[Mutable] Change namespace name

diff --git a/Engine/Plugins/Mutable/Source/CustomizableObject/Private/MuCO/CustomizableObjectInstance.cpp b/Engine/Plugins/Mutable/Source/CustomizableObject/Private/MuCO/CustomizableObjectInstance.cpp
--- a/Engine/Plugins/Mutable/Source/CustomizableObject/Private/MuCO/CustomizableObjectInstance.cpp
+++ b/Engine/Plugins/Mutable/Source/CustomizableObject/Private/MuCO/CustomizableObjectInstance.cpp
@@ -1361,7 +1361,7 @@ namespace
 		return Result;
 	}
 
-	FKAggregateGeom MakeAggGeomFromMutablePhysics(int32 BodyIndex, const mu::FPhysicsBody* MutablePhysicsBody)
+	FKAggregateGeom MakeAggGeomFromMutablePhysics(int32 BodyIndex, const UE::Mutable::Private::FPhysicsBody* MutablePhysicsBody)
 	{
 		FKAggregateGeom BodyAggGeom;
 
@@ -1506,7 +1506,7 @@ namespace
 	TObjectPtr<UPhysicsAsset> MakePhysicsAssetFromTemplateAndMutableBody(
 		const TSharedRef<FUpdateContextPrivate>& Context,
 		TObjectPtr<UPhysicsAsset> TemplateAsset,
-		const mu::FPhysicsBody* MutablePhysics,
+		const UE::Mutable::Private::FPhysicsBody* MutablePhysics,
 		FCustomizableObjectInstanceComponentIndex InstanceComponentIndex)
 	{
 		check(TemplateAsset);
@@ -1522,7 +1522,7 @@ namespace
 
 		Result->bNotForDedicatedServer = TemplateAsset->bNotForDedicatedServer;
 
-		const TMap<mu::FBoneName, TPair<FName, uint16>>& BoneInfoMap = Context->InstanceUpdateData.SkeletonsPerInstanceComponent[InstanceComponentIndex.GetValue()].BoneInfoMap;
+		const TMap<UE::Mutable::Private::FBoneName, TPair<FName, uint16>>& BoneInfoMap = Context->InstanceUpdateData.SkeletonsPerInstanceComponent[InstanceComponentIndex.GetValue()].BoneInfoMap;
 		TMap<FName, int32> BonesInUse;
 
 		const int32 MutablePhysicsBodyCount = MutablePhysics->GetBodyCount();
@@ -1649,7 +1649,7 @@ namespace
 UPhysicsAsset* GetOrBuildMainPhysicsAsset(
 	const TSharedRef<FUpdateContextPrivate>& Context,
 	TObjectPtr<UPhysicsAsset> TemplateAsset,
-	const mu::FPhysicsBody* MutablePhysics,
+	const UE::Mutable::Private::FPhysicsBody* MutablePhysics,
 	bool bDisableCollisionsBetweenDifferentAssets,
 	FCustomizableObjectInstanceComponentIndex InstanceComponentIndex)
 {
@@ -1711,7 +1711,7 @@ UPhysicsAsset* GetOrBuildMainPhysicsAsset(
 
 	Result->bNotForDedicatedServer = TemplateAsset->bNotForDedicatedServer;
 
-	const TMap<mu::FBoneName, TPair<FName, uint16>>& BoneInfoMap = Context->InstanceUpdateData.SkeletonsPerInstanceComponent[InstanceComponentIndex.GetValue()].BoneInfoMap;
+	const TMap<UE::Mutable::Private::FBoneName, TPair<FName, uint16>>& BoneInfoMap = Context->InstanceUpdateData.SkeletonsPerInstanceComponent[InstanceComponentIndex.GetValue()].BoneInfoMap;
 	TMap<FName, int32> BonesInUse;
 
 	const int32 MutablePhysicsBodyCount = MutablePhysics->GetBodyCount();
@@ -2010,7 +2010,7 @@ bool DoComponentsNeedUpdate(const TSharedRef<FUpdateContextPrivate>& Context, bo
 
 			if (!LOD.Mesh)
 			{
-				bHasInvalidMesh |= bComponentWithMesh; // Mesh will be invalid if the mu::mesh is missing on higher LODs.
+				bHasInvalidMesh |= bComponentWithMesh; // Mesh will be invalid if the UE::Mutable::Private::mesh is missing on higher LODs.
 				continue;
 			}
 
@@ -2034,7 +2034,7 @@ bool DoComponentsNeedUpdate(const TSharedRef<FUpdateContextPrivate>& Context, bo
 		bool bMeshChanged = true;
 
 		// Compare Mesh IDs 
-		const TArray<mu::FMeshId>* MeshIDs = Context->GetMeshDescriptors(ObjectComponentIndex);
+		const TArray<UE::Mutable::Private::FMeshId>* MeshIDs = Context->GetMeshDescriptors(ObjectComponentIndex);
 		if (ensure(MeshIDs))
 		{
 			if (const FCustomizableInstanceComponentData* ComponentData = InstancePrivate->GetComponentData(ObjectComponentIndex)) // Could be nullptr if the component has not been generated.
@@ -2129,7 +2129,7 @@ bool UCustomizableInstancePrivate::UpdateSkeletalMesh_PostBeginUpdate0(UCustomiz
 				continue;
 			}
 
-			if (ExtensionOutput.Data->Origin == mu::FExtensionData::EOrigin::Invalid)
+			if (ExtensionOutput.Data->Origin == UE::Mutable::Private::FExtensionData::EOrigin::Invalid)
 			{
 				// Null data was produced
 				//
@@ -2151,14 +2151,14 @@ bool UCustomizableInstancePrivate::UpdateSkeletalMesh_PostBeginUpdate0(UCustomiz
 			const FCustomizableObjectResourceData* ReferencedExtensionData = nullptr;
 			switch (ExtensionOutput.Data->Origin)
 			{
-				case mu::FExtensionData::EOrigin::ConstantAlwaysLoaded:
+				case UE::Mutable::Private::FExtensionData::EOrigin::ConstantAlwaysLoaded:
 				{
 					check(ModelResources.AlwaysLoadedExtensionData.IsValidIndex(ExtensionOutput.Data->Index));
 					ReferencedExtensionData = &ModelResources.AlwaysLoadedExtensionData[ExtensionOutput.Data->Index];
 				}
 				break;
 
-				case mu::FExtensionData::EOrigin::ConstantStreamed:
+				case UE::Mutable::Private::FExtensionData::EOrigin::ConstantStreamed:
 				{
 #if WITH_EDITOR
 					check(ModelResources.StreamedExtensionDataEditor.IsValidIndex(ExtensionOutput.Data->Index));
@@ -2266,7 +2266,7 @@ bool UCustomizableInstancePrivate::UpdateSkeletalMesh_PostBeginUpdate0(UCustomiz
 			return false;
 		}
 
-		TArray<mu::FMeshId>* MeshDescriptors = Context->GetMeshDescriptors(ObjectComponentIndex);
+		TArray<UE::Mutable::Private::FMeshId>* MeshDescriptors = Context->GetMeshDescriptors(ObjectComponentIndex);
 		check(MeshDescriptors);
 		check(MeshDescriptors->Num() == MAX_MESH_LOD_COUNT);
 		ComponentData->LastMeshIdPerLOD = *MeshDescriptors;
@@ -2289,7 +2289,7 @@ bool UCustomizableInstancePrivate::UpdateSkeletalMesh_PostBeginUpdate0(UCustomiz
 		uint32 MeshHash = INDEX_NONE;
 		for (int32 LODIndex = FirstLOD; LODIndex < Component->LODCount && LODIndex < MAX_MESH_LOD_COUNT; ++LODIndex)
 		{
-			const mu::FGeneratedMeshKey* MeshKey = (*MeshDescriptors)[LODIndex].GetKey();
+			const UE::Mutable::Private::FGeneratedMeshKey* MeshKey = (*MeshDescriptors)[LODIndex].GetKey();
 			check(MeshKey);
 			const uint32 MeshIDHash = GetTypeHashPersistent(*MeshKey);
 
@@ -2313,7 +2313,7 @@ bool UCustomizableInstancePrivate::UpdateSkeletalMesh_PostBeginUpdate0(UCustomiz
 		}
 
 		// We need the first valid mesh. get it from the component, considering that some LOSs may have been skipped.
-		TSharedPtr<const mu::FMesh> ComponentMesh;
+		TSharedPtr<const UE::Mutable::Private::FMesh> ComponentMesh;
 		for (int32 LODIndex = 0; LODIndex < Component->LODCount && !ComponentMesh; ++LODIndex)
 		{
 			ComponentMesh = Context->InstanceUpdateData.LODs[Component->FirstLOD + LODIndex].Mesh;
@@ -2341,7 +2341,7 @@ bool UCustomizableInstancePrivate::UpdateSkeletalMesh_PostBeginUpdate0(UCustomiz
         		UE_LOG(LogMutable, Error, TEXT("Referenced mesh [%s] was not pre-loaded. It will be sync-loaded probably causing a hitch. CO [%s]"), *Ref.ToString(), *GetNameSafe(CustomizableObject));
         	}
 
-			UStreamableRenderAsset* Asset = MutablePrivate::LoadObject(Ref);
+			UStreamableRenderAsset* Asset = UE::Mutable::Private::LoadObject(Ref);
 			USkeletalMesh* SkeletalMesh = Cast<USkeletalMesh>(Asset);
 			if (SkeletalMesh)
 			{
@@ -2407,7 +2407,7 @@ bool UCustomizableInstancePrivate::UpdateSkeletalMesh_PostBeginUpdate0(UCustomiz
 		}
 
 		// Build PhysicsAsset merging physics assets coming from SubMeshes of the newly generated Mesh
-		if (TSharedPtr<const mu::FPhysicsBody> MutablePhysics = ComponentMesh->GetPhysicsBody())
+		if (TSharedPtr<const UE::Mutable::Private::FPhysicsBody> MutablePhysics = ComponentMesh->GetPhysicsBody())
 		{
 			constexpr bool bDisallowCollisionBetweenAssets = true;
 			UPhysicsAsset* PhysicsAssetResult = GetOrBuildMainPhysicsAsset(
@@ -2427,7 +2427,7 @@ bool UCustomizableInstancePrivate::UpdateSkeletalMesh_PostBeginUpdate0(UCustomiz
 		const int32 NumAdditionalPhysicsNum = ComponentMesh->AdditionalPhysicsBodies.Num();
 		for (int32 I = 0; I < NumAdditionalPhysicsNum; ++I)
 		{
-			TSharedPtr<const mu::FPhysicsBody> AdditionalPhysicsBody = ComponentMesh->AdditionalPhysicsBodies[I];
+			TSharedPtr<const UE::Mutable::Private::FPhysicsBody> AdditionalPhysicsBody = ComponentMesh->AdditionalPhysicsBodies[I];
 			
 			check(AdditionalPhysicsBody);
 			if (!AdditionalPhysicsBody->bBodiesModified)
@@ -2440,7 +2440,7 @@ bool UCustomizableInstancePrivate::UpdateSkeletalMesh_PostBeginUpdate0(UCustomiz
 			const FAnimBpOverridePhysicsAssetsInfo& Info = ModelResources.AnimBpOverridePhysiscAssetsInfo[PhysicsBodyExternalId];
 
 			// Make sure the AnimInstance class is loaded. It is expected to be already loaded at this point though. 
-			UClass* AnimInstanceClassLoaded = MutablePrivate::LoadClass(Info.AnimInstanceClass);
+			UClass* AnimInstanceClassLoaded = UE::Mutable::Private::LoadClass(Info.AnimInstanceClass);
 			TSubclassOf<UAnimInstance> AnimInstanceClass = TSubclassOf<UAnimInstance>(AnimInstanceClassLoaded);
 			if (!ensureAlways(AnimInstanceClass))
 			{
@@ -2494,7 +2494,7 @@ bool UCustomizableInstancePrivate::UpdateSkeletalMesh_PostBeginUpdate0(UCustomiz
 		
 		if (Context->bUseMeshCache)
 		{
-			const TArray<mu::FMeshId>* MeshId = Context->GetMeshDescriptors(ObjectComponentIndex);
+			const TArray<UE::Mutable::Private::FMeshId>* MeshId = Context->GetMeshDescriptors(ObjectComponentIndex);
 			if (MeshId)
 			{
 				CustomizableObject->GetPrivate()->MeshCache.Add(*MeshId, CastChecked<UCustomizableObjectSkeletalMesh>(SkeletalMesh));
@@ -3249,7 +3249,7 @@ void UCustomizableInstancePrivate::SetReferenceSkeletalMesh() const
 			Parent->EmptyOverrideMaterials();
 
 			TSoftObjectPtr<USkeletalMesh> SoftObjectPtr = ModelResources->ReferenceSkeletalMeshesData[ObjectComponentIndex].SoftSkeletalMesh;
-			USkeletalMesh* SkeletalMesh = MutablePrivate::LoadObject(SoftObjectPtr);
+			USkeletalMesh* SkeletalMesh = UE::Mutable::Private::LoadObject(SoftObjectPtr);
 			Parent->SetSkeletalMesh(SkeletalMesh);
 		}
 	}
@@ -3344,7 +3344,7 @@ void FMutableUpdateCandidate::ApplyLODUpdateParamsToInstance(FUpdateContextPriva
 
 // The memory allocated in the function and pointed by the returned pointer is owned by the caller and must be freed. 
 // If assigned to a UTexture2D, it will be freed by that UTexture2D
-FTexturePlatformData* MutableCreateImagePlatformData(TSharedPtr<const mu::FImage> MutableImage, int32 OnlyLOD, uint16 FullSizeX, uint16 FullSizeY)
+FTexturePlatformData* MutableCreateImagePlatformData(TSharedPtr<const UE::Mutable::Private::FImage> MutableImage, int32 OnlyLOD, uint16 FullSizeX, uint16 FullSizeY)
 {
 	int32 SizeX = FMath::Max(MutableImage->GetSize()[0], FullSizeX);
 	int32 SizeY = FMath::Max(MutableImage->GetSize()[1], FullSizeY);
@@ -3414,7 +3414,7 @@ FTexturePlatformData* MutableCreateImagePlatformData(TSharedPtr<const mu::FImage
 
 	int32 EndLOD = OnlyLOD < 0 ? FullLODCount : FirstLOD + 1;
 	
-	mu::EImageFormat MutableFormat = MutableImage->GetFormat();
+	UE::Mutable::Private::EImageFormat MutableFormat = MutableImage->GetFormat();
 
 	int32 MaxPossibleSize = 0;
 		
@@ -3435,49 +3435,49 @@ FTexturePlatformData* MutableCreateImagePlatformData(TSharedPtr<const mu::FImage
 		//return nullptr;
 	}
 
-	mu::FImageOperator ImOp = mu::FImageOperator::GetDefault(mu::FImageOperator::FImagePixelFormatFunc());
+	UE::Mutable::Private::FImageOperator ImOp = UE::Mutable::Private::FImageOperator::GetDefault(UE::Mutable::Private::FImageOperator::FImagePixelFormatFunc());
 
 	EPixelFormat PlatformFormat = PF_Unknown;
 	switch (MutableFormat)
 	{
-	case mu::EImageFormat::RGB_UByte:
+	case UE::Mutable::Private::EImageFormat::RGB_UByte:
 		// performance penalty. can happen in states that remove compression.
 		PlatformFormat = PF_R8G8B8A8;	
 		UE_LOG(LogMutable, Display, TEXT("Building instance: a texture was generated in a format not supported by the hardware (RGB), this results in an additional conversion, so a performance penalty."));
 		break; 
 
-	case mu::EImageFormat::BGRA_UByte:			
+	case UE::Mutable::Private::EImageFormat::BGRA_UByte:			
 		// performance penalty. can happen with texture parameter images.
 		PlatformFormat = PF_R8G8B8A8;	
 		UE_LOG(LogMutable, Display, TEXT("Building instance: a texture was generated in a format not supported by the hardware (BGRA), this results in an additional conversion, so a performance penalty."));
 		break;
 
 	// Good cases:
-	case mu::EImageFormat::RGBA_UByte:		PlatformFormat = PF_R8G8B8A8;	break;
-	case mu::EImageFormat::BC1:				PlatformFormat = PF_DXT1;		break;
-	case mu::EImageFormat::BC2:				PlatformFormat = PF_DXT3;		break;
-	case mu::EImageFormat::BC3:				PlatformFormat = PF_DXT5;		break;
-	case mu::EImageFormat::BC4:				PlatformFormat = PF_BC4;		break;
-	case mu::EImageFormat::BC5:				PlatformFormat = PF_BC5;		break;
-	case mu::EImageFormat::L_UByte:			PlatformFormat = PF_G8;			break;
-	case mu::EImageFormat::ASTC_4x4_RGB_LDR:	PlatformFormat = PF_ASTC_4x4;	break;
-	case mu::EImageFormat::ASTC_4x4_RGBA_LDR:PlatformFormat = PF_ASTC_4x4;	break;
-	case mu::EImageFormat::ASTC_4x4_RG_LDR:	PlatformFormat = PF_ASTC_4x4;	break;
+	case UE::Mutable::Private::EImageFormat::RGBA_UByte:		PlatformFormat = PF_R8G8B8A8;	break;
+	case UE::Mutable::Private::EImageFormat::BC1:				PlatformFormat = PF_DXT1;		break;
+	case UE::Mutable::Private::EImageFormat::BC2:				PlatformFormat = PF_DXT3;		break;
+	case UE::Mutable::Private::EImageFormat::BC3:				PlatformFormat = PF_DXT5;		break;
+	case UE::Mutable::Private::EImageFormat::BC4:				PlatformFormat = PF_BC4;		break;
+	case UE::Mutable::Private::EImageFormat::BC5:				PlatformFormat = PF_BC5;		break;
+	case UE::Mutable::Private::EImageFormat::L_UByte:			PlatformFormat = PF_G8;			break;
+	case UE::Mutable::Private::EImageFormat::ASTC_4x4_RGB_LDR:	PlatformFormat = PF_ASTC_4x4;	break;
+	case UE::Mutable::Private::EImageFormat::ASTC_4x4_RGBA_LDR:PlatformFormat = PF_ASTC_4x4;	break;
+	case UE::Mutable::Private::EImageFormat::ASTC_4x4_RG_LDR:	PlatformFormat = PF_ASTC_4x4;	break;
 	default:
 		// Cannot prepare texture if it's not in the right format, this can happen if mutable is in debug mode or in case of bugs
 		UE_LOG(LogMutable, Warning, TEXT("Building instance: a texture was generated in an unsupported format, it will be converted to Unreal with a performance penalty."));
 
-		switch (mu::GetImageFormatData(MutableFormat).Channels)
+		switch (UE::Mutable::Private::GetImageFormatData(MutableFormat).Channels)
 		{
 		case 1:
 			PlatformFormat = PF_R8;
-			MutableImage = ImOp.ImagePixelFormat(0, MutableImage.Get(), mu::EImageFormat::L_UByte);
+			MutableImage = ImOp.ImagePixelFormat(0, MutableImage.Get(), UE::Mutable::Private::EImageFormat::L_UByte);
 			break;
 		case 2:
 		case 3:
 		case 4:
 			PlatformFormat = PF_R8G8B8A8;
-			MutableImage = ImOp.ImagePixelFormat(0, MutableImage.Get(), mu::EImageFormat::RGBA_UByte);
+			MutableImage = ImOp.ImagePixelFormat(0, MutableImage.Get(), UE::Mutable::Private::EImageFormat::RGBA_UByte);
 			break;
 		default: 
 			// Absolutely worst case
@@ -3527,13 +3527,13 @@ FTexturePlatformData* MutableCreateImagePlatformData(TSharedPtr<const mu::FImage
 			const uint8* MutableData = MutableImage->GetLODData(MipLevelMutable);
 			const uint32 SourceDataSize = MutableImage->GetLODDataSize(MipLevelMutable);
 
-			uint32 DestDataSize = (MutableFormat == mu::EImageFormat::RGB_UByte)
+			uint32 DestDataSize = (MutableFormat == UE::Mutable::Private::EImageFormat::RGB_UByte)
 					? (SourceDataSize/3) * 4
 					: SourceDataSize;
 			void* pData = Mip->BulkData.Realloc(DestDataSize);
 
 			// Special inefficient cases
-			if (MutableFormat== mu::EImageFormat::BGRA_UByte)
+			if (MutableFormat== UE::Mutable::Private::EImageFormat::BGRA_UByte)
 			{
 				check(SourceDataSize==DestDataSize);
 
@@ -3549,7 +3549,7 @@ FTexturePlatformData* MutableCreateImagePlatformData(TSharedPtr<const mu::FImage
 				}
 			}
 
-			else if (MutableFormat == mu::EImageFormat::RGB_UByte)
+			else if (MutableFormat == UE::Mutable::Private::EImageFormat::RGB_UByte)
 			{
 				MUTABLE_CPUPROFILER_SCOPE(Innefficent_RGB_Format_Conversion);
 
@@ -3617,33 +3617,33 @@ FTexturePlatformData* MutableCreateImagePlatformData(TSharedPtr<const mu::FImage
 }
 
 
-void ConvertImage(UTexture2D* Texture, TSharedPtr<const mu::FImage> MutableImage, const FMutableModelImageProperties& Props, int OnlyLOD, int32 ExtractChannel)
+void ConvertImage(UTexture2D* Texture, TSharedPtr<const UE::Mutable::Private::FImage> MutableImage, const FMutableModelImageProperties& Props, int OnlyLOD, int32 ExtractChannel)
 {
 	MUTABLE_CPUPROFILER_SCOPE(UCustomizableInstancePrivate::ConvertImage);
 
 	SetTexturePropertiesFromMutableImageProps(Texture, Props, false);
 
-	mu::EImageFormat MutableFormat = MutableImage->GetFormat();
+	UE::Mutable::Private::EImageFormat MutableFormat = MutableImage->GetFormat();
 
 	// Extract a single channel, if requested.
 	if (ExtractChannel >= 0)
 	{
-		mu::FImageOperator ImOp = mu::FImageOperator::GetDefault(mu::FImageOperator::FImagePixelFormatFunc());
+		UE::Mutable::Private::FImageOperator ImOp = UE::Mutable::Private::FImageOperator::GetDefault(UE::Mutable::Private::FImageOperator::FImagePixelFormatFunc());
 
-		MutableImage = ImOp.ImagePixelFormat( 4, MutableImage.Get(), mu::EImageFormat::RGBA_UByte );
+		MutableImage = ImOp.ImagePixelFormat( 4, MutableImage.Get(), UE::Mutable::Private::EImageFormat::RGBA_UByte );
 
 		uint8_t Channel = uint8_t( FMath::Clamp(ExtractChannel,0,3) );
-		MutableImage = ImOp.ImageSwizzle( mu::EImageFormat::L_UByte, &MutableImage, &Channel );
-		MutableFormat = mu::EImageFormat::L_UByte;
+		MutableImage = ImOp.ImageSwizzle( UE::Mutable::Private::EImageFormat::L_UByte, &MutableImage, &Channel );
+		MutableFormat = UE::Mutable::Private::EImageFormat::L_UByte;
 	}
 
 	// Hack: This format is unsupported in UE, but it shouldn't happen in production.
-	if (MutableFormat == mu::EImageFormat::RGB_UByte)
+	if (MutableFormat == UE::Mutable::Private::EImageFormat::RGB_UByte)
 	{
 		UE_LOG(LogMutable, Warning, TEXT("Building instance: a texture was generated in RGB format, which is slow to convert to Unreal."));
 
 		// Expand the image.
-		TSharedPtr<mu::FImage> Converted = MakeShared<mu::FImage>(MutableImage->GetSizeX(), MutableImage->GetSizeY(), MutableImage->GetLODCount(), mu::EImageFormat::RGBA_UByte, mu::EInitializationType::NotInitialized);
+		TSharedPtr<UE::Mutable::Private::FImage> Converted = MakeShared<UE::Mutable::Private::FImage>(MutableImage->GetSizeX(), MutableImage->GetSizeY(), MutableImage->GetLODCount(), UE::Mutable::Private::EImageFormat::RGBA_UByte, UE::Mutable::Private::EInitializationType::NotInitialized);
 
 		for (int32 LODIndex = 0; LODIndex < Converted->GetLODCount(); ++LODIndex)
 		{
@@ -3661,14 +3661,14 @@ void ConvertImage(UTexture2D* Texture, TSharedPtr<const mu::FImage> MutableImage
 
 		MutableImage = Converted;
 	}
-	else if (MutableFormat == mu::EImageFormat::BGRA_UByte)
+	else if (MutableFormat == UE::Mutable::Private::EImageFormat::BGRA_UByte)
 	{
 		UE_LOG(LogMutable, Warning, TEXT("Building instance: a texture was generated in BGRA format, which is slow to convert to Unreal."));
 
 		MUTABLE_CPUPROFILER_SCOPE(Swizzle);
 		// Swizzle the image.
 		// \TODO: Raise a warning?
-		TSharedPtr<mu::FImage> Converted = MakeShared<mu::FImage>(MutableImage->GetSizeX(), MutableImage->GetSizeY(), 1, mu::EImageFormat::RGBA_UByte, mu::EInitializationType::NotInitialized);
+		TSharedPtr<UE::Mutable::Private::FImage> Converted = MakeShared<UE::Mutable::Private::FImage>(MutableImage->GetSizeX(), MutableImage->GetSizeY(), 1, UE::Mutable::Private::EImageFormat::RGBA_UByte, UE::Mutable::Private::EInitializationType::NotInitialized);
 		int32 PixelCount = MutableImage->GetSizeX() * MutableImage->GetSizeY();
 
 		const uint8* pSource = MutableImage->GetLODData(0);
@@ -3889,7 +3889,7 @@ bool BuildSkeletonData(const TSharedRef<FUpdateContextPrivate>& Context, USkelet
 	const int32 RawRefBoneCount = ReferenceSkeleton.GetRawBoneNum();
 
 	const TArray<FInstanceUpdateData::FBone>& BonePose = Context->InstanceUpdateData.SkeletonsPerInstanceComponent[InstanceComponentIndex.GetValue()].BonePose;
-	TMap<mu::FBoneName, TPair<FName,uint16>>& BoneInfoMap = Context->InstanceUpdateData.SkeletonsPerInstanceComponent[InstanceComponentIndex.GetValue()].BoneInfoMap;
+	TMap<UE::Mutable::Private::FBoneName, TPair<FName,uint16>>& BoneInfoMap = Context->InstanceUpdateData.SkeletonsPerInstanceComponent[InstanceComponentIndex.GetValue()].BoneInfoMap;
 
 	{
 		MUTABLE_CPUPROFILER_SCOPE(BuildSkeletonData_BuildBoneInfoMap);
@@ -3899,7 +3899,7 @@ bool BuildSkeletonData(const TSharedRef<FUpdateContextPrivate>& Context, USkelet
 		for (int32 Index = 0; Index < RawRefBoneCount; ++Index)
 		{
 			const FName BoneName = RawRefBoneInfo[Index].Name;
-			if(const mu::FBoneName* Bone = Context->BoneNames->Find(BoneName))
+			if(const UE::Mutable::Private::FBoneName* Bone = Context->BoneNames->Find(BoneName))
 			{
 				TPair<FName, uint16>& BoneInfo = BoneInfoMap.Add(*Bone);
 				BoneInfo.Key = BoneName;
@@ -4011,7 +4011,7 @@ bool BuildSkeletonData(const TSharedRef<FUpdateContextPrivate>& Context, USkelet
 }
 
 
-void BuildMeshSockets(USkeletalMesh* SkeletalMesh, const UModelResources& ModelResources, const FMutableRefSkeletalMeshData& RefSkeletalMeshData, TSharedPtr<const mu::FMesh> MutableMesh)
+void BuildMeshSockets(USkeletalMesh* SkeletalMesh, const UModelResources& ModelResources, const FMutableRefSkeletalMeshData& RefSkeletalMeshData, TSharedPtr<const UE::Mutable::Private::FMesh> MutableMesh)
 {
 	// Build mesh sockets.
 	MUTABLE_CPUPROFILER_SCOPE(UCustomizableInstancePrivate::BuildMeshSockets);
@@ -4303,7 +4303,7 @@ void BuildOrCopyClothingData(
 		{
 			const FInstanceUpdateData::FLOD& LOD = Context->InstanceUpdateData.LODs[Component->FirstLOD+LODIndex];
 
-			if (TSharedPtr<const mu::FMesh> MutableMesh = LOD.Mesh)
+			if (TSharedPtr<const UE::Mutable::Private::FMesh> MutableMesh = LOD.Mesh)
 			{
 				for (int32 SectionIndex = 0; SectionIndex < MutableMesh->GetSurfaceCount(); ++SectionIndex)
 				{
@@ -4312,7 +4312,7 @@ void BuildOrCopyClothingData(
 						continue;
 					}
 					
-					const mu::FSurfaceSubMesh& SubMesh = MutableMesh->Surfaces[SectionIndex].SubMeshes[0];
+					const UE::Mutable::Private::FSurfaceSubMesh& SubMesh = MutableMesh->Surfaces[SectionIndex].SubMeshes[0];
 					const uint32 ClothResourceId = ModelResources.MeshMetadata[SubMesh.ExternalId].ClothingMetadataId;
 
 					if (ClothResourceId == 0)
@@ -4412,7 +4412,7 @@ void BuildOrCopyClothingData(
 		{
 			const FInstanceUpdateData::FLOD& LOD = Context->InstanceUpdateData.LODs[Component->FirstLOD+LODIndex];
 
-			if (TSharedPtr<const mu::FMesh> MutableMesh = LOD.Mesh)
+			if (TSharedPtr<const UE::Mutable::Private::FMesh> MutableMesh = LOD.Mesh)
 			{
 				GetSectionClothData(*MutableMesh, LODIndex, Context->InstanceUpdateData.ClothingMeshData, SectionsClothData, NumClothingDataNotFound);
 			}
@@ -4592,7 +4592,7 @@ void BuildOrCopyClothingData(
 	{
 		const FInstanceUpdateData::FLOD& LOD = Context->InstanceUpdateData.LODs[Component->FirstLOD + LODIndex];
 		
-		TSharedPtr<const mu::FMesh> MutableMesh = LOD.Mesh;
+		TSharedPtr<const UE::Mutable::Private::FMesh> MutableMesh = LOD.Mesh;
 		if (!MutableMesh)
 		{
 			continue;
@@ -4729,11 +4729,11 @@ bool BuildOrCopyRenderData(const TSharedRef<FUpdateContextPrivate>& Context, USk
 		
 		FSkeletalMeshLODRenderData& LODResource = RenderData->LODRenderData[LODIndex];
 
-		const TMap<mu::FBoneName, TPair<FName, uint16>>& BoneInfoMap = Context->InstanceUpdateData.SkeletonsPerInstanceComponent[InstanceComponentIndex.GetValue()].BoneInfoMap;
+		const TMap<UE::Mutable::Private::FBoneName, TPair<FName, uint16>>& BoneInfoMap = Context->InstanceUpdateData.SkeletonsPerInstanceComponent[InstanceComponentIndex.GetValue()].BoneInfoMap;
 		
 		// Set active and required bones
 		{
-			const TArray<mu::FBoneName>& ActiveBones = Context->InstanceUpdateData.ActiveBones;
+			const TArray<UE::Mutable::Private::FBoneName>& ActiveBones = Context->InstanceUpdateData.ActiveBones;
 			LODResource.ActiveBoneIndices.Reserve(LOD.ActiveBoneCount);
 
 			for (uint32 Index = 0; Index < LOD.ActiveBoneCount; ++Index)
@@ -4771,7 +4771,7 @@ bool BuildOrCopyRenderData(const TSharedRef<FUpdateContextPrivate>& Context, USk
 			else
 			{
 				// In case the surface does not have metadata, check if any submesh has surface metadata.
-				for (const mu::FSurfaceSubMesh& SubMesh : LOD.Mesh->Surfaces[MeshSectionIndex].SubMeshes)	
+				for (const UE::Mutable::Private::FSurfaceSubMesh& SubMesh : LOD.Mesh->Surfaces[MeshSectionIndex].SubMeshes)	
 				{
 					const FMutableMeshMetadata* FoundMeshMetadata = ModelResources.MeshMetadata.Find(SubMesh.ExternalId);
 
@@ -4807,7 +4807,7 @@ bool BuildOrCopyRenderData(const TSharedRef<FUpdateContextPrivate>& Context, USk
 		// Active SkinWeightProfiles ID and Name
 		TArray<TPair<uint32, FName>> ActiveSkinWeightProfiles;
 		
-		const mu::FMeshBufferSet& MutableMeshVertexBuffers = LOD.Mesh->GetVertexBuffers();
+		const UE::Mutable::Private::FMeshBufferSet& MutableMeshVertexBuffers = LOD.Mesh->GetVertexBuffers();
 		const int32 NumBuffers = MutableMeshVertexBuffers.GetBufferCount();
 
 		for (int32 BufferIndex = 0; BufferIndex < NumBuffers; ++BufferIndex)
@@ -4817,11 +4817,11 @@ bool BuildOrCopyRenderData(const TSharedRef<FUpdateContextPrivate>& Context, USk
 				continue;
 			}
 
-			mu::EMeshBufferSemantic Semantic;
+			UE::Mutable::Private::EMeshBufferSemantic Semantic;
 			int32 SemanticIndex;
 			MutableMeshVertexBuffers.GetChannel(BufferIndex, 0, &Semantic, &SemanticIndex, nullptr, nullptr, nullptr);
 
-			if (Semantic != mu::EMeshBufferSemantic::AltSkinWeight)
+			if (Semantic != UE::Mutable::Private::EMeshBufferSemantic::AltSkinWeight)
 			{
 				continue;
 			}
@@ -5094,7 +5094,7 @@ UE::Tasks::FTask LoadAdditionalAssetsAndData(const TSharedRef<FUpdateContextPriv
 			const FInstanceUpdateData::FComponent& Component = Components[InstanceComponentIndex];
 			FCustomizableInstanceComponentData* ComponentData = InstancePrivate->GetComponentData(Component.Id);
 
-			TSharedPtr<const mu::FMesh> FirstComponentMesh = Context->InstanceUpdateData.LODs.IsValidIndex(Component.FirstLOD) ?
+			TSharedPtr<const UE::Mutable::Private::FMesh> FirstComponentMesh = Context->InstanceUpdateData.LODs.IsValidIndex(Component.FirstLOD) ?
 				Context->InstanceUpdateData.LODs[Component.FirstLOD].Mesh :
 				nullptr;
 
@@ -5126,7 +5126,7 @@ UE::Tasks::FTask LoadAdditionalAssetsAndData(const TSharedRef<FUpdateContextPriv
 					{
 						const FInstanceUpdateData::FLOD& LOD = Context->InstanceUpdateData.LODs[Component.FirstLOD + LODIndex];
 
-						TSharedPtr<const mu::FMesh> MutableMesh = LOD.Mesh;
+						TSharedPtr<const UE::Mutable::Private::FMesh> MutableMesh = LOD.Mesh;
 						if (!MutableMesh)
 						{
 							continue;
@@ -5144,7 +5144,7 @@ UE::Tasks::FTask LoadAdditionalAssetsAndData(const TSharedRef<FUpdateContextPriv
 					{
 						const FInstanceUpdateData::FLOD& LOD = Context->InstanceUpdateData.LODs[Component.FirstLOD + LODIndex];
 
-						TSharedPtr<const mu::FMesh> MutableMesh = LOD.Mesh;
+						TSharedPtr<const UE::Mutable::Private::FMesh> MutableMesh = LOD.Mesh;
 						if (!MutableMesh)
 						{
 							continue;
@@ -5162,7 +5162,7 @@ UE::Tasks::FTask LoadAdditionalAssetsAndData(const TSharedRef<FUpdateContextPriv
 					{
 						const FInstanceUpdateData::FLOD& LOD = Context->InstanceUpdateData.LODs[Component.FirstLOD + LODIndex];
 
-						TSharedPtr<const mu::FMesh> MutableMesh = LOD.Mesh;
+						TSharedPtr<const UE::Mutable::Private::FMesh> MutableMesh = LOD.Mesh;
 						if (!MutableMesh)
 						{
 							continue;
@@ -5179,7 +5179,7 @@ UE::Tasks::FTask LoadAdditionalAssetsAndData(const TSharedRef<FUpdateContextPriv
 					{
 						const FInstanceUpdateData::FLOD& LOD = Context->InstanceUpdateData.LODs[Component.FirstLOD + LODIndex];
 
-						TSharedPtr<const mu::FMesh> MutableMesh = LOD.Mesh;
+						TSharedPtr<const UE::Mutable::Private::FMesh> MutableMesh = LOD.Mesh;
 						if (!MutableMesh)
 						{
 							continue;
@@ -5240,7 +5240,7 @@ UE::Tasks::FTask LoadAdditionalAssetsAndData(const TSharedRef<FUpdateContextPriv
 			{
 				const FInstanceUpdateData::FLOD& LOD = Context->InstanceUpdateData.LODs[Component.FirstLOD + LODIndex];
 
-				TSharedPtr<const mu::FMesh> MutableMesh = LOD.Mesh;
+				TSharedPtr<const UE::Mutable::Private::FMesh> MutableMesh = LOD.Mesh;
 				if (!MutableMesh)
 				{
 					continue;
@@ -5403,7 +5403,7 @@ UE::Tasks::FTask LoadAdditionalAssetsAndData(const TSharedRef<FUpdateContextPriv
 	// Copy FExtensionData Object node input from the Instance to the InstanceUpdateData
 	for (int32 ExtensionDataIndex = 0; ExtensionDataIndex < Context->MutableInstance->GetExtensionDataCount(); ExtensionDataIndex++)
 	{
-		TSharedPtr<const mu::FExtensionData> ExtensionData;
+		TSharedPtr<const UE::Mutable::Private::FExtensionData> ExtensionData;
 		FName Name;
 		Context->MutableInstance->GetExtensionData(ExtensionDataIndex, ExtensionData, Name);
 
@@ -5595,7 +5595,7 @@ void AdditionalAssetsAsyncLoaded(const TSharedRef<FUpdateContextPrivate> Context
 			{
 				if (const FCustomizableObjectAssetUserData* AUDResource = ModelResources.StreamedResourceDataEditor[ResourceIndex].Data.GetPtr<FCustomizableObjectAssetUserData>())
 				{
-					ComponentData.AssetUserDataArray.Add(MutablePrivate::LoadObject(AUDResource->AssetUserDataEditor)); // Already loaded
+					ComponentData.AssetUserDataArray.Add(UE::Mutable::Private::LoadObject(AUDResource->AssetUserDataEditor)); // Already loaded
 				}
 			}
 #else
@@ -6010,14 +6010,14 @@ void UCustomizableInstancePrivate::BuildMaterials(const TSharedRef<FUpdateContex
 				// For now use the first found, an aggregate may be needed. 
 				if (Surface.SurfaceMetadataId == 0 && LOD.Mesh)
 				{
-					int32 MeshSurfaceIndex = LOD.Mesh->Surfaces.IndexOfByPredicate([SurfaceId = Surface.SurfaceId](const mu::FMeshSurface& Surface)
+					int32 MeshSurfaceIndex = LOD.Mesh->Surfaces.IndexOfByPredicate([SurfaceId = Surface.SurfaceId](const UE::Mutable::Private::FMeshSurface& Surface)
 					{
 						return SurfaceId == Surface.Id;
 					});
 
 					if (MeshSurfaceIndex != INDEX_NONE)
 					{
-						for (const mu::FSurfaceSubMesh& SubMesh : LOD.Mesh->Surfaces[SurfaceIndex].SubMeshes)	
+						for (const UE::Mutable::Private::FSurfaceSubMesh& SubMesh : LOD.Mesh->Surfaces[SurfaceIndex].SubMeshes)	
 						{
 							const FMutableMeshMetadata* FoundMeshMetadata = ModelResources.MeshMetadata.Find(SubMesh.ExternalId);
 
@@ -6149,7 +6149,7 @@ void UCustomizableInstancePrivate::BuildMaterials(const TSharedRef<FUpdateContex
 					{
 						const FInstanceUpdateData::FImage& Image = Context->InstanceUpdateData.Images[Surface.FirstImage + ImageIndex];
 						FString KeyName = Image.Name.ToString();
-						TSharedPtr<const mu::FImage> MutableImage = Image.Image;
+						TSharedPtr<const UE::Mutable::Private::FImage> MutableImage = Image.Image;
 
 						UTexture2D* MutableTexture = nullptr; // Texture generated by mutable
 						UTexture* PassThroughTexture = nullptr; // Texture not generated by mutable
@@ -6189,7 +6189,7 @@ void UCustomizableInstancePrivate::BuildMaterials(const TSharedRef<FUpdateContex
 										TEXT("Pass-through texture with name %s hasn't been loaded yet in BuildMaterials(). Forcing sync load."),
 										*Ref.ToSoftObjectPath().ToString());
 									ensure(false);
-									PassThroughTexture = MutablePrivate::LoadObject(Ref);
+									PassThroughTexture = UE::Mutable::Private::LoadObject(Ref);
 								}
 							}
 
@@ -6217,7 +6217,7 @@ void UCustomizableInstancePrivate::BuildMaterials(const TSharedRef<FUpdateContex
 
 								FString MutableTextureName;
 								{
-									const uint32 MutableImageHash = mu::GetTypeHashPersistent(*Image.ImageID.GetKey());
+									const uint32 MutableImageHash = UE::Mutable::Private::GetTypeHashPersistent(*Image.ImageID.GetKey());
 									MutableTextureName = FString::Printf(TEXT("T_%s_%lu"), *Props.TextureParameterName, MutableImageHash);
 									MutableTextureName.ReplaceInline(TEXT(" "), TEXT("_"));
 
@@ -6975,7 +6975,7 @@ const FGameplayTagContainer& UCustomizableObjectInstance::GetAnimationGameplayTa
 	return GetPrivate()->AnimBPGameplayTags;
 }
 
-namespace MutablePrivate
+namespace UE::Mutable::Private
 {
 	template<typename DELEGATE>
 	void InternalForEachAnimInstance(UCustomizableInstancePrivate* Private, FName ComponentName, DELEGATE Delegate)
@@ -7028,7 +7028,7 @@ namespace MutablePrivate
 			if (!AnimBP.IsNull())
 			{
 				// force a load right now - we don't know whether we would have loaded already - this could be called in editor
-				const TSubclassOf<UAnimInstance> LiveAnimBP = MutablePrivate::LoadClass(AnimBP);
+				const TSubclassOf<UAnimInstance> LiveAnimBP = UE::Mutable::Private::LoadClass(AnimBP);
 				if (LiveAnimBP)
 				{
 					Delegate.Execute(Index, LiveAnimBP);
@@ -7042,13 +7042,13 @@ namespace MutablePrivate
 
 void UCustomizableObjectInstance::ForEachComponentAnimInstance(FName ComponentName, FEachComponentAnimInstanceClassDelegate Delegate) const
 {
-	MutablePrivate::InternalForEachAnimInstance<>(GetPrivate(), ComponentName, Delegate);
+	UE::Mutable::Private::InternalForEachAnimInstance<>(GetPrivate(), ComponentName, Delegate);
 }
 
 
 void UCustomizableObjectInstance::ForEachComponentAnimInstance(FName ComponentName, FEachComponentAnimInstanceClassNativeDelegate Delegate) const
 {
-	MutablePrivate::InternalForEachAnimInstance<>(GetPrivate(), ComponentName, Delegate);
+	UE::Mutable::Private::InternalForEachAnimInstance<>(GetPrivate(), ComponentName, Delegate);
 }
 
 // Deprecated
@@ -7058,7 +7058,7 @@ void UCustomizableObjectInstance::ForEachAnimInstance(int32 ObjectComponentIndex
 	if (CO)
 	{
 		FName ComponentName = CO->GetPrivate()->GetComponentName(FCustomizableObjectComponentIndex(ObjectComponentIndex));
-		MutablePrivate::InternalForEachAnimInstance<>(GetPrivate(), ComponentName, Delegate);
+		UE::Mutable::Private::InternalForEachAnimInstance<>(GetPrivate(), ComponentName, Delegate);
 	}
 }
 
@@ -7070,7 +7070,7 @@ void UCustomizableObjectInstance::ForEachAnimInstance(int32 ObjectComponentIndex
 	if (CO)
 	{
 		FName ComponentName = CO->GetPrivate()->GetComponentName(FCustomizableObjectComponentIndex(ObjectComponentIndex));
-		MutablePrivate::InternalForEachAnimInstance<>(GetPrivate(), ComponentName, Delegate);
+		UE::Mutable::Private::InternalForEachAnimInstance<>(GetPrivate(), ComponentName, Delegate);
 	}
 }
 
diff --git a/Engine/Plugins/Mutable/Source/CustomizableObject/Private/MuCO/CustomizableObjectInstanceDescriptor.cpp b/Engine/Plugins/Mutable/Source/CustomizableObject/Private/MuCO/CustomizableObjectInstanceDescriptor.cpp
--- a/Engine/Plugins/Mutable/Source/CustomizableObject/Private/MuCO/CustomizableObjectInstanceDescriptor.cpp
+++ b/Engine/Plugins/Mutable/Source/CustomizableObject/Private/MuCO/CustomizableObjectInstanceDescriptor.cpp
@@ -572,7 +572,7 @@ void FCustomizableObjectInstanceDescriptor::SetBuildParameterRelevancy(const boo
 }
  
  
-TSharedPtr<mu::FParameters> FCustomizableObjectInstanceDescriptor::GetParameters() const
+TSharedPtr<UE::Mutable::Private::FParameters> FCustomizableObjectInstanceDescriptor::GetParameters() const
 {
 	if (!CustomizableObject)
 	{
@@ -592,7 +592,7 @@ TSharedPtr<mu::FParameters> FCustomizableObjectInstanceDescriptor::GetParameters
  
 	const UModelResources& ModelResources = CustomizableObject->GetPrivate()->GetModelResourcesChecked();
 	
-	TSharedPtr<mu::FParameters> MutableParameters = mu::FModel::NewParameters(CustomizableObject->GetPrivate()->GetModel(),
+	TSharedPtr<UE::Mutable::Private::FParameters> MutableParameters = UE::Mutable::Private::FModel::NewParameters(CustomizableObject->GetPrivate()->GetModel(),
 		nullptr,
 		&ModelResources.TextureParameterDefaultValues,
 		&ModelResources.SkeletalMeshParameterDefaultValues,
@@ -603,11 +603,11 @@ TSharedPtr<mu::FParameters> FCustomizableObjectInstanceDescriptor::GetParameters
 	{
 		const FString& Name = MutableParameters->GetName(ParamIndex);
 		const FGuid& Uid = MutableParameters->GetUid(ParamIndex);
-		const mu::EParameterType MutableType = MutableParameters->GetType(ParamIndex);
+		const UE::Mutable::Private::EParameterType MutableType = MutableParameters->GetType(ParamIndex);
  
 		switch (MutableType)
 		{
-		case mu::EParameterType::Bool:
+		case UE::Mutable::Private::EParameterType::Bool:
 		{
 			for (const FCustomizableObjectBoolParameterValue& BoolParameter : BoolParameters)
 			{
@@ -621,13 +621,13 @@ TSharedPtr<mu::FParameters> FCustomizableObjectInstanceDescriptor::GetParameters
 			break;
 		}
  
-		case mu::EParameterType::Int:
+		case UE::Mutable::Private::EParameterType::Int:
 		{
 			for (const FCustomizableObjectIntParameterValue& IntParameter : IntParameters)
 			{
 				if (IntParameter.ParameterName.Equals(Name, ESearchCase::CaseSensitive) || (Uid.IsValid() && IntParameter.Id == Uid))
 				{
-					if (TSharedPtr<mu::FRangeIndex> RangeIdxPtr = MutableParameters->NewRangeIndex(ParamIndex))
+					if (TSharedPtr<UE::Mutable::Private::FRangeIndex> RangeIdxPtr = MutableParameters->NewRangeIndex(ParamIndex))
 					{
 						for (int32 RangeIndex = 0; RangeIndex < IntParameter.ParameterRangeValueNames.Num(); ++RangeIndex)
 						{
@@ -651,13 +651,13 @@ TSharedPtr<mu::FParameters> FCustomizableObjectInstanceDescriptor::GetParameters
 			break;
 		}
  
-		case mu::EParameterType::Float:
+		case UE::Mutable::Private::EParameterType::Float:
 		{
 			for (const FCustomizableObjectFloatParameterValue& FloatParameter : FloatParameters)
 			{
 				if (FloatParameter.ParameterName == Name || (Uid.IsValid() && FloatParameter.Id == Uid))
 				{
-					if (TSharedPtr<mu::FRangeIndex> RangeIdxPtr = MutableParameters->NewRangeIndex(ParamIndex))
+					if (TSharedPtr<UE::Mutable::Private::FRangeIndex> RangeIdxPtr = MutableParameters->NewRangeIndex(ParamIndex))
 					{
 						for (int32 RangeIndex = 0; RangeIndex < FloatParameter.ParameterRangeValues.Num(); ++RangeIndex)
 						{
@@ -677,7 +677,7 @@ TSharedPtr<mu::FParameters> FCustomizableObjectInstanceDescriptor::GetParameters
 			break;
 		}
  
-		case mu::EParameterType::Color:
+		case UE::Mutable::Private::EParameterType::Color:
 		{
 			for (const FCustomizableObjectVectorParameterValue& VectorParameter : VectorParameters)
 			{
@@ -692,7 +692,7 @@ TSharedPtr<mu::FParameters> FCustomizableObjectInstanceDescriptor::GetParameters
 			break;
 		}
  
-		case mu::EParameterType::Matrix:
+		case UE::Mutable::Private::EParameterType::Matrix:
 			{
 				for (const FCustomizableObjectTransformParameterValue& TransformParameter : TransformParameters)
 				{
@@ -707,13 +707,13 @@ TSharedPtr<mu::FParameters> FCustomizableObjectInstanceDescriptor::GetParameters
 				break;
 			}
 			
-		case mu::EParameterType::Projector:
+		case UE::Mutable::Private::EParameterType::Projector:
 		{
 			for (const auto& ProjectorParameter : ProjectorParameters)
 			{
 				if (ProjectorParameter.ParameterName == Name || (Uid.IsValid() && ProjectorParameter.Id == Uid))
 				{
-					auto CopyProjector = [&MutableParameters, ParamIndex](const FCustomizableObjectProjector& Value, const mu::FRangeIndex* RangeIdxPtr = nullptr)
+					auto CopyProjector = [&MutableParameters, ParamIndex](const FCustomizableObjectProjector& Value, const UE::Mutable::Private::FRangeIndex* RangeIdxPtr = nullptr)
 					{
 						switch (Value.ProjectionType)
 						{
@@ -752,7 +752,7 @@ TSharedPtr<mu::FParameters> FCustomizableObjectInstanceDescriptor::GetParameters
  
 					CopyProjector(ProjectorParameter.Value);
  
-					if (const TSharedPtr<mu::FRangeIndex> RangeIdxPtr = MutableParameters->NewRangeIndex(ParamIndex))
+					if (const TSharedPtr<UE::Mutable::Private::FRangeIndex> RangeIdxPtr = MutableParameters->NewRangeIndex(ParamIndex))
 					{
 						for (int32 RangeIndex = 0; RangeIndex < ProjectorParameter.RangeValues.Num(); ++RangeIndex)
 						{
@@ -766,13 +766,13 @@ TSharedPtr<mu::FParameters> FCustomizableObjectInstanceDescriptor::GetParameters
 			break;
 		}
  
-		case mu::EParameterType::Image:
+		case UE::Mutable::Private::EParameterType::Image:
 		{
 			for (const FCustomizableObjectTextureParameterValue& TextureParameter : TextureParameters)
 			{
 				if (TextureParameter.ParameterName == Name || (Uid.IsValid() && TextureParameter.Id == Uid))
 				{
-					if (TSharedPtr<mu::FRangeIndex> RangeIdxPtr = MutableParameters->NewRangeIndex(ParamIndex))
+					if (TSharedPtr<UE::Mutable::Private::FRangeIndex> RangeIdxPtr = MutableParameters->NewRangeIndex(ParamIndex))
 					{
 						for (int32 RangeIndex = 0; RangeIndex < TextureParameter.ParameterRangeValues.Num(); ++RangeIndex)
 						{
@@ -792,13 +792,13 @@ TSharedPtr<mu::FParameters> FCustomizableObjectInstanceDescriptor::GetParameters
 			break;
 		}
  
-		case mu::EParameterType::Mesh:
+		case UE::Mutable::Private::EParameterType::Mesh:
 		{
 			for (const FCustomizableObjectSkeletalMeshParameterValue& MeshParameter : SkeletalMeshParameters)
 			{
 				if (MeshParameter.ParameterName == Name || (Uid.IsValid() && MeshParameter.Id == Uid))
 				{
-					if (TSharedPtr<mu::FRangeIndex> RangeIdxPtr = MutableParameters->NewRangeIndex(ParamIndex))
+					if (TSharedPtr<UE::Mutable::Private::FRangeIndex> RangeIdxPtr = MutableParameters->NewRangeIndex(ParamIndex))
 					{
 						for (int32 RangeIndex = 0; RangeIndex < MeshParameter.ParameterRangeValues.Num(); ++RangeIndex)
 						{
@@ -818,13 +818,13 @@ TSharedPtr<mu::FParameters> FCustomizableObjectInstanceDescriptor::GetParameters
 			break;
 		}
 
-		case mu::EParameterType::Material:
+		case UE::Mutable::Private::EParameterType::Material:
 		{
 			for (const FCustomizableObjectMaterialParameterValue& MaterialParameter : MaterialParameters)
 			{
 				if (MaterialParameter.ParameterName == Name || (Uid.IsValid() && MaterialParameter.Id == Uid))
 				{
-					if (TSharedPtr<mu::FRangeIndex> RangeIdxPtr = MutableParameters->NewRangeIndex(ParamIndex))
+					if (TSharedPtr<UE::Mutable::Private::FRangeIndex> RangeIdxPtr = MutableParameters->NewRangeIndex(ParamIndex))
 					{
 						for (int32 RangeIndex = 0; RangeIndex < MaterialParameter.ParameterRangeValues.Num(); ++RangeIndex)
 						{
@@ -921,7 +921,7 @@ void FCustomizableObjectInstanceDescriptor::ReloadParameters()
 		return;
 	}
 	
-	TSharedPtr<mu::FParameters> MutableParameters = mu::FModel::NewParameters(
+	TSharedPtr<UE::Mutable::Private::FParameters> MutableParameters = UE::Mutable::Private::FModel::NewParameters(
 		CustomizableObject->GetPrivate()->GetModel(),
 		nullptr,
 		&ModelResources->TextureParameterDefaultValues,
@@ -933,11 +933,11 @@ void FCustomizableObjectInstanceDescriptor::ReloadParameters()
 	{
 		const FString& Name = MutableParameters->GetName(ParamIndex);
 		const FGuid& Uid = MutableParameters->GetUid(ParamIndex);
-		const mu::EParameterType MutableType = MutableParameters->GetType(ParamIndex);
+		const UE::Mutable::Private::EParameterType MutableType = MutableParameters->GetType(ParamIndex);
  
 		switch (MutableType)
 		{
-		case mu::EParameterType::Bool:
+		case UE::Mutable::Private::EParameterType::Bool:
 		{
 			FCustomizableObjectBoolParameterValue Param;
 			Param.ParameterName = Name;
@@ -961,7 +961,7 @@ void FCustomizableObjectInstanceDescriptor::ReloadParameters()
 			break;
 		}
  
-		case mu::EParameterType::Int:
+		case UE::Mutable::Private::EParameterType::Int:
 		{
 			FCustomizableObjectIntParameterValue Param;
 			Param.ParameterName = Name;
@@ -989,7 +989,7 @@ void FCustomizableObjectInstanceDescriptor::ReloadParameters()
 					return false;
 				};
 				
-				if (TSharedPtr<mu::FRangeIndex> RangeIdxPtr = MutableParameters->NewRangeIndex(ParamIndex)) // Is multidimensional
+				if (TSharedPtr<UE::Mutable::Private::FRangeIndex> RangeIdxPtr = MutableParameters->NewRangeIndex(ParamIndex)) // Is multidimensional
 				{
 					// Get num of ranges (layers) from the instance
 					int32 ValueCount = Result->ParameterRangeValueNames.Num();
@@ -1034,13 +1034,13 @@ void FCustomizableObjectInstanceDescriptor::ReloadParameters()
 			} 
 			else // Not found in Instance Parameters. Use Mutable Parameters.
 			{
-				if (TSharedPtr<mu::FRangeIndex> RangeIdxPtr = MutableParameters->NewRangeIndex(ParamIndex))
+				if (TSharedPtr<UE::Mutable::Private::FRangeIndex> RangeIdxPtr = MutableParameters->NewRangeIndex(ParamIndex))
 				{
 					const int32 ValueCount = MutableParameters->GetValueCount(ParamIndex);
  
 					for (int32 ValueIndex = 0; ValueIndex < ValueCount; ++ValueIndex)
 					{	
-						const TSharedPtr<mu::FRangeIndex> RangeValueIdxPtr = MutableParameters->GetValueIndex(ParamIndex, ValueIndex);
+						const TSharedPtr<UE::Mutable::Private::FRangeIndex> RangeValueIdxPtr = MutableParameters->GetValueIndex(ParamIndex, ValueIndex);
 						const int32 RangeIndex = RangeValueIdxPtr->GetPosition(0);
  
 						if (!Param.ParameterRangeValueNames.IsValidIndex(RangeIndex))
@@ -1065,7 +1065,7 @@ void FCustomizableObjectInstanceDescriptor::ReloadParameters()
 			break;
 		}
  
-		case mu::EParameterType::Float:
+		case UE::Mutable::Private::EParameterType::Float:
 		{
 			FCustomizableObjectFloatParameterValue Param;
 			Param.ParameterName = Name;
@@ -1078,7 +1078,7 @@ void FCustomizableObjectInstanceDescriptor::ReloadParameters()
 			
 			if (FCustomizableObjectFloatParameterValue* Result = OldFloatParameters.FindByPredicate(FindByNameAndUid))
 			{	
-				if (TSharedPtr<mu::FRangeIndex> RangeIdxPtr = MutableParameters->NewRangeIndex(ParamIndex))
+				if (TSharedPtr<UE::Mutable::Private::FRangeIndex> RangeIdxPtr = MutableParameters->NewRangeIndex(ParamIndex))
 				{
 					Param.ParameterRangeValues = Result->ParameterRangeValues;
 				}
@@ -1089,13 +1089,13 @@ void FCustomizableObjectInstanceDescriptor::ReloadParameters()
 			} 
 			else // Not found in Instance Parameters. Use Mutable Parameters.
 			{
-				if (TSharedPtr<mu::FRangeIndex> RangeIdxPtr = MutableParameters->NewRangeIndex(ParamIndex))
+				if (TSharedPtr<UE::Mutable::Private::FRangeIndex> RangeIdxPtr = MutableParameters->NewRangeIndex(ParamIndex))
 				{
 					const int32 ValueCount = MutableParameters->GetValueCount(ParamIndex);
  
 					for (int32 ValueIndex = 0; ValueIndex < ValueCount; ++ValueIndex)
 					{	
-						TSharedPtr<mu::FRangeIndex> RangeValueIdxPtr = MutableParameters->GetValueIndex(ParamIndex, ValueIndex);
+						TSharedPtr<UE::Mutable::Private::FRangeIndex> RangeValueIdxPtr = MutableParameters->GetValueIndex(ParamIndex, ValueIndex);
 						int32 RangeIndex = RangeValueIdxPtr->GetPosition(0);
  
 						if (!Param.ParameterRangeValues.IsValidIndex(RangeIndex))
@@ -1116,7 +1116,7 @@ void FCustomizableObjectInstanceDescriptor::ReloadParameters()
 			break;
 		}
  
-		case mu::EParameterType::Color:
+		case UE::Mutable::Private::EParameterType::Color:
 		{
 			FCustomizableObjectVectorParameterValue Param;
 			Param.ParameterName = Name;
@@ -1142,7 +1142,7 @@ void FCustomizableObjectInstanceDescriptor::ReloadParameters()
 			break;
 		}
  
-		case mu::EParameterType::Matrix:
+		case UE::Mutable::Private::EParameterType::Matrix:
 			{
 				FCustomizableObjectTransformParameterValue Param;
 				Param.ParameterName = Name;
@@ -1168,7 +1168,7 @@ void FCustomizableObjectInstanceDescriptor::ReloadParameters()
 				break;
 			}
 			
-		case mu::EParameterType::Projector:
+		case UE::Mutable::Private::EParameterType::Projector:
 		{
 			FCustomizableObjectProjectorParameterValue Param;
 			Param.ParameterName = Name;
@@ -1184,7 +1184,7 @@ void FCustomizableObjectInstanceDescriptor::ReloadParameters()
 			
 			if (FCustomizableObjectProjectorParameterValue* Result = OldProjectorParameters.FindByPredicate(FindByNameAndUid))
 			{	
-				if (TSharedPtr<mu::FRangeIndex> RangeIdxPtr = MutableParameters->NewRangeIndex(ParamIndex))
+				if (TSharedPtr<UE::Mutable::Private::FRangeIndex> RangeIdxPtr = MutableParameters->NewRangeIndex(ParamIndex))
 				{
 					Param.RangeValues = Result->RangeValues;
 					Param.Value.ProjectionType = DefaultProjectorValue.ProjectionType;
@@ -1202,9 +1202,9 @@ void FCustomizableObjectInstanceDescriptor::ReloadParameters()
 			} 
 			else // Not found in Instance Parameters. Use Mutable Parameters.
 			{
-				auto GetProjector = [&MutableParameters, ParamIndex](FCustomizableObjectProjector &Value, const TSharedPtr<mu::FRangeIndex> RangeIndex = nullptr)
+				auto GetProjector = [&MutableParameters, ParamIndex](FCustomizableObjectProjector &Value, const TSharedPtr<UE::Mutable::Private::FRangeIndex> RangeIndex = nullptr)
 				{
-					mu::EProjectorType Type;
+					UE::Mutable::Private::EProjectorType Type;
 					MutableParameters->GetProjectorValue(ParamIndex,
 						&Type,
 						&Value.Position,
@@ -1228,13 +1228,13 @@ void FCustomizableObjectInstanceDescriptor::ReloadParameters()
  
 				GetProjector(Param.Value);
  
-				if (TSharedPtr<mu::FRangeIndex> RangeIdxPtr = MutableParameters->NewRangeIndex(ParamIndex))
+				if (TSharedPtr<UE::Mutable::Private::FRangeIndex> RangeIdxPtr = MutableParameters->NewRangeIndex(ParamIndex))
 				{
 					const int32 ValueCount = MutableParameters->GetValueCount(ParamIndex);
  
 					for (int32 ValueIndex = 0; ValueIndex < ValueCount; ++ValueIndex)
 					{	
-						TSharedPtr<mu::FRangeIndex> RangeValueIdxPtr = MutableParameters->GetValueIndex(ParamIndex, ValueIndex);
+						TSharedPtr<UE::Mutable::Private::FRangeIndex> RangeValueIdxPtr = MutableParameters->GetValueIndex(ParamIndex, ValueIndex);
 						int32 RangeIndex = RangeValueIdxPtr->GetPosition(0);
  
 						if (!Param.RangeValues.IsValidIndex(RangeIndex))
@@ -1251,7 +1251,7 @@ void FCustomizableObjectInstanceDescriptor::ReloadParameters()
 			break;
 		}
  
-		case mu::EParameterType::Image:
+		case UE::Mutable::Private::EParameterType::Image:
 		{
 			FCustomizableObjectTextureParameterValue Param;
 			Param.ParameterName = Name;
@@ -1264,7 +1264,7 @@ void FCustomizableObjectInstanceDescriptor::ReloadParameters()
 			
 			if (FCustomizableObjectTextureParameterValue* Result = OldTextureParameters.FindByPredicate(FindByNameAndUid))
 			{	
-				if (TSharedPtr<mu::FRangeIndex> RangeIdxPtr = MutableParameters->NewRangeIndex(ParamIndex))
+				if (TSharedPtr<UE::Mutable::Private::FRangeIndex> RangeIdxPtr = MutableParameters->NewRangeIndex(ParamIndex))
 				{
 					Param.ParameterRangeValues = Result->ParameterRangeValues;
 				}
@@ -1275,13 +1275,13 @@ void FCustomizableObjectInstanceDescriptor::ReloadParameters()
 			} 
 			else // Not found in Instance Parameters. Use Mutable Parameters.
 			{
-				if (TSharedPtr<mu::FRangeIndex> RangeIdxPtr = MutableParameters->NewRangeIndex(ParamIndex))
+				if (TSharedPtr<UE::Mutable::Private::FRangeIndex> RangeIdxPtr = MutableParameters->NewRangeIndex(ParamIndex))
 				{
 					const int32 ValueCount = MutableParameters->GetValueCount(ParamIndex);
  
 					for (int32 ValueIndex = 0; ValueIndex < ValueCount; ++ValueIndex)
 					{	
-						TSharedPtr<mu::FRangeIndex> RangeValueIdxPtr = MutableParameters->GetValueIndex(ParamIndex, ValueIndex);
+						TSharedPtr<UE::Mutable::Private::FRangeIndex> RangeValueIdxPtr = MutableParameters->GetValueIndex(ParamIndex, ValueIndex);
 						int32 RangeIndex = RangeValueIdxPtr->GetPosition(0);
  
 						if (!Param.ParameterRangeValues.IsValidIndex(RangeIndex))
@@ -1309,7 +1309,7 @@ void FCustomizableObjectInstanceDescriptor::ReloadParameters()
 		}
  
  
-		case mu::EParameterType::Mesh:
+		case UE::Mutable::Private::EParameterType::Mesh:
 		{
 			FCustomizableObjectSkeletalMeshParameterValue Param;
 			Param.ParameterName = Name;
@@ -1322,7 +1322,7 @@ void FCustomizableObjectInstanceDescriptor::ReloadParameters()
  
 			if (FCustomizableObjectSkeletalMeshParameterValue* Result = OldSkeletalMeshParameters.FindByPredicate(FindByNameAndUid))
 			{
-				if (TSharedPtr<mu::FRangeIndex> RangeIdxPtr = MutableParameters->NewRangeIndex(ParamIndex))
+				if (TSharedPtr<UE::Mutable::Private::FRangeIndex> RangeIdxPtr = MutableParameters->NewRangeIndex(ParamIndex))
 				{
 					Param.ParameterRangeValues = Result->ParameterRangeValues;
 				}
@@ -1333,13 +1333,13 @@ void FCustomizableObjectInstanceDescriptor::ReloadParameters()
 			}
 			else // Not found in Instance Parameters. Use Mutable Parameters.
 			{
-				if (TSharedPtr<mu::FRangeIndex> RangeIdxPtr = MutableParameters->NewRangeIndex(ParamIndex))
+				if (TSharedPtr<UE::Mutable::Private::FRangeIndex> RangeIdxPtr = MutableParameters->NewRangeIndex(ParamIndex))
 				{
 					const int32 ValueCount = MutableParameters->GetValueCount(ParamIndex);
  
 					for (int32 ValueIndex = 0; ValueIndex < ValueCount; ++ValueIndex)
 					{
-						TSharedPtr<mu::FRangeIndex> RangeValueIdxPtr = MutableParameters->GetValueIndex(ParamIndex, ValueIndex);
+						TSharedPtr<UE::Mutable::Private::FRangeIndex> RangeValueIdxPtr = MutableParameters->GetValueIndex(ParamIndex, ValueIndex);
 						int32 RangeIndex = RangeValueIdxPtr->GetPosition(0);
  
 						if (!Param.ParameterRangeValues.IsValidIndex(RangeIndex))
@@ -1366,7 +1366,7 @@ void FCustomizableObjectInstanceDescriptor::ReloadParameters()
 			break;
 		}
  
-		case mu::EParameterType::Material:
+		case UE::Mutable::Private::EParameterType::Material:
 		{
 			FCustomizableObjectMaterialParameterValue Param;
 			Param.ParameterName = Name;
@@ -1379,7 +1379,7 @@ void FCustomizableObjectInstanceDescriptor::ReloadParameters()
 
 			if (FCustomizableObjectMaterialParameterValue* Result = OldMaterialParameters.FindByPredicate(FindByNameAndUid))
 			{
-				if (TSharedPtr<mu::FRangeIndex> RangeIdxPtr = MutableParameters->NewRangeIndex(ParamIndex))
+				if (TSharedPtr<UE::Mutable::Private::FRangeIndex> RangeIdxPtr = MutableParameters->NewRangeIndex(ParamIndex))
 				{
 					Param.ParameterRangeValues = Result->ParameterRangeValues;
 				}
@@ -1390,13 +1390,13 @@ void FCustomizableObjectInstanceDescriptor::ReloadParameters()
 			}
 			else // Not found in Instance Parameters. Use Mutable Parameters.
 			{
-				if (TSharedPtr<mu::FRangeIndex> RangeIdxPtr = MutableParameters->NewRangeIndex(ParamIndex))
+				if (TSharedPtr<UE::Mutable::Private::FRangeIndex> RangeIdxPtr = MutableParameters->NewRangeIndex(ParamIndex))
 				{
 					const int32 ValueCount = MutableParameters->GetValueCount(ParamIndex);
 
 					for (int32 ValueIndex = 0; ValueIndex < ValueCount; ++ValueIndex)
 					{
-						TSharedPtr<mu::FRangeIndex> RangeValueIdxPtr = MutableParameters->GetValueIndex(ParamIndex, ValueIndex);
+						TSharedPtr<UE::Mutable::Private::FRangeIndex> RangeValueIdxPtr = MutableParameters->GetValueIndex(ParamIndex, ValueIndex);
 						int32 RangeIndex = RangeValueIdxPtr->GetPosition(0);
 
 						if (!Param.ParameterRangeValues.IsValidIndex(RangeIndex))
diff --git a/Engine/Plugins/Mutable/Source/CustomizableObject/Private/MuCO/CustomizableObjectMeshUpdate.cpp b/Engine/Plugins/Mutable/Source/CustomizableObject/Private/MuCO/CustomizableObjectMeshUpdate.cpp
--- a/Engine/Plugins/Mutable/Source/CustomizableObject/Private/MuCO/CustomizableObjectMeshUpdate.cpp
+++ b/Engine/Plugins/Mutable/Source/CustomizableObject/Private/MuCO/CustomizableObjectMeshUpdate.cpp
@@ -170,7 +170,7 @@ void FCustomizableObjectMeshStreamIn::DoCancelMeshUpdate(const FContext& Context
 
 namespace impl
 {
-	void Task_Mutable_UpdateMesh_End(const TSharedPtr<FMutableUpdateMeshContext> OperationData, TRefCountPtr<FCustomizableObjectMeshStreamIn>& Task, mu::FInstance::FID InstanceID)
+	void Task_Mutable_UpdateMesh_End(const TSharedPtr<FMutableUpdateMeshContext> OperationData, TRefCountPtr<FCustomizableObjectMeshStreamIn>& Task, UE::Mutable::Private::FInstance::FID InstanceID)
 	{
 		MUTABLE_CPUPROFILER_SCOPE(Task_Mutable_UpdateMesh_End);
 
@@ -191,7 +191,7 @@ namespace impl
 	void Task_Mutable_UpdateMesh_Loop(
 		const TSharedPtr<FMutableUpdateMeshContext> OperationData,
 		TRefCountPtr<FCustomizableObjectMeshStreamIn>& Task,
-		mu::FInstance::FID InstanceID,
+		UE::Mutable::Private::FInstance::FID InstanceID,
 		int32 LODIndex)
 	{
 		MUTABLE_CPUPROFILER_SCOPE(Task_Mutable_UpdateMesh_Loop);
@@ -202,10 +202,10 @@ namespace impl
 			return;
 		}
 
-		const mu::FMeshId& MeshID = OperationData->MeshIDs[LODIndex];
+		const UE::Mutable::Private::FMeshId& MeshID = OperationData->MeshIDs[LODIndex];
 
-		constexpr mu::EMeshContentFlags MeshContentFilter = mu::EMeshContentFlags::AllFlags;
-		UE::Tasks::TTask<TSharedPtr<const mu::FMesh>> GetMeshTask = OperationData->System->GetMesh(InstanceID, MeshID, MeshContentFilter);
+		constexpr UE::Mutable::Private::EMeshContentFlags MeshContentFilter = UE::Mutable::Private::EMeshContentFlags::AllFlags;
+		UE::Tasks::TTask<TSharedPtr<const UE::Mutable::Private::FMesh>> GetMeshTask = OperationData->System->GetMesh(InstanceID, MeshID, MeshContentFilter);
 
 		UE::Tasks::AddNested(UE::Tasks::Launch(TEXT("Task_MutableGetMeshes_GetMesh_Post"), [=]() mutable
 			{
@@ -228,8 +228,8 @@ namespace impl
 
 		TRACE_BEGIN_REGION(UE_MUTABLE_UPDATE_MESH_REGION);
 
-		TSharedPtr<mu::FSystem> System = OperationData->System;
-		const TSharedPtr<mu::FModel, ESPMode::ThreadSafe> Model = OperationData->Model;
+		TSharedPtr<UE::Mutable::Private::FSystem> System = OperationData->System;
+		const TSharedPtr<UE::Mutable::Private::FModel, ESPMode::ThreadSafe> Model = OperationData->Model;
 
 #if WITH_EDITOR
 		// Recompiling a CO in the editor will invalidate the previously generated Model. Check that it is valid before accessing the streamed data.
@@ -242,14 +242,14 @@ namespace impl
 #endif
 
 		// For now, we are forcing the recreation of mutable-side instances with every update.
-		mu::FInstance::FID InstanceID = System->NewInstance(Model, OperationData->ExternalResourceProvider);
+		UE::Mutable::Private::FInstance::FID InstanceID = System->NewInstance(Model, OperationData->ExternalResourceProvider);
 		UE_LOG(LogMutable, Verbose, TEXT("Creating Mutable instance with id [%d] for a mesh update"), InstanceID);
 
 		// LOD mask, set to all ones to build all LODs
 		const uint32 LODMask = 0xFFFFFFFF;
 
 		// Main instance generation step
-		TSharedPtr<const mu::FInstance> Instance = System->BeginUpdate_MutableThread(InstanceID,
+		TSharedPtr<const UE::Mutable::Private::FInstance> Instance = System->BeginUpdate_MutableThread(InstanceID,
 			OperationData->Parameters,
 			OperationData->MeshIdRegistry,
 			OperationData->ImageIdRegistry,
@@ -314,7 +314,7 @@ void FCustomizableObjectMeshStreamIn::RequestMeshUpdate(const FContext& Context)
 			FMutableStreamRequest StreamRequest(RefThis->OperationData->ModelStreamableBulkData);
 
 			// Get morphs to be streamed.
-			for (TSharedPtr<const mu::FMesh>& Mesh : RefThis->OperationData->Meshes)
+			for (TSharedPtr<const UE::Mutable::Private::FMesh>& Mesh : RefThis->OperationData->Meshes)
 			{
 				if (!Mesh)
 				{
@@ -395,7 +395,7 @@ void FCustomizableObjectMeshStreamIn::ConvertMesh(const FContext& Context, bool&
 
 	for (int32 LODIndex = PendingFirstLODIdx; LODIndex < CurrentFirstLODIdx; ++LODIndex)
 	{
-		TSharedPtr<const mu::FMesh> MutableMesh = OperationData->Meshes[LODIndex + Context.AssetLODBias];
+		TSharedPtr<const UE::Mutable::Private::FMesh> MutableMesh = OperationData->Meshes[LODIndex + Context.AssetLODBias];
 
 		if (!MutableMesh)
 		{
diff --git a/Engine/Plugins/Mutable/Source/CustomizableObject/Private/MuCO/CustomizableObjectMeshUpdate.h b/Engine/Plugins/Mutable/Source/CustomizableObject/Private/MuCO/CustomizableObjectMeshUpdate.h
--- a/Engine/Plugins/Mutable/Source/CustomizableObject/Private/MuCO/CustomizableObjectMeshUpdate.h
+++ b/Engine/Plugins/Mutable/Source/CustomizableObject/Private/MuCO/CustomizableObjectMeshUpdate.h
@@ -17,30 +17,30 @@ CustomizableObjectMeshUpdate.h: Helpers to stream in skeletal mesh LODs.
 struct FModelStreamableBulkData;
 class UCustomizableObjectSkeletalMesh;
 
-namespace mu 
+namespace UE::Mutable::Private 
 { 
 	class FParameters;
 }
 
 struct FMutableUpdateMeshContext
 {
-	TSharedPtr<mu::FSystem> System;
-	TSharedPtr<mu::FModel> Model;
-	TSharedPtr<mu::FMeshIdRegistry> MeshIdRegistry;
-	TSharedPtr<mu::FImageIdRegistry> ImageIdRegistry;
-	TSharedPtr<mu::FMaterialIdRegistry> MaterialIdRegistry;
+	TSharedPtr<UE::Mutable::Private::FSystem> System;
+	TSharedPtr<UE::Mutable::Private::FModel> Model;
+	TSharedPtr<UE::Mutable::Private::FMeshIdRegistry> MeshIdRegistry;
+	TSharedPtr<UE::Mutable::Private::FImageIdRegistry> ImageIdRegistry;
+	TSharedPtr<UE::Mutable::Private::FMaterialIdRegistry> MaterialIdRegistry;
 	TSharedPtr<FUnrealMutableResourceProvider> ExternalResourceProvider;
-	TSharedPtr<const mu::FParameters> Parameters;
+	TSharedPtr<const UE::Mutable::Private::FParameters> Parameters;
 
 	TSharedPtr<FModelStreamableBulkData> ModelStreamableBulkData;
 	
 	int32 State = -1;
 
 	// IDs of the meshes to generate per LOD.
-	TArray<mu::FMeshId> MeshIDs;
+	TArray<UE::Mutable::Private::FMeshId> MeshIDs;
 
 	// Meshes generated per LOD
-	TArray<TSharedPtr<const mu::FMesh>> Meshes;
+	TArray<TSharedPtr<const UE::Mutable::Private::FMesh>> Meshes;
 
 	/** Serialized data per block. */
 	TMap<uint32, FMorphTargetMeshData> MorphTargetMeshData;
@@ -96,7 +96,7 @@ private:
 	/** Cancels the MeshUpdate task. */
 	void CancelMeshUpdate(const FContext& Context);
 
-	/** Converts from mu::FMesh to FSkeletalMeshLODRenderData. */
+	/** Converts from UE::Mutable::Private::FMesh to FSkeletalMeshLODRenderData. */
 	void ConvertMesh(const FContext& Context, bool& bOutMarkRenderStateDirty);
 
 	void MarkRenderStateDirty(const FContext& Context);
diff --git a/Engine/Plugins/Mutable/Source/CustomizableObject/Private/MuCO/CustomizableObjectMipDataProvider.cpp b/Engine/Plugins/Mutable/Source/CustomizableObject/Private/MuCO/CustomizableObjectMipDataProvider.cpp
--- a/Engine/Plugins/Mutable/Source/CustomizableObject/Private/MuCO/CustomizableObjectMipDataProvider.cpp
+++ b/Engine/Plugins/Mutable/Source/CustomizableObject/Private/MuCO/CustomizableObjectMipDataProvider.cpp
@@ -81,8 +81,8 @@ namespace Impl
 		if (FLogBenchmarkUtil::IsBenchmarkingReportingEnabled())
 		{
 			// Cache memory used when starting the update of the image
-			OperationData->ImageUpdateStartBytes = mu::FGlobalMemoryCounter::GetAbsoluteCounter();
-			mu::FGlobalMemoryCounter::Zero();
+			OperationData->ImageUpdateStartBytes = UE::Mutable::Private::FGlobalMemoryCounter::GetAbsoluteCounter();
+			UE::Mutable::Private::FGlobalMemoryCounter::Zero();
 		}
 		
 		// Any external texture that may be needed for this update will be requested from Mutable Core's GetImage
@@ -101,8 +101,8 @@ namespace Impl
 
 		TRACE_BEGIN_REGION(UE_MUTABLE_MIPDATA_PROVIDER_UPDATE_IMAGE_REGION);
 		
-		TSharedPtr<mu::FSystem> System = OperationData->UpdateContext->System;
-		const TSharedPtr<mu::FModel, ESPMode::ThreadSafe> Model = OperationData->UpdateContext->Model;
+		TSharedPtr<UE::Mutable::Private::FSystem> System = OperationData->UpdateContext->System;
+		const TSharedPtr<UE::Mutable::Private::FModel, ESPMode::ThreadSafe> Model = OperationData->UpdateContext->Model;
 
 		auto EndUpdateImage = [](TSharedPtr<FMutableImageOperationData>& OperationData)
 		{
@@ -131,10 +131,10 @@ namespace Impl
 #endif
 
 		// For now, we are forcing the recreation of mutable-side instances with every update.
-		mu::FInstance::FID InstanceID = System->NewInstance(Model, OperationData->UpdateContext->ExternalResourceProvider);
+		UE::Mutable::Private::FInstance::FID InstanceID = System->NewInstance(Model, OperationData->UpdateContext->ExternalResourceProvider);
 		UE_LOG(LogMutable, Verbose, TEXT("Creating Mutable instance with id [%d] for a single UpdateImage"), InstanceID)
 
-		TSharedPtr<const mu::FInstance> Instance;
+		TSharedPtr<const UE::Mutable::Private::FInstance> Instance;
 
 		// Main instance generation step
 		{
@@ -159,9 +159,9 @@ namespace Impl
 
 		// This ID may be different than the ID obtained the first time the image was generated, because the mutable
 		// runtime cannot remember all the resources it has built, and only remembers a fixed amount.
-		mu::FImageId MipImageID = Instance->GetImageId(ImageRef.Component, ImageRef.LOD, SurfaceIndex, ImageRef.Image);
+		UE::Mutable::Private::FImageId MipImageID = Instance->GetImageId(ImageRef.Component, ImageRef.LOD, SurfaceIndex, ImageRef.Image);
 
-		UE::Tasks::TTask<TSharedPtr<const mu::FImage>> GetImageTask = 
+		UE::Tasks::TTask<TSharedPtr<const UE::Mutable::Private::FImage>> GetImageTask = 
 				System->GetImage(InstanceID, MipImageID, ImageRef.BaseMip + OperationData->MipsToSkip, ImageRef.LOD);
 
 		UE::Tasks::AddNested(UE::Tasks::Launch(TEXT("MipDataProvider_EndUpdateImagesTask"),
@@ -169,7 +169,7 @@ namespace Impl
 				{
 					check(GetImageTask.IsCompleted());
 
-					TSharedPtr<const mu::FImage> ResultImage = GetImageTask.GetResult();
+					TSharedPtr<const UE::Mutable::Private::FImage> ResultImage = GetImageTask.GetResult();
 
 					check(ResultImage);
 
@@ -186,7 +186,7 @@ namespace Impl
 						UE_LOG(LogMutable, Warning, TEXT("Mutable generated an incomplete mip chain for image."));
 
 						// Force the right number of mips. The missing data will be black.
-						TSharedPtr<mu::FImage> NewImage = MakeShared<mu::FImage>(ResultImage->GetSizeX(), ResultImage->GetSizeY(), FullMipCount, ResultImage->GetFormat(), mu::EInitializationType::Black);
+						TSharedPtr<UE::Mutable::Private::FImage> NewImage = MakeShared<UE::Mutable::Private::FImage>(ResultImage->GetSizeX(), ResultImage->GetSizeY(), FullMipCount, ResultImage->GetFormat(), UE::Mutable::Private::EInitializationType::Black);
 
 						// Formats with BytesPerBlock == 0 will not allocate memory. This type of images are not expected here.
 						check(!NewImage->DataStorage.IsEmpty());
@@ -221,7 +221,7 @@ namespace Impl
 					{
 						double Time = FPlatformTime::Seconds() - StartTime;
 						// Report the peak memory used by the operation
-						const int64 PeakMemory = mu::FGlobalMemoryCounter::GetPeak();
+						const int64 PeakMemory = UE::Mutable::Private::FGlobalMemoryCounter::GetPeak();
 						// Report the peak memory used during the operation (operation + baseline)
 						const int64 RealMemoryPeak = PeakMemory + OperationData->ImageUpdateStartBytes;
 
@@ -366,7 +366,7 @@ int32 FMutableTextureMipDataProvider::GetMips(const FTextureUpdateContext& Conte
 
 			uint32 MaxBlockSize = 0;
 
-			TSharedPtr<mu::FModel, ESPMode::ThreadSafe> Model = UpdateContext->Model;
+			TSharedPtr<UE::Mutable::Private::FModel, ESPMode::ThreadSafe> Model = UpdateContext->Model;
 			check(Model);
 
 			if (!OperationData->Levels.IsEmpty())
@@ -402,8 +402,8 @@ int32 FMutableTextureMipDataProvider::GetMips(const FTextureUpdateContext& Conte
 				const FMutableStreamableBlock& Block = ModelStreamableBulkData->ModelStreamables[RomId];
 				PrefetchRequest->AddBlock(
 					ModelStreamableBulkData->ModelStreamables[RomId],
-					MutablePrivate::EStreamableDataType::Model,
-					(uint16)mu::EDataType::Image,
+					UE::Mutable::Private::EStreamableDataType::Model,
+					(uint16)UE::Mutable::Private::EDataType::Image,
 					TArrayView<uint8>(AllocatedMemory, Model->GetRomSize(RomId)));
 			}
 		}
@@ -483,7 +483,7 @@ bool FMutableTextureMipDataProvider::PollMips(const FTextureUpdateSyncOptions& S
 		// The counter must be zero meaning the Mutable image operation has finished
 		check(SyncOptions.Counter->GetValue() == 0);
 
-		TSharedPtr<const mu::FImage> Image = OperationData->Result;
+		TSharedPtr<const UE::Mutable::Private::FImage> Image = OperationData->Result;
 		
 		int32 ImageLODCount = 0;
 		if (Image)
diff --git a/Engine/Plugins/Mutable/Source/CustomizableObject/Private/MuCO/CustomizableObjectSystem.cpp b/Engine/Plugins/Mutable/Source/CustomizableObject/Private/MuCO/CustomizableObjectSystem.cpp
--- a/Engine/Plugins/Mutable/Source/CustomizableObject/Private/MuCO/CustomizableObjectSystem.cpp
+++ b/Engine/Plugins/Mutable/Source/CustomizableObject/Private/MuCO/CustomizableObjectSystem.cpp
@@ -596,13 +596,13 @@ void FUpdateContextPrivate::InitMeshDescriptors(int32 Size)
 }
 
 
-const TArray<TArray<mu::FMeshId>>& FUpdateContextPrivate::GetMeshDescriptors() const
+const TArray<TArray<UE::Mutable::Private::FMeshId>>& FUpdateContextPrivate::GetMeshDescriptors() const
 {
 	return MeshDescriptors;
 }
 
 
-TArray<mu::FMeshId>* FUpdateContextPrivate::GetMeshDescriptors(FCustomizableObjectComponentIndex Index)
+TArray<UE::Mutable::Private::FMeshId>* FUpdateContextPrivate::GetMeshDescriptors(FCustomizableObjectComponentIndex Index)
 {
 	if (MeshDescriptors.IsValidIndex(Index.GetValue()))
 	{
@@ -744,7 +744,7 @@ void FMutablePendingInstanceWork::AddDiscard(const FMutablePendingInstanceDiscar
 }
 
 
-void FMutablePendingInstanceWork::AddIDRelease(mu::FInstance::FID IDToRelease)
+void FMutablePendingInstanceWork::AddIDRelease(UE::Mutable::Private::FInstance::FID IDToRelease)
 {
 	PendingIDsToRelease.Add(IDToRelease);
 }
@@ -842,10 +842,10 @@ void UCustomizableObjectSystem::InitSystem()
 
 	Private->LastWorkingMemoryBytes = CVarWorkingMemoryKB->GetInt() * 1024;
 
-	mu::FSettings Settings;
+	UE::Mutable::Private::FSettings Settings;
 	Settings.SetProfile(false);
 	Settings.SetWorkingMemoryBytes(Private->LastWorkingMemoryBytes);
-	Private->MutableSystem = MakeShared<mu::FSystem>(Settings);
+	Private->MutableSystem = MakeShared<UE::Mutable::Private::FSystem>(Settings);
 	check(Private->MutableSystem);
 
 	Private->Streamer = MakeShared<FUnrealMutableModelBulkReader>();
@@ -1709,7 +1709,7 @@ void UCustomizableObjectSystemPrivate::InitDiscardResourcesSkeletalMesh(UCustomi
 }
 
 
-void UCustomizableObjectSystemPrivate::InitInstanceIDRelease(mu::FInstance::FID IDToRelease)
+void UCustomizableObjectSystemPrivate::InitInstanceIDRelease(UE::Mutable::Private::FInstance::FID IDToRelease)
 {
 	check(IsInGameThread());
 
@@ -1926,7 +1926,7 @@ namespace impl
 	struct FGetImageData
 	{
 		int32 ImageIndex;
-		mu::FImageId ImageID;
+		UE::Mutable::Private::FImageId ImageID;
 	};
 	
 
@@ -1941,8 +1941,8 @@ namespace impl
 	struct FGetMeshData
 	{
 		int32 InstanceUpdateLODIndex;
-		mu::FMeshId MeshID;
-		mu::EMeshContentFlags ContentFilter;
+		UE::Mutable::Private::FMeshId MeshID;
+		UE::Mutable::Private::EMeshContentFlags ContentFilter;
 	};
 
 
@@ -1959,16 +1959,16 @@ namespace impl
 	void Task_Mutable_GetImages_GetImage(
 		const TSharedRef<FUpdateContextPrivate>& OperationData,
 		double StartTime,
-		const TSharedPtr<TArray<mu::FImageId>>& ImagesInThisInstance,
+		const TSharedPtr<TArray<UE::Mutable::Private::FImageId>>& ImagesInThisInstance,
 		int32 ImageIndex,
-		UE::Tasks::TTask<mu::FExtendedImageDesc> GetImageDescTask);
+		UE::Tasks::TTask<UE::Mutable::Private::FExtendedImageDesc> GetImageDescTask);
 
 	
 	/** Process the next Image. If there are no more Images, go to the end of the task. */
 	void Task_Mutable_GetImages_Loop(
 		const TSharedRef<FUpdateContextPrivate>& OperationData,
 		double StartTime,
-		const TSharedPtr<TArray<mu::FImageId>>& ImagesInThisInstance,
+		const TSharedPtr<TArray<UE::Mutable::Private::FImageId>>& ImagesInThisInstance,
 		int32 ImageIndex);
 
 	
@@ -2060,7 +2060,7 @@ namespace impl
 				const FInstanceUpdateData::FImage& Image = OperationData->InstanceUpdateData.Images[Surface.FirstImage+ImageIndex];
 
 				const FName KeyName = Image.Name;
-				TSharedPtr<const mu::FImage> MutableImage = Image.Image;
+				TSharedPtr<const UE::Mutable::Private::FImage> MutableImage = Image.Image;
 
 				// If the image is null, it must be in the cache (or repeated in this instance), and we don't need to do anything here.
 				if (MutableImage)
@@ -2101,7 +2101,7 @@ namespace impl
 			{
 				FInstanceUpdateData::FLOD& LOD = OperationData->InstanceUpdateData.LODs[Component.FirstLOD + LODIndex];
 
-				TSharedPtr<const mu::FMesh> Mesh = LOD.Mesh;
+				TSharedPtr<const UE::Mutable::Private::FMesh> Mesh = LOD.Mesh;
 				if (!Mesh || Mesh->IsReference())
 				{
 					continue;
@@ -2127,7 +2127,7 @@ namespace impl
 				}
 
 				// Append BoneMap to the array of BoneMaps
-				const TArray<mu::FBoneName>& BoneMap = Mesh->GetBoneMap();
+				const TArray<UE::Mutable::Private::FBoneName>& BoneMap = Mesh->GetBoneMap();
 				LOD.FirstBoneMap = OperationData->InstanceUpdateData.BoneMaps.Num();
 				LOD.BoneMapCount = BoneMap.Num();
 				OperationData->InstanceUpdateData.BoneMaps.Append(BoneMap);
@@ -2137,7 +2137,7 @@ namespace impl
 				LOD.ActiveBoneCount = Mesh->GetBonePoseCount();
 				for (uint32 BoneIndex = 0; BoneIndex < LOD.ActiveBoneCount; ++BoneIndex)
 				{
-					const mu::FBoneName& BoneId = Mesh->GetBonePoseId(BoneIndex);
+					const UE::Mutable::Private::FBoneName& BoneId = Mesh->GetBonePoseId(BoneIndex);
 
 					OperationData->InstanceUpdateData.ActiveBones.Add(BoneId);
 
@@ -2249,8 +2249,8 @@ namespace impl
 				
 				for (int32 SectionIndex = 0; SectionIndex < LOD.Mesh->Surfaces.Num(); ++SectionIndex)
 				{
-					const mu::FMeshSurface& Surface = LOD.Mesh->Surfaces[SectionIndex];
-					for (const mu::FSurfaceSubMesh& SubMesh : Surface.SubMeshes)
+					const UE::Mutable::Private::FMeshSurface& Surface = LOD.Mesh->Surfaces[SectionIndex];
+					for (const UE::Mutable::Private::FSurfaceSubMesh& SubMesh : Surface.SubMeshes)
 					{
 						if (SubMesh.ExternalId == 0)
 						{
@@ -2332,7 +2332,7 @@ namespace impl
 		double StartTime,
 		const TSharedRef<TArray<FGetImageData>>& GetImagesData,
 		int32 GetImageIndex,
-		UE::Tasks::TTask<TSharedPtr<const mu::FImage>> GetImageTask)
+		UE::Tasks::TTask<TSharedPtr<const UE::Mutable::Private::FImage>> GetImageTask)
 	{
 		MUTABLE_CPUPROFILER_SCOPE(Task_Mutable_GetMeshes_GetImage_Post)
 
@@ -2381,7 +2381,7 @@ namespace impl
 
 		const FGetImageData& ImageData = (*GetImagesData)[GetImageIndex];
 		
-		UE::Tasks::TTask<TSharedPtr<const mu::FImage>> GetImageTask = OperationData->MutableSystem->GetImage(OperationData->InstanceID, ImageData.ImageID, 0, 0);
+		UE::Tasks::TTask<TSharedPtr<const UE::Mutable::Private::FImage>> GetImageTask = OperationData->MutableSystem->GetImage(OperationData->InstanceID, ImageData.ImageID, 0, 0);
 
 		UE::Tasks::AddNested(UE::Tasks::Launch(TEXT("Task_Mutable_GetMeshes_GetImage_Post"), [=]()
 		{
@@ -2403,7 +2403,7 @@ namespace impl
 		const UCustomizableObject* CustomizableObject = Instance->GetCustomizableObject();
 		const UModelResources& ModelResources = *CustomizableObject->GetPrivate()->GetModelResources();
 
-		const mu::FInstance* MutableInstance = OperationData->MutableInstance.Get();
+		const UE::Mutable::Private::FInstance* MutableInstance = OperationData->MutableInstance.Get();
 
 		TArray<uint32> SurfacesSharedId;
 
@@ -2575,8 +2575,8 @@ namespace impl
 		{
 			LastTask = UE::Tasks::Launch(TEXT("Task_MutableGetMeshes_GetMesh_GetMaterial"), [SurfaceIndex, OperationData]()
 			{
-				mu::FMaterialId MaterialId = OperationData->InstanceUpdateData.Surfaces[SurfaceIndex].MaterialId;
-				UE::Tasks::TTask<TSharedPtr<const mu::FMaterial>> GetMaterialTask = OperationData->MutableSystem->GetMaterial(OperationData->InstanceID, MaterialId);
+				UE::Mutable::Private::FMaterialId MaterialId = OperationData->InstanceUpdateData.Surfaces[SurfaceIndex].MaterialId;
+				UE::Tasks::TTask<TSharedPtr<const UE::Mutable::Private::FMaterial>> GetMaterialTask = OperationData->MutableSystem->GetMaterial(OperationData->InstanceID, MaterialId);
 				
 				UE::Tasks::AddNested(UE::Tasks::Launch(TEXT("Task_MutableGetMeshes_GetMesh_GetMaterial_Post"), [=]() mutable
 				{
@@ -2604,7 +2604,7 @@ namespace impl
 		double StartTime,
 		const TSharedRef<TArray<FGetMeshData>>& GetMeshesData,
 		int32 GetMeshIndex,
-		UE::Tasks::TTask<TSharedPtr<const mu::FMesh>> GetMeshTask)
+		UE::Tasks::TTask<TSharedPtr<const UE::Mutable::Private::FMesh>> GetMeshTask)
 	{
 		MUTABLE_CPUPROFILER_SCOPE(Task_MutableGetMeshes_GetMesh_Post)
 
@@ -2654,7 +2654,7 @@ namespace impl
 		}
 		
 		const FGetMeshData& MeshData = (*GetMeshesData)[GetMeshIndex];
-		UE::Tasks::TTask<TSharedPtr<const mu::FMesh>> GetMeshTask = 
+		UE::Tasks::TTask<TSharedPtr<const UE::Mutable::Private::FMesh>> GetMeshTask = 
 				OperationData->MutableSystem->GetMesh(OperationData->InstanceID, MeshData.MeshID, MeshData.ContentFilter);
 		
 		UE::Tasks::AddNested(UE::Tasks::Launch(TEXT("Task_MutableGetMeshes_GetMesh_Post"), [=]()
@@ -2693,13 +2693,13 @@ namespace impl
 				OperationData->MutableSystem->SetImagePixelConversionOverride( OperationData->PixelFormatOverride );
 			}
 			
-			TSharedPtr<const mu::FInstance> Instance = OperationData->MutableSystem->BeginUpdate_MutableThread(OperationData->InstanceID,
+			TSharedPtr<const UE::Mutable::Private::FInstance> Instance = OperationData->MutableSystem->BeginUpdate_MutableThread(OperationData->InstanceID,
 				OperationData->Parameters,
 				OperationData->MeshIdRegistry,
 				OperationData->ImageIdRegistry,
 				OperationData->MaterialIdRegistry,
 				OperationData->GetCapturedDescriptor().GetState(),
-				mu::FSystem::AllLODs);
+				UE::Mutable::Private::FSystem::AllLODs);
 
 			OperationData->MutableInstance = Instance;
 			
@@ -2757,13 +2757,13 @@ namespace impl
 					
 					LOD.MeshId = Instance->GetMeshId(InstanceComponentIndex, LODIndex);
 					
-					mu::EMeshContentFlags MeshContentFilter = mu::EMeshContentFlags::AllFlags; 
+					UE::Mutable::Private::EMeshContentFlags MeshContentFilter = UE::Mutable::Private::EMeshContentFlags::AllFlags; 
 					
 					if (!bForceGeometryGeneration)
 					{
 						if (LODIndex < FirstResidentLOD)
 						{
-							EnumRemoveFlags(MeshContentFilter, mu::EMeshContentFlags::GeometryData);
+							EnumRemoveFlags(MeshContentFilter, UE::Mutable::Private::EMeshContentFlags::GeometryData);
 						}
 					}
 					
@@ -2777,8 +2777,8 @@ namespace impl
 			{
 				LastTask = UE::Tasks::Launch(TEXT("Task_MutableGetMeshes_GetMesh_GetMaterialOverride"), [InstanceComponentIndex, OperationData]()
 				{
-					mu::FMaterialId MaterialId = OperationData->InstanceUpdateData.Components[InstanceComponentIndex].OverlayMaterialId;
-					UE::Tasks::TTask<TSharedPtr<const mu::FMaterial>> GetMaterialOverrideTask = OperationData->MutableSystem->GetMaterial(OperationData->InstanceID, MaterialId);
+					UE::Mutable::Private::FMaterialId MaterialId = OperationData->InstanceUpdateData.Components[InstanceComponentIndex].OverlayMaterialId;
+					UE::Tasks::TTask<TSharedPtr<const UE::Mutable::Private::FMaterial>> GetMaterialOverrideTask = OperationData->MutableSystem->GetMaterial(OperationData->InstanceID, MaterialId);
 
 					UE::Tasks::AddNested(UE::Tasks::Launch(TEXT("Task_MutableGetMeshes_GetMesh_GetMaterialOverride_Post"), [=]() mutable
 					{
@@ -2826,7 +2826,7 @@ namespace impl
 	void Task_Mutable_GetImages_GetImageDesc(
 		const TSharedRef<FUpdateContextPrivate>& OperationData,
 		double StartTime,
-		const TSharedPtr<TArray<mu::FImageId>>& ImagesInThisInstance,
+		const TSharedPtr<TArray<UE::Mutable::Private::FImageId>>& ImagesInThisInstance,
 		int32 ImageIndex)
 	{
 		MUTABLE_CPUPROFILER_SCOPE(Task_Mutable_GetImages_GetImageDesc)
@@ -2834,7 +2834,7 @@ namespace impl
 		FInstanceUpdateData::FImage& Image = OperationData->InstanceUpdateData.Images[ImageIndex];
 		
 		// This should only be done when using progressive images, since GetImageDesc does some actual processing.
-		UE::Tasks::TTask<mu::FExtendedImageDesc> GetImageDescTask = OperationData->MutableSystem->GetImageDesc(OperationData->InstanceID, Image.ImageID);
+		UE::Tasks::TTask<UE::Mutable::Private::FExtendedImageDesc> GetImageDescTask = OperationData->MutableSystem->GetImageDesc(OperationData->InstanceID, Image.ImageID);
 		
 		UE::Tasks::AddNested(UE::Tasks::Launch(TEXT("Task_Mutable_GetImages_GetImage"), [=]()
 		{
@@ -2849,9 +2849,9 @@ namespace impl
 	void Task_Mutable_GetImages_GetImage_Post(
 		const TSharedRef<FUpdateContextPrivate>& OperationData,
     	double StartTime,
-    	const TSharedPtr<TArray<mu::FImageId>>& ImagesInThisInstance,
+    	const TSharedPtr<TArray<UE::Mutable::Private::FImageId>>& ImagesInThisInstance,
     	int32 ImageIndex,
-    	UE::Tasks::TTask<TSharedPtr<const mu::FImage>> GetImageTask,
+    	UE::Tasks::TTask<TSharedPtr<const UE::Mutable::Private::FImage>> GetImageTask,
     	int32 MipSizeX,
 		int32 MipSizeY,
 		int32 FullLODContent,
@@ -2871,7 +2871,7 @@ namespace impl
 		{
 			// Generate a correctly-sized but empty image instead, to avoid crashes.
 			UE_LOG(LogMutable, Warning, TEXT("Mutable generated a wrongly-sized image %s."), *Image.ImageID.ToString());
-			Image.Image = MakeShared<mu::FImage>(MipSizeX, MipSizeY, FullLODContent - MipsToSkip, Image.Image->GetFormat(), mu::EInitializationType::Black);
+			Image.Image = MakeShared<UE::Mutable::Private::FImage>(MipSizeX, MipSizeY, FullLODContent - MipsToSkip, Image.Image->GetFormat(), UE::Mutable::Private::EInitializationType::Black);
 		}
 
 		// We need one mip or the complete chain. Otherwise there was a bug.
@@ -2889,7 +2889,7 @@ namespace impl
 			UE_LOG(LogMutable, Warning, TEXT("Mutable generated an incomplete mip chain for image %s."), *Image.ImageID.ToString());
 
 			// Force the right number of mips. The missing data will be black.
-			const TSharedPtr<mu::FImage> NewImage = MakeShared<mu::FImage>(Image.Image->GetSizeX(), Image.Image->GetSizeY(), FullMipCount, Image.Image->GetFormat(), mu::EInitializationType::Black);
+			const TSharedPtr<UE::Mutable::Private::FImage> NewImage = MakeShared<UE::Mutable::Private::FImage>(Image.Image->GetSizeX(), Image.Image->GetSizeY(), FullMipCount, Image.Image->GetFormat(), UE::Mutable::Private::EInitializationType::Black);
 			check(NewImage);	
 			// Formats with BytesPerBlock == 0 will not allocate memory. This type of images are not expected here.
 			check(!NewImage->DataStorage.IsEmpty());
@@ -2914,13 +2914,13 @@ namespace impl
 	void Task_Mutable_GetImages_GetImage(
 		const TSharedRef<FUpdateContextPrivate>& OperationData,
 		double StartTime,
-		const TSharedPtr<TArray<mu::FImageId>>& ImagesInThisInstance,
+		const TSharedPtr<TArray<UE::Mutable::Private::FImageId>>& ImagesInThisInstance,
 		int32 ImageIndex,
-		UE::Tasks::TTask<mu::FExtendedImageDesc> GetImageDescTask)
+		UE::Tasks::TTask<UE::Mutable::Private::FExtendedImageDesc> GetImageDescTask)
 	{
 		MUTABLE_CPUPROFILER_SCOPE(Task_Mutable_GetImages_GetImage)
 
-		mu::FExtendedImageDesc& ImageDesc = GetImageDescTask.GetResult();
+		UE::Mutable::Private::FExtendedImageDesc& ImageDesc = GetImageDescTask.GetResult();
 
 		FInstanceUpdateData::FImage& Image = OperationData->InstanceUpdateData.Images[ImageIndex];
 		Image.ConstantImagesNeededToGenerate = MoveTemp(ImageDesc.ConstantImagesNeededToGenerate);
@@ -2976,14 +2976,14 @@ namespace impl
 		const int32 MipSizeY = FMath::Max(Image.FullImageSizeY >> MipsToSkip, 1);
 		if (MipsToSkip > 0 && CustomizableObjectSystemPrivate->EnableSkipGenerateResidentMips != 0 && OperationData->LowPriorityTextures.Find(Image.Name.ToString()) != INDEX_NONE)
 		{
-			TSharedPtr<const mu::FImage> NewImage = MakeShared<mu::FImage>(MipSizeX, MipSizeY, FullLODCount - MipsToSkip, ImageDesc.m_format, mu::EInitializationType::Black);
+			TSharedPtr<const UE::Mutable::Private::FImage> NewImage = MakeShared<UE::Mutable::Private::FImage>(MipSizeX, MipSizeY, FullLODCount - MipsToSkip, ImageDesc.m_format, UE::Mutable::Private::EInitializationType::Black);
 
-			UE::Tasks::TTask<TSharedPtr<const mu::FImage>> DummyTask = UE::Tasks::MakeCompletedTask<TSharedPtr<const mu::FImage>>(NewImage);
+			UE::Tasks::TTask<TSharedPtr<const UE::Mutable::Private::FImage>> DummyTask = UE::Tasks::MakeCompletedTask<TSharedPtr<const UE::Mutable::Private::FImage>>(NewImage);
 			Task_Mutable_GetImages_GetImage_Post(OperationData, StartTime, ImagesInThisInstance, ImageIndex, DummyTask, MipSizeX, MipSizeY, FullLODCount, MipsToSkip);
 		}
 		else
 		{
-			const UE::Tasks::TTask<TSharedPtr<const mu::FImage>> GetImageTask = OperationData->MutableSystem->GetImage(OperationData->InstanceID, Image.ImageID, Image.BaseMip + MipsToSkip, Image.BaseLOD);
+			const UE::Tasks::TTask<TSharedPtr<const UE::Mutable::Private::FImage>> GetImageTask = OperationData->MutableSystem->GetImage(OperationData->InstanceID, Image.ImageID, Image.BaseMip + MipsToSkip, Image.BaseLOD);
 			
 			UE::Tasks::AddNested(UE::Tasks::Launch(TEXT("Task_Mutable_GetImages_GetImage_Post"), [=]()
 			{
@@ -2998,7 +2998,7 @@ namespace impl
 	void Task_Mutable_GetImages_Loop(
 		const TSharedRef<FUpdateContextPrivate>& OperationData,
 		double StartTime,
-		const TSharedPtr<TArray<mu::FImageId>>& ImagesInThisInstance,
+		const TSharedPtr<TArray<UE::Mutable::Private::FImageId>>& ImagesInThisInstance,
 		int32 ImageIndex)
 	{
 		MUTABLE_CPUPROFILER_SCOPE(Task_Mutable_GetImages_Loop)
@@ -3029,7 +3029,7 @@ namespace impl
 
 			const double StartTime = FPlatformTime::Seconds();		
 
-			const TSharedPtr<TArray<mu::FImageId>> ImagesInThisInstance = MakeShared<TArray<mu::FImageId>>();
+			const TSharedPtr<TArray<UE::Mutable::Private::FImageId>> ImagesInThisInstance = MakeShared<TArray<UE::Mutable::Private::FImageId>>();
 			Task_Mutable_GetImages_Loop(OperationData, StartTime, ImagesInThisInstance, 0);
 		}
 		
@@ -3044,7 +3044,7 @@ namespace impl
 
 
 	// This runs in a worker thread.
-	void Task_Mutable_ReleaseInstance(mu::FInstance::FID InstanceID, TSharedPtr<mu::FSystem> MutableSystem, bool bLiveUpdateMode)
+	void Task_Mutable_ReleaseInstance(UE::Mutable::Private::FInstance::FID InstanceID, TSharedPtr<UE::Mutable::Private::FSystem> MutableSystem, bool bLiveUpdateMode)
 	{
 		MUTABLE_CPUPROFILER_SCOPE(Task_Mutable_ReleaseInstance)
 
@@ -3072,7 +3072,7 @@ namespace impl
 
 
 	// This runs in a worker thread.
-	void Task_Mutable_ReleaseInstanceID(const mu::FInstance::FID InstanceID, const TSharedPtr<mu::FSystem>& MutableSystem)
+	void Task_Mutable_ReleaseInstanceID(const UE::Mutable::Private::FInstance::FID InstanceID, const TSharedPtr<UE::Mutable::Private::FSystem>& MutableSystem)
 	{
 		MUTABLE_CPUPROFILER_SCOPE(Task_Mutable_ReleaseInstanceID)
 
@@ -3094,8 +3094,8 @@ namespace impl
 
 		check(OperationData);
 		
-		TMap<mu::FImageId, FTexturePlatformData*>& ImageToPlatformDataMap = OperationData->ImageToPlatformDataMap;
-		for (const TPair<mu::FImageId, FTexturePlatformData*>& Pair : ImageToPlatformDataMap)
+		TMap<UE::Mutable::Private::FImageId, FTexturePlatformData*>& ImageToPlatformDataMap = OperationData->ImageToPlatformDataMap;
+		for (const TPair<UE::Mutable::Private::FImageId, FTexturePlatformData*>& Pair : ImageToPlatformDataMap)
 		{
 			delete Pair.Value; // If this is not null then it must mean it hasn't been used, otherwise they would have taken ownership and nulled it
 		}
@@ -3207,7 +3207,7 @@ namespace impl
 		if (FLogBenchmarkUtil::IsBenchmarkingReportingEnabled())
 		{
 			// Memory used in the context of this the update of mesh
-			OperationData->UpdateEndPeakBytes = mu::FGlobalMemoryCounter::GetPeak();
+			OperationData->UpdateEndPeakBytes = UE::Mutable::Private::FGlobalMemoryCounter::GetPeak();
 			// Memory used in the context of the mesh update + the baseline memory used by mutable when starting the update
 			OperationData->UpdateEndRealPeakBytes = OperationData->UpdateEndPeakBytes + OperationData->UpdateStartBytes;
 		}
@@ -3216,7 +3216,7 @@ namespace impl
 		
 		// Next Task: Release Mutable. We need this regardless if we cancel or not
 		//-------------------------------------------------------------		
-		const  TSharedPtr<mu::FSystem> MutableSystem = CustomizableObjectSystemPrivate->MutableSystem;
+		const  TSharedPtr<UE::Mutable::Private::FSystem> MutableSystem = CustomizableObjectSystemPrivate->MutableSystem;
 		CustomizableObjectSystemPrivate->LastUpdateMutableTask = CustomizableObjectSystemPrivate->MutableTaskGraph.AddMutableThreadTask(
 			TEXT("Task_Mutable_ReleaseInstance"),
 			[InstanceID = OperationData->InstanceID, MutableSystem, bLiveUpdateMode = OperationData->bLiveUpdateMode]()
@@ -3385,14 +3385,14 @@ namespace impl
 
 
 	/** Enqueue the release ID operation in the Mutable queue */
-	void Task_Game_ReleaseInstanceID(const mu::FInstance::FID IDToRelease)
+	void Task_Game_ReleaseInstanceID(const UE::Mutable::Private::FInstance::FID IDToRelease)
 	{
 		MUTABLE_CPUPROFILER_SCOPE(Task_Game_ReleaseInstanceID)
 
 		UCustomizableObjectSystem* System = UCustomizableObjectSystem::GetInstanceChecked();
 		UCustomizableObjectSystemPrivate* SystemPrivate = System->GetPrivate();
 
-		const TSharedPtr<mu::FSystem> MutableSystem = SystemPrivate->MutableSystem;
+		const TSharedPtr<UE::Mutable::Private::FSystem> MutableSystem = SystemPrivate->MutableSystem;
 
 		// Task: Release Instance ID
 		//-------------------------------------------------------------
@@ -3674,9 +3674,9 @@ namespace impl
 			if (FLogBenchmarkUtil::IsBenchmarkingReportingEnabled())
 			{
 				// Get the amount of mutable memory in use now
-				Operation->UpdateStartBytes = mu::FGlobalMemoryCounter::GetAbsoluteCounter();
+				Operation->UpdateStartBytes = UE::Mutable::Private::FGlobalMemoryCounter::GetAbsoluteCounter();
 				// Reset the counter to later get the peak during the updated
-				mu::FGlobalMemoryCounter::Zero();													
+				UE::Mutable::Private::FGlobalMemoryCounter::Zero();													
 			}
 
 			// Prepare streaming for the current customizable object
@@ -3686,7 +3686,7 @@ namespace impl
 			Operation->bLowPriorityTasksBlocked = true;
 			SystemPrivate->MutableTaskGraph.AllowLaunchingMutableTaskLowPriority(false, false);
 
-			const TSharedPtr<mu::FSystem> MutableSystem = SystemPrivate->MutableSystem;
+			const TSharedPtr<UE::Mutable::Private::FSystem> MutableSystem = SystemPrivate->MutableSystem;
 
 			const TSharedRef<FBoneNames> BoneNames = MakeShared<FBoneNames>(*Object->GetPrivate()->GetModelResources());
 			Operation->BoneNames = BoneNames;
@@ -3722,7 +3722,7 @@ namespace impl
 				Operation->ImageIdRegistry,
 				Operation->MaterialIdRegistry,
 				Operation->GetCapturedDescriptor().GetState(),
-				mu::FSystem::AllLODs);
+				UE::Mutable::Private::FSystem::AllLODs);
 			
 			Operation->NumInstanceComponents = Operation->MutableInstance->GetComponentCount();
 		}
@@ -3736,7 +3736,7 @@ namespace impl
 		for (int32 InstanceComponentIndex = 0; InstanceComponentIndex < Operation->NumInstanceComponents; ++InstanceComponentIndex)
 		{
 			FCustomizableObjectComponentIndex ObjectComponentIndex(Operation->MutableInstance->GetComponentId(InstanceComponentIndex));
-			TArray<mu::FMeshId>* MeshId = Operation->GetMeshDescriptors(ObjectComponentIndex);
+			TArray<UE::Mutable::Private::FMeshId>* MeshId = Operation->GetMeshDescriptors(ObjectComponentIndex);
 			check(MeshId);
 
 			if (!MeshId)
@@ -3760,7 +3760,7 @@ namespace impl
 
 		if (Operation->bUseMeshCache)
 		{
-			for (const TArray<mu::FMeshId>& MeshId : Operation->GetMeshDescriptors())
+			for (const TArray<UE::Mutable::Private::FMeshId>& MeshId : Operation->GetMeshDescriptors())
 			{
 				if (UCustomizableObjectSkeletalMesh* CachedMesh = ObjectPrivate->MeshCache.Get(MeshId))
 				{
@@ -4216,7 +4216,7 @@ void UCustomizableObjectSystem::DiscardInstances()
 									{
 										// Force load the reference mesh if necessary. 
 										TSoftObjectPtr<USkeletalMesh> SoftObjectPtr = ModelResources->ReferenceSkeletalMeshesData[ObjectComponentIndex].SoftSkeletalMesh;
-										SkeletalMesh = MutablePrivate::LoadObject(SoftObjectPtr);
+										SkeletalMesh = UE::Mutable::Private::LoadObject(SoftObjectPtr);
 									}
 
 									Parent->EmptyOverrideMaterials();
@@ -4660,7 +4660,7 @@ int32 UCustomizableObjectSystemPrivate::BlockTillAllRequestsFinished(float TimeL
 }
 
 
-void LoadMorphTargetsData(FMutableStreamRequest& MutableDataStreamer, const TSharedRef<const mu::FMesh>& Mesh, TMap<uint32, FMorphTargetMeshData>& StreamingResult)
+void LoadMorphTargetsData(FMutableStreamRequest& MutableDataStreamer, const TSharedRef<const UE::Mutable::Private::FMesh>& Mesh, TMap<uint32, FMorphTargetMeshData>& StreamingResult)
 {
 	MUTABLE_CPUPROFILER_SCOPE(LoadMorphTargetsData);
 
@@ -4714,14 +4714,14 @@ void LoadMorphTargetsData(FMutableStreamRequest& MutableDataStreamer, const TSha
 
 		MutableDataStreamer.AddBlock(
 			Block,
-			MutablePrivate::EStreamableDataType::RealTimeMorph,
+			UE::Mutable::Private::EStreamableDataType::RealTimeMorph,
 			0,
 			MakeArrayView(reinterpret_cast<uint8*>(ReadDestData.Data.GetData()), ReadDestData.Data.Num() * sizeof(FMorphTargetVertexData)));
 	}
 }
 
 
-void LoadMorphTargetsMetadata(const FMutableStreamRequest& MutableDataStreamer, const TSharedRef<const mu::FMesh>& Mesh, TMap<uint32, FMorphTargetMeshData>& StreamingResult)
+void LoadMorphTargetsMetadata(const FMutableStreamRequest& MutableDataStreamer, const TSharedRef<const UE::Mutable::Private::FMesh>& Mesh, TMap<uint32, FMorphTargetMeshData>& StreamingResult)
 {
 	MUTABLE_CPUPROFILER_SCOPE(LoadMorphTargetsMetadata);
 
@@ -4756,7 +4756,7 @@ void LoadMorphTargetsMetadata(const FMutableStreamRequest& MutableDataStreamer,
 }
 
 
-void LoadClothing(FMutableStreamRequest& MutableDataStreamer, const TSharedRef<const mu::FMesh>& Mesh, TMap<uint32, FClothingMeshData>& StreamingResult)
+void LoadClothing(FMutableStreamRequest& MutableDataStreamer, const TSharedRef<const UE::Mutable::Private::FMesh>& Mesh, TMap<uint32, FClothingMeshData>& StreamingResult)
 {
 	MUTABLE_CPUPROFILER_SCOPE(LoadClothing);
 
@@ -4809,22 +4809,22 @@ void LoadClothing(FMutableStreamRequest& MutableDataStreamer, const TSharedRef<c
 
 		MutableDataStreamer.AddBlock(
 			Block,
-			MutablePrivate::EStreamableDataType::Clothing,
+			UE::Mutable::Private::EStreamableDataType::Clothing,
 			0,
 			MakeArrayView(reinterpret_cast<uint8*>(ReadDestData.Data.GetData()), ReadDestData.Data.Num() * sizeof(FCustomizableObjectMeshToMeshVertData)));
 	}
 }
 
 
-void ReconstructMorphTargets(const mu::FMesh& Mesh, const TArray<FName>& GlobalNames, const TMap<uint32, FMappedMorphTargetMeshData>& MappedMorphTargets, TArray<FMorphTargetLODModel>& OutMorphTargets)
+void ReconstructMorphTargets(const UE::Mutable::Private::FMesh& Mesh, const TArray<FName>& GlobalNames, const TMap<uint32, FMappedMorphTargetMeshData>& MappedMorphTargets, TArray<FMorphTargetLODModel>& OutMorphTargets)
 {
-	const mu::FMeshBufferSet& MeshSet = Mesh.GetVertexBuffers();
+	const UE::Mutable::Private::FMeshBufferSet& MeshSet = Mesh.GetVertexBuffers();
 
 	int32 VertexMorphsInfoIndexAndCountBufferIndex, VertexMorphsInfoIndexAndCountBufferChannel;
-	MeshSet.FindChannel(mu::EMeshBufferSemantic::Other, 0, &VertexMorphsInfoIndexAndCountBufferIndex, &VertexMorphsInfoIndexAndCountBufferChannel);
+	MeshSet.FindChannel(UE::Mutable::Private::EMeshBufferSemantic::Other, 0, &VertexMorphsInfoIndexAndCountBufferIndex, &VertexMorphsInfoIndexAndCountBufferChannel);
 
 	int32 VertexMorphsResourceIdBufferIndex, VertexMorphsResourceIdBufferChannel;
-	MeshSet.FindChannel(mu::EMeshBufferSemantic::Other, 1, &VertexMorphsResourceIdBufferIndex, &VertexMorphsResourceIdBufferChannel);
+	MeshSet.FindChannel(UE::Mutable::Private::EMeshBufferSemantic::Other, 1, &VertexMorphsResourceIdBufferIndex, &VertexMorphsResourceIdBufferChannel);
 
 	if (VertexMorphsInfoIndexAndCountBufferIndex < 0 || VertexMorphsResourceIdBufferIndex < 0)
 	{
diff --git a/Engine/Plugins/Mutable/Source/CustomizableObject/Private/MuCO/LoadUtils.cpp b/Engine/Plugins/Mutable/Source/CustomizableObject/Private/MuCO/LoadUtils.cpp
--- a/Engine/Plugins/Mutable/Source/CustomizableObject/Private/MuCO/LoadUtils.cpp
+++ b/Engine/Plugins/Mutable/Source/CustomizableObject/Private/MuCO/LoadUtils.cpp
@@ -5,7 +5,7 @@
 #include "AssetRegistry/AssetData.h"
 
 
-namespace MutablePrivate
+namespace UE::Mutable::Private
 {
 	UObject* LoadObject(const FAssetData& DataAsset)
 	{
