[Mutable] Change namespace name

diff --git a/Engine/Plugins/Mutable/Source/CustomizableObject/Private/MuCO/MutableMeshBufferUtils.cpp b/Engine/Plugins/Mutable/Source/CustomizableObject/Private/MuCO/MutableMeshBufferUtils.cpp
--- a/Engine/Plugins/Mutable/Source/CustomizableObject/Private/MuCO/MutableMeshBufferUtils.cpp
+++ b/Engine/Plugins/Mutable/Source/CustomizableObject/Private/MuCO/MutableMeshBufferUtils.cpp
@@ -7,9 +7,9 @@
 #include "Rendering/StaticMeshVertexBuffer.h"
 
 
-void MutableMeshBufferUtils::SetupVertexPositionsBuffer(const int32& InCurrentVertexBuffer, mu::FMeshBufferSet& OutTargetVertexBuffers)
+void MutableMeshBufferUtils::SetupVertexPositionsBuffer(const int32& InCurrentVertexBuffer, UE::Mutable::Private::FMeshBufferSet& OutTargetVertexBuffers)
 {
-	using namespace mu;
+	using namespace UE::Mutable::Private;
 	const int32 ElementSize = sizeof(FPositionVertex);
 	constexpr int32 ChannelCount = 1;
 	const EMeshBufferSemantic Semantics[ChannelCount] = {EMeshBufferSemantic::Position};
@@ -27,13 +27,13 @@ void MutableMeshBufferUtils::SetupVertexPositionsBuffer(const int32& InCurrentVe
 
 
 void MutableMeshBufferUtils::SetupTangentBuffer(const int32& InCurrentVertexBuffer,
-                                                mu::FMeshBufferSet& OutTargetVertexBuffers)
+                                                UE::Mutable::Private::FMeshBufferSet& OutTargetVertexBuffers)
 {
 	// \todo: support for high precision?
 	typedef TStaticMeshVertexTangentDatum<TStaticMeshVertexTangentTypeSelector<
 		EStaticMeshVertexTangentBasisType::Default>::TangentTypeT> TangentType;
 
-	using namespace mu;
+	using namespace UE::Mutable::Private;
 	const int32 ElementSize = sizeof(TangentType);
 	constexpr int32 ChannelCount = 2;
 	const EMeshBufferSemantic Semantics[ChannelCount] = {EMeshBufferSemantic::Tangent, EMeshBufferSemantic::Normal};
@@ -52,13 +52,13 @@ void MutableMeshBufferUtils::SetupTangentBuffer(const int32& InCurrentVertexBuff
 
 
 void MutableMeshBufferUtils::SetupTexCoordinatesBuffer(const int32& InCurrentVertexBuffer, const int32& InChannelCount, bool bHighPrecision,
-                                                       mu::FMeshBufferSet& OutTargetVertexBuffers,const int32* InTextureSemanticIndicesOverride /* = nullptr */)
+                                                       UE::Mutable::Private::FMeshBufferSet& OutTargetVertexBuffers,const int32* InTextureSemanticIndicesOverride /* = nullptr */)
 {	
 	int32 UVSize = bHighPrecision 
 		? sizeof(TStaticMeshVertexUVsDatum<TStaticMeshVertexUVsTypeSelector<EStaticMeshVertexUVType::HighPrecision>::UVsTypeT>)
 		: sizeof(TStaticMeshVertexUVsDatum<TStaticMeshVertexUVsTypeSelector<EStaticMeshVertexUVType::Default>::UVsTypeT>);
 
-	using namespace mu;
+	using namespace UE::Mutable::Private;
 	const int32 ElementSize = UVSize * InChannelCount;
 	constexpr int32 MaxChannelCount = MaxTexCordChannelCount;
 	const EMeshBufferSemantic Semantics[MaxChannelCount] = { EMeshBufferSemantic::TexCoords, EMeshBufferSemantic::TexCoords, EMeshBufferSemantic::TexCoords, EMeshBufferSemantic::TexCoords };
@@ -80,9 +80,9 @@ void MutableMeshBufferUtils::SetupSkinBuffer(const int32& InCurrentVertexBuffer,
 	const int32& MaxBoneIndexTypeSizeBytes,
 	const int32& MaxBoneWeightTypeSizeBytes,
 	const int32& MaxNumBonesPerVertex,
-	mu::FMeshBufferSet& OutTargetVertexBuffers)
+	UE::Mutable::Private::FMeshBufferSet& OutTargetVertexBuffers)
 {
-	using namespace mu;
+	using namespace UE::Mutable::Private;
 	const int32 ElementSize = (MaxBoneWeightTypeSizeBytes + MaxBoneIndexTypeSizeBytes) * MaxNumBonesPerVertex;
 	constexpr int32 ChannelCount = 2;
 	const EMeshBufferSemantic Semantics[ChannelCount] = {EMeshBufferSemantic::BoneIndices, EMeshBufferSemantic::BoneWeights};
@@ -92,32 +92,32 @@ void MutableMeshBufferUtils::SetupSkinBuffer(const int32& InCurrentVertexBuffer,
 	switch (MaxBoneIndexTypeSizeBytes)
 	{
 	case 0: // Fallback to something in this case.
-	case 1: Formats[0] = mu::EMeshBufferFormat::UInt8;
+	case 1: Formats[0] = UE::Mutable::Private::EMeshBufferFormat::UInt8;
 		break;
-	case 2: Formats[0] = mu::EMeshBufferFormat::UInt16;
+	case 2: Formats[0] = UE::Mutable::Private::EMeshBufferFormat::UInt16;
 		break;
-	case 4: Formats[0] = mu::EMeshBufferFormat::UInt32;
+	case 4: Formats[0] = UE::Mutable::Private::EMeshBufferFormat::UInt32;
 		break;
 	default:
 		// unsupported bone index type
 		check(false);
-		Formats[0] = mu::EMeshBufferFormat::None;
+		Formats[0] = UE::Mutable::Private::EMeshBufferFormat::None;
 		break;
 	}
 
 	switch (MaxBoneWeightTypeSizeBytes)
 	{
 	case 0: // Fallback to something in this case.
-	case 1: Formats[1] = mu::EMeshBufferFormat::NUInt8;
+	case 1: Formats[1] = UE::Mutable::Private::EMeshBufferFormat::NUInt8;
 		break;
-	case 2: Formats[1] = mu::EMeshBufferFormat::NUInt16;
+	case 2: Formats[1] = UE::Mutable::Private::EMeshBufferFormat::NUInt16;
 		break;
-	case 4: Formats[1] = mu::EMeshBufferFormat::NUInt32;
+	case 4: Formats[1] = UE::Mutable::Private::EMeshBufferFormat::NUInt32;
 		break;
 	default:
 		// unsupported bone weight type
 		check(false);
-		Formats[1] = mu::EMeshBufferFormat::None;
+		Formats[1] = UE::Mutable::Private::EMeshBufferFormat::None;
 		break;
 	}
 
@@ -134,9 +134,9 @@ void MutableMeshBufferUtils::SetupSkinBuffer(const int32& InCurrentVertexBuffer,
 
 
 void MutableMeshBufferUtils::SetupVertexColorBuffer(const int32& InCurrentVertexBuffer,
-                                                    mu::FMeshBufferSet& OutTargetVertexBuffers)
+                                                    UE::Mutable::Private::FMeshBufferSet& OutTargetVertexBuffers)
 {
-	using namespace mu;
+	using namespace UE::Mutable::Private;
 	const int32 ElementSize = sizeof(FColor);
 	constexpr int32 ChannelCount = 1;
 	const EMeshBufferSemantic Semantics[ChannelCount] = {EMeshBufferSemantic::Color};
@@ -151,11 +151,11 @@ void MutableMeshBufferUtils::SetupVertexColorBuffer(const int32& InCurrentVertex
 }
 
 
-void MutableMeshBufferUtils::SetupIndexBuffer(mu::FMeshBufferSet& OutTargetIndexBuffers)
+void MutableMeshBufferUtils::SetupIndexBuffer(UE::Mutable::Private::FMeshBufferSet& OutTargetIndexBuffers)
 {
 	OutTargetIndexBuffers.SetBufferCount(1);
 
-	using namespace mu;
+	using namespace UE::Mutable::Private;
 	const int32 ElementSize = sizeof(uint32);
 	//SkeletalMesh->GetImportedResource()->LODModels[LOD].MultiSizeIndexContainer.GetDataTypeSize();
 	constexpr int32 ChannelCount = 1;
@@ -178,9 +178,9 @@ void MutableMeshBufferUtils::SetupSkinWeightProfileBuffer(const int32& InCurrent
 	const int32& MaxBoneWeightTypeSizeBytes,
 	const int32& MaxNumBonesPerVertex,
 	const int32 SemanticsIndex,
-	mu::FMeshBufferSet& OutTargetVertexBuffers)
+	UE::Mutable::Private::FMeshBufferSet& OutTargetVertexBuffers)
 {
-	using namespace mu;
+	using namespace UE::Mutable::Private;
 	const int32 ElementSize = sizeof(int32) + (MaxBoneIndexTypeSizeBytes + 1) * MaxNumBonesPerVertex;
 	constexpr int32 ChannelCount = 3;
 	const EMeshBufferSemantic Semantics[ChannelCount] = { EMeshBufferSemantic::AltSkinWeight, EMeshBufferSemantic::BoneIndices, EMeshBufferSemantic::BoneWeights };
@@ -189,35 +189,35 @@ void MutableMeshBufferUtils::SetupSkinWeightProfileBuffer(const int32& InCurrent
 	EMeshBufferFormat Formats[ChannelCount] = { EMeshBufferFormat::Int32, EMeshBufferFormat::UInt8, EMeshBufferFormat::NUInt8 };
 	switch (MaxBoneIndexTypeSizeBytes)
 	{
-	case 1: Formats[1] = mu::EMeshBufferFormat::UInt8;
+	case 1: Formats[1] = UE::Mutable::Private::EMeshBufferFormat::UInt8;
 		break;
-	case 2: Formats[1] = mu::EMeshBufferFormat::UInt16;
+	case 2: Formats[1] = UE::Mutable::Private::EMeshBufferFormat::UInt16;
 		break;
-	case 4: Formats[1] = mu::EMeshBufferFormat::UInt32;
+	case 4: Formats[1] = UE::Mutable::Private::EMeshBufferFormat::UInt32;
 		break;
 	default:
 		// unsupported bone index type
 		check(false);
-		Formats[1] = mu::EMeshBufferFormat::None;
+		Formats[1] = UE::Mutable::Private::EMeshBufferFormat::None;
 		break;
 	}
 
 	switch (MaxBoneWeightTypeSizeBytes)
 	{
-	case 1: Formats[2] = mu::EMeshBufferFormat::NUInt8;
+	case 1: Formats[2] = UE::Mutable::Private::EMeshBufferFormat::NUInt8;
 		break;
 	case 2:
 	{
 		unimplemented()
-		Formats[2] = mu::EMeshBufferFormat::NUInt16;
+		Formats[2] = UE::Mutable::Private::EMeshBufferFormat::NUInt16;
 		break;
 	}
-	case 4: Formats[2] = mu::EMeshBufferFormat::NUInt32;
+	case 4: Formats[2] = UE::Mutable::Private::EMeshBufferFormat::NUInt32;
 		break;
 	default:
 		// unsupported bone weight type
 		check(false);
-		Formats[2] = mu::EMeshBufferFormat::None;
+		Formats[2] = UE::Mutable::Private::EMeshBufferFormat::None;
 		break;
 	}
 
diff --git a/Engine/Plugins/Mutable/Source/CustomizableObject/Private/MuCO/MutableProjectorTypeUtils.cpp b/Engine/Plugins/Mutable/Source/CustomizableObject/Private/MuCO/MutableProjectorTypeUtils.cpp
--- a/Engine/Plugins/Mutable/Source/CustomizableObject/Private/MuCO/MutableProjectorTypeUtils.cpp
+++ b/Engine/Plugins/Mutable/Source/CustomizableObject/Private/MuCO/MutableProjectorTypeUtils.cpp
@@ -5,21 +5,21 @@
 #include "MuCO/CustomizableObjectParameterTypeDefinitions.h"
 #include "MuR/Parameters.h"
 
-ECustomizableObjectProjectorType ProjectorUtils::GetEquivalentProjectorType (mu::EProjectorType ProjectorType)
+ECustomizableObjectProjectorType ProjectorUtils::GetEquivalentProjectorType (UE::Mutable::Private::EProjectorType ProjectorType)
 {
 	// Translate projector type from Mutable Core enum type to CO enum type
 	switch (ProjectorType)
 	{
-	case mu::EProjectorType::Planar:
+	case UE::Mutable::Private::EProjectorType::Planar:
 		return ECustomizableObjectProjectorType::Planar;
 		
-	case mu::EProjectorType::Cylindrical:
+	case UE::Mutable::Private::EProjectorType::Cylindrical:
 		return ECustomizableObjectProjectorType::Cylindrical;
 		
-	case mu::EProjectorType::Wrapping:
+	case UE::Mutable::Private::EProjectorType::Wrapping:
 		return ECustomizableObjectProjectorType::Wrapping;
 		
-	case mu::EProjectorType::Count:
+	case UE::Mutable::Private::EProjectorType::Count:
 	default:
 		checkNoEntry();
 		return ECustomizableObjectProjectorType::Planar;
@@ -27,21 +27,21 @@ ECustomizableObjectProjectorType ProjectorUtils::GetEquivalentProjectorType (mu:
 }
 
 
-mu::EProjectorType ProjectorUtils::GetEquivalentProjectorType (ECustomizableObjectProjectorType ProjectorType)
+UE::Mutable::Private::EProjectorType ProjectorUtils::GetEquivalentProjectorType (ECustomizableObjectProjectorType ProjectorType)
 {
-	if (GetEquivalentProjectorType(mu::EProjectorType::Planar) == ProjectorType)
+	if (GetEquivalentProjectorType(UE::Mutable::Private::EProjectorType::Planar) == ProjectorType)
 	{
-		return mu::EProjectorType::Planar;
+		return UE::Mutable::Private::EProjectorType::Planar;
 	}
-	if (GetEquivalentProjectorType(mu::EProjectorType::Wrapping) == ProjectorType)
+	if (GetEquivalentProjectorType(UE::Mutable::Private::EProjectorType::Wrapping) == ProjectorType)
 	{
-		return mu::EProjectorType::Wrapping;
+		return UE::Mutable::Private::EProjectorType::Wrapping;
 	}
-	if (GetEquivalentProjectorType(mu::EProjectorType::Cylindrical) == ProjectorType)
+	if (GetEquivalentProjectorType(UE::Mutable::Private::EProjectorType::Cylindrical) == ProjectorType)
 	{
-		return mu::EProjectorType::Cylindrical;
+		return UE::Mutable::Private::EProjectorType::Cylindrical;
 	}
 
 	checkNoEntry();
-	return mu::EProjectorType::Count;		// Invalid 
+	return UE::Mutable::Private::EProjectorType::Count;		// Invalid 
 }
diff --git a/Engine/Plugins/Mutable/Source/CustomizableObject/Private/MuCO/MutableStreamRequest.cpp b/Engine/Plugins/Mutable/Source/CustomizableObject/Private/MuCO/MutableStreamRequest.cpp
--- a/Engine/Plugins/Mutable/Source/CustomizableObject/Private/MuCO/MutableStreamRequest.cpp
+++ b/Engine/Plugins/Mutable/Source/CustomizableObject/Private/MuCO/MutableStreamRequest.cpp
@@ -25,7 +25,7 @@ const TSharedPtr<FModelStreamableBulkData>& FMutableStreamRequest::GetModelStrea
 }
 
 
-void FMutableStreamRequest::AddBlock(const FMutableStreamableBlock& Block, MutablePrivate::EStreamableDataType DataType, uint16 ResourceType, TArrayView<uint8> AllocatedMemoryView)
+void FMutableStreamRequest::AddBlock(const FMutableStreamableBlock& Block, UE::Mutable::Private::EStreamableDataType DataType, uint16 ResourceType, TArrayView<uint8> AllocatedMemoryView)
 {
 	if (bIsStreaming)
 	{
@@ -293,7 +293,7 @@ void FMutableStreamRequest::Cancel()
 
 
 FMutableStreamRequest::FBlockReadInfo::FBlockReadInfo(uint64 InOffset, IAsyncReadFileHandle* InFileHandle, TArrayView<uint8>& InAllocatedMemoryView,
-	uint32 InFileId, MutablePrivate::EStreamableDataType InDataType, uint16 InResourceType, uint16 InResourceFlags)
+	uint32 InFileId, UE::Mutable::Private::EStreamableDataType InDataType, uint16 InResourceType, uint16 InResourceFlags)
 	: Offset(InOffset)
 	, FileHandle(InFileHandle)
 	, AllocatedMemoryView(InAllocatedMemoryView)
diff --git a/Engine/Plugins/Mutable/Source/CustomizableObject/Private/MuCO/MutableStreamRequest.h b/Engine/Plugins/Mutable/Source/CustomizableObject/Private/MuCO/MutableStreamRequest.h
--- a/Engine/Plugins/Mutable/Source/CustomizableObject/Private/MuCO/MutableStreamRequest.h
+++ b/Engine/Plugins/Mutable/Source/CustomizableObject/Private/MuCO/MutableStreamRequest.h
@@ -30,7 +30,7 @@ public:
 	const TSharedPtr<FModelStreamableBulkData>& GetModelStreamableBulkData() const;
 
 	/** Add a block to stream. */
-	void AddBlock(const FMutableStreamableBlock& Block, MutablePrivate::EStreamableDataType DataType, uint16 InResourceType, TArrayView<uint8> AllocatedMemoryView);
+	void AddBlock(const FMutableStreamableBlock& Block, UE::Mutable::Private::EStreamableDataType DataType, uint16 InResourceType, TArrayView<uint8> AllocatedMemoryView);
 
 	/** Stream the previously requested data. */
 	UE::Tasks::FTask Stream();
@@ -42,7 +42,7 @@ private:
 	struct FBlockReadInfo
 	{
 		FBlockReadInfo() = default;
-		FBlockReadInfo(uint64 InOffset, IAsyncReadFileHandle* InFileHandle, TArrayView<uint8>& InAllocatedMemoryView, uint32 InFileId, MutablePrivate::EStreamableDataType InDataType, uint16 InResourceType, uint16 InResourceFlags);
+		FBlockReadInfo(uint64 InOffset, IAsyncReadFileHandle* InFileHandle, TArrayView<uint8>& InAllocatedMemoryView, uint32 InFileId, UE::Mutable::Private::EStreamableDataType InDataType, uint16 InResourceType, uint16 InResourceFlags);
 
 		uint64 Offset;
 		IAsyncReadFileHandle* FileHandle;
@@ -51,7 +51,7 @@ private:
 		uint16 ResourceType;
 		uint16 ResourceFlags;
 
-		MutablePrivate::EStreamableDataType DataType;
+		UE::Mutable::Private::EStreamableDataType DataType;
 	};
 	
 	const TSharedPtr<FModelStreamableBulkData> ModelStreamableBulkData;
diff --git a/Engine/Plugins/Mutable/Source/CustomizableObject/Private/MuCO/UnrealConversionUtils.cpp b/Engine/Plugins/Mutable/Source/CustomizableObject/Private/MuCO/UnrealConversionUtils.cpp
--- a/Engine/Plugins/Mutable/Source/CustomizableObject/Private/MuCO/UnrealConversionUtils.cpp
+++ b/Engine/Plugins/Mutable/Source/CustomizableObject/Private/MuCO/UnrealConversionUtils.cpp
@@ -172,21 +172,21 @@ namespace UnrealConversionUtils
 
 	void SetupRenderSections(
 		FSkeletalMeshLODRenderData& LODResource,
-		const mu::FMesh* InMutableMesh,
-		const TArray<mu::FBoneName>& InBoneMap,
-		const TMap<mu::FBoneName, TPair<FName, uint16>>& BoneInfoMap,
+		const UE::Mutable::Private::FMesh* InMutableMesh,
+		const TArray<UE::Mutable::Private::FBoneName>& InBoneMap,
+		const TMap<UE::Mutable::Private::FBoneName, TPair<FName, uint16>>& BoneInfoMap,
 		const int32 InFirstBoneMapIndex,
 		const TArray<const FMutableSurfaceMetadata*>& SurfacesMetadata)
 	{
 		check(InMutableMesh);
 
-		const mu::FMeshBufferSet& MutableMeshVertexBuffers = InMutableMesh->GetVertexBuffers();
+		const UE::Mutable::Private::FMeshBufferSet& MutableMeshVertexBuffers = InMutableMesh->GetVertexBuffers();
 
 		// Find the number of influences from this mesh
 		int32 NumBoneInfluences = 0;
 		int32 boneIndexBuffer = -1;
 		int32 boneIndexChannel = -1;
-		MutableMeshVertexBuffers.FindChannel(mu::EMeshBufferSemantic::BoneIndices, 0, &boneIndexBuffer, &boneIndexChannel);
+		MutableMeshVertexBuffers.FindChannel(UE::Mutable::Private::EMeshBufferSemantic::BoneIndices, 0, &boneIndexBuffer, &boneIndexChannel);
 		if (boneIndexBuffer >= 0 || boneIndexChannel >= 0)
 		{
 			MutableMeshVertexBuffers.GetChannel(boneIndexBuffer, boneIndexChannel,
@@ -235,7 +235,7 @@ namespace UnrealConversionUtils
 			{
 				if (InBoneMap.IsValidIndex(FirstBone))
 				{
-					mu::FBoneName FirstBoneName = InBoneMap[FirstBone];
+					UE::Mutable::Private::FBoneName FirstBoneName = InBoneMap[FirstBone];
 					if (const TPair<FName, uint16>* Found = BoneInfoMap.Find(FirstBoneName))
 					{
 						Section.BoneMap.Add(Found->Value);
@@ -248,12 +248,12 @@ namespace UnrealConversionUtils
 
 	void InitVertexBuffersWithDummyData(
 		FSkeletalMeshLODRenderData& LODResource,
-		const mu::FMesh* InMutableMesh,
+		const UE::Mutable::Private::FMesh* InMutableMesh,
 		const bool bAllowCPUAccess)
 	{
 		MUTABLE_CPUPROFILER_SCOPE(InitVertexBuffersWithDummyData);
 
-		const mu::FMeshBufferSet& MutableMeshVertexBuffers = InMutableMesh->GetVertexBuffers();
+		const UE::Mutable::Private::FMeshBufferSet& MutableMeshVertexBuffers = InMutableMesh->GetVertexBuffers();
 		check(MutableMeshVertexBuffers.GetElementCount() > 0);
 
 		const bool bUseFullPrecisionUVs = true;
@@ -276,21 +276,21 @@ namespace UnrealConversionUtils
 			LODResource.StaticVertexBuffers.StaticMeshVertexBuffer.SetMetaData(NumTexCoords, NumVertices, bUseFullPrecisionUVs, bUseFullPrecisionTangentBasis);
 		}
 
-		mu::EMeshBufferFormat BoneIndexFormat = mu::EMeshBufferFormat::None;
+		UE::Mutable::Private::EMeshBufferFormat BoneIndexFormat = UE::Mutable::Private::EMeshBufferFormat::None;
 		int32 NumBoneInfluences = 0;
 		int32 BoneIndexBuffer = -1;
 		int32 BoneIndexChannel = -1;
-		MutableMeshVertexBuffers.FindChannel(mu::EMeshBufferSemantic::BoneIndices, 0, &BoneIndexBuffer, &BoneIndexChannel);
+		MutableMeshVertexBuffers.FindChannel(UE::Mutable::Private::EMeshBufferSemantic::BoneIndices, 0, &BoneIndexBuffer, &BoneIndexChannel);
 		if (BoneIndexBuffer >= 0 || BoneIndexChannel >= 0)
 		{
 			MutableMeshVertexBuffers.GetChannel(BoneIndexBuffer, BoneIndexChannel,
 				nullptr, nullptr, &BoneIndexFormat, &NumBoneInfluences, nullptr);
 		}
 
-		mu::EMeshBufferFormat BoneWeightFormat = mu::EMeshBufferFormat::None;
+		UE::Mutable::Private::EMeshBufferFormat BoneWeightFormat = UE::Mutable::Private::EMeshBufferFormat::None;
 		int32 BoneWeightBuffer = -1;
 		int32 BoneWeightChannel = -1;
-		MutableMeshVertexBuffers.FindChannel(mu::EMeshBufferSemantic::BoneWeights, 0, &BoneWeightBuffer, &BoneWeightChannel);
+		MutableMeshVertexBuffers.FindChannel(UE::Mutable::Private::EMeshBufferSemantic::BoneWeights, 0, &BoneWeightBuffer, &BoneWeightChannel);
 		if (BoneWeightBuffer >= 0 || BoneWeightChannel >= 0)
 		{
 			MutableMeshVertexBuffers.GetChannel(BoneWeightBuffer, BoneWeightChannel,
@@ -299,8 +299,8 @@ namespace UnrealConversionUtils
 
 		// Skin Weights
 		{
-			const bool bUse16BitBoneIndex = BoneIndexFormat == mu::EMeshBufferFormat::UInt16;
-			const bool bUse16BitBoneWeights = BoneWeightFormat == mu::EMeshBufferFormat::NUInt16;
+			const bool bUse16BitBoneIndex = BoneIndexFormat == UE::Mutable::Private::EMeshBufferFormat::UInt16;
+			const bool bUse16BitBoneWeights = BoneWeightFormat == UE::Mutable::Private::EMeshBufferFormat::NUInt16;
 
 			LODResource.SkinWeightVertexBuffer.SetUse16BitBoneIndex(bUse16BitBoneIndex);
 			LODResource.SkinWeightVertexBuffer.SetUse16BitBoneWeight(bUse16BitBoneWeights);
@@ -327,15 +327,15 @@ namespace UnrealConversionUtils
 		{
 			if (MutableMeshVertexBuffers.GetBufferChannelCount(Buffer) > 0)
 			{
-				mu::EMeshBufferSemantic Semantic;
-				mu::EMeshBufferFormat Format;
+				UE::Mutable::Private::EMeshBufferSemantic Semantic;
+				UE::Mutable::Private::EMeshBufferFormat Format;
 				int32 SemanticIndex;
 				int32 ComponentCount;
 				int32 Offset;
 				MutableMeshVertexBuffers.GetChannel(Buffer, 0, &Semantic, &SemanticIndex, &Format, &ComponentCount, &Offset);
 
 				// colour buffer?
-				if (Semantic == mu::EMeshBufferSemantic::Color)
+				if (Semantic == UE::Mutable::Private::EMeshBufferSemantic::Color)
 				{
 					LODResource.StaticVertexBuffers.ColorVertexBuffer.Init(1, bAllowCPUAccess);
 					LODResource.StaticVertexBuffers.ColorVertexBuffer.SetMetaData(LODResource.StaticVertexBuffers.ColorVertexBuffer.GetStride(), NumVertices);
@@ -349,13 +349,13 @@ namespace UnrealConversionUtils
 
 	void CopyMutableVertexBuffers(
 		FSkeletalMeshLODRenderData& LODResource,
-		const mu::FMesh* MutableMesh,
+		const UE::Mutable::Private::FMesh* MutableMesh,
 		const bool bAllowCPUAccess)
 
 	{
 		MUTABLE_CPUPROFILER_SCOPE(CopyMutableVertexBuffers);
 
-		const mu::FMeshBufferSet& MutableMeshVertexBuffers = MutableMesh->GetVertexBuffers();
+		const UE::Mutable::Private::FMeshBufferSet& MutableMeshVertexBuffers = MutableMesh->GetVertexBuffers();
 		const int32 NumVertices = MutableMeshVertexBuffers.IsDescriptor()
 				? 0 
 				: MutableMeshVertexBuffers.GetElementCount();
@@ -366,10 +366,10 @@ namespace UnrealConversionUtils
 		{
 			int32 TexCoordsBufferIndex = -1;
 			int32 TexCoordsChannelIndex = -1;
-			MutableMesh->VertexBuffers.FindChannel(mu::EMeshBufferSemantic::TexCoords, 0, &TexCoordsBufferIndex, &TexCoordsChannelIndex);
+			MutableMesh->VertexBuffers.FindChannel(UE::Mutable::Private::EMeshBufferSemantic::TexCoords, 0, &TexCoordsBufferIndex, &TexCoordsChannelIndex);
 			if (TexCoordsBufferIndex >= 0 && TexCoordsChannelIndex >= 0)
 			{
-				bUseFullPrecisionUVs = MutableMesh->VertexBuffers.Buffers[TexCoordsBufferIndex].Channels[TexCoordsChannelIndex].Format == mu::EMeshBufferFormat::Float32;
+				bUseFullPrecisionUVs = MutableMesh->VertexBuffers.Buffers[TexCoordsBufferIndex].Channels[TexCoordsChannelIndex].Format == UE::Mutable::Private::EMeshBufferFormat::Float32;
 			}
 		}
 
@@ -386,33 +386,33 @@ namespace UnrealConversionUtils
 			MutableMeshVertexBuffers.GetBufferData(MUTABLE_VERTEXBUFFER_TEXCOORDS)
 		);
 
-		mu::EMeshBufferFormat BoneIndexFormat = mu::EMeshBufferFormat::None;
+		UE::Mutable::Private::EMeshBufferFormat BoneIndexFormat = UE::Mutable::Private::EMeshBufferFormat::None;
 		int32 NumBoneInfluences = 0;
 		int32 BoneIndexBuffer = -1;
 		int32 BoneIndexChannel = -1;
-		MutableMeshVertexBuffers.FindChannel(mu::EMeshBufferSemantic::BoneIndices, 0, &BoneIndexBuffer, &BoneIndexChannel);
+		MutableMeshVertexBuffers.FindChannel(UE::Mutable::Private::EMeshBufferSemantic::BoneIndices, 0, &BoneIndexBuffer, &BoneIndexChannel);
 		if (BoneIndexBuffer >= 0 || BoneIndexChannel >= 0)
 		{
 			MutableMeshVertexBuffers.GetChannel(BoneIndexBuffer, BoneIndexChannel,
 				nullptr, nullptr, &BoneIndexFormat, &NumBoneInfluences, nullptr);
 		}
 
-		mu::EMeshBufferFormat BoneWeightFormat = mu::EMeshBufferFormat::None;
+		UE::Mutable::Private::EMeshBufferFormat BoneWeightFormat = UE::Mutable::Private::EMeshBufferFormat::None;
 		int32 BoneWeightBuffer = -1;
 		int32 BoneWeightChannel = -1;
-		MutableMeshVertexBuffers.FindChannel(mu::EMeshBufferSemantic::BoneWeights, 0, &BoneWeightBuffer, &BoneWeightChannel);
+		MutableMeshVertexBuffers.FindChannel(UE::Mutable::Private::EMeshBufferSemantic::BoneWeights, 0, &BoneWeightBuffer, &BoneWeightChannel);
 		if (BoneWeightBuffer >= 0 || BoneWeightChannel >= 0)
 		{
 			MutableMeshVertexBuffers.GetChannel(BoneWeightBuffer, BoneWeightChannel,
 				nullptr, nullptr, &BoneWeightFormat, nullptr, nullptr);
 		}
 
-		if (BoneIndexFormat == mu::EMeshBufferFormat::UInt16)
+		if (BoneIndexFormat == UE::Mutable::Private::EMeshBufferFormat::UInt16)
 		{
 			LODResource.SkinWeightVertexBuffer.SetUse16BitBoneIndex(true);
 		}
 
-		if (BoneWeightFormat == mu::EMeshBufferFormat::NUInt16)
+		if (BoneWeightFormat == UE::Mutable::Private::EMeshBufferFormat::NUInt16)
 		{
 			LODResource.SkinWeightVertexBuffer.SetUse16BitBoneWeight(true);
 		}
@@ -433,15 +433,15 @@ namespace UnrealConversionUtils
 		{
 			if (MutableMeshVertexBuffers.GetBufferChannelCount(Buffer) > 0)
 			{
-				mu::EMeshBufferSemantic Semantic;
-				mu::EMeshBufferFormat Format;
+				UE::Mutable::Private::EMeshBufferSemantic Semantic;
+				UE::Mutable::Private::EMeshBufferFormat Format;
 				int32 SemanticIndex;
 				int32 ComponentCount;
 				int32 Offset;
 				MutableMeshVertexBuffers.GetChannel(Buffer, 0, &Semantic, &SemanticIndex, &Format, &ComponentCount, &Offset);
 
 				// colour buffer?
-				if (Semantic == mu::EMeshBufferSemantic::Color)
+				if (Semantic == UE::Mutable::Private::EMeshBufferSemantic::Color)
 				{
 					const void* DataPtr = MutableMeshVertexBuffers.GetBufferData(Buffer);
 					FColorVertexBuffers_InitWithMutableData(LODResource.StaticVertexBuffers, NumVertices, DataPtr);
@@ -452,7 +452,7 @@ namespace UnrealConversionUtils
 	}
 
 		
-	void InitIndexBuffersWithDummyData(FSkeletalMeshLODRenderData& LODResource, const mu::FMesh* InMutableMesh)
+	void InitIndexBuffersWithDummyData(FSkeletalMeshLODRenderData& LODResource, const UE::Mutable::Private::FMesh* InMutableMesh)
 	{
 		MUTABLE_CPUPROFILER_SCOPE(InitIndexBuffersWithDummyData);
 
@@ -468,7 +468,7 @@ namespace UnrealConversionUtils
 
 	bool CopyMutableIndexBuffers(
 		FSkeletalMeshLODRenderData& LODResource,
-		const mu::FMesh* InMutableMesh,
+		const UE::Mutable::Private::FMesh* InMutableMesh,
 		const TArray<uint32>& SurfaceIDs,
 		bool& bOutMarkRenderStateDirty)
 	{
@@ -510,7 +510,7 @@ namespace UnrealConversionUtils
 			FSkelMeshRenderSection& Section = LODResource.RenderSections[SectionIndex];
 			
 			const int32 SurfaceID = SurfaceIDs[SectionIndex];
-			const mu::FMeshSurface* Surface = InMutableMesh->Surfaces.FindByPredicate([SurfaceID](const mu::FMeshSurface& Surface)
+			const UE::Mutable::Private::FMeshSurface* Surface = InMutableMesh->Surfaces.FindByPredicate([SurfaceID](const UE::Mutable::Private::FMeshSurface& Surface)
 			{
 				return Surface.Id == SurfaceID;
 			});
@@ -581,7 +581,7 @@ namespace UnrealConversionUtils
 		FSkeletalMeshLODRenderData& LODResource,
 		USkeletalMesh& SkeletalMesh,
 		int32 LODIndex,
-		const mu::FMesh* InMutableMesh,
+		const UE::Mutable::Private::FMesh* InMutableMesh,
 		const TArray<TPair<uint32, FName>>& ActiveSkinWeightProfiles)
 	{
 		MUTABLE_CPUPROFILER_SCOPE(CopyMutableSkinWeightProfilesBuffers);
@@ -597,14 +597,14 @@ namespace UnrealConversionUtils
 
 		TMap<FAuxMutableSkinWeightVertexKey, int32> HashToUniqueWeightIndexMap;
 
-		const mu::FMeshBufferSet& MutableMeshVertexBuffers = InMutableMesh->GetVertexBuffers();
+		const UE::Mutable::Private::FMeshBufferSet& MutableMeshVertexBuffers = InMutableMesh->GetVertexBuffers();
 		for (const TPair<uint32, FName>& Profile : ActiveSkinWeightProfiles)
 		{
 			int32 BufferIndex, ChannelIndex;
-			mu::EMeshBufferFormat Format;
+			UE::Mutable::Private::EMeshBufferFormat Format;
 			int32 MutNumInfluences;
 
-			MutableMeshVertexBuffers.FindChannel(mu::EMeshBufferSemantic::AltSkinWeight, Profile.Key, &BufferIndex, &ChannelIndex);
+			MutableMeshVertexBuffers.FindChannel(UE::Mutable::Private::EMeshBufferSemantic::AltSkinWeight, Profile.Key, &BufferIndex, &ChannelIndex);
 			
 			if (BufferIndex == INDEX_NONE)
 			{
@@ -618,13 +618,13 @@ namespace UnrealConversionUtils
 
 			// BoneIndices channel info
 			MutableMeshVertexBuffers.GetChannel(BufferIndex, 1, nullptr, nullptr, &Format, &MutNumInfluences, nullptr);
-			const uint8 MutBoneIndexTypeByteSize = Format == mu::EMeshBufferFormat::UInt16 ? 2 : 1;
+			const uint8 MutBoneIndexTypeByteSize = Format == UE::Mutable::Private::EMeshBufferFormat::UInt16 ? 2 : 1;
 			const uint8 MutBoneIndicesDataSize = MutBoneIndexTypeByteSize * MutNumInfluences;
 			check(BoneIndexTypeByteSize == MutBoneIndexTypeByteSize);
 
 			// BoneWeights channel info
 			MutableMeshVertexBuffers.GetChannel(BufferIndex, 2, nullptr, nullptr, &Format, &MutNumInfluences, nullptr);
-			const uint8 MutBoneWeightTypeByteSize = Format == mu::EMeshBufferFormat::NUInt16 ? 2 : 1;
+			const uint8 MutBoneWeightTypeByteSize = Format == UE::Mutable::Private::EMeshBufferFormat::NUInt16 ? 2 : 1;
 			const uint8 MutBoneWeightsDataSize = MutBoneWeightTypeByteSize * MutNumInfluences;
 			check(BoneWeightTypeByteSize == MutBoneWeightTypeByteSize);
 
@@ -833,7 +833,7 @@ namespace UnrealConversionUtils
 	}
 
 
-	void MorphTargetVertexInfoBuffers(FSkeletalMeshLODRenderData& LODResource, const USkeletalMesh& Owner, const mu::FMesh& MutableMesh, const TMap<uint32, FMorphTargetMeshData>& MorphTargetMeshData, int32 LODIndex)
+	void MorphTargetVertexInfoBuffers(FSkeletalMeshLODRenderData& LODResource, const USkeletalMesh& Owner, const UE::Mutable::Private::FMesh& MutableMesh, const TMap<uint32, FMorphTargetMeshData>& MorphTargetMeshData, int32 LODIndex)
 	{
 		if (Owner.GetMorphTargets().IsEmpty())
 		{
@@ -891,13 +891,13 @@ namespace UnrealConversionUtils
 	}
 
 	
-	UE::Tasks::FTask ConvertSkeletalMeshFromRuntimeData(USkeletalMesh* SkeletalMesh, int32 LODIndex, int32 SectionIndex, TSharedPtr<mu::FMesh> Result, const TSharedRef<FBoneNames>& BoneNames)
+	UE::Tasks::FTask ConvertSkeletalMeshFromRuntimeData(USkeletalMesh* SkeletalMesh, int32 LODIndex, int32 SectionIndex, TSharedPtr<UE::Mutable::Private::FMesh> Result, const TSharedRef<FBoneNames>& BoneNames)
 	{
 		MUTABLE_CPUPROFILER_SCOPE(ConvertSkeletalMeshFromRuntimeData);
 
 		if (!SkeletalMesh)
 		{
-			return UE::Tasks::MakeCompletedTask<TSharedPtr<mu::FMesh>>();
+			return UE::Tasks::MakeCompletedTask<TSharedPtr<UE::Mutable::Private::FMesh>>();
 		}
 		
 		struct FMeshConversionContext
@@ -1056,7 +1056,7 @@ namespace UnrealConversionUtils
 					}
 
 					// Create the skeleton, poses, and BoneMap for this mesh
-					TSharedPtr<mu::FSkeleton> MutableSkeleton = MakeShared<mu::FSkeleton>();
+					TSharedPtr<UE::Mutable::Private::FSkeleton> MutableSkeleton = MakeShared<UE::Mutable::Private::FSkeleton>();
 					Result->SetSkeleton(MutableSkeleton);
 
 					const int32 NumRequiredBones = RequiredBones.Num();
@@ -1064,7 +1064,7 @@ namespace UnrealConversionUtils
 					MutableSkeleton->SetBoneCount(NumRequiredBones);
 
 					// MutableBoneMap will not keep an index to the Skeleton, but to the BoneName
-					TArray<mu::FBoneName> MutableBoneMap;
+					TArray<UE::Mutable::Private::FBoneName> MutableBoneMap;
 					MutableBoneMap.SetNum(BoneMap.Num());
 
 					TArray<FMatrix> ComposedRefPoseMatrices;
@@ -1081,7 +1081,7 @@ namespace UnrealConversionUtils
 						// Set bone hierarchy
 
 						// Get the bone id
-						mu::FBoneName BoneName = BoneNames->FindOrAdd(BoneInfo.Name);
+						UE::Mutable::Private::FBoneName BoneName = BoneNames->FindOrAdd(BoneInfo.Name);
 						
 						MutableSkeleton->SetBoneName(BoneIndex, BoneName);
 						MutableSkeleton->SetBoneParent(BoneIndex, ParentBoneIndex);
@@ -1109,9 +1109,9 @@ namespace UnrealConversionUtils
 						FTransform3f BoneTransform;
 						BoneTransform.SetFromMatrix(FMatrix44f(ComposedRefPoseMatrices[BoneIndex]));
 
-						mu::EBoneUsageFlags BoneUsageFlags = mu::EBoneUsageFlags::None;
-						EnumAddFlags(BoneUsageFlags, BoneMapIndex != INDEX_NONE ? mu::EBoneUsageFlags::Skinning : mu::EBoneUsageFlags::None);
-						EnumAddFlags(BoneUsageFlags, ParentBoneIndex == INDEX_NONE ? mu::EBoneUsageFlags::Root : mu::EBoneUsageFlags::None);
+						UE::Mutable::Private::EBoneUsageFlags BoneUsageFlags = UE::Mutable::Private::EBoneUsageFlags::None;
+						EnumAddFlags(BoneUsageFlags, BoneMapIndex != INDEX_NONE ? UE::Mutable::Private::EBoneUsageFlags::Skinning : UE::Mutable::Private::EBoneUsageFlags::None);
+						EnumAddFlags(BoneUsageFlags, ParentBoneIndex == INDEX_NONE ? UE::Mutable::Private::EBoneUsageFlags::Root : UE::Mutable::Private::EBoneUsageFlags::None);
 
 						Result->SetBonePose(BoneIndex, BoneName, BoneTransform, BoneUsageFlags);
 					}
@@ -1124,8 +1124,8 @@ namespace UnrealConversionUtils
 				int32 FirstVertexIndex = Context->OriginalLOD->RenderSections[SectionIndex].BaseVertexIndex;
 				int32 VertexCount = Section.GetNumVertices();
 
-				mu::FMeshBufferSet& Vertices = Result->GetVertexBuffers();
-				Vertices.SetElementCount(VertexCount, mu::EMemoryInitPolicy::Zeroed);
+				UE::Mutable::Private::FMeshBufferSet& Vertices = Result->GetVertexBuffers();
+				Vertices.SetElementCount(VertexCount, UE::Mutable::Private::EMemoryInitPolicy::Zeroed);
 				Vertices.SetBufferCount(5);
 
 				// Position buffer
@@ -1146,10 +1146,10 @@ namespace UnrealConversionUtils
 				{
 					bool bIsHighPrecision = Context->DataLOD->StaticVertexBuffers.StaticMeshVertexBuffer.GetUseHighPrecisionTangentBasis();
 
-					const mu::EMeshBufferSemantic Semantics[2] = { mu::EMeshBufferSemantic::Tangent, mu::EMeshBufferSemantic::Normal };
+					const UE::Mutable::Private::EMeshBufferSemantic Semantics[2] = { UE::Mutable::Private::EMeshBufferSemantic::Tangent, UE::Mutable::Private::EMeshBufferSemantic::Normal };
 					const int32 SemanticIndices[2] = { 0, 0 };
 					const int32 Components[2] = { 4, 4 };
-					mu::EMeshBufferFormat Formats[2] = { mu::EMeshBufferFormat::PackedDirS8, mu::EMeshBufferFormat::PackedDirS8_W_TangentSign };
+					UE::Mutable::Private::EMeshBufferFormat Formats[2] = { UE::Mutable::Private::EMeshBufferFormat::PackedDirS8, UE::Mutable::Private::EMeshBufferFormat::PackedDirS8_W_TangentSign };
 					int32 Offsets[2] = { 0, 4 };
 					int32 ElementSize = 8;
 
@@ -1157,8 +1157,8 @@ namespace UnrealConversionUtils
 					{
 						// Not really supported
 						ensure(false);
-						Formats[0] = mu::EMeshBufferFormat::Int16;
-						Formats[1] = mu::EMeshBufferFormat::Int16;
+						Formats[0] = UE::Mutable::Private::EMeshBufferFormat::Int16;
+						Formats[1] = UE::Mutable::Private::EMeshBufferFormat::Int16;
 						Offsets[1] = 8;
 						ElementSize = 16;
 					}
@@ -1177,9 +1177,9 @@ namespace UnrealConversionUtils
 					int32 NumTexCoords = Context->OriginalLOD->GetNumTexCoords();
 
 					constexpr int32 MaxChannelCount = 4;
-					mu::EMeshBufferSemantic Semantics[MaxChannelCount];
+					UE::Mutable::Private::EMeshBufferSemantic Semantics[MaxChannelCount];
 					int32 SemanticIndices[MaxChannelCount];
-					mu::EMeshBufferFormat Formats[MaxChannelCount];
+					UE::Mutable::Private::EMeshBufferFormat Formats[MaxChannelCount];
 					int32 Components[MaxChannelCount];
 					int32 Offsets[MaxChannelCount];
 
@@ -1188,9 +1188,9 @@ namespace UnrealConversionUtils
 
 					for (int32 UV = 0; UV < NumTexCoords; ++UV)
 					{
-						Semantics[UV] = mu::EMeshBufferSemantic::TexCoords;
+						Semantics[UV] = UE::Mutable::Private::EMeshBufferSemantic::TexCoords;
 						SemanticIndices[UV] = UV;
-						Formats[UV] = (UVSize == 8) ? mu::EMeshBufferFormat::Float32 : mu::EMeshBufferFormat::Float16;
+						Formats[UV] = (UVSize == 8) ? UE::Mutable::Private::EMeshBufferFormat::Float32 : UE::Mutable::Private::EMeshBufferFormat::Float16;
 						Components[UV] = 2;
 						Offsets[UV] = UVSize * UV;
 					}
@@ -1248,13 +1248,13 @@ namespace UnrealConversionUtils
 					int32 IndexCount = Section.NumTriangles * 3;
 					int32 ElementSize = Context->DataLOD->MultiSizeIndexContainer.GetDataTypeSize();
 
-					mu::FMeshBufferSet& Indices = Result->GetIndexBuffers();
+					UE::Mutable::Private::FMeshBufferSet& Indices = Result->GetIndexBuffers();
 					Indices.SetBufferCount(1);
 					Indices.SetElementCount(IndexCount);
 					constexpr int32 ChannelCount = 1;
-					const mu::EMeshBufferSemantic Semantics[ChannelCount] = { mu::EMeshBufferSemantic::VertexIndex };
+					const UE::Mutable::Private::EMeshBufferSemantic Semantics[ChannelCount] = { UE::Mutable::Private::EMeshBufferSemantic::VertexIndex };
 					const int32 SemanticIndices[ChannelCount] = { 0 };
-					mu::EMeshBufferFormat Formats[ChannelCount] = { ElementSize == 4 ? mu::EMeshBufferFormat::UInt32 : mu::EMeshBufferFormat::UInt16 };
+					UE::Mutable::Private::EMeshBufferFormat Formats[ChannelCount] = { ElementSize == 4 ? UE::Mutable::Private::EMeshBufferFormat::UInt32 : UE::Mutable::Private::EMeshBufferFormat::UInt16 };
 					const int32 Components[ChannelCount] = { 1 };
 					const int32 Offsets[ChannelCount] = { 0 };
 					Indices.SetBuffer(0, ElementSize, ChannelCount, Semantics, SemanticIndices, Formats, Components, Offsets);
@@ -1417,7 +1417,7 @@ namespace UnrealConversionUtils
 				if (SourceMorphDataView.GetData())
 				{	
 					MUTABLE_CPUPROFILER_SCOPE(ExtractSectionMorphs)
-					mu::FMeshMorph& Morph = Result->Morph;
+					UE::Mutable::Private::FMeshMorph& Morph = Result->Morph;
 				
 					if (RuntimeCompressedMorphData.IsEmpty())
 					{
@@ -1684,17 +1684,17 @@ namespace UnrealConversionUtils
 	}
 
 
-	void GetSectionClothData(const mu::FMesh& MutableMesh, int32 LODIndex, const TMap<uint32, FClothingMeshData>& ClothingMeshData, TArray<FSectionClothData>& SectionsClothData, int32& NumClothingDataNotFound)
+	void GetSectionClothData(const UE::Mutable::Private::FMesh& MutableMesh, int32 LODIndex, const TMap<uint32, FClothingMeshData>& ClothingMeshData, TArray<FSectionClothData>& SectionsClothData, int32& NumClothingDataNotFound)
 	{
 		MUTABLE_CPUPROFILER_SCOPE(GetSectionClothData);
 		
-		const mu::FMeshBufferSet& MeshSet = MutableMesh.GetVertexBuffers();
+		const UE::Mutable::Private::FMeshBufferSet& MeshSet = MutableMesh.GetVertexBuffers();
 
 		int32 ClothingIndexBuffer, ClothingIndexChannel;
-		MeshSet.FindChannel(mu::EMeshBufferSemantic::Other, 2, &ClothingIndexBuffer, &ClothingIndexChannel);
+		MeshSet.FindChannel(UE::Mutable::Private::EMeshBufferSemantic::Other, 2, &ClothingIndexBuffer, &ClothingIndexChannel);
 
 		int32 ClothingResourceBuffer, ClothingResourceChannel;
-		MeshSet.FindChannel(mu::EMeshBufferSemantic::Other, 3, &ClothingResourceBuffer, &ClothingResourceChannel);
+		MeshSet.FindChannel(UE::Mutable::Private::EMeshBufferSemantic::Other, 3, &ClothingResourceBuffer, &ClothingResourceChannel);
 
 		if (ClothingIndexBuffer >= 0 && ClothingResourceBuffer >= 0)
 		{
@@ -1818,7 +1818,7 @@ namespace UnrealConversionUtils
 	}
 
 
-	void ClothVertexBuffers(FSkeletalMeshLODRenderData& LODResource, const mu::FMesh& MutableMesh, const TMap<uint32, FClothingMeshData>& ClothingMeshData, int32 LODIndex)
+	void ClothVertexBuffers(FSkeletalMeshLODRenderData& LODResource, const UE::Mutable::Private::FMesh& MutableMesh, const TMap<uint32, FClothingMeshData>& ClothingMeshData, int32 LODIndex)
 	{
 		TArray<FSectionClothData> SectionsClothData;
 		SectionsClothData.Reserve(32);
diff --git a/Engine/Plugins/Mutable/Source/CustomizableObject/Private/MuCO/UnrealMutableImageProvider.cpp b/Engine/Plugins/Mutable/Source/CustomizableObject/Private/MuCO/UnrealMutableImageProvider.cpp
--- a/Engine/Plugins/Mutable/Source/CustomizableObject/Private/MuCO/UnrealMutableImageProvider.cpp
+++ b/Engine/Plugins/Mutable/Source/CustomizableObject/Private/MuCO/UnrealMutableImageProvider.cpp
@@ -25,7 +25,7 @@
 
 namespace
 {
-	void ConvertTextureUnrealPlatformToMutable(mu::FImage* OutResult, UTexture2D* Texture, uint8 MipmapsToSkip)
+	void ConvertTextureUnrealPlatformToMutable(UE::Mutable::Private::FImage* OutResult, UTexture2D* Texture, uint8 MipmapsToSkip)
 	{		
 		check(Texture->GetPlatformData()->Mips[MipmapsToSkip].BulkData.IsBulkDataLoaded());
 
@@ -35,14 +35,14 @@ namespace
 		check(SizeX > 0 && SizeY > 0);
 
 		EPixelFormat Format = Texture->GetPlatformData()->PixelFormat;
-		mu::EImageFormat MutableFormat = mu::EImageFormat::None;
+		UE::Mutable::Private::EImageFormat MutableFormat = UE::Mutable::Private::EImageFormat::None;
 
 		switch (Format)
 		{
-		case EPixelFormat::PF_B8G8R8A8: MutableFormat = mu::EImageFormat::BGRA_UByte; break;
+		case EPixelFormat::PF_B8G8R8A8: MutableFormat = UE::Mutable::Private::EImageFormat::BGRA_UByte; break;
 			// This format is deprecated and using the enum fails to compile in some cases.
-			//case ETextureSourceFormat::TSF_RGBA8: MutableFormat = mu::EImageFormat::RGBA_UByte; break;
-		case EPixelFormat::PF_G8: MutableFormat = mu::EImageFormat::L_UByte; break;
+			//case ETextureSourceFormat::TSF_RGBA8: MutableFormat = UE::Mutable::Private::EImageFormat::RGBA_UByte; break;
+		case EPixelFormat::PF_G8: MutableFormat = UE::Mutable::Private::EImageFormat::L_UByte; break;
 		default:
 			break;
 		}
@@ -53,37 +53,37 @@ namespace
 
 		if (pSource)
 		{
-			OutResult->Init(SizeX, SizeY, LODs, MutableFormat, mu::EInitializationType::NotInitialized);
+			OutResult->Init(SizeX, SizeY, LODs, MutableFormat, UE::Mutable::Private::EInitializationType::NotInitialized);
 			FMemory::Memcpy(OutResult->GetLODData(0), pSource, OutResult->GetLODDataSize(0));
 			Texture->GetPlatformData()->Mips[MipmapsToSkip].BulkData.Unlock();
 		}
 		else
 		{
 			check(false);
-			OutResult->Init(SizeX, SizeY, LODs, MutableFormat, mu::EInitializationType::Black);
+			OutResult->Init(SizeX, SizeY, LODs, MutableFormat, UE::Mutable::Private::EInitializationType::Black);
 		}
 	}
 }
 
 
-mu::EImageFormat GetMutablePixelFormat(EPixelFormat InTextureFormat)
+UE::Mutable::Private::EImageFormat GetMutablePixelFormat(EPixelFormat InTextureFormat)
 {
 	switch (InTextureFormat)
 	{
-	case PF_B8G8R8A8: return mu::EImageFormat::BGRA_UByte;
-	case PF_R8G8B8A8: return mu::EImageFormat::RGBA_UByte;
-	case PF_DXT1: return mu::EImageFormat::BC1;
-	case PF_DXT3: return mu::EImageFormat::BC2;
-	case PF_DXT5: return mu::EImageFormat::BC3;
-	case PF_BC4: return mu::EImageFormat::BC4;
-	case PF_BC5: return mu::EImageFormat::BC5;
-	case PF_G8: return mu::EImageFormat::L_UByte;
-	case PF_ASTC_4x4: return mu::EImageFormat::ASTC_4x4_RGBA_LDR;
-	case PF_ASTC_6x6: return mu::EImageFormat::ASTC_6x6_RGBA_LDR;
-	case PF_ASTC_8x8: return mu::EImageFormat::ASTC_8x8_RGBA_LDR;
-	case PF_ASTC_10x10: return mu::EImageFormat::ASTC_10x10_RGBA_LDR;
-	case PF_ASTC_12x12: return mu::EImageFormat::ASTC_12x12_RGBA_LDR;
-	default: return mu::EImageFormat::None;
+	case PF_B8G8R8A8: return UE::Mutable::Private::EImageFormat::BGRA_UByte;
+	case PF_R8G8B8A8: return UE::Mutable::Private::EImageFormat::RGBA_UByte;
+	case PF_DXT1: return UE::Mutable::Private::EImageFormat::BC1;
+	case PF_DXT3: return UE::Mutable::Private::EImageFormat::BC2;
+	case PF_DXT5: return UE::Mutable::Private::EImageFormat::BC3;
+	case PF_BC4: return UE::Mutable::Private::EImageFormat::BC4;
+	case PF_BC5: return UE::Mutable::Private::EImageFormat::BC5;
+	case PF_G8: return UE::Mutable::Private::EImageFormat::L_UByte;
+	case PF_ASTC_4x4: return UE::Mutable::Private::EImageFormat::ASTC_4x4_RGBA_LDR;
+	case PF_ASTC_6x6: return UE::Mutable::Private::EImageFormat::ASTC_6x6_RGBA_LDR;
+	case PF_ASTC_8x8: return UE::Mutable::Private::EImageFormat::ASTC_8x8_RGBA_LDR;
+	case PF_ASTC_10x10: return UE::Mutable::Private::EImageFormat::ASTC_10x10_RGBA_LDR;
+	case PF_ASTC_12x12: return UE::Mutable::Private::EImageFormat::ASTC_12x12_RGBA_LDR;
+	default: return UE::Mutable::Private::EImageFormat::None;
 	}
 }
 
@@ -94,7 +94,7 @@ FUnrealMutableResourceProvider::FUnrealMutableResourceProvider(const TSharedRef<
 
 
 //-------------------------------------------------------------------------------------------------
-TTuple<UE::Tasks::FTask, TFunction<void()>> FUnrealMutableResourceProvider::GetImageAsync(UTexture* Texture, uint8 MipmapsToSkip, TFunction<void(TSharedPtr<mu::FImage>)>& ResultCallback)
+TTuple<UE::Tasks::FTask, TFunction<void()>> FUnrealMutableResourceProvider::GetImageAsync(UTexture* Texture, uint8 MipmapsToSkip, TFunction<void(TSharedPtr<UE::Mutable::Private::FImage>)>& ResultCallback)
 {
 	// Thread: worker
 	MUTABLE_CPUPROFILER_SCOPE(FUnrealMutableImageProvider::GetImageAsync);
@@ -106,7 +106,7 @@ TTuple<UE::Tasks::FTask, TFunction<void()>> FUnrealMutableResourceProvider::GetI
 	EPixelFormat Format = EPixelFormat::PF_Unknown;
 	int32 BulkDataSize = 0;
 
-	mu::EImageFormat MutImageFormat = mu::EImageFormat::None;
+	UE::Mutable::Private::EImageFormat MutImageFormat = UE::Mutable::Private::EImageFormat::None;
 	int32 MutImageDataSize = 0;
 
 	auto TrivialReturn = []() -> TTuple<UE::Tasks::FTask, TFunction<void()>>
@@ -136,7 +136,7 @@ TTuple<UE::Tasks::FTask, TFunction<void()>> FUnrealMutableResourceProvider::GetI
 	check(MipIndex >= 0);
 		
 	// In the editor the src data can be directly accessed
-	TSharedPtr<mu::FImage> Image = MakeShared<mu::FImage>();
+	TSharedPtr<UE::Mutable::Private::FImage> Image = MakeShared<UE::Mutable::Private::FImage>();
 
 	FMutableSourceTextureData Tex(*Texture2D);
 	EUnrealToMutableConversionError Error = ConvertTextureUnrealSourceToMutable(Image.Get(), Tex, MipIndex);
@@ -172,7 +172,7 @@ TTuple<UE::Tasks::FTask, TFunction<void()>> FUnrealMutableResourceProvider::GetI
 	MutImageFormat = GetMutablePixelFormat(Format);
 
 	// Check if it's a format we support
-	if (MutImageFormat == mu::EImageFormat::None)
+	if (MutImageFormat == UE::Mutable::Private::EImageFormat::None)
 	{
 		UE_LOG(LogMutable, Warning, TEXT("Failed to get Image Parameter [%s]. Unexpected image format. EImageFormat [%s]."), *Texture2D->GetName(), GetPixelFormatString(Format));
 		ResultCallback(CreateDummy());
@@ -183,7 +183,7 @@ TTuple<UE::Tasks::FTask, TFunction<void()>> FUnrealMutableResourceProvider::GetI
 	int32 SizeY = Texture2D->GetSizeY() >> MipIndex;
 
 	check(LODs == 1);
-	TSharedPtr<mu::FImage> Image = MakeShared<mu::FImage>(SizeX, SizeY, LODs, MutImageFormat, mu::EInitializationType::NotInitialized);
+	TSharedPtr<UE::Mutable::Private::FImage> Image = MakeShared<UE::Mutable::Private::FImage>(SizeX, SizeY, LODs, MutImageFormat, UE::Mutable::Private::EInitializationType::NotInitialized);
 	TArrayView<uint8> MutImageDataView = Image->DataStorage.GetLOD(0);
 
 	// In a packaged game the bulk data has to be loaded
@@ -326,7 +326,7 @@ TTuple<UE::Tasks::FTask, TFunction<void()>> FUnrealMutableResourceProvider::GetI
 
 
 //-------------------------------------------------------------------------------------------------
-TTuple<UE::Tasks::FTask, TFunction<void()>> FUnrealMutableResourceProvider::GetReferencedImageAsync(int32 Id, uint8 MipmapsToSkip, TFunction<void(TSharedPtr<mu::FImage>)>& ResultCallback)
+TTuple<UE::Tasks::FTask, TFunction<void()>> FUnrealMutableResourceProvider::GetReferencedImageAsync(int32 Id, uint8 MipmapsToSkip, TFunction<void(TSharedPtr<UE::Mutable::Private::FImage>)>& ResultCallback)
 {
 	// Thread: worker
 	MUTABLE_CPUPROFILER_SCOPE(FUnrealMutableImageProvider::GetReferencedImageAsync);
@@ -367,7 +367,7 @@ TTuple<UE::Tasks::FTask, TFunction<void()>> FUnrealMutableResourceProvider::GetR
 	const int32 MipIndex = FMath::Min(static_cast<int32>(MipmapsToSkip), Texture2D->Source.CopyTornOff().GetNumMips() - 1);
 	check(MipIndex >= 0);
 
-	TSharedPtr<mu::FImage> Image = MakeShared<mu::FImage>();
+	TSharedPtr<UE::Mutable::Private::FImage> Image = MakeShared<UE::Mutable::Private::FImage>();
 	
 	FMutableSourceTextureData Tex(*Texture2D);
 	EUnrealToMutableConversionError Error = ConvertTextureUnrealSourceToMutable(Image.Get(), Tex, MipIndex);
@@ -396,7 +396,7 @@ TTuple<UE::Tasks::FTask, TFunction<void()>> FUnrealMutableResourceProvider::GetR
 
 
 // This should mantain parity with the descriptor of the images generated by GetImageAsync 
-mu::FExtendedImageDesc FUnrealMutableResourceProvider::GetImageDesc(UTexture* Texture)
+UE::Mutable::Private::FExtendedImageDesc FUnrealMutableResourceProvider::GetImageDesc(UTexture* Texture)
 {
 	MUTABLE_CPUPROFILER_SCOPE(FUnrealMutableImageProvider::GetImageDesc);
 
@@ -414,12 +414,12 @@ mu::FExtendedImageDesc FUnrealMutableResourceProvider::GetImageDesc(UTexture* Te
 #if WITH_EDITOR
 	const FTextureSource& Source = Texture->Source.CopyTornOff();
 
-	const mu::FImageSize ImageSize = mu::FImageSize(Source.GetSizeX(), Source.GetSizeY());
+	const UE::Mutable::Private::FImageSize ImageSize = UE::Mutable::Private::FImageSize(Source.GetSizeX(), Source.GetSizeY());
 	const uint8 LODs = 1;
 
-	return mu::FExtendedImageDesc { mu::FImageDesc { ImageSize, mu::EImageFormat::None, LODs }, 0 };
+	return UE::Mutable::Private::FExtendedImageDesc { UE::Mutable::Private::FImageDesc { ImageSize, UE::Mutable::Private::EImageFormat::None, LODs }, 0 };
 #else
-	mu::FExtendedImageDesc Result;	
+	UE::Mutable::Private::FExtendedImageDesc Result;	
 
 	// It's safe to access TextureToLoad because ExternalImagesLock guarantees that the data in GlobalExternalImages is valid,
 	// not being modified by the game thread at the moment and the texture cannot be GCed because of the AddReferencedObjects
@@ -438,26 +438,26 @@ mu::FExtendedImageDesc FUnrealMutableResourceProvider::GetImageDesc(UTexture* Te
 	
 	// Texture format and the equivalent mutable format
 	const EPixelFormat Format = Texture2D->GetPlatformData()->PixelFormat;
-	const mu::EImageFormat MutableFormat = GetMutablePixelFormat(Format);
+	const UE::Mutable::Private::EImageFormat MutableFormat = GetMutablePixelFormat(Format);
 
 	// Check if it's a format we support
-	if (MutableFormat == mu::EImageFormat::None)
+	if (MutableFormat == UE::Mutable::Private::EImageFormat::None)
 	{
 		UE_LOG(LogMutable, Warning, TEXT("Failed to get Image Parameter descriptor. Unexpected image format. EImageFormat [%s]."), GetPixelFormatString(Format));
 		return CreateDummyDesc();
 	}
 
-	const mu::FImageDesc ImageDesc = mu::FImageDesc 
-		{ mu::FImageSize(Texture2D->GetSizeX(), Texture2D->GetSizeY()), MutableFormat, 1 }; 
+	const UE::Mutable::Private::FImageDesc ImageDesc = UE::Mutable::Private::FImageDesc 
+		{ UE::Mutable::Private::FImageSize(Texture2D->GetSizeX(), Texture2D->GetSizeY()), MutableFormat, 1 }; 
 
-	Result = mu::FExtendedImageDesc { ImageDesc, (uint8)FirstLODAvailable }; 
+	Result = UE::Mutable::Private::FExtendedImageDesc { ImageDesc, (uint8)FirstLODAvailable }; 
 
 	return Result;
 #endif
 }
 
 
-TTuple<UE::Tasks::FTask, TFunction<void()>> FUnrealMutableResourceProvider::GetMeshAsync(USkeletalMesh* SkeletalMesh, int32 LODIndex, int32 SectionIndex, TFunction<void(TSharedPtr<mu::FMesh>)>& ResultCallback)
+TTuple<UE::Tasks::FTask, TFunction<void()>> FUnrealMutableResourceProvider::GetMeshAsync(USkeletalMesh* SkeletalMesh, int32 LODIndex, int32 SectionIndex, TFunction<void(TSharedPtr<UE::Mutable::Private::FMesh>)>& ResultCallback)
 {
 	// Thread: worker
 	MUTABLE_CPUPROFILER_SCOPE(FUnrealMutableImageProvider::GetMeshAsync);
@@ -467,7 +467,7 @@ TTuple<UE::Tasks::FTask, TFunction<void()>> FUnrealMutableResourceProvider::GetM
 		return MakeTuple(UE::Tasks::MakeCompletedTask<void>(), []() -> void {});
 	};
 
-	TSharedPtr<mu::FMesh> Result = MakeShared<mu::FMesh>();
+	TSharedPtr<UE::Mutable::Private::FMesh> Result = MakeShared<UE::Mutable::Private::FMesh>();
 
 	if (!SkeletalMesh)
 	{
@@ -512,7 +512,7 @@ void FUnrealMutableResourceProvider::CacheRuntimeReferencedImages(const TArray<T
 		{
 			UE_LOG(LogMutable, Warning, TEXT("Runtime Referenced Texture [%s] was not async loaded. Forcing load sync."), *RuntimeReferencedTexture->GetPathName());
 			
-			Texture = MutablePrivate::LoadObject(RuntimeReferencedTexture);
+			Texture = UE::Mutable::Private::LoadObject(RuntimeReferencedTexture);
 			if (!Texture)
 			{
 				UE_LOG(LogMutable, Warning, TEXT("Failed to force load sync [%s]."), *RuntimeReferencedTexture->GetPathName());
@@ -526,7 +526,7 @@ void FUnrealMutableResourceProvider::CacheRuntimeReferencedImages(const TArray<T
 #endif
 
 
-TSharedPtr<mu::FImage> FUnrealMutableResourceProvider::CreateDummy()
+TSharedPtr<UE::Mutable::Private::FImage> FUnrealMutableResourceProvider::CreateDummy()
 {
 	// Create a dummy image
 	const int32 Size = DUMMY_IMAGE_DESC.m_size[0];
@@ -539,10 +539,10 @@ TSharedPtr<mu::FImage> FUnrealMutableResourceProvider::CreateDummy()
 	uint8 Colors[CheckerTileCount][4] = {{255, 255, 0, 0}, {0, 0, 255, 0}};
 #endif
 
-	TSharedPtr<mu::FImage> pResult = MakeShared<mu::FImage>(Size, Size, DUMMY_IMAGE_DESC.m_lods, DUMMY_IMAGE_DESC.m_format, mu::EInitializationType::NotInitialized);
+	TSharedPtr<UE::Mutable::Private::FImage> pResult = MakeShared<UE::Mutable::Private::FImage>(Size, Size, DUMMY_IMAGE_DESC.m_lods, DUMMY_IMAGE_DESC.m_format, UE::Mutable::Private::EInitializationType::NotInitialized);
 
 	check(pResult->GetLODCount() == 1);
-	check(pResult->GetFormat() == mu::EImageFormat::RGBA_UByte || pResult->GetFormat() == mu::EImageFormat::BGRA_UByte);
+	check(pResult->GetFormat() == UE::Mutable::Private::EImageFormat::RGBA_UByte || pResult->GetFormat() == UE::Mutable::Private::EImageFormat::BGRA_UByte);
 	uint8* pData = pResult->GetLODData(0);
 	for (int32 X = 0; X < Size; ++X)
 	{
@@ -561,8 +561,8 @@ TSharedPtr<mu::FImage> FUnrealMutableResourceProvider::CreateDummy()
 }
 
 
-mu::FExtendedImageDesc FUnrealMutableResourceProvider::CreateDummyDesc()
+UE::Mutable::Private::FExtendedImageDesc FUnrealMutableResourceProvider::CreateDummyDesc()
 {
-	return mu::FExtendedImageDesc{ {DUMMY_IMAGE_DESC}, 0 };
+	return UE::Mutable::Private::FExtendedImageDesc{ {DUMMY_IMAGE_DESC}, 0 };
 }
 
diff --git a/Engine/Plugins/Mutable/Source/CustomizableObject/Private/MuCO/UnrealMutableModelDiskStreamer.cpp b/Engine/Plugins/Mutable/Source/CustomizableObject/Private/MuCO/UnrealMutableModelDiskStreamer.cpp
--- a/Engine/Plugins/Mutable/Source/CustomizableObject/Private/MuCO/UnrealMutableModelDiskStreamer.cpp
+++ b/Engine/Plugins/Mutable/Source/CustomizableObject/Private/MuCO/UnrealMutableModelDiskStreamer.cpp
@@ -91,7 +91,7 @@ bool FUnrealMutableModelBulkReader::PrepareStreamingForObject(UCustomizableObjec
 	if (!bAlreadyStreaming)
 	{
 		FObjectData NewData;
-		NewData.Model = TWeakPtr<const mu::FModel>(CustomizableObject->GetPrivate()->GetModel());
+		NewData.Model = TWeakPtr<const UE::Mutable::Private::FModel>(CustomizableObject->GetPrivate()->GetModel());
 
 		const TSharedPtr<FModelStreamableBulkData> ModelStreamableBulkData = CustomizableObject->GetPrivate()->GetModelStreamableBulkData();
 		NewData.ModelStreamableBulkData = TWeakPtr<FModelStreamableBulkData>(ModelStreamableBulkData);
@@ -108,7 +108,7 @@ bool FUnrealMutableModelBulkReader::PrepareStreamingForObject(UCustomizableObjec
 
 #if WITH_EDITOR
 		FString FullFileName = CustomizableObject->GetPrivate()->GetCompiledDataFileName(nullptr, true);
-		NewData.BulkFilePrefix = *FullFileName + GetDataTypeExtension(MutablePrivate::EStreamableDataType::Model);
+		NewData.BulkFilePrefix = *FullFileName + GetDataTypeExtension(UE::Mutable::Private::EStreamableDataType::Model);
 #else
 		if (ModelStreamableBulkData->StreamableBulkData.IsEmpty())
 		{
@@ -225,7 +225,7 @@ FORCEINLINE FString MakeBlockFilePath(const FString& Prefix, const FMutableStrea
 }
 };
 
-bool FUnrealMutableModelBulkReader::DoesBlockExist(const mu::FModel* Model, uint32 Key)
+bool FUnrealMutableModelBulkReader::DoesBlockExist(const UE::Mutable::Private::FModel* Model, uint32 Key)
 {
 	MUTABLE_CPUPROFILER_SCOPE(FUnrealMutableModelBulkReader::DoesBlockExist)
 
@@ -278,7 +278,7 @@ bool FUnrealMutableModelBulkReader::DoesBlockExist(const mu::FModel* Model, uint
 }
 
 
-mu::FModelReader::FOperationID FUnrealMutableModelBulkReader::BeginReadBlock(const mu::FModel* Model, uint32 Key, void* Buffer, uint64 Size, mu::EDataType ResourceType, TFunction<void(bool bSuccess)>* CompletionCallback)
+UE::Mutable::Private::FModelReader::FOperationID FUnrealMutableModelBulkReader::BeginReadBlock(const UE::Mutable::Private::FModel* Model, uint32 Key, void* Buffer, uint64 Size, UE::Mutable::Private::EDataType ResourceType, TFunction<void(bool bSuccess)>* CompletionCallback)
 {
 	MUTABLE_CPUPROFILER_SCOPE(FUnrealMutableModelBulkStreamer::OpenReadFile);
 
@@ -458,7 +458,7 @@ mu::FModelReader::FOperationID FUnrealMutableModelBulkReader::BeginReadBlock(con
 }
 
 
-bool FUnrealMutableModelBulkReader::IsReadCompleted(mu::FModelReader::FOperationID OperationId)
+bool FUnrealMutableModelBulkReader::IsReadCompleted(UE::Mutable::Private::FModelReader::FOperationID OperationId)
 {
 	MUTABLE_CPUPROFILER_SCOPE(FUnrealMutableModelBulkStreamer::IsReadCompleted);
 
@@ -489,7 +489,7 @@ bool FUnrealMutableModelBulkReader::IsReadCompleted(mu::FModelReader::FOperation
 }
 
 
-bool FUnrealMutableModelBulkReader::EndRead(mu::FModelReader::FOperationID OperationId)
+bool FUnrealMutableModelBulkReader::EndRead(UE::Mutable::Private::FModelReader::FOperationID OperationId)
 {
 	MUTABLE_CPUPROFILER_SCOPE(FUnrealMutableModelBulkStreamer::EndRead);
 
@@ -587,7 +587,7 @@ void FUnrealMutableModelBulkWriterEditor::CloseWriteFile()
 
 
 //-------------------------------------------------------------------------------------------------
-FUnrealMutableModelBulkWriterCook::FUnrealMutableModelBulkWriterCook(FArchive* InMainDataArchive, MutablePrivate::FModelStreamableData* InStreamedData)
+FUnrealMutableModelBulkWriterCook::FUnrealMutableModelBulkWriterCook(FArchive* InMainDataArchive, UE::Mutable::Private::FModelStreamableData* InStreamedData)
 {
 	MainDataArchive = InMainDataArchive;
 	StreamedData = InStreamedData;
diff --git a/Engine/Plugins/Mutable/Source/CustomizableObject/Private/MuCO/UnrealToMutableTextureConversionUtils.cpp b/Engine/Plugins/Mutable/Source/CustomizableObject/Private/MuCO/UnrealToMutableTextureConversionUtils.cpp
--- a/Engine/Plugins/Mutable/Source/CustomizableObject/Private/MuCO/UnrealToMutableTextureConversionUtils.cpp
+++ b/Engine/Plugins/Mutable/Source/CustomizableObject/Private/MuCO/UnrealToMutableTextureConversionUtils.cpp
@@ -215,7 +215,7 @@ bool FMutableSourceTextureData::IsNormalComposite() const
 }
 
 
-EUnrealToMutableConversionError ConvertTextureUnrealSourceToMutable(mu::FImage* OutResult, FMutableSourceTextureData& Tex, uint8 MipmapsToSkip)
+EUnrealToMutableConversionError ConvertTextureUnrealSourceToMutable(UE::Mutable::Private::FImage* OutResult, FMutableSourceTextureData& Tex, uint8 MipmapsToSkip)
 {
 	MUTABLE_CPUPROFILER_SCOPE(ConvertTextureUnrealSourceToMutable);
 
@@ -293,7 +293,7 @@ EUnrealToMutableConversionError ConvertTextureUnrealSourceToMutable(mu::FImage*
 		MUTABLE_CPUPROFILER_SCOPE(NoConvert);
         check(LODs == 1);
 
-		OutResult->Init(SizeX, SizeY, LODs, mu::EImageFormat::L_UByte, mu::EInitializationType::NotInitialized);
+		OutResult->Init(SizeX, SizeY, LODs, UE::Mutable::Private::EImageFormat::L_UByte, UE::Mutable::Private::EInitializationType::NotInitialized);
 		OutResult->DataStorage.GetInternalArray(0) = MoveTemp(TempImage.RawData);
 		break;
 	}
@@ -310,7 +310,7 @@ EUnrealToMutableConversionError ConvertTextureUnrealSourceToMutable(mu::FImage*
 		if (bHasAlphaChannel)
 		{
 			MUTABLE_CPUPROFILER_SCOPE(ToRGBA);
-			OutResult->Init(SizeX, SizeY, LODs, mu::EImageFormat::RGBA_UByte, mu::EInitializationType::NotInitialized);
+			OutResult->Init(SizeX, SizeY, LODs, UE::Mutable::Private::EImageFormat::RGBA_UByte, UE::Mutable::Private::EInitializationType::NotInitialized);
             check(LODs == 1);
             uint8* DataDest = OutResult->GetLODData(0);
 
@@ -329,8 +329,8 @@ EUnrealToMutableConversionError ConvertTextureUnrealSourceToMutable(mu::FImage*
 		{
 			MUTABLE_CPUPROFILER_SCOPE(ToRGB);
 
-			// TODO: add support for a mu::IF_RGBX_UBYTE?			
-			OutResult->Init(SizeX, SizeY, LODs, mu::EImageFormat::RGB_UByte, mu::EInitializationType::NotInitialized);
+			// TODO: add support for a UE::Mutable::Private::IF_RGBX_UBYTE?			
+			OutResult->Init(SizeX, SizeY, LODs, UE::Mutable::Private::EImageFormat::RGB_UByte, UE::Mutable::Private::EInitializationType::NotInitialized);
 			check(LODs == 1);
             uint8* DataDest = OutResult->GetLODData(0);
 
diff --git a/Engine/Plugins/Mutable/Source/CustomizableObject/Public/MuCO/CustomizableObjectInstance.h b/Engine/Plugins/Mutable/Source/CustomizableObject/Public/MuCO/CustomizableObjectInstance.h
--- a/Engine/Plugins/Mutable/Source/CustomizableObject/Public/MuCO/CustomizableObjectInstance.h
+++ b/Engine/Plugins/Mutable/Source/CustomizableObject/Public/MuCO/CustomizableObjectInstance.h
@@ -9,7 +9,7 @@
 
 #include "CustomizableObjectInstance.generated.h"
 
-namespace mu
+namespace UE::Mutable::Private
 {
 	class FImage;
 }
diff --git a/Engine/Plugins/Mutable/Source/CustomizableObject/Public/MuCO/CustomizableObjectInstanceDescriptor.h b/Engine/Plugins/Mutable/Source/CustomizableObject/Public/MuCO/CustomizableObjectInstanceDescriptor.h
--- a/Engine/Plugins/Mutable/Source/CustomizableObject/Public/MuCO/CustomizableObjectInstanceDescriptor.h
+++ b/Engine/Plugins/Mutable/Source/CustomizableObject/Public/MuCO/CustomizableObjectInstanceDescriptor.h
@@ -24,7 +24,7 @@ struct FCustomizableObjectComponentIndex;
 
 typedef TMap<const UCustomizableObjectInstance*, FMutableUpdateCandidate> FMutableInstanceUpdateMap;
 
-namespace mu
+namespace UE::Mutable::Private
 {
 	class FParameters;
 	
@@ -296,7 +296,7 @@ struct FCustomizableObjectInstanceDescriptor
 	UE_API void UpdateLayer(const FName& ParamName, int32 Index, const FMultilayerProjectorLayer& Layer);
 
 	/** Return a Mutable Core object containing all parameters. */
-	UE_API TSharedPtr<mu::FParameters> GetParameters() const;
+	UE_API TSharedPtr<UE::Mutable::Private::FParameters> GetParameters() const;
 
 	UE_API FString ToString() const;
 	
diff --git a/Engine/Plugins/Mutable/Source/CustomizableObject/Public/MuCO/CustomizableObjectSkeletalMesh.h b/Engine/Plugins/Mutable/Source/CustomizableObject/Public/MuCO/CustomizableObjectSkeletalMesh.h
--- a/Engine/Plugins/Mutable/Source/CustomizableObject/Public/MuCO/CustomizableObjectSkeletalMesh.h
+++ b/Engine/Plugins/Mutable/Source/CustomizableObject/Public/MuCO/CustomizableObjectSkeletalMesh.h
@@ -46,19 +46,19 @@ class UCustomizableObjectSkeletalMesh : public USkeletalMesh
 
 	TSharedPtr<FModelStreamableBulkData> ModelStreamableBulkData;
 
-	TSharedPtr<mu::FModel> Model;
-	TSharedPtr<mu::FMeshIdRegistry> MeshIdRegistry;
-	TSharedPtr<mu::FImageIdRegistry> ImageIdRegistry;
-	TSharedPtr<mu::FMaterialIdRegistry> MaterialIdRegistry;
+	TSharedPtr<UE::Mutable::Private::FModel> Model;
+	TSharedPtr<UE::Mutable::Private::FMeshIdRegistry> MeshIdRegistry;
+	TSharedPtr<UE::Mutable::Private::FImageIdRegistry> ImageIdRegistry;
+	TSharedPtr<UE::Mutable::Private::FMaterialIdRegistry> MaterialIdRegistry;
 
 	TSharedPtr<FUnrealMutableResourceProvider> ExternalResourceProvider;
 
 	/** Instance parameters at the time of creation. */
-	TSharedPtr<mu::FParameters> Parameters;
+	TSharedPtr<UE::Mutable::Private::FParameters> Parameters;
 	int32 State = -1;
 
 	/** Mesh IDs per LOD */
-	TArray<mu::FMeshId> MeshIDs;
+	TArray<UE::Mutable::Private::FMeshId> MeshIDs;
 	TArray<TArray<uint32>> SurfaceIDs;
 	TArray<TPair<uint32, FName>> SkinWeightProfileIDs;
 
diff --git a/Engine/Plugins/Mutable/Source/CustomizableObjectEditor/Private/MuCOE/AssetDefinition_CustomizableObject.cpp b/Engine/Plugins/Mutable/Source/CustomizableObjectEditor/Private/MuCOE/AssetDefinition_CustomizableObject.cpp
--- a/Engine/Plugins/Mutable/Source/CustomizableObjectEditor/Private/MuCOE/AssetDefinition_CustomizableObject.cpp
+++ b/Engine/Plugins/Mutable/Source/CustomizableObjectEditor/Private/MuCOE/AssetDefinition_CustomizableObject.cpp
@@ -59,7 +59,7 @@ EAssetCommandResult UAssetDefinition_CustomizableObject::ActivateAssets(const FA
 		if ( ActivateArgs.Assets.Num() == 1 )
 		{
 			const FAssetData* FirstAsset = AssetsToActivate.GetData();
-			GEditor->GetEditorSubsystem<UAssetEditorSubsystem>()->OpenEditorForAsset(MutablePrivate::LoadObject(*FirstAsset));
+			GEditor->GetEditorSubsystem<UAssetEditorSubsystem>()->OpenEditorForAsset(UE::Mutable::Private::LoadObject(*FirstAsset));
 		}
 		else if ( ActivateArgs.Assets.Num() > 1 )
 		{
@@ -68,7 +68,7 @@ EAssetCommandResult UAssetDefinition_CustomizableObject::ActivateAssets(const FA
 			ObjectsToOpen.Reserve(AssetsToActivate.Num());
 			for (const FAssetData& ToOpen : AssetsToActivate)
 			{
-				ObjectsToOpen.Add(MutablePrivate::LoadObject(ToOpen));
+				ObjectsToOpen.Add(UE::Mutable::Private::LoadObject(ToOpen));
 			}
 			ObjectsToOpen.Shrink();
 			
@@ -231,7 +231,7 @@ namespace MenuExtension_CustomizableObject
 						continue;
 					}
 
-					UCustomizableObject* Object = Cast<UCustomizableObject>(MutablePrivate::LoadObject(Asset));
+					UCustomizableObject* Object = Cast<UCustomizableObject>(UE::Mutable::Private::LoadObject(Asset));
 					if (Object)
 					{
 						TSharedRef<FCompilationRequest> NewRequest = MakeShared<FCompilationRequest>(*Object);
diff --git a/Engine/Plugins/Mutable/Source/CustomizableObjectEditor/Private/MuCOE/AssetDefinition_CustomizableObjectInstance.cpp b/Engine/Plugins/Mutable/Source/CustomizableObjectEditor/Private/MuCOE/AssetDefinition_CustomizableObjectInstance.cpp
--- a/Engine/Plugins/Mutable/Source/CustomizableObjectEditor/Private/MuCOE/AssetDefinition_CustomizableObjectInstance.cpp
+++ b/Engine/Plugins/Mutable/Source/CustomizableObjectEditor/Private/MuCOE/AssetDefinition_CustomizableObjectInstance.cpp
@@ -47,7 +47,7 @@ EAssetCommandResult UAssetDefinition_CustomizableObjectInstance::ActivateAssets(
 		if (ActivateArgs.Assets.Num() == 1)
 		{
 			const FAssetData* FirstAsset = ActivateArgs.Assets.GetData();
-			GEditor->GetEditorSubsystem<UAssetEditorSubsystem>()->OpenEditorForAsset(MutablePrivate::LoadObject(*FirstAsset));
+			GEditor->GetEditorSubsystem<UAssetEditorSubsystem>()->OpenEditorForAsset(UE::Mutable::Private::LoadObject(*FirstAsset));
 		}
 		else if (ActivateArgs.Assets.Num() > 1)
 		{
@@ -55,7 +55,7 @@ EAssetCommandResult UAssetDefinition_CustomizableObjectInstance::ActivateAssets(
 			ObjectsToOpen.Reserve(ActivateArgs.Assets.Num());
 			for (const FAssetData& ToOpen : ActivateArgs.Assets)
 			{
-				ObjectsToOpen.Add(MutablePrivate::LoadObject(ToOpen));
+				ObjectsToOpen.Add(UE::Mutable::Private::LoadObject(ToOpen));
 			}
 			ObjectsToOpen.Shrink();
 			
diff --git a/Engine/Plugins/Mutable/Source/CustomizableObjectEditor/Private/MuCOE/CustomizableInstanceDetails.cpp b/Engine/Plugins/Mutable/Source/CustomizableObjectEditor/Private/MuCOE/CustomizableInstanceDetails.cpp
--- a/Engine/Plugins/Mutable/Source/CustomizableObjectEditor/Private/MuCOE/CustomizableInstanceDetails.cpp
+++ b/Engine/Plugins/Mutable/Source/CustomizableObjectEditor/Private/MuCOE/CustomizableInstanceDetails.cpp
@@ -1209,7 +1209,7 @@ TSharedRef<SWidget> FCustomizableInstanceDetails::OnGenerateWidgetIntParameter(T
 
 			if (AssetData.IsValid())
 			{
-				if (UTexture2D* UIThumbnail = Cast<UTexture2D>(MutablePrivate::LoadObject(AssetData)))
+				if (UTexture2D* UIThumbnail = Cast<UTexture2D>(UE::Mutable::Private::LoadObject(AssetData)))
 				{
 					// We need to store the generated texture.
 					TSharedRef<FDeferredCleanupSlateBrush> Brush = FDeferredCleanupSlateBrush::CreateBrush(UIThumbnail, FVector2D(68.f, 68.0f), /* Texture size (64) + padding (4) */
@@ -1284,7 +1284,7 @@ TSharedRef<SWidget> FCustomizableInstanceDetails::GenerateFloatWidget(const UCus
 	if (const TSoftObjectPtr<UObject>* FloatDecoratorAsset = UIMetadata.ExtraAssets.Find(UIMetadataKeyWords::FloatDecoratorName))
 	{
 		// Checking if there is an image decorator for the float slider:
-		if (UTexture2D* DecoratorTexture = Cast<UTexture2D>(MutablePrivate::LoadObject(*FloatDecoratorAsset)))
+		if (UTexture2D* DecoratorTexture = Cast<UTexture2D>(UE::Mutable::Private::LoadObject(*FloatDecoratorAsset)))
 		{
 			TSharedRef<FDeferredCleanupSlateBrush> Brush = FDeferredCleanupSlateBrush::CreateBrush(DecoratorTexture, FVector2D(DecoratorTexture->GetSizeX(), 1.0f), FLinearColor(1.0f, 1.0f, 1.0f, 1.0f), ESlateBrushTileType::NoTile, ESlateBrushImageType::Linear);
 			DynamicBrushes.Add(Brush);
diff --git a/Engine/Plugins/Mutable/Source/CustomizableObjectEditor/Private/MuCOE/CustomizableObjectCompileRunnable.cpp b/Engine/Plugins/Mutable/Source/CustomizableObjectEditor/Private/MuCOE/CustomizableObjectCompileRunnable.cpp
--- a/Engine/Plugins/Mutable/Source/CustomizableObjectEditor/Private/MuCOE/CustomizableObjectCompileRunnable.cpp
+++ b/Engine/Plugins/Mutable/Source/CustomizableObjectEditor/Private/MuCOE/CustomizableObjectCompileRunnable.cpp
@@ -51,7 +51,7 @@ TAutoConsoleVariable<bool> CVarMutableCompilerFastCompression(
 	ECVF_Default);
 
 
-FCustomizableObjectCompileRunnable::FCustomizableObjectCompileRunnable(mu::Ptr<mu::Node> Root, const TSharedRef<FCustomizableObjectCompiler>& InCompiler)
+FCustomizableObjectCompileRunnable::FCustomizableObjectCompileRunnable(UE::Mutable::Private::Ptr<UE::Mutable::Private::Node> Root, const TSharedRef<FCustomizableObjectCompiler>& InCompiler)
 	: MutableRoot(Root)
 	, WeakCompiler(InCompiler)
 	, bThreadCompleted(false)
@@ -60,11 +60,11 @@ FCustomizableObjectCompileRunnable::FCustomizableObjectCompileRunnable(mu::Ptr<m
 }
 
 
-TSharedPtr<mu::FImage> FCustomizableObjectCompileRunnable::LoadImageResourceReferenced(int32 ID)
+TSharedPtr<UE::Mutable::Private::FImage> FCustomizableObjectCompileRunnable::LoadImageResourceReferenced(int32 ID)
 {
 	MUTABLE_CPUPROFILER_SCOPE(LoadResourceReferenced);
 
-	TSharedPtr<mu::FImage> Image;
+	TSharedPtr<UE::Mutable::Private::FImage> Image;
 	if (!ReferencedTextures.IsValidIndex(ID))
 	{
 		// The id is not valid for this CO
@@ -76,7 +76,7 @@ TSharedPtr<mu::FImage> FCustomizableObjectCompileRunnable::LoadImageResourceRefe
 	FMutableSourceTextureData& TextureData = ReferencedTextures[ID];
 
 	// In the editor the src data can be directly accessed
-	Image = MakeShared<mu::FImage>();
+	Image = MakeShared<UE::Mutable::Private::FImage>();
 	int32 MipmapsToSkip = 0;
 	EUnrealToMutableConversionError Error = ConvertTextureUnrealSourceToMutable(Image.Get(), TextureData, MipmapsToSkip);
 
@@ -100,8 +100,8 @@ uint32 FCustomizableObjectCompileRunnable::Run()
 	uint32 Result = 1;
 	ErrorMsg = FString();
 
-	// Translate CO compile options into mu::CompilerOptions
-	mu::Ptr<mu::CompilerOptions> CompilerOptions = new mu::CompilerOptions();
+	// Translate CO compile options into UE::Mutable::Private::CompilerOptions
+	UE::Mutable::Private::Ptr<UE::Mutable::Private::CompilerOptions> CompilerOptions = new UE::Mutable::Private::CompilerOptions();
 
 	bool bUseDiskCache = Options.bUseDiskCompilation;
 	if (CVarMutableCompilerDiskCache->GetBool())
@@ -157,12 +157,12 @@ uint32 FCustomizableObjectCompileRunnable::Run()
 	}
 
 	CompilerOptions->SetReferencedResourceCallback(
-		[this](int32 ID, TSharedPtr<TSharedPtr<mu::FImage>> ResolvedImage, bool bRunImmediatlyIfPossible)
+		[this](int32 ID, TSharedPtr<TSharedPtr<UE::Mutable::Private::FImage>> ResolvedImage, bool bRunImmediatlyIfPossible)
 		{
 			UE::Tasks::FTask LaunchTask = UE::Tasks::Launch(TEXT("LoadImageReferenceTasks"),
 				[ID,ResolvedImage,this]()
 				{
-					TSharedPtr<mu::FImage> Result = LoadImageResourceReferenced(ID);
+					TSharedPtr<UE::Mutable::Private::FImage> Result = LoadImageResourceReferenced(ID);
 					*ResolvedImage = Result;
 				},
 				LowLevelTasks::ETaskPriority::BackgroundLow
@@ -171,7 +171,7 @@ uint32 FCustomizableObjectCompileRunnable::Run()
 			return LaunchTask;
 		},
 
-		[this](int32 ID, const FString& MorphNameRef, TSharedPtr<TSharedPtr<mu::FMesh>> ResolvedMesh, bool bRunImmediatlyIfPossible) -> UE::Tasks::FTask
+		[this](int32 ID, const FString& MorphNameRef, TSharedPtr<TSharedPtr<UE::Mutable::Private::FMesh>> ResolvedMesh, bool bRunImmediatlyIfPossible) -> UE::Tasks::FTask
 		{
 			MUTABLE_CPUPROFILER_SCOPE(LoadMeshReferenceTasks);
 
@@ -199,7 +199,7 @@ uint32 FCustomizableObjectCompileRunnable::Run()
 					// If we are shutting down, we are not allowed to try to load anything.
 					if (IsEngineExitRequested())
 					{
-						*ResolvedMesh = MakeShared<mu::FMesh>();
+						*ResolvedMesh = MakeShared<UE::Mutable::Private::FMesh>();
 						CompletionEvent.Trigger();
 						return;
 					}
@@ -207,7 +207,7 @@ uint32 FCustomizableObjectCompileRunnable::Run()
 					TSharedPtr<FCustomizableObjectCompiler> Compiler = WeakCompiler.Pin();
 					if (!Compiler)
 					{
-						*ResolvedMesh = MakeShared<mu::FMesh>();
+						*ResolvedMesh = MakeShared<UE::Mutable::Private::FMesh>();
 						CompletionEvent.Trigger();
 						return;
 					}
@@ -221,9 +221,9 @@ uint32 FCustomizableObjectCompileRunnable::Run()
 					MeshConversionContext->CompilationContext = Compiler->CompilationContext;
 
 					// Todo: Load assets async
-					MeshConversionContext->ReferencedObjects.Emplace(TStrongObjectPtr<UObject>(MutablePrivate::LoadObject(MeshData.Mesh)));
-					MeshConversionContext->ReferencedObjects.Emplace(TStrongObjectPtr<UObject>(MutablePrivate::LoadObject(MeshData.TableReferenceSkeletalMesh)));
-					MeshConversionContext->ReferencedObjects.Emplace(TStrongObjectPtr<UObject>(MutablePrivate::LoadClass(MeshData.AnimInstance)));
+					MeshConversionContext->ReferencedObjects.Emplace(TStrongObjectPtr<UObject>(UE::Mutable::Private::LoadObject(MeshData.Mesh)));
+					MeshConversionContext->ReferencedObjects.Emplace(TStrongObjectPtr<UObject>(UE::Mutable::Private::LoadObject(MeshData.TableReferenceSkeletalMesh)));
+					MeshConversionContext->ReferencedObjects.Emplace(TStrongObjectPtr<UObject>(UE::Mutable::Private::LoadClass(MeshData.AnimInstance)));
 
 					*ResolvedMesh = ConvertSkeletalMeshToMutable(MeshConversionContext);
 
@@ -243,7 +243,7 @@ uint32 FCustomizableObjectCompileRunnable::Run()
 				}
 				else
 				{
-					*ResolvedMesh = MakeShared<mu::FMesh>();
+					*ResolvedMesh = MakeShared<UE::Mutable::Private::FMesh>();
 					CompletionEvent.Trigger();
 				}
 			}
@@ -279,27 +279,27 @@ uint32 FCustomizableObjectCompileRunnable::Run()
 			}
 		};
 
-	mu::Ptr<mu::Compiler> Compiler = new mu::Compiler(CompilerOptions, WaitCallback);
+	UE::Mutable::Private::Ptr<UE::Mutable::Private::Compiler> Compiler = new UE::Mutable::Private::Compiler(CompilerOptions, WaitCallback);
 
 	UE_LOG(LogMutable, Verbose, TEXT("PROFILE: [ %16.8f ] FCustomizableObjectCompileRunnable Compile start."), FPlatformTime::Seconds());
 	Model = Compiler->Compile(MutableRoot);
 
 	// Dump all the log messages from the compiler
-	TSharedPtr<const mu::FErrorLog> pLog = Compiler->GetLog();
+	TSharedPtr<const UE::Mutable::Private::FErrorLog> pLog = Compiler->GetLog();
 	for (int32 i = 0; i < pLog->GetMessageCount(); ++i)
 	{
 		const FString& Message = pLog->GetMessageText(i);
-		const mu::ErrorLogMessageType MessageType = pLog->GetMessageType(i);
-		const mu::ErrorLogMessageAttachedDataView MessageAttachedData = pLog->GetMessageAttachedData(i);
+		const UE::Mutable::Private::ErrorLogMessageType MessageType = pLog->GetMessageType(i);
+		const UE::Mutable::Private::ErrorLogMessageAttachedDataView MessageAttachedData = pLog->GetMessageAttachedData(i);
 
-		if (MessageType == mu::ELMT_WARNING || MessageType == mu::ELMT_ERROR)
+		if (MessageType == UE::Mutable::Private::ELMT_WARNING || MessageType == UE::Mutable::Private::ELMT_ERROR)
 		{
-			const EMessageSeverity::Type Severity = MessageType == mu::ELMT_WARNING ? EMessageSeverity::Warning : EMessageSeverity::Error;
+			const EMessageSeverity::Type Severity = MessageType == UE::Mutable::Private::ELMT_WARNING ? EMessageSeverity::Warning : EMessageSeverity::Error;
 			const ELoggerSpamBin SpamBin = [&] {
 				switch (pLog->GetMessageSpamBin(i)) {
-				case mu::ErrorLogMessageSpamBin::ELMSB_UNKNOWN_TAG:
+				case UE::Mutable::Private::ErrorLogMessageSpamBin::ELMSB_UNKNOWN_TAG:
 					return ELoggerSpamBin::TagsNotFound;
-				case mu::ErrorLogMessageSpamBin::ELMSB_ALL:
+				case UE::Mutable::Private::ErrorLogMessageSpamBin::ELMSB_ALL:
 				default:
 					return ELoggerSpamBin::ShowAll;
 			}
@@ -350,7 +350,7 @@ const TArray<FCustomizableObjectCompileRunnable::FError>& FCustomizableObjectCom
 
 
 FCustomizableObjectSaveDDRunnable::FCustomizableObjectSaveDDRunnable(const TSharedPtr<FCompilationRequest>& InRequest,
-	TSharedPtr<MutablePrivate::FMutableCachedPlatformData>& InPlatformData)
+	TSharedPtr<UE::Mutable::Private::FMutableCachedPlatformData>& InPlatformData)
 {
 	MUTABLE_CPUPROFILER_SCOPE(FCustomizableObjectSaveDDRunnable::FCustomizableObjectSaveDDRunnable);
 
@@ -437,12 +437,12 @@ void FCustomizableObjectSaveDDRunnable::CachePlatformData()
 		// Generate list of files and update streamable blocks ids and offsets
 		if (Options.bUseBulkData)
 		{
-			MutablePrivate::GenerateBulkDataFilesListWithFileLimit(PlatformData->Model, *PlatformData->ModelStreamableBulkData.Get(), MAX_uint8, false, PlatformData->BulkDataFiles);
+			UE::Mutable::Private::GenerateBulkDataFilesListWithFileLimit(PlatformData->Model, *PlatformData->ModelStreamableBulkData.Get(), MAX_uint8, false, PlatformData->BulkDataFiles);
 		}
 		else
 		{
 			const uint64 PackageDataBytesLimit = Options.bIsCooking ? Options.PackagedDataBytesLimit : MAX_uint64;
-			MutablePrivate::GenerateBulkDataFilesListWithSizeLimit(PlatformData->Model, *PlatformData->ModelStreamableBulkData.Get(), Options.TargetPlatform, PackageDataBytesLimit, PlatformData->BulkDataFiles);
+			UE::Mutable::Private::GenerateBulkDataFilesListWithSizeLimit(PlatformData->Model, *PlatformData->ModelStreamableBulkData.Get(), Options.TargetPlatform, PackageDataBytesLimit, PlatformData->BulkDataFiles);
 		}
 	}
 
@@ -451,9 +451,9 @@ void FCustomizableObjectSaveDDRunnable::CachePlatformData()
 		FMemoryWriter64 ModelMemoryWriter(ModelData);
 		FUnrealMutableModelBulkWriterCook Streamer(&ModelMemoryWriter, &PlatformData->ModelStreamableData);
 
-		// Serialize mu::FModel and streamable resources 
+		// Serialize UE::Mutable::Private::FModel and streamable resources 
 		constexpr bool bDropData = true;
-		mu::FModel::Serialise(PlatformData->Model.Get(), Streamer, bDropData);
+		UE::Mutable::Private::FModel::Serialise(PlatformData->Model.Get(), Streamer, bDropData);
 	}
 }
 
@@ -483,7 +483,7 @@ void FCustomizableObjectSaveDDRunnable::StoreCachedPlatformDataInDDC(bool& bStor
 
 		// Generate list of files and update streamable blocks ids and offsets
 		constexpr int32 MaxDDCFiles = 1 << 13;
-		MutablePrivate::GenerateBulkDataFilesListWithFileLimit(PlatformData->Model, ModelStreamablesDDC, MaxDDCFiles, !Options.bIsCooking, BulkDataFilesDDC);
+		UE::Mutable::Private::GenerateBulkDataFilesListWithFileLimit(PlatformData->Model, ModelStreamablesDDC, MaxDDCFiles, !Options.bIsCooking, BulkDataFilesDDC);
 
 		if (!Options.bIsCooking)
 		{
@@ -491,7 +491,7 @@ void FCustomizableObjectSaveDDRunnable::StoreCachedPlatformDataInDDC(bool& bStor
 			PutValueRequests.Reserve(BulkDataFilesDDC.Num());
 
 			const auto WriteBulkDataDDC = [&PutValueRequests, &ModelStreamablesDDC, SharedName, DefaultDDCPolicy = DefaultDDCPolicy, DDCKey = DDCKey]
-			(MutablePrivate::FFile& File, TArray64<uint8>& FileBulkData, uint32 FileIndex)
+			(UE::Mutable::Private::FFile& File, TArray64<uint8>& FileBulkData, uint32 FileIndex)
 				{
 					FCachePutValueRequest& ValueRequest = PutValueRequests.AddDefaulted_GetRef();
 					ValueRequest.Name = SharedName;
@@ -503,7 +503,7 @@ void FCustomizableObjectSaveDDRunnable::StoreCachedPlatformDataInDDC(bool& bStor
 				};
 
 			constexpr bool bDropData = false;
-			MutablePrivate::SerializeBulkDataFiles(*PlatformData.Get(), BulkDataFilesDDC, WriteBulkDataDDC, bDropData);
+			UE::Mutable::Private::SerializeBulkDataFiles(*PlatformData.Get(), BulkDataFilesDDC, WriteBulkDataDDC, bDropData);
 
 			FRequestOwner RequestOwner(UE::DerivedData::EPriority::Blocking);
 			GetCache().PutValue(PutValueRequests, RequestOwner,
@@ -518,7 +518,7 @@ void FCustomizableObjectSaveDDRunnable::StoreCachedPlatformDataInDDC(bool& bStor
 		}
 		else
 		{
-			const auto WriteBulkDataDDC = [&RecordBuilder](MutablePrivate::FFile& File, TArray64<uint8>& FileBulkData, uint32 FileIndex)
+			const auto WriteBulkDataDDC = [&RecordBuilder](UE::Mutable::Private::FFile& File, TArray64<uint8>& FileBulkData, uint32 FileIndex)
 				{
 					const FValueId ValueId = GetDerivedDataValueIdForResource(File.DataType, FileIndex, File.ResourceType, File.Flags);
 					const FValue Value = FValue::Compress(FSharedBuffer::MakeView(FileBulkData.GetData(), FileBulkData.Num()));
@@ -527,7 +527,7 @@ void FCustomizableObjectSaveDDRunnable::StoreCachedPlatformDataInDDC(bool& bStor
 				};
 
 			constexpr bool bDropData = false;
-			MutablePrivate::SerializeBulkDataFiles(*PlatformData.Get(), BulkDataFilesDDC, WriteBulkDataDDC, bDropData);
+			UE::Mutable::Private::SerializeBulkDataFiles(*PlatformData.Get(), BulkDataFilesDDC, WriteBulkDataDDC, bDropData);
 		}
 	}
 
@@ -540,7 +540,7 @@ void FCustomizableObjectSaveDDRunnable::StoreCachedPlatformDataInDDC(bool& bStor
 		MemoryWriterDDC << ModelStreamablesDDC;
 
 		const FValue ModelStreamablesValue = FValue::Compress(FSharedBuffer::MakeView(ModelStreamablesBytesDDC.GetData(), ModelStreamablesBytesDDC.Num()));
-		RecordBuilder.AddValue(MutablePrivate::GetDerivedDataModelStreamableBulkDataId(), ModelStreamablesValue);
+		RecordBuilder.AddValue(UE::Mutable::Private::GetDerivedDataModelStreamableBulkDataId(), ModelStreamablesValue);
 	}
 
 
@@ -554,7 +554,7 @@ void FCustomizableObjectSaveDDRunnable::StoreCachedPlatformDataInDDC(bool& bStor
 		MemoryWriter << BulkDataFilesDDC;
 
 		const FValue BulkDataFilesValue = FValue::Compress(FSharedBuffer::MakeView(BulkDataFilesBytes.GetData(), BulkDataFilesBytes.Num()));
-		RecordBuilder.AddValue(MutablePrivate::GetDerivedDataBulkDataFilesId(), BulkDataFilesValue);
+		RecordBuilder.AddValue(UE::Mutable::Private::GetDerivedDataBulkDataFilesId(), BulkDataFilesValue);
 	}
 
 	// Store ModelResources bytes as a FValue
@@ -562,7 +562,7 @@ void FCustomizableObjectSaveDDRunnable::StoreCachedPlatformDataInDDC(bool& bStor
 		MUTABLE_CPUPROFILER_SCOPE(SerializeModelResourcesForDDC);
 
 		const FValue ModelResourcesValue = FValue::Compress(FSharedBuffer::MakeView(ModelResourcesData.GetData(), ModelResourcesData.Num()));
-		RecordBuilder.AddValue(MutablePrivate::GetDerivedDataModelResourcesId(), ModelResourcesValue);
+		RecordBuilder.AddValue(UE::Mutable::Private::GetDerivedDataModelResourcesId(), ModelResourcesValue);
 	}
 
 	// Store Model bytes as a FValue
@@ -570,7 +570,7 @@ void FCustomizableObjectSaveDDRunnable::StoreCachedPlatformDataInDDC(bool& bStor
 		MUTABLE_CPUPROFILER_SCOPE(SerializeModelForDDC);
 
 		const FValue ModelValue = FValue::Compress(FSharedBuffer::MakeView(ModelData.GetData(), ModelData.Num()));
-		RecordBuilder.AddValue(MutablePrivate::GetDerivedDataModelId(), ModelValue);
+		RecordBuilder.AddValue(UE::Mutable::Private::GetDerivedDataModelId(), ModelValue);
 	}
 
 	// Push record to the DDC
@@ -622,10 +622,10 @@ void FCustomizableObjectSaveDDRunnable::StoreCachedPlatformDataToDisk()
 	bool bSuccess = true;
 	TArray<TUniquePtr<FArchive>> MemoryWriters;
 
-	const int32 NumDataTypes = static_cast<int32>(MutablePrivate::EStreamableDataType::DataTypeCount);
+	const int32 NumDataTypes = static_cast<int32>(UE::Mutable::Private::EStreamableDataType::DataTypeCount);
 	for (int32 DataType = 0; DataType < NumDataTypes; ++DataType)
 	{
-		const FString FilePath = FullFileName + GetDataTypeExtension(static_cast<MutablePrivate::EStreamableDataType>(DataType));
+		const FString FilePath = FullFileName + GetDataTypeExtension(static_cast<UE::Mutable::Private::EStreamableDataType>(DataType));
 		if (FileManager.FileExists(*FilePath)
 			&& !FileManager.Delete(*FilePath, true, false, true))
 		{
@@ -645,7 +645,7 @@ void FCustomizableObjectSaveDDRunnable::StoreCachedPlatformDataToDisk()
 	if (bSuccess)
 	{
 		// Serialize Streamable resources
-		const auto WriteBulkDataToDisk = [&MemoryWriters](MutablePrivate::FFile& File, TArray64<uint8>& FileBulkData, uint32 FileIndex)
+		const auto WriteBulkDataToDisk = [&MemoryWriters](UE::Mutable::Private::FFile& File, TArray64<uint8>& FileBulkData, uint32 FileIndex)
 			{
 				const int32 DataTypeIndex = static_cast<int32>(File.DataType);
 				if (ensure(MemoryWriters.IsValidIndex(DataTypeIndex)))
@@ -656,7 +656,7 @@ void FCustomizableObjectSaveDDRunnable::StoreCachedPlatformDataToDisk()
 
 		// Serialize streamable resources into a single file and fix offsets
 		constexpr bool bDropData = true;
-		MutablePrivate::SerializeBulkDataFiles(*PlatformData.Get(), PlatformData->BulkDataFiles, WriteBulkDataToDisk, bDropData);
+		UE::Mutable::Private::SerializeBulkDataFiles(*PlatformData.Get(), PlatformData->BulkDataFiles, WriteBulkDataToDisk, bDropData);
 
 
 		// Serialize Model and ModelResources. Store after SerializeBulkDataFiles fixes the HashToStreamableFiles offsets.
diff --git a/Engine/Plugins/Mutable/Source/CustomizableObjectEditor/Private/MuCOE/CustomizableObjectCompileRunnable.h b/Engine/Plugins/Mutable/Source/CustomizableObjectEditor/Private/MuCOE/CustomizableObjectCompileRunnable.h
--- a/Engine/Plugins/Mutable/Source/CustomizableObjectEditor/Private/MuCOE/CustomizableObjectCompileRunnable.h
+++ b/Engine/Plugins/Mutable/Source/CustomizableObjectEditor/Private/MuCOE/CustomizableObjectCompileRunnable.h
@@ -47,14 +47,14 @@ public:
 
 private:
 
-	mu::Ptr<mu::Node> MutableRoot;
+	UE::Mutable::Private::Ptr<UE::Mutable::Private::Node> MutableRoot;
 	TArray<FError> ArrayErrors;
 
-	TSharedPtr<mu::FImage> LoadImageResourceReferenced(int32 ID);
+	TSharedPtr<UE::Mutable::Private::FImage> LoadImageResourceReferenced(int32 ID);
 
 public:
 
-	FCustomizableObjectCompileRunnable(mu::Ptr<mu::Node> Root, const TSharedRef<FCustomizableObjectCompiler>& InCompiler);
+	FCustomizableObjectCompileRunnable(UE::Mutable::Private::Ptr<UE::Mutable::Private::Node> Root, const TSharedRef<FCustomizableObjectCompiler>& InCompiler);
 
 	// FRunnable interface
 	uint32 Run() override;
@@ -67,7 +67,7 @@ public:
 	//
 	const TArray<FError>& GetArrayErrors() const;
 
-	TSharedPtr<mu::FModel, ESPMode::ThreadSafe> Model;
+	TSharedPtr<UE::Mutable::Private::FModel, ESPMode::ThreadSafe> Model;
 
 	FCompilationOptions Options;
 
@@ -87,7 +87,7 @@ class FCustomizableObjectSaveDDRunnable : public FRunnable
 {
 public:
 
-	FCustomizableObjectSaveDDRunnable(const TSharedPtr<FCompilationRequest>& InRequest, TSharedPtr<MutablePrivate::FMutableCachedPlatformData>& InPlatformData);
+	FCustomizableObjectSaveDDRunnable(const TSharedPtr<FCompilationRequest>& InRequest, TSharedPtr<UE::Mutable::Private::FMutableCachedPlatformData>& InPlatformData);
 
 	// FRunnable interface
 	uint32 Run() override;
@@ -127,8 +127,8 @@ public:
 	TArray64<uint8> ModelResourcesData;
 
 	// Cached platform data
-	TSharedPtr<MutablePrivate::FMutableCachedPlatformData> PlatformData;
+	TSharedPtr<UE::Mutable::Private::FMutableCachedPlatformData> PlatformData;
 
 	// DDC Helpers
-	TArray<MutablePrivate::FFile> BulkDataFilesDDC;
+	TArray<UE::Mutable::Private::FFile> BulkDataFilesDDC;
 };
diff --git a/Engine/Plugins/Mutable/Source/CustomizableObjectEditor/Private/MuCOE/CustomizableObjectCompiler.cpp b/Engine/Plugins/Mutable/Source/CustomizableObjectEditor/Private/MuCOE/CustomizableObjectCompiler.cpp
--- a/Engine/Plugins/Mutable/Source/CustomizableObjectEditor/Private/MuCOE/CustomizableObjectCompiler.cpp
+++ b/Engine/Plugins/Mutable/Source/CustomizableObjectEditor/Private/MuCOE/CustomizableObjectCompiler.cpp
@@ -204,7 +204,7 @@ void FCustomizableObjectCompiler::PreloadingReferencerAssetsCallback(bool bAsync
 
 		for (FSoftObjectPath AssetToStream : AssetsToStream)
 		{
-			ArrayGCProtect.Add(MutablePrivate::LoadObject(AssetToStream)); // Already loaded.
+			ArrayGCProtect.Add(UE::Mutable::Private::LoadObject(AssetToStream)); // Already loaded.
 		}
 
 		AsynchronousStreamableHandlePtr = nullptr;
@@ -339,7 +339,7 @@ void FCustomizableObjectCompiler::Compile(const TSharedRef<FCompilationRequest>&
 	CompilationStartTime = FPlatformTime::Seconds();
 
 	// Platform data to cache all compiled resources
-	PlatformData = MakeShared<MutablePrivate::FMutableCachedPlatformData>();
+	PlatformData = MakeShared<UE::Mutable::Private::FMutableCachedPlatformData>();
 
 	// Now that we know for sure that the CO is locked and there are no pending updates of instances using the CO,
 	// destroy any live update instances, as they become invalid when recompiling the CO
@@ -566,7 +566,7 @@ void ProcessChildObjectsRecursively(const UCustomizableObject* ParentObject, FMu
 	{
 		FSoftObjectPath SoftObjectPath = AssetData.GetSoftObjectPath();
 		
-		UCustomizableObject* ChildObject = Cast<UCustomizableObject>(MutablePrivate::LoadObject(SoftObjectPath));
+		UCustomizableObject* ChildObject = Cast<UCustomizableObject>(UE::Mutable::Private::LoadObject(SoftObjectPath));
 		if (!ChildObject || ChildObject->HasAnyFlags(RF_Transient))
 		{
 			continue;
@@ -630,7 +630,7 @@ void ProcessChildObjectsRecursively(const UCustomizableObject* ParentObject, FMu
 }
 
 
-mu::Ptr<mu::NodeObject> GenerateMutableRoot( 
+UE::Mutable::Private::Ptr<UE::Mutable::Private::NodeObject> GenerateMutableRoot( 
 	const UCustomizableObject* Object, 
 	FMutableGraphGenerationContext& GenerationContext)
 {
@@ -773,7 +773,7 @@ mu::Ptr<mu::NodeObject> GenerateMutableRoot(
 
 	// Generate the object expression
 	UE_LOG(LogMutable, Verbose, TEXT("PROFILE: [ %16.8f ] GenerateMutableSource start."), FPlatformTime::Seconds());
-	mu::Ptr<mu::NodeObject> MutableRoot = GenerateMutableSource(RootNodeObject->OutputPin(), GenerationContext);
+	UE::Mutable::Private::Ptr<UE::Mutable::Private::NodeObject> MutableRoot = GenerateMutableSource(RootNodeObject->OutputPin(), GenerationContext);
 	UE_LOG(LogMutable, Verbose, TEXT("PROFILE: [ %16.8f ] GenerateMutableSource end."), FPlatformTime::Seconds());
 
 	// Generate ReferenceSkeletalMeshes data;
@@ -915,7 +915,7 @@ void FCustomizableObjectCompiler::CompileInternal(bool bAsync)
 	CompilationLogsContainer.ClearMessagesArray();
 
 	// Generate the mutable node expression
-	mu::Ptr<mu::NodeObject> MutableRoot = GenerateMutableRoot(CurrentObject, GenerationContext);
+	UE::Mutable::Private::Ptr<UE::Mutable::Private::NodeObject> MutableRoot = GenerateMutableRoot(CurrentObject, GenerationContext);
 	if (!MutableRoot)
 	{
 		CompilerLog(FText(LOCTEXT("FailedToGenerateRoot", "Failed to generate the mutable node graph. Object not built.")));
@@ -929,7 +929,7 @@ void FCustomizableObjectCompiler::CompileInternal(bool bAsync)
 			{
 				check(Pair.Value.ID == NewCompileTimeReferencedTextures.Num());
 
-				FMutableSourceTextureData Tex(*MutablePrivate::LoadObject(Pair.Key));
+				FMutableSourceTextureData Tex(*UE::Mutable::Private::LoadObject(Pair.Key));
 				NewCompileTimeReferencedTextures.Add(Tex);
 			}
 		}
@@ -1350,10 +1350,10 @@ bool FCustomizableObjectCompiler::TryLoadCompiledDataFromDDC(UCustomizableObject
 	FCacheRecordPolicyBuilder PolicyBuilder(DefaultPolicy | ECachePolicy::SkipData);
 
 	// Overwrite the request policy for the resources we want to pull
-	PolicyBuilder.AddValuePolicy(MutablePrivate::GetDerivedDataModelResourcesId(), DefaultPolicy);
-	PolicyBuilder.AddValuePolicy(MutablePrivate::GetDerivedDataModelId(), DefaultPolicy);
-	PolicyBuilder.AddValuePolicy(MutablePrivate::GetDerivedDataModelStreamableBulkDataId(), DefaultPolicy);
-	PolicyBuilder.AddValuePolicy(MutablePrivate::GetDerivedDataBulkDataFilesId(), DefaultPolicy);
+	PolicyBuilder.AddValuePolicy(UE::Mutable::Private::GetDerivedDataModelResourcesId(), DefaultPolicy);
+	PolicyBuilder.AddValuePolicy(UE::Mutable::Private::GetDerivedDataModelId(), DefaultPolicy);
+	PolicyBuilder.AddValuePolicy(UE::Mutable::Private::GetDerivedDataModelStreamableBulkDataId(), DefaultPolicy);
+	PolicyBuilder.AddValuePolicy(UE::Mutable::Private::GetDerivedDataBulkDataFilesId(), DefaultPolicy);
 
 	FCacheGetRequest Request;
 	Request.Name = GetPathNameSafe(CurrentObject);
@@ -1371,16 +1371,16 @@ bool FCustomizableObjectCompiler::TryLoadCompiledDataFromDDC(UCustomizableObject
 			
 			if (Response.Status == EStatus::Ok)
 			{
-				const FCompressedBuffer& ModelCompressedBuffer = Response.Record.GetValue(MutablePrivate::GetDerivedDataModelId()).GetData();
+				const FCompressedBuffer& ModelCompressedBuffer = Response.Record.GetValue(UE::Mutable::Private::GetDerivedDataModelId()).GetData();
 				DDCHeapMemory->ModelBytesDDC = ModelCompressedBuffer.Decompress();
 
-				const FCompressedBuffer& ModelResourcesCompressedBuffer = Response.Record.GetValue(MutablePrivate::GetDerivedDataModelResourcesId()).GetData();
+				const FCompressedBuffer& ModelResourcesCompressedBuffer = Response.Record.GetValue(UE::Mutable::Private::GetDerivedDataModelResourcesId()).GetData();
 				DDCHeapMemory->ModelResourcesBytesDDC = ModelResourcesCompressedBuffer.Decompress();
 
-				const FCompressedBuffer& ModelStreamablesCompressedBuffer = Response.Record.GetValue(MutablePrivate::GetDerivedDataModelStreamableBulkDataId()).GetData();
+				const FCompressedBuffer& ModelStreamablesCompressedBuffer = Response.Record.GetValue(UE::Mutable::Private::GetDerivedDataModelStreamableBulkDataId()).GetData();
 				DDCHeapMemory->ModelStreamablesBytesDDC = ModelStreamablesCompressedBuffer.Decompress();
 
-				const FCompressedBuffer& BulkDataFilesCompressedBuffer = Response.Record.GetValue(MutablePrivate::GetDerivedDataBulkDataFilesId()).GetData();
+				const FCompressedBuffer& BulkDataFilesCompressedBuffer = Response.Record.GetValue(UE::Mutable::Private::GetDerivedDataBulkDataFilesId()).GetData();
 				DDCHeapMemory->BulkDataFilesBytesDDC = BulkDataFilesCompressedBuffer.Decompress();
 			}
 			
@@ -1467,18 +1467,18 @@ void FCustomizableObjectCompiler::FinishLoadingModelDataFromDDC()
 					const FCompilationOptions& Options = CurrentRequest->Options;
 
 					// Value Id to file mapping to reconstruct the cached data
-					TMap<FValueId, MutablePrivate::FFile> ValueIdToFile;
+					TMap<FValueId, UE::Mutable::Private::FFile> ValueIdToFile;
 
 					{
 						MUTABLE_CPUPROFILER_SCOPE(BuildValueIdToFile);
-						TArray<MutablePrivate::FFile> BulkDataFiles;
+						TArray<UE::Mutable::Private::FFile> BulkDataFiles;
 						FMemoryReaderView FilesReader(DDCHeapMemory->BulkDataFilesBytesDDC.GetView());
 						FilesReader << BulkDataFiles;
 
 						ValueIdToFile.Reserve(BulkDataFiles.Num());
 
 						uint32 FileIndex = 0;
-						for (MutablePrivate::FFile& File : BulkDataFiles)
+						for (UE::Mutable::Private::FFile& File : BulkDataFiles)
 						{
 							FValueId ValueId = GetDerivedDataValueIdForResource(File.DataType, FileIndex, File.ResourceType, File.Flags);
 							ValueIdToFile.Add(ValueId, MoveTemp(File));
@@ -1496,7 +1496,7 @@ void FCustomizableObjectCompiler::FinishLoadingModelDataFromDDC()
 					{
 						check(Value.IsValid());
 
-						const MutablePrivate::FFile* File = ValueIdToFile.Find(Value.GetId());
+						const UE::Mutable::Private::FFile* File = ValueIdToFile.Find(Value.GetId());
 						if (!File) // Skip value. It is not a streamable binary blob.
 						{
 							continue;
@@ -1512,25 +1512,25 @@ void FCustomizableObjectCompiler::FinishLoadingModelDataFromDDC()
 						// Filter and cache the data by DataType
 						switch (File->DataType)
 						{
-						case MutablePrivate::EStreamableDataType::Model:
+						case UE::Mutable::Private::EStreamableDataType::Model:
 						{
-							for (const MutablePrivate::FBlock& Block : File->Blocks)
+							for (const UE::Mutable::Private::FBlock& Block : File->Blocks)
 							{
 								PlatformData->ModelStreamableData.Set(Block.Id, TempData.GetData() + Block.Offset, Block.Size);
 							}
 							break;
 						}
-						case MutablePrivate::EStreamableDataType::RealTimeMorph:
+						case UE::Mutable::Private::EStreamableDataType::RealTimeMorph:
 						{
-							for (const MutablePrivate::FBlock& Block : File->Blocks)
+							for (const UE::Mutable::Private::FBlock& Block : File->Blocks)
 							{
 								PlatformData->MorphStreamableData.Set(Block.Id, TempData.GetData() + Block.Offset, Block.Size);
 							}
 							break;
 						}
-						case MutablePrivate::EStreamableDataType::Clothing:
+						case UE::Mutable::Private::EStreamableDataType::Clothing:
 						{
-							for (const MutablePrivate::FBlock& Block : File->Blocks)
+							for (const UE::Mutable::Private::FBlock& Block : File->Blocks)
 							{
 								PlatformData->ClothingStreamableData.Set(Block.Id, TempData.GetData() + Block.Offset, Block.Size);
 							}
@@ -1547,12 +1547,12 @@ void FCustomizableObjectCompiler::FinishLoadingModelDataFromDDC()
 					if (CVarMutableUseBulkData.GetValueOnAnyThread())
 					{
 						const uint32 NumBulkDataFilesPerBucket = MAX_uint8;
-						MutablePrivate::GenerateBulkDataFilesListWithFileLimit(PlatformData->Model, *PlatformData->ModelStreamableBulkData.Get(),
+						UE::Mutable::Private::GenerateBulkDataFilesListWithFileLimit(PlatformData->Model, *PlatformData->ModelStreamableBulkData.Get(),
 							NumBulkDataFilesPerBucket, false, PlatformData->BulkDataFiles);
 					}
 					else
 					{
-						MutablePrivate::GenerateBulkDataFilesListWithSizeLimit(PlatformData->Model, *PlatformData->ModelStreamableBulkData.Get(),
+						UE::Mutable::Private::GenerateBulkDataFilesListWithSizeLimit(PlatformData->Model, *PlatformData->ModelStreamableBulkData.Get(),
 							Options.TargetPlatform, Options.PackagedDataBytesLimit, PlatformData->BulkDataFiles);
 					}
 				}
@@ -1591,7 +1591,7 @@ void FCustomizableObjectCompiler::FinishLoadingStreamableDataFromDDC()
 	{
 		if (CurrentOptions.bIsCooking)
 		{
-			MutablePrivate::FMutableCachedPlatformData& CachedPlatformData = CurrentObject->GetPrivate()->CachedPlatformsData.Add(CurrentOptions.TargetPlatform->PlatformName(), {});
+			UE::Mutable::Private::FMutableCachedPlatformData& CachedPlatformData = CurrentObject->GetPrivate()->CachedPlatformsData.Add(CurrentOptions.TargetPlatform->PlatformName(), {});
 			CachedPlatformData = MoveTemp(*PlatformData.Get());
 		}
 
@@ -1608,7 +1608,7 @@ void FCustomizableObjectCompiler::FinishLoadingStreamableDataFromDDC()
 }
 
 
-mu::Ptr<mu::Node> FCustomizableObjectCompiler::Export(UCustomizableObject* Object, const FCompilationOptions& InCompilerOptions, 
+UE::Mutable::Private::Ptr<UE::Mutable::Private::Node> FCustomizableObjectCompiler::Export(UCustomizableObject* Object, const FCompilationOptions& InCompilerOptions, 
 	TArray<TSoftObjectPtr<const UTexture>>& OutRuntimeReferencedTextures,
 	TArray<FMutableSourceTextureData>& OutCompilerReferencedTextures, 
 	TArray<TSoftObjectPtr<const UStreamableRenderAsset>>& OutRuntimeReferencedMeshes,
@@ -1635,7 +1635,7 @@ mu::Ptr<mu::Node> FCustomizableObjectCompiler::Export(UCustomizableObject* Objec
 	FMutableGraphGenerationContext GenerationContext(CompilationContext);
 	
 	// Generate the mutable node expression
-	mu::Ptr<mu::NodeObject> MutableRoot = GenerateMutableRoot(Object, GenerationContext);
+	UE::Mutable::Private::Ptr<UE::Mutable::Private::NodeObject> MutableRoot = GenerateMutableRoot(Object, GenerationContext);
 	if (!MutableRoot)
 	{
 		CompilerLog(LOCTEXT("FailedToExport", "Failed to generate the mutable node graph. Object not built."), nullptr);
@@ -1655,7 +1655,7 @@ mu::Ptr<mu::Node> FCustomizableObjectCompiler::Export(UCustomizableObject* Objec
 	{
 		check(Pair.Value.ID == OutCompilerReferencedTextures.Num());
 
-		FMutableSourceTextureData Tex(*MutablePrivate::LoadObject(Pair.Key));
+		FMutableSourceTextureData Tex(*UE::Mutable::Private::LoadObject(Pair.Key));
 		OutCompilerReferencedTextures.Add(Tex);
 	}
 
@@ -1685,7 +1685,7 @@ void FCustomizableObjectCompiler::FinishCompilationTask()
 	check(CompileTask.IsValid());
 
 	UpdateCompilerLogData();
-	TSharedPtr<mu::FModel, ESPMode::ThreadSafe> Model = CompileTask->Model;
+	TSharedPtr<UE::Mutable::Private::FModel, ESPMode::ThreadSafe> Model = CompileTask->Model;
 	PlatformData->Model = Model;
 
 	// Generate a map that using the resource id tells the offset and size of the resource inside the bulk data
@@ -1862,7 +1862,7 @@ void FCustomizableObjectCompiler::FinishSavingDerivedDataTask()
 
 		check(!CurrentObject->GetPrivate()->CachedPlatformsData.Find(PlatformName));
 
-		MutablePrivate::FMutableCachedPlatformData& Data = CurrentObject->GetPrivate()->CachedPlatformsData.Add(PlatformName);
+		UE::Mutable::Private::FMutableCachedPlatformData& Data = CurrentObject->GetPrivate()->CachedPlatformsData.Add(PlatformName);
 		Data = MoveTemp(*PlatformData.Get());
 
 		Data.ModelResources->InitCookData(*CurrentObject);
