[Mutable] Change namespace name

diff --git a/Engine/Plugins/Mutable/Source/CustomizableObjectEditor/Private/MuCOE/CustomizableObjectCompiler.h b/Engine/Plugins/Mutable/Source/CustomizableObjectEditor/Private/MuCOE/CustomizableObjectCompiler.h
--- a/Engine/Plugins/Mutable/Source/CustomizableObjectEditor/Private/MuCOE/CustomizableObjectCompiler.h
+++ b/Engine/Plugins/Mutable/Source/CustomizableObjectEditor/Private/MuCOE/CustomizableObjectCompiler.h
@@ -26,14 +26,14 @@ class UCustomizableObjectNode;
 class UModelResources;
 struct FMutableGraphGenerationContext;
 
-namespace mu
+namespace UE::Mutable::Private
 {
 	class NodeObject;
 	class Node;
 }
 
 
-mu::Ptr<mu::NodeObject> GenerateMutableRoot(const UCustomizableObject* Object, FMutableGraphGenerationContext& GenerationContext);
+UE::Mutable::Private::Ptr<UE::Mutable::Private::NodeObject> GenerateMutableRoot(const UCustomizableObject* Object, FMutableGraphGenerationContext& GenerationContext);
 
 
 class FCustomizableObjectCompiler : public TSharedFromThis<FCustomizableObjectCompiler>, public FTickableEditorObject, public FTickableCookObject, public FTickableGameObject, public FGCObject
@@ -53,7 +53,7 @@ public:
 	virtual void TickCook(float DeltaTime, bool bCookCompete) override;
 	
 	/** Generate the Mutable Graph from the Unreal Graph. */
-	mu::Ptr<mu::Node> Export(UCustomizableObject* Object, const FCompilationOptions& Options, 
+	UE::Mutable::Private::Ptr<UE::Mutable::Private::Node> Export(UCustomizableObject* Object, const FCompilationOptions& Options, 
 		TArray<TSoftObjectPtr<const UTexture>>& OutRuntimeReferencedTextures, 
 		TArray<FMutableSourceTextureData>& OutCompilerReferencedTextures,
 		TArray<TSoftObjectPtr<const UStreamableRenderAsset>>& OutRuntimeReferencedMeshes,
@@ -162,7 +162,7 @@ private:
 	TArray<TObjectPtr<UObject>> ArrayGCProtect;
 
 	/** Cached Platform Data. */
-	TSharedPtr<MutablePrivate::FMutableCachedPlatformData> PlatformData;
+	TSharedPtr<UE::Mutable::Private::FMutableCachedPlatformData> PlatformData;
 
 	// Protected from GC with FCustomizableObjectCompiler::AddReferencedObjects
 	TObjectPtr<UCustomizableObject> CurrentObject = nullptr;
diff --git a/Engine/Plugins/Mutable/Source/CustomizableObjectEditor/Private/MuCOE/CustomizableObjectCookPackageSplitter.cpp b/Engine/Plugins/Mutable/Source/CustomizableObjectEditor/Private/MuCOE/CustomizableObjectCookPackageSplitter.cpp
--- a/Engine/Plugins/Mutable/Source/CustomizableObjectEditor/Private/MuCOE/CustomizableObjectCookPackageSplitter.cpp
+++ b/Engine/Plugins/Mutable/Source/CustomizableObjectEditor/Private/MuCOE/CustomizableObjectCookPackageSplitter.cpp
@@ -16,8 +16,8 @@ namespace
 UModelResources* FindModelResources(UCustomizableObject& Object)
 {
 	// All platforms should have the same resources
-	const TMap<FString, MutablePrivate::FMutableCachedPlatformData>& CachePlatforms = Object.GetPrivate()->CachedPlatformsData;
-	for (const TPair<FString, MutablePrivate::FMutableCachedPlatformData>& PlatformData : CachePlatforms)
+	const TMap<FString, UE::Mutable::Private::FMutableCachedPlatformData>& CachePlatforms = Object.GetPrivate()->CachedPlatformsData;
+	for (const TPair<FString, UE::Mutable::Private::FMutableCachedPlatformData>& PlatformData : CachePlatforms)
 	{
 		if (PlatformData.Value.ModelResources)
 		{
@@ -85,7 +85,7 @@ EMoveContainerError MoveContainerToNewOuter(
 	
 	OutContainer = nullptr;
 
-	UCustomizableObjectResourceDataContainer* Container = MutablePrivate::LoadObject(StreamedResourceData->GetPath());
+	UCustomizableObjectResourceDataContainer* Container = UE::Mutable::Private::LoadObject(StreamedResourceData->GetPath());
 	if (!Container)
 	{
 		return EMoveContainerError::FailedToLoadContainer;
@@ -279,8 +279,8 @@ bool FCustomizableObjectCookPackageSplitter::PreSaveGeneratorPackage(FPopulateCo
 	}
 
 	// All platforms should have the same resources
-	const TMap<FString, MutablePrivate::FMutableCachedPlatformData>& CachePlatforms = Object->GetPrivate()->CachedPlatformsData;
-	for (const TPair<FString, MutablePrivate::FMutableCachedPlatformData>& PlatformData : CachePlatforms)
+	const TMap<FString, UE::Mutable::Private::FMutableCachedPlatformData>& CachePlatforms = Object->GetPrivate()->CachedPlatformsData;
+	for (const TPair<FString, UE::Mutable::Private::FMutableCachedPlatformData>& PlatformData : CachePlatforms)
 	{
 		if (PlatformData.Value.ModelResources)
 		{
@@ -384,7 +384,7 @@ bool FCustomizableObjectCookPackageSplitter::PopulateGeneratedPackage(FPopulateC
 
 	// [TEMP] Loading a package referencing the CO before PostSaveGeneratedPackage is called causes a name collision.
 	// Duplicate the object with the new outer instead of moving it until it is fixed.
-	UObject* Container = MutablePrivate::LoadObject(ResourceData->GetPath());
+	UObject* Container = UE::Mutable::Private::LoadObject(ResourceData->GetPath());
 	EMoveContainerError Error = Container ? EMoveContainerError::None : EMoveContainerError::FailedToLoadContainer;
 	if (Container)
 	{
diff --git a/Engine/Plugins/Mutable/Source/CustomizableObjectEditor/Private/MuCOE/CustomizableObjectCustomSettings.cpp b/Engine/Plugins/Mutable/Source/CustomizableObjectEditor/Private/MuCOE/CustomizableObjectCustomSettings.cpp
--- a/Engine/Plugins/Mutable/Source/CustomizableObjectEditor/Private/MuCOE/CustomizableObjectCustomSettings.cpp
+++ b/Engine/Plugins/Mutable/Source/CustomizableObjectEditor/Private/MuCOE/CustomizableObjectCustomSettings.cpp
@@ -453,7 +453,7 @@ void SCustomizableObjectCustomSettings::SetViewportLightsByAsset(const FAssetDat
 		return;	
 	}
 	
-	UCustomizableObjectEditorViewportLights* InLights = Cast<UCustomizableObjectEditorViewportLights>(MutablePrivate::LoadObject(InAsset));	
+	UCustomizableObjectEditorViewportLights* InLights = Cast<UCustomizableObjectEditorViewportLights>(UE::Mutable::Private::LoadObject(InAsset));	
 	if (!InLights)
 	{
 		return;
diff --git a/Engine/Plugins/Mutable/Source/CustomizableObjectEditor/Private/MuCOE/CustomizableObjectDetails.cpp b/Engine/Plugins/Mutable/Source/CustomizableObjectEditor/Private/MuCOE/CustomizableObjectDetails.cpp
--- a/Engine/Plugins/Mutable/Source/CustomizableObjectEditor/Private/MuCOE/CustomizableObjectDetails.cpp
+++ b/Engine/Plugins/Mutable/Source/CustomizableObjectEditor/Private/MuCOE/CustomizableObjectDetails.cpp
@@ -34,7 +34,7 @@ public:
 		ParameterIndex = -1;
 	}
 
-	TSharedPtr<mu::FModel, ESPMode::ThreadSafe> Model;
+	TSharedPtr<UE::Mutable::Private::FModel, ESPMode::ThreadSafe> Model;
 	int StateIndex;
 	int ParameterIndex;
 
@@ -48,7 +48,7 @@ public:
 		}
 		else
 		{
-			TSharedPtr<mu::FParameters> TempParams = mu::FModel::NewParameters(Model);
+			TSharedPtr<UE::Mutable::Private::FParameters> TempParams = UE::Mutable::Private::FModel::NewParameters(Model);
 			res = TempParams->GetName( ParameterIndex );
 		}
 
diff --git a/Engine/Plugins/Mutable/Source/CustomizableObjectEditor/Private/MuCOE/CustomizableObjectEditor.cpp b/Engine/Plugins/Mutable/Source/CustomizableObjectEditor/Private/MuCOE/CustomizableObjectEditor.cpp
--- a/Engine/Plugins/Mutable/Source/CustomizableObjectEditor/Private/MuCOE/CustomizableObjectEditor.cpp
+++ b/Engine/Plugins/Mutable/Source/CustomizableObjectEditor/Private/MuCOE/CustomizableObjectEditor.cpp
@@ -865,13 +865,13 @@ void FCustomizableObjectEditor::ShowGizmoClipMesh(UCustomizableObjectNode& Node,
 		{
 			if (const UCustomizableObjectNodeStaticMesh* StaticMeshNode = Cast<UCustomizableObjectNodeStaticMesh>(ConnectedNode))
 			{
-				ClipMesh = MutablePrivate::LoadObject(StaticMeshNode->GetMesh());
+				ClipMesh = UE::Mutable::Private::LoadObject(StaticMeshNode->GetMesh());
 				StaticMeshNode->GetPinSection(*ConnectedPin, LODIndex, SectionIndex);
 				MaterialSlotIndex = SectionIndex;
 			}
 			else if (const UCustomizableObjectNodeSkeletalMesh* SkeletalMeshNode = Cast<UCustomizableObjectNodeSkeletalMesh>(ConnectedNode))
 			{
-				ClipMesh = MutablePrivate::LoadObject(SkeletalMeshNode->GetMesh());
+				ClipMesh = UE::Mutable::Private::LoadObject(SkeletalMeshNode->GetMesh());
 				SkeletalMeshNode->GetPinSection(*ConnectedPin, LODIndex, SectionIndex);
 				MaterialSlotIndex = SkeletalMeshNode->GetSkeletalMaterialIndexFor(*ConnectedPin);
 			}
@@ -1381,7 +1381,7 @@ const ITargetPlatform* GetTargetPlatform(const TSharedPtr<FString>& TargetPlatfo
 }
 
 
-mu::Ptr<mu::Node> FCustomizableObjectEditor::DebuggerCompile(TSharedPtr<mu::FModel>& Model,
+UE::Mutable::Private::Ptr<UE::Mutable::Private::Node> FCustomizableObjectEditor::DebuggerCompile(TSharedPtr<UE::Mutable::Private::FModel>& Model,
 	TArray<TSoftObjectPtr<const UTexture>>& RuntimeTextures,
 	TArray<FMutableSourceTextureData>& CompilerTextures,
 	TArray<TSoftObjectPtr<const UStreamableRenderAsset>>& RuntimeMeshes,
@@ -1404,7 +1404,7 @@ mu::Ptr<mu::Node> FCustomizableObjectEditor::DebuggerCompile(TSharedPtr<mu::FMod
 			CompilerTextures.Empty();
 			RuntimeMeshes.Empty();
 			CompilerMeshes.Empty();
-			mu::Ptr<mu::Node> RootNode = Compiler->Export(CustomizableObject, CompileOptions, RuntimeTextures, CompilerTextures, RuntimeMeshes, CompilerMeshes);
+			UE::Mutable::Private::Ptr<UE::Mutable::Private::Node> RootNode = Compiler->Export(CustomizableObject, CompileOptions, RuntimeTextures, CompilerTextures, RuntimeMeshes, CompilerMeshes);
 			if (!RootNode)
 			{
 				// TODO: Show errors
@@ -1428,7 +1428,7 @@ mu::Ptr<mu::Node> FCustomizableObjectEditor::DebuggerCompile(TSharedPtr<mu::FMod
 	CompilerTextures.Empty();
 	RuntimeMeshes.Empty();
 	CompilerMeshes.Empty();
-	mu::Ptr<mu::Node> RootNode = Compiler->Export(CustomizableObject, CompileOptions, RuntimeTextures, CompilerTextures, RuntimeMeshes, CompilerMeshes);
+	UE::Mutable::Private::Ptr<UE::Mutable::Private::Node> RootNode = Compiler->Export(CustomizableObject, CompileOptions, RuntimeTextures, CompilerTextures, RuntimeMeshes, CompilerMeshes);
 	if (!RootNode)
 	{
 		// TODO: Show errors
@@ -1452,13 +1452,13 @@ mu::Ptr<mu::Node> FCustomizableObjectEditor::DebuggerCompile(TSharedPtr<mu::FMod
 
 void FCustomizableObjectEditor::GraphViewer() const
 {
-	TSharedPtr<mu::FModel> Model;
+	TSharedPtr<UE::Mutable::Private::FModel> Model;
 	TArray<TSoftObjectPtr<const UTexture>> RuntimeTextures;
 	TArray<FMutableSourceTextureData> CompilerTextures;
 	TArray<TSoftObjectPtr<const UStreamableRenderAsset>> RuntimeMeshes;
 	TArray<FMutableSourceMeshData> CompilerMeshes;
 
-	mu::Ptr<mu::Node> RootNode = DebuggerCompile(Model, RuntimeTextures, CompilerTextures, RuntimeMeshes, CompilerMeshes);
+	UE::Mutable::Private::Ptr<UE::Mutable::Private::Node> RootNode = DebuggerCompile(Model, RuntimeTextures, CompilerTextures, RuntimeMeshes, CompilerMeshes);
 	
 	const TSharedPtr<SDockTab> NewMutableObjectTab = SNew(SDockTab)
 	.Label(LOCTEXT("GraphViewer", "Graph Viewer"))
@@ -1477,13 +1477,13 @@ void FCustomizableObjectEditor::GraphViewer() const
 
 void FCustomizableObjectEditor::CodeViewer() const
 {
-	TSharedPtr<mu::FModel> Model;
+	TSharedPtr<UE::Mutable::Private::FModel> Model;
 	TArray<TSoftObjectPtr<const UTexture>> RuntimeTextures;
 	TArray<FMutableSourceTextureData> CompilerTextures;
 	TArray<TSoftObjectPtr<const UStreamableRenderAsset>> RuntimeMeshes;
 	TArray<FMutableSourceMeshData> CompilerMeshes;
 
-	mu::Ptr<mu::Node> RootNode = DebuggerCompile(Model, RuntimeTextures, CompilerTextures, RuntimeMeshes, CompilerMeshes);
+	UE::Mutable::Private::Ptr<UE::Mutable::Private::Node> RootNode = DebuggerCompile(Model, RuntimeTextures, CompilerTextures, RuntimeMeshes, CompilerMeshes);
 	
 	const TSharedRef<FBoneNames> BoneNames = MakeShared<FBoneNames>(*CustomizableObject->GetPrivate()->GetModelResources());
 	TSharedPtr<FUnrealMutableResourceProvider> ExternalResourceProvider = MakeShared<FUnrealMutableResourceProvider>(BoneNames);
diff --git a/Engine/Plugins/Mutable/Source/CustomizableObjectEditor/Private/MuCOE/CustomizableObjectEditor.h b/Engine/Plugins/Mutable/Source/CustomizableObjectEditor/Private/MuCOE/CustomizableObjectEditor.h
--- a/Engine/Plugins/Mutable/Source/CustomizableObjectEditor/Private/MuCOE/CustomizableObjectEditor.h
+++ b/Engine/Plugins/Mutable/Source/CustomizableObjectEditor/Private/MuCOE/CustomizableObjectEditor.h
@@ -169,7 +169,7 @@ public:
 	/** Callback to notify the editor when the PreviewInstance has been updated */
 	void OnUpdatePreviewInstance(UCustomizableObjectInstance* Instance);
 
-	mu::Ptr<mu::Node> DebuggerCompile(TSharedPtr<mu::FModel>& Model,
+	UE::Mutable::Private::Ptr<UE::Mutable::Private::Node> DebuggerCompile(TSharedPtr<UE::Mutable::Private::FModel>& Model,
 		TArray<TSoftObjectPtr<const UTexture>>& RuntimeTextures,
 		TArray<FMutableSourceTextureData>& CompilerTextures,
 		TArray<TSoftObjectPtr<const UStreamableRenderAsset>>& RuntimeMeshes,
diff --git a/Engine/Plugins/Mutable/Source/CustomizableObjectEditor/Private/MuCOE/CustomizableObjectEditorLevelProfile.cpp b/Engine/Plugins/Mutable/Source/CustomizableObjectEditor/Private/MuCOE/CustomizableObjectEditorLevelProfile.cpp
--- a/Engine/Plugins/Mutable/Source/CustomizableObjectEditor/Private/MuCOE/CustomizableObjectEditorLevelProfile.cpp
+++ b/Engine/Plugins/Mutable/Source/CustomizableObjectEditor/Private/MuCOE/CustomizableObjectEditorLevelProfile.cpp
@@ -38,7 +38,7 @@ bool LevelProfileManage::LoadProfileUObjectNames(const FString& AssetPath,
 											     TArray<FString>& ArrayPostProcess,
 											     TArray<FString>& ArraySkyLight)
 {
-	UWorld* World = MutablePrivate::LoadObject<UWorld>(nullptr, *AssetPath);
+	UWorld* World = UE::Mutable::Private::LoadObject<UWorld>(nullptr, *AssetPath);
 
 	if (World == nullptr)
 	{
@@ -100,7 +100,7 @@ bool LevelProfileManage::LoadProfileUObjects(const FString& AssetPath,
 								             APostProcessVolume*& PostProcessVolume,
 								             ASkyLight*& SkyLight)
 {
-	UWorld* World = MutablePrivate::LoadObject<UWorld>(nullptr, *AssetPath);
+	UWorld* World = UE::Mutable::Private::LoadObject<UWorld>(nullptr, *AssetPath);
 
 	DirectionalLight = nullptr;
 	PostProcessVolume = nullptr;
diff --git a/Engine/Plugins/Mutable/Source/CustomizableObjectEditor/Private/MuCOE/CustomizableObjectEditorModule.cpp b/Engine/Plugins/Mutable/Source/CustomizableObjectEditor/Private/MuCOE/CustomizableObjectEditorModule.cpp
--- a/Engine/Plugins/Mutable/Source/CustomizableObjectEditor/Private/MuCOE/CustomizableObjectEditorModule.cpp
+++ b/Engine/Plugins/Mutable/Source/CustomizableObjectEditor/Private/MuCOE/CustomizableObjectEditorModule.cpp
@@ -558,7 +558,7 @@ void GetReferencingPackages(const UCustomizableObject& Object, TArray<FAssetData
 	{
 		const TSoftObjectPtr SoftObjectPtr(ReferencingCustomizableObject.ToSoftObjectPath());
 
-		const UCustomizableObject* ChildCustomizableObject = Cast<UCustomizableObject>(MutablePrivate::LoadObject(SoftObjectPtr));
+		const UCustomizableObject* ChildCustomizableObject = Cast<UCustomizableObject>(UE::Mutable::Private::LoadObject(SoftObjectPtr));
 		if (!ChildCustomizableObject)
 		{
 			continue;
@@ -1051,7 +1051,7 @@ TMap<FName, FGuid> FCustomizableObjectEditorModule::GetParticipatingObjects(cons
 			}
 			else
 			{
-				DataTable = MutablePrivate::LoadObject(TableNode->Table);
+				DataTable = UE::Mutable::Private::LoadObject(TableNode->Table);
 			}
 
 			if (DataTable)
@@ -1312,7 +1312,7 @@ void FCustomizableObjectEditorModule::OnPreBeginPIE(const bool bIsSimulatingInEd
 			continue;
 		}
 
-		UCustomizableObject* Object = Cast<UCustomizableObject>(MutablePrivate::LoadObject(Asset));
+		UCustomizableObject* Object = Cast<UCustomizableObject>(UE::Mutable::Private::LoadObject(Asset));
 		if (!Object || Object->IsCompiled() || Object->GetPrivate()->IsLocked())
 		{
 			continue;
diff --git a/Engine/Plugins/Mutable/Source/CustomizableObjectEditor/Private/MuCOE/CustomizableObjectEditorViewportClient.cpp b/Engine/Plugins/Mutable/Source/CustomizableObjectEditor/Private/MuCOE/CustomizableObjectEditorViewportClient.cpp
--- a/Engine/Plugins/Mutable/Source/CustomizableObjectEditor/Private/MuCOE/CustomizableObjectEditorViewportClient.cpp
+++ b/Engine/Plugins/Mutable/Source/CustomizableObjectEditor/Private/MuCOE/CustomizableObjectEditorViewportClient.cpp
@@ -134,7 +134,7 @@ FCustomizableObjectEditorViewportClient::FCustomizableObjectEditorViewportClient
 		PluginConfig->GetString(TEXT("EditorDefaults"), TEXT("ClipMorphMaterialName"), ClipMorphMaterialName);
 		if (!ClipMorphMaterialName.IsEmpty())
 		{
-			ClipMorphMaterial = TStrongObjectPtr<UMaterial>(MutablePrivate::LoadObject<UMaterial>(nullptr, *ClipMorphMaterialName));
+			ClipMorphMaterial = TStrongObjectPtr<UMaterial>(UE::Mutable::Private::LoadObject<UMaterial>(nullptr, *ClipMorphMaterialName));
 			ensure(ClipMorphMaterial);
 		}
 
@@ -143,7 +143,7 @@ FCustomizableObjectEditorViewportClient::FCustomizableObjectEditorViewportClient
 		PluginConfig->GetString(TEXT("EditorDefaults"), TEXT("ClipMeshMaterialName"), ClipMeshMaterialName);
 		if (!ClipMeshMaterialName.IsEmpty())
 		{
-			ClipMeshMaterial = TStrongObjectPtr<UMaterial>(MutablePrivate::LoadObject<UMaterial>(nullptr, *ClipMeshMaterialName));
+			ClipMeshMaterial = TStrongObjectPtr<UMaterial>(UE::Mutable::Private::LoadObject<UMaterial>(nullptr, *ClipMeshMaterialName));
 			ensure(ClipMeshMaterial);
 		}
 	}
diff --git a/Engine/Plugins/Mutable/Source/CustomizableObjectEditor/Private/MuCOE/CustomizableObjectIdentifierCustomization.cpp b/Engine/Plugins/Mutable/Source/CustomizableObjectEditor/Private/MuCOE/CustomizableObjectIdentifierCustomization.cpp
--- a/Engine/Plugins/Mutable/Source/CustomizableObjectEditor/Private/MuCOE/CustomizableObjectIdentifierCustomization.cpp
+++ b/Engine/Plugins/Mutable/Source/CustomizableObjectEditor/Private/MuCOE/CustomizableObjectIdentifierCustomization.cpp
@@ -51,7 +51,7 @@ void FCustomizableObjectIdentifierCustomization::CustomizeHeader(TSharedRef<clas
 	{
 		for (auto Itr = OutAssetData.CreateIterator(); Itr; ++Itr)
 		{
-			UCustomizableObject* CustomizableObject = Cast<UCustomizableObject>(MutablePrivate::LoadObject(*Itr));
+			UCustomizableObject* CustomizableObject = Cast<UCustomizableObject>(UE::Mutable::Private::LoadObject(*Itr));
 
 			if (CustomizableObject)
 			{
@@ -95,7 +95,7 @@ FString FCustomizableObjectIdentifierCustomization::OnGetObjectPath() const
 
 void FCustomizableObjectIdentifierCustomization::OnEventChanged(const FAssetData& InAssetData)
 {
-	SelectedCustomizableObject = Cast<UCustomizableObject>(MutablePrivate::LoadObject(InAssetData));
+	SelectedCustomizableObject = Cast<UCustomizableObject>(UE::Mutable::Private::LoadObject(InAssetData));
 
 	if (SelectedCustomizableObject)
 	{
diff --git a/Engine/Plugins/Mutable/Source/CustomizableObjectEditor/Private/MuCOE/CustomizableObjectInstanceFactory.cpp b/Engine/Plugins/Mutable/Source/CustomizableObjectEditor/Private/MuCOE/CustomizableObjectInstanceFactory.cpp
--- a/Engine/Plugins/Mutable/Source/CustomizableObjectEditor/Private/MuCOE/CustomizableObjectInstanceFactory.cpp
+++ b/Engine/Plugins/Mutable/Source/CustomizableObjectEditor/Private/MuCOE/CustomizableObjectInstanceFactory.cpp
@@ -82,7 +82,7 @@ bool UCustomizableObjectInstanceFactory::CanCreateActorFrom(const FAssetData& As
 
     if (AssetData.GetClass()->IsChildOf(UCustomizableObjectInstance::StaticClass()))
     {
-        if (UCustomizableObjectInstance* CustomizableObjectInstance = Cast<UCustomizableObjectInstance>(MutablePrivate::LoadObject(AssetData)))
+        if (UCustomizableObjectInstance* CustomizableObjectInstance = Cast<UCustomizableObjectInstance>(UE::Mutable::Private::LoadObject(AssetData)))
         {
 			TArray<FName> ComponentNames = CustomizableObjectInstance->GetComponentNames();
 
diff --git a/Engine/Plugins/Mutable/Source/CustomizableObjectEditor/Private/MuCOE/CustomizableObjectLayout.cpp b/Engine/Plugins/Mutable/Source/CustomizableObjectEditor/Private/MuCOE/CustomizableObjectLayout.cpp
--- a/Engine/Plugins/Mutable/Source/CustomizableObjectEditor/Private/MuCOE/CustomizableObjectLayout.cpp
+++ b/Engine/Plugins/Mutable/Source/CustomizableObjectEditor/Private/MuCOE/CustomizableObjectLayout.cpp
@@ -39,22 +39,22 @@ void UCustomizableObjectLayout::SetLayout(int32 LODIndex, int32 MatIndex, int32
 	UVChannel = UVIndex;
 }
 
-mu::EAutoBlocksStrategy ConvertAutomaticBlocksStrategy(const ECustomizableObjectLayoutAutomaticBlocksStrategy InAutomaticBlockStrategy)
+UE::Mutable::Private::EAutoBlocksStrategy ConvertAutomaticBlocksStrategy(const ECustomizableObjectLayoutAutomaticBlocksStrategy InAutomaticBlockStrategy)
 {
-	mu::EAutoBlocksStrategy AutoBlockStrategy = mu::EAutoBlocksStrategy::Rectangles;
+	UE::Mutable::Private::EAutoBlocksStrategy AutoBlockStrategy = UE::Mutable::Private::EAutoBlocksStrategy::Rectangles;
 
 	switch (InAutomaticBlockStrategy)
 	{
 	case ECustomizableObjectLayoutAutomaticBlocksStrategy::Rectangles:
-		AutoBlockStrategy = mu::EAutoBlocksStrategy::Rectangles;
+		AutoBlockStrategy = UE::Mutable::Private::EAutoBlocksStrategy::Rectangles;
 		break;
 
 	case ECustomizableObjectLayoutAutomaticBlocksStrategy::UVIslands:
-		AutoBlockStrategy = mu::EAutoBlocksStrategy::UVIslands;
+		AutoBlockStrategy = UE::Mutable::Private::EAutoBlocksStrategy::UVIslands;
 		break;
 
 	case ECustomizableObjectLayoutAutomaticBlocksStrategy::Ignore:
-		AutoBlockStrategy = mu::EAutoBlocksStrategy::Ignore;
+		AutoBlockStrategy = UE::Mutable::Private::EAutoBlocksStrategy::Ignore;
 		break;
 
 	default:
@@ -66,22 +66,22 @@ mu::EAutoBlocksStrategy ConvertAutomaticBlocksStrategy(const ECustomizableObject
 
 
 
-mu::EPackStrategy ConvertLayoutStrategy(const ECustomizableObjectTextureLayoutPackingStrategy LayoutPackStrategy)
+UE::Mutable::Private::EPackStrategy ConvertLayoutStrategy(const ECustomizableObjectTextureLayoutPackingStrategy LayoutPackStrategy)
 {
-	mu::EPackStrategy PackStrategy = mu::EPackStrategy::Fixed;
+	UE::Mutable::Private::EPackStrategy PackStrategy = UE::Mutable::Private::EPackStrategy::Fixed;
 
 	switch (LayoutPackStrategy)
 	{
 	case ECustomizableObjectTextureLayoutPackingStrategy::Fixed:
-		PackStrategy = mu::EPackStrategy::Fixed;
+		PackStrategy = UE::Mutable::Private::EPackStrategy::Fixed;
 		break;
 
 	case ECustomizableObjectTextureLayoutPackingStrategy::Resizable:
-		PackStrategy = mu::EPackStrategy::Resizeable;
+		PackStrategy = UE::Mutable::Private::EPackStrategy::Resizeable;
 		break;
 
 	case ECustomizableObjectTextureLayoutPackingStrategy::Overlay:
-		PackStrategy = mu::EPackStrategy::Overlay;
+		PackStrategy = UE::Mutable::Private::EPackStrategy::Overlay;
 		break;
 
 	default:
@@ -137,7 +137,7 @@ void UCustomizableObjectLayout::GenerateAutomaticBlocksFromUVs()
 	TSharedPtr<FMutableCompilationContext> CompilationContext = MakeShared<FMutableCompilationContext>(Object, Compiler, Options);
 	FMutableGraphGenerationContext GenerationContext(CompilationContext);
 
-	mu::Ptr<mu::NodeLayout> LayoutNode = CreateMutableLayoutNode(GenerationContext, this, false);
+	UE::Mutable::Private::Ptr<UE::Mutable::Private::NodeLayout> LayoutNode = CreateMutableLayoutNode(GenerationContext, this, false);
 	check(LayoutNode);
 
 	if (PackingStrategy == ECustomizableObjectTextureLayoutPackingStrategy::Overlay)
@@ -155,7 +155,7 @@ void UCustomizableObjectLayout::GenerateAutomaticBlocksFromUVs()
 		}
 	else // Convert Mesh and generate blocks
 	{
-		TSharedPtr<mu::FMesh> MutableMesh;
+		TSharedPtr<UE::Mutable::Private::FMesh> MutableMesh;
 
 		if (USkeletalMesh* SkeletalMesh = Cast<USkeletalMesh>(GenerationContext.LoadObject(GetMesh())))
 		{
@@ -228,7 +228,7 @@ void UCustomizableObjectLayout::GenerateAutomaticBlocksFromUVs()
 		Block.bIsAutomatic = true;
 		Block.Id = FGuid::NewGuid();
 
-		TSharedPtr<mu::FImage> Mask = LayoutNode->Blocks[i].Mask;
+		TSharedPtr<UE::Mutable::Private::FImage> Mask = LayoutNode->Blocks[i].Mask;
 		if (Mask)
 		{
 			UTexture2D* UnrealImage = NewObject<UTexture2D>(UTexture2D::StaticClass());
@@ -265,7 +265,7 @@ void UCustomizableObjectLayout::ConsolidateAutomaticBlocks()
 
 void UCustomizableObjectLayout::GetUVs(TArray<FVector2f>& UVs) const
 {
-	if (const UObject* Mesh = MutablePrivate::LoadObject(GetMesh()))
+	if (const UObject* Mesh = UE::Mutable::Private::LoadObject(GetMesh()))
 	{
 		if (const USkeletalMesh* SkeletalMesh = Cast<const USkeletalMesh>(Mesh))
 		{
diff --git a/Engine/Plugins/Mutable/Source/CustomizableObjectEditor/Private/MuCOE/CustomizableObjectLayout.h b/Engine/Plugins/Mutable/Source/CustomizableObjectEditor/Private/MuCOE/CustomizableObjectLayout.h
--- a/Engine/Plugins/Mutable/Source/CustomizableObjectEditor/Private/MuCOE/CustomizableObjectLayout.h
+++ b/Engine/Plugins/Mutable/Source/CustomizableObjectEditor/Private/MuCOE/CustomizableObjectLayout.h
@@ -20,7 +20,7 @@ enum class ECustomizableObjectTextureLayoutPackingStrategy : uint8
 	Overlay = 2 UMETA(DisplayName = "Overlay Layout")
 };
 
-mu::EPackStrategy ConvertLayoutStrategy(const ECustomizableObjectTextureLayoutPackingStrategy LayoutPackStrategy);
+UE::Mutable::Private::EPackStrategy ConvertLayoutStrategy(const ECustomizableObjectTextureLayoutPackingStrategy LayoutPackStrategy);
 
 UENUM()
 enum class ECustomizableObjectLayoutAutomaticBlocksStrategy : uint8
@@ -33,7 +33,7 @@ enum class ECustomizableObjectLayoutAutomaticBlocksStrategy : uint8
 	Ignore = 2 UMETA(DisplayName = "Ignore"),
 };
 
-mu::EAutoBlocksStrategy ConvertAutomaticBlocksStrategy(const ECustomizableObjectLayoutAutomaticBlocksStrategy LayoutPackStrategy);
+UE::Mutable::Private::EAutoBlocksStrategy ConvertAutomaticBlocksStrategy(const ECustomizableObjectLayoutAutomaticBlocksStrategy LayoutPackStrategy);
 
 UENUM()
 enum class ECustomizableObjectLayoutAutomaticBlocksMergeStrategy : uint8
diff --git a/Engine/Plugins/Mutable/Source/CustomizableObjectEditor/Private/MuCOE/EdGraphSchema_CustomizableObject.cpp b/Engine/Plugins/Mutable/Source/CustomizableObjectEditor/Private/MuCOE/EdGraphSchema_CustomizableObject.cpp
--- a/Engine/Plugins/Mutable/Source/CustomizableObjectEditor/Private/MuCOE/EdGraphSchema_CustomizableObject.cpp
+++ b/Engine/Plugins/Mutable/Source/CustomizableObjectEditor/Private/MuCOE/EdGraphSchema_CustomizableObject.cpp
@@ -119,7 +119,7 @@ class IToolkit;
 
 #define SNAP_GRID (16) // @todo ensure this is the same as SNodePanel::GetSnapGridSize()
 
-namespace MutablePrivate
+namespace UE::Mutable::Private
 {
 	// Maximum distance a drag can be off a node edge to require 'push off' from node
 	const int32 NodeDistance = 60;
@@ -232,11 +232,11 @@ UEdGraphNode* FCustomizableObjectSchemaAction_NewNode::CreateNode(class UEdGraph
 		UEdGraphNode* PinNode = FromPin->GetOwningNode();
 		const float XDelta = FMath::Abs(PinNode->NodePosX - Location.X);
 
-		if (XDelta < MutablePrivate::NodeDistance)
+		if (XDelta < UE::Mutable::Private::NodeDistance)
 		{
 			// Set location to edge of current node minus the max move distance
 			// to force node to push off from connect node enough to give selection handle
-			XLocation = PinNode->NodePosX - MutablePrivate::NodeDistance;
+			XLocation = PinNode->NodePosX - UE::Mutable::Private::NodeDistance;
 		}
 	}
 
@@ -1129,7 +1129,7 @@ enum class ESpawnableObjectType : int32
 
 bool IsSpawnableAsset(const FAssetData& InAsset, ESpawnableObjectType& OutObjectType)
 {
-	UObject* Object = MutablePrivate::LoadObject(InAsset);
+	UObject* Object = UE::Mutable::Private::LoadObject(InAsset);
 
 	// Type used to know what kind of UObject this asset is
 	OutObjectType = ESpawnableObjectType::None;
@@ -1181,7 +1181,7 @@ void UEdGraphSchema_CustomizableObject::DroppedAssetsOnGraph(const TArray<FAsset
 
 		// At this point we know we are working with an asset we can spawn as a mutable node.
 		
-		UObject* Object = MutablePrivate::LoadObject(Asset);
+		UObject* Object = UE::Mutable::Private::LoadObject(Asset);
 		UEdGraphNode* GraphNode = nullptr;
 
 		// Depending on the UObjectType spawn one or another mutable node
diff --git a/Engine/Plugins/Mutable/Source/CustomizableObjectEditor/Private/MuCOE/ExtensionDataCompilerInterface.cpp b/Engine/Plugins/Mutable/Source/CustomizableObjectEditor/Private/MuCOE/ExtensionDataCompilerInterface.cpp
--- a/Engine/Plugins/Mutable/Source/CustomizableObjectEditor/Private/MuCOE/ExtensionDataCompilerInterface.cpp
+++ b/Engine/Plugins/Mutable/Source/CustomizableObjectEditor/Private/MuCOE/ExtensionDataCompilerInterface.cpp
@@ -13,10 +13,10 @@ FExtensionDataCompilerInterface::FExtensionDataCompilerInterface(FMutableGraphGe
 {
 }
 
-TSharedPtr<const mu::FExtensionData> FExtensionDataCompilerInterface::MakeStreamedExtensionData(FInstancedStruct&& Data)
+TSharedPtr<const UE::Mutable::Private::FExtensionData> FExtensionDataCompilerInterface::MakeStreamedExtensionData(FInstancedStruct&& Data)
 {
-	TSharedPtr<mu::FExtensionData> Result = MakeShared<mu::FExtensionData>();
-	Result->Origin = mu::FExtensionData::EOrigin::ConstantStreamed;
+	TSharedPtr<UE::Mutable::Private::FExtensionData> Result = MakeShared<UE::Mutable::Private::FExtensionData>();
+	Result->Origin = UE::Mutable::Private::FExtensionData::EOrigin::ConstantStreamed;
 	Result->Index = GenerationContext.StreamedExtensionData.Num();
 	
 	FCustomizableObjectResourceData* CompileTimeExtensionData = &GenerationContext.StreamedExtensionData.AddDefaulted_GetRef();
@@ -25,10 +25,10 @@ TSharedPtr<const mu::FExtensionData> FExtensionDataCompilerInterface::MakeStream
 	return Result;
 }
 
-TSharedPtr<const mu::FExtensionData> FExtensionDataCompilerInterface::MakeAlwaysLoadedExtensionData(FInstancedStruct&& Data)
+TSharedPtr<const UE::Mutable::Private::FExtensionData> FExtensionDataCompilerInterface::MakeAlwaysLoadedExtensionData(FInstancedStruct&& Data)
 {
-	TSharedPtr<mu::FExtensionData> Result = MakeShared<mu::FExtensionData>();
-	Result->Origin = mu::FExtensionData::EOrigin::ConstantAlwaysLoaded;
+	TSharedPtr<UE::Mutable::Private::FExtensionData> Result = MakeShared<UE::Mutable::Private::FExtensionData>();
+	Result->Origin = UE::Mutable::Private::FExtensionData::EOrigin::ConstantAlwaysLoaded;
 	Result->Index = GenerationContext.AlwaysLoadedExtensionData.Num();
 
 	FCustomizableObjectResourceData* CompileTimeExtensionData = &GenerationContext.AlwaysLoadedExtensionData.AddDefaulted_GetRef();
diff --git a/Engine/Plugins/Mutable/Source/CustomizableObjectEditor/Private/MuCOE/GenerateMutableSource/GenerateMutableSource.cpp b/Engine/Plugins/Mutable/Source/CustomizableObjectEditor/Private/MuCOE/GenerateMutableSource/GenerateMutableSource.cpp
--- a/Engine/Plugins/Mutable/Source/CustomizableObjectEditor/Private/MuCOE/GenerateMutableSource/GenerateMutableSource.cpp
+++ b/Engine/Plugins/Mutable/Source/CustomizableObjectEditor/Private/MuCOE/GenerateMutableSource/GenerateMutableSource.cpp
@@ -149,17 +149,17 @@ void FMutableCompilationContext::Log(const FText& Message, const UObject* Contex
 }
 
 
-mu::FBoneName FMutableCompilationContext::GetBoneUnique(const FName& InBoneName)
+UE::Mutable::Private::FBoneName FMutableCompilationContext::GetBoneUnique(const FName& InBoneName)
 {
 	const FString BoneNameString = InBoneName.ToString().ToLower();
-	mu::FBoneName Bone(CityHash32(reinterpret_cast<const char*>(*BoneNameString), BoneNameString.Len() * sizeof(FString::ElementType)));
+	UE::Mutable::Private::FBoneName Bone(CityHash32(reinterpret_cast<const char*>(*BoneNameString), BoneNameString.Len() * sizeof(FString::ElementType)));
 
 	bool bNewRemappedBoneName = false;
 
 	FString& BoneName = UniqueBoneNames.FindOrAdd(Bone, BoneNameString);
 	while (BoneName != InBoneName)
 	{
-		if (mu::FBoneName* RemappedBoneName = RemappedBoneNames.Find(BoneNameString))
+		if (UE::Mutable::Private::FBoneName* RemappedBoneName = RemappedBoneNames.Find(BoneNameString))
 		{
 			Bone.Id = RemappedBoneName->Id;
 			break;
@@ -239,7 +239,7 @@ FString FMutableCompilationContext::GetObjectName() const
 }
 
 
-bool FMutableCompilationContext::FindBone(const FName& InBoneName, mu::FBoneName& OutBoneName) const
+bool FMutableCompilationContext::FindBone(const FName& InBoneName, UE::Mutable::Private::FBoneName& OutBoneName) const
 {
 	const FString BoneNameString = InBoneName.ToString().ToLower();
 	OutBoneName.Id = CityHash32(reinterpret_cast<const char*>(*BoneNameString), BoneNameString.Len() * sizeof(FString::ElementType));
@@ -248,7 +248,7 @@ bool FMutableCompilationContext::FindBone(const FName& InBoneName, mu::FBoneName
 		return true;
 	}
 
-	if (const mu::FBoneName* BoneName = RemappedBoneNames.Find(BoneNameString))
+	if (const UE::Mutable::Private::FBoneName* BoneName = RemappedBoneNames.Find(BoneNameString))
 	{
 		OutBoneName = *BoneName;
 		return true;
@@ -288,7 +288,7 @@ void FMutableGraphGenerationContext::Log(const FText& Message, const UObject* Co
 }
 
 
-TSharedPtr<mu::FMesh> FMutableGraphGenerationContext::FindGeneratedMesh( const FGeneratedMeshData::FKey& Key )
+TSharedPtr<UE::Mutable::Private::FMesh> FMutableGraphGenerationContext::FindGeneratedMesh( const FGeneratedMeshData::FKey& Key )
 {
 	for (const FGeneratedMeshData& d : GeneratedMeshes)
 	{
@@ -377,7 +377,7 @@ FMutableComponentInfo* FMutableGraphGenerationContext::GetCurrentComponentInfo()
 
 UObject* FMutableGraphGenerationContext::LoadObject(const FSoftObjectPtr& SoftObject)
 {
-	return  MutablePrivate::LoadObject(SoftObject);
+	return  UE::Mutable::Private::LoadObject(SoftObject);
 }
 
 
@@ -535,7 +535,7 @@ UTexture2D* FindReferenceImage(const UEdGraphPin* Pin, FMutableGraphGenerationCo
 
 	else if (const UCustomizableObjectNodeTextureFromChannels* TypedNodeFrom = Cast<UCustomizableObjectNodeTextureFromChannels>(Node))
 	{
-		mu::NodeImagePtr RNode, GNode, BNode, ANode;
+		UE::Mutable::Private::NodeImagePtr RNode, GNode, BNode, ANode;
 
 		if (const UEdGraphPin* ConnectedPin = FollowInputPin(*TypedNodeFrom->RPin()))
 		{
@@ -688,15 +688,15 @@ UTexture2D* FindReferenceImage(const UEdGraphPin* Pin, FMutableGraphGenerationCo
 }
 
 
-mu::Ptr<mu::NodeMeshApplyPose> CreateNodeMeshApplyPose(FMutableGraphGenerationContext& GenerationContext, mu::NodeMeshPtr InputMeshNode, const TArray<FName>& ArrayBoneName, const TArray<FTransform>& ArrayTransform)
+UE::Mutable::Private::Ptr<UE::Mutable::Private::NodeMeshApplyPose> CreateNodeMeshApplyPose(FMutableGraphGenerationContext& GenerationContext, UE::Mutable::Private::NodeMeshPtr InputMeshNode, const TArray<FName>& ArrayBoneName, const TArray<FTransform>& ArrayTransform)
 {
 	check(ArrayBoneName.Num() == ArrayTransform.Num());
 
-	TSharedPtr<mu::FMesh> MutableMesh = MakeShared<mu::FMesh>();
-	mu::Ptr<mu::NodeMeshConstant> PoseNodeMesh = new mu::NodeMeshConstant;
+	TSharedPtr<UE::Mutable::Private::FMesh> MutableMesh = MakeShared<UE::Mutable::Private::FMesh>();
+	UE::Mutable::Private::Ptr<UE::Mutable::Private::NodeMeshConstant> PoseNodeMesh = new UE::Mutable::Private::NodeMeshConstant;
 	PoseNodeMesh->Value = MutableMesh;
 
-	TSharedPtr<mu::FSkeleton> MutableSkeleton = MakeShared<mu::FSkeleton>();
+	TSharedPtr<UE::Mutable::Private::FSkeleton> MutableSkeleton = MakeShared<UE::Mutable::Private::FSkeleton>();
 	MutableMesh->SetSkeleton(MutableSkeleton);
 	MutableMesh->SetBonePoseCount(ArrayBoneName.Num());
 	MutableSkeleton->SetBoneCount(ArrayBoneName.Num());
@@ -704,14 +704,14 @@ mu::Ptr<mu::NodeMeshApplyPose> CreateNodeMeshApplyPose(FMutableGraphGenerationCo
 	for (int32 i = 0; i < ArrayBoneName.Num(); ++i)
 	{
 		const FName BoneName = ArrayBoneName[i];
-		const mu::FBoneName& BoneId = GenerationContext.CompilationContext->GetBoneUnique(BoneName);
+		const UE::Mutable::Private::FBoneName& BoneId = GenerationContext.CompilationContext->GetBoneUnique(BoneName);
 
 		MutableSkeleton->SetDebugName(i, BoneName);
 		MutableSkeleton->SetBoneName(i, { BoneId });
-		MutableMesh->SetBonePose(i, BoneId, (FTransform3f)ArrayTransform[i], mu::EBoneUsageFlags::Skinning);
+		MutableMesh->SetBonePose(i, BoneId, (FTransform3f)ArrayTransform[i], UE::Mutable::Private::EBoneUsageFlags::Skinning);
 	}
 
-	mu::Ptr<mu::NodeMeshApplyPose> NodeMeshApplyPose = new mu::NodeMeshApplyPose;
+	UE::Mutable::Private::Ptr<UE::Mutable::Private::NodeMeshApplyPose> NodeMeshApplyPose = new UE::Mutable::Private::NodeMeshApplyPose;
 	NodeMeshApplyPose->Base = InputMeshNode;
 	NodeMeshApplyPose->Pose = PoseNodeMesh;
 
@@ -800,7 +800,7 @@ uint32 AddUniqueSurfaceMetadata(const FMutableSourceSurfaceMetadata& SourceMetad
 		return *Found;
 	}
 	
-	USkeletalMesh* MetadataBaseMesh = Cast<USkeletalMesh>(MutablePrivate::LoadObject(SourceMetadata.Mesh));
+	USkeletalMesh* MetadataBaseMesh = Cast<USkeletalMesh>(UE::Mutable::Private::LoadObject(SourceMetadata.Mesh));
 	
 	FMutableSurfaceMetadata SurfaceMetadata;
 
@@ -832,7 +832,7 @@ uint32 AddUniqueSurfaceMetadata(const FMutableSourceSurfaceMetadata& SourceMetad
 
 
 // Convert a CustomizableObject Source Graph into a mutable source graph  
-mu::Ptr<mu::NodeObject> GenerateMutableSource(const UEdGraphPin * Pin, FMutableGraphGenerationContext & GenerationContext)
+UE::Mutable::Private::Ptr<UE::Mutable::Private::NodeObject> GenerateMutableSource(const UEdGraphPin * Pin, FMutableGraphGenerationContext & GenerationContext)
 {
 	MUTABLE_CPUPROFILER_SCOPE(GenerateMutableSource);
 
@@ -846,18 +846,18 @@ mu::Ptr<mu::NodeObject> GenerateMutableSource(const UEdGraphPin * Pin, FMutableG
 	const FGeneratedKey Key(reinterpret_cast<void*>(&GenerateMutableSource), *Pin, *Node, GenerationContext, true);
 	if (const FGeneratedData* Generated = GenerationContext.Generated.Find(Key))
 	{
-		return static_cast<mu::NodeObject*>(Generated->Node.get());
+		return static_cast<UE::Mutable::Private::NodeObject*>(Generated->Node.get());
 	}
 
 	// Bool that determines if a node can be added to the cache of nodes.
 	// Most nodes need to be added to the cache but there are some that don't. For exampel, MacroInstanceNodes
 	bool bCacheNode = true;
 
-	mu::Ptr<mu::NodeObject> Result;
+	UE::Mutable::Private::Ptr<UE::Mutable::Private::NodeObject> Result;
 	
 	if (const UCustomizableObjectNodeObject* TypedNodeObj = Cast<UCustomizableObjectNodeObject>(Node))
 	{
-		mu::Ptr<mu::NodeObjectNew> ObjectNode = new mu::NodeObjectNew();
+		UE::Mutable::Private::Ptr<UE::Mutable::Private::NodeObjectNew> ObjectNode = new UE::Mutable::Private::NodeObjectNew();
 		Result = ObjectNode;
 
 		ObjectNode->SetMessageContext(Node);
@@ -942,7 +942,7 @@ mu::Ptr<mu::NodeObject> GenerateMutableSource(const UEdGraphPin * Pin, FMutableG
 			TArray<UEdGraphPin*> ConnectedComponentPins = FollowInputPinArray(*ComponentsPin);
 			for (const UEdGraphPin* ComponentNodePin : ConnectedComponentPins)
 			{
-				mu::Ptr<mu::NodeComponent> ComponentNode = GenerateMutableSourceComponent(ComponentNodePin, GenerationContext);
+				UE::Mutable::Private::Ptr<UE::Mutable::Private::NodeComponent> ComponentNode = GenerateMutableSourceComponent(ComponentNodePin, GenerationContext);
 				ObjectNode->Components.Add(ComponentNode);
 			}
 		}
@@ -959,7 +959,7 @@ mu::Ptr<mu::NodeObject> GenerateMutableSource(const UEdGraphPin * Pin, FMutableG
 					GenerationContext.CurrentMeshComponent = Pair.ComponentName;
 					GenerationContext.CurrentAutoLODStrategy = Pair.Node->AutoLODStrategy;
 					
-					mu::Ptr<mu::NodeModifier> ModifierNode = GenerateMutableSourceModifier(ConnectedModifier, GenerationContext);
+					UE::Mutable::Private::Ptr<UE::Mutable::Private::NodeModifier> ModifierNode = GenerateMutableSourceModifier(ConnectedModifier, GenerationContext);
 					ObjectNode->Modifiers.Add(ModifierNode);
 
 					GenerationContext.CurrentMeshComponent = {};
@@ -996,7 +996,7 @@ mu::Ptr<mu::NodeObject> GenerateMutableSource(const UEdGraphPin * Pin, FMutableG
 				
 				if (const ICustomizableObjectExtensionNode* ExtensionNode = Cast<ICustomizableObjectExtensionNode>(ConnectedNode))
 				{
-					if (mu::Ptr<mu::NodeExtensionData> GeneratedNode = ExtensionNode->GenerateMutableNode(GenerationContext.ExtensionDataCompilerInterface))
+					if (UE::Mutable::Private::Ptr<UE::Mutable::Private::NodeExtensionData> GeneratedNode = ExtensionNode->GenerateMutableNode(GenerationContext.ExtensionDataCompilerInterface))
 					{
 						ObjectNode->AddExtensionDataNode(GeneratedNode, TCHAR_TO_ANSI(*ExtensionInputPin.GlobalPinName.ToString()));
 					}
@@ -1010,14 +1010,14 @@ mu::Ptr<mu::NodeObject> GenerateMutableSource(const UEdGraphPin * Pin, FMutableG
 		ObjectNode->Children.Reserve(ConnectedChildrenPins.Num());
 		for (int32 ChildIndex = 0; ChildIndex < ConnectedChildrenPins.Num(); ++ChildIndex)
 		{
-			mu::Ptr<mu::NodeObject> ChildNode = GenerateMutableSource(ConnectedChildrenPins[ChildIndex], GenerationContext);
+			UE::Mutable::Private::Ptr<UE::Mutable::Private::NodeObject> ChildNode = GenerateMutableSource(ConnectedChildrenPins[ChildIndex], GenerationContext);
 			ObjectNode->Children.Add(ChildNode);
 		}
 	}
 
 	else if (const UCustomizableObjectNodeObjectGroup* TypedNodeGroup = Cast<UCustomizableObjectNodeObjectGroup>(Node))
 	{
-		mu::Ptr<mu::NodeObjectGroup> GroupNode = new mu::NodeObjectGroup();
+		UE::Mutable::Private::Ptr<UE::Mutable::Private::NodeObjectGroup> GroupNode = new UE::Mutable::Private::NodeObjectGroup();
 		Result = GroupNode;
 
 		// All sockets from all mesh parts plugged into this group node will have the following priority when there's a socket name clash
@@ -1053,13 +1053,13 @@ mu::Ptr<mu::NodeObject> GenerateMutableSource(const UEdGraphPin * Pin, FMutableG
 			GenerationContext.CurrentGroupProjectors.Pop();
 		};
 
-		mu::NodeObjectGroup::EChildSelection Type = mu::NodeObjectGroup::CS_ALWAYS_ALL;
+		UE::Mutable::Private::NodeObjectGroup::EChildSelection Type = UE::Mutable::Private::NodeObjectGroup::CS_ALWAYS_ALL;
 		switch (TypedNodeGroup->GroupType)
 		{
-		case ECustomizableObjectGroupType::COGT_ALL: Type = mu::NodeObjectGroup::CS_ALWAYS_ALL; break;
-		case ECustomizableObjectGroupType::COGT_TOGGLE: Type = mu::NodeObjectGroup::CS_TOGGLE_EACH; break;
-		case ECustomizableObjectGroupType::COGT_ONE: Type = mu::NodeObjectGroup::CS_ALWAYS_ONE; break;
-		case ECustomizableObjectGroupType::COGT_ONE_OR_NONE: Type = mu::NodeObjectGroup::CS_ONE_OR_NONE; break;
+		case ECustomizableObjectGroupType::COGT_ALL: Type = UE::Mutable::Private::NodeObjectGroup::CS_ALWAYS_ALL; break;
+		case ECustomizableObjectGroupType::COGT_TOGGLE: Type = UE::Mutable::Private::NodeObjectGroup::CS_TOGGLE_EACH; break;
+		case ECustomizableObjectGroupType::COGT_ONE: Type = UE::Mutable::Private::NodeObjectGroup::CS_ALWAYS_ONE; break;
+		case ECustomizableObjectGroupType::COGT_ONE_OR_NONE: Type = UE::Mutable::Private::NodeObjectGroup::CS_ONE_OR_NONE; break;
 		default:
 			GenerationContext.Log(LOCTEXT("UnsupportedGroupType", "Object Group Type not supported. Setting to 'ALL'."), Node);
 			break;
@@ -1091,7 +1091,7 @@ mu::Ptr<mu::NodeObject> GenerateMutableSource(const UEdGraphPin * Pin, FMutableG
 		const int32 TotalNumChildren = NumChildren + ExternalChildNodes.Num();
 
 		GroupNode->Children.SetNum(TotalNumChildren);
-		GroupNode->DefaultValue = (Type == mu::NodeObjectGroup::CS_ONE_OR_NONE ? -1 : 0);
+		GroupNode->DefaultValue = (Type == UE::Mutable::Private::NodeObjectGroup::CS_ONE_OR_NONE ? -1 : 0);
 		int32 ChildIndex = 0;
 
 		// UI data
@@ -1116,7 +1116,7 @@ mu::Ptr<mu::NodeObject> GenerateMutableSource(const UEdGraphPin * Pin, FMutableG
 			UCustomizableObjectNodeObject* CustomizableObjectNodeObject = Cast<UCustomizableObjectNodeObject>(MacroContextPin->GetOwningNode());
 
 			const FString* SelectedOptionName = GenerationContext.CompilationContext->Options.ParamNamesToSelectedOptions.Find(TypedNodeGroup->GetGroupName(&GenerationContext.MacroNodesStack)); // If the param is in the map restrict to only the selected option
-			mu::Ptr<mu::NodeObject> ChildNode;
+			UE::Mutable::Private::Ptr<UE::Mutable::Private::NodeObject> ChildNode;
 
 			if (bConnectAtLeastTheLastChild || !SelectedOptionName || (CustomizableObjectNodeObject && *SelectedOptionName == CustomizableObjectNodeObject->GetObjectName(&GenerationContext.MacroNodesStack)) )
 			{
@@ -1146,7 +1146,7 @@ mu::Ptr<mu::NodeObject> GenerateMutableSource(const UEdGraphPin * Pin, FMutableG
 			}
 			else
 			{
-				ChildNode = new mu::NodeObjectNew;
+				ChildNode = new UE::Mutable::Private::NodeObjectNew;
 				ChildNode->SetName(CustomizableObjectNodeObject->GetObjectName(&GenerationContext.MacroNodesStack));
 				GroupNode->Children[ChildIndex] = ChildNode;
 			}
@@ -1242,7 +1242,7 @@ mu::Ptr<mu::NodeObject> GenerateMutableSource(const UEdGraphPin * Pin, FMutableG
 			UCustomizableObjectNodeObject* CustomizableObjectNodeObject = Cast<UCustomizableObjectNodeObject>(ExternalChildNode->OutputPin()->GetOwningNode());
 
 			const FString* SelectedOptionName = GenerationContext.CompilationContext->Options.ParamNamesToSelectedOptions.Find(TypedNodeGroup->GetGroupName(&GenerationContext.MacroNodesStack)); // If the param is in the map restrict to only the selected option
-			mu::Ptr<mu::NodeObject> ChildNode;
+			UE::Mutable::Private::Ptr<UE::Mutable::Private::NodeObject> ChildNode;
 
 			if (bConnectAtLeastTheLastChild || !SelectedOptionName || (CustomizableObjectNodeObject  && *SelectedOptionName == CustomizableObjectNodeObject->GetObjectName(&GenerationContext.MacroNodesStack)) )
 			{
@@ -1277,7 +1277,7 @@ mu::Ptr<mu::NodeObject> GenerateMutableSource(const UEdGraphPin * Pin, FMutableG
 			}
 			else
 			{
-				ChildNode = new mu::NodeObjectNew;
+				ChildNode = new UE::Mutable::Private::NodeObjectNew;
 				ChildNode->SetName(CustomizableObjectNodeObject->GetObjectName(&GenerationContext.MacroNodesStack));
 				GroupNode->Children[ChildIndex] = ChildNode;
 			}
@@ -1326,14 +1326,14 @@ mu::Ptr<mu::NodeObject> GenerateMutableSource(const UEdGraphPin * Pin, FMutableG
 	{
 		bCacheNode = false;
 
-		Result = GenerateMutableSourceMacro<mu::NodeObject>(*Pin, GenerationContext, GenerateMutableSource);
+		Result = GenerateMutableSourceMacro<UE::Mutable::Private::NodeObject>(*Pin, GenerationContext, GenerateMutableSource);
 	}
 
 	else if (const UCustomizableObjectNodeTunnel* TypedNodeTunnel = Cast<UCustomizableObjectNodeTunnel>(Node))
 	{
 		bCacheNode = false;
 
-		Result = GenerateMutableSourceMacro<mu::NodeObject>(*Pin, GenerationContext, GenerateMutableSource);
+		Result = GenerateMutableSourceMacro<UE::Mutable::Private::NodeObject>(*Pin, GenerationContext, GenerateMutableSource);
 	}
 
 	else
@@ -1351,7 +1351,7 @@ mu::Ptr<mu::NodeObject> GenerateMutableSource(const UEdGraphPin * Pin, FMutableG
 }
 
 
-int32 AddTagToMutableMeshUnique(mu::FMesh& MutableMesh, const FString& Tag)
+int32 AddTagToMutableMeshUnique(UE::Mutable::Private::FMesh& MutableMesh, const FString& Tag)
 {
 	const int32 TagCount = MutableMesh.GetTagCount();
 
@@ -1635,16 +1635,16 @@ int32 GetTextureSizeInGame(const UTexture2D& Texture, const UTextureLODSettings&
 }
 
 
-mu::FImageDesc GenerateImageDescriptor(UTexture* Texture)
+UE::Mutable::Private::FImageDesc GenerateImageDescriptor(UTexture* Texture)
 {
 	check(Texture);
-	mu::FImageDesc ImageDesc;
+	UE::Mutable::Private::FImageDesc ImageDesc;
 
 	ImageDesc.m_size[0] = Texture->Source.GetSizeX();
 	ImageDesc.m_size[1] = Texture->Source.GetSizeY();
 	ImageDesc.m_lods = Texture->Source.GetNumMips();
 
-	mu::EImageFormat MutableFormat = mu::EImageFormat::RGBA_UByte;
+	UE::Mutable::Private::EImageFormat MutableFormat = UE::Mutable::Private::EImageFormat::RGBA_UByte;
 	ETextureSourceFormat SourceFormat = Texture->Source.GetFormat();
 	switch (SourceFormat)
 	{
@@ -1652,7 +1652,7 @@ mu::FImageDesc GenerateImageDescriptor(UTexture* Texture)
 	case ETextureSourceFormat::TSF_G16:
 	case ETextureSourceFormat::TSF_R16F:
 	case ETextureSourceFormat::TSF_R32F:
-		MutableFormat = mu::EImageFormat::L_UByte;
+		MutableFormat = UE::Mutable::Private::EImageFormat::L_UByte;
 		break;
 
 	default:
@@ -1665,7 +1665,7 @@ mu::FImageDesc GenerateImageDescriptor(UTexture* Texture)
 }
 
 
-TSharedPtr<mu::FImage> GenerateImageConstant(UTexture* Texture, FMutableGraphGenerationContext& GenerationContext, bool bIsReference)
+TSharedPtr<UE::Mutable::Private::FImage> GenerateImageConstant(UTexture* Texture, FMutableGraphGenerationContext& GenerationContext, bool bIsReference)
 {
 	MUTABLE_CPUPROFILER_SCOPE(GenerateImageConstant);
 
@@ -1692,7 +1692,7 @@ TSharedPtr<mu::FImage> GenerateImageConstant(UTexture* Texture, FMutableGraphGen
 	// Create a descriptor for the image.
 	// \TODO: If passthrough (bIsReference) we should apply lod bias, and max texture size to this desc.
 	// For now it is not a problem because passthrough textures shouldn't mix with any other operations.
-	mu::FImageDesc ImageDesc = GenerateImageDescriptor(Texture);
+	UE::Mutable::Private::FImageDesc ImageDesc = GenerateImageDescriptor(Texture);
 
 	FMutableGraphGenerationContext::FGeneratedReferencedTexture InvalidEntry;
 	InvalidEntry.ID = TNumericLimits<uint32>::Max();
@@ -1722,12 +1722,12 @@ TSharedPtr<mu::FImage> GenerateImageConstant(UTexture* Texture, FMutableGraphGen
 	}
 	
 	// Compile-time references that are left should be resolved immediately (should only happen in editor).
-	TSharedPtr<mu::FImage> Result = mu::FImage::CreateAsReference(Entry->ID, ImageDesc, bForceLoad);
+	TSharedPtr<UE::Mutable::Private::FImage> Result = UE::Mutable::Private::FImage::CreateAsReference(Entry->ID, ImageDesc, bForceLoad);
 	return Result;
 }
 
 
-TSharedPtr<mu::FMesh> GenerateMeshConstant(
+TSharedPtr<UE::Mutable::Private::FMesh> GenerateMeshConstant(
 	const FMutableSourceMeshData& Source,
 	FMutableGraphGenerationContext& GenerationContext)
 {
@@ -1783,12 +1783,12 @@ TSharedPtr<mu::FMesh> GenerateMeshConstant(
 		Entry->ID = Num - 1;
 	}
 
-	TSharedPtr<mu::FMesh> Result = mu::FMesh::CreateAsReference(Entry->ID, bForceLoad);
+	TSharedPtr<UE::Mutable::Private::FMesh> Result = UE::Mutable::Private::FMesh::CreateAsReference(Entry->ID, bForceLoad);
 	return Result;
 }
 
 
-void AddSocketTagsToMesh(const USkeletalMesh* SourceMesh, TSharedPtr<mu::FMesh> MutableMesh, FMutableGraphGenerationContext& GenerationContext)
+void AddSocketTagsToMesh(const USkeletalMesh* SourceMesh, TSharedPtr<UE::Mutable::Private::FMesh> MutableMesh, FMutableGraphGenerationContext& GenerationContext)
 {
 	for (int32 SocketIndex = 0; SocketIndex < SourceMesh->NumSockets(); ++SocketIndex)
 	{
diff --git a/Engine/Plugins/Mutable/Source/CustomizableObjectEditor/Private/MuCOE/GenerateMutableSource/GenerateMutableSource.h b/Engine/Plugins/Mutable/Source/CustomizableObjectEditor/Private/MuCOE/GenerateMutableSource/GenerateMutableSource.h
--- a/Engine/Plugins/Mutable/Source/CustomizableObjectEditor/Private/MuCOE/GenerateMutableSource/GenerateMutableSource.h
+++ b/Engine/Plugins/Mutable/Source/CustomizableObjectEditor/Private/MuCOE/GenerateMutableSource/GenerateMutableSource.h
@@ -74,7 +74,7 @@ struct FGeneratedImageProperties
 	/** Name in the Material. */
 	FString TextureParameterName;
 
-	/** Name in the mu::Surface. */
+	/** Name in the UE::Mutable::Private::Surface. */
 	int32 ImagePropertiesIndex = INDEX_NONE;
 
 	TEnumAsByte<TextureCompressionSettings> CompressionSettings = TC_Default;
@@ -123,7 +123,7 @@ struct FMorphNodeData
 	FString MorphTargetName;
 
 	// Pin to the node that generates the factor of the morph
-	mu::Ptr<mu::NodeScalar> FactorNode;
+	UE::Mutable::Private::Ptr<UE::Mutable::Private::NodeScalar> FactorNode;
 
 	// Pin of the mesh where the morphs will ble apllied
 	const UEdGraphPin* MeshPin;
@@ -238,14 +238,14 @@ struct FMutableGraphMeshGenerationData
 // Data stored for each processed unreal graph node, stored in the cache.
 struct FGeneratedData
 {
-	FGeneratedData(const UEdGraphNode* InSource, mu::NodePtr InNode)
+	FGeneratedData(const UEdGraphNode* InSource, UE::Mutable::Private::NodePtr InNode)
 	{
 		Source = InSource;
 		Node = InNode;
 	}
 
 	const UEdGraphNode* Source;
-	mu::NodePtr Node;
+	UE::Mutable::Private::NodePtr Node;
 };
 
 
@@ -299,12 +299,12 @@ struct FClothingMeshDataSource
 struct FGroupProjectorTempData
 {
 	class UCustomizableObjectNodeGroupProjectorParameter* CustomizableObjectNodeGroupProjectorParameter;
-	mu::Ptr<mu::NodeProjectorParameter> NodeProjectorParameterPtr;
-	mu::Ptr<mu::NodeImage> NodeImagePtr;
-	mu::Ptr<mu::NodeRange> NodeRange;
-	mu::Ptr<mu::NodeScalarParameter> NodeOpacityParameter;
+	UE::Mutable::Private::Ptr<UE::Mutable::Private::NodeProjectorParameter> NodeProjectorParameterPtr;
+	UE::Mutable::Private::Ptr<UE::Mutable::Private::NodeImage> NodeImagePtr;
+	UE::Mutable::Private::Ptr<UE::Mutable::Private::NodeRange> NodeRange;
+	UE::Mutable::Private::Ptr<UE::Mutable::Private::NodeScalarParameter> NodeOpacityParameter;
 
-	mu::Ptr<mu::NodeScalarEnumParameter> PoseOptionsParameter;
+	UE::Mutable::Private::Ptr<UE::Mutable::Private::NodeScalarEnumParameter> PoseOptionsParameter;
 	TArray<FPoseBoneData> PoseBoneDataArray;
 
 	bool bAlternateResStateNameWarningDisplayed = false; // Used to display this warning only once
@@ -333,15 +333,15 @@ struct FGroupNodeIdsTempData
 
 struct FGroupProjectorImageInfo
 {
-	mu::Ptr<mu::NodeImage> ImageNode;
-	mu::Ptr<mu::NodeImage> ImageResizeNode;
-	mu::Ptr<mu::NodeSurfaceNew> SurfNode;
+	UE::Mutable::Private::Ptr<UE::Mutable::Private::NodeImage> ImageNode;
+	UE::Mutable::Private::Ptr<UE::Mutable::Private::NodeImage> ImageResizeNode;
+	UE::Mutable::Private::Ptr<UE::Mutable::Private::NodeSurfaceNew> SurfNode;
 	UCustomizableObjectNodeMaterialBase* TypedNodeMat;
 	FString TextureName;
 	FString RealTextureName;
 	int32 UVLayout = 0;
 
-	FGroupProjectorImageInfo(mu::NodeImagePtr InImageNode, const FString& InTextureName, const FString& InRealTextureName, UCustomizableObjectNodeMaterialBase* InTypedNodeMat, mu::Ptr<mu::NodeSurfaceNew> InSurfNode, int32 InUVLayout)
+	FGroupProjectorImageInfo(UE::Mutable::Private::NodeImagePtr InImageNode, const FString& InTextureName, const FString& InRealTextureName, UCustomizableObjectNodeMaterialBase* InTypedNodeMat, UE::Mutable::Private::Ptr<UE::Mutable::Private::NodeSurfaceNew> InSurfNode, int32 InUVLayout)
 		: TypedNodeMat(InTypedNodeMat), TextureName(InTextureName), RealTextureName(InRealTextureName),
 		UVLayout(InUVLayout)
 	{
@@ -499,10 +499,10 @@ struct FMutableCompilationContext
 	TArray<FMutableComponentInfo> ComponentInfos;
 
 	/** Array of unique Bone identifiers. */
-	TMap<mu::FBoneName, FString> UniqueBoneNames;
+	TMap<UE::Mutable::Private::FBoneName, FString> UniqueBoneNames;
 
 	/** Bone identifiers that had a collision. */
-	TMap<FString, mu::FBoneName> RemappedBoneNames;
+	TMap<FString, UE::Mutable::Private::FBoneName> RemappedBoneNames;
 
 private:
 
@@ -528,10 +528,10 @@ public:
 	FMutableComponentInfo* GetComponentInfo(FName ComponentName);
 
 	/** Get unique identifier for BoneName built from its FString. */
-	mu::FBoneName GetBoneUnique(const FName& BoneName);
+	UE::Mutable::Private::FBoneName GetBoneUnique(const FName& BoneName);
 
 	/** */
-	bool FindBone(const FName& BoneName, mu::FBoneName& OutBone) const;
+	bool FindBone(const FName& BoneName, UE::Mutable::Private::FBoneName& OutBone) const;
 
 	/** */
 	uint32 GetSkinWeightProfileIdUnique(const FName ProfileName);
@@ -595,7 +595,7 @@ struct FMutableGraphGenerationContext
 	struct FGeneratedDataTablesData
 	{
 		// Pointer to the generated mutable Table
-		mu::Ptr<mu::FTable> GeneratedTable;
+		UE::Mutable::Private::Ptr<UE::Mutable::Private::FTable> GeneratedTable;
 
 		// Table Node used to fill this info
 		const UCustomizableObjectNodeTable* ReferenceNode;
@@ -611,10 +611,10 @@ struct FMutableGraphGenerationContext
 	TMap<FGeneratedGroupProjectorsKey, FGroupProjectorTempData> GeneratedGroupProjectors;
 
 	/** Key is the Node Uid. */
-	TMap<FString, mu::Ptr<mu::NodeScalarParameter>> GeneratedScalarParameters;
+	TMap<FString, UE::Mutable::Private::Ptr<UE::Mutable::Private::NodeScalarParameter>> GeneratedScalarParameters;
 
 	/** Key is the Node Uid. */
-	TMap<FString, mu::Ptr<mu::NodeScalarEnumParameter>> GeneratedEnumParameters;
+	TMap<FString, UE::Mutable::Private::Ptr<UE::Mutable::Private::NodeScalarEnumParameter>> GeneratedEnumParameters;
 	
 	struct FGeneratedCompositeDataTablesData
 	{
@@ -633,13 +633,13 @@ struct FMutableGraphGenerationContext
 
 	// Cache of generated images, because sometimes they are reused by LOD, we use this as a second
 	// level cache
-	TMap<FGeneratedImageKey, mu::NodeImagePtr> GeneratedImages;
+	TMap<FGeneratedImageKey, UE::Mutable::Private::NodeImagePtr> GeneratedImages;
 
 	/** Data stored per-generated passthrough texture. */
 	struct FGeneratedReferencedTexture
 	{
 		uint32 ID;
-		//mu::FImageDesc ImageDesc;
+		//UE::Mutable::Private::FImageDesc ImageDesc;
 	};
 
 	/** Data stored per-generated passthrough mesh. */
@@ -716,7 +716,7 @@ struct FMutableGraphGenerationContext
 		FKey Key;
 
 		/** Generated mesh. */
-		TSharedPtr<mu::FMesh> Generated;
+		TSharedPtr<UE::Mutable::Private::FMesh> Generated;
 	};
 	TArray<FGeneratedMeshData> GeneratedMeshes;
 
@@ -724,7 +724,7 @@ struct FMutableGraphGenerationContext
 	{
 		FString PinName;
 		FName PinType;
-		const mu::Ptr<mu::FTable> Table;
+		const UE::Mutable::Private::Ptr<UE::Mutable::Private::FTable> Table;
 		const UCustomizableObjectNodeTable* TableNode;
 
 		bool operator==(const FGeneratedTableImageData& Other) const
@@ -751,7 +751,7 @@ struct FMutableGraphGenerationContext
 	void Log(const FText& Message, const UObject* Context = nullptr, const EMessageSeverity::Type MessageSeverity = EMessageSeverity::Warning, const bool bAddBaseObjectInfo = true, const ELoggerSpamBin SpamBin = ELoggerSpamBin::ShowAll) const;
 
 	/** Find a mesh if already generated for a given source and flags. */
-	TSharedPtr<mu::FMesh> FindGeneratedMesh(const FGeneratedMeshData::FKey& Key);
+	TSharedPtr<UE::Mutable::Private::FMesh> FindGeneratedMesh(const FGeneratedMeshData::FKey& Key);
 
 	/** Adds a streamed resource of type AssetUserData.
 	  * Returns resource index in the array of streamed resources. */
@@ -773,13 +773,13 @@ struct FMutableGraphGenerationContext
 	template<typename T>
 	T* LoadObject(const TSoftObjectPtr<T>& SoftObject)
 	{
-		return MutablePrivate::LoadObject<T>(SoftObject);
+		return UE::Mutable::Private::LoadObject<T>(SoftObject);
 	}
 
 	template<typename T>
 	UClass* LoadClass(const TSoftClassPtr<T>& SoftClass)
 	{
-		return MutablePrivate::LoadClass<T>(SoftClass);
+		return UE::Mutable::Private::LoadClass<T>(SoftClass);
 	}
 
 	/** Only compiled components. All components types. Index is the ObjectComponentIndex. */
@@ -810,7 +810,7 @@ struct FMutableGraphGenerationContext
 	FName CurrentMeshComponent;
 
 	/** If this is set, we are genreating materials for a "passthrough" component, with a fixed mesh. */
-	mu::Ptr<mu::NodeMesh> ComponentMeshOverride;
+	UE::Mutable::Private::Ptr<UE::Mutable::Private::NodeMesh> ComponentMeshOverride;
 
 	TMap<FName, uint8> NumLODs;
 
@@ -861,7 +861,7 @@ struct FMutableGraphGenerationContext
 	FString CurrentMaterialTableParameter;
 
 	/** */
-	mu::ETableColumnType CurrentTableColumnType = mu::ETableColumnType::None;
+	UE::Mutable::Private::ETableColumnType CurrentTableColumnType = UE::Mutable::Private::ETableColumnType::None;
 
 	// Current material parameter id to find the corresponding column in a mutable table
 	FString CurrentMaterialParameterId;
@@ -906,8 +906,8 @@ struct FMutableGraphGenerationContext
 	NOTE: This is set when the material pin of a material breake node is compiled and then it must be reset to empty. */
 	FMaterialBreakParameter CurrentMaterialBreakParameter;
 
-	TMap<const UCustomizableObjectNodeMaterialParameter*, mu::Ptr<mu::NodeMaterialParameter>> MaterialParameterNodesCache;
-	TMap<const UCONodeMaterialConstant*, mu::Ptr<mu::NodeMaterialConstant>> MaterialConstantNodesCache;
+	TMap<const UCustomizableObjectNodeMaterialParameter*, UE::Mutable::Private::Ptr<UE::Mutable::Private::NodeMaterialParameter>> MaterialParameterNodesCache;
+	TMap<const UCONodeMaterialConstant*, UE::Mutable::Private::Ptr<UE::Mutable::Private::NodeMaterialConstant>> MaterialConstantNodesCache;
 };
 
 
@@ -967,7 +967,7 @@ namespace Private
 } //Private
 
 //
-mu::Ptr<mu::NodeObject> GenerateMutableSource(const class UEdGraphPin* Pin, FMutableGraphGenerationContext& GenerationContext);
+UE::Mutable::Private::Ptr<UE::Mutable::Private::NodeObject> GenerateMutableSource(const class UEdGraphPin* Pin, FMutableGraphGenerationContext& GenerationContext);
 
 /** Populate an array with all the information related to the reference skeletal meshes we might need in-game to generate instances */
 void PopulateReferenceSkeletalMeshesData(FMutableGraphGenerationContext& GenerationContext);
@@ -980,13 +980,13 @@ void CheckNumOutputs(const UEdGraphPin& Pin, const FMutableGraphGenerationContex
 UTexture2D* FindReferenceImage(const UEdGraphPin* Pin, FMutableGraphGenerationContext& GenerationContext);
 
 
-mu::Ptr<mu::NodeMeshApplyPose> CreateNodeMeshApplyPose(FMutableGraphGenerationContext& GenerationContext, mu::NodeMeshPtr InputMeshNode, const TArray<FName>& ArrayBoneName, const TArray<FTransform>& ArrayTransform);
+UE::Mutable::Private::Ptr<UE::Mutable::Private::NodeMeshApplyPose> CreateNodeMeshApplyPose(FMutableGraphGenerationContext& GenerationContext, UE::Mutable::Private::NodeMeshPtr InputMeshNode, const TArray<FName>& ArrayBoneName, const TArray<FTransform>& ArrayTransform);
 
 
 /** Adds Tag to MutableMesh uniquely, returns the index were the tag has been inserted or the index where an intance of the tag has been found */
-int32 AddTagToMutableMeshUnique(mu::FMesh& MutableMesh, const FString& Tag);
+int32 AddTagToMutableMeshUnique(UE::Mutable::Private::FMesh& MutableMesh, const FString& Tag);
 
-void AddSocketTagsToMesh(const USkeletalMesh* SourceMesh, TSharedPtr<mu::FMesh> MutableMesh, FMutableGraphGenerationContext& GenerationContext);
+void AddSocketTagsToMesh(const USkeletalMesh* SourceMesh, TSharedPtr<UE::Mutable::Private::FMesh> MutableMesh, FMutableGraphGenerationContext& GenerationContext);
 
 // Generates the tag for an animation instance
 FString GenerateAnimationInstanceTag(const int32 AnimInstanceIndex, const FName& SlotIndex);
@@ -1006,11 +1006,11 @@ int32 GetMaxTextureSize(const UTexture2D& ReferenceTexture, const UTextureLODSet
 // Max texture size of the texture with per platform MaxTextureSize and LODBias applied.
 int32 GetTextureSizeInGame(const UTexture2D& Texture, const UTextureLODSettings& LODSettings);
 
-TSharedPtr<mu::FImage> GenerateImageConstant(UTexture*, FMutableGraphGenerationContext&, bool bIsReference);
-TSharedPtr<mu::FMesh> GenerateMeshConstant(const FMutableSourceMeshData& Source, FMutableGraphGenerationContext&);
+TSharedPtr<UE::Mutable::Private::FImage> GenerateImageConstant(UTexture*, FMutableGraphGenerationContext&, bool bIsReference);
+TSharedPtr<UE::Mutable::Private::FMesh> GenerateMeshConstant(const FMutableSourceMeshData& Source, FMutableGraphGenerationContext&);
 
 /** Generates a mutable image descriptor from an unreal engine texture */
-mu::FImageDesc GenerateImageDescriptor(UTexture* Texture);
+UE::Mutable::Private::FImageDesc GenerateImageDescriptor(UTexture* Texture);
 
 /** 
  * Add SurfaceMetadata gathered form Material and MeshSection to HashSurfaceMetadataSet.
@@ -1021,4 +1021,4 @@ uint32 AddUniqueSurfaceMetadata(const FMutableSourceSurfaceMetadata& MeshData, F
 
 /** Convert a mesh to Mutable format.
 */
-TSharedPtr<mu::FMesh> ConvertSkeletalMeshToMutable(TSharedRef<FMeshConversionContext> MeshConversionContext);
+TSharedPtr<UE::Mutable::Private::FMesh> ConvertSkeletalMeshToMutable(TSharedRef<FMeshConversionContext> MeshConversionContext);
diff --git a/Engine/Plugins/Mutable/Source/CustomizableObjectEditor/Private/MuCOE/GenerateMutableSource/GenerateMutableSourceColor.cpp b/Engine/Plugins/Mutable/Source/CustomizableObjectEditor/Private/MuCOE/GenerateMutableSource/GenerateMutableSourceColor.cpp
--- a/Engine/Plugins/Mutable/Source/CustomizableObjectEditor/Private/MuCOE/GenerateMutableSource/GenerateMutableSourceColor.cpp
+++ b/Engine/Plugins/Mutable/Source/CustomizableObjectEditor/Private/MuCOE/GenerateMutableSource/GenerateMutableSourceColor.cpp
@@ -40,7 +40,7 @@
 #define LOCTEXT_NAMESPACE "CustomizableObjectEditor"
 
 
-mu::Ptr<mu::NodeColour> GenerateMutableSourceColor(const UEdGraphPin* Pin, FMutableGraphGenerationContext& GenerationContext)
+UE::Mutable::Private::Ptr<UE::Mutable::Private::NodeColour> GenerateMutableSourceColor(const UEdGraphPin* Pin, FMutableGraphGenerationContext& GenerationContext)
 {
 	check(Pin)
 	RETURN_ON_CYCLE(*Pin, GenerationContext)
@@ -54,7 +54,7 @@ mu::Ptr<mu::NodeColour> GenerateMutableSourceColor(const UEdGraphPin* Pin, FMuta
 	const FGeneratedKey Key(reinterpret_cast<void*>(&GenerateMutableSourceColor), *Pin, *Node, GenerationContext);
 	if (const FGeneratedData* Generated = GenerationContext.Generated.Find(Key))
 	{
-		return static_cast<mu::NodeColour*>(Generated->Node.get());
+		return static_cast<UE::Mutable::Private::NodeColour*>(Generated->Node.get());
 	}
 
 	if (Node->IsNodeOutDatedAndNeedsRefresh())
@@ -66,11 +66,11 @@ mu::Ptr<mu::NodeColour> GenerateMutableSourceColor(const UEdGraphPin* Pin, FMuta
 	// Most nodes need to be added to the cache but there are some that don't. For exampel, MacroInstanceNodes
 	bool bCacheNode = true;
 
-	mu::Ptr<mu::NodeColour> Result;
+	UE::Mutable::Private::Ptr<UE::Mutable::Private::NodeColour> Result;
 
 	if (const UCustomizableObjectNodeColorConstant* TypedNodeColorConst = Cast<UCustomizableObjectNodeColorConstant>(Node))
 	{
-		mu::Ptr<mu::NodeColourConstant> ColorNode = new mu::NodeColourConstant();
+		UE::Mutable::Private::Ptr<UE::Mutable::Private::NodeColourConstant> ColorNode = new UE::Mutable::Private::NodeColourConstant();
 		Result = ColorNode;
 
 		ColorNode->Value = TypedNodeColorConst->Value;
@@ -78,7 +78,7 @@ mu::Ptr<mu::NodeColour> GenerateMutableSourceColor(const UEdGraphPin* Pin, FMuta
 
 	else if (const UCustomizableObjectNodeColorParameter* TypedNodeColorParam = Cast<UCustomizableObjectNodeColorParameter>(Node))
 	{
-		mu::Ptr<mu::NodeColourParameter> ColorNode = new mu::NodeColourParameter();
+		UE::Mutable::Private::Ptr<UE::Mutable::Private::NodeColourParameter> ColorNode = new UE::Mutable::Private::NodeColourParameter();
 		Result = ColorNode;
 
 		ColorNode->Name = TypedNodeColorParam->GetParameterName(&GenerationContext.MacroNodesStack);
@@ -92,7 +92,7 @@ mu::Ptr<mu::NodeColour> GenerateMutableSourceColor(const UEdGraphPin* Pin, FMuta
 
 	else if (const UCustomizableObjectNodeColorSwitch* TypedNodeColorSwitch = Cast<UCustomizableObjectNodeColorSwitch>(Node))
 	{
-		Result = [&]() -> mu::Ptr<mu::NodeColourSwitch>
+		Result = [&]() -> UE::Mutable::Private::Ptr<UE::Mutable::Private::NodeColourSwitch>
 		{
 			if (const int32 NumParameters = FollowInputPinArray(*TypedNodeColorSwitch->SwitchParameter()).Num();
 				NumParameters != 1)
@@ -106,7 +106,7 @@ mu::Ptr<mu::NodeColour> GenerateMutableSourceColor(const UEdGraphPin* Pin, FMuta
 			}
 
 			const UEdGraphPin* EnumPin = FollowInputPin(*TypedNodeColorSwitch->SwitchParameter());
-			mu::Ptr<mu::NodeScalar> SwitchParam = EnumPin ? GenerateMutableSourceFloat(EnumPin, GenerationContext) : nullptr;
+			UE::Mutable::Private::Ptr<UE::Mutable::Private::NodeScalar> SwitchParam = EnumPin ? GenerateMutableSourceFloat(EnumPin, GenerationContext) : nullptr;
 
 			// Switch Param not generated
 			if (!SwitchParam)
@@ -121,7 +121,7 @@ mu::Ptr<mu::NodeColour> GenerateMutableSourceColor(const UEdGraphPin* Pin, FMuta
 				return nullptr;
 			}
 
-			if (SwitchParam->GetType() != mu::NodeScalarEnumParameter::GetStaticType())
+			if (SwitchParam->GetType() != UE::Mutable::Private::NodeScalarEnumParameter::GetStaticType())
 			{
 				const FText Message = LOCTEXT("WrongSwitchParamType", "Switch parameter of incorrect type.");
 				GenerationContext.Log(Message, Node);
@@ -131,14 +131,14 @@ mu::Ptr<mu::NodeColour> GenerateMutableSourceColor(const UEdGraphPin* Pin, FMuta
 
 			const int32 NumSwitchOptions = TypedNodeColorSwitch->GetNumElements();
 
-			mu::NodeScalarEnumParameter* EnumParameter = static_cast<mu::NodeScalarEnumParameter*>(SwitchParam.get());
+			UE::Mutable::Private::NodeScalarEnumParameter* EnumParameter = static_cast<UE::Mutable::Private::NodeScalarEnumParameter*>(SwitchParam.get());
 			if (NumSwitchOptions != EnumParameter->Options.Num())
 			{
 				const FText Message = LOCTEXT("MismatchedSwitch", "Switch enum and switch node have different number of options. Please refresh the switch node to make sure the outcomes are labeled properly.");
 				GenerationContext.Log(Message, Node);
 			}
 
-			mu::Ptr<mu::NodeColourSwitch> SwitchNode = new mu::NodeColourSwitch;
+			UE::Mutable::Private::Ptr<UE::Mutable::Private::NodeColourSwitch> SwitchNode = new UE::Mutable::Private::NodeColourSwitch;
 			SwitchNode->Parameter = SwitchParam;
 			SwitchNode->Options.SetNum(NumSwitchOptions);
 
@@ -157,31 +157,31 @@ mu::Ptr<mu::NodeColour> GenerateMutableSourceColor(const UEdGraphPin* Pin, FMuta
 
 	else if (const UCustomizableObjectNodeTextureSample* TypedNodeTexSample = Cast<UCustomizableObjectNodeTextureSample>(Node))
 	{
-		mu::Ptr<mu::NodeColourSampleImage> ColorNode = new mu::NodeColourSampleImage();
+		UE::Mutable::Private::Ptr<UE::Mutable::Private::NodeColourSampleImage> ColorNode = new UE::Mutable::Private::NodeColourSampleImage();
 		Result = ColorNode;
 
 		if (const UEdGraphPin* ConnectedPin = FollowInputPin(*TypedNodeTexSample->TexturePin()))
 		{
-			mu::NodeImagePtr TextureNode = GenerateMutableSourceImage(ConnectedPin, GenerationContext, 0);
+			UE::Mutable::Private::NodeImagePtr TextureNode = GenerateMutableSourceImage(ConnectedPin, GenerationContext, 0);
 			ColorNode->Image = TextureNode;
 		}
 
 		if (const UEdGraphPin* ConnectedPin = FollowInputPin(*TypedNodeTexSample->XPin()))
 		{
-			mu::Ptr<mu::NodeScalar> XNode = GenerateMutableSourceFloat(ConnectedPin, GenerationContext);
+			UE::Mutable::Private::Ptr<UE::Mutable::Private::NodeScalar> XNode = GenerateMutableSourceFloat(ConnectedPin, GenerationContext);
 			ColorNode->X = XNode;
 		}
 
 		if (const UEdGraphPin* ConnectedPin = FollowInputPin(*TypedNodeTexSample->YPin()))
 		{
-			mu::Ptr<mu::NodeScalar> YNode = GenerateMutableSourceFloat(ConnectedPin, GenerationContext);
+			UE::Mutable::Private::Ptr<UE::Mutable::Private::NodeScalar> YNode = GenerateMutableSourceFloat(ConnectedPin, GenerationContext);
 			ColorNode->Y = YNode;
 		}
 	}
 
 	else if (const UCustomizableObjectNodeColorArithmeticOp* TypedNodeColorArith = Cast<UCustomizableObjectNodeColorArithmeticOp>(Node))
 	{
-		mu::Ptr<mu::NodeColourArithmeticOperation> OpNode = new mu::NodeColourArithmeticOperation();
+		UE::Mutable::Private::Ptr<UE::Mutable::Private::NodeColourArithmeticOperation> OpNode = new UE::Mutable::Private::NodeColourArithmeticOperation();
 		Result = OpNode;
 
 		if (const UEdGraphPin* ConnectedPin = FollowInputPin(*TypedNodeColorArith->XPin()))
@@ -197,61 +197,61 @@ mu::Ptr<mu::NodeColour> GenerateMutableSourceColor(const UEdGraphPin* Pin, FMuta
 		switch (TypedNodeColorArith->Operation)
 		{
 		case EColorArithmeticOperation::E_Add:
-			OpNode->Operation = mu::NodeColourArithmeticOperation::EOperation::Add;
+			OpNode->Operation = UE::Mutable::Private::NodeColourArithmeticOperation::EOperation::Add;
 			break;
 
 		case EColorArithmeticOperation::E_Sub:
-			OpNode->Operation = mu::NodeColourArithmeticOperation::EOperation::Subtract;
+			OpNode->Operation = UE::Mutable::Private::NodeColourArithmeticOperation::EOperation::Subtract;
 			break;
 
 		case EColorArithmeticOperation::E_Mul:
-			OpNode->Operation = mu::NodeColourArithmeticOperation::EOperation::Multiply;
+			OpNode->Operation = UE::Mutable::Private::NodeColourArithmeticOperation::EOperation::Multiply;
 			break;
 
 		case EColorArithmeticOperation::E_Div:
-			OpNode->Operation = mu::NodeColourArithmeticOperation::EOperation::Divide;
+			OpNode->Operation = UE::Mutable::Private::NodeColourArithmeticOperation::EOperation::Divide;
 			break;
 		}
 	}
 
 	else if (const UCustomizableObjectNodeColorFromFloats* TypedNodeFrom = Cast<UCustomizableObjectNodeColorFromFloats>(Node))
 	{
-		mu::Ptr<mu::NodeColourFromScalars> OpNode = new mu::NodeColourFromScalars();
+		UE::Mutable::Private::Ptr<UE::Mutable::Private::NodeColourFromScalars> OpNode = new UE::Mutable::Private::NodeColourFromScalars();
 		Result = OpNode;
 
 		if (const UEdGraphPin* ConnectedPin = FollowInputPin(*TypedNodeFrom->RPin()))
 		{
-			mu::Ptr<mu::NodeScalar> FloatNode = GenerateMutableSourceFloat(ConnectedPin, GenerationContext);
+			UE::Mutable::Private::Ptr<UE::Mutable::Private::NodeScalar> FloatNode = GenerateMutableSourceFloat(ConnectedPin, GenerationContext);
 			OpNode->X = FloatNode;
 		}
 
 		if (const UEdGraphPin* ConnectedPin = FollowInputPin(*TypedNodeFrom->GPin()))
 		{
-			mu::Ptr<mu::NodeScalar> FloatNode = GenerateMutableSourceFloat(ConnectedPin, GenerationContext);
+			UE::Mutable::Private::Ptr<UE::Mutable::Private::NodeScalar> FloatNode = GenerateMutableSourceFloat(ConnectedPin, GenerationContext);
 			OpNode->Y = FloatNode;
 		}
 
 		if (const UEdGraphPin* ConnectedPin = FollowInputPin(*TypedNodeFrom->BPin()))
 		{
-			mu::Ptr<mu::NodeScalar> FloatNode = GenerateMutableSourceFloat(ConnectedPin, GenerationContext);
+			UE::Mutable::Private::Ptr<UE::Mutable::Private::NodeScalar> FloatNode = GenerateMutableSourceFloat(ConnectedPin, GenerationContext);
 			OpNode->Z = FloatNode;
 		}
 
 		if (const UEdGraphPin* ConnectedPin = FollowInputPin(*TypedNodeFrom->APin()))
 		{
-			mu::Ptr<mu::NodeScalar> FloatNode = GenerateMutableSourceFloat(ConnectedPin, GenerationContext);
+			UE::Mutable::Private::Ptr<UE::Mutable::Private::NodeScalar> FloatNode = GenerateMutableSourceFloat(ConnectedPin, GenerationContext);
 			OpNode->W = FloatNode;
 		}
 	}
 
 	else if (const UCustomizableObjectNodeColorVariation* TypedNodeColorVar = Cast<const UCustomizableObjectNodeColorVariation>(Node))
 	{
-		mu::Ptr<mu::NodeColourVariation> ColorNode = new mu::NodeColourVariation();
+		UE::Mutable::Private::Ptr<UE::Mutable::Private::NodeColourVariation> ColorNode = new UE::Mutable::Private::NodeColourVariation();
 		Result = ColorNode;
 
 		if (const UEdGraphPin* ConnectedPin = FollowInputPin(*TypedNodeColorVar->DefaultPin()))
 		{
-			mu::Ptr<mu::NodeColour> ChildNode = GenerateMutableSourceColor(ConnectedPin, GenerationContext);
+			UE::Mutable::Private::Ptr<UE::Mutable::Private::NodeColour> ChildNode = GenerateMutableSourceColor(ConnectedPin, GenerationContext);
 			if (ChildNode)
 			{
 				ColorNode->DefaultColour = ChildNode;
@@ -273,7 +273,7 @@ mu::Ptr<mu::NodeColour> GenerateMutableSourceColor(const UEdGraphPin* Pin, FMuta
 
 			if (const UEdGraphPin* ConnectedPin = FollowInputPin(*VariationPin))
 			{
-				mu::Ptr<mu::NodeColour> ChildNode = GenerateMutableSourceColor(ConnectedPin, GenerationContext);
+				UE::Mutable::Private::Ptr<UE::Mutable::Private::NodeColour> ChildNode = GenerateMutableSourceColor(ConnectedPin, GenerationContext);
 				ColorNode->Variations[VariationIndex].Colour = ChildNode;
 			}
 		}
@@ -281,7 +281,7 @@ mu::Ptr<mu::NodeColour> GenerateMutableSourceColor(const UEdGraphPin* Pin, FMuta
 
 	else if (const UCustomizableObjectNodeColorToSRGB* TypedNodeColorToSRGB = Cast<UCustomizableObjectNodeColorToSRGB>(Node))
 	{
-		mu::Ptr<mu::NodeColorToSRGB> ColorNode = new mu::NodeColorToSRGB();
+		UE::Mutable::Private::Ptr<UE::Mutable::Private::NodeColorToSRGB> ColorNode = new UE::Mutable::Private::NodeColorToSRGB();
 		Result = ColorNode;
 
 		if (const UEdGraphPin* ConnectedPin = FollowInputPin(*TypedNodeColorToSRGB->GetInputPin()))
@@ -293,7 +293,7 @@ mu::Ptr<mu::NodeColour> GenerateMutableSourceColor(const UEdGraphPin* Pin, FMuta
 	else if (const UCustomizableObjectNodeTable* TypedNodeTable = Cast<UCustomizableObjectNodeTable>(Node))
 	{
 		//This node will add a white color in case of error
-		mu::Ptr<mu::NodeColourConstant> WhiteColorNode = new mu::NodeColourConstant();
+		UE::Mutable::Private::Ptr<UE::Mutable::Private::NodeColourConstant> WhiteColorNode = new UE::Mutable::Private::NodeColourConstant();
 		WhiteColorNode->Value = FVector4f(1.0f, 1.0f, 1.0f, 1.0f);
 
 		Result = WhiteColorNode;
@@ -317,12 +317,12 @@ mu::Ptr<mu::NodeColour> GenerateMutableSourceColor(const UEdGraphPin* Pin, FMuta
 			if (bSuccess)
 			{
 				// Generating a new data table if not exists
-				mu::Ptr<mu::FTable> Table;
+				UE::Mutable::Private::Ptr<UE::Mutable::Private::FTable> Table;
 				Table = GenerateMutableSourceTable(DataTable, TypedNodeTable, GenerationContext);
 
 				if (Table)
 				{
-					mu::Ptr<mu::NodeColourTable> ColorTableNode = new mu::NodeColourTable();
+					UE::Mutable::Private::Ptr<UE::Mutable::Private::NodeColourTable> ColorTableNode = new UE::Mutable::Private::NodeColourTable();
 
 					// Generating a new Color column if not exists
 					if (Table->FindColumn(ColumnName) == INDEX_NONE)
@@ -366,24 +366,24 @@ mu::Ptr<mu::NodeColour> GenerateMutableSourceColor(const UEdGraphPin* Pin, FMuta
 	{
 		bCacheNode = false;
 
-		Result = GenerateMutableSourceMacro<mu::NodeColour>(*Pin, GenerationContext, GenerateMutableSourceColor);
+		Result = GenerateMutableSourceMacro<UE::Mutable::Private::NodeColour>(*Pin, GenerationContext, GenerateMutableSourceColor);
 	}
 
 	else if (const UCustomizableObjectNodeTunnel* TypedNodeTunnel = Cast<UCustomizableObjectNodeTunnel>(Node))
 	{
 		bCacheNode = false;
 
-		Result = GenerateMutableSourceMacro<mu::NodeColour>(*Pin, GenerationContext, GenerateMutableSourceColor);
+		Result = GenerateMutableSourceMacro<UE::Mutable::Private::NodeColour>(*Pin, GenerationContext, GenerateMutableSourceColor);
 	}
 
 	else if (const UCONodeMaterialBreak* TypedNodeMaterialBreak = Cast<UCONodeMaterialBreak>(Node))
 	{
-		Result = new mu::NodeColourConstant();
+		Result = new UE::Mutable::Private::NodeColourConstant();
 		check(TypedNodeMaterialBreak->MaterialPinRef.Get());
 
 		if (const UEdGraphPin* LinkedMaterial = FollowInputPin(*TypedNodeMaterialBreak->MaterialPinRef.Get()))
 		{
-			mu::Ptr<mu::NodeColourMaterialBreak> MaterialBreak = new mu::NodeColourMaterialBreak();
+			UE::Mutable::Private::Ptr<UE::Mutable::Private::NodeColourMaterialBreak> MaterialBreak = new UE::Mutable::Private::NodeColourMaterialBreak();
 
 			FName ParameterName = TypedNodeMaterialBreak->GetPinParameterName(*Pin);
 			check(!ParameterName.IsNone());
@@ -393,7 +393,7 @@ mu::Ptr<mu::NodeColour> GenerateMutableSourceColor(const UEdGraphPin* Pin, FMuta
 
 			// Generate the material that is going to be broken into parameters
 			GenerationContext.CurrentMaterialBreakParameter = FMutableGraphGenerationContext::FMaterialBreakParameter{ ParameterName, EMaterialParameterType::Vector };
-			mu::Ptr<mu::NodeMaterial> LinkedNode = GenerateMutableSourceMaterial(LinkedMaterial, GenerationContext);
+			UE::Mutable::Private::Ptr<UE::Mutable::Private::NodeMaterial> LinkedNode = GenerateMutableSourceMaterial(LinkedMaterial, GenerationContext);
 			MaterialBreak->MaterialSource = LinkedNode;
 			GenerationContext.CurrentMaterialBreakParameter = FMutableGraphGenerationContext::FMaterialBreakParameter{ };
 
diff --git a/Engine/Plugins/Mutable/Source/CustomizableObjectEditor/Private/MuCOE/GenerateMutableSource/GenerateMutableSourceColor.h b/Engine/Plugins/Mutable/Source/CustomizableObjectEditor/Private/MuCOE/GenerateMutableSource/GenerateMutableSourceColor.h
--- a/Engine/Plugins/Mutable/Source/CustomizableObjectEditor/Private/MuCOE/GenerateMutableSource/GenerateMutableSourceColor.h
+++ b/Engine/Plugins/Mutable/Source/CustomizableObjectEditor/Private/MuCOE/GenerateMutableSource/GenerateMutableSourceColor.h
@@ -8,4 +8,4 @@ class UEdGraphPin;
 struct FMutableGraphGenerationContext;
 
 /** Convert a CustomizableObject Source Graph into a mutable source graph. */
-mu::Ptr<mu::NodeColour> GenerateMutableSourceColor(const UEdGraphPin* Pin, FMutableGraphGenerationContext& GenerationContext);
+UE::Mutable::Private::Ptr<UE::Mutable::Private::NodeColour> GenerateMutableSourceColor(const UEdGraphPin* Pin, FMutableGraphGenerationContext& GenerationContext);
diff --git a/Engine/Plugins/Mutable/Source/CustomizableObjectEditor/Private/MuCOE/GenerateMutableSource/GenerateMutableSourceComponent.cpp b/Engine/Plugins/Mutable/Source/CustomizableObjectEditor/Private/MuCOE/GenerateMutableSource/GenerateMutableSourceComponent.cpp
--- a/Engine/Plugins/Mutable/Source/CustomizableObjectEditor/Private/MuCOE/GenerateMutableSource/GenerateMutableSourceComponent.cpp
+++ b/Engine/Plugins/Mutable/Source/CustomizableObjectEditor/Private/MuCOE/GenerateMutableSource/GenerateMutableSourceComponent.cpp
@@ -42,7 +42,7 @@
 /** Generate LOD pins of the given NodeComponentBase (NodeComponent, NodeComponentExtend...).
  * @param TypedComponentMesh Given component node.
  * @param NodeComponent Core node to connect LOD generated pins. */
-void GenerateMutableSourceComponentMesh(FMutableGraphGenerationContext& GenerationContext, const ICustomizableObjectNodeComponentMeshInterface& TypedComponentMesh, mu::Ptr<mu::NodeComponent> NodeComponent)
+void GenerateMutableSourceComponentMesh(FMutableGraphGenerationContext& GenerationContext, const ICustomizableObjectNodeComponentMeshInterface& TypedComponentMesh, UE::Mutable::Private::Ptr<UE::Mutable::Private::NodeComponent> NodeComponent)
 {
 	int32 FirstLOD = -1;
 
@@ -53,10 +53,10 @@ void GenerateMutableSourceComponentMesh(FMutableGraphGenerationContext& Generati
 
 		if (!NodeComponent->LODs.IsValidIndex(CurrentLOD))
 		{
-			NodeComponent->LODs.Add(new mu::NodeLOD());
+			NodeComponent->LODs.Add(new UE::Mutable::Private::NodeLOD());
 		}
 
-		mu::Ptr<mu::NodeLOD> LODNode = NodeComponent->LODs[CurrentLOD];
+		UE::Mutable::Private::Ptr<UE::Mutable::Private::NodeLOD> LODNode = NodeComponent->LODs[CurrentLOD];
 
 		LODNode->SetMessageContext(&TypedComponentMesh);
 
@@ -101,7 +101,7 @@ void GenerateMutableSourceComponentMesh(FMutableGraphGenerationContext& Generati
 					continue;
 				}
 				
-				mu::Ptr<mu::NodeSurface> SurfaceNode = GenerateMutableSourceSurface(ChildNodePin, GenerationContext);
+				UE::Mutable::Private::Ptr<UE::Mutable::Private::NodeSurface> SurfaceNode = GenerateMutableSourceSurface(ChildNodePin, GenerationContext);
 				LODNode->Surfaces.Add(SurfaceNode);
 			}
 		}
@@ -113,7 +113,7 @@ void GenerateMutableSourceComponentMesh(FMutableGraphGenerationContext& Generati
 }
 
 
-mu::Ptr<mu::NodeComponent> GenerateMutableSourceComponent(const UEdGraphPin* Pin, FMutableGraphGenerationContext& GenerationContext)
+UE::Mutable::Private::Ptr<UE::Mutable::Private::NodeComponent> GenerateMutableSourceComponent(const UEdGraphPin* Pin, FMutableGraphGenerationContext& GenerationContext)
 {
 	check(Pin)
 	RETURN_ON_CYCLE(*Pin, GenerationContext)
@@ -125,14 +125,14 @@ mu::Ptr<mu::NodeComponent> GenerateMutableSourceComponent(const UEdGraphPin* Pin
 	const FGeneratedKey Key(reinterpret_cast<void*>(&GenerateMutableSourceComponent), *Pin, *Node, GenerationContext, false);
 	if (const FGeneratedData* Generated = GenerationContext.Generated.Find(Key))
 	{
-		return static_cast<mu::NodeComponent*>(Generated->Node.get());
+		return static_cast<UE::Mutable::Private::NodeComponent*>(Generated->Node.get());
 	}
 
 	// Bool that determines if a node can be added to the cache of nodes.
 	// Most nodes need to be added to the cache but there are some that don't. For example, MacroInstanceNodes
 	bool bCacheNode = true;
 	
-	mu::Ptr<mu::NodeComponent> Result;
+	UE::Mutable::Private::Ptr<UE::Mutable::Private::NodeComponent> Result;
 	
 	if (const UCustomizableObjectNodeComponentMesh* TypedComponentMesh = Cast<UCustomizableObjectNodeComponentMesh>(Node))
 	{
@@ -144,7 +144,7 @@ mu::Ptr<mu::NodeComponent> GenerateMutableSourceComponent(const UEdGraphPin* Pin
 			return nullptr; // Not generated in the first pass.
 		}
 		
-		mu::Ptr<mu::NodeComponentNew> NodeComponentNew = new mu::NodeComponentNew();
+		UE::Mutable::Private::Ptr<UE::Mutable::Private::NodeComponentNew> NodeComponentNew = new UE::Mutable::Private::NodeComponentNew();
 		const FName& ComponentName = TypedComponentMesh->GetComponentName(&GenerationContext.MacroNodesStack);
 
 		NodeComponentNew->Id = GenerationContext.ComponentNames.Find(ComponentName);
@@ -154,12 +154,12 @@ mu::Ptr<mu::NodeComponent> GenerateMutableSourceComponent(const UEdGraphPin* Pin
 		if (const UEdGraphPin* ConnectedPin = MaterialAssetPin ? FollowInputPin(*MaterialAssetPin) : nullptr)
 		{
 			GenerationContext.CurrentMaterialParameterId = ConnectedPin->PinId.ToString();
-			mu::Ptr<mu::NodeMaterial> MaterialNode = GenerateMutableSourceMaterial(ConnectedPin, GenerationContext);
+			UE::Mutable::Private::Ptr<UE::Mutable::Private::NodeMaterial> MaterialNode = GenerateMutableSourceMaterial(ConnectedPin, GenerationContext);
 			NodeComponentNew->OverlayMaterial = MaterialNode;
 		}
 		else if(UMaterialInterface* OverlayMaterial = GenerationContext.LoadObject(TypedComponentMesh->GetOverlayMaterial()))
 		{
-			mu::Ptr<mu::NodeMaterialConstant> MaterialNode = new mu::NodeMaterialConstant();
+			UE::Mutable::Private::Ptr<UE::Mutable::Private::NodeMaterialConstant> MaterialNode = new UE::Mutable::Private::NodeMaterialConstant();
 			MaterialNode->MaterialId = GenerationContext.ReferencedMaterials.AddUnique(OverlayMaterial);
 			NodeComponentNew->OverlayMaterial = MaterialNode;
 		}
@@ -196,7 +196,7 @@ mu::Ptr<mu::NodeComponent> GenerateMutableSourceComponent(const UEdGraphPin* Pin
 			TArray<const UCustomizableObjectNodeMacroInstance*> MacroContextCopy = GenerationContext.MacroNodesStack;
 			GenerationContext.MacroNodesStack = FindResult->MacroContext;
 
-			mu::Ptr<mu::NodeComponent> ParentNodeComponent = GenerateMutableSourceComponent(TypedParentComponentMesh->OutputPin.Get(), GenerationContext);
+			UE::Mutable::Private::Ptr<UE::Mutable::Private::NodeComponent> ParentNodeComponent = GenerateMutableSourceComponent(TypedParentComponentMesh->OutputPin.Get(), GenerationContext);
 
 			// Get the parent name using the parent's macro context
 			GenerationContext.CurrentMeshComponent = TypedParentComponentMesh->GetComponentName(&GenerationContext.MacroNodesStack);
@@ -204,7 +204,7 @@ mu::Ptr<mu::NodeComponent> GenerateMutableSourceComponent(const UEdGraphPin* Pin
 			// Restore GenerationContext MacroContex
 			GenerationContext.MacroNodesStack = MacroContextCopy;
 
-			mu::Ptr<mu::NodeComponentEdit> NodeComponentEdit = new mu::NodeComponentEdit();
+			UE::Mutable::Private::Ptr<UE::Mutable::Private::NodeComponentEdit> NodeComponentEdit = new UE::Mutable::Private::NodeComponentEdit();
 			NodeComponentEdit->Parent = ParentNodeComponent.get();
 			NodeComponentEdit->SetMessageContext(TypedComponentMeshExtend);
 		
@@ -254,20 +254,20 @@ mu::Ptr<mu::NodeComponent> GenerateMutableSourceComponent(const UEdGraphPin* Pin
 		}
 
 		// Create the referenced mesh node.
-		mu::Ptr<mu::NodeMeshConstant> MeshNode;
+		UE::Mutable::Private::Ptr<UE::Mutable::Private::NodeMeshConstant> MeshNode;
 		{
-			MeshNode = new mu::NodeMeshConstant();
+			MeshNode = new UE::Mutable::Private::NodeMeshConstant();
 
 			FMutableSourceMeshData Source;
 			Source.Mesh = SkeletalMesh;
 			Source.bIsPassthrough = true;
-			TSharedPtr<mu::FMesh> MutableMesh = GenerateMutableSkeletalMesh(Source, 0, 0, 0, 0, {}, GenerationContext, Node);
+			TSharedPtr<UE::Mutable::Private::FMesh> MutableMesh = GenerateMutableSkeletalMesh(Source, 0, 0, 0, 0, {}, GenerationContext, Node);
 
 			MeshNode->Value = MutableMesh;
 		}
 
 		// Create the component node
-		mu::Ptr<mu::NodeComponentNew> ComponentNode = new mu::NodeComponentNew;
+		UE::Mutable::Private::Ptr<UE::Mutable::Private::NodeComponentNew> ComponentNode = new UE::Mutable::Private::NodeComponentNew;
 		if (GenerationContext.ComponentNames.Contains(ComponentName))
 		{
 			FString Msg = FString::Printf(TEXT("More than one component with the same name [%s] have been found. This is not supported."), *ComponentName.ToString());
@@ -280,7 +280,7 @@ mu::Ptr<mu::NodeComponent> GenerateMutableSourceComponent(const UEdGraphPin* Pin
 		int32 SkeletalMeshLODCount = Model->LODModels.Num();
 		for (int32 LODIndex=0; LODIndex<SkeletalMeshLODCount; ++LODIndex)
 		{
-			mu::Ptr<mu::NodeLOD> LODNode = new mu::NodeLOD;
+			UE::Mutable::Private::Ptr<UE::Mutable::Private::NodeLOD> LODNode = new UE::Mutable::Private::NodeLOD;
 			ComponentNode->LODs.Add(LODNode);
 
 			const FSkeletalMeshLODModel& LODModel = Model->LODModels[LODIndex];
@@ -294,7 +294,7 @@ mu::Ptr<mu::NodeComponent> GenerateMutableSourceComponent(const UEdGraphPin* Pin
 					{
 						GenerationContext.ComponentMeshOverride = MeshNode;
 					
-						mu::Ptr<mu::NodeSurface> SurfaceNode = GenerateMutableSourceSurface(ConnectedMaterialPin, GenerationContext);
+						UE::Mutable::Private::Ptr<UE::Mutable::Private::NodeSurface> SurfaceNode = GenerateMutableSourceSurface(ConnectedMaterialPin, GenerationContext);
 						LODNode->Surfaces.Add(SurfaceNode);
 
 						GenerationContext.ComponentMeshOverride = nullptr;
@@ -302,7 +302,7 @@ mu::Ptr<mu::NodeComponent> GenerateMutableSourceComponent(const UEdGraphPin* Pin
 					else
 					{
 						// Add an empty surface node anyway.
-						mu::Ptr<mu::NodeSurfaceNew> SurfaceNode = new mu::NodeSurfaceNew;
+						UE::Mutable::Private::Ptr<UE::Mutable::Private::NodeSurfaceNew> SurfaceNode = new UE::Mutable::Private::NodeSurfaceNew;
 						SurfaceNode->Mesh = MeshNode;
 						SurfaceNode->SurfaceGuid = GenerationContext.GetNodeIdUnique(TypedComponentPassthroughMesh);
 						LODNode->Surfaces.Add(SurfaceNode);
@@ -314,7 +314,7 @@ mu::Ptr<mu::NodeComponent> GenerateMutableSourceComponent(const UEdGraphPin* Pin
 					// This may happen if the node was saved without pins but with a SKM set (old bug)
 					
 					// Add an empty surface node anyway.
-					mu::Ptr<mu::NodeSurfaceNew> SurfaceNode = new mu::NodeSurfaceNew;
+					UE::Mutable::Private::Ptr<UE::Mutable::Private::NodeSurfaceNew> SurfaceNode = new UE::Mutable::Private::NodeSurfaceNew;
 					SurfaceNode->Mesh = MeshNode;
 					SurfaceNode->SurfaceGuid = GenerationContext.GetNodeIdUnique(TypedComponentPassthroughMesh);
 					LODNode->Surfaces.Add(SurfaceNode);
@@ -336,7 +336,7 @@ mu::Ptr<mu::NodeComponent> GenerateMutableSourceComponent(const UEdGraphPin* Pin
 				// Check Switch Parameter arity preconditions.
 				if (const UEdGraphPin* EnumPin = FollowInputPin(*SwitchParameter))
 				{
-					mu::Ptr<mu::NodeScalar> SwitchParam = GenerateMutableSourceFloat(EnumPin, GenerationContext);
+					UE::Mutable::Private::Ptr<UE::Mutable::Private::NodeScalar> SwitchParam = GenerateMutableSourceFloat(EnumPin, GenerationContext);
 
 					// Switch Param not generated
 					if (!SwitchParam)
@@ -351,7 +351,7 @@ mu::Ptr<mu::NodeComponent> GenerateMutableSourceComponent(const UEdGraphPin* Pin
 						return Result;
 					}
 
-					if (SwitchParam->GetType() != mu::NodeScalarEnumParameter::GetStaticType())
+					if (SwitchParam->GetType() != UE::Mutable::Private::NodeScalarEnumParameter::GetStaticType())
 					{
 						const FText Message = LOCTEXT("WrongSwitchParamType", "Switch parameter of incorrect type.");
 						GenerationContext.Log(Message, Node);
@@ -361,14 +361,14 @@ mu::Ptr<mu::NodeComponent> GenerateMutableSourceComponent(const UEdGraphPin* Pin
 
 					const int32 NumSwitchOptions = TypedNodeSwitch->GetNumElements();
 
-					mu::NodeScalarEnumParameter* EnumParameter = static_cast<mu::NodeScalarEnumParameter*>(SwitchParam.get());
+					UE::Mutable::Private::NodeScalarEnumParameter* EnumParameter = static_cast<UE::Mutable::Private::NodeScalarEnumParameter*>(SwitchParam.get());
 					if (NumSwitchOptions != EnumParameter->Options.Num())
 					{
 						const FText Message = LOCTEXT("MismatchedSwitch", "Switch enum and switch node have different number of options. Please refresh the switch node to make sure the outcomes are labeled properly.");
 						GenerationContext.Log(Message, Node);
 					}
 
-					mu::Ptr<mu::NodeComponentSwitch> SwitchNode = new mu::NodeComponentSwitch;
+					UE::Mutable::Private::Ptr<UE::Mutable::Private::NodeComponentSwitch> SwitchNode = new UE::Mutable::Private::NodeComponentSwitch;
 					SwitchNode->Parameter = SwitchParam;
 					SwitchNode->Options.SetNum(NumSwitchOptions);
 
@@ -376,7 +376,7 @@ mu::Ptr<mu::NodeComponent> GenerateMutableSourceComponent(const UEdGraphPin* Pin
 					{
 						if (const UEdGraphPin* ConnectedPin = FollowInputPin(*TypedNodeSwitch->GetElementPin(SelectorIndex)))
 						{
-							mu::Ptr<mu::NodeComponent> ChildNode = GenerateMutableSourceComponent(ConnectedPin, GenerationContext);
+							UE::Mutable::Private::Ptr<UE::Mutable::Private::NodeComponent> ChildNode = GenerateMutableSourceComponent(ConnectedPin, GenerationContext);
 							if (ChildNode)
 							{
 								SwitchNode->Options[SelectorIndex] = ChildNode;
@@ -401,12 +401,12 @@ mu::Ptr<mu::NodeComponent> GenerateMutableSourceComponent(const UEdGraphPin* Pin
 
 	else if (const UCustomizableObjectNodeComponentVariation* TypedNodeVar = Cast<UCustomizableObjectNodeComponentVariation>(Node))
 	{
-		mu::Ptr<mu::NodeComponentVariation> SurfNode = new mu::NodeComponentVariation();
+		UE::Mutable::Private::Ptr<UE::Mutable::Private::NodeComponentVariation> SurfNode = new UE::Mutable::Private::NodeComponentVariation();
 		Result = SurfNode;
 
 		for (const UEdGraphPin* ConnectedPin : FollowInputPinArray(*TypedNodeVar->DefaultPin()))
 		{
-			mu::Ptr<mu::NodeComponent> ChildNode = GenerateMutableSourceComponent(ConnectedPin, GenerationContext);
+			UE::Mutable::Private::Ptr<UE::Mutable::Private::NodeComponent> ChildNode = GenerateMutableSourceComponent(ConnectedPin, GenerationContext);
 			if (ChildNode)
 			{
 				SurfNode->DefaultComponent = ChildNode;
@@ -421,7 +421,7 @@ mu::Ptr<mu::NodeComponent> GenerateMutableSourceComponent(const UEdGraphPin* Pin
 		SurfNode->Variations.SetNum(NumVariations);
 		for (int VariationIndex = 0; VariationIndex < NumVariations; ++VariationIndex)
 		{
-			mu::NodeSurfacePtr VariationSurfaceNode;
+			UE::Mutable::Private::NodeSurfacePtr VariationSurfaceNode;
 
 			if (UEdGraphPin* VariationPin = TypedNodeVar->VariationPin(VariationIndex))
 			{
@@ -429,7 +429,7 @@ mu::Ptr<mu::NodeComponent> GenerateMutableSourceComponent(const UEdGraphPin* Pin
 				for (const UEdGraphPin* ConnectedPin : FollowInputPinArray(*VariationPin))
 				{
 					// Is it a modifier?
-					mu::Ptr<mu::NodeComponent> ChildNode = GenerateMutableSourceComponent(ConnectedPin, GenerationContext);
+					UE::Mutable::Private::Ptr<UE::Mutable::Private::NodeComponent> ChildNode = GenerateMutableSourceComponent(ConnectedPin, GenerationContext);
 					if (ChildNode)
 					{
 						SurfNode->Variations[VariationIndex].Component = ChildNode;
@@ -447,13 +447,13 @@ mu::Ptr<mu::NodeComponent> GenerateMutableSourceComponent(const UEdGraphPin* Pin
 	{
 		bCacheNode = false;
 
-		Result = GenerateMutableSourceMacro<mu::NodeComponent>(*Pin, GenerationContext, GenerateMutableSourceComponent);
+		Result = GenerateMutableSourceMacro<UE::Mutable::Private::NodeComponent>(*Pin, GenerationContext, GenerateMutableSourceComponent);
 	}
 
 	else if (const UCustomizableObjectNodeTunnel* TypedNodeTunnel = Cast<UCustomizableObjectNodeTunnel>(Node))
 	{
 		bCacheNode = false;
-		Result = GenerateMutableSourceMacro<mu::NodeComponent>(*Pin, GenerationContext, GenerateMutableSourceComponent);
+		Result = GenerateMutableSourceMacro<UE::Mutable::Private::NodeComponent>(*Pin, GenerationContext, GenerateMutableSourceComponent);
 	}
 
 	else
diff --git a/Engine/Plugins/Mutable/Source/CustomizableObjectEditor/Private/MuCOE/GenerateMutableSource/GenerateMutableSourceComponent.h b/Engine/Plugins/Mutable/Source/CustomizableObjectEditor/Private/MuCOE/GenerateMutableSource/GenerateMutableSourceComponent.h
--- a/Engine/Plugins/Mutable/Source/CustomizableObjectEditor/Private/MuCOE/GenerateMutableSource/GenerateMutableSourceComponent.h
+++ b/Engine/Plugins/Mutable/Source/CustomizableObjectEditor/Private/MuCOE/GenerateMutableSource/GenerateMutableSourceComponent.h
@@ -8,13 +8,13 @@
 class UEdGraphPin;
 struct FMutableGraphGenerationContext;
 
-namespace mu
+namespace UE::Mutable::Private
 {
 	class NodeComponent;
 }
 
 
-mu::Ptr<mu::NodeComponent> GenerateMutableSourceComponent(const UEdGraphPin* Pin, FMutableGraphGenerationContext& GenerationContext);
+UE::Mutable::Private::Ptr<UE::Mutable::Private::NodeComponent> GenerateMutableSourceComponent(const UEdGraphPin* Pin, FMutableGraphGenerationContext& GenerationContext);
 
 
 void FirstPass(UCustomizableObjectNodeComponentMesh& Node, FMutableGraphGenerationContext& GenerationContext);
