[Mutable] Change namespace name

diff --git a/Engine/Plugins/Mutable/Source/CustomizableObjectEditor/Private/MuCOE/GenerateMutableSource/GenerateMutableSourceFloat.cpp b/Engine/Plugins/Mutable/Source/CustomizableObjectEditor/Private/MuCOE/GenerateMutableSource/GenerateMutableSourceFloat.cpp
--- a/Engine/Plugins/Mutable/Source/CustomizableObjectEditor/Private/MuCOE/GenerateMutableSource/GenerateMutableSourceFloat.cpp
+++ b/Engine/Plugins/Mutable/Source/CustomizableObjectEditor/Private/MuCOE/GenerateMutableSource/GenerateMutableSourceFloat.cpp
@@ -35,7 +35,7 @@
 #define LOCTEXT_NAMESPACE "CustomizableObjectEditor"
 
 
-mu::Ptr<mu::NodeScalar> GenerateMutableSourceFloat(const UEdGraphPin* Pin, FMutableGraphGenerationContext& GenerationContext)
+UE::Mutable::Private::Ptr<UE::Mutable::Private::NodeScalar> GenerateMutableSourceFloat(const UEdGraphPin* Pin, FMutableGraphGenerationContext& GenerationContext)
 {
 	check(Pin)
 	RETURN_ON_CYCLE(*Pin, GenerationContext)
@@ -50,7 +50,7 @@ mu::Ptr<mu::NodeScalar> GenerateMutableSourceFloat(const UEdGraphPin* Pin, FMuta
 	const FGeneratedKey Key(reinterpret_cast<void*>(&GenerateMutableSourceFloat), *Pin, *Node, GenerationContext);
 	if (const FGeneratedData* Generated = GenerationContext.Generated.Find(Key))
 	{
-		return static_cast<mu::NodeScalar*>(Generated->Node.get());
+		return static_cast<UE::Mutable::Private::NodeScalar*>(Generated->Node.get());
 	}
 
 	if (Node->IsNodeOutDatedAndNeedsRefresh())
@@ -62,11 +62,11 @@ mu::Ptr<mu::NodeScalar> GenerateMutableSourceFloat(const UEdGraphPin* Pin, FMuta
 	// Most nodes need to be added to the cache but there are some that don't. For exampel, MacroInstanceNodes
 	bool bCacheNode = true;
 
-	mu::Ptr<mu::NodeScalar> Result;
+	UE::Mutable::Private::Ptr<UE::Mutable::Private::NodeScalar> Result;
 	
 	if (const UCustomizableObjectNodeFloatConstant* FloatConstantNode = Cast<UCustomizableObjectNodeFloatConstant>(Node))
 	{
-		mu::Ptr<mu::NodeScalarConstant> ScalarNode = new mu::NodeScalarConstant();
+		UE::Mutable::Private::Ptr<UE::Mutable::Private::NodeScalarConstant> ScalarNode = new UE::Mutable::Private::NodeScalarConstant();
 		Result = ScalarNode;
 
 		ScalarNode->Value = FloatConstantNode->Value;
@@ -74,7 +74,7 @@ mu::Ptr<mu::NodeScalar> GenerateMutableSourceFloat(const UEdGraphPin* Pin, FMuta
 
 	else if (const UCustomizableObjectNodeFloatParameter* FloatParameterNode = Cast<UCustomizableObjectNodeFloatParameter>(Node))
 	{
-		mu::Ptr<mu::NodeScalarParameter> ScalarNode = new mu::NodeScalarParameter();
+		UE::Mutable::Private::Ptr<UE::Mutable::Private::NodeScalarParameter> ScalarNode = new UE::Mutable::Private::NodeScalarParameter();
 		Result = ScalarNode;
 
 		ScalarNode->Name = FloatParameterNode->GetParameterName(&GenerationContext.MacroNodesStack);
@@ -88,7 +88,7 @@ mu::Ptr<mu::NodeScalar> GenerateMutableSourceFloat(const UEdGraphPin* Pin, FMuta
 
 	else if (const UCustomizableObjectNodeEnumParameter* EnumParamNode = Cast<UCustomizableObjectNodeEnumParameter>(Node))
 	{
-		mu::Ptr<mu::NodeScalarEnumParameter> EnumParameterNode = new mu::NodeScalarEnumParameter;
+		UE::Mutable::Private::Ptr<UE::Mutable::Private::NodeScalarEnumParameter> EnumParameterNode = new UE::Mutable::Private::NodeScalarEnumParameter;
 
 		const int32 NumSelectors = EnumParamNode->Values.Num();
 
@@ -120,7 +120,7 @@ mu::Ptr<mu::NodeScalar> GenerateMutableSourceFloat(const UEdGraphPin* Pin, FMuta
 	else if (const UCustomizableObjectNodeFloatSwitch* TypedNodeFloatSwitch = Cast<UCustomizableObjectNodeFloatSwitch>(Node))
 	{
 		// Using a lambda so control flow is easier to manage.
-		Result = [&]() -> mu::Ptr<mu::NodeScalar>
+		Result = [&]() -> UE::Mutable::Private::Ptr<UE::Mutable::Private::NodeScalar>
 		{
 			const UEdGraphPin* SwitchParameter = TypedNodeFloatSwitch->SwitchParameter();
 
@@ -137,7 +137,7 @@ mu::Ptr<mu::NodeScalar> GenerateMutableSourceFloat(const UEdGraphPin* Pin, FMuta
 			}
 
 			const UEdGraphPin* EnumPin = FollowInputPin(*SwitchParameter);
-			mu::Ptr<mu::NodeScalar> SwitchParam = GenerateMutableSourceFloat(EnumPin, GenerationContext);
+			UE::Mutable::Private::Ptr<UE::Mutable::Private::NodeScalar> SwitchParam = GenerateMutableSourceFloat(EnumPin, GenerationContext);
 
 			// Switch Param not generated
 			if (!SwitchParam)
@@ -152,7 +152,7 @@ mu::Ptr<mu::NodeScalar> GenerateMutableSourceFloat(const UEdGraphPin* Pin, FMuta
 				return nullptr;
 			}
 
-			if (SwitchParam->GetType() != mu::NodeScalarEnumParameter::GetStaticType())
+			if (SwitchParam->GetType() != UE::Mutable::Private::NodeScalarEnumParameter::GetStaticType())
 			{
 				const FText Message = LOCTEXT("WrongSwitchParamType", "Switch parameter of incorrect type.");
 				GenerationContext.Log(Message, Node);
@@ -162,14 +162,14 @@ mu::Ptr<mu::NodeScalar> GenerateMutableSourceFloat(const UEdGraphPin* Pin, FMuta
 
 			const int32 NumSwitchOptions = TypedNodeFloatSwitch->GetNumElements();
 
-			mu::NodeScalarEnumParameter* EnumParameter = static_cast<mu::NodeScalarEnumParameter*>(SwitchParam.get());
+			UE::Mutable::Private::NodeScalarEnumParameter* EnumParameter = static_cast<UE::Mutable::Private::NodeScalarEnumParameter*>(SwitchParam.get());
 			if (NumSwitchOptions != EnumParameter->Options.Num())
 			{
 				const FText Message = LOCTEXT("MismatchedSwitch", "Switch enum and switch node have different number of options. Please refresh the switch node to make sure the outcomes are labeled properly.");
 				GenerationContext.Log(Message, Node);
 			}
 
-			mu::Ptr<mu::NodeScalarSwitch> SwitchNode = new mu::NodeScalarSwitch;
+			UE::Mutable::Private::Ptr<UE::Mutable::Private::NodeScalarSwitch> SwitchNode = new UE::Mutable::Private::NodeScalarSwitch;
 			SwitchNode->Parameter = SwitchParam;
 			SwitchNode->Options.SetNum(NumSwitchOptions);
 
@@ -190,7 +190,7 @@ mu::Ptr<mu::NodeScalar> GenerateMutableSourceFloat(const UEdGraphPin* Pin, FMuta
 
 	else if (const UCustomizableObjectNodeCurve* TypedNodeCurve = Cast<UCustomizableObjectNodeCurve>(Node))
 	{
-		mu::Ptr<mu::NodeScalarCurve> CurveNode = new mu::NodeScalarCurve();
+		UE::Mutable::Private::Ptr<UE::Mutable::Private::NodeScalarCurve> CurveNode = new UE::Mutable::Private::NodeScalarCurve();
 		Result = CurveNode;
 
 		if (const UEdGraphPin* ConnectedPin = FollowInputPin(*TypedNodeCurve->InputPin()))
@@ -237,12 +237,12 @@ mu::Ptr<mu::NodeScalar> GenerateMutableSourceFloat(const UEdGraphPin* Pin, FMuta
 
 	else if (const UCustomizableObjectNodeFloatVariation* TypedNodeFloatVar = Cast<const UCustomizableObjectNodeFloatVariation>(Node))
 	{
-		mu::Ptr<mu::NodeScalarVariation> FloatNode = new mu::NodeScalarVariation();
+		UE::Mutable::Private::Ptr<UE::Mutable::Private::NodeScalarVariation> FloatNode = new UE::Mutable::Private::NodeScalarVariation();
 		Result = FloatNode;
 
 		if (const UEdGraphPin* ConnectedPin = FollowInputPin(*TypedNodeFloatVar->DefaultPin()))
 		{
-			mu::Ptr<mu::NodeScalar> ChildNode = GenerateMutableSourceFloat(ConnectedPin, GenerationContext);
+			UE::Mutable::Private::Ptr<UE::Mutable::Private::NodeScalar> ChildNode = GenerateMutableSourceFloat(ConnectedPin, GenerationContext);
 			if (ChildNode)
 			{
 				FloatNode->DefaultScalar = ChildNode;
@@ -264,7 +264,7 @@ mu::Ptr<mu::NodeScalar> GenerateMutableSourceFloat(const UEdGraphPin* Pin, FMuta
 
 			if (const UEdGraphPin* ConnectedPin = FollowInputPin(*VariationPin))
 			{
-				mu::Ptr<mu::NodeScalar> ChildNode = GenerateMutableSourceFloat(ConnectedPin, GenerationContext);
+				UE::Mutable::Private::Ptr<UE::Mutable::Private::NodeScalar> ChildNode = GenerateMutableSourceFloat(ConnectedPin, GenerationContext);
 				FloatNode->Variations[VariationIndex].Scalar = ChildNode;
 			}
 		}
@@ -272,37 +272,37 @@ mu::Ptr<mu::NodeScalar> GenerateMutableSourceFloat(const UEdGraphPin* Pin, FMuta
 
 	else if (const UCustomizableObjectNodeFloatArithmeticOp* TypedNodeFloatArith = Cast<UCustomizableObjectNodeFloatArithmeticOp>(Node))
 	{
-		mu::Ptr<mu::NodeScalarArithmeticOperation> OpNode = new mu::NodeScalarArithmeticOperation();
+		UE::Mutable::Private::Ptr<UE::Mutable::Private::NodeScalarArithmeticOperation> OpNode = new UE::Mutable::Private::NodeScalarArithmeticOperation();
 		Result = OpNode;
 
 		if (const UEdGraphPin* ConnectedPin = FollowInputPin(*TypedNodeFloatArith->XPin()))
 		{
-			mu::Ptr<mu::NodeScalar> XNode = GenerateMutableSourceFloat(ConnectedPin, GenerationContext);
+			UE::Mutable::Private::Ptr<UE::Mutable::Private::NodeScalar> XNode = GenerateMutableSourceFloat(ConnectedPin, GenerationContext);
 			OpNode->A = XNode;
 		}
 
 		if (const UEdGraphPin* ConnectedPin = FollowInputPin(*TypedNodeFloatArith->YPin()))
 		{
-			mu::Ptr<mu::NodeScalar> YNode = GenerateMutableSourceFloat(ConnectedPin, GenerationContext);
+			UE::Mutable::Private::Ptr<UE::Mutable::Private::NodeScalar> YNode = GenerateMutableSourceFloat(ConnectedPin, GenerationContext);
 			OpNode->B = YNode;
 		}
 
 		switch (TypedNodeFloatArith->Operation)
 		{
 		case EFloatArithmeticOperation::E_Add:
-			OpNode->Operation = mu::NodeScalarArithmeticOperation::EOperation::AO_ADD;
+			OpNode->Operation = UE::Mutable::Private::NodeScalarArithmeticOperation::EOperation::AO_ADD;
 			break;
 
 		case EFloatArithmeticOperation::E_Sub:
-			OpNode->Operation = mu::NodeScalarArithmeticOperation::EOperation::AO_SUBTRACT;
+			OpNode->Operation = UE::Mutable::Private::NodeScalarArithmeticOperation::EOperation::AO_SUBTRACT;
 			break;
 
 		case EFloatArithmeticOperation::E_Mul:
-			OpNode->Operation = mu::NodeScalarArithmeticOperation::EOperation::AO_MULTIPLY;
+			OpNode->Operation = UE::Mutable::Private::NodeScalarArithmeticOperation::EOperation::AO_MULTIPLY;
 			break;
 
 		case EFloatArithmeticOperation::E_Div:
-			OpNode->Operation = mu::NodeScalarArithmeticOperation::EOperation::AO_DIVIDE;
+			OpNode->Operation = UE::Mutable::Private::NodeScalarArithmeticOperation::EOperation::AO_DIVIDE;
 			break;
 
 		default:
@@ -313,7 +313,7 @@ mu::Ptr<mu::NodeScalar> GenerateMutableSourceFloat(const UEdGraphPin* Pin, FMuta
 	else if (const UCustomizableObjectNodeTable* TypedNodeTable = Cast<UCustomizableObjectNodeTable>(Node))
 	{
 		//This node will add a default value in case of error
-		mu::Ptr<mu::NodeScalarConstant> ConstantValue = new mu::NodeScalarConstant();
+		UE::Mutable::Private::Ptr<UE::Mutable::Private::NodeScalarConstant> ConstantValue = new UE::Mutable::Private::NodeScalarConstant();
 		ConstantValue->Value = 1.0f;
 
 		Result = ConstantValue;
@@ -343,12 +343,12 @@ mu::Ptr<mu::NodeScalar> GenerateMutableSourceFloat(const UEdGraphPin* Pin, FMuta
 			if (bSuccess)
 			{
 				// Generating a new data table if not exists
-				mu::Ptr<mu::FTable> Table;
+				UE::Mutable::Private::Ptr<UE::Mutable::Private::FTable> Table;
 				Table = GenerateMutableSourceTable(DataTable, TypedNodeTable, GenerationContext);
 
 				if (Table)
 				{
-					mu::Ptr<mu::NodeScalarTable> ScalarTableNode = new mu::NodeScalarTable();
+					UE::Mutable::Private::Ptr<UE::Mutable::Private::NodeScalarTable> ScalarTableNode = new UE::Mutable::Private::NodeScalarTable();
 
 					// Generating a new Float column if not exists
 					if (Table->FindColumn(ColumnName) == INDEX_NONE)
@@ -391,23 +391,23 @@ mu::Ptr<mu::NodeScalar> GenerateMutableSourceFloat(const UEdGraphPin* Pin, FMuta
 	else if (const UCustomizableObjectNodeMacroInstance* TypedNodeMacro = Cast<UCustomizableObjectNodeMacroInstance>(Node))
 	{
 		bCacheNode = false;
-		Result = GenerateMutableSourceMacro<mu::NodeScalar>(*Pin, GenerationContext, GenerateMutableSourceFloat);
+		Result = GenerateMutableSourceMacro<UE::Mutable::Private::NodeScalar>(*Pin, GenerationContext, GenerateMutableSourceFloat);
 	}
 
 	else if (const UCustomizableObjectNodeTunnel* TypedNodeTunnel = Cast<UCustomizableObjectNodeTunnel>(Node))
 	{
 		bCacheNode = false;
-		Result = GenerateMutableSourceMacro<mu::NodeScalar>(*Pin, GenerationContext, GenerateMutableSourceFloat);
+		Result = GenerateMutableSourceMacro<UE::Mutable::Private::NodeScalar>(*Pin, GenerationContext, GenerateMutableSourceFloat);
 	}
 
 	else if (const UCONodeMaterialBreak* TypedNodeMaterialBreak = Cast<UCONodeMaterialBreak>(Node))
 	{
-		Result = new mu::NodeScalarConstant();
+		Result = new UE::Mutable::Private::NodeScalarConstant();
 		check(TypedNodeMaterialBreak->MaterialPinRef.Get());
 
 		if (const UEdGraphPin* LinkedMaterial = FollowInputPin(*TypedNodeMaterialBreak->MaterialPinRef.Get()))
 		{
-			mu::Ptr<mu::NodeScalarMaterialBreak> MaterialBreak = new mu::NodeScalarMaterialBreak();
+			UE::Mutable::Private::Ptr<UE::Mutable::Private::NodeScalarMaterialBreak> MaterialBreak = new UE::Mutable::Private::NodeScalarMaterialBreak();
 
 			FName ParameterName = TypedNodeMaterialBreak->GetPinParameterName(*Pin);
 			check(!ParameterName.IsNone());
@@ -417,7 +417,7 @@ mu::Ptr<mu::NodeScalar> GenerateMutableSourceFloat(const UEdGraphPin* Pin, FMuta
 
 			// Generate the material that is going to be broken into parameters
 			GenerationContext.CurrentMaterialBreakParameter = FMutableGraphGenerationContext::FMaterialBreakParameter{ ParameterName, EMaterialParameterType::Scalar };
-			mu::Ptr<mu::NodeMaterial> LinkedNode = GenerateMutableSourceMaterial(LinkedMaterial, GenerationContext);
+			UE::Mutable::Private::Ptr<UE::Mutable::Private::NodeMaterial> LinkedNode = GenerateMutableSourceMaterial(LinkedMaterial, GenerationContext);
 			MaterialBreak->MaterialSource = LinkedNode;
 			GenerationContext.CurrentMaterialBreakParameter = FMutableGraphGenerationContext::FMaterialBreakParameter{ };
 
diff --git a/Engine/Plugins/Mutable/Source/CustomizableObjectEditor/Private/MuCOE/GenerateMutableSource/GenerateMutableSourceFloat.h b/Engine/Plugins/Mutable/Source/CustomizableObjectEditor/Private/MuCOE/GenerateMutableSource/GenerateMutableSourceFloat.h
--- a/Engine/Plugins/Mutable/Source/CustomizableObjectEditor/Private/MuCOE/GenerateMutableSource/GenerateMutableSourceFloat.h
+++ b/Engine/Plugins/Mutable/Source/CustomizableObjectEditor/Private/MuCOE/GenerateMutableSource/GenerateMutableSourceFloat.h
@@ -8,4 +8,4 @@ class UEdGraphPin;
 struct FMutableGraphGenerationContext;
 
 /** Convert a CustomizableObject Source Graph into a mutable source graph. */
-mu::Ptr<mu::NodeScalar> GenerateMutableSourceFloat(const UEdGraphPin* Pin, FMutableGraphGenerationContext& GenerationContext);
+UE::Mutable::Private::Ptr<UE::Mutable::Private::NodeScalar> GenerateMutableSourceFloat(const UEdGraphPin* Pin, FMutableGraphGenerationContext& GenerationContext);
diff --git a/Engine/Plugins/Mutable/Source/CustomizableObjectEditor/Private/MuCOE/GenerateMutableSource/GenerateMutableSourceGroupProjector.cpp b/Engine/Plugins/Mutable/Source/CustomizableObjectEditor/Private/MuCOE/GenerateMutableSource/GenerateMutableSourceGroupProjector.cpp
--- a/Engine/Plugins/Mutable/Source/CustomizableObjectEditor/Private/MuCOE/GenerateMutableSource/GenerateMutableSourceGroupProjector.cpp
+++ b/Engine/Plugins/Mutable/Source/CustomizableObjectEditor/Private/MuCOE/GenerateMutableSource/GenerateMutableSourceGroupProjector.cpp
@@ -35,7 +35,7 @@ class UPoseAsset;
 #define LOCTEXT_NAMESPACE "CustomizableObjectEditor"
 
 
-mu::NodeImagePtr GenerateMutableSourceGroupProjector(const int32 NodeLOD, const int32 ImageIndex, mu::NodeMeshPtr MeshNode, FMutableGraphGenerationContext& GenerationContext,
+UE::Mutable::Private::NodeImagePtr GenerateMutableSourceGroupProjector(const int32 NodeLOD, const int32 ImageIndex, UE::Mutable::Private::NodeMeshPtr MeshNode, FMutableGraphGenerationContext& GenerationContext,
 	UCustomizableObjectNodeMaterialBase* TypedNodeMat, UCustomizableObjectNodeModifierExtendMeshSection* TypedNodeExt, bool& bShareProjectionTexturesBetweenLODs, bool& bIsGroupProjectorImage,
 	UTexture2D*& GroupProjectionReferenceTexture)
 {
@@ -43,10 +43,10 @@ mu::NodeImagePtr GenerateMutableSourceGroupProjector(const int32 NodeLOD, const
 
 	if (!MeshNode.get())
 	{
-		return mu::NodeImagePtr();
+		return UE::Mutable::Private::NodeImagePtr();
 	}
 
-	TArray<mu::Ptr<mu::NodeImageProject>> ImageNodes;
+	TArray<UE::Mutable::Private::Ptr<UE::Mutable::Private::NodeImageProject>> ImageNodes;
 	TArray<FGroupProjectorTempData> ImageNodes_ProjectorTempData;
 
 	int32 TextureSize = 512;
@@ -127,23 +127,23 @@ mu::NodeImagePtr GenerateMutableSourceGroupProjector(const int32 NodeLOD, const
 					continue;
 				}
 				
-				mu::Ptr<mu::NodeImageProject> ImageNode = new mu::NodeImageProject();
+				UE::Mutable::Private::Ptr<UE::Mutable::Private::NodeImageProject> ImageNode = new UE::Mutable::Private::NodeImageProject();
 				bIsGroupProjectorImage = true;
 				ImageNode->Layout = ProjectorTempData.CustomizableObjectNodeGroupProjectorParameter->UVLayout;
 
 				{
-					mu::Ptr<mu::NodeScalarConstant> ScalarNode = new mu::NodeScalarConstant;
+					UE::Mutable::Private::Ptr<UE::Mutable::Private::NodeScalarConstant> ScalarNode = new UE::Mutable::Private::NodeScalarConstant;
 					ScalarNode->Value = 120.f;
 					ImageNode->AngleFadeStart = ScalarNode;
 				}
 
 				{
-					mu::Ptr<mu::NodeScalarConstant> ScalarNode = new mu::NodeScalarConstant;
+					UE::Mutable::Private::Ptr<UE::Mutable::Private::NodeScalarConstant> ScalarNode = new UE::Mutable::Private::NodeScalarConstant;
 					ScalarNode->Value = 150.f;
 					ImageNode->AngleFadeEnd = ScalarNode;
 				}
 
-				mu::Ptr<mu::NodeMeshSwitch> MeshSwitchNode = new mu::NodeMeshSwitch;
+				UE::Mutable::Private::Ptr<UE::Mutable::Private::NodeMeshSwitch> MeshSwitchNode = new UE::Mutable::Private::NodeMeshSwitch;
 				MeshSwitchNode->Parameter = ProjectorTempData.PoseOptionsParameter;
 				MeshSwitchNode->Options.SetNum(ProjectorTempData.CustomizableObjectNodeGroupProjectorParameter->OptionPoses.Num() + 1);
 				MeshSwitchNode->Options[0] = MeshNode;
@@ -152,7 +152,7 @@ mu::NodeImagePtr GenerateMutableSourceGroupProjector(const int32 NodeLOD, const
 				{
 					if (ProjectorTempData.CustomizableObjectNodeGroupProjectorParameter->OptionPoses[SelectorIndex].OptionPose)
 					{
-						mu::Ptr<mu::NodeMeshApplyPose> NodeMeshApplyPose = CreateNodeMeshApplyPose(GenerationContext, MeshNode,
+						UE::Mutable::Private::Ptr<UE::Mutable::Private::NodeMeshApplyPose> NodeMeshApplyPose = CreateNodeMeshApplyPose(GenerationContext, MeshNode,
 							ProjectorTempData.PoseBoneDataArray[SelectorIndex].ArrayBoneName, ProjectorTempData.PoseBoneDataArray[SelectorIndex].ArrayTransform);
 
 						if (!NodeMeshApplyPose)
@@ -184,27 +184,27 @@ mu::NodeImagePtr GenerateMutableSourceGroupProjector(const int32 NodeLOD, const
 	
 	if (ImageNodes.Num() == 0)
 	{
-		return mu::NodeImagePtr();
+		return UE::Mutable::Private::NodeImagePtr();
 	}
 	
-	mu::Ptr<mu::NodeColourConstant> ZeroColorNode = new mu::NodeColourConstant();
+	UE::Mutable::Private::Ptr<UE::Mutable::Private::NodeColourConstant> ZeroColorNode = new UE::Mutable::Private::NodeColourConstant();
 	ZeroColorNode->Value = FVector4f(0.f, 0.f, 0.f, 1.0f);
 
-	mu::Ptr<mu::NodeImagePlainColour> ZeroPlainColourNode = new mu::NodeImagePlainColour;
+	UE::Mutable::Private::Ptr<UE::Mutable::Private::NodeImagePlainColour> ZeroPlainColourNode = new UE::Mutable::Private::NodeImagePlainColour;
 	ZeroPlainColourNode->SizeX = TextureSize;
 	ZeroPlainColourNode->SizeY = TextureSize;
 	ZeroPlainColourNode->Colour = ZeroColorNode;
 
-	mu::Ptr<mu::NodeImageSwizzle> ZeroChannelNode = new mu::NodeImageSwizzle;
-	ZeroChannelNode->SetFormat( mu::EImageFormat::L_UByte );
+	UE::Mutable::Private::Ptr<UE::Mutable::Private::NodeImageSwizzle> ZeroChannelNode = new UE::Mutable::Private::NodeImageSwizzle;
+	ZeroChannelNode->SetFormat( UE::Mutable::Private::EImageFormat::L_UByte );
 	ZeroChannelNode->Sources[0] = ZeroPlainColourNode;
 	ZeroChannelNode->SourceChannels[0] = 2; // Just take a zeroed channel for the base alpha
 
-	mu::Ptr<mu::NodeScalarConstant> OneConstantNode = new mu::NodeScalarConstant;
+	UE::Mutable::Private::Ptr<UE::Mutable::Private::NodeScalarConstant> OneConstantNode = new UE::Mutable::Private::NodeScalarConstant;
 	OneConstantNode->Value = 1.f;
 
-	mu::NodeImagePtr ResultAlpha = ZeroChannelNode;
-	mu::NodeImagePtr ResultImage = ZeroPlainColourNode;
+	UE::Mutable::Private::NodeImagePtr ResultAlpha = ZeroChannelNode;
+	UE::Mutable::Private::NodeImagePtr ResultImage = ZeroPlainColourNode;
 
 	for (int32 i = 0; i < ImageNodes.Num(); ++i)
 	{
@@ -221,69 +221,69 @@ mu::NodeImagePtr GenerateMutableSourceGroupProjector(const int32 NodeLOD, const
 			{
 				TextureSize = NewTextureSize;
 
-				ZeroPlainColourNode = new mu::NodeImagePlainColour;
+				ZeroPlainColourNode = new UE::Mutable::Private::NodeImagePlainColour;
 				ZeroPlainColourNode->SizeX = TextureSize;
 				ZeroPlainColourNode->SizeY = TextureSize;
 				ZeroPlainColourNode->Colour = ZeroColorNode;
 			}
 		}
 
-		mu::Ptr<mu::NodeImageSwizzle> ImageNodesAlphaChannelNode = new mu::NodeImageSwizzle;
-		ImageNodesAlphaChannelNode->SetFormat( mu::EImageFormat::L_UByte );
+		UE::Mutable::Private::Ptr<UE::Mutable::Private::NodeImageSwizzle> ImageNodesAlphaChannelNode = new UE::Mutable::Private::NodeImageSwizzle;
+		ImageNodesAlphaChannelNode->SetFormat( UE::Mutable::Private::EImageFormat::L_UByte );
 		ImageNodesAlphaChannelNode->Sources[0] = ImageNodes[i];
 		ImageNodesAlphaChannelNode->SourceChannels[0] = 3;
 
-		mu::Ptr<mu::NodeColourFromScalars> ColourFromScalars = new mu::NodeColourFromScalars;
+		UE::Mutable::Private::Ptr<UE::Mutable::Private::NodeColourFromScalars> ColourFromScalars = new UE::Mutable::Private::NodeColourFromScalars;
 		ColourFromScalars->X = ImageNodes_ProjectorTempData[i].NodeOpacityParameter;
 		ColourFromScalars->Y = ImageNodes_ProjectorTempData[i].NodeOpacityParameter;
 		ColourFromScalars->Z = ImageNodes_ProjectorTempData[i].NodeOpacityParameter;
 		ColourFromScalars->W = OneConstantNode;
 
-		mu::Ptr<mu::NodeImageLayerColour> OpacityMultiLayerNode = new mu::NodeImageLayerColour;
-		OpacityMultiLayerNode->Type = mu::EBlendType::BT_MULTIPLY;
+		UE::Mutable::Private::Ptr<UE::Mutable::Private::NodeImageLayerColour> OpacityMultiLayerNode = new UE::Mutable::Private::NodeImageLayerColour;
+		OpacityMultiLayerNode->Type = UE::Mutable::Private::EBlendType::BT_MULTIPLY;
 		OpacityMultiLayerNode->Colour = ColourFromScalars;
 		OpacityMultiLayerNode->Base = ImageNodesAlphaChannelNode;
 		//OpacityMultiLayerNode->SetMask(OneChannelNode); // No mask needed
 
-		mu::Ptr<mu::NodeImageSwizzle> MultiplySwizzleNode = new mu::NodeImageSwizzle;
-		MultiplySwizzleNode->SetFormat( mu::EImageFormat::L_UByte );
+		UE::Mutable::Private::Ptr<UE::Mutable::Private::NodeImageSwizzle> MultiplySwizzleNode = new UE::Mutable::Private::NodeImageSwizzle;
+		MultiplySwizzleNode->SetFormat( UE::Mutable::Private::EImageFormat::L_UByte );
 		MultiplySwizzleNode->Sources[0] = OpacityMultiLayerNode;
 		MultiplySwizzleNode->SourceChannels[0] = 0;
 
-		mu::Ptr<mu::NodeImageMultiLayer> BaseAlphaMultiLayerNode = new mu::NodeImageMultiLayer;
+		UE::Mutable::Private::Ptr<UE::Mutable::Private::NodeImageMultiLayer> BaseAlphaMultiLayerNode = new UE::Mutable::Private::NodeImageMultiLayer;
 		BaseAlphaMultiLayerNode->Range = ImageNodes_ProjectorTempData[i].NodeRange;
-		BaseAlphaMultiLayerNode->Type = mu::EBlendType::BT_LIGHTEN;
+		BaseAlphaMultiLayerNode->Type = UE::Mutable::Private::EBlendType::BT_LIGHTEN;
 		BaseAlphaMultiLayerNode->Base = ResultAlpha;
 		BaseAlphaMultiLayerNode->Blended = MultiplySwizzleNode;
 		//BaseAlphaMultiLayerNode->SetMask(MultiplySwizzleNode); // No mask needed
 		ResultAlpha = BaseAlphaMultiLayerNode;
 
-		mu::Ptr<mu::NodeImageMultiLayer> BaseMultiLayerNode = new mu::NodeImageMultiLayer;
+		UE::Mutable::Private::Ptr<UE::Mutable::Private::NodeImageMultiLayer> BaseMultiLayerNode = new UE::Mutable::Private::NodeImageMultiLayer;
 		BaseMultiLayerNode->Range = ImageNodes_ProjectorTempData[i].NodeRange;
-		BaseMultiLayerNode->Type = mu::EBlendType::BT_BLEND;
+		BaseMultiLayerNode->Type = UE::Mutable::Private::EBlendType::BT_BLEND;
 		BaseMultiLayerNode->Base = ResultImage;
 		BaseMultiLayerNode->Blended = ImageNodes[i];
 		BaseMultiLayerNode->Mask = MultiplySwizzleNode;
 		ResultImage = BaseMultiLayerNode;
 	}
 
-	mu::Ptr<mu::NodeImageSwizzle> SwizzleNodeR = new mu::NodeImageSwizzle;
-	SwizzleNodeR->SetFormat(mu::EImageFormat::L_UByte);
+	UE::Mutable::Private::Ptr<UE::Mutable::Private::NodeImageSwizzle> SwizzleNodeR = new UE::Mutable::Private::NodeImageSwizzle;
+	SwizzleNodeR->SetFormat(UE::Mutable::Private::EImageFormat::L_UByte);
 	SwizzleNodeR->Sources[0] = ResultImage;
 	SwizzleNodeR->SourceChannels[0] = 0;
 
-	mu::Ptr<mu::NodeImageSwizzle> SwizzleNodeG = new mu::NodeImageSwizzle;
-	SwizzleNodeG->SetFormat(mu::EImageFormat::L_UByte);
+	UE::Mutable::Private::Ptr<UE::Mutable::Private::NodeImageSwizzle> SwizzleNodeG = new UE::Mutable::Private::NodeImageSwizzle;
+	SwizzleNodeG->SetFormat(UE::Mutable::Private::EImageFormat::L_UByte);
 	SwizzleNodeG->Sources[0] = ResultImage;
 	SwizzleNodeG->SourceChannels[0] = 1;
 
-	mu::Ptr<mu::NodeImageSwizzle> SwizzleNodeB = new mu::NodeImageSwizzle;
-	SwizzleNodeB->SetFormat(mu::EImageFormat::L_UByte);
+	UE::Mutable::Private::Ptr<UE::Mutable::Private::NodeImageSwizzle> SwizzleNodeB = new UE::Mutable::Private::NodeImageSwizzle;
+	SwizzleNodeB->SetFormat(UE::Mutable::Private::EImageFormat::L_UByte);
 	SwizzleNodeB->Sources[0] = ResultImage;
 	SwizzleNodeB->SourceChannels[0] = 2;
 
-	mu::Ptr<mu::NodeImageSwizzle> FinalSwizzleNode = new mu::NodeImageSwizzle;
-	FinalSwizzleNode->SetFormat(mu::EImageFormat::RGBA_UByte);
+	UE::Mutable::Private::Ptr<UE::Mutable::Private::NodeImageSwizzle> FinalSwizzleNode = new UE::Mutable::Private::NodeImageSwizzle;
+	FinalSwizzleNode->SetFormat(UE::Mutable::Private::EImageFormat::RGBA_UByte);
 	FinalSwizzleNode->Sources[0] = SwizzleNodeR;
 	FinalSwizzleNode->SourceChannels[0] = 0;
 	FinalSwizzleNode->Sources[1] = SwizzleNodeG;
@@ -312,8 +312,8 @@ TOptional<FGroupProjectorTempData> GenerateMutableGroupProjector(UCustomizableOb
 	
 	const FCompilationOptions& CompilationOptions = GenerationContext.CompilationContext->Options;
 
-	// The static cast works here because it's already known to be a mu::NodeProjectorParameter* because of the UE5 Cast in the previous line
-	GroupProjectorTempData.NodeProjectorParameterPtr = static_cast<mu::NodeProjectorParameter*>(GenerateMutableSourceProjector(&ProjParamNode->OutputPin(), GenerationContext).get());
+	// The static cast works here because it's already known to be a UE::Mutable::Private::NodeProjectorParameter* because of the UE5 Cast in the previous line
+	GroupProjectorTempData.NodeProjectorParameterPtr = static_cast<UE::Mutable::Private::NodeProjectorParameter*>(GenerateMutableSourceProjector(&ProjParamNode->OutputPin(), GenerationContext).get());
 
 	if (GroupProjectorTempData.NodeProjectorParameterPtr)
 	{
@@ -341,14 +341,14 @@ TOptional<FGroupProjectorTempData> GenerateMutableGroupProjector(UCustomizableOb
 
 		GroupProjectorTempData.CustomizableObjectNodeGroupProjectorParameter = ProjParamNode;
 
-		mu::Ptr<mu::NodeScalarParameter> NodeScalarParam;
-		if (mu::Ptr<mu::NodeScalarParameter>* FindResult = GenerationContext.GeneratedScalarParameters.Find(NumLayersParamUid.ToString()))
+		UE::Mutable::Private::Ptr<UE::Mutable::Private::NodeScalarParameter> NodeScalarParam;
+		if (UE::Mutable::Private::Ptr<UE::Mutable::Private::NodeScalarParameter>* FindResult = GenerationContext.GeneratedScalarParameters.Find(NumLayersParamUid.ToString()))
 		{
 			NodeScalarParam = *FindResult;
 		}
 		else
 		{
-			NodeScalarParam = new mu::NodeScalarParameter;
+			NodeScalarParam = new UE::Mutable::Private::NodeScalarParameter;
 			FString NodeScalarParamName = ProjParamNode->GetParameterName(&GenerationContext.MacroNodesStack) + NUM_LAYERS_PARAMETER_POSTFIX;
 			NodeScalarParam->Name = NodeScalarParamName;
 			NodeScalarParam->UID = NumLayersParamUid.ToString();
@@ -360,20 +360,20 @@ TOptional<FGroupProjectorTempData> GenerateMutableGroupProjector(UCustomizableOb
 			GenerationContext.GeneratedScalarParameters.Add(NumLayersParamUid.ToString(), NodeScalarParam);
 		}
 
-		mu::Ptr<mu::NodeRangeFromScalar> NodeRangeFromScalar = new mu::NodeRangeFromScalar;
+		UE::Mutable::Private::Ptr<UE::Mutable::Private::NodeRangeFromScalar> NodeRangeFromScalar = new UE::Mutable::Private::NodeRangeFromScalar;
 		NodeRangeFromScalar->Size = NodeScalarParam;
 		GroupProjectorTempData.NodeRange = NodeRangeFromScalar;
 		GroupProjectorTempData.NodeProjectorParameterPtr->SetRangeCount(1);
 		GroupProjectorTempData.NodeProjectorParameterPtr->SetRange(0, NodeRangeFromScalar);
 		
-		mu::Ptr<mu::NodeScalarEnumParameter> PoseEnumParameterNode;
-		if (mu::Ptr<mu::NodeScalarEnumParameter>* FindResult = GenerationContext.GeneratedEnumParameters.Find(SelectedPoseParamUid.ToString()))
+		UE::Mutable::Private::Ptr<UE::Mutable::Private::NodeScalarEnumParameter> PoseEnumParameterNode;
+		if (UE::Mutable::Private::Ptr<UE::Mutable::Private::NodeScalarEnumParameter>* FindResult = GenerationContext.GeneratedEnumParameters.Find(SelectedPoseParamUid.ToString()))
 		{
 			PoseEnumParameterNode = *FindResult;
 		}
 		else
 		{
-			PoseEnumParameterNode = new mu::NodeScalarEnumParameter;
+			PoseEnumParameterNode = new UE::Mutable::Private::NodeScalarEnumParameter;
 			FString PoseNodeEnumParamName = ProjParamNode->GetParameterName(&GenerationContext.MacroNodesStack) + POSE_PARAMETER_POSTFIX;
 			PoseEnumParameterNode->Name = PoseNodeEnumParamName;
 			PoseEnumParameterNode->UID = SelectedPoseParamUid.ToString();
@@ -387,14 +387,14 @@ TOptional<FGroupProjectorTempData> GenerateMutableGroupProjector(UCustomizableOb
 		FMutableParameterData(ProjParamNode->ParamUIMetadata, EMutableParameterType::Int));
 		}
 		
-		mu::Ptr<mu::NodeScalarParameter> OpacityParameterNode;
-		if (mu::Ptr<mu::NodeScalarParameter>* FindResult = GenerationContext.GeneratedScalarParameters.Find(OpacityParamUid.ToString()))
+		UE::Mutable::Private::Ptr<UE::Mutable::Private::NodeScalarParameter> OpacityParameterNode;
+		if (UE::Mutable::Private::Ptr<UE::Mutable::Private::NodeScalarParameter>* FindResult = GenerationContext.GeneratedScalarParameters.Find(OpacityParamUid.ToString()))
 		{
 			OpacityParameterNode = *FindResult;
 		}
 		else
 		{
-			OpacityParameterNode = new mu::NodeScalarParameter;
+			OpacityParameterNode = new UE::Mutable::Private::NodeScalarParameter;
 			FString OpacityParameterNodeName = ProjParamNode->GetParameterName(&GenerationContext.MacroNodesStack) + OPACITY_PARAMETER_POSTFIX;
 			OpacityParameterNode->Name = OpacityParameterNodeName;
 			OpacityParameterNode->UID = OpacityParamUid.ToString();
@@ -450,14 +450,14 @@ TOptional<FGroupProjectorTempData> GenerateMutableGroupProjector(UCustomizableOb
 				GroupProjectorTempData.PoseBoneDataArray[PoseIndex].ArrayBoneName, GroupProjectorTempData.PoseBoneDataArray[PoseIndex].ArrayTransform);
 		}
 	
-		mu::Ptr<mu::NodeScalarEnumParameter> EnumParameterNode;
-		if (mu::Ptr<mu::NodeScalarEnumParameter>* FindResult = GenerationContext.GeneratedEnumParameters.Find(SelectedImageParamUid.ToString()))
+		UE::Mutable::Private::Ptr<UE::Mutable::Private::NodeScalarEnumParameter> EnumParameterNode;
+		if (UE::Mutable::Private::Ptr<UE::Mutable::Private::NodeScalarEnumParameter>* FindResult = GenerationContext.GeneratedEnumParameters.Find(SelectedImageParamUid.ToString()))
 		{
 			EnumParameterNode = *FindResult;
 		}
 		else
 		{
-			EnumParameterNode = new mu::NodeScalarEnumParameter;
+			EnumParameterNode = new UE::Mutable::Private::NodeScalarEnumParameter;
 
 			FString NodeEnumParamName = ProjParamNode->GetParameterName(&GenerationContext.MacroNodesStack) + IMAGE_PARAMETER_POSTFIX;
 			EnumParameterNode->Name = NodeEnumParamName;
@@ -490,7 +490,7 @@ TOptional<FGroupProjectorTempData> GenerateMutableGroupProjector(UCustomizableOb
 			GenerationContext.ParameterUIDataMap.Add(NodeEnumParamName, ParameterUIData);
 		}
 
-		mu::Ptr<mu::NodeImageSwitch> SwitchNode = new mu::NodeImageSwitch;
+		UE::Mutable::Private::Ptr<UE::Mutable::Private::NodeImageSwitch> SwitchNode = new UE::Mutable::Private::NodeImageSwitch;
 		SwitchNode->Parameter = EnumParameterNode;
 		SwitchNode->Options.SetNum(ArrayOptionTexture.Num());
 
@@ -498,9 +498,9 @@ TOptional<FGroupProjectorTempData> GenerateMutableGroupProjector(UCustomizableOb
 		{
 			if (const TObjectPtr<UTexture2D>& Texture = ArrayOptionTexture[SelectorIndex].OptionTexture)
 			{
-				TSharedPtr<mu::FImage> ImageConstant = GenerateImageConstant(ArrayOptionTexture[SelectorIndex].OptionTexture, GenerationContext, false);
+				TSharedPtr<UE::Mutable::Private::FImage> ImageConstant = GenerateImageConstant(ArrayOptionTexture[SelectorIndex].OptionTexture, GenerationContext, false);
 
-				mu::Ptr<mu::NodeImageConstant> ImageNode = new mu::NodeImageConstant();
+				UE::Mutable::Private::Ptr<UE::Mutable::Private::NodeImageConstant> ImageNode = new UE::Mutable::Private::NodeImageConstant();
 				ImageNode->SetValue(ImageConstant);
 
 				const uint32 MipsToSkip = ComputeLODBiasForTexture(GenerationContext, *Texture, ProjParamNode->ReferenceTexture);
diff --git a/Engine/Plugins/Mutable/Source/CustomizableObjectEditor/Private/MuCOE/GenerateMutableSource/GenerateMutableSourceGroupProjector.h b/Engine/Plugins/Mutable/Source/CustomizableObjectEditor/Private/MuCOE/GenerateMutableSource/GenerateMutableSourceGroupProjector.h
--- a/Engine/Plugins/Mutable/Source/CustomizableObjectEditor/Private/MuCOE/GenerateMutableSource/GenerateMutableSourceGroupProjector.h
+++ b/Engine/Plugins/Mutable/Source/CustomizableObjectEditor/Private/MuCOE/GenerateMutableSource/GenerateMutableSourceGroupProjector.h
@@ -17,7 +17,7 @@ class UTexture2D;
 struct FMutableGraphGenerationContext;
 
 
-mu::NodeImagePtr GenerateMutableSourceGroupProjector(const int32 NodeLOD, const int32 ImageIndex, mu::NodeMeshPtr MeshNode, FMutableGraphGenerationContext& GenerationContext,
+UE::Mutable::Private::NodeImagePtr GenerateMutableSourceGroupProjector(const int32 NodeLOD, const int32 ImageIndex, UE::Mutable::Private::NodeMeshPtr MeshNode, FMutableGraphGenerationContext& GenerationContext,
 	UCustomizableObjectNodeMaterialBase* TypedNodeMat, UCustomizableObjectNodeModifierExtendMeshSection* TypedNodeExt, bool& bShareProjectionTexturesBetweenLODs, bool& bIsGroupProjectorImage,
 	UTexture2D*& GroupProjectionReferenceTexture);
 
diff --git a/Engine/Plugins/Mutable/Source/CustomizableObjectEditor/Private/MuCOE/GenerateMutableSource/GenerateMutableSourceImage.cpp b/Engine/Plugins/Mutable/Source/CustomizableObjectEditor/Private/MuCOE/GenerateMutableSource/GenerateMutableSourceImage.cpp
--- a/Engine/Plugins/Mutable/Source/CustomizableObjectEditor/Private/MuCOE/GenerateMutableSource/GenerateMutableSourceImage.cpp
+++ b/Engine/Plugins/Mutable/Source/CustomizableObjectEditor/Private/MuCOE/GenerateMutableSource/GenerateMutableSourceImage.cpp
@@ -64,11 +64,11 @@
 #define LOCTEXT_NAMESPACE "CustomizableObjectEditor"
 
 
-mu::Ptr<mu::NodeImage> ResizeTextureByNumMips(const mu::Ptr<mu::NodeImage>& ImageConstant, int32 MipsToSkip)
+UE::Mutable::Private::Ptr<UE::Mutable::Private::NodeImage> ResizeTextureByNumMips(const UE::Mutable::Private::Ptr<UE::Mutable::Private::NodeImage>& ImageConstant, int32 MipsToSkip)
 {
 	if (MipsToSkip > 0)
 	{
-		mu::Ptr<mu::NodeImageResize> ImageResize = new mu::NodeImageResize();
+		UE::Mutable::Private::Ptr<UE::Mutable::Private::NodeImageResize> ImageResize = new UE::Mutable::Private::NodeImageResize();
 		ImageResize->Base = ImageConstant;
 		ImageResize->bRelative = true;
 		const float Factor = FMath::Pow(0.5f, MipsToSkip);
@@ -81,7 +81,7 @@ mu::Ptr<mu::NodeImage> ResizeTextureByNumMips(const mu::Ptr<mu::NodeImage>& Imag
 }
 
 
-mu::Ptr<mu::NodeImage> GenerateMutableSourceImage(const UEdGraphPin* Pin, FMutableGraphGenerationContext& GenerationContext, int32 ReferenceTextureSize)
+UE::Mutable::Private::Ptr<UE::Mutable::Private::NodeImage> GenerateMutableSourceImage(const UEdGraphPin* Pin, FMutableGraphGenerationContext& GenerationContext, int32 ReferenceTextureSize)
 {
 	check(Pin)
 	RETURN_ON_CYCLE(*Pin, GenerationContext)
@@ -97,7 +97,7 @@ mu::Ptr<mu::NodeImage> GenerateMutableSourceImage(const UEdGraphPin* Pin, FMutab
 	
 	if (const FGeneratedData* Generated = GenerationContext.Generated.Find(Key))
 	{
-		return static_cast<mu::NodeImage*>(Generated->Node.get());
+		return static_cast<UE::Mutable::Private::NodeImage*>(Generated->Node.get());
 	}
 
 	if (Node->IsNodeOutDatedAndNeedsRefresh())
@@ -109,7 +109,7 @@ mu::Ptr<mu::NodeImage> GenerateMutableSourceImage(const UEdGraphPin* Pin, FMutab
 	// Most nodes need to be added to the cache but there are some that don't. For exampel, MacroInstanceNodes
 	bool bCacheNode = true;
 
-	mu::Ptr<mu::NodeImage> Result;
+	UE::Mutable::Private::Ptr<UE::Mutable::Private::NodeImage> Result;
 
 	const FCompilationOptions& CompilationOptions = GenerationContext.CompilationContext->Options;
 
@@ -120,19 +120,19 @@ mu::Ptr<mu::NodeImage> GenerateMutableSourceImage(const UEdGraphPin* Pin, FMutab
 		{
 			// Check the specific image cache
 			FGeneratedImageKey ImageKey = FGeneratedImageKey(Pin);
-			mu::Ptr<mu::NodeImage> ImageNode;
+			UE::Mutable::Private::Ptr<UE::Mutable::Private::NodeImage> ImageNode;
 			
 			const uint32 MipsToSkip = ComputeLODBiasForTexture(GenerationContext, *BaseTexture, nullptr, ReferenceTextureSize);
 
-			if (mu::Ptr<mu::NodeImage>* Cached = GenerationContext.GeneratedImages.Find(ImageKey))
+			if (UE::Mutable::Private::Ptr<UE::Mutable::Private::NodeImage>* Cached = GenerationContext.GeneratedImages.Find(ImageKey))
 			{
 				ImageNode = *Cached;
 			}
 			else
 			{
-				TSharedPtr<mu::FImage> ImageConstant = GenerateImageConstant(BaseTexture, GenerationContext, false);
+				TSharedPtr<UE::Mutable::Private::FImage> ImageConstant = GenerateImageConstant(BaseTexture, GenerationContext, false);
 
-				mu::Ptr<mu::NodeImageConstant> ConstantImageNode = new mu::NodeImageConstant;
+				UE::Mutable::Private::Ptr<UE::Mutable::Private::NodeImageConstant> ConstantImageNode = new UE::Mutable::Private::NodeImageConstant;
 				ConstantImageNode->SetValue( ImageConstant );
 				
 				const UTextureLODSettings& TextureLODSettings = GenerationContext.CompilationContext->Options.TargetPlatform->GetTextureLODSettings();
@@ -160,7 +160,7 @@ mu::Ptr<mu::NodeImage> GenerateMutableSourceImage(const UEdGraphPin* Pin, FMutab
 
 	else if (const UCustomizableObjectNodeTextureParameter* TypedNodeParam = Cast<UCustomizableObjectNodeTextureParameter>(Node))
 	{
-		mu::Ptr<mu::NodeImageParameter> TextureNode = new mu::NodeImageParameter();
+		UE::Mutable::Private::Ptr<UE::Mutable::Private::NodeImageParameter> TextureNode = new UE::Mutable::Private::NodeImageParameter();
 
 		TextureNode->Name = TypedNodeParam->GetParameterName(&GenerationContext.MacroNodesStack);
 		TextureNode->UID = GenerationContext.GetNodeIdUnique(Node).ToString();
@@ -175,20 +175,20 @@ mu::Ptr<mu::NodeImage> GenerateMutableSourceImage(const UEdGraphPin* Pin, FMutab
 			EMutableParameterType::Texture));
 
 		// Force the same format that the default texture if any.
-		mu::Ptr<mu::NodeImageFormat> FormatNode = new mu::NodeImageFormat();
+		UE::Mutable::Private::Ptr<UE::Mutable::Private::NodeImageFormat> FormatNode = new UE::Mutable::Private::NodeImageFormat();
 		// \TODO: Take it from default?
 		//if (TypedNodeParam->DefaultValue)
 		//{
-		//	FormatNode->SetFormat(mu::EImageFormat::RGBA_UByte);
+		//	FormatNode->SetFormat(UE::Mutable::Private::EImageFormat::RGBA_UByte);
 		//}
 		//else
 		{
 			// Force an "easy format" on the texture.
-			FormatNode->Format = mu::EImageFormat::RGBA_UByte;
+			FormatNode->Format = UE::Mutable::Private::EImageFormat::RGBA_UByte;
 		}
 		FormatNode->Source = TextureNode;
 
-		mu::Ptr<mu::NodeImageResize> ResizeNode = new mu::NodeImageResize();
+		UE::Mutable::Private::Ptr<UE::Mutable::Private::NodeImageResize> ResizeNode = new UE::Mutable::Private::NodeImageResize();
 		ResizeNode->Base = FormatNode;
 		ResizeNode->bRelative = false;
 
@@ -225,14 +225,14 @@ mu::Ptr<mu::NodeImage> GenerateMutableSourceImage(const UEdGraphPin* Pin, FMutab
 
 	else if (const ICustomizableObjectNodeMeshInterface* TypedNodeMesh = Cast<ICustomizableObjectNodeMeshInterface>(Node))
 	{
-		mu::Ptr<mu::NodeImageConstant> ImageNode = new mu::NodeImageConstant();
+		UE::Mutable::Private::Ptr<UE::Mutable::Private::NodeImageConstant> ImageNode = new UE::Mutable::Private::NodeImageConstant();
 		Result = ImageNode;
 
 		UTexture2D* Texture = TypedNodeMesh->FindTextureForPin(Pin);
 
 		if (Texture)
 		{
-			TSharedPtr<mu::FImage> ImageConstant = GenerateImageConstant(Texture, GenerationContext, false);
+			TSharedPtr<UE::Mutable::Private::FImage> ImageConstant = GenerateImageConstant(Texture, GenerationContext, false);
 			ImageNode->SetValue(ImageConstant);
 
 			const uint32 MipsToSkip = ComputeLODBiasForTexture(GenerationContext, *Texture, nullptr, ReferenceTextureSize);
@@ -250,18 +250,18 @@ mu::Ptr<mu::NodeImage> GenerateMutableSourceImage(const UEdGraphPin* Pin, FMutab
 		}
 		else
 		{
-			Result = mu::Ptr<mu::NodeImage>();
+			Result = UE::Mutable::Private::Ptr<UE::Mutable::Private::NodeImage>();
 		}
 	}
 
 	else if (const UCustomizableObjectNodeTextureInterpolate* TypedNodeInterp = Cast<UCustomizableObjectNodeTextureInterpolate>(Node))
 	{
-		mu::Ptr<mu::NodeImageInterpolate> ImageNode = new mu::NodeImageInterpolate();
+		UE::Mutable::Private::Ptr<UE::Mutable::Private::NodeImageInterpolate> ImageNode = new UE::Mutable::Private::NodeImageInterpolate();
 		Result = ImageNode;
 
 		if (const UEdGraphPin* ConnectedPin = FollowInputPin(*TypedNodeInterp->FactorPin()))
 		{
-			mu::Ptr<mu::NodeScalar> FactorNode = GenerateMutableSourceFloat(ConnectedPin, GenerationContext);
+			UE::Mutable::Private::Ptr<UE::Mutable::Private::NodeScalar> FactorNode = GenerateMutableSourceFloat(ConnectedPin, GenerationContext);
 			ImageNode->Factor = FactorNode;
 		}
 
@@ -270,7 +270,7 @@ mu::Ptr<mu::NodeImage> GenerateMutableSourceImage(const UEdGraphPin* Pin, FMutab
 		{
 			if (const UEdGraphPin* ConnectedPin = FollowInputPin(*TypedNodeInterp->Targets(LayerIndex)))
 			{
-				mu::Ptr<mu::NodeImage> TargetNode = GenerateMutableSourceImage(ConnectedPin, GenerationContext, ReferenceTextureSize);
+				UE::Mutable::Private::Ptr<UE::Mutable::Private::NodeImage> TargetNode = GenerateMutableSourceImage(ConnectedPin, GenerationContext, ReferenceTextureSize);
 
 				if (TargetNode)
 				{
@@ -293,40 +293,40 @@ mu::Ptr<mu::NodeImage> GenerateMutableSourceImage(const UEdGraphPin* Pin, FMutab
 		{
 			if (const UEdGraphPin* OtherPin = FollowInputPin(*TypedNodeLayer->LayerPin(LayerIndex)))
 			{
-				mu::Ptr<mu::NodeImage> MaskNode = nullptr;
+				UE::Mutable::Private::Ptr<UE::Mutable::Private::NodeImage> MaskNode = nullptr;
 				if (const UEdGraphPin* ConnectedPin = FollowInputPin(*TypedNodeLayer->MaskPin(LayerIndex)))
 				{
 					MaskNode = GenerateMutableSourceImage(ConnectedPin, GenerationContext, ReferenceTextureSize);
 				}
 
-				mu::EBlendType Type = mu::EBlendType::BT_BLEND;
+				UE::Mutable::Private::EBlendType Type = UE::Mutable::Private::EBlendType::BT_BLEND;
 				switch (TypedNodeLayer->Layers[LayerIndex].Effect)
 				{
-				case COTLE_MODULATE: Type = mu::EBlendType::BT_BLEND; break;
-				case COTLE_MULTIPLY: Type = mu::EBlendType::BT_MULTIPLY; break;
-				case COTLE_SOFTLIGHT: Type = mu::EBlendType::BT_SOFTLIGHT; break;
-				case COTLE_HARDLIGHT: Type = mu::EBlendType::BT_HARDLIGHT; break;
-				case COTLE_DODGE: Type = mu::EBlendType::BT_DODGE; break;
-				case COTLE_BURN: Type = mu::EBlendType::BT_BURN; break;
-				case COTLE_SCREEN: Type = mu::EBlendType::BT_SCREEN; break;
-				case COTLE_OVERLAY: Type = mu::EBlendType::BT_OVERLAY; break;
-				case COTLE_ALPHA_OVERLAY: Type = mu::EBlendType::BT_LIGHTEN; break;
-				case COTLE_NORMAL_COMBINE: Type = mu::EBlendType::BT_NORMAL_COMBINE; break;
+				case COTLE_MODULATE: Type = UE::Mutable::Private::EBlendType::BT_BLEND; break;
+				case COTLE_MULTIPLY: Type = UE::Mutable::Private::EBlendType::BT_MULTIPLY; break;
+				case COTLE_SOFTLIGHT: Type = UE::Mutable::Private::EBlendType::BT_SOFTLIGHT; break;
+				case COTLE_HARDLIGHT: Type = UE::Mutable::Private::EBlendType::BT_HARDLIGHT; break;
+				case COTLE_DODGE: Type = UE::Mutable::Private::EBlendType::BT_DODGE; break;
+				case COTLE_BURN: Type = UE::Mutable::Private::EBlendType::BT_BURN; break;
+				case COTLE_SCREEN: Type = UE::Mutable::Private::EBlendType::BT_SCREEN; break;
+				case COTLE_OVERLAY: Type = UE::Mutable::Private::EBlendType::BT_OVERLAY; break;
+				case COTLE_ALPHA_OVERLAY: Type = UE::Mutable::Private::EBlendType::BT_LIGHTEN; break;
+				case COTLE_NORMAL_COMBINE: Type = UE::Mutable::Private::EBlendType::BT_NORMAL_COMBINE; break;
 				default:
 					GenerationContext.Log(LOCTEXT("UnsupportedImageEffect", "Texture layer effect not supported. Setting to 'Blend'."), Node);
 					break;
 				}
 
-				if (Type == mu::EBlendType::BT_BLEND && !MaskNode)
+				if (Type == UE::Mutable::Private::EBlendType::BT_BLEND && !MaskNode)
 				{
 					GenerationContext.Log(LOCTEXT("ModulateWithoutMask", "Texture layer effect uses Modulate without a mask. It will replace everything below it!"), Node);
 				}
 				
 				if (OtherPin->PinType.PinCategory == Schema->PC_Texture)
 				{
-					mu::Ptr<mu::NodeImage> BlendNode = GenerateMutableSourceImage(OtherPin, GenerationContext, ReferenceTextureSize);
+					UE::Mutable::Private::Ptr<UE::Mutable::Private::NodeImage> BlendNode = GenerateMutableSourceImage(OtherPin, GenerationContext, ReferenceTextureSize);
 
-					mu::Ptr<mu::NodeImageLayer> LayerNode = new mu::NodeImageLayer;
+					UE::Mutable::Private::Ptr<UE::Mutable::Private::NodeImageLayer> LayerNode = new UE::Mutable::Private::NodeImageLayer;
 					LayerNode->Type = Type;
 					LayerNode->Base = Result;
 					LayerNode->Blended = BlendNode;
@@ -336,9 +336,9 @@ mu::Ptr<mu::NodeImage> GenerateMutableSourceImage(const UEdGraphPin* Pin, FMutab
 
 				else if (OtherPin->PinType.PinCategory == Schema->PC_Color)
 				{
-					mu::Ptr<mu::NodeColour> ColorNode = GenerateMutableSourceColor(OtherPin, GenerationContext);
+					UE::Mutable::Private::Ptr<UE::Mutable::Private::NodeColour> ColorNode = GenerateMutableSourceColor(OtherPin, GenerationContext);
 
-					mu::Ptr<mu::NodeImageLayerColour> LayerNode = new mu::NodeImageLayerColour;
+					UE::Mutable::Private::Ptr<UE::Mutable::Private::NodeImageLayerColour> LayerNode = new UE::Mutable::Private::NodeImageLayerColour;
 					LayerNode->Type = Type;
 					LayerNode->Base = Result;
 					LayerNode->Colour = ColorNode;
@@ -361,7 +361,7 @@ mu::Ptr<mu::NodeImage> GenerateMutableSourceImage(const UEdGraphPin* Pin, FMutab
 			// Check Switch Parameter arity preconditions.
 			if (const UEdGraphPin* ConnectedPin = FollowInputPin(*SwitchParameter))
 			{
-				mu::Ptr<mu::NodeScalar> SwitchParam = GenerateMutableSourceFloat(ConnectedPin, GenerationContext);
+				UE::Mutable::Private::Ptr<UE::Mutable::Private::NodeScalar> SwitchParam = GenerateMutableSourceFloat(ConnectedPin, GenerationContext);
 				// Switch Param not generated
 				if (!SwitchParam)
 				{
@@ -371,7 +371,7 @@ mu::Ptr<mu::NodeImage> GenerateMutableSourceImage(const UEdGraphPin* Pin, FMutab
 					return Result;
 				}
 
-				if (SwitchParam->GetType() != mu::NodeScalarEnumParameter::GetStaticType())
+				if (SwitchParam->GetType() != UE::Mutable::Private::NodeScalarEnumParameter::GetStaticType())
 				{
 					const FText Message = LOCTEXT("WrongSwitchParamType", "Switch parameter of incorrect type.");
 					GenerationContext.Log(Message, Node);
@@ -381,14 +381,14 @@ mu::Ptr<mu::NodeImage> GenerateMutableSourceImage(const UEdGraphPin* Pin, FMutab
 
 				const int32 NumSwitchOptions = TypedNodeTextureSwitch->GetNumElements();
 
-				mu::NodeScalarEnumParameter* EnumParameter = static_cast<mu::NodeScalarEnumParameter*>(SwitchParam.get());
+				UE::Mutable::Private::NodeScalarEnumParameter* EnumParameter = static_cast<UE::Mutable::Private::NodeScalarEnumParameter*>(SwitchParam.get());
 				if (NumSwitchOptions != EnumParameter->Options.Num())
 				{
 					const FText Message = LOCTEXT("MismatchedSwitch", "Switch enum and switch node have different number of options. Please refresh the switch node to make sure the outcomes are labeled properly.");
 					GenerationContext.Log(Message, Node);
 				}
 
-				mu::Ptr<mu::NodeImageSwitch> SwitchNode = new mu::NodeImageSwitch;
+				UE::Mutable::Private::Ptr<UE::Mutable::Private::NodeImageSwitch> SwitchNode = new UE::Mutable::Private::NodeImageSwitch;
 				SwitchNode->Parameter = SwitchParam;
 				SwitchNode->Options.SetNum(NumSwitchOptions);
 
@@ -420,12 +420,12 @@ mu::Ptr<mu::NodeImage> GenerateMutableSourceImage(const UEdGraphPin* Pin, FMutab
 	else if (const UCustomizableObjectNodeTextureVariation* TypedNodeImageVar = Cast<const UCustomizableObjectNodeTextureVariation>(Node))
 	{
 		// UCustomizableObjectNodePassThroughTextureVariation nodes are also handled here
-		mu::Ptr<mu::NodeImageVariation> TextureNode = new mu::NodeImageVariation();
+		UE::Mutable::Private::Ptr<UE::Mutable::Private::NodeImageVariation> TextureNode = new UE::Mutable::Private::NodeImageVariation();
 		Result = TextureNode;
 
 		if (const UEdGraphPin* ConnectedPin = FollowInputPin(*TypedNodeImageVar->DefaultPin()))
 		{
-			mu::Ptr<mu::NodeImage> ChildNode = GenerateMutableSourceImage(ConnectedPin, GenerationContext, ReferenceTextureSize);
+			UE::Mutable::Private::Ptr<UE::Mutable::Private::NodeImage> ChildNode = GenerateMutableSourceImage(ConnectedPin, GenerationContext, ReferenceTextureSize);
 			if (ChildNode)
 			{
 				TextureNode->DefaultImage = ChildNode;
@@ -451,7 +451,7 @@ mu::Ptr<mu::NodeImage> GenerateMutableSourceImage(const UEdGraphPin* Pin, FMutab
 
 			if (const UEdGraphPin* ConnectedPin = FollowInputPin(*VariationPin))
 			{
-				mu::Ptr<mu::NodeImage> ChildNode = GenerateMutableSourceImage(ConnectedPin, GenerationContext, ReferenceTextureSize);
+				UE::Mutable::Private::Ptr<UE::Mutable::Private::NodeImage> ChildNode = GenerateMutableSourceImage(ConnectedPin, GenerationContext, ReferenceTextureSize);
 				TextureNode->Variations[VariationIndex].Image = ChildNode;
 			}
 		}
@@ -459,14 +459,14 @@ mu::Ptr<mu::NodeImage> GenerateMutableSourceImage(const UEdGraphPin* Pin, FMutab
 
 	else if (const UCustomizableObjectNodeTextureFromColor* TypedNodeFromColor = Cast<UCustomizableObjectNodeTextureFromColor>(Node))
 	{
-		mu::Ptr<mu::NodeColour> color;
+		UE::Mutable::Private::Ptr<UE::Mutable::Private::NodeColour> color;
 
 		if (const UEdGraphPin* ConnectedPin = FollowInputPin(*TypedNodeFromColor->ColorPin()))
 		{
 			color = GenerateMutableSourceColor(ConnectedPin, GenerationContext);
 		}
 
-		mu::Ptr<mu::NodeImagePlainColour> ImageFromColour = new mu::NodeImagePlainColour;
+		UE::Mutable::Private::Ptr<UE::Mutable::Private::NodeImagePlainColour> ImageFromColour = new UE::Mutable::Private::NodeImagePlainColour;
 		Result = ImageFromColour;
 
 		if (color)
@@ -483,7 +483,7 @@ mu::Ptr<mu::NodeImage> GenerateMutableSourceImage(const UEdGraphPin* Pin, FMutab
 
 	else if (const UCustomizableObjectNodeTextureFromFloats* TypedNodeFromFloats = Cast<UCustomizableObjectNodeTextureFromFloats>(Node))
 	{		
-		mu::Ptr<mu::NodeColourFromScalars> Color = new mu::NodeColourFromScalars();
+		UE::Mutable::Private::Ptr<UE::Mutable::Private::NodeColourFromScalars> Color = new UE::Mutable::Private::NodeColourFromScalars();
 
 		if (const UEdGraphPin* RPin = FollowInputPin(*TypedNodeFromFloats->RPin()))
 		{
@@ -505,7 +505,7 @@ mu::Ptr<mu::NodeImage> GenerateMutableSourceImage(const UEdGraphPin* Pin, FMutab
 			Color->W = GenerateMutableSourceFloat(APin, GenerationContext);
 		}
 
-		mu::Ptr<mu::NodeImagePlainColour> ImageFromColour = new mu::NodeImagePlainColour;
+		UE::Mutable::Private::Ptr<UE::Mutable::Private::NodeImagePlainColour> ImageFromColour = new UE::Mutable::Private::NodeImagePlainColour;
 		Result = ImageFromColour;
 
 		if (Color)
@@ -522,7 +522,7 @@ mu::Ptr<mu::NodeImage> GenerateMutableSourceImage(const UEdGraphPin* Pin, FMutab
 
 	else if (const UCustomizableObjectNodeTextureFromChannels* TypedNodeFrom = Cast<UCustomizableObjectNodeTextureFromChannels>(Node))
 	{
-		mu::Ptr<mu::NodeImage> RNode, GNode, BNode, ANode;
+		UE::Mutable::Private::Ptr<UE::Mutable::Private::NodeImage> RNode, GNode, BNode, ANode;
 		bool RGB = false;
 
 		if (const UEdGraphPin* ConnectedPin = FollowInputPin(*TypedNodeFrom->RPin()))
@@ -545,12 +545,12 @@ mu::Ptr<mu::NodeImage> GenerateMutableSourceImage(const UEdGraphPin* Pin, FMutab
 			ANode = GenerateMutableSourceImage(ConnectedPin, GenerationContext, ReferenceTextureSize);
 		}
 
-		mu::Ptr<mu::NodeImageSwizzle> SwizzleNode = new mu::NodeImageSwizzle;
+		UE::Mutable::Private::Ptr<UE::Mutable::Private::NodeImageSwizzle> SwizzleNode = new UE::Mutable::Private::NodeImageSwizzle;
 		Result = SwizzleNode;
 
 		if (RGB && ANode)
 		{
-			SwizzleNode->SetFormat(mu::EImageFormat::RGBA_UByte);
+			SwizzleNode->SetFormat(UE::Mutable::Private::EImageFormat::RGBA_UByte);
 			SwizzleNode->Sources[0] = RNode;
 			SwizzleNode->SourceChannels[0] = 0;
 			SwizzleNode->Sources[1] = GNode;
@@ -562,7 +562,7 @@ mu::Ptr<mu::NodeImage> GenerateMutableSourceImage(const UEdGraphPin* Pin, FMutab
 		}
 		else if (RGB)
 		{
-			SwizzleNode->SetFormat(mu::EImageFormat::RGB_UByte);
+			SwizzleNode->SetFormat(UE::Mutable::Private::EImageFormat::RGB_UByte);
 			SwizzleNode->Sources[0] = RNode;
 			SwizzleNode->SourceChannels[0] = 0;
 			SwizzleNode->Sources[1] = GNode;
@@ -572,13 +572,13 @@ mu::Ptr<mu::NodeImage> GenerateMutableSourceImage(const UEdGraphPin* Pin, FMutab
 		}
 		else if (RNode)
 		{
-			SwizzleNode->SetFormat(mu::EImageFormat::L_UByte);
+			SwizzleNode->SetFormat(UE::Mutable::Private::EImageFormat::L_UByte);
 			SwizzleNode->Sources[0] = RNode;
 			SwizzleNode->SourceChannels[0] = 0;
 		}
 		else if (ANode)
 		{
-			SwizzleNode->SetFormat(mu::EImageFormat::L_UByte);
+			SwizzleNode->SetFormat(UE::Mutable::Private::EImageFormat::L_UByte);
 			SwizzleNode->Sources[0] = ANode;
 			SwizzleNode->SourceChannels[0] = 0;
 		}
@@ -586,15 +586,15 @@ mu::Ptr<mu::NodeImage> GenerateMutableSourceImage(const UEdGraphPin* Pin, FMutab
 
 	else if (const UCustomizableObjectNodeTextureToChannels* TypedNodeTo = Cast<UCustomizableObjectNodeTextureToChannels>(Node))
 	{
-		mu::Ptr<mu::NodeImage> BaseNode;
+		UE::Mutable::Private::Ptr<UE::Mutable::Private::NodeImage> BaseNode;
 		if (const UEdGraphPin* ConnectedPin = FollowInputPin(*TypedNodeTo->InputPin()))
 		{
 			BaseNode = GenerateMutableSourceImage(ConnectedPin, GenerationContext, ReferenceTextureSize);
 		}
 
-		mu::Ptr<mu::NodeImageSwizzle> SwizzleNode = new mu::NodeImageSwizzle;
+		UE::Mutable::Private::Ptr<UE::Mutable::Private::NodeImageSwizzle> SwizzleNode = new UE::Mutable::Private::NodeImageSwizzle;
 		Result = SwizzleNode;
-		SwizzleNode->SetFormat(mu::EImageFormat::L_UByte);
+		SwizzleNode->SetFormat(UE::Mutable::Private::EImageFormat::L_UByte);
 		SwizzleNode->Sources[0] = BaseNode;
 
 		if (Pin == TypedNodeTo->RPin())
@@ -621,7 +621,7 @@ mu::Ptr<mu::NodeImage> GenerateMutableSourceImage(const UEdGraphPin* Pin, FMutab
 
 	else if (const UCustomizableObjectNodeTextureProject* TypedNodeProject = Cast<UCustomizableObjectNodeTextureProject>(Node))
 	{
-		mu::Ptr<mu::NodeImageProject> ImageNode = new mu::NodeImageProject();
+		UE::Mutable::Private::Ptr<UE::Mutable::Private::NodeImageProject> ImageNode = new UE::Mutable::Private::NodeImageProject();
 		Result = ImageNode;
 
 		if (!FollowInputPin(*TypedNodeProject->MeshPin()))
@@ -661,37 +661,37 @@ mu::Ptr<mu::NodeImage> GenerateMutableSourceImage(const UEdGraphPin* Pin, FMutab
 		ImageNode->bIsRGBFadingEnabled = TypedNodeProject->bEnableAngleFadeOutForRGB;
 		ImageNode->bIsAlphaFadingEnabled = TypedNodeProject->bEnableAngleFadeOutForAlpha;
 		ImageNode->SamplingMethod = Invoke(
-			[](ETextureProjectSamplingMethod Method) -> mu::ESamplingMethod 
+			[](ETextureProjectSamplingMethod Method) -> UE::Mutable::Private::ESamplingMethod 
 			{
 				switch(Method)
 				{
-				case ETextureProjectSamplingMethod::Point:    return mu::ESamplingMethod::Point;
-				case ETextureProjectSamplingMethod::BiLinear: return mu::ESamplingMethod::BiLinear;
-				default: { check(false); return mu::ESamplingMethod::Point; }
+				case ETextureProjectSamplingMethod::Point:    return UE::Mutable::Private::ESamplingMethod::Point;
+				case ETextureProjectSamplingMethod::BiLinear: return UE::Mutable::Private::ESamplingMethod::BiLinear;
+				default: { check(false); return UE::Mutable::Private::ESamplingMethod::Point; }
 				}
 			}, 
 			TypedNodeProject->SamplingMethod);
 
 		ImageNode->MinFilterMethod = Invoke(
-			[](ETextureProjectMinFilterMethod Method) -> mu::EMinFilterMethod 
+			[](ETextureProjectMinFilterMethod Method) -> UE::Mutable::Private::EMinFilterMethod 
 			{
 				switch(Method)
 				{
-				case ETextureProjectMinFilterMethod::None:               return mu::EMinFilterMethod::None;
-				case ETextureProjectMinFilterMethod::TotalAreaHeuristic: return mu::EMinFilterMethod::TotalAreaHeuristic;
-				default: { check(false); return mu::EMinFilterMethod::None; }
+				case ETextureProjectMinFilterMethod::None:               return UE::Mutable::Private::EMinFilterMethod::None;
+				case ETextureProjectMinFilterMethod::TotalAreaHeuristic: return UE::Mutable::Private::EMinFilterMethod::TotalAreaHeuristic;
+				default: { check(false); return UE::Mutable::Private::EMinFilterMethod::None; }
 				}
 			}, 
 			TypedNodeProject->MinFilterMethod);
 
 		if (const UEdGraphPin* ConnectedPin = FollowInputPin(*TypedNodeProject->AngleFadeStartPin()))
 		{
-			mu::Ptr<mu::NodeScalar> ScalarNode = GenerateMutableSourceFloat(ConnectedPin, GenerationContext);
+			UE::Mutable::Private::Ptr<UE::Mutable::Private::NodeScalar> ScalarNode = GenerateMutableSourceFloat(ConnectedPin, GenerationContext);
 			ImageNode->AngleFadeStart = ScalarNode;
 		}
 		else
 		{
-			mu::Ptr<mu::NodeScalarConstant> ScalarNode = new mu::NodeScalarConstant;
+			UE::Mutable::Private::Ptr<UE::Mutable::Private::NodeScalarConstant> ScalarNode = new UE::Mutable::Private::NodeScalarConstant;
 			float value = FCString::Atof(*TypedNodeProject->AngleFadeStartPin()->DefaultValue);
 			ScalarNode->Value = value;
 			ImageNode->AngleFadeStart = ScalarNode;
@@ -699,12 +699,12 @@ mu::Ptr<mu::NodeImage> GenerateMutableSourceImage(const UEdGraphPin* Pin, FMutab
 
 		if (const UEdGraphPin* ConnectedPin = FollowInputPin(*TypedNodeProject->AngleFadeEndPin()))
 		{
-			mu::Ptr<mu::NodeScalar> ScalarNode = GenerateMutableSourceFloat(ConnectedPin, GenerationContext);
+			UE::Mutable::Private::Ptr<UE::Mutable::Private::NodeScalar> ScalarNode = GenerateMutableSourceFloat(ConnectedPin, GenerationContext);
 			ImageNode->AngleFadeEnd = ScalarNode;
 		}
 		else
 		{
-			mu::Ptr<mu::NodeScalarConstant> ScalarNode = new mu::NodeScalarConstant;
+			UE::Mutable::Private::Ptr<UE::Mutable::Private::NodeScalarConstant> ScalarNode = new UE::Mutable::Private::NodeScalarConstant;
 			float value = FCString::Atof(*TypedNodeProject->AngleFadeEndPin()->DefaultValue);
 			ScalarNode->Value = value;
 			ImageNode->AngleFadeEnd = ScalarNode;
@@ -715,20 +715,20 @@ mu::Ptr<mu::NodeImage> GenerateMutableSourceImage(const UEdGraphPin* Pin, FMutab
 			FLayoutGenerationFlags LayoutGenerationFlags;
 			LayoutGenerationFlags.TexturePinModes.Init(EPinMode::Mutable, TEXSTREAM_MAX_NUM_UVCHANNELS);
 			GenerationContext.LayoutGenerationFlags.Push(LayoutGenerationFlags);
-			mu::Ptr<mu::NodeMesh> MeshNode = GenerateMutableSourceMesh(ConnectedPin, GenerationContext, FMutableSourceMeshData(), false, false);
+			UE::Mutable::Private::Ptr<UE::Mutable::Private::NodeMesh> MeshNode = GenerateMutableSourceMesh(ConnectedPin, GenerationContext, FMutableSourceMeshData(), false, false);
 			GenerationContext.LayoutGenerationFlags.Pop();
 			ImageNode->Mesh = MeshNode;
 		}
 
 		if (const UEdGraphPin* ConnectedPin = FollowInputPin(*TypedNodeProject->MeshMaskPin()))
 		{
-			mu::Ptr<mu::NodeImage> MeshMaskNode = GenerateMutableSourceImage(ConnectedPin, GenerationContext, ReferenceTextureSize);
+			UE::Mutable::Private::Ptr<UE::Mutable::Private::NodeImage> MeshMaskNode = GenerateMutableSourceImage(ConnectedPin, GenerationContext, ReferenceTextureSize);
 			ImageNode->Mask = MeshMaskNode;
 		}
 
 		if (const UEdGraphPin* ConnectedPin = FollowInputPin(*TypedNodeProject->ProjectorPin()))
 		{
-			mu::Ptr<mu::NodeProjector> ProjectorNode = GenerateMutableSourceProjector(ConnectedPin, GenerationContext);
+			UE::Mutable::Private::Ptr<UE::Mutable::Private::NodeProjector> ProjectorNode = GenerateMutableSourceProjector(ConnectedPin, GenerationContext);
 			ImageNode->Projector = ProjectorNode;
 		}
 
@@ -746,104 +746,104 @@ mu::Ptr<mu::NodeImage> GenerateMutableSourceImage(const UEdGraphPin* Pin, FMutab
 
 		if (const UEdGraphPin* ConnectedPin = FollowInputPin(*TypedNodeProject->TexturePins(TexIndex)))
 		{
-			mu::Ptr<mu::NodeImage> SourceNode = GenerateMutableSourceImage(ConnectedPin, GenerationContext, FMath::Max(TextureSize.X, TextureSize.Y));
+			UE::Mutable::Private::Ptr<UE::Mutable::Private::NodeImage> SourceNode = GenerateMutableSourceImage(ConnectedPin, GenerationContext, FMath::Max(TextureSize.X, TextureSize.Y));
 			ImageNode->Image = SourceNode;
 		}
 	}
 
 	else if (const UCustomizableObjectNodeTextureBinarise* TypedNodeTexBin = Cast<UCustomizableObjectNodeTextureBinarise>(Node))
 	{
-		mu::Ptr<mu::NodeImageBinarise> BinariseNode = new mu::NodeImageBinarise();
+		UE::Mutable::Private::Ptr<UE::Mutable::Private::NodeImageBinarise> BinariseNode = new UE::Mutable::Private::NodeImageBinarise();
 		Result = BinariseNode;
 
 		if (const UEdGraphPin* ConnectedPin = FollowInputPin(*TypedNodeTexBin->GetBaseImagePin()))
 		{
-			mu::Ptr<mu::NodeImage> BaseImageNode = GenerateMutableSourceImage(ConnectedPin, GenerationContext, ReferenceTextureSize);
+			UE::Mutable::Private::Ptr<UE::Mutable::Private::NodeImage> BaseImageNode = GenerateMutableSourceImage(ConnectedPin, GenerationContext, ReferenceTextureSize);
 			BinariseNode->Base = BaseImageNode;
 		}
 
 		if (const UEdGraphPin* ConnectedPin = FollowInputPin(*TypedNodeTexBin->GetThresholdPin()))
 		{
-			mu::Ptr<mu::NodeScalar> ThresholdNode = GenerateMutableSourceFloat(ConnectedPin, GenerationContext);
+			UE::Mutable::Private::Ptr<UE::Mutable::Private::NodeScalar> ThresholdNode = GenerateMutableSourceFloat(ConnectedPin, GenerationContext);
 			BinariseNode->Threshold = ThresholdNode;
 		}
 	}
 
 	else if (const UCustomizableObjectNodeTextureInvert* TypedNodeTexInv = Cast<UCustomizableObjectNodeTextureInvert>(Node))
 	{
-		mu::Ptr<mu::NodeImageInvert> InvertNode = new mu::NodeImageInvert();
+		UE::Mutable::Private::Ptr<UE::Mutable::Private::NodeImageInvert> InvertNode = new UE::Mutable::Private::NodeImageInvert();
 		Result = InvertNode;
 
 		if (const UEdGraphPin* ConnectedPin = FollowInputPin(*TypedNodeTexInv->GetBaseImagePin()))
 		{
-			mu::Ptr<mu::NodeImage> BaseImageNode = GenerateMutableSourceImage(ConnectedPin, GenerationContext, ReferenceTextureSize);
+			UE::Mutable::Private::Ptr<UE::Mutable::Private::NodeImage> BaseImageNode = GenerateMutableSourceImage(ConnectedPin, GenerationContext, ReferenceTextureSize);
 			InvertNode->Base = BaseImageNode;
 		}
 	}
 
 	else if (const UCustomizableObjectNodeTextureColourMap* TypedNodeColourMap = Cast<UCustomizableObjectNodeTextureColourMap>(Node))
 	{
-		mu::Ptr<mu::NodeImageColourMap> ColourMapNode = new mu::NodeImageColourMap();
+		UE::Mutable::Private::Ptr<UE::Mutable::Private::NodeImageColourMap> ColourMapNode = new UE::Mutable::Private::NodeImageColourMap();
 
 		Result = ColourMapNode;
 
 		if (const UEdGraphPin* ConnectedPin = FollowInputPin(*TypedNodeColourMap->GetMapPin()))
 		{
-			mu::Ptr<mu::NodeImage> GradientImageNode = GenerateMutableSourceImage(ConnectedPin, GenerationContext, ReferenceTextureSize);
+			UE::Mutable::Private::Ptr<UE::Mutable::Private::NodeImage> GradientImageNode = GenerateMutableSourceImage(ConnectedPin, GenerationContext, ReferenceTextureSize);
 			ColourMapNode->Map = GradientImageNode;
 		}
 
 		if (const UEdGraphPin* ConnectedPin = FollowInputPin(*TypedNodeColourMap->GetMaskPin()))
 		{
-			mu::Ptr<mu::NodeImage> GradientImageNode = GenerateMutableSourceImage(ConnectedPin, GenerationContext, ReferenceTextureSize);
+			UE::Mutable::Private::Ptr<UE::Mutable::Private::NodeImage> GradientImageNode = GenerateMutableSourceImage(ConnectedPin, GenerationContext, ReferenceTextureSize);
 			ColourMapNode->Mask = GradientImageNode;
 		}
 
 		if (const UEdGraphPin* ConnectedPin = FollowInputPin(*TypedNodeColourMap->GetBasePin()))
 		{
-			mu::Ptr<mu::NodeImage> SourceImageNode = GenerateMutableSourceImage(ConnectedPin, GenerationContext, ReferenceTextureSize);
+			UE::Mutable::Private::Ptr<UE::Mutable::Private::NodeImage> SourceImageNode = GenerateMutableSourceImage(ConnectedPin, GenerationContext, ReferenceTextureSize);
 			ColourMapNode->Base = SourceImageNode;
 		}
 	}
 
 	else if (const UCustomizableObjectNodeTextureTransform* TypedNodeTransform = Cast<UCustomizableObjectNodeTextureTransform>(Node))
 	{
-		mu::Ptr<mu::NodeImageTransform> TransformNode = new mu::NodeImageTransform();
+		UE::Mutable::Private::Ptr<UE::Mutable::Private::NodeImageTransform> TransformNode = new UE::Mutable::Private::NodeImageTransform();
 		Result = TransformNode;
 		
 		if ( UEdGraphPin* BaseImagePin = FollowInputPin(*TypedNodeTransform->GetBaseImagePin()) )
 		{
-			mu::Ptr<mu::NodeImage> ImageNode = GenerateMutableSourceImage( BaseImagePin, GenerationContext, ReferenceTextureSize);
+			UE::Mutable::Private::Ptr<UE::Mutable::Private::NodeImage> ImageNode = GenerateMutableSourceImage( BaseImagePin, GenerationContext, ReferenceTextureSize);
 			TransformNode->Base = ImageNode;
 		}
 
 		if ( UEdGraphPin* OffsetXPin = FollowInputPin(*TypedNodeTransform->GetOffsetXPin()) )
 		{
-			mu::Ptr<mu::NodeScalar> OffsetXNode = GenerateMutableSourceFloat(OffsetXPin, GenerationContext);
+			UE::Mutable::Private::Ptr<UE::Mutable::Private::NodeScalar> OffsetXNode = GenerateMutableSourceFloat(OffsetXPin, GenerationContext);
 			TransformNode->OffsetX = OffsetXNode;
 		}
 
 		if ( UEdGraphPin* OffsetYPin = FollowInputPin(*TypedNodeTransform->GetOffsetYPin()) )
 		{
-			mu::Ptr<mu::NodeScalar> OffsetYNode = GenerateMutableSourceFloat(OffsetYPin, GenerationContext);
+			UE::Mutable::Private::Ptr<UE::Mutable::Private::NodeScalar> OffsetYNode = GenerateMutableSourceFloat(OffsetYPin, GenerationContext);
 			TransformNode->OffsetY = OffsetYNode;
 		}
 
 		if ( UEdGraphPin* ScaleXPin = FollowInputPin(*TypedNodeTransform->GetScaleXPin()) )
 		{
-			mu::Ptr<mu::NodeScalar> ScaleXNode = GenerateMutableSourceFloat(ScaleXPin, GenerationContext);
+			UE::Mutable::Private::Ptr<UE::Mutable::Private::NodeScalar> ScaleXNode = GenerateMutableSourceFloat(ScaleXPin, GenerationContext);
 			TransformNode->ScaleX = ScaleXNode;
 		}
 
 		if ( UEdGraphPin* ScaleYPin = FollowInputPin(*TypedNodeTransform->GetScaleYPin()) )
 		{
-			mu::Ptr<mu::NodeScalar> ScaleYNode = GenerateMutableSourceFloat(ScaleYPin, GenerationContext);
+			UE::Mutable::Private::Ptr<UE::Mutable::Private::NodeScalar> ScaleYNode = GenerateMutableSourceFloat(ScaleYPin, GenerationContext);
 			TransformNode->ScaleY = ScaleYNode; 
 		}
 
 		if ( UEdGraphPin* RotationPin = FollowInputPin(*TypedNodeTransform->GetRotationPin()) )
 		{
-			mu::Ptr<mu::NodeScalar> RotationNode = GenerateMutableSourceFloat(RotationPin, GenerationContext);
+			UE::Mutable::Private::Ptr<UE::Mutable::Private::NodeScalar> RotationNode = GenerateMutableSourceFloat(RotationPin, GenerationContext);
 			TransformNode->Rotation = RotationNode; 
 		}
 
@@ -851,10 +851,10 @@ mu::Ptr<mu::NodeImage> GenerateMutableSourceImage(const UEdGraphPin* Pin, FMutab
 		{
 			switch (TypedNodeTransform->AddressMode)
 			{
-			case ETextureTransformAddressMode::Wrap:		 return mu::EAddressMode::Wrap;
-			case ETextureTransformAddressMode::ClampToEdge:  return mu::EAddressMode::ClampToEdge;
-			case ETextureTransformAddressMode::ClampToBlack: return mu::EAddressMode::ClampToBlack;
-			default: { check(false); return mu::EAddressMode::None; }
+			case ETextureTransformAddressMode::Wrap:		 return UE::Mutable::Private::EAddressMode::Wrap;
+			case ETextureTransformAddressMode::ClampToEdge:  return UE::Mutable::Private::EAddressMode::ClampToEdge;
+			case ETextureTransformAddressMode::ClampToBlack: return UE::Mutable::Private::EAddressMode::ClampToBlack;
+			default: { check(false); return UE::Mutable::Private::EAddressMode::None; }
 			}
 		});
 
@@ -888,18 +888,18 @@ mu::Ptr<mu::NodeImage> GenerateMutableSourceImage(const UEdGraphPin* Pin, FMutab
 
 	else if (const UCustomizableObjectNodeTextureSaturate* TypedNodeSaturate = Cast<UCustomizableObjectNodeTextureSaturate>(Node))
 	{
-		mu::Ptr<mu::NodeImageSaturate> SaturateNode = new mu::NodeImageSaturate();
+		UE::Mutable::Private::Ptr<UE::Mutable::Private::NodeImageSaturate> SaturateNode = new UE::Mutable::Private::NodeImageSaturate();
 		Result = SaturateNode;
 	
 		if ( UEdGraphPin* BaseImagePin = FollowInputPin(*TypedNodeSaturate->GetBaseImagePin()) )
 		{
-			mu::Ptr<mu::NodeImage> ImageNode = GenerateMutableSourceImage(BaseImagePin, GenerationContext, ReferenceTextureSize);
+			UE::Mutable::Private::Ptr<UE::Mutable::Private::NodeImage> ImageNode = GenerateMutableSourceImage(BaseImagePin, GenerationContext, ReferenceTextureSize);
 			SaturateNode->Source = ImageNode;
 		}
 
 		if ( UEdGraphPin* FactorPin = FollowInputPin(*TypedNodeSaturate->GetFactorPin()))
 		{
-			mu::Ptr<mu::NodeScalar> FactorNode = GenerateMutableSourceFloat(FactorPin, GenerationContext);
+			UE::Mutable::Private::Ptr<UE::Mutable::Private::NodeScalar> FactorNode = GenerateMutableSourceFloat(FactorPin, GenerationContext);
 			SaturateNode->Factor = FactorNode;
 		}
 	}
@@ -909,7 +909,7 @@ mu::Ptr<mu::NodeImage> GenerateMutableSourceImage(const UEdGraphPin* Pin, FMutab
 		UTexture* BaseTexture = TypedNodePassThroughTex->PassThroughTexture;
 		if (BaseTexture)
 		{
-			mu::Ptr<mu::NodeImageConstant> ImageNode = new mu::NodeImageConstant();
+			UE::Mutable::Private::Ptr<UE::Mutable::Private::NodeImageConstant> ImageNode = new UE::Mutable::Private::NodeImageConstant();
 			Result = ImageNode;
 
 			ImageNode->SetValue(GenerateImageConstant(BaseTexture, GenerationContext, true));
@@ -929,7 +929,7 @@ mu::Ptr<mu::NodeImage> GenerateMutableSourceImage(const UEdGraphPin* Pin, FMutab
 			// Check Switch Parameter arity preconditions.
 			if (const UEdGraphPin* ConnectedPin = FollowInputPin(*SwitchParameter))
 			{
-				mu::Ptr<mu::NodeScalar> SwitchParam = GenerateMutableSourceFloat(ConnectedPin, GenerationContext);
+				UE::Mutable::Private::Ptr<UE::Mutable::Private::NodeScalar> SwitchParam = GenerateMutableSourceFloat(ConnectedPin, GenerationContext);
 				// Switch Param not generated
 				if (!SwitchParam)
 				{
@@ -939,7 +939,7 @@ mu::Ptr<mu::NodeImage> GenerateMutableSourceImage(const UEdGraphPin* Pin, FMutab
 					return Result;
 				}
 
-				if (SwitchParam->GetType() != mu::NodeScalarEnumParameter::GetStaticType())
+				if (SwitchParam->GetType() != UE::Mutable::Private::NodeScalarEnumParameter::GetStaticType())
 				{
 					const FText Message = LOCTEXT("WrongSwitchParamType", "Switch parameter of incorrect type.");
 					GenerationContext.Log(Message, Node);
@@ -949,7 +949,7 @@ mu::Ptr<mu::NodeImage> GenerateMutableSourceImage(const UEdGraphPin* Pin, FMutab
 
 				const int32 NumSwitchOptions = TypedNodePassThroughTextureSwitch->GetNumElements();
 
-				mu::NodeScalarEnumParameter* EnumParameter = static_cast<mu::NodeScalarEnumParameter*>(SwitchParam.get());
+				UE::Mutable::Private::NodeScalarEnumParameter* EnumParameter = static_cast<UE::Mutable::Private::NodeScalarEnumParameter*>(SwitchParam.get());
 				if (NumSwitchOptions != EnumParameter->Options.Num())
 				{
 					const FText Message = LOCTEXT("MismatchedSwitch", "Switch enum and switch node have different number of options. Please refresh the switch node to make sure the outcomes are labeled properly.");
@@ -957,7 +957,7 @@ mu::Ptr<mu::NodeImage> GenerateMutableSourceImage(const UEdGraphPin* Pin, FMutab
 				}
 
 				// TODO: Implement Mutable core pass-through switch nodes
-				mu::Ptr<mu::NodeImageSwitch> SwitchNode = new mu::NodeImageSwitch;
+				UE::Mutable::Private::Ptr<UE::Mutable::Private::NodeImageSwitch> SwitchNode = new UE::Mutable::Private::NodeImageSwitch;
 				SwitchNode->Parameter = SwitchParam;
 				SwitchNode->Options.SetNum(NumSwitchOptions);
 
@@ -965,7 +965,7 @@ mu::Ptr<mu::NodeImage> GenerateMutableSourceImage(const UEdGraphPin* Pin, FMutab
 				{
 					if (const UEdGraphPin* TexturePin = FollowInputPin(*TypedNodePassThroughTextureSwitch->GetElementPin(SelectorIndex)))
 					{
-						mu::Ptr<mu::NodeImage> PassThroughImage = GenerateMutableSourceImage(TexturePin, GenerationContext, ReferenceTextureSize);
+						UE::Mutable::Private::Ptr<UE::Mutable::Private::NodeImage> PassThroughImage = GenerateMutableSourceImage(TexturePin, GenerationContext, ReferenceTextureSize);
 						SwitchNode->Options[SelectorIndex] = PassThroughImage;
 					}
 					else
@@ -990,9 +990,9 @@ mu::Ptr<mu::NodeImage> GenerateMutableSourceImage(const UEdGraphPin* Pin, FMutab
 	// If the node is a plain colour node, generate an image out of it
 	else if (Pin->PinType.PinCategory == Schema->PC_Color)
 	{
-		mu::Ptr<mu::NodeColour> ColorNode = GenerateMutableSourceColor(Pin, GenerationContext);
+		UE::Mutable::Private::Ptr<UE::Mutable::Private::NodeColour> ColorNode = GenerateMutableSourceColor(Pin, GenerationContext);
 
-		mu::Ptr<mu::NodeImagePlainColour> ImageNode = new mu::NodeImagePlainColour;
+		UE::Mutable::Private::Ptr<UE::Mutable::Private::NodeImagePlainColour> ImageNode = new UE::Mutable::Private::NodeImagePlainColour;
 		ImageNode->SizeX = 16;
 		ImageNode->SizeY = 16;
 		ImageNode->Colour = ColorNode;
@@ -1003,8 +1003,8 @@ mu::Ptr<mu::NodeImage> GenerateMutableSourceImage(const UEdGraphPin* Pin, FMutab
 	{
 		if (Pin->PinType.PinCategory == Schema->PC_Color)
 		{
-			mu::Ptr<mu::NodeColour> ColorNode = GenerateMutableSourceColor(Pin, GenerationContext);
-			mu::Ptr<mu::NodeImagePlainColour> ImageNode = new mu::NodeImagePlainColour;
+			UE::Mutable::Private::Ptr<UE::Mutable::Private::NodeColour> ColorNode = GenerateMutableSourceColor(Pin, GenerationContext);
+			UE::Mutable::Private::Ptr<UE::Mutable::Private::NodeImagePlainColour> ImageNode = new UE::Mutable::Private::NodeImagePlainColour;
 
 			ImageNode->SizeX = 16;
 			ImageNode->SizeY = 16;
@@ -1014,13 +1014,13 @@ mu::Ptr<mu::NodeImage> GenerateMutableSourceImage(const UEdGraphPin* Pin, FMutab
 		}
 		else
 		{
-			TObjectPtr<UDataTable> Table = MutablePrivate::LoadObject(TypedNodeTable->Table);
-			TObjectPtr<UScriptStruct> Structure = MutablePrivate::LoadObject(TypedNodeTable->Structure);
+			TObjectPtr<UDataTable> Table = UE::Mutable::Private::LoadObject(TypedNodeTable->Table);
+			TObjectPtr<UScriptStruct> Structure = UE::Mutable::Private::LoadObject(TypedNodeTable->Structure);
 			const FString TableName = Table ? GetNameSafe(Table).ToLower() : GetNameSafe(Structure).ToLower();
 			const uint32 TableId = CityHash32(reinterpret_cast<const char*>(*TableName), TableName.Len() * sizeof(FString::ElementType));
 
 			// This node will add a checker texture in case of error
-			mu::Ptr<mu::NodeImageConstant> EmptyNode = new mu::NodeImageConstant();
+			UE::Mutable::Private::Ptr<UE::Mutable::Private::NodeImageConstant> EmptyNode = new UE::Mutable::Private::NodeImageConstant();
 			Result = EmptyNode;
 
 			bool bSuccess = true;
@@ -1059,12 +1059,12 @@ mu::Ptr<mu::NodeImage> GenerateMutableSourceImage(const UEdGraphPin* Pin, FMutab
 				if (bSuccess)
 				{
 					// Generating a new data table if not exists
-					mu::Ptr<mu::FTable> GeneratedTable = nullptr;
+					UE::Mutable::Private::Ptr<UE::Mutable::Private::FTable> GeneratedTable = nullptr;
 					GeneratedTable = GenerateMutableSourceTable(DataTable, TypedNodeTable, GenerationContext);
 
 					if (GeneratedTable)
 					{
-						mu::Ptr<mu::NodeImageTable> ImageTableNode = new mu::NodeImageTable();
+						UE::Mutable::Private::Ptr<UE::Mutable::Private::NodeImageTable> ImageTableNode = new UE::Mutable::Private::NodeImageTable();
 
 						if (Pin->PinType.PinCategory == Schema->PC_Material)
 						{
@@ -1105,7 +1105,7 @@ mu::Ptr<mu::NodeImage> GenerateMutableSourceImage(const UEdGraphPin* Pin, FMutab
 						{
 							int32 Dummy = -1; // TODO MTBL-1512
 							bool Dummy2 = false;
-							GenerationContext.CurrentTableColumnType = mu::ETableColumnType::Image;
+							GenerationContext.CurrentTableColumnType = UE::Mutable::Private::ETableColumnType::Image;
 							bSuccess = GenerateTableColumn(TypedNodeTable, Pin, GeneratedTable, ColumnName, Property, 
 								FMutableSourceMeshData(), Dummy, Dummy, GenerationContext.CurrentLOD, Dummy, Dummy2, GenerationContext);
 
@@ -1131,7 +1131,7 @@ mu::Ptr<mu::NodeImage> GenerateMutableSourceImage(const UEdGraphPin* Pin, FMutab
 							UTexture2D* DefaultTexture2D = Cast<UTexture2D>(DefaultTexture);
 							if (!bIsPassthrough && DefaultTexture2D)
 							{
-								mu::FImageDesc ImageDesc = GenerateImageDescriptor(DefaultTexture2D);
+								UE::Mutable::Private::FImageDesc ImageDesc = GenerateImageDescriptor(DefaultTexture2D);
 
 								uint32 LODBias = ReferenceTextureSize > 0 ? ComputeLODBiasForTexture(GenerationContext, *DefaultTexture2D, nullptr, ReferenceTextureSize) : 0;
 								ImageDesc.m_size[0] = ImageDesc.m_size[0] >> LODBias;
@@ -1222,12 +1222,12 @@ mu::Ptr<mu::NodeImage> GenerateMutableSourceImage(const UEdGraphPin* Pin, FMutab
 
 	else if (const UCONodeMaterialBreak* TypedNodeMaterialBreak = Cast<UCONodeMaterialBreak>(Node))
 	{
-		Result = new mu::NodeImageConstant();
+		Result = new UE::Mutable::Private::NodeImageConstant();
 		check(TypedNodeMaterialBreak->MaterialPinRef.Get());
 
 		if (const UEdGraphPin* LinkedMaterial = FollowInputPin(*TypedNodeMaterialBreak->MaterialPinRef.Get()))
 		{
-			mu::Ptr<mu::NodeImageMaterialBreak> MaterialBreak = new mu::NodeImageMaterialBreak();
+			UE::Mutable::Private::Ptr<UE::Mutable::Private::NodeImageMaterialBreak> MaterialBreak = new UE::Mutable::Private::NodeImageMaterialBreak();
 
 			FName ParameterName = TypedNodeMaterialBreak->GetPinParameterName(*Pin);
 			check(!ParameterName.IsNone());
@@ -1237,7 +1237,7 @@ mu::Ptr<mu::NodeImage> GenerateMutableSourceImage(const UEdGraphPin* Pin, FMutab
 
 			// Generate the material that is going to be broken into parameters
 			GenerationContext.CurrentMaterialBreakParameter = FMutableGraphGenerationContext::FMaterialBreakParameter{ ParameterName, EMaterialParameterType::Texture };
-			mu::Ptr<mu::NodeMaterial> LinkedNode = GenerateMutableSourceMaterial(LinkedMaterial, GenerationContext);
+			UE::Mutable::Private::Ptr<UE::Mutable::Private::NodeMaterial> LinkedNode = GenerateMutableSourceMaterial(LinkedMaterial, GenerationContext);
 			MaterialBreak->MaterialSource = LinkedNode;
 			GenerationContext.CurrentMaterialBreakParameter = FMutableGraphGenerationContext::FMaterialBreakParameter{ };
 
diff --git a/Engine/Plugins/Mutable/Source/CustomizableObjectEditor/Private/MuCOE/GenerateMutableSource/GenerateMutableSourceImage.h b/Engine/Plugins/Mutable/Source/CustomizableObjectEditor/Private/MuCOE/GenerateMutableSource/GenerateMutableSourceImage.h
--- a/Engine/Plugins/Mutable/Source/CustomizableObjectEditor/Private/MuCOE/GenerateMutableSource/GenerateMutableSourceImage.h
+++ b/Engine/Plugins/Mutable/Source/CustomizableObjectEditor/Private/MuCOE/GenerateMutableSource/GenerateMutableSourceImage.h
@@ -12,16 +12,16 @@ class UEdGraphPin;
 class UTexture2D;
 struct FMutableGraphGenerationContext;
 
-namespace mu
+namespace UE::Mutable::Private
 {
 	class FImage;
 	class NodeImage;
 }
 
 
-//TSharedPtr<mu::FImage> ConvertTextureUnrealToMutable(UTexture2D* Texture, const UCustomizableObjectNode* Node, FCustomizableObjectCompiler* Compiler, bool bIsNormalComposite);
+//TSharedPtr<UE::Mutable::Private::FImage> ConvertTextureUnrealToMutable(UTexture2D* Texture, const UCustomizableObjectNode* Node, FCustomizableObjectCompiler* Compiler, bool bIsNormalComposite);
 
-mu::Ptr<mu::NodeImage> ResizeTextureByNumMips(const mu::Ptr<mu::NodeImage>& Image, int32 MipsToSkip);
+UE::Mutable::Private::Ptr<UE::Mutable::Private::NodeImage> ResizeTextureByNumMips(const UE::Mutable::Private::Ptr<UE::Mutable::Private::NodeImage>& Image, int32 MipsToSkip);
 
 /** Convert a CustomizableObject Source Graph from an Image pin into a mutable source graph. */
-mu::Ptr<mu::NodeImage> GenerateMutableSourceImage(const UEdGraphPin* Pin, FMutableGraphGenerationContext& GenerationContext, int32 ReferenceTextureSize);
+UE::Mutable::Private::Ptr<UE::Mutable::Private::NodeImage> GenerateMutableSourceImage(const UEdGraphPin* Pin, FMutableGraphGenerationContext& GenerationContext, int32 ReferenceTextureSize);
diff --git a/Engine/Plugins/Mutable/Source/CustomizableObjectEditor/Private/MuCOE/GenerateMutableSource/GenerateMutableSourceLayout.cpp b/Engine/Plugins/Mutable/Source/CustomizableObjectEditor/Private/MuCOE/GenerateMutableSource/GenerateMutableSourceLayout.cpp
--- a/Engine/Plugins/Mutable/Source/CustomizableObjectEditor/Private/MuCOE/GenerateMutableSource/GenerateMutableSourceLayout.cpp
+++ b/Engine/Plugins/Mutable/Source/CustomizableObjectEditor/Private/MuCOE/GenerateMutableSource/GenerateMutableSourceLayout.cpp
@@ -12,16 +12,16 @@
 
 #define LOCTEXT_NAMESPACE "CustomizableObjectEditor"
 
-mu::Ptr<mu::NodeLayout> CreateDefaultLayout()
+UE::Mutable::Private::Ptr<UE::Mutable::Private::NodeLayout> CreateDefaultLayout()
 {
 	constexpr int32 GridSize = 4;
 
-	mu::Ptr<mu::NodeLayout> LayoutNode = new mu::NodeLayout();
+	UE::Mutable::Private::Ptr<UE::Mutable::Private::NodeLayout> LayoutNode = new UE::Mutable::Private::NodeLayout();
 	LayoutNode->Size = { GridSize, GridSize };
 	LayoutNode->MaxSize = { GridSize, GridSize };
-	LayoutNode->AutoBlockStrategy = mu::EAutoBlocksStrategy::Ignore;
-	LayoutNode->Strategy = mu::EPackStrategy::Resizeable;
-	LayoutNode->ReductionMethod = mu::EReductionMethod::Halve;
+	LayoutNode->AutoBlockStrategy = UE::Mutable::Private::EAutoBlocksStrategy::Ignore;
+	LayoutNode->Strategy = UE::Mutable::Private::EPackStrategy::Resizeable;
+	LayoutNode->ReductionMethod = UE::Mutable::Private::EReductionMethod::Halve;
 	LayoutNode->TexCoordsIndex = 0;
 	LayoutNode->Blocks.SetNum(1);
 	LayoutNode->Blocks[0].Min = { 0, 0 };
@@ -34,9 +34,9 @@ mu::Ptr<mu::NodeLayout> CreateDefaultLayout()
 }
 
 
-mu::Ptr<mu::NodeLayout> CreateMutableLayoutNode(FMutableGraphGenerationContext& GenerationContext, const UCustomizableObjectLayout* UnrealLayout, bool bIgnoreLayoutWarnings)
+UE::Mutable::Private::Ptr<UE::Mutable::Private::NodeLayout> CreateMutableLayoutNode(FMutableGraphGenerationContext& GenerationContext, const UCustomizableObjectLayout* UnrealLayout, bool bIgnoreLayoutWarnings)
 {
-	mu::Ptr<mu::NodeLayout> LayoutNode = new mu::NodeLayout;
+	UE::Mutable::Private::Ptr<UE::Mutable::Private::NodeLayout> LayoutNode = new UE::Mutable::Private::NodeLayout;
 
 	LayoutNode->Size = FIntVector2(UnrealLayout->GetGridSize().X, UnrealLayout->GetGridSize().Y);
 	LayoutNode->MaxSize = FIntVector2(UnrealLayout->GetMaxGridSize().X, UnrealLayout->GetMaxGridSize().Y);
@@ -45,7 +45,7 @@ mu::Ptr<mu::NodeLayout> CreateMutableLayoutNode(FMutableGraphGenerationContext&
 	LayoutNode->AutoBlockStrategy = ConvertAutomaticBlocksStrategy(UnrealLayout->AutomaticBlocksStrategy);
 	LayoutNode->bMergeChildBlocks = UnrealLayout->AutomaticBlocksMergeStrategy == ECustomizableObjectLayoutAutomaticBlocksMergeStrategy::MergeChildBlocks;
 	LayoutNode->Strategy = ConvertLayoutStrategy(UnrealLayout->PackingStrategy);
-	LayoutNode->ReductionMethod = UnrealLayout->BlockReductionMethod == ECustomizableObjectLayoutBlockReductionMethod::Halve ? mu::EReductionMethod::Halve : mu::EReductionMethod::Unitary;
+	LayoutNode->ReductionMethod = UnrealLayout->BlockReductionMethod == ECustomizableObjectLayoutBlockReductionMethod::Halve ? UE::Mutable::Private::EReductionMethod::Halve : UE::Mutable::Private::EReductionMethod::Unitary;
 
 	if (bIgnoreLayoutWarnings)
 	{
@@ -67,9 +67,9 @@ mu::Ptr<mu::NodeLayout> CreateMutableLayoutNode(FMutableGraphGenerationContext&
 }
 
 
-mu::FSourceLayoutBlock ToMutable(FMutableGraphGenerationContext& GenerationContext, const FCustomizableObjectLayoutBlock& UnrealBlock)
+UE::Mutable::Private::FSourceLayoutBlock ToMutable(FMutableGraphGenerationContext& GenerationContext, const FCustomizableObjectLayoutBlock& UnrealBlock)
 {
-	mu::FSourceLayoutBlock MutableBlock;
+	UE::Mutable::Private::FSourceLayoutBlock MutableBlock;
 
 	MutableBlock.Min = { uint16(UnrealBlock.Min.X), uint16(UnrealBlock.Min.Y) };
 	FIntPoint Size = UnrealBlock.Max - UnrealBlock.Min;
@@ -82,7 +82,7 @@ mu::FSourceLayoutBlock ToMutable(FMutableGraphGenerationContext& GenerationConte
 	if (UnrealBlock.Mask)
 	{
 		// In the editor the src data can be directly accessed
-		TSharedPtr<mu::FImage> MaskImage = MakeShared<mu::FImage>();
+		TSharedPtr<UE::Mutable::Private::FImage> MaskImage = MakeShared<UE::Mutable::Private::FImage>();
 
 		FMutableSourceTextureData Tex(*UnrealBlock.Mask);
 		EUnrealToMutableConversionError Error = ConvertTextureUnrealSourceToMutable(MaskImage.Get(), Tex, 0);
diff --git a/Engine/Plugins/Mutable/Source/CustomizableObjectEditor/Private/MuCOE/GenerateMutableSource/GenerateMutableSourceLayout.h b/Engine/Plugins/Mutable/Source/CustomizableObjectEditor/Private/MuCOE/GenerateMutableSource/GenerateMutableSourceLayout.h
--- a/Engine/Plugins/Mutable/Source/CustomizableObjectEditor/Private/MuCOE/GenerateMutableSource/GenerateMutableSourceLayout.h
+++ b/Engine/Plugins/Mutable/Source/CustomizableObjectEditor/Private/MuCOE/GenerateMutableSource/GenerateMutableSourceLayout.h
@@ -9,8 +9,8 @@ class UEdGraphPin;
 struct FMutableGraphGenerationContext;
 
 
-extern mu::Ptr<mu::NodeLayout> CreateDefaultLayout();
+extern UE::Mutable::Private::Ptr<UE::Mutable::Private::NodeLayout> CreateDefaultLayout();
 
-extern mu::Ptr<mu::NodeLayout> CreateMutableLayoutNode(FMutableGraphGenerationContext& GenerationContext, const UCustomizableObjectLayout*, bool bIgnoreLayoutWarnings);
+extern UE::Mutable::Private::Ptr<UE::Mutable::Private::NodeLayout> CreateMutableLayoutNode(FMutableGraphGenerationContext& GenerationContext, const UCustomizableObjectLayout*, bool bIgnoreLayoutWarnings);
 
-extern mu::FSourceLayoutBlock ToMutable(FMutableGraphGenerationContext& GenerationContext, const FCustomizableObjectLayoutBlock&);
+extern UE::Mutable::Private::FSourceLayoutBlock ToMutable(FMutableGraphGenerationContext& GenerationContext, const FCustomizableObjectLayoutBlock&);
diff --git a/Engine/Plugins/Mutable/Source/CustomizableObjectEditor/Private/MuCOE/GenerateMutableSource/GenerateMutableSourceMacro.h b/Engine/Plugins/Mutable/Source/CustomizableObjectEditor/Private/MuCOE/GenerateMutableSource/GenerateMutableSourceMacro.h
--- a/Engine/Plugins/Mutable/Source/CustomizableObjectEditor/Private/MuCOE/GenerateMutableSource/GenerateMutableSourceMacro.h
+++ b/Engine/Plugins/Mutable/Source/CustomizableObjectEditor/Private/MuCOE/GenerateMutableSource/GenerateMutableSourceMacro.h
@@ -10,9 +10,9 @@
 #define LOCTEXT_NAMESPACE "CustomizableObjectEditor"
 
 template<class T>
-mu::Ptr<T> GenerateMutableSourceMacro(const UEdGraphPin& Pin, FMutableGraphGenerationContext& GenerationContext, const TFunction<mu::Ptr<T>(const UEdGraphPin*, FMutableGraphGenerationContext&)>& GenerationFunction)
+UE::Mutable::Private::Ptr<T> GenerateMutableSourceMacro(const UEdGraphPin& Pin, FMutableGraphGenerationContext& GenerationContext, const TFunction<UE::Mutable::Private::Ptr<T>(const UEdGraphPin*, FMutableGraphGenerationContext&)>& GenerationFunction)
 {
-	mu::Ptr<T> Result;
+	UE::Mutable::Private::Ptr<T> Result;
 	UCustomizableObjectNode* Node = CastChecked<UCustomizableObjectNode>(Pin.GetOwningNode());
 
 	if (const UCustomizableObjectNodeMacroInstance* TypedNodeMacro = Cast<UCustomizableObjectNodeMacroInstance>(Node))
