[Mutable] Change namespace name

diff --git a/Engine/Plugins/Mutable/Source/CustomizableObjectEditor/Private/MuCOE/GenerateMutableSource/GenerateMutableSourceMaterial.cpp b/Engine/Plugins/Mutable/Source/CustomizableObjectEditor/Private/MuCOE/GenerateMutableSource/GenerateMutableSourceMaterial.cpp
--- a/Engine/Plugins/Mutable/Source/CustomizableObjectEditor/Private/MuCOE/GenerateMutableSource/GenerateMutableSourceMaterial.cpp
+++ b/Engine/Plugins/Mutable/Source/CustomizableObjectEditor/Private/MuCOE/GenerateMutableSource/GenerateMutableSourceMaterial.cpp
@@ -38,7 +38,7 @@
 #define LOCTEXT_NAMESPACE "CustomizableObjectEditor"
 
 
-mu::Ptr<mu::NodeMaterial> GenerateMutableSourceMaterial(const UEdGraphPin* Pin, FMutableGraphGenerationContext& GenerationContext)
+UE::Mutable::Private::Ptr<UE::Mutable::Private::NodeMaterial> GenerateMutableSourceMaterial(const UEdGraphPin* Pin, FMutableGraphGenerationContext& GenerationContext)
 {
 	check(Pin)
 	RETURN_ON_CYCLE(*Pin, GenerationContext)
@@ -52,7 +52,7 @@ mu::Ptr<mu::NodeMaterial> GenerateMutableSourceMaterial(const UEdGraphPin* Pin,
 	const FGeneratedKey Key(reinterpret_cast<void*>(&GenerateMutableSourceMaterial), *Pin, *Node, GenerationContext);
 	if (const FGeneratedData* Generated = GenerationContext.Generated.Find(Key))
 	{
-		return static_cast<mu::NodeMaterial*>(Generated->Node.get());
+		return static_cast<UE::Mutable::Private::NodeMaterial*>(Generated->Node.get());
 	}
 
 	if (Node->IsNodeOutDatedAndNeedsRefresh())
@@ -64,11 +64,11 @@ mu::Ptr<mu::NodeMaterial> GenerateMutableSourceMaterial(const UEdGraphPin* Pin,
 	// Most nodes need to be added to the cache but there are some that don't. For exampel, MacroInstanceNodes
 	bool bCacheNode = true;
 
-	mu::Ptr<mu::NodeMaterial> Result;
+	UE::Mutable::Private::Ptr<UE::Mutable::Private::NodeMaterial> Result;
 	
 	if (const UCONodeMaterialConstant* TypedNodeMaterialConstant = Cast<UCONodeMaterialConstant>(Node))
 	{
-		mu::Ptr<mu::NodeMaterialConstant> MaterialNode;
+		UE::Mutable::Private::Ptr<UE::Mutable::Private::NodeMaterialConstant> MaterialNode;
 		bCacheNode = false; //custom node cache
 
 		if (UMaterialInterface* Material = GenerationContext.LoadObject(TypedNodeMaterialConstant->Material))
@@ -79,7 +79,7 @@ mu::Ptr<mu::NodeMaterial> GenerateMutableSourceMaterial(const UEdGraphPin* Pin,
 			}
 			else
 			{
-				MaterialNode = new mu::NodeMaterialConstant();
+				MaterialNode = new UE::Mutable::Private::NodeMaterialConstant();
 
 				GenerationContext.CurrentReferencedMaterialIndex = GenerationContext.ReferencedMaterials.AddUnique(Material);
 				MaterialNode->MaterialId = GenerationContext.CurrentReferencedMaterialIndex;
@@ -106,10 +106,10 @@ mu::Ptr<mu::NodeMaterial> GenerateMutableSourceMaterial(const UEdGraphPin* Pin,
 
 						if (bParameterFound && BaseTexture)
 						{
-							TSharedPtr<mu::FImage> ImageConstant = GenerateImageConstant(BaseTexture, GenerationContext, false);
+							TSharedPtr<UE::Mutable::Private::FImage> ImageConstant = GenerateImageConstant(BaseTexture, GenerationContext, false);
 							const uint32 MipsToSkip = ComputeLODBiasForTexture(GenerationContext, *BaseTexture);
 
-							mu::Ptr<mu::NodeImageConstant> ConstantImageNode = new mu::NodeImageConstant;
+							UE::Mutable::Private::Ptr<UE::Mutable::Private::NodeImageConstant> ConstantImageNode = new UE::Mutable::Private::NodeImageConstant;
 							ConstantImageNode->SetValue(ImageConstant);
 
 							const UTextureLODSettings& TextureLODSettings = GenerationContext.CompilationContext->Options.TargetPlatform->GetTextureLODSettings();
@@ -123,7 +123,7 @@ mu::Ptr<mu::NodeMaterial> GenerateMutableSourceMaterial(const UEdGraphPin* Pin,
 							const FString TextureName = GetNameSafe(BaseTexture).ToLower();
 							ConstantImageNode->SourceDataDescriptor.SourceId = CityHash32(reinterpret_cast<const char*>(*TextureName), TextureName.Len() * sizeof(FString::ElementType));
 
-							mu::Ptr<mu::NodeImage> FinalNodeImag = ResizeTextureByNumMips(ConstantImageNode, MipsToSkip);
+							UE::Mutable::Private::Ptr<UE::Mutable::Private::NodeImage> FinalNodeImag = ResizeTextureByNumMips(ConstantImageNode, MipsToSkip);
 
 							if (FinalNodeImag)
 							{
@@ -172,7 +172,7 @@ mu::Ptr<mu::NodeMaterial> GenerateMutableSourceMaterial(const UEdGraphPin* Pin,
 			// Check Switch Parameter arity preconditions.
 			if (const UEdGraphPin* ConnectedPin = FollowInputPin(*SwitchParameter))
 			{
-				mu::Ptr<mu::NodeScalar> SwitchParam = GenerateMutableSourceFloat(ConnectedPin, GenerationContext);
+				UE::Mutable::Private::Ptr<UE::Mutable::Private::NodeScalar> SwitchParam = GenerateMutableSourceFloat(ConnectedPin, GenerationContext);
 
 				// Switch Param not generated
 				if (!SwitchParam)
@@ -183,7 +183,7 @@ mu::Ptr<mu::NodeMaterial> GenerateMutableSourceMaterial(const UEdGraphPin* Pin,
 					return Result;
 				}
 
-				if (SwitchParam->GetType() != mu::NodeScalarEnumParameter::GetStaticType())
+				if (SwitchParam->GetType() != UE::Mutable::Private::NodeScalarEnumParameter::GetStaticType())
 				{
 					const FText Message = LOCTEXT("WrongSwitchParamType", "Switch parameter of incorrect type.");
 					GenerationContext.Log(Message, Node);
@@ -193,14 +193,14 @@ mu::Ptr<mu::NodeMaterial> GenerateMutableSourceMaterial(const UEdGraphPin* Pin,
 
 				const int32 NumSwitchOptions = TypedNodeMaterialSwitch->GetNumElements();
 
-				mu::NodeScalarEnumParameter* EnumParameter = static_cast<mu::NodeScalarEnumParameter*>(SwitchParam.get());
+				UE::Mutable::Private::NodeScalarEnumParameter* EnumParameter = static_cast<UE::Mutable::Private::NodeScalarEnumParameter*>(SwitchParam.get());
 				if (NumSwitchOptions != EnumParameter->Options.Num())
 				{
 					const FText Message = LOCTEXT("MismatchedSwitch", "Switch enum and switch node have different number of options. Please refresh the switch node to make sure the outcomes are labeled properly.");
 					GenerationContext.Log(Message, Node);
 				}
 
-				mu::Ptr<mu::NodeMaterialSwitch> SwitchNode = new mu::NodeMaterialSwitch;
+				UE::Mutable::Private::Ptr<UE::Mutable::Private::NodeMaterialSwitch> SwitchNode = new UE::Mutable::Private::NodeMaterialSwitch;
 				SwitchNode->Parameter = SwitchParam;
 				SwitchNode->Options.SetNum(NumSwitchOptions);
 
@@ -231,12 +231,12 @@ mu::Ptr<mu::NodeMaterial> GenerateMutableSourceMaterial(const UEdGraphPin* Pin,
 
 	else if (const UCONodeMaterialVariation* TypedNodeMaterialVariation = Cast<UCONodeMaterialVariation>(Node))
 	{
-		mu::Ptr<mu::NodeMaterialVariation> MaterialNode = new mu::NodeMaterialVariation();
+		UE::Mutable::Private::Ptr<UE::Mutable::Private::NodeMaterialVariation> MaterialNode = new UE::Mutable::Private::NodeMaterialVariation();
 		Result = MaterialNode;
 
 		if (const UEdGraphPin* ConnectedPin = FollowInputPin(*TypedNodeMaterialVariation->DefaultPin()))
 		{
-			mu::Ptr<mu::NodeMaterial> ChildNode = GenerateMutableSourceMaterial(ConnectedPin, GenerationContext);
+			UE::Mutable::Private::Ptr<UE::Mutable::Private::NodeMaterial> ChildNode = GenerateMutableSourceMaterial(ConnectedPin, GenerationContext);
 			if (ChildNode)
 			{
 				MaterialNode->DefaultMaterial = ChildNode;
@@ -257,7 +257,7 @@ mu::Ptr<mu::NodeMaterial> GenerateMutableSourceMaterial(const UEdGraphPin* Pin,
 
 				if (const UEdGraphPin* ConnectedPin = FollowInputPin(*VariationPin))
 				{
-					mu::Ptr<mu::NodeMaterial> ChildNode = GenerateMutableSourceMaterial(ConnectedPin, GenerationContext);
+					UE::Mutable::Private::Ptr<UE::Mutable::Private::NodeMaterial> ChildNode = GenerateMutableSourceMaterial(ConnectedPin, GenerationContext);
 					MaterialNode->Variations[VariationIndex].Material = ChildNode;
 				}
 			}
@@ -267,7 +267,7 @@ mu::Ptr<mu::NodeMaterial> GenerateMutableSourceMaterial(const UEdGraphPin* Pin,
 	else if (const UCustomizableObjectNodeMaterialParameter* MaterialParameterNode = Cast<UCustomizableObjectNodeMaterialParameter>(Node))
 	{
 		bCacheNode = false;
-		mu::Ptr<mu::NodeMaterialParameter> MaterialNode;
+		UE::Mutable::Private::Ptr<UE::Mutable::Private::NodeMaterialParameter> MaterialNode;
 
 		if (GenerationContext.MaterialParameterNodesCache.Contains(MaterialParameterNode))
 		{
@@ -275,13 +275,13 @@ mu::Ptr<mu::NodeMaterial> GenerateMutableSourceMaterial(const UEdGraphPin* Pin,
 		}
 		else
 		{
-			MaterialNode = new mu::NodeMaterialParameter();
+			MaterialNode = new UE::Mutable::Private::NodeMaterialParameter();
 			MaterialNode->Name = MaterialParameterNode->GetParameterName(&GenerationContext.MacroNodesStack);
 			MaterialNode->UID = GenerationContext.GetNodeIdUnique(Node).ToString();
 
 			if (MaterialParameterNode->DefaultValue)
 			{
-				GenerationContext.MaterialParameterDefaultValues.Add(FName(MaterialParameterNode->GetParameterName()), MutablePrivate::LoadObject(MaterialParameterNode->DefaultValue));
+				GenerationContext.MaterialParameterDefaultValues.Add(FName(MaterialParameterNode->GetParameterName()), UE::Mutable::Private::LoadObject(MaterialParameterNode->DefaultValue));
 			}
 
 			GenerationContext.ParameterUIDataMap.Add(MaterialParameterNode->GetParameterName(&GenerationContext.MacroNodesStack), FMutableParameterData(
@@ -305,22 +305,22 @@ mu::Ptr<mu::NodeMaterial> GenerateMutableSourceMaterial(const UEdGraphPin* Pin,
 				case EMaterialParameterType::Texture:
 				{
 					//TODO(Max): Move this to the generatemutablesourceimage:
-					mu::Ptr<mu::NodeImageFromMaterialParameter> TextureNode = new mu::NodeImageFromMaterialParameter();
+					UE::Mutable::Private::Ptr<UE::Mutable::Private::NodeImageFromMaterialParameter> TextureNode = new UE::Mutable::Private::NodeImageFromMaterialParameter();
 					TextureNode->ImageParameterName = GenerationContext.CurrentMaterialBreakParameter.ParameterName;
 
 					// Force the same format that the default texture if any.
-					mu::Ptr<mu::NodeImageFormat> FormatNode = new mu::NodeImageFormat();
+					UE::Mutable::Private::Ptr<UE::Mutable::Private::NodeImageFormat> FormatNode = new UE::Mutable::Private::NodeImageFormat();
 
 					// Force an "easy format" on the texture.
-					FormatNode->Format = mu::EImageFormat::RGBA_UByte;
+					FormatNode->Format = UE::Mutable::Private::EImageFormat::RGBA_UByte;
 					FormatNode->Source = TextureNode;
 
 					// Resize image node.
-					mu::Ptr<mu::NodeImageResize> ResizeNode = new mu::NodeImageResize();
+					UE::Mutable::Private::Ptr<UE::Mutable::Private::NodeImageResize> ResizeNode = new UE::Mutable::Private::NodeImageResize();
 					ResizeNode->Base = FormatNode;
 					ResizeNode->bRelative = false;
 
-					mu::FImageSize TextureSize(1);
+					UE::Mutable::Private::FImageSize TextureSize(1);
 
 					//Find reference Texture
 					UTexture* Texture = nullptr;
@@ -363,7 +363,7 @@ mu::Ptr<mu::NodeMaterial> GenerateMutableSourceMaterial(const UEdGraphPin* Pin,
 	else if (const UCustomizableObjectNodeTable* TypedNodeTable = Cast<UCustomizableObjectNodeTable>(Node))
 	{
 		//This node will add a default value in case of error
-		mu::Ptr<mu::NodeMaterialConstant> ConstantValue = new mu::NodeMaterialConstant();
+		UE::Mutable::Private::Ptr<UE::Mutable::Private::NodeMaterialConstant> ConstantValue = new UE::Mutable::Private::NodeMaterialConstant();
 		ConstantValue->MaterialId = GenerationContext.CurrentReferencedMaterialIndex;
 
 		Result = ConstantValue;
@@ -420,12 +420,12 @@ mu::Ptr<mu::NodeMaterial> GenerateMutableSourceMaterial(const UEdGraphPin* Pin,
 			if (bSuccess)
 			{
 				// Generating a new data table if not exists
-				mu::Ptr<mu::FTable> Table;
+				UE::Mutable::Private::Ptr<UE::Mutable::Private::FTable> Table;
 				Table = GenerateMutableSourceTable(DataTable, TypedNodeTable, GenerationContext);
 
 				if (Table)
 				{
-					mu::Ptr<mu::NodeMaterialTable> MaterialTableNode = new mu::NodeMaterialTable();
+					UE::Mutable::Private::Ptr<UE::Mutable::Private::NodeMaterialTable> MaterialTableNode = new UE::Mutable::Private::NodeMaterialTable();
 
 					// Materials use the parameter id as column names
 					ColumnName = GenerationContext.CurrentMaterialParameterId;
@@ -435,7 +435,7 @@ mu::Ptr<mu::NodeMaterial> GenerateMutableSourceMaterial(const UEdGraphPin* Pin,
 					{
 						int32 Dummy = -1; // TODO MTBL-1512
 						bool Dummy2 = false;
-						GenerationContext.CurrentTableColumnType = mu::ETableColumnType::Material;
+						GenerationContext.CurrentTableColumnType = UE::Mutable::Private::ETableColumnType::Material;
 						bSuccess = GenerateTableColumn(TypedNodeTable, Pin, Table, ColumnName, Property, FMutableSourceMeshData(), Dummy, Dummy, GenerationContext.CurrentLOD, Dummy, Dummy2, GenerationContext);
 
 						if (!bSuccess)
@@ -472,13 +472,13 @@ mu::Ptr<mu::NodeMaterial> GenerateMutableSourceMaterial(const UEdGraphPin* Pin,
 	else if (const UCustomizableObjectNodeMacroInstance* TypedNodeMacro = Cast<UCustomizableObjectNodeMacroInstance>(Node))
 	{
 		bCacheNode = false;
-		Result = GenerateMutableSourceMacro<mu::NodeMaterial>(*Pin, GenerationContext, GenerateMutableSourceMaterial);
+		Result = GenerateMutableSourceMacro<UE::Mutable::Private::NodeMaterial>(*Pin, GenerationContext, GenerateMutableSourceMaterial);
 	}
 
 	else if (const UCustomizableObjectNodeTunnel* TypedNodeTunnel = Cast<UCustomizableObjectNodeTunnel>(Node))
 	{
 		bCacheNode = false;
-		Result = GenerateMutableSourceMacro<mu::NodeMaterial>(*Pin, GenerationContext, GenerateMutableSourceMaterial);
+		Result = GenerateMutableSourceMacro<UE::Mutable::Private::NodeMaterial>(*Pin, GenerationContext, GenerateMutableSourceMaterial);
 	}
 
 	else
diff --git a/Engine/Plugins/Mutable/Source/CustomizableObjectEditor/Private/MuCOE/GenerateMutableSource/GenerateMutableSourceMaterial.h b/Engine/Plugins/Mutable/Source/CustomizableObjectEditor/Private/MuCOE/GenerateMutableSource/GenerateMutableSourceMaterial.h
--- a/Engine/Plugins/Mutable/Source/CustomizableObjectEditor/Private/MuCOE/GenerateMutableSource/GenerateMutableSourceMaterial.h
+++ b/Engine/Plugins/Mutable/Source/CustomizableObjectEditor/Private/MuCOE/GenerateMutableSource/GenerateMutableSourceMaterial.h
@@ -8,4 +8,4 @@ class UEdGraphPin;
 struct FMutableGraphGenerationContext;
 
 /** Convert a CustomizableObject Source Graph into a mutable source graph. */
-mu::Ptr<mu::NodeMaterial> GenerateMutableSourceMaterial(const UEdGraphPin* Pin, FMutableGraphGenerationContext& GenerationContext);
+UE::Mutable::Private::Ptr<UE::Mutable::Private::NodeMaterial> GenerateMutableSourceMaterial(const UEdGraphPin* Pin, FMutableGraphGenerationContext& GenerationContext);
diff --git a/Engine/Plugins/Mutable/Source/CustomizableObjectEditor/Private/MuCOE/GenerateMutableSource/GenerateMutableSourceMesh.cpp b/Engine/Plugins/Mutable/Source/CustomizableObjectEditor/Private/MuCOE/GenerateMutableSource/GenerateMutableSourceMesh.cpp
--- a/Engine/Plugins/Mutable/Source/CustomizableObjectEditor/Private/MuCOE/GenerateMutableSource/GenerateMutableSourceMesh.cpp
+++ b/Engine/Plugins/Mutable/Source/CustomizableObjectEditor/Private/MuCOE/GenerateMutableSource/GenerateMutableSourceMesh.cpp
@@ -62,8 +62,8 @@
 
 
 void SetSurfaceFormat(FMutableCompilationContext& Context,
-	mu::FMeshBufferSet& OutVertexBufferFormat,
-	mu::FMeshBufferSet& OutIndexBufferFormat,
+	UE::Mutable::Private::FMeshBufferSet& OutVertexBufferFormat,
+	UE::Mutable::Private::FMeshBufferSet& OutIndexBufferFormat,
 	const FMutableGraphMeshGenerationData& MeshData)
 {
 	// Limit skinning weights if necessary
@@ -138,7 +138,7 @@ void SetSurfaceFormat(FMutableCompilationContext& Context,
 	// MorphTarget vertex tracking info buffers
 	if (MeshData.bHasRealTimeMorphs)
 	{
-		using namespace mu;
+		using namespace UE::Mutable::Private;
 		{
 			const int32 ElementSize = sizeof(uint32);
 			constexpr int32 ChannelCount = 1;
@@ -170,7 +170,7 @@ void SetSurfaceFormat(FMutableCompilationContext& Context,
 	if (MeshData.bHasClothing)
 	{
 		{
-			using namespace mu;
+			using namespace UE::Mutable::Private;
 			const int32 ElementSize = sizeof(int32);
 			constexpr int32 ChannelCount = 1;
 			const EMeshBufferSemantic Semantics[ChannelCount] = { EMeshBufferSemantic::Other };
@@ -184,7 +184,7 @@ void SetSurfaceFormat(FMutableCompilationContext& Context,
 		}
 
 		{
-			using namespace mu;
+			using namespace UE::Mutable::Private;
 			const int32 ElementSize = sizeof(uint32);
 			constexpr int32 ChannelCount = 1;
 			const EMeshBufferSemantic Semantics[ChannelCount] = { EMeshBufferSemantic::Other };
@@ -545,14 +545,14 @@ bool IsSkeletalMeshCompatibleWithRefSkeleton(FMutableComponentInfo& ComponentInf
 	return true;
 }
 
-void SetAndPropagatePoseBoneUsage(mu::FMesh& MutableMesh, int32 PoseIndex, mu::EBoneUsageFlags Usage)
+void SetAndPropagatePoseBoneUsage(UE::Mutable::Private::FMesh& MutableMesh, int32 PoseIndex, UE::Mutable::Private::EBoneUsageFlags Usage)
 {
 	if (!MutableMesh.GetSkeleton())
 	{
 		return;
 	}
 
-	const mu::FSkeleton& MutableSkeleton = *MutableMesh.GetSkeleton();
+	const UE::Mutable::Private::FSkeleton& MutableSkeleton = *MutableMesh.GetSkeleton();
 
 	if (!MutableMesh.BonePoses.IsValidIndex(PoseIndex))
 	{
@@ -590,7 +590,7 @@ TArray<TTuple<UPhysicsAsset*, int32>> GetPhysicsAssetsFromAnimInstance(const TSo
 	}
 
 	// TODO: Get Physics Asset using AssetRegistry
-	UClass* AnimInstanceClass = MutablePrivate::LoadClass(AnimInstance);
+	UClass* AnimInstanceClass = UE::Mutable::Private::LoadClass(AnimInstance);
 	UAnimBlueprintGeneratedClass* AnimClass = Cast<UAnimBlueprintGeneratedClass>(AnimInstanceClass);
 
 	if (AnimClass)
@@ -634,7 +634,7 @@ TArray<TTuple<UPhysicsAsset*, int32>> GetPhysicsAssetsFromAnimInstance(FMutableG
 	return Result;
 }
 
-TArray<uint8> MakePhysicsAssetBodySetupRelevancyMap(const FMutableCompilationContext& Context, UPhysicsAsset* Asset, const TSharedPtr<mu::FMesh>& Mesh)
+TArray<uint8> MakePhysicsAssetBodySetupRelevancyMap(const FMutableCompilationContext& Context, UPhysicsAsset* Asset, const TSharedPtr<UE::Mutable::Private::FMesh>& Mesh)
 {
 	const int32 BodySetupsNum = Asset->SkeletalBodySetups.Num();
 
@@ -648,14 +648,14 @@ TArray<uint8> MakePhysicsAssetBodySetupRelevancyMap(const FMutableCompilationCon
 
 	for (int32 BodyIndex = 0; BodyIndex < BodySetupsNum; ++BodyIndex)
 	{
-		mu::FBoneName Bone; 
+		UE::Mutable::Private::FBoneName Bone; 
 		RelevancyMap[BodyIndex] = Context.FindBone(Asset->SkeletalBodySetups[BodyIndex]->BoneName, Bone);
 	}
 
 	return RelevancyMap;
 }
 
-TSharedPtr<mu::FPhysicsBody> MakePhysicsBodyFromAsset(FMutableCompilationContext& Context, UPhysicsAsset* Asset, const TArray<uint8>& BodySetupRelevancyMap)
+TSharedPtr<UE::Mutable::Private::FPhysicsBody> MakePhysicsBodyFromAsset(FMutableCompilationContext& Context, UPhysicsAsset* Asset, const TArray<uint8>& BodySetupRelevancyMap)
 {
 	check(Asset);
 	check(Asset->SkeletalBodySetups.Num() == BodySetupRelevancyMap.Num());
@@ -665,7 +665,7 @@ TSharedPtr<mu::FPhysicsBody> MakePhysicsBodyFromAsset(FMutableCompilationContext
 	
 	const int32 NumRelevantSetups = Algo::CountIf(BodySetupRelevancyMap, [](uint8 V) { return V; });
 
-	TSharedPtr<mu::FPhysicsBody> PhysicsBody = MakeShared<mu::FPhysicsBody>();
+	TSharedPtr<UE::Mutable::Private::FPhysicsBody> PhysicsBody = MakeShared<UE::Mutable::Private::FPhysicsBody>();
 	
 	PhysicsBody->SetBodyCount(NumRelevantSetups);
 
@@ -688,7 +688,7 @@ TSharedPtr<mu::FPhysicsBody> MakePhysicsBodyFromAsset(FMutableCompilationContext
 
 		TObjectPtr<USkeletalBodySetup>& BodySetup = SkeletalBodySetups[SourceBodyIndex++];
 
-		const mu::FBoneName& BodyBoneId = Context.GetBoneUnique(BodySetup->BoneName);
+		const UE::Mutable::Private::FBoneName& BodyBoneId = Context.GetBoneUnique(BodySetup->BoneName);
 		PhysicsBody->SetBodyBoneId(B, BodyBoneId);
 		
 		const int32 NumSpheres = BodySetup->AggGeom.SphereElems.Num();
@@ -782,11 +782,11 @@ TSharedPtr<mu::FPhysicsBody> MakePhysicsBodyFromAsset(FMutableCompilationContext
 }
 
 
-namespace MutablePrivate
+namespace UE::Mutable::Private
 {
 
 	/** Copy mesh data from a single source buffer into a set of destination buffers assuming the format and components is the same. */
-	void CopyBufferClearingPadding(UE::Tasks::FTaskEvent CompletionEvent, mu::FMeshBufferSet& BufferSet, const mu::FMeshBufferSet& SourceBufferFormat, const void* InSourceData)
+	void CopyBufferClearingPadding(UE::Tasks::FTaskEvent CompletionEvent, UE::Mutable::Private::FMeshBufferSet& BufferSet, const UE::Mutable::Private::FMeshBufferSet& SourceBufferFormat, const void* InSourceData)
 	{
 		MUTABLE_CPUPROFILER_SCOPE(CopyBufferClearingPadding);
 
@@ -801,21 +801,21 @@ namespace MutablePrivate
 						int32 ElementSize = BufferSet.Buffers[BufferIndex].ElementSize;
 
 						check(SourceBufferFormat.Buffers.Num() == 1);
-						const mu::FMeshBuffer& SourceBuffer = SourceBufferFormat.Buffers[0];
+						const UE::Mutable::Private::FMeshBuffer& SourceBuffer = SourceBufferFormat.Buffers[0];
 
 						uint8* TargetData = BufferSet.GetBufferData(BufferIndex);
 						TArray<const uint8*, TInlineAllocator<8>> SourceDataPerChannel;
 						SourceDataPerChannel.SetNumUninitialized(ChannelCount);
 						for (int32 ChannelIndex = 0; ChannelIndex < ChannelCount; ++ChannelIndex)
 						{
-							const mu::FMeshBufferChannel& DestinationChannel = BufferSet.Buffers[BufferIndex].Channels[ChannelIndex];
+							const UE::Mutable::Private::FMeshBufferChannel& DestinationChannel = BufferSet.Buffers[BufferIndex].Channels[ChannelIndex];
 
 							int32 SourceBufferIndex = -1;
 							int32 SourceChannelIndex = -1;
 							SourceBufferFormat.FindChannel(DestinationChannel.Semantic, DestinationChannel.SemanticIndex, &SourceBufferIndex, &SourceChannelIndex);
 							check(SourceBufferIndex == 0 && SourceChannelIndex >= 0);
 
-							const mu::FMeshBufferChannel& SourceChannel = SourceBuffer.Channels[SourceChannelIndex];
+							const UE::Mutable::Private::FMeshBufferChannel& SourceChannel = SourceBuffer.Channels[SourceChannelIndex];
 							check(SourceChannel.Format == DestinationChannel.Format);
 							check(SourceChannel.ComponentCount == DestinationChannel.ComponentCount);
 
@@ -827,7 +827,7 @@ namespace MutablePrivate
 							int32 CurrentOffset = 0;
 							for (int32 ChannelIndex = 0; ChannelIndex < ChannelCount; ++ChannelIndex)
 							{
-								const mu::FMeshBufferChannel& Channel = BufferSet.Buffers[BufferIndex].Channels[ChannelIndex];
+								const UE::Mutable::Private::FMeshBufferChannel& Channel = BufferSet.Buffers[BufferIndex].Channels[ChannelIndex];
 
 								int32 ChannelOffset = Channel.Offset;
 
@@ -862,7 +862,7 @@ namespace MutablePrivate
 }
 
 
-TSharedPtr<mu::FMesh> ConvertSkeletalMeshToMutable(TSharedRef<FMeshConversionContext> MeshConversionContext)
+TSharedPtr<UE::Mutable::Private::FMesh> ConvertSkeletalMeshToMutable(TSharedRef<FMeshConversionContext> MeshConversionContext)
 {
 	const FMutableSourceMeshData& Source = MeshConversionContext->Source;
 	FMutableCompilationContext& Context = *MeshConversionContext->CompilationContext;
@@ -888,7 +888,7 @@ TSharedPtr<mu::FMesh> ConvertSkeletalMeshToMutable(TSharedRef<FMeshConversionCon
 	{
 		if (!Source.bMeshMustExist)
 		{
-			return MakeShared<mu::FMesh>(); // Return empty mesh to preserve the layouts
+			return MakeShared<UE::Mutable::Private::FMesh>(); // Return empty mesh to preserve the layouts
 		}
 		else
 		{
@@ -910,7 +910,7 @@ TSharedPtr<mu::FMesh> ConvertSkeletalMeshToMutable(TSharedRef<FMeshConversionCon
 	{
 		if (!Source.bMeshMustExist)
 		{
-			return MakeShared<mu::FMesh>(); // Return empty mesh to preserve the layouts
+			return MakeShared<UE::Mutable::Private::FMesh>(); // Return empty mesh to preserve the layouts
 		}
 		else
 		{
@@ -938,13 +938,13 @@ TSharedPtr<mu::FMesh> ConvertSkeletalMeshToMutable(TSharedRef<FMeshConversionCon
 	const bool bIgnoreMorphs = EnumHasAnyFlags(CurrentFlags, EMutableMeshConversionFlags::IgnoreMorphs);
 	const bool bDoNotCreateMeshMetadata = EnumHasAnyFlags(CurrentFlags, EMutableMeshConversionFlags::DoNotCreateMeshMetadata);
 
-	TSharedPtr<mu::FMesh> MutableMesh = MakeShared<mu::FMesh>();
+	TSharedPtr<UE::Mutable::Private::FMesh> MutableMesh = MakeShared<UE::Mutable::Private::FMesh>();
 
 	bool bBoneMapModified = false;
 	TArray<FBoneIndexType> BoneMap;
 	TArray<FBoneIndexType> RemappedBoneMapIndices;
 
-	// Check if the Skeleton is valid and build the mu::Skeleton
+	// Check if the Skeleton is valid and build the UE::Mutable::Private::Skeleton
 	if (!bIgnoreSkinning)
 	{
 		const USkeleton* Skeleton = SkeletalMesh->GetSkeleton();
@@ -1052,7 +1052,7 @@ TSharedPtr<mu::FMesh> ConvertSkeletalMeshToMutable(TSharedRef<FMeshConversionCon
 		}
 
 		// Create the skeleton, poses, and BoneMap for this mesh
-		TSharedPtr<mu::FSkeleton> MutableSkeleton = MakeShared<mu::FSkeleton>();
+		TSharedPtr<UE::Mutable::Private::FSkeleton> MutableSkeleton = MakeShared<UE::Mutable::Private::FSkeleton>();
 		MutableMesh->SetSkeleton(MutableSkeleton);
 
 		const int32 NumRequiredBones = RequiredBones.Num();
@@ -1060,7 +1060,7 @@ TSharedPtr<mu::FMesh> ConvertSkeletalMeshToMutable(TSharedRef<FMeshConversionCon
 		MutableSkeleton->SetBoneCount(NumRequiredBones);
 
 		// MutableBoneMap will not keep an index to the Skeleton, but to the BoneName
-		TArray<mu::FBoneName> MutableBoneMap;
+		TArray<UE::Mutable::Private::FBoneName> MutableBoneMap;
 		MutableBoneMap.SetNum(BoneMap.Num());
 
 		TArray<FMatrix> ComposedRefPoseMatrices;
@@ -1075,7 +1075,7 @@ TSharedPtr<mu::FMesh> ConvertSkeletalMeshToMutable(TSharedRef<FMeshConversionCon
 			const int32 ParentBoneIndex = RequiredBones.Find(BoneInfo.ParentIndex);
 
 			// Set bone hierarchy
-			const mu::FBoneName& BoneName = Context.GetBoneUnique(BoneInfo.Name);
+			const UE::Mutable::Private::FBoneName& BoneName = Context.GetBoneUnique(BoneInfo.Name);
 
 			MutableSkeleton->SetBoneName(BoneIndex, BoneName);
 			MutableSkeleton->SetBoneParent(BoneIndex, ParentBoneIndex);
@@ -1103,9 +1103,9 @@ TSharedPtr<mu::FMesh> ConvertSkeletalMeshToMutable(TSharedRef<FMeshConversionCon
 			FTransform3f BoneTransform;
 			BoneTransform.SetFromMatrix(FMatrix44f(ComposedRefPoseMatrices[BoneIndex]));
 
-			mu::EBoneUsageFlags BoneUsageFlags = mu::EBoneUsageFlags::None;
-			EnumAddFlags(BoneUsageFlags, BoneMapIndex != INDEX_NONE ? mu::EBoneUsageFlags::Skinning : mu::EBoneUsageFlags::None);
-			EnumAddFlags(BoneUsageFlags, ParentBoneIndex == INDEX_NONE ? mu::EBoneUsageFlags::Root : mu::EBoneUsageFlags::None);
+			UE::Mutable::Private::EBoneUsageFlags BoneUsageFlags = UE::Mutable::Private::EBoneUsageFlags::None;
+			EnumAddFlags(BoneUsageFlags, BoneMapIndex != INDEX_NONE ? UE::Mutable::Private::EBoneUsageFlags::Skinning : UE::Mutable::Private::EBoneUsageFlags::None);
+			EnumAddFlags(BoneUsageFlags, ParentBoneIndex == INDEX_NONE ? UE::Mutable::Private::EBoneUsageFlags::Root : UE::Mutable::Private::EBoneUsageFlags::None);
 
 			MutableMesh->SetBonePose(BoneIndex, BoneName, BoneTransform, BoneUsageFlags);
 		}
@@ -1127,9 +1127,9 @@ TSharedPtr<mu::FMesh> ConvertSkeletalMeshToMutable(TSharedRef<FMeshConversionCon
 	const bool bUseUnlimitedInfluences = FGPUBaseSkinVertexFactory::UseUnlimitedBoneInfluences(MaxSectionInfluences, Context.Options.TargetPlatform);
 
 	// Create a mutable vertex buffer definition that matches the unreal soft vertex format.
-	mu::FMeshBufferSet UnrealSourceVertexFormat;
+	UE::Mutable::Private::FMeshBufferSet UnrealSourceVertexFormat;
 	{
-		using namespace mu;
+		using namespace UE::Mutable::Private;
 		const int32 ElementSize = sizeof(FSoftSkinVertex);
 		const int32 ChannelCount = 11;
 		const EMeshBufferSemantic Semantics[ChannelCount] = { EMeshBufferSemantic::Position, EMeshBufferSemantic::Tangent, EMeshBufferSemantic::Binormal, EMeshBufferSemantic::Normal, EMeshBufferSemantic::TexCoords, EMeshBufferSemantic::TexCoords, EMeshBufferSemantic::TexCoords, EMeshBufferSemantic::TexCoords, EMeshBufferSemantic::Color, EMeshBufferSemantic::BoneIndices, EMeshBufferSemantic::BoneWeights };
@@ -1180,7 +1180,7 @@ TSharedPtr<mu::FMesh> ConvertSkeletalMeshToMutable(TSharedRef<FMeshConversionCon
 
 	// Create the mutable mesh with the same data without padding and separating the UVs, also add skinning only optionally.
 	{
-		using namespace mu;
+		using namespace UE::Mutable::Private;
 
 		// Base channels
 		{
@@ -1229,7 +1229,7 @@ TSharedPtr<mu::FMesh> ConvertSkeletalMeshToMutable(TSharedRef<FMeshConversionCon
 		MutableMesh->GetVertexBuffers().SetBufferCount(3);
 
 		// Skinning buffer
-		using namespace mu;
+		using namespace UE::Mutable::Private;
 		const int32 ChannelCount = 2;
 		const EMeshBufferSemantic Semantics[ChannelCount] = { EMeshBufferSemantic::BoneIndices, EMeshBufferSemantic::BoneWeights };
 		const int32 SemanticIndices[ChannelCount] = { 0, 0 };
@@ -1370,7 +1370,7 @@ TSharedPtr<mu::FMesh> ConvertSkeletalMeshToMutable(TSharedRef<FMeshConversionCon
 
 	UE::Tasks::FTaskEvent CompletionEvent(TEXT("MeshConversionCompletion"));
 
-	MutablePrivate::CopyBufferClearingPadding(CompletionEvent, MutableMesh->GetVertexBuffers(), UnrealSourceVertexFormat, SourceVertexData);
+	UE::Mutable::Private::CopyBufferClearingPadding(CompletionEvent, MutableMesh->GetVertexBuffers(), UnrealSourceVertexFormat, SourceVertexData);
 
 	// Indices
 	CompletionEvent.AddPrerequisites(
@@ -1382,7 +1382,7 @@ TSharedPtr<mu::FMesh> ConvertSkeletalMeshToMutable(TSharedRef<FMeshConversionCon
 				MutableMesh->GetIndexBuffers().SetBufferCount(1);
 				MutableMesh->GetIndexBuffers().SetElementCount(IndexCount);
 
-				using namespace mu;
+				using namespace UE::Mutable::Private;
 
 				check(LODModel.IndexBuffer.IsValidIndex(IndexStart) && LODModel.IndexBuffer.IsValidIndex(IndexStart + IndexCount - 1));
 				const uint32* IndexDataPtr = &LODModel.IndexBuffer[IndexStart];
@@ -1538,7 +1538,7 @@ TSharedPtr<mu::FMesh> ConvertSkeletalMeshToMutable(TSharedRef<FMeshConversionCon
 			MutableMesh->GetVertexBuffers().SetBufferCount(NextBufferIndex + 2);
 			// MorphTarget vertex block offset.
 			{
-				using namespace mu;
+				using namespace UE::Mutable::Private;
 				const int32 ElementSize = sizeof(uint32);
 				const int32 ChannelCount = 1;
 				const EMeshBufferSemantic Semantics[ChannelCount] = { EMeshBufferSemantic::Other };
@@ -1547,14 +1547,14 @@ TSharedPtr<mu::FMesh> ConvertSkeletalMeshToMutable(TSharedRef<FMeshConversionCon
 				int32 Components[ChannelCount] = { 1 };
 				const int32 Offsets[ChannelCount] = { 0 };
 
-				MutableMesh->GetVertexBuffers().SetBuffer(NextBufferIndex, ElementSize, ChannelCount, Semantics, SemanticIndices, Formats, Components, Offsets, mu::EMemoryInitPolicy::Zeroed);
+				MutableMesh->GetVertexBuffers().SetBuffer(NextBufferIndex, ElementSize, ChannelCount, Semantics, SemanticIndices, Formats, Components, Offsets, UE::Mutable::Private::EMemoryInitPolicy::Zeroed);
 			}
 
 			// MorphTarget vertex morph count.
 
 			// MorphTarget vertex block id.
 			{
-				using namespace mu;
+				using namespace UE::Mutable::Private;
 				const int32 ElementSize = sizeof(uint32);
 				const int32 ChannelCount = 1;
 				const EMeshBufferSemantic Semantics[ChannelCount] = { EMeshBufferSemantic::Other };
@@ -1563,7 +1563,7 @@ TSharedPtr<mu::FMesh> ConvertSkeletalMeshToMutable(TSharedRef<FMeshConversionCon
 				int32 Components[ChannelCount] = { 1 };
 				const int32 Offsets[ChannelCount] = { 0 };
 
-				MutableMesh->GetVertexBuffers().SetBuffer(NextBufferIndex + 1, ElementSize, ChannelCount, Semantics, SemanticIndices, Formats, Components, Offsets, mu::EMemoryInitPolicy::Zeroed);
+				MutableMesh->GetVertexBuffers().SetBuffer(NextBufferIndex + 1, ElementSize, ChannelCount, Semantics, SemanticIndices, Formats, Components, Offsets, UE::Mutable::Private::EMemoryInitPolicy::Zeroed);
 			}
 
 			// Setup MorphTarget reconstruction data.
@@ -1884,7 +1884,7 @@ TSharedPtr<mu::FMesh> ConvertSkeletalMeshToMutable(TSharedRef<FMeshConversionCon
 		{
 			MutableMesh->GetVertexBuffers().SetBufferCount(NextBufferIndex + 2);
 			{
-				using namespace mu;
+				using namespace UE::Mutable::Private;
 				const int32 ElementSize = sizeof(int32);
 				const int32 ChannelCount = 1;
 				const EMeshBufferSemantic Semantics[ChannelCount] = { EMeshBufferSemantic::Other };
@@ -1893,11 +1893,11 @@ TSharedPtr<mu::FMesh> ConvertSkeletalMeshToMutable(TSharedRef<FMeshConversionCon
 				int32 Components[ChannelCount] = { 1 };
 				const int32 Offsets[ChannelCount] = { 0 };
 
-				MutableMesh->GetVertexBuffers().SetBuffer(NextBufferIndex, ElementSize, ChannelCount, Semantics, SemanticIndices, Formats, Components, Offsets, mu::EMemoryInitPolicy::Zeroed);
+				MutableMesh->GetVertexBuffers().SetBuffer(NextBufferIndex, ElementSize, ChannelCount, Semantics, SemanticIndices, Formats, Components, Offsets, UE::Mutable::Private::EMemoryInitPolicy::Zeroed);
 			}
 
 			{
-				using namespace mu;
+				using namespace UE::Mutable::Private;
 				const int32 ElementSize = sizeof(uint32);
 				const int32 ChannelCount = 1;
 				const EMeshBufferSemantic Semantics[ChannelCount] = { EMeshBufferSemantic::Other };
@@ -1906,7 +1906,7 @@ TSharedPtr<mu::FMesh> ConvertSkeletalMeshToMutable(TSharedRef<FMeshConversionCon
 				int32 Components[ChannelCount] = { 1 };
 				const int32 Offsets[ChannelCount] = { 0 };
 
-				MutableMesh->GetVertexBuffers().SetBuffer(NextBufferIndex + 1, ElementSize, ChannelCount, Semantics, SemanticIndices, Formats, Components, Offsets, mu::EMemoryInitPolicy::Zeroed);
+				MutableMesh->GetVertexBuffers().SetBuffer(NextBufferIndex + 1, ElementSize, ChannelCount, Semantics, SemanticIndices, Formats, Components, Offsets, UE::Mutable::Private::EMemoryInitPolicy::Zeroed);
 			}
 
 			TArrayView<int32> ClothSectionIndexView(reinterpret_cast<int32*>(
@@ -2070,7 +2070,7 @@ TSharedPtr<mu::FMesh> ConvertSkeletalMeshToMutable(TSharedRef<FMeshConversionCon
 	// SkinWeightProfiles vertex info.
 	if (!bIgnoreSkinning && Context.Options.bSkinWeightProfilesEnabled)
 	{
-		using namespace mu;
+		using namespace UE::Mutable::Private;
 
 		// TODO: Remove BoneWeightFormat after merge
 		const int32 BoneWeightTypeSizeBytes = sizeof(TDecay<decltype(DeclVal<FRawSkinWeight>().InfluenceWeights[0])>::Type);
@@ -2192,7 +2192,7 @@ TSharedPtr<mu::FMesh> ConvertSkeletalMeshToMutable(TSharedRef<FMeshConversionCon
 				continue;
 			}
 
-			const mu::FBoneName& BoneName = Context.GetBoneUnique(BodySetup->BoneName);
+			const UE::Mutable::Private::FBoneName& BoneName = Context.GetBoneUnique(BodySetup->BoneName);
 			const int32 BonePoseIndex = MutableMesh->FindBonePose(BoneName);
 
 			if (BonePoseIndex == INDEX_NONE)
@@ -2202,7 +2202,7 @@ TSharedPtr<mu::FMesh> ConvertSkeletalMeshToMutable(TSharedRef<FMeshConversionCon
 
 			RelevantBodySetups.Add(BodySetup);
 			DiscardedBodySetups[BodySetupIndex] = 0;
-			EnumAddFlags(MutableMesh->BonePoses[BonePoseIndex].BoneUsageFlags, mu::EBoneUsageFlags::Physics);
+			EnumAddFlags(MutableMesh->BonePoses[BonePoseIndex].BoneUsageFlags, UE::Mutable::Private::EBoneUsageFlags::Physics);
 		}
 
 		const int32 NumDiscardedSetups = Algo::CountIf(DiscardedBodySetups, [](const uint8& V) { return V; });
@@ -2236,7 +2236,7 @@ TSharedPtr<mu::FMesh> ConvertSkeletalMeshToMutable(TSharedRef<FMeshConversionCon
 			Context.Log(FText::FromString(PhysicsSetupsRemovedMsg), Source.MessageContext, EMessageSeverity::Warning);
 		}
 
-		TSharedPtr<mu::FPhysicsBody> PhysicsBody = MakeShared<mu::FPhysicsBody>();
+		TSharedPtr<UE::Mutable::Private::FPhysicsBody> PhysicsBody = MakeShared<UE::Mutable::Private::FPhysicsBody>();
 
 		const int32 NumBodySetups = RelevantBodySetups.Num();
 		PhysicsBody->SetBodyCount(NumBodySetups);
@@ -2255,7 +2255,7 @@ TSharedPtr<mu::FMesh> ConvertSkeletalMeshToMutable(TSharedRef<FMeshConversionCon
 		{
 			TObjectPtr<USkeletalBodySetup>& BodySetup = RelevantBodySetups[B];
 
-			const mu::FBoneName& BoneId = Context.GetBoneUnique(BodySetup->BoneName);
+			const UE::Mutable::Private::FBoneName& BoneId = Context.GetBoneUnique(BodySetup->BoneName);
 			PhysicsBody->SetBodyBoneId(B, BoneId);
 
 			const int32 NumSpheres = BodySetup->AggGeom.SphereElems.Num();
@@ -2363,10 +2363,10 @@ TSharedPtr<mu::FMesh> ConvertSkeletalMeshToMutable(TSharedRef<FMeshConversionCon
 	// Set Bone Parenting usages. This has to be done after all primary usages are set.
 	for (int32 I = MutableMesh->GetBonePoseCount() - 1; I >= 0; --I)
 	{
-		mu::FMesh::FBonePose& BonePose = MutableMesh->BonePoses[I];
+		UE::Mutable::Private::FMesh::FBonePose& BonePose = MutableMesh->BonePoses[I];
 
-		constexpr mu::EBoneUsageFlags FlagsToPropagate =
-			mu::EBoneUsageFlags::Skinning | mu::EBoneUsageFlags::Physics | mu::EBoneUsageFlags::Deform;
+		constexpr UE::Mutable::Private::EBoneUsageFlags FlagsToPropagate =
+			UE::Mutable::Private::EBoneUsageFlags::Skinning | UE::Mutable::Private::EBoneUsageFlags::Physics | UE::Mutable::Private::EBoneUsageFlags::Deform;
 		if (EnumHasAnyFlags(BonePose.BoneUsageFlags, FlagsToPropagate))
 		{
 			const int32 Index = MutableMesh->GetSkeleton()->FindBone(BonePose.BoneId);
@@ -2383,13 +2383,13 @@ TSharedPtr<mu::FMesh> ConvertSkeletalMeshToMutable(TSharedRef<FMeshConversionCon
 				continue;
 			}
 
-			const mu::EBoneUsageFlags ParentPropagationFlags =
-				(EnumHasAnyFlags(BonePose.BoneUsageFlags, mu::EBoneUsageFlags::Skinning)
-					? mu::EBoneUsageFlags::SkinningParent : mu::EBoneUsageFlags::None) |
-				(EnumHasAnyFlags(BonePose.BoneUsageFlags, mu::EBoneUsageFlags::Physics)
-					? mu::EBoneUsageFlags::PhysicsParent : mu::EBoneUsageFlags::None) |
-				(EnumHasAnyFlags(BonePose.BoneUsageFlags, mu::EBoneUsageFlags::Deform)
-					? mu::EBoneUsageFlags::DeformParent : mu::EBoneUsageFlags::None);
+			const UE::Mutable::Private::EBoneUsageFlags ParentPropagationFlags =
+				(EnumHasAnyFlags(BonePose.BoneUsageFlags, UE::Mutable::Private::EBoneUsageFlags::Skinning)
+					? UE::Mutable::Private::EBoneUsageFlags::SkinningParent : UE::Mutable::Private::EBoneUsageFlags::None) |
+				(EnumHasAnyFlags(BonePose.BoneUsageFlags, UE::Mutable::Private::EBoneUsageFlags::Physics)
+					? UE::Mutable::Private::EBoneUsageFlags::PhysicsParent : UE::Mutable::Private::EBoneUsageFlags::None) |
+				(EnumHasAnyFlags(BonePose.BoneUsageFlags, UE::Mutable::Private::EBoneUsageFlags::Deform)
+					? UE::Mutable::Private::EBoneUsageFlags::DeformParent : UE::Mutable::Private::EBoneUsageFlags::None);
 
 			SetAndPropagatePoseBoneUsage(*MutableMesh, ParentIndex, ParentPropagationFlags);
 		}
@@ -2416,7 +2416,7 @@ TSharedPtr<mu::FMesh> ConvertSkeletalMeshToMutable(TSharedRef<FMeshConversionCon
 
 			const int32 PhysicsAssetId = Context.AnimBpOverridePhysicsAssetsInfo.AddUnique(Info);
 
-			TSharedPtr<mu::FPhysicsBody> MutableBody = MakePhysicsBodyFromAsset(Context,
+			TSharedPtr<UE::Mutable::Private::FPhysicsBody> MutableBody = MakePhysicsBodyFromAsset(Context,
 				PropertyAsset,
 				MakePhysicsAssetBodySetupRelevancyMap(Context, PropertyAsset, MutableMesh));
 			MutableBody->CustomId = PhysicsAssetId;
@@ -2457,10 +2457,10 @@ TSharedPtr<mu::FMesh> ConvertSkeletalMeshToMutable(TSharedRef<FMeshConversionCon
 		}
 
 		// Ensure Surface Data
-		mu::FMeshSurface& MeshSurface = MutableMesh->Surfaces.Emplace_GetRef();
+		UE::Mutable::Private::FMeshSurface& MeshSurface = MutableMesh->Surfaces.Emplace_GetRef();
 
 		MeshSurface.SubMeshes.Emplace(
-			mu::FSurfaceSubMesh
+			UE::Mutable::Private::FSurfaceSubMesh
 			{
 				0, MutableMesh->VertexBuffers.GetElementCount(),
 				0, MutableMesh->IndexBuffers.GetElementCount(),
@@ -2499,23 +2499,23 @@ TSharedPtr<mu::FMesh> ConvertSkeletalMeshToMutable(TSharedRef<FMeshConversionCon
 			MeshData.bHasClothing = Context.Options.bClothingEnabled && ImportedModel->LODModels[Source.LODIndex].HasClothData();
 		}
 
-		TSharedPtr<mu::FMesh> ForcedFormat = MakeShared<mu::FMesh>();
+		TSharedPtr<UE::Mutable::Private::FMesh> ForcedFormat = MakeShared<UE::Mutable::Private::FMesh>();
 		SetSurfaceFormat(Context, ForcedFormat->GetVertexBuffers(), ForcedFormat->GetIndexBuffers(), MeshData);
 
-		TSharedPtr<mu::FMesh> FormattedMesh = MakeShared<mu::FMesh>();
+		TSharedPtr<UE::Mutable::Private::FMesh> FormattedMesh = MakeShared<UE::Mutable::Private::FMesh>();
 		bool bOutSuccess = false;
 		constexpr bool bKeepSystemBuffers = true;
 		constexpr bool bFormatVertices = true;
 		constexpr bool bFormatIndices = true;
 		constexpr bool bIgnoreMissingChannels = true;
-		mu::MeshFormat(FormattedMesh.Get(), MutableMesh.Get(), ForcedFormat.Get(),
+		UE::Mutable::Private::MeshFormat(FormattedMesh.Get(), MutableMesh.Get(), ForcedFormat.Get(),
 			bKeepSystemBuffers,
 			bFormatVertices,
 			bFormatIndices,
 			bIgnoreMissingChannels,
 			bOutSuccess);
 
-		mu::MeshOptimizeBuffers(FormattedMesh.Get());
+		UE::Mutable::Private::MeshOptimizeBuffers(FormattedMesh.Get());
 
 		check(bOutSuccess);
 
@@ -2526,7 +2526,7 @@ TSharedPtr<mu::FMesh> ConvertSkeletalMeshToMutable(TSharedRef<FMeshConversionCon
 }
 
 
-TSharedPtr<mu::FMesh> ConvertSkeletalMeshToMutable(
+TSharedPtr<UE::Mutable::Private::FMesh> ConvertSkeletalMeshToMutable(
 	FMutableSourceMeshData& Source,
 	bool bMeshMustExist, int32 LODIndex, int32 SectionIndex,
 	FMutableGraphGenerationContext& GenerationContext, 
@@ -2559,7 +2559,7 @@ TSharedPtr<mu::FMesh> ConvertSkeletalMeshToMutable(
 		}
 	}
 
-	TSharedPtr<mu::FMesh> MutableMesh;
+	TSharedPtr<UE::Mutable::Private::FMesh> MutableMesh;
 	if (!bForceImmediateConversion)
 	{
 		MutableMesh = GenerateMeshConstant(Source, GenerationContext);
@@ -2573,9 +2573,9 @@ TSharedPtr<mu::FMesh> ConvertSkeletalMeshToMutable(
 		MeshConversionContext->CompilationContext = GenerationContext.CompilationContext;
 
 		// Load assets sync
-		MeshConversionContext->ReferencedObjects.Emplace(TStrongObjectPtr<UObject>(MutablePrivate::LoadObject(Source.Mesh)));
-		MeshConversionContext->ReferencedObjects.Emplace(TStrongObjectPtr<UObject>(MutablePrivate::LoadObject(Source.TableReferenceSkeletalMesh)));
-		MeshConversionContext->ReferencedObjects.Emplace(TStrongObjectPtr<UObject>(MutablePrivate::LoadClass(Source.AnimInstance)));
+		MeshConversionContext->ReferencedObjects.Emplace(TStrongObjectPtr<UObject>(UE::Mutable::Private::LoadObject(Source.Mesh)));
+		MeshConversionContext->ReferencedObjects.Emplace(TStrongObjectPtr<UObject>(UE::Mutable::Private::LoadObject(Source.TableReferenceSkeletalMesh)));
+		MeshConversionContext->ReferencedObjects.Emplace(TStrongObjectPtr<UObject>(UE::Mutable::Private::LoadClass(Source.AnimInstance)));
 
 		MutableMesh = ConvertSkeletalMeshToMutable(MeshConversionContext);
 	}
@@ -2583,7 +2583,7 @@ TSharedPtr<mu::FMesh> ConvertSkeletalMeshToMutable(
 	return MutableMesh;
 }
 
-TSharedPtr<mu::FMesh> ConvertStaticMeshToMutable(const UStaticMesh* StaticMesh, int32 LODIndex, int32 SectionIndex, FMutableGraphGenerationContext& GenerationContext, const UCustomizableObjectNode* CurrentNode)
+TSharedPtr<UE::Mutable::Private::FMesh> ConvertStaticMeshToMutable(const UStaticMesh* StaticMesh, int32 LODIndex, int32 SectionIndex, FMutableGraphGenerationContext& GenerationContext, const UCustomizableObjectNode* CurrentNode)
 {
 	if (!StaticMesh->GetRenderData() ||
 		!StaticMesh->GetRenderData()->LODResources.IsValidIndex(LODIndex) ||
@@ -2594,7 +2594,7 @@ TSharedPtr<mu::FMesh> ConvertStaticMeshToMutable(const UStaticMesh* StaticMesh,
 		return nullptr;
 	}
 
-	TSharedPtr<mu::FMesh> MutableMesh = MakeShared<mu::FMesh>();
+	TSharedPtr<UE::Mutable::Private::FMesh> MutableMesh = MakeShared<UE::Mutable::Private::FMesh>();
 
 	// Vertices
 	int32 VertexStart = StaticMesh->GetRenderData()->LODResources[LODIndex].Sections[SectionIndex].MinVertexIndex;
@@ -2602,7 +2602,7 @@ TSharedPtr<mu::FMesh> ConvertStaticMeshToMutable(const UStaticMesh* StaticMesh,
 
 	MutableMesh->GetVertexBuffers().SetElementCount(VertexCount);
 	{
-		using namespace mu;
+		using namespace UE::Mutable::Private;
 
 		MutableMesh->GetVertexBuffers().SetBufferCount(5);
 
@@ -2727,7 +2727,7 @@ TSharedPtr<mu::FMesh> ConvertStaticMeshToMutable(const UStaticMesh* StaticMesh,
 		MutableMesh->GetIndexBuffers().SetBufferCount(1);
 		MutableMesh->GetIndexBuffers().SetElementCount(IndexCount);
 
-		using namespace mu;
+		using namespace UE::Mutable::Private;
 		const int ElementSize = 2;
 		const int ChannelCount = 1;
 		const EMeshBufferSemantic Semantics[ChannelCount] = { EMeshBufferSemantic::VertexIndex };
@@ -2771,7 +2771,7 @@ TSharedPtr<mu::FMesh> ConvertStaticMeshToMutable(const UStaticMesh* StaticMesh,
 
 // Convert a Mesh constant to a mutable format. UniqueTags are the tags that make this Mesh unique that cannot be merged in the cache 
 // with the exact same Mesh with other tags
-TSharedPtr<mu::FMesh> GenerateMutableSkeletalMesh(FMutableSourceMeshData& Source, int32 LODIndexConnected, int32 SectionIndexConnected,
+TSharedPtr<UE::Mutable::Private::FMesh> GenerateMutableSkeletalMesh(FMutableSourceMeshData& Source, int32 LODIndexConnected, int32 SectionIndexConnected,
 									  int32 LODIndex, int32 SectionIndex, const FString& UniqueTags, FMutableGraphGenerationContext & GenerationContext, 
 									  const UCustomizableObjectNode* CurrentNode)
 {
@@ -2779,7 +2779,7 @@ TSharedPtr<mu::FMesh> GenerateMutableSkeletalMesh(FMutableSourceMeshData& Source
 	EMutableMeshConversionFlags CurrentFlags = GenerationContext.MeshGenerationFlags.Last();
 
 	FMutableGraphGenerationContext::FGeneratedMeshData::FKey Key = { Source.Mesh, LODIndex, GenerationContext.CurrentLOD, SectionIndex, CurrentFlags, UniqueTags, CurrentNode };
-	TSharedPtr<mu::FMesh> MutableMesh = GenerationContext.FindGeneratedMesh(Key);
+	TSharedPtr<UE::Mutable::Private::FMesh> MutableMesh = GenerationContext.FindGeneratedMesh(Key);
 	if (MutableMesh)
 	{
 		return MutableMesh;
@@ -2814,7 +2814,7 @@ TSharedPtr<mu::FMesh> GenerateMutableSkeletalMesh(FMutableSourceMeshData& Source
 
 // Convert a Mesh constant to a mutable format. UniqueTags are the tags that make this Mesh unique that cannot be merged in the cache 
 // with the exact same Mesh with other tags
-TSharedPtr<mu::FMesh> GenerateMutableStaticMesh(TSoftObjectPtr<UStreamableRenderAsset> Mesh, const TSoftClassPtr<UAnimInstance>& AnimInstance,
+TSharedPtr<UE::Mutable::Private::FMesh> GenerateMutableStaticMesh(TSoftObjectPtr<UStreamableRenderAsset> Mesh, const TSoftClassPtr<UAnimInstance>& AnimInstance,
 	int32 LODIndex, int32 SectionIndex, const FString& UniqueTags, FMutableGraphGenerationContext& GenerationContext,
 	const UCustomizableObjectNode* CurrentNode, USkeletalMesh* TableReferenceSkeletalMesh, bool bIsPassthrough)
 {
@@ -2822,7 +2822,7 @@ TSharedPtr<mu::FMesh> GenerateMutableStaticMesh(TSoftObjectPtr<UStreamableRender
 	EMutableMeshConversionFlags CurrentFlags = GenerationContext.MeshGenerationFlags.Last();
 
 	FMutableGraphGenerationContext::FGeneratedMeshData::FKey Key = { Mesh, LODIndex, GenerationContext.CurrentLOD, SectionIndex, CurrentFlags, UniqueTags, CurrentNode };
-	TSharedPtr<mu::FMesh> MutableMesh = GenerationContext.FindGeneratedMesh(Key);
+	TSharedPtr<UE::Mutable::Private::FMesh> MutableMesh = GenerationContext.FindGeneratedMesh(Key);
 	if (MutableMesh)
 	{
 		return MutableMesh;
@@ -2830,7 +2830,7 @@ TSharedPtr<mu::FMesh> GenerateMutableStaticMesh(TSoftObjectPtr<UStreamableRender
 
 	// When we want to defer the mesh conversion to the core compilation stage, but it is not supported for static meshes yet.
 	// Meanwhile:
-	UObject* LoadedMesh = MutablePrivate::LoadObject(Mesh);
+	UObject* LoadedMesh = UE::Mutable::Private::LoadObject(Mesh);
 	if (const UStaticMesh* StaticMesh = Cast<UStaticMesh>(LoadedMesh))
 	{
 		MutableMesh = ConvertStaticMeshToMutable(StaticMesh, LODIndex, SectionIndex, GenerationContext, CurrentNode);
@@ -2849,7 +2849,7 @@ TSharedPtr<mu::FMesh> GenerateMutableStaticMesh(TSoftObjectPtr<UStreamableRender
 }
 
 
-TSharedPtr<mu::FMesh> BuildMorphedMutableMesh(const UEdGraphPin* BaseSourcePin, const FString& MorphTargetName, FMutableGraphGenerationContext& GenerationContext, const bool bOnlyConnectedLOD, const FName& RowName)
+TSharedPtr<UE::Mutable::Private::FMesh> BuildMorphedMutableMesh(const UEdGraphPin* BaseSourcePin, const FString& MorphTargetName, FMutableGraphGenerationContext& GenerationContext, const bool bOnlyConnectedLOD, const FName& RowName)
 {
 	check(BaseSourcePin);
 
@@ -2876,7 +2876,7 @@ TSharedPtr<mu::FMesh> BuildMorphedMutableMesh(const UEdGraphPin* BaseSourcePin,
 	else if (const ICustomizableObjectNodeMeshInterface* TypedNodeSkeletalMesh = Cast<ICustomizableObjectNodeMeshInterface>(Node))
 	{
 		TypedNodeSkeletalMesh->GetPinSection(*BaseSourcePin, LODIndexConnected, SectionIndexConnected);
-		SkeletalMesh = Cast<USkeletalMesh>(MutablePrivate::LoadObject(TypedNodeSkeletalMesh->GetMesh()));
+		SkeletalMesh = Cast<USkeletalMesh>(UE::Mutable::Private::LoadObject(TypedNodeSkeletalMesh->GetMesh()));
 	}
 	
 	else if (const UCustomizableObjectNodeTable* TypedNodeTable = Cast<UCustomizableObjectNodeTable>(Node))
@@ -2895,7 +2895,7 @@ TSharedPtr<mu::FMesh> BuildMorphedMutableMesh(const UEdGraphPin* BaseSourcePin,
 		unimplemented();	
 	}
 
-	TSharedPtr<mu::FMesh> MorphedSourceMesh;
+	TSharedPtr<UE::Mutable::Private::FMesh> MorphedSourceMesh;
 
 	if (SkeletalMesh)
 	{
@@ -2903,7 +2903,7 @@ TSharedPtr<mu::FMesh> BuildMorphedMutableMesh(const UEdGraphPin* BaseSourcePin,
 		// Get the base mesh
 		FMutableSourceMeshData Source;
 		Source.Mesh = SkeletalMesh;
-		TSharedPtr<mu::FMesh> BaseSourceMesh = GenerateMutableSkeletalMesh(Source, LODIndexConnected, SectionIndexConnected,
+		TSharedPtr<UE::Mutable::Private::FMesh> BaseSourceMesh = GenerateMutableSkeletalMesh(Source, LODIndexConnected, SectionIndexConnected,
 			LODIndex, SectionIndex, FString(), GenerationContext, Node);
 
 		if (BaseSourceMesh)
@@ -2920,9 +2920,9 @@ TSharedPtr<mu::FMesh> BuildMorphedMutableMesh(const UEdGraphPin* BaseSourcePin,
 }
 
 
-mu::Ptr<mu::NodeScalar> GenerateMorphFactor(const UCustomizableObjectNode* Node, const UEdGraphPin& FactorPin, FMutableGraphGenerationContext& GenerationContext)
+UE::Mutable::Private::Ptr<UE::Mutable::Private::NodeScalar> GenerateMorphFactor(const UCustomizableObjectNode* Node, const UEdGraphPin& FactorPin, FMutableGraphGenerationContext& GenerationContext)
 {
-	mu::Ptr<mu::NodeScalar> FactorNode;
+	UE::Mutable::Private::Ptr<UE::Mutable::Private::NodeScalar> FactorNode;
 
 	if (const UEdGraphPin* ConnectedPin = FollowInputPin(FactorPin))
 	{
@@ -3421,10 +3421,10 @@ bool GetAndValidateReshapePhysicsToDeform(
 }
 
 
-mu::NodeMeshPtr GenerateMorphMesh(const UEdGraphPin* Pin,
+UE::Mutable::Private::NodeMeshPtr GenerateMorphMesh(const UEdGraphPin* Pin,
 	TArray<FMorphNodeData> TypedNodeMorphs,
 	int32 MorphIndex,
-	mu::Ptr<mu::NodeMesh> SourceNode,
+	UE::Mutable::Private::Ptr<UE::Mutable::Private::NodeMesh> SourceNode,
 	FMutableGraphGenerationContext & GenerationContext,
 	const bool bOnlyConnectedLOD,
 	const FString& TableColumnName = "")
@@ -3439,7 +3439,7 @@ mu::NodeMeshPtr GenerateMorphMesh(const UEdGraphPin* Pin,
 	const UCustomizableObjectNode* MorphNode = TypedNodeMorphs[MorphIndex].OwningNode;
 	check(MorphNode);
 	
-	mu::Ptr<mu::NodeMeshMorph> Result = new mu::NodeMeshMorph();
+	UE::Mutable::Private::Ptr<UE::Mutable::Private::NodeMeshMorph> Result = new UE::Mutable::Private::NodeMeshMorph();
 	
 	// Factor
 	//GenerateMorphFactor(MorphNode, , GenerationContext, Result);
@@ -3453,21 +3453,21 @@ mu::NodeMeshPtr GenerateMorphMesh(const UEdGraphPin* Pin,
 	}
 	else
 	{
-		mu::NodeMeshPtr NextMorph = GenerateMorphMesh(Pin, TypedNodeMorphs, MorphIndex + 1, SourceNode, GenerationContext, bOnlyConnectedLOD, TableColumnName); // TODO FutureGMT change to a for. This recursion can be problematic with the production cache
+		UE::Mutable::Private::NodeMeshPtr NextMorph = GenerateMorphMesh(Pin, TypedNodeMorphs, MorphIndex + 1, SourceNode, GenerationContext, bOnlyConnectedLOD, TableColumnName); // TODO FutureGMT change to a for. This recursion can be problematic with the production cache
 		Result->Base = NextMorph;
 	}
 	
 	// Target
-	mu::NodeMeshPtr BaseSourceMesh = SourceNode;
+	UE::Mutable::Private::NodeMeshPtr BaseSourceMesh = SourceNode;
 
-	TSharedPtr<mu::FMesh> MorphedSourceMesh;
+	TSharedPtr<UE::Mutable::Private::FMesh> MorphedSourceMesh;
 
 	bool bSuccess = false;
 	
 	if (const UCustomizableObjectNodeTable* TypedNodeTable = Cast<UCustomizableObjectNodeTable>(Pin->GetOwningNode()))
 	{
-		TObjectPtr<UDataTable> Table = MutablePrivate::LoadObject(TypedNodeTable->Table);
-		TObjectPtr<UScriptStruct> Structure = MutablePrivate::LoadObject(TypedNodeTable->Structure);
+		TObjectPtr<UDataTable> Table = UE::Mutable::Private::LoadObject(TypedNodeTable->Table);
+		TObjectPtr<UScriptStruct> Structure = UE::Mutable::Private::LoadObject(TypedNodeTable->Structure);
 		const FString TableName = Table ? GetNameSafe(Table).ToLower() : GetNameSafe(Structure).ToLower();
 		const uint32 TableId = CityHash32(reinterpret_cast<const char*>(*TableName), TableName.Len() * sizeof(FString::ElementType));
 
@@ -3480,7 +3480,7 @@ mu::NodeMeshPtr GenerateMorphMesh(const UEdGraphPin* Pin,
 
 		// Should exist
 		FGeneratedMutableDataTableKey MutableTableKey(DataTable->GetName(), TypedNodeTable->VersionColumn, TypedNodeTable->CompilationFilterOptions);
-		mu::Ptr<mu::FTable> GeneratedTable = GenerationContext.GeneratedTables[MutableTableKey].GeneratedTable;
+		UE::Mutable::Private::Ptr<UE::Mutable::Private::FTable> GeneratedTable = GenerationContext.GeneratedTables[MutableTableKey].GeneratedTable;
 		check(GeneratedTable);
 
 		FString ColumnName = TableColumnName + TypedNodeMorphs[MorphIndex].MorphTargetName;
@@ -3495,10 +3495,10 @@ mu::NodeMeshPtr GenerateMorphMesh(const UEdGraphPin* Pin,
 
 			if (ColumnIndex == INDEX_NONE)
 			{
-				ColumnIndex = GeneratedTable->AddColumn(ColumnName, mu::ETableColumnType::Mesh);
+				ColumnIndex = GeneratedTable->AddColumn(ColumnName, UE::Mutable::Private::ETableColumnType::Mesh);
 			}
 
-			TSharedPtr<mu::FMesh> MorphedSourceTableMesh = BuildMorphedMutableMesh(Pin, TypedNodeMorphs[MorphIndex].MorphTargetName, GenerationContext, bOnlyConnectedLOD, RowName);
+			TSharedPtr<UE::Mutable::Private::FMesh> MorphedSourceTableMesh = BuildMorphedMutableMesh(Pin, TypedNodeMorphs[MorphIndex].MorphTargetName, GenerationContext, bOnlyConnectedLOD, RowName);
 			GeneratedTable->SetCell(ColumnIndex, RowId, MorphedSourceTableMesh);
 		}
 
@@ -3506,7 +3506,7 @@ mu::NodeMeshPtr GenerateMorphMesh(const UEdGraphPin* Pin,
 		{
 			bSuccess = true;
 
-			mu::Ptr<mu::NodeMeshTable> MorphedSourceMeshNodeTable = new mu::NodeMeshTable;
+			UE::Mutable::Private::Ptr<UE::Mutable::Private::NodeMeshTable> MorphedSourceMeshNodeTable = new UE::Mutable::Private::NodeMeshTable;
 			MorphedSourceMeshNodeTable->Table = GeneratedTable;
 			MorphedSourceMeshNodeTable->ColumnName = ColumnName;
 			MorphedSourceMeshNodeTable->ParameterName = TypedNodeTable->ParameterName;
@@ -3514,7 +3514,7 @@ mu::NodeMeshPtr GenerateMorphMesh(const UEdGraphPin* Pin,
 			MorphedSourceMeshNodeTable->SourceDataDescriptor.SourceId = TableId;
 			MorphedSourceMeshNodeTable->SourceDataDescriptor.OptionalMaxLODSize = 0;
 
-			mu::Ptr<mu::NodeMeshMakeMorph> Morph = new mu::NodeMeshMakeMorph;
+			UE::Mutable::Private::Ptr<UE::Mutable::Private::NodeMeshMakeMorph> Morph = new UE::Mutable::Private::NodeMeshMakeMorph;
 			Morph->Base = BaseSourceMesh;
 			Morph->Target = MorphedSourceMeshNodeTable;
 			Morph->bOnlyPositionAndNormal = true;
@@ -3530,11 +3530,11 @@ mu::NodeMeshPtr GenerateMorphMesh(const UEdGraphPin* Pin,
 
 		if (MorphedSourceMesh)
 		{
-			mu::Ptr<mu::NodeMeshConstant> MorphedSourceMeshNode = new mu::NodeMeshConstant;
+			UE::Mutable::Private::Ptr<UE::Mutable::Private::NodeMeshConstant> MorphedSourceMeshNode = new UE::Mutable::Private::NodeMeshConstant;
 			MorphedSourceMeshNode->Value = MorphedSourceMesh;
 			MorphedSourceMeshNode->SetMessageContext(MorphNode);
 
-			mu::Ptr<mu::NodeMeshMakeMorph> Morph = new mu::NodeMeshMakeMorph;
+			UE::Mutable::Private::Ptr<UE::Mutable::Private::NodeMeshMakeMorph> Morph = new UE::Mutable::Private::NodeMeshMakeMorph;
 			Morph->Base = BaseSourceMesh;
 			Morph->Target = MorphedSourceMeshNode;
 			Morph->bOnlyPositionAndNormal = true;
@@ -3624,7 +3624,7 @@ TArray<FMorphNodeData> GenerateMeshMorphStackDefinition(const UEdGraphPin* Pin,
 			if (MorphPin->LinkedTo.Num() && MorphTargetNames.Contains(MorphPinName))
 			{
 				// Generate Factor
-				mu::Ptr<mu::NodeScalar> FactorNode = GenerateMorphFactor(Node, *MorphPin, GenerationContext);
+				UE::Mutable::Private::Ptr<UE::Mutable::Private::NodeScalar> FactorNode = GenerateMorphFactor(Node, *MorphPin, GenerationContext);
 
 				OutMorphData.Add({ Node, MorphPinName, FactorNode, MeshPin});
 			}
@@ -3688,7 +3688,7 @@ TArray<FMorphNodeData> GenerateMeshMorphStackDefinition(const UEdGraphPin* Pin,
 }
 
 
-mu::Ptr<mu::NodeMesh> GenerateMutableSourceMesh(const UEdGraphPin* Pin,
+UE::Mutable::Private::Ptr<UE::Mutable::Private::NodeMesh> GenerateMutableSourceMesh(const UEdGraphPin* Pin,
 	FMutableGraphGenerationContext& GenerationContext,
 	const FMutableSourceMeshData& BaseMeshData,
 	const bool bLinkedToExtendMaterial,
@@ -3706,7 +3706,7 @@ mu::Ptr<mu::NodeMesh> GenerateMutableSourceMesh(const UEdGraphPin* Pin,
 	const FGeneratedKey Key(reinterpret_cast<void*>(&GenerateMutableSourceMesh), *Pin , *Node, GenerationContext, true, bOnlyConnectedLOD);
 	if (const FGeneratedData* Generated = GenerationContext.Generated.Find(Key))
 	{
-		return static_cast<mu::NodeMesh*>(Generated->Node.get());
+		return static_cast<UE::Mutable::Private::NodeMesh*>(Generated->Node.get());
 	}
 	
 	if (Node->IsNodeOutDatedAndNeedsRefresh())
@@ -3719,14 +3719,14 @@ mu::Ptr<mu::NodeMesh> GenerateMutableSourceMesh(const UEdGraphPin* Pin,
 	bool bCacheNode = true;
 
 	//SkeletalMesh Result
-	mu::Ptr<mu::NodeMesh> Result;
+	UE::Mutable::Private::Ptr<UE::Mutable::Private::NodeMesh> Result;
 
 	//SkeletalMesh + Morphs Result
-	mu::Ptr<mu::NodeMesh> MorphResult;
+	UE::Mutable::Private::Ptr<UE::Mutable::Private::NodeMesh> MorphResult;
 	
 	if (const UCustomizableObjectNodeSkeletalMesh* TypedNodeSkel = Cast<UCustomizableObjectNodeSkeletalMesh>(Node))
 	{
-		mu::Ptr<mu::NodeMeshConstant> MeshNode = new mu::NodeMeshConstant();
+		UE::Mutable::Private::Ptr<UE::Mutable::Private::NodeMeshConstant> MeshNode = new UE::Mutable::Private::NodeMeshConstant();
 		Result = MeshNode;
 
 		if (TypedNodeSkel->SkeletalMesh.IsNull())
@@ -3810,7 +3810,7 @@ mu::Ptr<mu::NodeMesh> GenerateMutableSourceMesh(const UEdGraphPin* Pin,
 				Source.Mesh = SkeletalMesh;
 				Source.AnimInstance = TypedNodeSkel->AnimInstance;
 
-				TSharedPtr<mu::FMesh> MutableMesh = GenerateMutableSkeletalMesh(Source, LODIndexConnected, SectionIndexConnected,
+				TSharedPtr<UE::Mutable::Private::FMesh> MutableMesh = GenerateMutableSkeletalMesh(Source, LODIndexConnected, SectionIndexConnected,
 					LODIndex, SectionIndex, MeshUniqueTags, GenerationContext, TypedNodeSkel);
 				if (MutableMesh)
 				{
@@ -3881,7 +3881,7 @@ mu::Ptr<mu::NodeMesh> GenerateMutableSourceMesh(const UEdGraphPin* Pin,
 						const UCustomizableObjectLayout* Layout = Layouts.IsValidIndex(LayoutIndex) ? Layouts[LayoutIndex] : nullptr;
 						if (ensure(Layout))
 						{
-							mu::Ptr<mu::NodeLayout> LayoutNode = CreateMutableLayoutNode(GenerationContext, Layout, bLinkedToExtendMaterial); // TODO PERE: Figure out
+							UE::Mutable::Private::Ptr<UE::Mutable::Private::NodeLayout> LayoutNode = CreateMutableLayoutNode(GenerationContext, Layout, bLinkedToExtendMaterial); // TODO PERE: Figure out
 							LayoutNode->SetMessageContext(Node);
 
 							MeshNode->Layouts[LayoutIndex] = LayoutNode;
@@ -3926,7 +3926,7 @@ mu::Ptr<mu::NodeMesh> GenerateMutableSourceMesh(const UEdGraphPin* Pin,
 			return {};
 		}
 
-		mu::Ptr<mu::NodeMeshConstant> MeshNode = new mu::NodeMeshConstant();
+		UE::Mutable::Private::Ptr<UE::Mutable::Private::NodeMeshConstant> MeshNode = new UE::Mutable::Private::NodeMeshConstant();
 		Result = MeshNode;
 
 		if (TypedNodeStatic->StaticMesh)
@@ -3955,7 +3955,7 @@ mu::Ptr<mu::NodeMesh> GenerateMutableSourceMesh(const UEdGraphPin* Pin,
 			check(TypedNodeStatic->LODs.IsValidIndex(LODIndex) && SectionIndex < TypedNodeStatic->LODs[LODIndex].Materials.Num());
 
 			constexpr bool bIsPassthrough = false;
-			TSharedPtr<mu::FMesh> MutableMesh = GenerateMutableStaticMesh(TypedNodeStatic->StaticMesh, TSoftClassPtr<UAnimInstance>(), LODIndex, SectionIndex,
+			TSharedPtr<UE::Mutable::Private::FMesh> MutableMesh = GenerateMutableStaticMesh(TypedNodeStatic->StaticMesh, TSoftClassPtr<UAnimInstance>(), LODIndex, SectionIndex,
 														  FString(), GenerationContext, TypedNodeStatic, nullptr, bIsPassthrough);
 			if (MutableMesh)
 			{
@@ -3964,7 +3964,7 @@ mu::Ptr<mu::NodeMesh> GenerateMutableSourceMesh(const UEdGraphPin* Pin,
 				// Layouts
 				MeshNode->Layouts.SetNum(1);
 
-				mu::Ptr<mu::NodeLayout> LayoutNode;
+				UE::Mutable::Private::Ptr<UE::Mutable::Private::NodeLayout> LayoutNode;
 
 				const TArray<UCustomizableObjectLayout*>& Layouts = TypedNodeStatic->GetLayouts(*Pin);
 				const UCustomizableObjectLayout* Layout = Layouts.IsValidIndex(0) ? Layouts[0] : nullptr;
@@ -4018,7 +4018,7 @@ mu::Ptr<mu::NodeMesh> GenerateMutableSourceMesh(const UEdGraphPin* Pin,
 			}
 
 			// Factor
-			mu::Ptr<mu::NodeScalar> FactorNode = GenerateMorphFactor(Node, *TypedNodeMorph->FactorPin(), GenerationContext);
+			UE::Mutable::Private::Ptr<UE::Mutable::Private::NodeScalar> FactorNode = GenerateMorphFactor(Node, *TypedNodeMorph->FactorPin(), GenerationContext);
 
 			// Mesh Morph Stack Management
 			FMorphNodeData NewMorphData = { TypedNodeMorph, MorphTargetName, FactorNode, TypedNodeMorph->MeshPin() };
@@ -4028,7 +4028,7 @@ mu::Ptr<mu::NodeMesh> GenerateMutableSourceMesh(const UEdGraphPin* Pin,
 		}
 		else
 		{
-			mu::Ptr<mu::NodeMeshMorph> MeshNode = new mu::NodeMeshMorph();
+			UE::Mutable::Private::Ptr<UE::Mutable::Private::NodeMeshMorph> MeshNode = new UE::Mutable::Private::NodeMeshMorph();
 			Result = MeshNode;
 		}
 	}
@@ -4075,7 +4075,7 @@ mu::Ptr<mu::NodeMesh> GenerateMutableSourceMesh(const UEdGraphPin* Pin,
 			// Check Switch Parameter arity preconditions.
 			if (const UEdGraphPin* EnumPin = FollowInputPin(*SwitchParameter))
 			{
-				mu::Ptr<mu::NodeScalar> SwitchParam = GenerateMutableSourceFloat(EnumPin, GenerationContext);
+				UE::Mutable::Private::Ptr<UE::Mutable::Private::NodeScalar> SwitchParam = GenerateMutableSourceFloat(EnumPin, GenerationContext);
 
 				// Switch Param not generated
 				if (!SwitchParam)
@@ -4090,7 +4090,7 @@ mu::Ptr<mu::NodeMesh> GenerateMutableSourceMesh(const UEdGraphPin* Pin,
 					return Result;
 				}
 
-				if (SwitchParam->GetType() != mu::NodeScalarEnumParameter::GetStaticType())
+				if (SwitchParam->GetType() != UE::Mutable::Private::NodeScalarEnumParameter::GetStaticType())
 				{
 					const FText Message = LOCTEXT("WrongSwitchParamType", "Switch parameter of incorrect type.");
 					GenerationContext.Log(Message, Node);
@@ -4100,14 +4100,14 @@ mu::Ptr<mu::NodeMesh> GenerateMutableSourceMesh(const UEdGraphPin* Pin,
 
 				const int32 NumSwitchOptions = TypedNodeMeshSwitch->GetNumElements();
 
-				mu::NodeScalarEnumParameter* EnumParameter = static_cast<mu::NodeScalarEnumParameter*>(SwitchParam.get());
+				UE::Mutable::Private::NodeScalarEnumParameter* EnumParameter = static_cast<UE::Mutable::Private::NodeScalarEnumParameter*>(SwitchParam.get());
 				if (NumSwitchOptions != EnumParameter->Options.Num())
 				{
 					const FText Message = LOCTEXT("MismatchedSwitch", "Switch enum and switch node have different number of options. Please refresh the switch node to make sure the outcomes are labeled properly.");
 					GenerationContext.Log(Message, Node);
 				}
 
-				mu::Ptr<mu::NodeMeshSwitch> SwitchNode = new mu::NodeMeshSwitch;
+				UE::Mutable::Private::Ptr<UE::Mutable::Private::NodeMeshSwitch> SwitchNode = new UE::Mutable::Private::NodeMeshSwitch;
 				SwitchNode->Parameter = SwitchParam;
 				SwitchNode->Options.SetNum(NumSwitchOptions);
 
@@ -4133,12 +4133,12 @@ mu::Ptr<mu::NodeMesh> GenerateMutableSourceMesh(const UEdGraphPin* Pin,
 
 	else if (const UCustomizableObjectNodeMeshVariation* TypedNodeMeshVar = Cast<const UCustomizableObjectNodeMeshVariation>(Node))
 	{
-		mu::Ptr<mu::NodeMeshVariation> MeshNode = new mu::NodeMeshVariation();
+		UE::Mutable::Private::Ptr<UE::Mutable::Private::NodeMeshVariation> MeshNode = new UE::Mutable::Private::NodeMeshVariation();
 		Result = MeshNode;
 
 		if (const UEdGraphPin* ConnectedPin = FollowInputPin(*TypedNodeMeshVar->DefaultPin()))
 		{
-			mu::NodeMeshPtr ChildNode = GenerateMutableSourceMesh(ConnectedPin, GenerationContext, BaseMeshData, false, bOnlyConnectedLOD);
+			UE::Mutable::Private::NodeMeshPtr ChildNode = GenerateMutableSourceMesh(ConnectedPin, GenerationContext, BaseMeshData, false, bOnlyConnectedLOD);
 			if (ChildNode)
 			{
 				MeshNode->DefaultMesh = ChildNode;
@@ -4163,7 +4163,7 @@ mu::Ptr<mu::NodeMesh> GenerateMutableSourceMesh(const UEdGraphPin* Pin,
 
 			if (const UEdGraphPin* ConnectedPin = FollowInputPin(*VariationPin))
 			{
-				mu::NodeMeshPtr ChildNode = GenerateMutableSourceMesh(ConnectedPin, GenerationContext, BaseMeshData, false, bOnlyConnectedLOD);
+				UE::Mutable::Private::NodeMeshPtr ChildNode = GenerateMutableSourceMesh(ConnectedPin, GenerationContext, BaseMeshData, false, bOnlyConnectedLOD);
 				MeshNode->Variations[VariationIndex].Mesh = ChildNode;
 			}
 		}
@@ -4171,12 +4171,12 @@ mu::Ptr<mu::NodeMesh> GenerateMutableSourceMesh(const UEdGraphPin* Pin,
 
 	else if (const UCustomizableObjectNodeMeshReshape* TypedNodeReshape = Cast<const UCustomizableObjectNodeMeshReshape>(Node))
 	{
-		mu::Ptr<mu::NodeMeshReshape> MeshNode = new mu::NodeMeshReshape();
+		UE::Mutable::Private::Ptr<UE::Mutable::Private::NodeMeshReshape> MeshNode = new UE::Mutable::Private::NodeMeshReshape();
 		Result = MeshNode;
 
 		if (const UEdGraphPin* ConnectedPin = FollowInputPin(*TypedNodeReshape->BaseMeshPin()))
 		{
-			mu::Ptr<mu::NodeMesh> ChildNode = GenerateMutableSourceMesh(ConnectedPin, GenerationContext, BaseMeshData, false, bOnlyConnectedLOD);
+			UE::Mutable::Private::Ptr<UE::Mutable::Private::NodeMesh> ChildNode = GenerateMutableSourceMesh(ConnectedPin, GenerationContext, BaseMeshData, false, bOnlyConnectedLOD);
 			if (ChildNode)
 			{
 				MeshNode->BaseMesh = ChildNode;
@@ -4234,14 +4234,14 @@ mu::Ptr<mu::NodeMesh> GenerateMutableSourceMesh(const UEdGraphPin* Pin,
 				}
 			}
 
-			auto ConvertColorUsage = [](EMeshReshapeVertexColorChannelUsage Usage) -> mu::EVertexColorUsage
+			auto ConvertColorUsage = [](EMeshReshapeVertexColorChannelUsage Usage) -> UE::Mutable::Private::EVertexColorUsage
 			{
 				switch (Usage)
 				{
-				case EMeshReshapeVertexColorChannelUsage::None:			  return mu::EVertexColorUsage::None;
-				case EMeshReshapeVertexColorChannelUsage::RigidClusterId: return mu::EVertexColorUsage::ReshapeClusterId;
-				case EMeshReshapeVertexColorChannelUsage::MaskWeight:     return mu::EVertexColorUsage::ReshapeMaskWeight;
-				default: check(false); return mu::EVertexColorUsage::None;
+				case EMeshReshapeVertexColorChannelUsage::None:			  return UE::Mutable::Private::EVertexColorUsage::None;
+				case EMeshReshapeVertexColorChannelUsage::RigidClusterId: return UE::Mutable::Private::EVertexColorUsage::ReshapeClusterId;
+				case EMeshReshapeVertexColorChannelUsage::MaskWeight:     return UE::Mutable::Private::EVertexColorUsage::ReshapeMaskWeight;
+				default: check(false); return UE::Mutable::Private::EVertexColorUsage::None;
 				};
 			};
 
@@ -4300,7 +4300,7 @@ mu::Ptr<mu::NodeMesh> GenerateMutableSourceMesh(const UEdGraphPin* Pin,
 			
 		if (const UEdGraphPin* ConnectedPin = FollowInputPin(*TypedNodeReshape->BaseShapePin()))
 		{
-			mu::Ptr<mu::NodeMesh> ChildNode = GenerateMutableSourceMesh(ConnectedPin, GenerationContext, BaseMeshData, false, true);
+			UE::Mutable::Private::Ptr<UE::Mutable::Private::NodeMesh> ChildNode = GenerateMutableSourceMesh(ConnectedPin, GenerationContext, BaseMeshData, false, true);
 	
 			if (ChildNode)
 			{
@@ -4314,7 +4314,7 @@ mu::Ptr<mu::NodeMesh> GenerateMutableSourceMesh(const UEdGraphPin* Pin,
 
 		if (const UEdGraphPin* ConnectedPin = FollowInputPin(*TypedNodeReshape->TargetShapePin()))
 		{
-			mu::Ptr<mu::NodeMesh> ChildNode = GenerateMutableSourceMesh(ConnectedPin, GenerationContext, BaseMeshData, false, true);
+			UE::Mutable::Private::Ptr<UE::Mutable::Private::NodeMesh> ChildNode = GenerateMutableSourceMesh(ConnectedPin, GenerationContext, BaseMeshData, false, true);
 			
 			if (ChildNode)
 			{
@@ -4333,7 +4333,7 @@ mu::Ptr<mu::NodeMesh> GenerateMutableSourceMesh(const UEdGraphPin* Pin,
 	{
 		if (const UEdGraphPin* InputMeshPin = FollowInputPin(*TypedNode->GetInputMeshPin()))
 		{
-			mu::Ptr<mu::NodeMesh> InputMeshNode = GenerateMutableSourceMesh(InputMeshPin, GenerationContext, BaseMeshData, false, bOnlyConnectedLOD);
+			UE::Mutable::Private::Ptr<UE::Mutable::Private::NodeMesh> InputMeshNode = GenerateMutableSourceMesh(InputMeshPin, GenerationContext, BaseMeshData, false, bOnlyConnectedLOD);
 
 			if (GenerationContext.GetCurrentComponentInfo()->RefSkeletalMesh)
 			{
@@ -4342,7 +4342,7 @@ mu::Ptr<mu::NodeMesh> GenerateMutableSourceMesh(const UEdGraphPin* Pin,
 					TArray<FName> ArrayBoneName;
 					TArray<FTransform> ArrayTransform;
 					UCustomizableObjectNodeAnimationPose::StaticRetrievePoseInformation(TypedNode->PoseAsset, GenerationContext.GetCurrentComponentInfo()->RefSkeletalMesh.Get(), ArrayBoneName, ArrayTransform);
-					mu::Ptr<mu::NodeMeshApplyPose> NodeMeshApplyPose = CreateNodeMeshApplyPose(GenerationContext, InputMeshNode, ArrayBoneName, ArrayTransform);
+					UE::Mutable::Private::Ptr<UE::Mutable::Private::NodeMeshApplyPose> NodeMeshApplyPose = CreateNodeMeshApplyPose(GenerationContext, InputMeshNode, ArrayBoneName, ArrayTransform);
 
 					if (NodeMeshApplyPose)
 					{
@@ -4358,8 +4358,8 @@ mu::Ptr<mu::NodeMesh> GenerateMutableSourceMesh(const UEdGraphPin* Pin,
 				}
 				else if (const UEdGraphPin* TablePosePin = FollowInputPin(*TypedNode->GetTablePosePin()))
 				{
-					mu::Ptr<mu::NodeMeshApplyPose> NodeMeshApplyPose = new mu::NodeMeshApplyPose();
-					mu::Ptr<mu::NodeMesh> MeshTableNode = GenerateMutableSourceMesh(TablePosePin, GenerationContext, BaseMeshData, false, bOnlyConnectedLOD);
+					UE::Mutable::Private::Ptr<UE::Mutable::Private::NodeMeshApplyPose> NodeMeshApplyPose = new UE::Mutable::Private::NodeMeshApplyPose();
+					UE::Mutable::Private::Ptr<UE::Mutable::Private::NodeMesh> MeshTableNode = GenerateMutableSourceMesh(TablePosePin, GenerationContext, BaseMeshData, false, bOnlyConnectedLOD);
 
 					NodeMeshApplyPose->Base = InputMeshNode;
 					NodeMeshApplyPose->Pose = MeshTableNode;
@@ -4383,12 +4383,12 @@ mu::Ptr<mu::NodeMesh> GenerateMutableSourceMesh(const UEdGraphPin* Pin,
 
 	else if (const UCustomizableObjectNodeTable* TypedNodeTable = Cast<UCustomizableObjectNodeTable>(Node))
 	{
-		TObjectPtr<UDataTable> Table = MutablePrivate::LoadObject(TypedNodeTable->Table);
-		TObjectPtr<UScriptStruct> Structure = MutablePrivate::LoadObject(TypedNodeTable->Structure);
+		TObjectPtr<UDataTable> Table = UE::Mutable::Private::LoadObject(TypedNodeTable->Table);
+		TObjectPtr<UScriptStruct> Structure = UE::Mutable::Private::LoadObject(TypedNodeTable->Structure);
 		const FString TableName = Table ? GetNameSafe(Table).ToLower() : GetNameSafe(Structure).ToLower();
 		const uint32 TableId = CityHash32(reinterpret_cast<const char*>(*TableName), TableName.Len() * sizeof(FString::ElementType));
 
-		mu::Ptr<mu::NodeMeshConstant> EmptyNode = new mu::NodeMeshConstant();
+		UE::Mutable::Private::Ptr<UE::Mutable::Private::NodeMeshConstant> EmptyNode = new UE::Mutable::Private::NodeMeshConstant();
 		Result = EmptyNode;
 		bool bSuccess = true;
 
@@ -4425,12 +4425,12 @@ mu::Ptr<mu::NodeMesh> GenerateMutableSourceMesh(const UEdGraphPin* Pin,
 			if (bSuccess)
 			{
 				// Generating a new data table if not exists
-				mu::Ptr<mu::FTable> GeneratedTable = nullptr;
+				UE::Mutable::Private::Ptr<UE::Mutable::Private::FTable> GeneratedTable = nullptr;
 				GeneratedTable = GenerateMutableSourceTable(DataTable, TypedNodeTable, GenerationContext);
 
 				if (GeneratedTable)
 				{
-					mu::Ptr<mu::NodeMeshTable> MeshTableNode = new mu::NodeMeshTable();
+					UE::Mutable::Private::Ptr<UE::Mutable::Private::NodeMeshTable> MeshTableNode = new UE::Mutable::Private::NodeMeshTable();
 					
 					int32 LODIndexConnected = -1; // LOD which the pin is connected to
 					int32 SectionIndexConnected = -1;
@@ -4510,7 +4510,7 @@ mu::Ptr<mu::NodeMesh> GenerateMutableSourceMesh(const UEdGraphPin* Pin,
 							
 								// In tables, mimic the legacy behaviour and ignore all layout warnings beyond LOD 0.
 								bool bIgnoreLayoutWarnings = true;
-								mu::Ptr<mu::NodeLayout> LayoutNode = CreateMutableLayoutNode(GenerationContext, Layouts[LayoutIndex], bIgnoreLayoutWarnings);
+								UE::Mutable::Private::Ptr<UE::Mutable::Private::NodeLayout> LayoutNode = CreateMutableLayoutNode(GenerationContext, Layouts[LayoutIndex], bIgnoreLayoutWarnings);
 							
 								MeshTableNode->Layouts[LayoutIndex] = LayoutNode;
 							}
@@ -4538,14 +4538,14 @@ mu::Ptr<mu::NodeMesh> GenerateMutableSourceMesh(const UEdGraphPin* Pin,
 
 	else if (const UCustomizableObjectNodeSkeletalMeshParameter* TypedNodeParam = Cast<UCustomizableObjectNodeSkeletalMeshParameter>(Node))
 	{
-		mu::Ptr<mu::NodeMeshParameter> ParameterNode = new mu::NodeMeshParameter();
+		UE::Mutable::Private::Ptr<UE::Mutable::Private::NodeMeshParameter> ParameterNode = new UE::Mutable::Private::NodeMeshParameter();
 
 		ParameterNode->Name = TypedNodeParam->GetParameterName();
 		ParameterNode->UID = GenerationContext.GetNodeIdUnique(Node).ToString();
 
 		if (TypedNodeParam->DefaultValue)
 		{
-			GenerationContext.SkeletalMeshParameterDefaultValues.Add(FName(TypedNodeParam->GetParameterName()), MutablePrivate::LoadObject(TypedNodeParam->DefaultValue));
+			GenerationContext.SkeletalMeshParameterDefaultValues.Add(FName(TypedNodeParam->GetParameterName()), UE::Mutable::Private::LoadObject(TypedNodeParam->DefaultValue));
 		}
 
 		int32 LODIndex = 0;
@@ -4596,7 +4596,7 @@ mu::Ptr<mu::NodeMesh> GenerateMutableSourceMesh(const UEdGraphPin* Pin,
 				if (ensure(Layout))
 				{
 					
-					mu::Ptr<mu::NodeLayout> LayoutNode = CreateMutableLayoutNode(GenerationContext, Layout, bLinkedToExtendMaterial);
+					UE::Mutable::Private::Ptr<UE::Mutable::Private::NodeLayout> LayoutNode = CreateMutableLayoutNode(GenerationContext, Layout, bLinkedToExtendMaterial);
 					LayoutNode->SetMessageContext(Node);
 
 					ParameterNode->Layouts[LayoutIndex] = LayoutNode;
diff --git a/Engine/Plugins/Mutable/Source/CustomizableObjectEditor/Private/MuCOE/GenerateMutableSource/GenerateMutableSourceMesh.h b/Engine/Plugins/Mutable/Source/CustomizableObjectEditor/Private/MuCOE/GenerateMutableSource/GenerateMutableSourceMesh.h
--- a/Engine/Plugins/Mutable/Source/CustomizableObjectEditor/Private/MuCOE/GenerateMutableSource/GenerateMutableSourceMesh.h
+++ b/Engine/Plugins/Mutable/Source/CustomizableObjectEditor/Private/MuCOE/GenerateMutableSource/GenerateMutableSourceMesh.h
@@ -20,7 +20,7 @@ struct FMutableGraphMeshGenerationData;
 struct FMutableCompilationContext;
 struct FMutableSourceMeshData;
 
-namespace mu
+namespace UE::Mutable::Private
 {
 	class FMeshBufferSet;
 }
@@ -55,7 +55,7 @@ void GetLODAndSectionForAutomaticLODs(const FMutableGraphGenerationContext& Cont
  * @param GenerationContext 
  * @param CurrentNode The node being compiled, can be a skeletal mesh or a table node.
  * @return Mutable Mesh. Nullptr if there has been an error. Empty mesh if the Skeletal Mesh does not contain the requested LOD + Section. */
-TSharedPtr<mu::FMesh> ConvertSkeletalMeshToMutable(FMutableSourceMeshData& Source,
+TSharedPtr<UE::Mutable::Private::FMesh> ConvertSkeletalMeshToMutable(FMutableSourceMeshData& Source,
                                          bool bMeshMustExist,
                                          int32 LODIndex, int32 SectionIndex,
                                          FMutableGraphGenerationContext& GenerationContext,
@@ -63,7 +63,7 @@ TSharedPtr<mu::FMesh> ConvertSkeletalMeshToMutable(FMutableSourceMeshData& Sourc
 										 bool bForceImmediateConversion = false );
 
 
-TSharedPtr<mu::FMesh> ConvertStaticMeshToMutable(const UStaticMesh* StaticMesh, int32 LODIndex, int32 SectionIndex,
+TSharedPtr<UE::Mutable::Private::FMesh> ConvertStaticMeshToMutable(const UStaticMesh* StaticMesh, int32 LODIndex, int32 SectionIndex,
 	FMutableGraphGenerationContext& GenerationContext, const UCustomizableObjectNode* CurrentNode);
 
 
@@ -77,21 +77,21 @@ TSharedPtr<mu::FMesh> ConvertStaticMeshToMutable(const UStaticMesh* StaticMesh,
  * @param GenerationContext
  * @param CurrentNode The node being compiled, can be a skeletal mesh or a table node.
  * @return Mutable Mesh. Nullptr if there has been an error. Empty mesh if the Skeletal Mesh does not contain the requested LOD + Section. */
-TSharedPtr<mu::FMesh> GenerateMutableSkeletalMesh(FMutableSourceMeshData& Source,
+TSharedPtr<UE::Mutable::Private::FMesh> GenerateMutableSkeletalMesh(FMutableSourceMeshData& Source,
                                 int32 LODIndexConnected, int32 SectionIndexConnected,
                                 int32 LODIndex, int32 SectionIndex, const FString& MeshUniqueTags, 
                                 FMutableGraphGenerationContext& GenerationContext, 
 								const UCustomizableObjectNode* CurrentNode);
 
 /** */
-TSharedPtr<mu::FMesh> GenerateMutableStaticMesh(TSoftObjectPtr<UStreamableRenderAsset> Mesh, const TSoftClassPtr<UAnimInstance>& AnimBp,
+TSharedPtr<UE::Mutable::Private::FMesh> GenerateMutableStaticMesh(TSoftObjectPtr<UStreamableRenderAsset> Mesh, const TSoftClassPtr<UAnimInstance>& AnimBp,
 	int32 LODIndex, int32 SectionIndex, const FString& MeshUniqueTags,
 	FMutableGraphGenerationContext& GenerationContext,
 	const UCustomizableObjectNode* CurrentNode, USkeletalMesh* TableReferenceSkeletalMesh,
 	bool bIsPassthrough);
 
 
-TSharedPtr<mu::FMesh> BuildMorphedMutableMesh(const UEdGraphPin* BaseSourcePin, const FString& MorphTargetName,
+TSharedPtr<UE::Mutable::Private::FMesh> BuildMorphedMutableMesh(const UEdGraphPin* BaseSourcePin, const FString& MorphTargetName,
 	FMutableGraphGenerationContext& GenerationContext, bool bOnlyConnectedLOD, const FName& RowName = "");
 
 
@@ -103,11 +103,11 @@ TSharedPtr<mu::FMesh> BuildMorphedMutableMesh(const UEdGraphPin* BaseSourcePin,
  * @param bLinkedToExtendMaterial 
  * @param bOnlyConnectedLOD Corrected LOD and Section will unconditionally always be the connected ones.
  * @return  Mutable Mesh Node. */
-mu::Ptr<mu::NodeMesh> GenerateMutableSourceMesh(const UEdGraphPin* Pin, FMutableGraphGenerationContext& GenerationContext, const FMutableSourceMeshData& BaseMeshData, bool bLinkedToExtendMaterial, bool bOnlyConnectedLOD);
+UE::Mutable::Private::Ptr<UE::Mutable::Private::NodeMesh> GenerateMutableSourceMesh(const UEdGraphPin* Pin, FMutableGraphGenerationContext& GenerationContext, const FMutableSourceMeshData& BaseMeshData, bool bLinkedToExtendMaterial, bool bOnlyConnectedLOD);
 
 
 /** */
 extern void SetSurfaceFormat(FMutableCompilationContext&,
-	mu::FMeshBufferSet& OutVertexBufferFormat, 
-	mu::FMeshBufferSet& OutIndexBufferFormat, 
+	UE::Mutable::Private::FMeshBufferSet& OutVertexBufferFormat, 
+	UE::Mutable::Private::FMeshBufferSet& OutIndexBufferFormat, 
 	const FMutableGraphMeshGenerationData& MeshData);
diff --git a/Engine/Plugins/Mutable/Source/CustomizableObjectEditor/Private/MuCOE/GenerateMutableSource/GenerateMutableSourceModifier.cpp b/Engine/Plugins/Mutable/Source/CustomizableObjectEditor/Private/MuCOE/GenerateMutableSource/GenerateMutableSourceModifier.cpp
--- a/Engine/Plugins/Mutable/Source/CustomizableObjectEditor/Private/MuCOE/GenerateMutableSource/GenerateMutableSourceModifier.cpp
+++ b/Engine/Plugins/Mutable/Source/CustomizableObjectEditor/Private/MuCOE/GenerateMutableSource/GenerateMutableSourceModifier.cpp
@@ -49,7 +49,7 @@
 
 #define LOCTEXT_NAMESPACE "CustomizableObjectEditor"
 
-mu::Ptr<mu::NodeModifier> GenerateMutableSourceModifier(const UEdGraphPin * Pin, FMutableGraphGenerationContext & GenerationContext)
+UE::Mutable::Private::Ptr<UE::Mutable::Private::NodeModifier> GenerateMutableSourceModifier(const UEdGraphPin * Pin, FMutableGraphGenerationContext & GenerationContext)
 {
 	check(Pin)
 	RETURN_ON_CYCLE(*Pin, GenerationContext)
@@ -63,10 +63,10 @@ mu::Ptr<mu::NodeModifier> GenerateMutableSourceModifier(const UEdGraphPin * Pin,
 	
 	if (const FGeneratedData* Generated = GenerationContext.Generated.Find(Key))
 	{
-		return static_cast<mu::NodeModifier*>(Generated->Node.get());
+		return static_cast<UE::Mutable::Private::NodeModifier*>(Generated->Node.get());
 	}
 	
-	mu::Ptr<mu::NodeModifier> Result;
+	UE::Mutable::Private::Ptr<UE::Mutable::Private::NodeModifier> Result;
 
 	// Bool that determines if a node can be added to the cache of nodes.
 	// Most nodes need to be added to the cache but there are some that don't. For exampel, MacroInstanceNodes
@@ -85,7 +85,7 @@ mu::Ptr<mu::NodeModifier> GenerateMutableSourceModifier(const UEdGraphPin * Pin,
 		// By not caching the mutable node we avoid this from even happening
 		bCacheNode = false;
 		
-		mu::Ptr<mu::NodeModifierMeshClipMorphPlane> ClipNode = new mu::NodeModifierMeshClipMorphPlane();
+		UE::Mutable::Private::Ptr<UE::Mutable::Private::NodeModifierMeshClipMorphPlane> ClipNode = new UE::Mutable::Private::NodeModifierMeshClipMorphPlane();
 		Result = ClipNode;
 
 		const FVector Origin = TypedNodeClip->GetOriginWithOffset();
@@ -114,28 +114,28 @@ mu::Ptr<mu::NodeModifier> GenerateMutableSourceModifier(const UEdGraphPin * Pin,
 			EMutableMeshConversionFlags::DoNotCreateMeshMetadata;
 		GenerationContext.MeshGenerationFlags.Push(ModifiersMeshFlags);
 
-		mu::Ptr<mu::NodeModifierMeshClipDeform> ClipNode = new mu::NodeModifierMeshClipDeform();
+		UE::Mutable::Private::Ptr<UE::Mutable::Private::NodeModifierMeshClipDeform> ClipNode = new UE::Mutable::Private::NodeModifierMeshClipDeform();
 		Result = ClipNode;
 	
 		ClipNode->FaceCullStrategy = TypedNodeClipDeform->FaceCullStrategy;
 
 		if (const UEdGraphPin* ConnectedPin = FollowInputPin(*TypedNodeClipDeform->ClipShapePin()))
 		{
-			mu::Ptr<mu::NodeMesh> ClipMesh = GenerateMutableSourceMesh(ConnectedPin, GenerationContext, FMutableSourceMeshData(), false, true);
+			UE::Mutable::Private::Ptr<UE::Mutable::Private::NodeMesh> ClipMesh = GenerateMutableSourceMesh(ConnectedPin, GenerationContext, FMutableSourceMeshData(), false, true);
 
 			ClipNode->ClipMesh = ClipMesh;
 
-			mu::EShapeBindingMethod BindingMethod = mu::EShapeBindingMethod::ClipDeformClosestProject;
+			UE::Mutable::Private::EShapeBindingMethod BindingMethod = UE::Mutable::Private::EShapeBindingMethod::ClipDeformClosestProject;
 			switch(TypedNodeClipDeform->BindingMethod)
 			{
 				case EShapeBindingMethod::ClosestProject:
-					BindingMethod = mu::EShapeBindingMethod::ClipDeformClosestProject;
+					BindingMethod = UE::Mutable::Private::EShapeBindingMethod::ClipDeformClosestProject;
 					break;
 				case EShapeBindingMethod::NormalProject:
-					BindingMethod = mu::EShapeBindingMethod::ClipDeformNormalProject;
+					BindingMethod = UE::Mutable::Private::EShapeBindingMethod::ClipDeformNormalProject;
 					break;
 				case EShapeBindingMethod::ClosestToSurface:
-					BindingMethod = mu::EShapeBindingMethod::ClipDeformClosestToSurface;
+					BindingMethod = UE::Mutable::Private::EShapeBindingMethod::ClipDeformClosestToSurface;
 					break;
 				default:
 					check(false);
@@ -170,18 +170,18 @@ mu::Ptr<mu::NodeModifier> GenerateMutableSourceModifier(const UEdGraphPin * Pin,
 		// needs to be different for each object. If it were added to the Generated cache, all the objects would get the same.
 		bCacheNode = false;
 
-		mu::Ptr<mu::NodeModifierMeshClipWithMesh> ClipNode = new mu::NodeModifierMeshClipWithMesh();
+		UE::Mutable::Private::Ptr<UE::Mutable::Private::NodeModifierMeshClipWithMesh> ClipNode = new UE::Mutable::Private::NodeModifierMeshClipWithMesh();
 		Result = ClipNode;
 
 		ClipNode->FaceCullStrategy = TypedNodeClipMesh->FaceCullStrategy;
 
 		if (const UEdGraphPin* ConnectedPin = FollowInputPin(*TypedNodeClipMesh->GetClipMeshPin()))
 		{
-			mu::NodeMeshPtr ClipMesh = GenerateMutableSourceMesh(ConnectedPin, GenerationContext, FMutableSourceMeshData(), false, true);
+			UE::Mutable::Private::NodeMeshPtr ClipMesh = GenerateMutableSourceMesh(ConnectedPin, GenerationContext, FMutableSourceMeshData(), false, true);
 
 			if (FMatrix Matrix = TypedNodeClipMesh->Transform.ToMatrixWithScale(); Matrix != FMatrix::Identity)
 			{
-				mu::Ptr<mu::NodeMeshTransform> TransformMesh = new mu::NodeMeshTransform();
+				UE::Mutable::Private::Ptr<UE::Mutable::Private::NodeMeshTransform> TransformMesh = new UE::Mutable::Private::NodeMeshTransform();
 				TransformMesh->Source = ClipMesh;
 
 				TransformMesh->Transform = FMatrix44f(Matrix);
@@ -216,14 +216,14 @@ mu::Ptr<mu::NodeModifier> GenerateMutableSourceModifier(const UEdGraphPin * Pin,
 		// needs to be different for each object. If it were added to the Generated cache, all the objects would get the same.
 		bCacheNode = false;
 
-		mu::Ptr<mu::NodeModifierMeshClipWithUVMask> ClipNode = new mu::NodeModifierMeshClipWithUVMask();
+		UE::Mutable::Private::Ptr<UE::Mutable::Private::NodeModifierMeshClipWithUVMask> ClipNode = new UE::Mutable::Private::NodeModifierMeshClipWithUVMask();
 		Result = ClipNode;
 
 		ClipNode->FaceCullStrategy = TypedNodeClipUVMask->FaceCullStrategy;
 
 		if (const UEdGraphPin* ConnectedPin = FollowInputPin(*TypedNodeClipUVMask->ClipMaskPin()))
 		{
-			mu::Ptr<mu::NodeImage> ClipMask = GenerateMutableSourceImage(ConnectedPin, GenerationContext, 0);
+			UE::Mutable::Private::Ptr<UE::Mutable::Private::NodeImage> ClipMask = GenerateMutableSourceImage(ConnectedPin, GenerationContext, 0);
 
 			ClipNode->ClipMask = ClipMask;
 		}
@@ -247,7 +247,7 @@ mu::Ptr<mu::NodeModifier> GenerateMutableSourceModifier(const UEdGraphPin * Pin,
 		const EMutableMeshConversionFlags ModifiersMeshFlags = EMutableMeshConversionFlags::None;
 		GenerationContext.MeshGenerationFlags.Push(ModifiersMeshFlags);
 
-		mu::Ptr<mu::NodeModifierSurfaceEdit> SurfNode = new mu::NodeModifierSurfaceEdit();
+		UE::Mutable::Private::Ptr<UE::Mutable::Private::NodeModifierSurfaceEdit> SurfNode = new UE::Mutable::Private::NodeModifierSurfaceEdit();
 		Result = SurfNode;
 		SurfNode->ModifierGuid = GenerationContext.GetNodeIdUnique(Node);
 
@@ -266,7 +266,7 @@ mu::Ptr<mu::NodeModifier> GenerateMutableSourceModifier(const UEdGraphPin * Pin,
 			GenerationContext.FromLOD = TypedNodeExt->FirstLOD;
 			GenerationContext.CurrentLOD = LODIndex;
 
-			mu::Ptr<mu::NodeMesh> AddMeshNode;
+			UE::Mutable::Private::Ptr<UE::Mutable::Private::NodeMesh> AddMeshNode;
 			if (const UEdGraphPin* ConnectedPin = FollowInputPin(*TypedNodeExt->AddMeshPin()))
 			{
 				// Flags to know which UV channels need layout
@@ -313,7 +313,7 @@ mu::Ptr<mu::NodeModifier> GenerateMutableSourceModifier(const UEdGraphPin * Pin,
 			SurfNode->LODs[LODIndex].Textures.SetNum(NumImages);
 			for (int32 ImageIndex = 0; ImageIndex < NumImages; ++ImageIndex)
 			{
-				mu::NodeImagePtr ImageNode;
+				UE::Mutable::Private::NodeImagePtr ImageNode;
 				FString MaterialParameterName;
 
 				if (!ImageNode) // If
@@ -385,7 +385,7 @@ mu::Ptr<mu::NodeModifier> GenerateMutableSourceModifier(const UEdGraphPin * Pin,
 			EMutableMeshConversionFlags::DoNotCreateMeshMetadata;
 		GenerationContext.MeshGenerationFlags.Push(ModifiersMeshFlags);
 
-		mu::Ptr<mu::NodeModifierSurfaceEdit> SurfNode = new mu::NodeModifierSurfaceEdit();
+		UE::Mutable::Private::Ptr<UE::Mutable::Private::NodeModifierSurfaceEdit> SurfNode = new UE::Mutable::Private::NodeModifierSurfaceEdit();
 		Result = SurfNode;
 
 		SurfNode->MultipleTagsPolicy = TypedNodeRem->MultipleTagPolicy;
@@ -404,7 +404,7 @@ mu::Ptr<mu::NodeModifier> GenerateMutableSourceModifier(const UEdGraphPin * Pin,
 				GenerationContext.FromLOD = 0;
 				GenerationContext.CurrentLOD = LODIndex;
 
-				mu::Ptr<mu::NodeMesh> RemoveMeshNode = GenerateMutableSourceMesh(ConnectedPin, GenerationContext, FMutableSourceMeshData(), false, true);
+				UE::Mutable::Private::Ptr<UE::Mutable::Private::NodeMesh> RemoveMeshNode = GenerateMutableSourceMesh(ConnectedPin, GenerationContext, FMutableSourceMeshData(), false, true);
 				SurfNode->LODs[LODIndex].MeshRemove = RemoveMeshNode;
 			}
 		}
@@ -416,7 +416,7 @@ mu::Ptr<mu::NodeModifier> GenerateMutableSourceModifier(const UEdGraphPin * Pin,
 
 	else if (const UCustomizableObjectNodeModifierRemoveMeshBlocks* TypedNodeRemBlocks = Cast<UCustomizableObjectNodeModifierRemoveMeshBlocks>(Node))
 	{
-		mu::Ptr<mu::NodeModifierMeshClipWithUVMask> ClipNode = new mu::NodeModifierMeshClipWithUVMask();
+		UE::Mutable::Private::Ptr<UE::Mutable::Private::NodeModifierMeshClipWithUVMask> ClipNode = new UE::Mutable::Private::NodeModifierMeshClipWithUVMask();
 		Result = ClipNode;
 
 		ClipNode->FaceCullStrategy = TypedNodeRemBlocks->FaceCullStrategy;
@@ -425,7 +425,7 @@ mu::Ptr<mu::NodeModifier> GenerateMutableSourceModifier(const UEdGraphPin * Pin,
 		ClipNode->RequiredTags = TypedNodeRemBlocks->GetNodeRequiredTags(&GenerationContext.MacroNodesStack);
 
 		check(TypedNodeRemBlocks->Layout->GetMesh().IsNull());
-		mu::Ptr<mu::NodeLayout> SourceLayout = CreateMutableLayoutNode(GenerationContext, TypedNodeRemBlocks->Layout, true);
+		UE::Mutable::Private::Ptr<UE::Mutable::Private::NodeLayout> SourceLayout = CreateMutableLayoutNode(GenerationContext, TypedNodeRemBlocks->Layout, true);
 		ClipNode->ClipLayout = SourceLayout;
 		ClipNode->LayoutIndex = TypedNodeRemBlocks->ParentLayoutIndex;
 	}
@@ -439,7 +439,7 @@ mu::Ptr<mu::NodeModifier> GenerateMutableSourceModifier(const UEdGraphPin * Pin,
 			EMutableMeshConversionFlags::DoNotCreateMeshMetadata;
 		GenerationContext.MeshGenerationFlags.Push(ModifiersMeshFlags);
 
-		mu::Ptr<mu::NodeModifierSurfaceEdit> SurfNode = new mu::NodeModifierSurfaceEdit();
+		UE::Mutable::Private::Ptr<UE::Mutable::Private::NodeModifierSurfaceEdit> SurfNode = new UE::Mutable::Private::NodeModifierSurfaceEdit();
 		Result = SurfNode;
 
 		SurfNode->MultipleTagsPolicy = TypedNodeEdit->MultipleTagPolicy;
@@ -468,12 +468,12 @@ mu::Ptr<mu::NodeModifier> GenerateMutableSourceModifier(const UEdGraphPin * Pin,
 
 					const UEdGraphPin* ConnectedImagePin = FollowInputPin(*TypedNodeEdit->GetUsedImagePin(ImageId));
 
-					mu::NodeModifierSurfaceEdit::FTexture& ImagePatch = SurfNode->LODs[LODIndex].Textures[ImageIndex];
+					UE::Mutable::Private::NodeModifierSurfaceEdit::FTexture& ImagePatch = SurfNode->LODs[LODIndex].Textures[ImageIndex];
 
 					ImagePatch.MaterialParameterName = TypedNodeEdit->GetParameterName(EMaterialParameterType::Texture, ImageIndex).ToString();
 
 					// \todo: expose these two options?
-					ImagePatch.PatchBlendType = mu::EBlendType::BT_BLEND;
+					ImagePatch.PatchBlendType = UE::Mutable::Private::EBlendType::BT_BLEND;
 					ImagePatch.bPatchApplyToAlpha = true;
 
 					// ReferenceTextureSize is used to limit the size of textures contributing to the final image.
@@ -518,7 +518,7 @@ mu::Ptr<mu::NodeModifier> GenerateMutableSourceModifier(const UEdGraphPin * Pin,
 			EMutableMeshConversionFlags::DoNotCreateMeshMetadata;
 		GenerationContext.MeshGenerationFlags.Push(ModifiersMeshFlags);
 
-		mu::Ptr<mu::NodeModifierSurfaceEdit> SurfNode = new mu::NodeModifierSurfaceEdit();
+		UE::Mutable::Private::Ptr<UE::Mutable::Private::NodeModifierSurfaceEdit> SurfNode = new UE::Mutable::Private::NodeModifierSurfaceEdit();
 		Result = SurfNode;
 
 		// This modifier needs to be applied right after the mesh constant is generated
@@ -593,7 +593,7 @@ mu::Ptr<mu::NodeModifier> GenerateMutableSourceModifier(const UEdGraphPin * Pin,
 			// If is a valid factor, continue the Generation
 			if (validStaticFactor)
 			{
-				mu::Ptr<mu::NodeScalar> FactorNode = GenerateMutableSourceFloat(ConnectedPin, GenerationContext);
+				UE::Mutable::Private::Ptr<UE::Mutable::Private::NodeScalar> FactorNode = GenerateMutableSourceFloat(ConnectedPin, GenerationContext);
 				SurfNode->MorphFactor = FactorNode;
 			}
 		}
@@ -614,7 +614,7 @@ mu::Ptr<mu::NodeModifier> GenerateMutableSourceModifier(const UEdGraphPin * Pin,
 		// needs to be different for each object. If it were added to the Generated cache, all the objects would get the same.
 		bCacheNode = false;
 
-		mu::Ptr<mu::NodeModifierMeshTransformInMesh> TransformNode = new mu::NodeModifierMeshTransformInMesh();
+		UE::Mutable::Private::Ptr<UE::Mutable::Private::NodeModifierMeshTransformInMesh> TransformNode = new UE::Mutable::Private::NodeModifierMeshTransformInMesh();
 		Result = TransformNode;
 
 		if (const UEdGraphPin* ConnectedPin = FollowInputPin(*TypedNodeTransformMesh->GetTransformPin()))
@@ -625,11 +625,11 @@ mu::Ptr<mu::NodeModifier> GenerateMutableSourceModifier(const UEdGraphPin * Pin,
 		// If no bounding mesh is provided, we transform the entire mesh.
 		if (const UEdGraphPin* ConnectedPin = FollowInputPin(*TypedNodeTransformMesh->GetBoundingMeshPin()))
 		{
-			mu::Ptr<mu::NodeMesh> BoundingMesh = GenerateMutableSourceMesh(ConnectedPin, GenerationContext, FMutableSourceMeshData(), false, true);
+			UE::Mutable::Private::Ptr<UE::Mutable::Private::NodeMesh> BoundingMesh = GenerateMutableSourceMesh(ConnectedPin, GenerationContext, FMutableSourceMeshData(), false, true);
 
 			if (FMatrix Matrix = TypedNodeTransformMesh->BoundingMeshTransform.ToMatrixWithScale(); Matrix != FMatrix::Identity)
 			{
-				mu::Ptr<mu::NodeMeshTransform> TransformMesh = new mu::NodeMeshTransform();
+				UE::Mutable::Private::Ptr<UE::Mutable::Private::NodeMeshTransform> TransformMesh = new UE::Mutable::Private::NodeMeshTransform();
 				TransformMesh->Source = BoundingMesh;
 
 				TransformMesh->Transform = FMatrix44f(Matrix);
@@ -649,7 +649,7 @@ mu::Ptr<mu::NodeModifier> GenerateMutableSourceModifier(const UEdGraphPin * Pin,
 	{
 		bCacheNode = false;
 
-		mu::Ptr<mu::NodeModifierMeshTransformWithBone> TransformNode = new mu::NodeModifierMeshTransformWithBone();
+		UE::Mutable::Private::Ptr<UE::Mutable::Private::NodeModifierMeshTransformWithBone> TransformNode = new UE::Mutable::Private::NodeModifierMeshTransformWithBone();
 		Result = TransformNode;
 
 		TransformNode->BoneName = GenerationContext.CompilationContext->GetBoneUnique(*TypedNodeTransformWithBone->BoneName);
@@ -667,13 +667,13 @@ mu::Ptr<mu::NodeModifier> GenerateMutableSourceModifier(const UEdGraphPin * Pin,
 	else if (const UCustomizableObjectNodeMacroInstance* TypedNodeMacro = Cast<UCustomizableObjectNodeMacroInstance>(Node))
 	{
 		bCacheNode = false;
-		Result = GenerateMutableSourceMacro<mu::NodeModifier>(*Pin, GenerationContext, GenerateMutableSourceModifier);
+		Result = GenerateMutableSourceMacro<UE::Mutable::Private::NodeModifier>(*Pin, GenerationContext, GenerateMutableSourceModifier);
 	}
 
 	else if (const UCustomizableObjectNodeTunnel* TypedNodeTunnel = Cast<UCustomizableObjectNodeTunnel>(Node))
 	{
 		bCacheNode = false;
-		Result = GenerateMutableSourceMacro<mu::NodeModifier>(*Pin, GenerationContext, GenerateMutableSourceModifier);
+		Result = GenerateMutableSourceMacro<UE::Mutable::Private::NodeModifier>(*Pin, GenerationContext, GenerateMutableSourceModifier);
 	}
 
 	else
diff --git a/Engine/Plugins/Mutable/Source/CustomizableObjectEditor/Private/MuCOE/GenerateMutableSource/GenerateMutableSourceModifier.h b/Engine/Plugins/Mutable/Source/CustomizableObjectEditor/Private/MuCOE/GenerateMutableSource/GenerateMutableSourceModifier.h
--- a/Engine/Plugins/Mutable/Source/CustomizableObjectEditor/Private/MuCOE/GenerateMutableSource/GenerateMutableSourceModifier.h
+++ b/Engine/Plugins/Mutable/Source/CustomizableObjectEditor/Private/MuCOE/GenerateMutableSource/GenerateMutableSourceModifier.h
@@ -7,4 +7,4 @@
 class UEdGraphPin;
 struct FMutableGraphGenerationContext;
 
-mu::Ptr<mu::NodeModifier> GenerateMutableSourceModifier(const UEdGraphPin* Pin, FMutableGraphGenerationContext& GenerationContext);
+UE::Mutable::Private::Ptr<UE::Mutable::Private::NodeModifier> GenerateMutableSourceModifier(const UEdGraphPin* Pin, FMutableGraphGenerationContext& GenerationContext);
