[Mutable] Change namespace name

diff --git a/Engine/Plugins/Mutable/Source/CustomizableObjectEditor/Private/MuCOE/GenerateMutableSource/GenerateMutableSourceProjector.cpp b/Engine/Plugins/Mutable/Source/CustomizableObjectEditor/Private/MuCOE/GenerateMutableSource/GenerateMutableSourceProjector.cpp
--- a/Engine/Plugins/Mutable/Source/CustomizableObjectEditor/Private/MuCOE/GenerateMutableSource/GenerateMutableSourceProjector.cpp
+++ b/Engine/Plugins/Mutable/Source/CustomizableObjectEditor/Private/MuCOE/GenerateMutableSource/GenerateMutableSourceProjector.cpp
@@ -11,7 +11,7 @@
 #define LOCTEXT_NAMESPACE "CustomizableObjectEditor"
 
 
-mu::Ptr<mu::NodeProjector> GenerateMutableSourceProjector(const UEdGraphPin* Pin, FMutableGraphGenerationContext& GenerationContext)
+UE::Mutable::Private::Ptr<UE::Mutable::Private::NodeProjector> GenerateMutableSourceProjector(const UEdGraphPin* Pin, FMutableGraphGenerationContext& GenerationContext)
 {
 	check(Pin)
 	RETURN_ON_CYCLE(*Pin, GenerationContext)
@@ -23,20 +23,20 @@ mu::Ptr<mu::NodeProjector> GenerateMutableSourceProjector(const UEdGraphPin* Pin
 	const FGeneratedKey Key(reinterpret_cast<void*>(&GenerateMutableSourceProjector), *Pin, *Node, GenerationContext, true);
 	if (const FGeneratedData* Generated = GenerationContext.Generated.Find(Key))
 	{
-		return static_cast<mu::NodeProjector*>(Generated->Node.get());
+		return static_cast<UE::Mutable::Private::NodeProjector*>(Generated->Node.get());
 	}
 
-	mu::Ptr<mu::NodeProjector> Result;
+	UE::Mutable::Private::Ptr<UE::Mutable::Private::NodeProjector> Result;
 	
 	if (const UCustomizableObjectNodeProjectorConstant* TypedNodeConst = Cast<UCustomizableObjectNodeProjectorConstant>(Node))
 	{
-		mu::Ptr<mu::NodeProjectorConstant> ProjectorNode = new mu::NodeProjectorConstant();
+		UE::Mutable::Private::Ptr<UE::Mutable::Private::NodeProjectorConstant> ProjectorNode = new UE::Mutable::Private::NodeProjectorConstant();
 		Result = ProjectorNode;
 
-		const mu::EProjectorType ProjectorType = ProjectorUtils::GetEquivalentProjectorType(TypedNodeConst->Value.ProjectionType);
+		const UE::Mutable::Private::EProjectorType ProjectorType = ProjectorUtils::GetEquivalentProjectorType(TypedNodeConst->Value.ProjectionType);
 		switch (ProjectorType) 
 		{ 
-			case mu::EProjectorType::Planar:
+			case UE::Mutable::Private::EProjectorType::Planar:
 				ProjectorNode->SetValue(
 					ProjectorType,
 					TypedNodeConst->Value.Position,
@@ -46,7 +46,7 @@ mu::Ptr<mu::NodeProjector> GenerateMutableSourceProjector(const UEdGraphPin* Pin
 					TypedNodeConst->Value.Angle);
 				break;
 			
-			case mu::EProjectorType::Cylindrical:
+			case UE::Mutable::Private::EProjectorType::Cylindrical:
 				{
 					// Apply strange swizzle for scales
 					// TODO: try to avoid this
@@ -62,7 +62,7 @@ mu::Ptr<mu::NodeProjector> GenerateMutableSourceProjector(const UEdGraphPin* Pin
 				}
 				break;
 			
-			case mu::EProjectorType::Wrapping:
+			case UE::Mutable::Private::EProjectorType::Wrapping:
 				ProjectorNode->SetValue(
 					ProjectorType,
 					TypedNodeConst->Value.Position,
@@ -72,7 +72,7 @@ mu::Ptr<mu::NodeProjector> GenerateMutableSourceProjector(const UEdGraphPin* Pin
 					TypedNodeConst->Value.Angle);
 				break;
 			
-			case mu::EProjectorType::Count: 
+			case UE::Mutable::Private::EProjectorType::Count: 
 			default:
 				checkNoEntry();
 		}
@@ -80,7 +80,7 @@ mu::Ptr<mu::NodeProjector> GenerateMutableSourceProjector(const UEdGraphPin* Pin
 
 	else if (const UCustomizableObjectNodeProjectorParameter* TypedNodeParam = Cast<UCustomizableObjectNodeProjectorParameter>(Node))
 	{
-		mu::Ptr<mu::NodeProjectorParameter> ProjectorNode = new mu::NodeProjectorParameter();
+		UE::Mutable::Private::Ptr<UE::Mutable::Private::NodeProjectorParameter> ProjectorNode = new UE::Mutable::Private::NodeProjectorParameter();
 		Result = ProjectorNode;
 
 		ProjectorNode->SetName(TypedNodeParam->GetParameterName(&GenerationContext.MacroNodesStack));
@@ -89,7 +89,7 @@ mu::Ptr<mu::NodeProjector> GenerateMutableSourceProjector(const UEdGraphPin* Pin
 		{
 		case 0:
 			ProjectorNode->SetDefaultValue(
-				mu::EProjectorType::Planar,
+				UE::Mutable::Private::EProjectorType::Planar,
 				TypedNodeParam->DefaultValue.Position,
 				TypedNodeParam->DefaultValue.Direction,
 				TypedNodeParam->DefaultValue.Up,
@@ -104,7 +104,7 @@ mu::Ptr<mu::NodeProjector> GenerateMutableSourceProjector(const UEdGraphPin* Pin
 			float Radius = FMath::Abs(TypedNodeParam->DefaultValue.Scale[0] / 2.0f);
 			float Height = TypedNodeParam->DefaultValue.Scale[2];
 			ProjectorNode->SetDefaultValue(
-				mu::EProjectorType::Cylindrical,
+				UE::Mutable::Private::EProjectorType::Cylindrical,
 				TypedNodeParam->DefaultValue.Position,
 				-TypedNodeParam->DefaultValue.Direction,
 				-TypedNodeParam->DefaultValue.Up,
@@ -115,7 +115,7 @@ mu::Ptr<mu::NodeProjector> GenerateMutableSourceProjector(const UEdGraphPin* Pin
 
 		case 2:
 			ProjectorNode->SetDefaultValue(
-				mu::EProjectorType::Wrapping,
+				UE::Mutable::Private::EProjectorType::Wrapping,
 				TypedNodeParam->DefaultValue.Position,
 				TypedNodeParam->DefaultValue.Direction,
 				TypedNodeParam->DefaultValue.Up,
diff --git a/Engine/Plugins/Mutable/Source/CustomizableObjectEditor/Private/MuCOE/GenerateMutableSource/GenerateMutableSourceProjector.h b/Engine/Plugins/Mutable/Source/CustomizableObjectEditor/Private/MuCOE/GenerateMutableSource/GenerateMutableSourceProjector.h
--- a/Engine/Plugins/Mutable/Source/CustomizableObjectEditor/Private/MuCOE/GenerateMutableSource/GenerateMutableSourceProjector.h
+++ b/Engine/Plugins/Mutable/Source/CustomizableObjectEditor/Private/MuCOE/GenerateMutableSource/GenerateMutableSourceProjector.h
@@ -8,4 +8,4 @@ class UEdGraphPin;
 struct FMutableGraphGenerationContext;
 
 /** Convert a CustomizableObject Source Graph into a mutable source graph. */
-mu::Ptr<mu::NodeProjector> GenerateMutableSourceProjector(const UEdGraphPin* Pin, FMutableGraphGenerationContext& GenerationContext);
+UE::Mutable::Private::Ptr<UE::Mutable::Private::NodeProjector> GenerateMutableSourceProjector(const UEdGraphPin* Pin, FMutableGraphGenerationContext& GenerationContext);
diff --git a/Engine/Plugins/Mutable/Source/CustomizableObjectEditor/Private/MuCOE/GenerateMutableSource/GenerateMutableSourceSurface.cpp b/Engine/Plugins/Mutable/Source/CustomizableObjectEditor/Private/MuCOE/GenerateMutableSource/GenerateMutableSourceSurface.cpp
--- a/Engine/Plugins/Mutable/Source/CustomizableObjectEditor/Private/MuCOE/GenerateMutableSource/GenerateMutableSourceSurface.cpp
+++ b/Engine/Plugins/Mutable/Source/CustomizableObjectEditor/Private/MuCOE/GenerateMutableSource/GenerateMutableSourceSurface.cpp
@@ -63,7 +63,7 @@
 #define LOCTEXT_NAMESPACE "CustomizableObjectEditor"
 
 
-mu::Ptr<mu::NodeSurface> GenerateMutableSourceSurface(const UEdGraphPin * Pin, FMutableGraphGenerationContext & GenerationContext)
+UE::Mutable::Private::Ptr<UE::Mutable::Private::NodeSurface> GenerateMutableSourceSurface(const UEdGraphPin * Pin, FMutableGraphGenerationContext & GenerationContext)
 {
 	MUTABLE_CPUPROFILER_SCOPE(GenerateMutableSourceSurface);
 
@@ -81,10 +81,10 @@ mu::Ptr<mu::NodeSurface> GenerateMutableSourceSurface(const UEdGraphPin * Pin, F
 	const FGeneratedKey Key(reinterpret_cast<void*>(&GenerateMutableSourceSurface), *Pin, *Node, GenerationContext, true);
 	if (const FGeneratedData* Generated = GenerationContext.Generated.Find(Key))
 	{
-		return static_cast<mu::NodeSurface*>(Generated->Node.get());
+		return static_cast<UE::Mutable::Private::NodeSurface*>(Generated->Node.get());
 	}
 	
-	mu::Ptr<mu::NodeSurface> Result;
+	UE::Mutable::Private::Ptr<UE::Mutable::Private::NodeSurface> Result;
 
 	const int32 LOD = Node->IsAffectedByLOD() ? GenerationContext.CurrentLOD : 0;
 	
@@ -128,7 +128,7 @@ mu::Ptr<mu::NodeSurface> GenerateMutableSourceSurface(const UEdGraphPin * Pin, F
 			return Result;
 		}
 
-		mu::Ptr<mu::NodeSurfaceNew> SurfNode = new mu::NodeSurfaceNew();
+		UE::Mutable::Private::Ptr<UE::Mutable::Private::NodeSurfaceNew> SurfNode = new UE::Mutable::Private::NodeSurfaceNew();
 		Result = SurfNode;
 
 		FGuid& SurfaceGuid = GenerationContext.SurfaceGuids.FindOrAdd(FSharedSurfaceKey(TypedNodeMat, GenerationContext.MacroNodesStack));
@@ -167,7 +167,7 @@ mu::Ptr<mu::NodeSurface> GenerateMutableSourceSurface(const UEdGraphPin * Pin, F
 			}
 		}
 
-		mu::Ptr<mu::NodeMesh> MeshNode;
+		UE::Mutable::Private::Ptr<UE::Mutable::Private::NodeMesh> MeshNode;
 		
 		if (bGeneratingImplicitComponent)
 		{
@@ -239,7 +239,7 @@ mu::Ptr<mu::NodeSurface> GenerateMutableSourceSurface(const UEdGraphPin * Pin, F
 
 		if (!SurfNode->Material)
 		{
-			mu::Ptr<mu::NodeMaterialConstant> MaterialNode = new mu::NodeMaterialConstant();
+			UE::Mutable::Private::Ptr<UE::Mutable::Private::NodeMaterialConstant> MaterialNode = new UE::Mutable::Private::NodeMaterialConstant();
 			MaterialNode->MaterialId = GenerationContext.CurrentReferencedMaterialIndex;
 			SurfNode->Material = MaterialNode;
 		}
@@ -271,7 +271,7 @@ mu::Ptr<mu::NodeSurface> GenerateMutableSourceSurface(const UEdGraphPin * Pin, F
 					}
 
 					// This is a connected pass-through texture that simply has to be passed to the core
-					mu::Ptr<mu::NodeImage> PassThroughImagePtr = GenerateMutableSourceImage(ConnectedPin, GenerationContext, 0);
+					UE::Mutable::Private::Ptr<UE::Mutable::Private::NodeImage> PassThroughImagePtr = GenerateMutableSourceImage(ConnectedPin, GenerationContext, 0);
 					SurfNode->Images[ImageIndex].Image = PassThroughImagePtr;
 
 					check(Props.ImagePropertiesIndex != INDEX_NONE);
@@ -285,7 +285,7 @@ mu::Ptr<mu::NodeSurface> GenerateMutableSourceSurface(const UEdGraphPin * Pin, F
 			}
 			else
 			{
-				mu::NodeImagePtr GroupProjectionImg;
+				UE::Mutable::Private::NodeImagePtr GroupProjectionImg;
 				UTexture2D* GroupProjectionReferenceTexture = nullptr;
 				const FString ImageName = TypedNodeMat->GetParameterName(EMaterialParameterType::Texture, ImageIndex).ToString();
 				const FNodeMaterialParameterId ImageId = TypedNodeMat->GetParameterId(EMaterialParameterType::Texture, ImageIndex);
@@ -392,7 +392,7 @@ mu::Ptr<mu::NodeSurface> GenerateMutableSourceSurface(const UEdGraphPin * Pin, F
 						}
 
 						// Generate the texture nodes
-						mu::NodeImagePtr ImageNode = [&]()
+						UE::Mutable::Private::NodeImagePtr ImageNode = [&]()
 						{
 							if (TypedNodeMat->IsImageMutableMode(ImageIndex))
 							{
@@ -417,12 +417,12 @@ mu::Ptr<mu::NodeSurface> GenerateMutableSourceSurface(const UEdGraphPin * Pin, F
 
 									if (Texture2D)
 									{
-										mu::Ptr<mu::NodeImageConstant> ConstImageNode = new mu::NodeImageConstant();
-										TSharedPtr<mu::FImage> ImageConstant = GenerateImageConstant(Texture2D, GenerationContext, false);
+										UE::Mutable::Private::Ptr<UE::Mutable::Private::NodeImageConstant> ConstImageNode = new UE::Mutable::Private::NodeImageConstant();
+										TSharedPtr<UE::Mutable::Private::FImage> ImageConstant = GenerateImageConstant(Texture2D, GenerationContext, false);
 										ConstImageNode->SetValue(ImageConstant);
 
 										const uint32 MipsToSkip = ComputeLODBiasForTexture(GenerationContext, *Texture2D, nullptr, Props.TextureSize);
-										mu::Ptr<mu::NodeImage> Result =  ResizeTextureByNumMips(ConstImageNode, MipsToSkip);
+										UE::Mutable::Private::Ptr<UE::Mutable::Private::NodeImage> Result =  ResizeTextureByNumMips(ConstImageNode, MipsToSkip);
 	
 										const UTextureLODSettings& TextureLODSettings = GenerationContext.CompilationContext->Options.TargetPlatform->GetTextureLODSettings();
 										const FTextureLODGroup& LODGroupInfo = TextureLODSettings.GetTextureLODGroup(Texture2D->LODGroup);
@@ -438,13 +438,13 @@ mu::Ptr<mu::NodeSurface> GenerateMutableSourceSurface(const UEdGraphPin * Pin, F
 									}
 									else
 									{
-										return mu::NodeImagePtr();
+										return UE::Mutable::Private::NodeImagePtr();
 									}
 								}
 							}
 							else
 							{
-								return mu::NodeImagePtr();
+								return UE::Mutable::Private::NodeImagePtr();
 							}
 						}();
 
@@ -457,30 +457,30 @@ mu::Ptr<mu::NodeSurface> GenerateMutableSourceSurface(const UEdGraphPin * Pin, F
 						{
 							// Apply base LODBias. It will be propagated to most images.
 							const uint32 BaseLODBias = ComputeLODBiasForTexture(GenerationContext, *ReferenceTexture);
-							mu::NodeImagePtr LastImage = ResizeTextureByNumMips(ImageNode, BaseLODBias);
+							UE::Mutable::Private::NodeImagePtr LastImage = ResizeTextureByNumMips(ImageNode, BaseLODBias);
 
 							if (ReferenceTexture->MipGenSettings != TextureMipGenSettings::TMGS_NoMipmaps)
 							{
-								mu::EMipmapFilterType MipGenerationFilterType = Invoke([&]()
+								UE::Mutable::Private::EMipmapFilterType MipGenerationFilterType = Invoke([&]()
 								{
 									if (ReferenceTexture)
 									{
 										switch (ReferenceTexture->MipGenSettings)
 										{
-											case TextureMipGenSettings::TMGS_SimpleAverage: return mu::EMipmapFilterType::SimpleAverage;
-											case TextureMipGenSettings::TMGS_Unfiltered:    return mu::EMipmapFilterType::Unfiltered;
-											default: return mu::EMipmapFilterType::SimpleAverage;
+											case TextureMipGenSettings::TMGS_SimpleAverage: return UE::Mutable::Private::EMipmapFilterType::SimpleAverage;
+											case TextureMipGenSettings::TMGS_Unfiltered:    return UE::Mutable::Private::EMipmapFilterType::Unfiltered;
+											default: return UE::Mutable::Private::EMipmapFilterType::SimpleAverage;
 										}
 									}
 
-									return mu::EMipmapFilterType::SimpleAverage;
+									return UE::Mutable::Private::EMipmapFilterType::SimpleAverage;
 								});
 
 
-								mu::Ptr<mu::NodeImageMipmap> MipmapImage = new mu::NodeImageMipmap();
+								UE::Mutable::Private::Ptr<UE::Mutable::Private::NodeImageMipmap> MipmapImage = new UE::Mutable::Private::NodeImageMipmap();
 								MipmapImage->Source = LastImage;
 								MipmapImage->Settings.FilterType = MipGenerationFilterType;
-								MipmapImage->Settings.AddressMode = mu::EAddressMode::None;
+								MipmapImage->Settings.AddressMode = UE::Mutable::Private::EAddressMode::None;
 
 								MipmapImage->SetMessageContext(Node);
 								LastImage = MipmapImage;
@@ -489,26 +489,26 @@ mu::Ptr<mu::NodeSurface> GenerateMutableSourceSurface(const UEdGraphPin * Pin, F
 							// Apply composite image. This needs to be computed after mipmaps generation. 	
 							if (ReferenceTexture && ReferenceTexture->GetCompositeTexture() && ReferenceTexture->CompositeTextureMode != CTM_Disabled)
 							{
-								mu::Ptr<mu::NodeImageNormalComposite> CompositedImage = new mu::NodeImageNormalComposite();
+								UE::Mutable::Private::Ptr<UE::Mutable::Private::NodeImageNormalComposite> CompositedImage = new UE::Mutable::Private::NodeImageNormalComposite();
 								CompositedImage->Base = LastImage;
 								CompositedImage->Power = ReferenceTexture->CompositePower;
 
-								mu::ECompositeImageMode CompositeImageMode = [CompositeTextureMode = ReferenceTexture->CompositeTextureMode]() -> mu::ECompositeImageMode
+								UE::Mutable::Private::ECompositeImageMode CompositeImageMode = [CompositeTextureMode = ReferenceTexture->CompositeTextureMode]() -> UE::Mutable::Private::ECompositeImageMode
 								{
 									switch (CompositeTextureMode)
 									{
-										case CTM_NormalRoughnessToRed: return mu::ECompositeImageMode::CIM_NormalRoughnessToRed;
-										case CTM_NormalRoughnessToGreen: return mu::ECompositeImageMode::CIM_NormalRoughnessToGreen;
-										case CTM_NormalRoughnessToBlue: return mu::ECompositeImageMode::CIM_NormalRoughnessToBlue;
-										case CTM_NormalRoughnessToAlpha: return mu::ECompositeImageMode::CIM_NormalRoughnessToAlpha;
+										case CTM_NormalRoughnessToRed: return UE::Mutable::Private::ECompositeImageMode::CIM_NormalRoughnessToRed;
+										case CTM_NormalRoughnessToGreen: return UE::Mutable::Private::ECompositeImageMode::CIM_NormalRoughnessToGreen;
+										case CTM_NormalRoughnessToBlue: return UE::Mutable::Private::ECompositeImageMode::CIM_NormalRoughnessToBlue;
+										case CTM_NormalRoughnessToAlpha: return UE::Mutable::Private::ECompositeImageMode::CIM_NormalRoughnessToAlpha;
 
-										default: return mu::ECompositeImageMode::CIM_Disabled;
+										default: return UE::Mutable::Private::ECompositeImageMode::CIM_Disabled;
 									}
 								}();
 
 								CompositedImage->Mode = CompositeImageMode;
 
-								mu::Ptr<mu::NodeImageConstant> CompositeNormalImage = new mu::NodeImageConstant();
+								UE::Mutable::Private::Ptr<UE::Mutable::Private::NodeImageConstant> CompositeNormalImage = new UE::Mutable::Private::NodeImageConstant();
 
 								UTexture2D* ReferenceCompositeNormalTexture = Cast<UTexture2D>(ReferenceTexture->GetCompositeTexture());
 								if (ReferenceCompositeNormalTexture)
@@ -516,14 +516,14 @@ mu::Ptr<mu::NodeSurface> GenerateMutableSourceSurface(const UEdGraphPin * Pin, F
 									// GenerationContext.ArrayTextureUnrealToMutableTask.Add(FTextureUnrealToMutableTask(CompositeNormalImage, ReferenceCompositeNormalTexture, Node, true));
 									// TODO: The normal composite part is not propagated, so it will be unsupported. Create a task that performs the required transforms at mutable image level, and add the right operations here
 									// instead of propagating the flag and doing them on unreal-convert.
-									TSharedPtr<mu::FImage> ImageConstant = GenerateImageConstant(ReferenceCompositeNormalTexture, GenerationContext, false);
+									TSharedPtr<UE::Mutable::Private::FImage> ImageConstant = GenerateImageConstant(ReferenceCompositeNormalTexture, GenerationContext, false);
 									CompositeNormalImage->SetValue(ImageConstant);
 
-									mu::Ptr<mu::NodeImageMipmap> NormalCompositeMipmapImage = new mu::NodeImageMipmap();
+									UE::Mutable::Private::Ptr<UE::Mutable::Private::NodeImageMipmap> NormalCompositeMipmapImage = new UE::Mutable::Private::NodeImageMipmap();
 									const uint32 MipsToSkip = ComputeLODBiasForTexture(GenerationContext, *ReferenceCompositeNormalTexture, ReferenceTexture);
 									NormalCompositeMipmapImage->Source = ResizeTextureByNumMips(CompositeNormalImage, MipsToSkip);
-									NormalCompositeMipmapImage->Settings.FilterType = mu::EMipmapFilterType::SimpleAverage;
-									NormalCompositeMipmapImage->Settings.AddressMode = mu::EAddressMode::None;
+									NormalCompositeMipmapImage->Settings.FilterType = UE::Mutable::Private::EMipmapFilterType::SimpleAverage;
+									NormalCompositeMipmapImage->Settings.AddressMode = UE::Mutable::Private::EAddressMode::None;
 
 									CompositedImage->Normal = NormalCompositeMipmapImage;
 
@@ -545,10 +545,10 @@ mu::Ptr<mu::NodeSurface> GenerateMutableSourceSurface(const UEdGraphPin * Pin, F
 								LastImage = CompositedImage;
 							}
 
-							mu::Ptr<mu::NodeImage> FormatSource = LastImage;
-							mu::Ptr<mu::NodeImageFormat> FormatImage = new mu::NodeImageFormat();
+							UE::Mutable::Private::Ptr<UE::Mutable::Private::NodeImage> FormatSource = LastImage;
+							UE::Mutable::Private::Ptr<UE::Mutable::Private::NodeImageFormat> FormatImage = new UE::Mutable::Private::NodeImageFormat();
 							FormatImage->Source = LastImage.get();
-							FormatImage->Format = mu::EImageFormat::RGBA_UByte;
+							FormatImage->Format = UE::Mutable::Private::EImageFormat::RGBA_UByte;
 							FormatImage->SetMessageContext(Node);
 							LastImage = FormatImage;
 
@@ -599,8 +599,8 @@ mu::Ptr<mu::NodeSurface> GenerateMutableSourceSurface(const UEdGraphPin * Pin, F
 									EPixelFormat UnrealTargetPlatformFormatAlpha = TextureFormat->GetEncodedPixelFormat(BuildSettingsPerLayer[0], true);
 
 									// \TODO: The QualityFix filter is used while the internal mutable runtime compression doesn't provide enough quality for some large block formats.
-									mu::EImageFormat MutableFormat = QualityAndPerformanceFix(UnrealToMutablePixelFormat(UnrealTargetPlatformFormat,false));
-									mu::EImageFormat MutableFormatIfAlpha = QualityAndPerformanceFix(UnrealToMutablePixelFormat(UnrealTargetPlatformFormatAlpha,true));
+									UE::Mutable::Private::EImageFormat MutableFormat = QualityAndPerformanceFix(UnrealToMutablePixelFormat(UnrealTargetPlatformFormat,false));
+									UE::Mutable::Private::EImageFormat MutableFormatIfAlpha = QualityAndPerformanceFix(UnrealToMutablePixelFormat(UnrealTargetPlatformFormatAlpha,true));
 
 									// Temp hack to enable RG->LA 
 									if (GenerationContext.CompilationContext->Options.TargetPlatform)
@@ -619,8 +619,8 @@ mu::Ptr<mu::NodeSurface> GenerateMutableSourceSurface(const UEdGraphPin * Pin, F
 											if (bIsNormalMapFormat)
 											{
 												// Insert a channel swizzle
-												mu::Ptr<mu::NodeImageSwizzle> Swizzle = new mu::NodeImageSwizzle;
-												Swizzle->SetFormat( mu::EImageFormat::RGBA_UByte );
+												UE::Mutable::Private::Ptr<UE::Mutable::Private::NodeImageSwizzle> Swizzle = new UE::Mutable::Private::NodeImageSwizzle;
+												Swizzle->SetFormat( UE::Mutable::Private::EImageFormat::RGBA_UByte );
 												Swizzle->Sources[0] = FormatSource;
 												Swizzle->Sources[1] = FormatSource;
 												Swizzle->Sources[2] = FormatSource;
@@ -636,21 +636,21 @@ mu::Ptr<mu::NodeSurface> GenerateMutableSourceSurface(const UEdGraphPin * Pin, F
 									}
 
 									// Unsupported format: look for something generic
-									if (MutableFormat == mu::EImageFormat::None)
+									if (MutableFormat == UE::Mutable::Private::EImageFormat::None)
 									{
 										const FString ReplacedImageFormatMsg = FString::Printf(TEXT("In object [%s] the unsupported image format %d is used, RGBA_UByte will be used instead."), *GenerationContext.GetObjectName(), UnrealTargetPlatformFormat );
 										const FText ReplacedImageFormatText = FText::FromString(ReplacedImageFormatMsg);
 										GenerationContext.Log(ReplacedImageFormatText, Node, EMessageSeverity::Info);
 										UE_LOG(LogMutable, Log, TEXT("%s"), *ReplacedImageFormatMsg);
-										MutableFormat = mu::EImageFormat::RGBA_UByte;
+										MutableFormat = UE::Mutable::Private::EImageFormat::RGBA_UByte;
 									}
-									if (MutableFormatIfAlpha == mu::EImageFormat::None)
+									if (MutableFormatIfAlpha == UE::Mutable::Private::EImageFormat::None)
 									{
 										const FString ReplacedImageFormatMsg = FString::Printf(TEXT("In object [%s] the unsupported image format %d is used, RGBA_UByte will be used instead."), *GenerationContext.GetObjectName(), UnrealTargetPlatformFormatAlpha);
 										const FText ReplacedImageFormatText = FText::FromString(ReplacedImageFormatMsg);
 										GenerationContext.Log(ReplacedImageFormatText, Node, EMessageSeverity::Info);
 										UE_LOG(LogMutable, Log, TEXT("%s"), *ReplacedImageFormatMsg);
-										MutableFormatIfAlpha = mu::EImageFormat::RGBA_UByte;
+										MutableFormatIfAlpha = UE::Mutable::Private::EImageFormat::RGBA_UByte;
 									}
 
 									FormatImage->Format = MutableFormat;
@@ -719,7 +719,7 @@ mu::Ptr<mu::NodeSurface> GenerateMutableSourceSurface(const UEdGraphPin * Pin, F
 			{				
 				if (const UEdGraphPin* ConnectedPin = FollowInputPin(*VectorPin))
 				{
-					mu::Ptr<mu::NodeColour> ColorNode = GenerateMutableSourceColor(ConnectedPin, GenerationContext);
+					UE::Mutable::Private::Ptr<UE::Mutable::Private::NodeColour> ColorNode = GenerateMutableSourceColor(ConnectedPin, GenerationContext);
 
 					// Encoding material layer in mutable name
 					if (const int32 LayerIndex = TypedNodeMat->GetParameterLayerIndex(EMaterialParameterType::Vector, VectorIndex); LayerIndex != INDEX_NONE)
@@ -750,7 +750,7 @@ mu::Ptr<mu::NodeSurface> GenerateMutableSourceSurface(const UEdGraphPin * Pin, F
 			{
 				if (const UEdGraphPin* ConnectedPin = FollowInputPin(*ScalarPin))
 				{
-					mu::Ptr<mu::NodeScalar> ScalarNode = GenerateMutableSourceFloat(ConnectedPin, GenerationContext);
+					UE::Mutable::Private::Ptr<UE::Mutable::Private::NodeScalar> ScalarNode = GenerateMutableSourceFloat(ConnectedPin, GenerationContext);
 
 					// Encoding material layer in mutable name
 					if (const int32 LayerIndex = TypedNodeMat->GetParameterLayerIndex(EMaterialParameterType::Scalar, ScalarIndex); LayerIndex != INDEX_NONE)
@@ -776,14 +776,14 @@ mu::Ptr<mu::NodeSurface> GenerateMutableSourceSurface(const UEdGraphPin * Pin, F
 
 	else if (const UCustomizableObjectNodeMaterialVariation* TypedNodeVar = Cast<UCustomizableObjectNodeMaterialVariation>(Node))
 	{
-		mu::Ptr<mu::NodeSurfaceVariation> SurfNode = new mu::NodeSurfaceVariation();
+		UE::Mutable::Private::Ptr<UE::Mutable::Private::NodeSurfaceVariation> SurfNode = new UE::Mutable::Private::NodeSurfaceVariation();
 		Result = SurfNode;
 
-		mu::NodeSurfaceVariation::VariationType muType = mu::NodeSurfaceVariation::VariationType::Tag;
+		UE::Mutable::Private::NodeSurfaceVariation::VariationType muType = UE::Mutable::Private::NodeSurfaceVariation::VariationType::Tag;
 		switch (TypedNodeVar->Type)
 		{
-		case ECustomizableObjectNodeMaterialVariationType::Tag: muType = mu::NodeSurfaceVariation::VariationType::Tag; break;
-		case ECustomizableObjectNodeMaterialVariationType::State: muType = mu::NodeSurfaceVariation::VariationType::State; break;
+		case ECustomizableObjectNodeMaterialVariationType::Tag: muType = UE::Mutable::Private::NodeSurfaceVariation::VariationType::Tag; break;
+		case ECustomizableObjectNodeMaterialVariationType::State: muType = UE::Mutable::Private::NodeSurfaceVariation::VariationType::State; break;
 		default:
 			check(false);
 			break;
@@ -793,7 +793,7 @@ mu::Ptr<mu::NodeSurface> GenerateMutableSourceSurface(const UEdGraphPin * Pin, F
 		for (const UEdGraphPin* ConnectedPin : FollowInputPinArray(*TypedNodeVar->DefaultPin()))
 		{
 			// Is it a modifier?
-			mu::NodeSurfacePtr ChildNode = GenerateMutableSourceSurface(ConnectedPin, GenerationContext);
+			UE::Mutable::Private::NodeSurfacePtr ChildNode = GenerateMutableSourceSurface(ConnectedPin, GenerationContext);
 			if (ChildNode)
 			{
 				SurfNode->DefaultSurfaces.Add(ChildNode);
@@ -808,7 +808,7 @@ mu::Ptr<mu::NodeSurface> GenerateMutableSourceSurface(const UEdGraphPin * Pin, F
 		SurfNode->Variations.SetNum(NumVariations);
 		for (int VariationIndex = 0; VariationIndex < NumVariations; ++VariationIndex)
 		{
-			mu::NodeSurfacePtr VariationSurfaceNode;
+			UE::Mutable::Private::NodeSurfacePtr VariationSurfaceNode;
 
 			if (UEdGraphPin* VariationPin = TypedNodeVar->VariationPin(VariationIndex))
 			{
@@ -816,7 +816,7 @@ mu::Ptr<mu::NodeSurface> GenerateMutableSourceSurface(const UEdGraphPin * Pin, F
 				for (const UEdGraphPin* ConnectedPin : FollowInputPinArray(*VariationPin))
 				{
 					// Is it a modifier?
-					mu::NodeSurfacePtr ChildNode = GenerateMutableSourceSurface(ConnectedPin, GenerationContext);
+					UE::Mutable::Private::NodeSurfacePtr ChildNode = GenerateMutableSourceSurface(ConnectedPin, GenerationContext);
 					if (ChildNode)
 					{
 						SurfNode->Variations[VariationIndex].Surfaces.Add( ChildNode );
@@ -840,7 +840,7 @@ mu::Ptr<mu::NodeSurface> GenerateMutableSourceSurface(const UEdGraphPin * Pin, F
 			// Check Switch Parameter arity preconditions.
 			if (const UEdGraphPin* EnumPin = FollowInputPin(*SwitchParameter))
 			{
-				mu::Ptr<mu::NodeScalar> SwitchParam = GenerateMutableSourceFloat(EnumPin, GenerationContext);
+				UE::Mutable::Private::Ptr<UE::Mutable::Private::NodeScalar> SwitchParam = GenerateMutableSourceFloat(EnumPin, GenerationContext);
 
 				// Switch Param not generated
 				if (!SwitchParam)
@@ -855,7 +855,7 @@ mu::Ptr<mu::NodeSurface> GenerateMutableSourceSurface(const UEdGraphPin * Pin, F
 					return Result;
 				}
 
-				if (SwitchParam->GetType() != mu::NodeScalarEnumParameter::GetStaticType())
+				if (SwitchParam->GetType() != UE::Mutable::Private::NodeScalarEnumParameter::GetStaticType())
 				{
 					const FText Message = LOCTEXT("WrongSwitchParamType", "Switch parameter of incorrect type.");
 					GenerationContext.Log(Message, Node);
@@ -865,14 +865,14 @@ mu::Ptr<mu::NodeSurface> GenerateMutableSourceSurface(const UEdGraphPin * Pin, F
 
 				const int32 NumSwitchOptions = TypedNodeSwitch->GetNumElements();
 
-				mu::NodeScalarEnumParameter* EnumParameter = static_cast<mu::NodeScalarEnumParameter*>(SwitchParam.get());
+				UE::Mutable::Private::NodeScalarEnumParameter* EnumParameter = static_cast<UE::Mutable::Private::NodeScalarEnumParameter*>(SwitchParam.get());
 				if (NumSwitchOptions != EnumParameter->Options.Num())
 				{
 					const FText Message = LOCTEXT("MismatchedSwitch", "Switch enum and switch node have different number of options. Please refresh the switch node to make sure the outcomes are labeled properly.");
 					GenerationContext.Log(Message, Node);
 				}
 
-				mu::Ptr<mu::NodeSurfaceSwitch> SwitchNode = new mu::NodeSurfaceSwitch;
+				UE::Mutable::Private::Ptr<UE::Mutable::Private::NodeSurfaceSwitch> SwitchNode = new UE::Mutable::Private::NodeSurfaceSwitch;
 				SwitchNode->Parameter = SwitchParam;
 				SwitchNode->Options.SetNum(NumSwitchOptions);
 
@@ -880,7 +880,7 @@ mu::Ptr<mu::NodeSurface> GenerateMutableSourceSurface(const UEdGraphPin * Pin, F
 				{
 					if (const UEdGraphPin* ConnectedPin = FollowInputPin(*TypedNodeSwitch->GetElementPin(SelectorIndex)))
 					{
-						mu::NodeSurfacePtr ChildNode = GenerateMutableSourceSurface(ConnectedPin, GenerationContext);
+						UE::Mutable::Private::NodeSurfacePtr ChildNode = GenerateMutableSourceSurface(ConnectedPin, GenerationContext);
 						if (ChildNode)
 						{
 							SwitchNode->Options[SelectorIndex] = ChildNode;
@@ -907,13 +907,13 @@ mu::Ptr<mu::NodeSurface> GenerateMutableSourceSurface(const UEdGraphPin * Pin, F
 	else if (const UCustomizableObjectNodeMacroInstance* TypedNodeMacro = Cast<UCustomizableObjectNodeMacroInstance>(Node))
 	{
 		bCacheNode = false;
-		Result = GenerateMutableSourceMacro<mu::NodeSurface>(*Pin, GenerationContext, GenerateMutableSourceSurface);
+		Result = GenerateMutableSourceMacro<UE::Mutable::Private::NodeSurface>(*Pin, GenerationContext, GenerateMutableSourceSurface);
 	}
 
 	else if (const UCustomizableObjectNodeTunnel* TypedNodeTunnel = Cast<UCustomizableObjectNodeTunnel>(Node))
 	{
 		bCacheNode = false;
-		Result = GenerateMutableSourceMacro<mu::NodeSurface>(*Pin, GenerationContext, GenerateMutableSourceSurface);
+		Result = GenerateMutableSourceMacro<UE::Mutable::Private::NodeSurface>(*Pin, GenerationContext, GenerateMutableSourceSurface);
 	}
 
 	else
diff --git a/Engine/Plugins/Mutable/Source/CustomizableObjectEditor/Private/MuCOE/GenerateMutableSource/GenerateMutableSourceSurface.h b/Engine/Plugins/Mutable/Source/CustomizableObjectEditor/Private/MuCOE/GenerateMutableSource/GenerateMutableSourceSurface.h
--- a/Engine/Plugins/Mutable/Source/CustomizableObjectEditor/Private/MuCOE/GenerateMutableSource/GenerateMutableSourceSurface.h
+++ b/Engine/Plugins/Mutable/Source/CustomizableObjectEditor/Private/MuCOE/GenerateMutableSource/GenerateMutableSourceSurface.h
@@ -9,11 +9,11 @@ class UEdGraphPin;
 struct FMutableGraphGenerationContext;
 
 enum class ECustomizableObjectNumBoneInfluences:uint8;
-namespace mu
+namespace UE::Mutable::Private
 {
 	class FMeshBufferSet;
 	class NodeSurface;
 }
 
-extern mu::Ptr<mu::NodeSurface> GenerateMutableSourceSurface(const UEdGraphPin* Pin, FMutableGraphGenerationContext & GenerationContext);
+extern UE::Mutable::Private::Ptr<UE::Mutable::Private::NodeSurface> GenerateMutableSourceSurface(const UEdGraphPin* Pin, FMutableGraphGenerationContext & GenerationContext);
 
diff --git a/Engine/Plugins/Mutable/Source/CustomizableObjectEditor/Private/MuCOE/GenerateMutableSource/GenerateMutableSourceTable.cpp b/Engine/Plugins/Mutable/Source/CustomizableObjectEditor/Private/MuCOE/GenerateMutableSource/GenerateMutableSourceTable.cpp
--- a/Engine/Plugins/Mutable/Source/CustomizableObjectEditor/Private/MuCOE/GenerateMutableSource/GenerateMutableSourceTable.cpp
+++ b/Engine/Plugins/Mutable/Source/CustomizableObjectEditor/Private/MuCOE/GenerateMutableSource/GenerateMutableSourceTable.cpp
@@ -29,7 +29,7 @@
 #define LOCTEXT_NAMESPACE "CustomizableObjectEditor"
 
 
-bool FillTableColumn(const UCustomizableObjectNodeTable* TableNode, mu::Ptr<mu::FTable> MutableTable, const FString& ColumnName, const FString& RowName, const uint32 RowId, uint8* CellData, const FProperty* ColumnProperty,
+bool FillTableColumn(const UCustomizableObjectNodeTable* TableNode, UE::Mutable::Private::Ptr<UE::Mutable::Private::FTable> MutableTable, const FString& ColumnName, const FString& RowName, const uint32 RowId, uint8* CellData, const FProperty* ColumnProperty,
 	const FMutableSourceMeshData& BaseMeshData, const int LODIndexConnected, const int32 SectionIndexConnected, int32 LODIndex, int32 SectionIndex, const bool bOnlyConnectedLOD, FMutableGraphGenerationContext& GenerationContext)
 {
 	int32 CurrentColumn;
@@ -58,10 +58,10 @@ bool FillTableColumn(const UCustomizableObjectNodeTable* TableNode, mu::Ptr<mu::
 
 				if (CurrentColumn == -1)
 				{
-					CurrentColumn = MutableTable->AddColumn(MutableColumnName, mu::ETableColumnType::Mesh);
+					CurrentColumn = MutableTable->AddColumn(MutableColumnName, UE::Mutable::Private::ETableColumnType::Mesh);
 				}
 
-				TSharedPtr<mu::FMesh> EmptySkeletalMesh = nullptr;
+				TSharedPtr<UE::Mutable::Private::FMesh> EmptySkeletalMesh = nullptr;
 				MutableTable->SetCell(CurrentColumn, RowId, EmptySkeletalMesh);
 
 				return true;
@@ -223,7 +223,7 @@ bool FillTableColumn(const UCustomizableObjectNodeTable* TableNode, mu::Ptr<mu::
 
 			if (CurrentColumn == -1)
 			{
-				CurrentColumn = MutableTable->AddColumn(MutableColumnName, mu::ETableColumnType::Mesh);
+				CurrentColumn = MutableTable->AddColumn(MutableColumnName, UE::Mutable::Private::ETableColumnType::Mesh);
 			}
 
 			// First process the mesh tags that are going to make the mesh unique and affect whether it's repeated in 
@@ -277,7 +277,7 @@ bool FillTableColumn(const UCustomizableObjectNodeTable* TableNode, mu::Ptr<mu::
 			FMutableSourceMeshData Source = BaseMeshData;
 			Source.Mesh = SoftObject.ToSoftObjectPath();
 			Source.TableReferenceSkeletalMesh = ReferenceSkeletalMesh;
-			TSharedPtr<mu::FMesh> MutableMesh = GenerateMutableSkeletalMesh(Source, LODIndexConnected, SectionIndexConnected,
+			TSharedPtr<UE::Mutable::Private::FMesh> MutableMesh = GenerateMutableSkeletalMesh(Source, LODIndexConnected, SectionIndexConnected,
 																LODIndex, SectionIndex, MeshUniqueTags, GenerationContext, TableNode);
 
 			if (MutableMesh)
@@ -370,11 +370,11 @@ bool FillTableColumn(const UCustomizableObjectNodeTable* TableNode, mu::Ptr<mu::
 
 			if (CurrentColumn == -1)
 			{
-				CurrentColumn = MutableTable->AddColumn(MutableColumnName, mu::ETableColumnType::Mesh);
+				CurrentColumn = MutableTable->AddColumn(MutableColumnName, UE::Mutable::Private::ETableColumnType::Mesh);
 			}
 
 			constexpr bool bIsPassthrough = false;
-			TSharedPtr<mu::FMesh> MutableMesh = GenerateMutableStaticMesh(StaticMesh, TSoftClassPtr<UAnimInstance>(), CurrentLOD, SectionIndex,
+			TSharedPtr<UE::Mutable::Private::FMesh> MutableMesh = GenerateMutableStaticMesh(StaticMesh, TSoftClassPtr<UAnimInstance>(), CurrentLOD, SectionIndex,
 														  FString(), GenerationContext, TableNode, nullptr, bIsPassthrough);
 
 			if (MutableMesh)
@@ -414,11 +414,11 @@ bool FillTableColumn(const UCustomizableObjectNodeTable* TableNode, mu::Ptr<mu::
 
 			if (CurrentColumn == INDEX_NONE)
 			{
-				CurrentColumn = MutableTable->AddColumn(ColumnName, mu::ETableColumnType::Image);
+				CurrentColumn = MutableTable->AddColumn(ColumnName, UE::Mutable::Private::ETableColumnType::Image);
 			}
 
 			const bool bIsPassthroughTexture = TableNode->GetColumnImageMode(ColumnPropertyName) == ETableTextureType::PASSTHROUGH_TEXTURE;
-			mu::Ptr<mu::TResourceProxyMemory<mu::FImage>> Proxy = new mu::TResourceProxyMemory<mu::FImage>(GenerateImageConstant(Texture, GenerationContext, bIsPassthroughTexture));
+			UE::Mutable::Private::Ptr<UE::Mutable::Private::TResourceProxyMemory<UE::Mutable::Private::FImage>> Proxy = new UE::Mutable::Private::TResourceProxyMemory<UE::Mutable::Private::FImage>(GenerateImageConstant(Texture, GenerationContext, bIsPassthroughTexture));
 			MutableTable->SetCell(CurrentColumn, RowId, Proxy.get());
 		}
 
@@ -477,13 +477,13 @@ bool FillTableColumn(const UCustomizableObjectNodeTable* TableNode, mu::Ptr<mu::
 				MaterialInstance = ReferenceMaterial;
 			}
 
-			if (GenerationContext.CurrentTableColumnType == mu::ETableColumnType::Material)
+			if (GenerationContext.CurrentTableColumnType == UE::Mutable::Private::ETableColumnType::Material)
 			{
 				CurrentColumn = MutableTable.get()->FindColumn(ColumnName);
 
 				if (CurrentColumn == -1)
 				{
-					CurrentColumn = MutableTable->AddColumn(ColumnName, mu::ETableColumnType::Material);
+					CurrentColumn = MutableTable->AddColumn(ColumnName, UE::Mutable::Private::ETableColumnType::Material);
 				}
 
 				int32 ReferenceMaterialId = GenerationContext.ReferencedMaterials.AddUnique(MaterialInstance);
@@ -515,7 +515,7 @@ bool FillTableColumn(const UCustomizableObjectNodeTable* TableNode, mu::Ptr<mu::
 				if (ColumnIndex == INDEX_NONE)
 				{
 					// If there is no column with the parameters name, we generate a new one
-					ColumnIndex = MutableTable->AddColumn(ColumnName, mu::ETableColumnType::Image);
+					ColumnIndex = MutableTable->AddColumn(ColumnName, UE::Mutable::Private::ETableColumnType::Image);
 				}
 
 				UTexture* ParentTextureValue = nullptr;
@@ -550,7 +550,7 @@ bool FillTableColumn(const UCustomizableObjectNodeTable* TableNode, mu::Ptr<mu::
 				}
 
 				bool bIsPassthroughTexture = false;
-				mu::Ptr<mu::TResourceProxyMemory<mu::FImage>> Proxy = new mu::TResourceProxyMemory<mu::FImage>(GenerateImageConstant(ParameterTexture, GenerationContext, bIsPassthroughTexture));
+				UE::Mutable::Private::Ptr<UE::Mutable::Private::TResourceProxyMemory<UE::Mutable::Private::FImage>> Proxy = new UE::Mutable::Private::TResourceProxyMemory<UE::Mutable::Private::FImage>(GenerateImageConstant(ParameterTexture, GenerationContext, bIsPassthroughTexture));
 				MutableTable->SetCell(ColumnIndex, RowId, Proxy.get());
 
 				return true;
@@ -567,15 +567,15 @@ bool FillTableColumn(const UCustomizableObjectNodeTable* TableNode, mu::Ptr<mu::
 
 				if (CurrentColumn == -1)
 				{
-					CurrentColumn = MutableTable->AddColumn(ColumnName, mu::ETableColumnType::Mesh);
+					CurrentColumn = MutableTable->AddColumn(ColumnName, UE::Mutable::Private::ETableColumnType::Mesh);
 				}
 
 				TArray<FName> ArrayBoneName;
 				TArray<FTransform> ArrayTransform;
 				UCustomizableObjectNodeAnimationPose::StaticRetrievePoseInformation(PoseAsset, GenerationContext.GetCurrentComponentInfo()->RefSkeletalMesh.Get(), ArrayBoneName, ArrayTransform);
 
-				TSharedPtr<mu::FMesh> MutableMesh = MakeShared<mu::FMesh>();
-				TSharedPtr<mu::FSkeleton> MutableSkeleton = MakeShared<mu::FSkeleton>();
+				TSharedPtr<UE::Mutable::Private::FMesh> MutableMesh = MakeShared<UE::Mutable::Private::FMesh>();
+				TSharedPtr<UE::Mutable::Private::FSkeleton> MutableSkeleton = MakeShared<UE::Mutable::Private::FSkeleton>();
 
 				MutableMesh->SetSkeleton(MutableSkeleton);
 				MutableMesh->SetBonePoseCount(ArrayBoneName.Num());
@@ -584,11 +584,11 @@ bool FillTableColumn(const UCustomizableObjectNodeTable* TableNode, mu::Ptr<mu::
 				for (int32 i = 0; i < ArrayBoneName.Num(); ++i)
 				{
 					const FName BoneName = ArrayBoneName[i];
-					const mu::FBoneName& BoneId = GenerationContext.CompilationContext->GetBoneUnique(BoneName);
+					const UE::Mutable::Private::FBoneName& BoneId = GenerationContext.CompilationContext->GetBoneUnique(BoneName);
 
 					MutableSkeleton->SetDebugName(i, BoneName);
 					MutableSkeleton->SetBoneName(i, { BoneId });
-					MutableMesh->SetBonePose(i, BoneId, (FTransform3f)ArrayTransform[i], mu::EBoneUsageFlags::Skinning);
+					MutableMesh->SetBonePose(i, BoneId, (FTransform3f)ArrayTransform[i], UE::Mutable::Private::EBoneUsageFlags::Skinning);
 				}
 
 				MutableTable->SetCell(CurrentColumn, RowId, MutableMesh);
@@ -613,7 +613,7 @@ bool FillTableColumn(const UCustomizableObjectNodeTable* TableNode, mu::Ptr<mu::
 
 			if (CurrentColumn == INDEX_NONE)
 			{
-				CurrentColumn = MutableTable->AddColumn(ColumnName, mu::ETableColumnType::Color);
+				CurrentColumn = MutableTable->AddColumn(ColumnName, UE::Mutable::Private::ETableColumnType::Color);
 			}
 
 			// Setting cell value
@@ -634,7 +634,7 @@ bool FillTableColumn(const UCustomizableObjectNodeTable* TableNode, mu::Ptr<mu::
 
 		if (CurrentColumn == INDEX_NONE)
 		{
-			CurrentColumn = MutableTable->AddColumn(ColumnName, mu::ETableColumnType::Scalar);
+			CurrentColumn = MutableTable->AddColumn(ColumnName, UE::Mutable::Private::ETableColumnType::Scalar);
 		}
 
 		// Setting cell value
@@ -648,7 +648,7 @@ bool FillTableColumn(const UCustomizableObjectNodeTable* TableNode, mu::Ptr<mu::
 	
 		if (CurrentColumn == INDEX_NONE)
 		{
-			CurrentColumn = MutableTable->AddColumn(ColumnName, mu::ETableColumnType::Scalar);
+			CurrentColumn = MutableTable->AddColumn(ColumnName, UE::Mutable::Private::ETableColumnType::Scalar);
 		}
 	
 		// Setting cell value
@@ -831,7 +831,7 @@ TArray<FName> GetRowsToCompile(const UDataTable& DataTable, const UCustomizableO
 }
 
 
-bool GenerateTableColumn(const UCustomizableObjectNodeTable* TableNode, const UEdGraphPin* Pin, mu::Ptr<mu::FTable> MutableTable, const FString& DataTableColumnName, const FProperty* ColumnProperty, 
+bool GenerateTableColumn(const UCustomizableObjectNodeTable* TableNode, const UEdGraphPin* Pin, UE::Mutable::Private::Ptr<UE::Mutable::Private::FTable> MutableTable, const FString& DataTableColumnName, const FProperty* ColumnProperty, 
 	const FMutableSourceMeshData& BaseMeshData,	const int32 LODIndexConnected, const int32 SectionIndexConnected, const int32 LODIndex, const int32 SectionIndex, const bool bOnlyConnectedLOD, FMutableGraphGenerationContext& GenerationContext)
 {
 	MUTABLE_CPUPROFILER_SCOPE(GenerateTableColumn);
@@ -1035,7 +1035,7 @@ void GenerateTableParameterUIData(const UDataTable* DataTable, const UCustomizab
 }
 
 
-mu::Ptr<mu::FTable> GenerateMutableSourceTable(const UDataTable* DataTable, const UCustomizableObjectNodeTable* TableNode, FMutableGraphGenerationContext& GenerationContext)
+UE::Mutable::Private::Ptr<UE::Mutable::Private::FTable> GenerateMutableSourceTable(const UDataTable* DataTable, const UCustomizableObjectNodeTable* TableNode, FMutableGraphGenerationContext& GenerationContext)
 {
 	check(DataTable && TableNode);
 
@@ -1063,7 +1063,7 @@ mu::Ptr<mu::FTable> GenerateMutableSourceTable(const UDataTable* DataTable, cons
 		return CachedTable->GeneratedTable;
 	}
 
-	mu::Ptr<mu::FTable> MutableTable = new mu::FTable();
+	UE::Mutable::Private::Ptr<UE::Mutable::Private::FTable> MutableTable = new UE::Mutable::Private::FTable();
 
 	if (const UScriptStruct* TableStruct = DataTable->GetRowStruct())
 	{
@@ -1093,7 +1093,7 @@ mu::Ptr<mu::FTable> GenerateMutableSourceTable(const UDataTable* DataTable, cons
 		}
 
 		// Adding the Name Column
-		MutableTable->AddColumn("Name", mu::ETableColumnType::String);
+		MutableTable->AddColumn("Name", UE::Mutable::Private::ETableColumnType::String);
 
 		// Always generate "None" row
 		{
@@ -1145,7 +1145,7 @@ UDataTable* GetDataTable(const UCustomizableObjectNodeTable* TableNode, FMutable
 	}
 	else
 	{
-		OutDataTable = MutablePrivate::LoadObject(TableNode->Table);
+		OutDataTable = UE::Mutable::Private::LoadObject(TableNode->Table);
 	}
 
 	return OutDataTable;
@@ -1154,7 +1154,7 @@ UDataTable* GetDataTable(const UCustomizableObjectNodeTable* TableNode, FMutable
 
 UDataTable* GenerateDataTableFromStruct(const UCustomizableObjectNodeTable* TableNode, FMutableGraphGenerationContext& GenerationContext)
 {
-	TObjectPtr<UScriptStruct> Structure = MutablePrivate::LoadObject(TableNode->Structure);
+	TObjectPtr<UScriptStruct> Structure = UE::Mutable::Private::LoadObject(TableNode->Structure);
 
 	if (!Structure)
 	{
@@ -1194,7 +1194,7 @@ UDataTable* GenerateDataTableFromStruct(const UCustomizableObjectNodeTable* Tabl
 	{
 		if (DataTableAsset.IsValid())
 		{
-			if (UDataTable* DataTable = Cast<UDataTable>(MutablePrivate::LoadObject(DataTableAsset)))
+			if (UDataTable* DataTable = Cast<UDataTable>(UE::Mutable::Private::LoadObject(DataTableAsset)))
 			{
 				ParentTables.Add(DataTable);
 			}
@@ -1304,7 +1304,7 @@ void LogRowGenerationMessage(const UCustomizableObjectNodeTable* TableNode, cons
 }
 
 
-bool GenerateNoneRow(const UCustomizableObjectNodeTable& TableNode, const UEdGraphPin* Pin, const FString& ColumnName, mu::Ptr<mu::FTable> MutableTable, FMutableGraphGenerationContext& GenerationContext)
+bool GenerateNoneRow(const UCustomizableObjectNodeTable& TableNode, const UEdGraphPin* Pin, const FString& ColumnName, UE::Mutable::Private::Ptr<UE::Mutable::Private::FTable> MutableTable, FMutableGraphGenerationContext& GenerationContext)
 {
 	if (!Pin || !MutableTable)
 	{
@@ -1321,12 +1321,12 @@ bool GenerateNoneRow(const UCustomizableObjectNodeTable& TableNode, const UEdGra
 	const UEdGraphSchema_CustomizableObject* Schema = GetDefault<UEdGraphSchema_CustomizableObject>();
 	check(Schema);
 
-	if (Pin->PinType.PinCategory == Schema->PC_Material && GenerationContext.CurrentTableColumnType == mu::ETableColumnType::Material)
+	if (Pin->PinType.PinCategory == Schema->PC_Material && GenerationContext.CurrentTableColumnType == UE::Mutable::Private::ETableColumnType::Material)
 	{
 		// Add a scalar column if it has not been added yet
 		if (ColumnId == INDEX_NONE)
 		{
-			ColumnId = MutableTable->AddColumn(ColumnName, mu::ETableColumnType::Material);
+			ColumnId = MutableTable->AddColumn(ColumnName, UE::Mutable::Private::ETableColumnType::Material);
 		}
 
 		MutableTable->SetCell(ColumnId, RowId, GenerationContext.CurrentReferencedMaterialIndex);
@@ -1336,7 +1336,7 @@ bool GenerateNoneRow(const UCustomizableObjectNodeTable& TableNode, const UEdGra
 		// Add a color column if it has not been added yet
 		if (ColumnId == INDEX_NONE)
 		{
-			ColumnId = MutableTable->AddColumn(ColumnName, mu::ETableColumnType::Color);
+			ColumnId = MutableTable->AddColumn(ColumnName, UE::Mutable::Private::ETableColumnType::Color);
 		}
 
 		// HACK: Encoding an invalid value (Nan) for table option "None".
diff --git a/Engine/Plugins/Mutable/Source/CustomizableObjectEditor/Private/MuCOE/GenerateMutableSource/GenerateMutableSourceTable.h b/Engine/Plugins/Mutable/Source/CustomizableObjectEditor/Private/MuCOE/GenerateMutableSource/GenerateMutableSourceTable.h
--- a/Engine/Plugins/Mutable/Source/CustomizableObjectEditor/Private/MuCOE/GenerateMutableSource/GenerateMutableSourceTable.h
+++ b/Engine/Plugins/Mutable/Source/CustomizableObjectEditor/Private/MuCOE/GenerateMutableSource/GenerateMutableSourceTable.h
@@ -27,7 +27,7 @@ struct FMutableSourceMeshData;
  * @param bOnlyConnectedLOD Corrected LOD and Section will unconditionally always be the connected ones.
  * @param GenerationContext 
  * @return  */
-bool GenerateTableColumn(const UCustomizableObjectNodeTable* TableNode, const UEdGraphPin* Pin, mu::Ptr<mu::FTable> MutableTable, const FString& DataTableColumnName, const FProperty* ColumnProperty,
+bool GenerateTableColumn(const UCustomizableObjectNodeTable* TableNode, const UEdGraphPin* Pin, UE::Mutable::Private::Ptr<UE::Mutable::Private::FTable> MutableTable, const FString& DataTableColumnName, const FProperty* ColumnProperty,
 	const FMutableSourceMeshData& BaseMeshData,
 	int32 LODIndexConnected, int32 SectionIndexConnected,
 	int32 LODIndex, int32 SectionIndex,
@@ -52,7 +52,7 @@ bool GenerateTableColumn(const UCustomizableObjectNodeTable* TableNode, const UE
  * @param bOnlyConnectedLOD Corrected LOD and Section will unconditionally always be the connected ones.
  * @param GenerationContext 
  * @return  */
-bool FillTableColumn(const UCustomizableObjectNodeTable* TableNode, mu::Ptr<mu::FTable> MutableTable, const FString& ColumnName,
+bool FillTableColumn(const UCustomizableObjectNodeTable* TableNode, UE::Mutable::Private::Ptr<UE::Mutable::Private::FTable> MutableTable, const FString& ColumnName,
 	const FString& RowName, uint32 RowId, uint8* CellData, const FProperty* ColumnProperty,
 	const FMutableSourceMeshData& BaseMeshData,
 	int32 LODIndexConnected, int32 SectionIndexConnected, int32 LODIndex, int32 SectionIndex,
@@ -60,7 +60,7 @@ bool FillTableColumn(const UCustomizableObjectNodeTable* TableNode, mu::Ptr<mu::
 	FMutableGraphGenerationContext& GenerationContext);
 
 
-mu::Ptr<mu::FTable> GenerateMutableSourceTable(const UDataTable* DataTable, const UCustomizableObjectNodeTable* TableNode, FMutableGraphGenerationContext& GenerationContext);
+UE::Mutable::Private::Ptr<UE::Mutable::Private::FTable> GenerateMutableSourceTable(const UDataTable* DataTable, const UCustomizableObjectNodeTable* TableNode, FMutableGraphGenerationContext& GenerationContext);
 
 /** Generates a new FParameterUIData for a new table parameter. Also, adds the UIData to the parameters map. */
 void GenerateTableParameterUIData(const UDataTable* DataTable, const UCustomizableObjectNodeTable* TableNode, FMutableGraphGenerationContext& GenerationContext);
@@ -77,7 +77,7 @@ TArray<FName> GetRowsToCompile(const UDataTable& DataTable, const UCustomizableO
 UDataTable* GenerateDataTableFromStruct(const UCustomizableObjectNodeTable* TableNode, FMutableGraphGenerationContext& GenerationContext);
 
 /** Generates and populates the "None" row for each column. */
-bool GenerateNoneRow(const UCustomizableObjectNodeTable& TableNode, const UEdGraphPin* Pin, const FString& ColumnName, mu::Ptr<mu::FTable> MutableTable, FMutableGraphGenerationContext& GenerationContext);
+bool GenerateNoneRow(const UCustomizableObjectNodeTable& TableNode, const UEdGraphPin* Pin, const FString& ColumnName, UE::Mutable::Private::Ptr<UE::Mutable::Private::FTable> MutableTable, FMutableGraphGenerationContext& GenerationContext);
 
 /** This method adds the original Data Table(s) of the processed row at the end of the log message when the data table is a Composite Data table. */
 void LogRowGenerationMessage(const UCustomizableObjectNodeTable* TableNode, const UDataTable* DataTable, FMutableGraphGenerationContext& GenerationContext, const FString& Message, const FString& RowName);
diff --git a/Engine/Plugins/Mutable/Source/CustomizableObjectEditor/Private/MuCOE/GenerateMutableSource/GenerateMutableSourceTransform.cpp b/Engine/Plugins/Mutable/Source/CustomizableObjectEditor/Private/MuCOE/GenerateMutableSource/GenerateMutableSourceTransform.cpp
--- a/Engine/Plugins/Mutable/Source/CustomizableObjectEditor/Private/MuCOE/GenerateMutableSource/GenerateMutableSourceTransform.cpp
+++ b/Engine/Plugins/Mutable/Source/CustomizableObjectEditor/Private/MuCOE/GenerateMutableSource/GenerateMutableSourceTransform.cpp
@@ -12,7 +12,7 @@
 #define LOCTEXT_NAMESPACE "CustomizableObjectEditor"
 
 
-mu::Ptr<mu::NodeMatrix> GenerateMutableSourceTransform(const UEdGraphPin* Pin, FMutableGraphGenerationContext& GenerationContext)
+UE::Mutable::Private::Ptr<UE::Mutable::Private::NodeMatrix> GenerateMutableSourceTransform(const UEdGraphPin* Pin, FMutableGraphGenerationContext& GenerationContext)
 {
 	check(Pin);
 
@@ -24,7 +24,7 @@ mu::Ptr<mu::NodeMatrix> GenerateMutableSourceTransform(const UEdGraphPin* Pin, F
 	const FGeneratedKey Key(reinterpret_cast<void*>(&GenerateMutableSourceTransform), *Pin, *Node, GenerationContext);
 	if (const FGeneratedData* Generated = GenerationContext.Generated.Find(Key))
 	{
-		return static_cast<mu::NodeMatrix*>(Generated->Node.get());
+		return static_cast<UE::Mutable::Private::NodeMatrix*>(Generated->Node.get());
 	}
 
 	if (Node->IsNodeOutDatedAndNeedsRefresh())
@@ -32,16 +32,16 @@ mu::Ptr<mu::NodeMatrix> GenerateMutableSourceTransform(const UEdGraphPin* Pin, F
 		Node->SetRefreshNodeWarning();
 	}
 
-	mu::Ptr<mu::NodeMatrix> Result;
+	UE::Mutable::Private::Ptr<UE::Mutable::Private::NodeMatrix> Result;
 	if (const UCustomizableObjectNodeTransformConstant* TypedNodeTransformConstant = Cast<UCustomizableObjectNodeTransformConstant>(Node))
 	{
-		mu::Ptr<mu::NodeMatrixConstant> MatrixNode = new mu::NodeMatrixConstant{};
+		UE::Mutable::Private::Ptr<UE::Mutable::Private::NodeMatrixConstant> MatrixNode = new UE::Mutable::Private::NodeMatrixConstant{};
 		Result = MatrixNode;
 		MatrixNode->Value = FMatrix44f(TypedNodeTransformConstant->Value.ToMatrixWithScale());
 	}
 	else if (const UCustomizableObjectNodeTransformParameter* TypedNodeTransformParameter = Cast<UCustomizableObjectNodeTransformParameter>(Node))
 	{
-		mu::Ptr<mu::NodeMatrixParameter> MatrixNode = new mu::NodeMatrixParameter();
+		UE::Mutable::Private::Ptr<UE::Mutable::Private::NodeMatrixParameter> MatrixNode = new UE::Mutable::Private::NodeMatrixParameter();
 		Result = MatrixNode;
 		
 		MatrixNode->Name = TypedNodeTransformParameter->GetParameterName(&GenerationContext.MacroNodesStack);
diff --git a/Engine/Plugins/Mutable/Source/CustomizableObjectEditor/Private/MuCOE/GenerateMutableSource/GenerateMutableSourceTransform.h b/Engine/Plugins/Mutable/Source/CustomizableObjectEditor/Private/MuCOE/GenerateMutableSource/GenerateMutableSourceTransform.h
--- a/Engine/Plugins/Mutable/Source/CustomizableObjectEditor/Private/MuCOE/GenerateMutableSource/GenerateMutableSourceTransform.h
+++ b/Engine/Plugins/Mutable/Source/CustomizableObjectEditor/Private/MuCOE/GenerateMutableSource/GenerateMutableSourceTransform.h
@@ -8,4 +8,4 @@ class UEdGraphPin;
 struct FMutableGraphGenerationContext;
 
 /** Convert a CustomizableObject Source Graph into a mutable source graph. */
-mu::Ptr<mu::NodeMatrix> GenerateMutableSourceTransform(const UEdGraphPin* Pin, FMutableGraphGenerationContext& GenerationContext);
+UE::Mutable::Private::Ptr<UE::Mutable::Private::NodeMatrix> GenerateMutableSourceTransform(const UEdGraphPin* Pin, FMutableGraphGenerationContext& GenerationContext);
diff --git a/Engine/Plugins/Mutable/Source/CustomizableObjectEditor/Private/MuCOE/GraphTraversal.cpp b/Engine/Plugins/Mutable/Source/CustomizableObjectEditor/Private/MuCOE/GraphTraversal.cpp
--- a/Engine/Plugins/Mutable/Source/CustomizableObjectEditor/Private/MuCOE/GraphTraversal.cpp
+++ b/Engine/Plugins/Mutable/Source/CustomizableObjectEditor/Private/MuCOE/GraphTraversal.cpp
@@ -781,7 +781,7 @@ void GetNodeGroupObjectNodeMappingImmersive(UCustomizableObject* Object, IAssetR
 
 	for (FAssetData& AssetData : ArrayAssetData)
 	{
-		UCustomizableObject* ChildObject = Cast<UCustomizableObject>(MutablePrivate::LoadObject(AssetData));
+		UCustomizableObject* ChildObject = Cast<UCustomizableObject>(UE::Mutable::Private::LoadObject(AssetData));
 		if (!ChildObject)
 		{
 			continue;
diff --git a/Engine/Plugins/Mutable/Source/CustomizableObjectEditor/Private/MuCOE/MutableMeshPreviewUtils.cpp b/Engine/Plugins/Mutable/Source/CustomizableObjectEditor/Private/MuCOE/MutableMeshPreviewUtils.cpp
--- a/Engine/Plugins/Mutable/Source/CustomizableObjectEditor/Private/MuCOE/MutableMeshPreviewUtils.cpp
+++ b/Engine/Plugins/Mutable/Source/CustomizableObjectEditor/Private/MuCOE/MutableMeshPreviewUtils.cpp
@@ -30,17 +30,17 @@ namespace MutableMeshPreviewUtils
 		 * @return An identical mutable mesh to the provided one in terms of data but with the buffer structure that UE
 		 * requires
 		 */
-		TSharedPtr<const mu::FMesh> GenerateCompatibleMutableMesh(const mu::FMesh* InMutableMesh)
+		TSharedPtr<const UE::Mutable::Private::FMesh> GenerateCompatibleMutableMesh(const UE::Mutable::Private::FMesh* InMutableMesh)
 		{
 			// 1) Generate a mutable mesh to be the one hosting the buffer format we desire (it will be required for 
-			TSharedPtr<mu::FMesh> FormattedMutableMesh = MakeShared<mu::FMesh>();
-			mu::FMeshBufferSet& FormattedVertexBuffers = FormattedMutableMesh->GetVertexBuffers();
-			mu::FMeshBufferSet& FormattedIndexBuffers = FormattedMutableMesh->GetIndexBuffers();
+			TSharedPtr<UE::Mutable::Private::FMesh> FormattedMutableMesh = MakeShared<UE::Mutable::Private::FMesh>();
+			UE::Mutable::Private::FMeshBufferSet& FormattedVertexBuffers = FormattedMutableMesh->GetVertexBuffers();
+			UE::Mutable::Private::FMeshBufferSet& FormattedIndexBuffers = FormattedMutableMesh->GetIndexBuffers();
 			
 			// 2) Generate the buffers for that mesh with the format unreal expects
 			{
 				// 2.1)		Locate what buffers are present on the mutable constant mesh
-				const mu::FMeshBufferSet& OriginVertexBuffers = InMutableMesh->GetVertexBuffers();
+				const UE::Mutable::Private::FMeshBufferSet& OriginVertexBuffers = InMutableMesh->GetVertexBuffers();
 
 				// Data extracted from the original buffers to be used later when setting up the formatted mesh
 				bool bOriginDoesHaveVertexColorData = false;
@@ -60,8 +60,8 @@ namespace MutableMeshPreviewUtils
 					const int32 BufferChannelCount = OriginVertexBuffers.GetBufferChannelCount(BufferIndex);
 					for (int32 ChannelIndex = 0; ChannelIndex < BufferChannelCount; ChannelIndex ++)
 					{
-						mu::EMeshBufferSemantic BufferSemantic = mu::EMeshBufferSemantic::None;
-						mu::EMeshBufferFormat BufferFormat = mu::EMeshBufferFormat::None;
+						UE::Mutable::Private::EMeshBufferSemantic BufferSemantic = UE::Mutable::Private::EMeshBufferSemantic::None;
+						UE::Mutable::Private::EMeshBufferFormat BufferFormat = UE::Mutable::Private::EMeshBufferFormat::None;
 						int32 BufferComponentCount = 0;
 						int32 SemanticIndex = 0;
 
@@ -79,16 +79,16 @@ namespace MutableMeshPreviewUtils
 
 						switch (BufferSemantic)
 						{
-							case mu::EMeshBufferSemantic::Color:
+							case UE::Mutable::Private::EMeshBufferSemantic::Color:
 							{
 								// Does the original mesh have vert color data?
 								bOriginDoesHaveVertexColorData = true;
 								break;
 							}
 
-							case mu::EMeshBufferSemantic::TexCoords:
+							case UE::Mutable::Private::EMeshBufferSemantic::TexCoords:
 							{
-								// Store the texture semantic index to be later used during the buffer formatting process on  mu::MeshFormat
+								// Store the texture semantic index to be later used during the buffer formatting process on  UE::Mutable::Private::MeshFormat
 								TextureSemanticIndices[NumbOfTextCoordChannels] = SemanticIndex;
 									
 								// Store the amount of located texture data channels
@@ -96,18 +96,18 @@ namespace MutableMeshPreviewUtils
 								break;
 							}
 
-							case mu::EMeshBufferSemantic::BoneIndices:
+							case UE::Mutable::Private::EMeshBufferSemantic::BoneIndices:
 							{
 								// Store the amount of bones a vertex can be skinned to.
 								MaxNumBonesPerVertex = BufferComponentCount;
-								BoneIndicesSizeBytes = BufferFormat == mu::EMeshBufferFormat::Int16 ? 2 : 1;
+								BoneIndicesSizeBytes = BufferFormat == UE::Mutable::Private::EMeshBufferFormat::Int16 ? 2 : 1;
 								break;
 							}
 
-							case mu::EMeshBufferSemantic::BoneWeights:
+							case UE::Mutable::Private::EMeshBufferSemantic::BoneWeights:
 							{
 								// Store the amount of bones a vertex can be skinned to.
-								BoneWeightsSizeBytes = BufferFormat == mu::EMeshBufferFormat::NUInt16 ? 2 : 1;
+								BoneWeightsSizeBytes = BufferFormat == UE::Mutable::Private::EMeshBufferFormat::NUInt16 ? 2 : 1;
 								break;
 							}
 
@@ -174,11 +174,11 @@ namespace MutableMeshPreviewUtils
 			// At this point a new mutable mesh with no data but with the appropriate structure for UE to be able to convert it to an
 			// Skeletal mesh is ready to be used as an argument for the process of reformatting the Mutable Mesh we want to preview
 
-			// 3) Call mu::MeshFormat to get the mesh with the wanted buffer structure to have the data of the original mesh
-			TSharedPtr<mu::FMesh> FormatResult = MakeShared<mu::FMesh>();
+			// 3) Call UE::Mutable::Private::MeshFormat to get the mesh with the wanted buffer structure to have the data of the original mesh
+			TSharedPtr<UE::Mutable::Private::FMesh> FormatResult = MakeShared<UE::Mutable::Private::FMesh>();
 
 			bool bOutSuccess = false;
-			mu::MeshFormat(FormatResult.Get(), InMutableMesh, FormattedMutableMesh.Get(),
+			UE::Mutable::Private::MeshFormat(FormatResult.Get(), InMutableMesh, FormattedMutableMesh.Get(),
 				false, true, true, false,
 				bOutSuccess);
 		
@@ -192,19 +192,19 @@ namespace MutableMeshPreviewUtils
 		 * badly located inside the buffer set.
 		 * @return True if the buffers are where they should be, false if not correctly placed or missing.
 		 */
-		bool AreVertexBuffersCorrectlyLocated(const mu::FMesh* InMutableMesh, bool bLogFindingsToConsole = false )
+		bool AreVertexBuffersCorrectlyLocated(const UE::Mutable::Private::FMesh* InMutableMesh, bool bLogFindingsToConsole = false )
 		{
 			bool bCanBeConverted = true;
 
 			// At the moment requires the presence of all "mandatory" buffers. Only meshes designed to be displayed by unreal
 			// will therefore be processable at the moment and not always
-			const mu::FMeshBufferSet& VertexBuffers = InMutableMesh->GetVertexBuffers();
+			const UE::Mutable::Private::FMeshBufferSet& VertexBuffers = InMutableMesh->GetVertexBuffers();
 
 			int32 Buffer = -1;
 			int32 Channel = -1;
 
 			// POSITION BUFFER INDEX CHECK
-			VertexBuffers.FindChannel(mu::EMeshBufferSemantic::Position, 0, &Buffer, &Channel);
+			VertexBuffers.FindChannel(UE::Mutable::Private::EMeshBufferSemantic::Position, 0, &Buffer, &Channel);
 			if (Buffer == -1)
 			{
 				if (bLogFindingsToConsole)
@@ -228,7 +228,7 @@ namespace MutableMeshPreviewUtils
 			}
 
 			// SKIN WEIGHTS CHECK
-			VertexBuffers.FindChannel(mu::EMeshBufferSemantic::BoneIndices, 0, &Buffer, &Channel);
+			VertexBuffers.FindChannel(UE::Mutable::Private::EMeshBufferSemantic::BoneIndices, 0, &Buffer, &Channel);
 			if (Buffer == -1)
 			{
 				if (bLogFindingsToConsole)
@@ -240,7 +240,7 @@ namespace MutableMeshPreviewUtils
 			}
 
 			// TANGENT BUFFER INDEX CHECK
-			VertexBuffers.FindChannel(mu::EMeshBufferSemantic::Tangent, 0, &Buffer, &Channel);
+			VertexBuffers.FindChannel(UE::Mutable::Private::EMeshBufferSemantic::Tangent, 0, &Buffer, &Channel);
 			if (Buffer == -1)
 			{
 				if (bLogFindingsToConsole)
@@ -264,7 +264,7 @@ namespace MutableMeshPreviewUtils
 			}
 
 			// TEX COORDS BUFFER INDEX CHECK
-			VertexBuffers.FindChannel(mu::EMeshBufferSemantic::TexCoords, 0, &Buffer, &Channel);
+			VertexBuffers.FindChannel(UE::Mutable::Private::EMeshBufferSemantic::TexCoords, 0, &Buffer, &Channel);
 			if (Buffer == -1)
 			{
 				if (bLogFindingsToConsole)
@@ -295,7 +295,7 @@ namespace MutableMeshPreviewUtils
 		 * @param OutSkeletalMesh - The skeletal mesh whose skeleton we want to generate.
 		 * @return True if the operation could be performed successfully, false if not.
 		 */
-		bool BuildSkeletalMeshSkeletonData(const mu::FMesh* InMutableMesh, USkeletalMesh* OutSkeletalMesh)
+		bool BuildSkeletalMeshSkeletonData(const UE::Mutable::Private::FMesh* InMutableMesh, USkeletalMesh* OutSkeletalMesh)
 		{
 			if (!InMutableMesh || !OutSkeletalMesh)
 			{
@@ -353,7 +353,7 @@ namespace MutableMeshPreviewUtils
 		 * @param InMutableMesh - The reference mutable mesh to be reading data from
 		 * @param OutSkeletalMesh - The skeletal mesh to be getting it's materials set up.
 		 */
-		void BuildSkeletalMeshElementData(const mu::FMesh* InMutableMesh, USkeletalMesh* OutSkeletalMesh)
+		void BuildSkeletalMeshElementData(const UE::Mutable::Private::FMesh* InMutableMesh, USkeletalMesh* OutSkeletalMesh)
 		{
 			// Set up Unreal's default material
 			UMaterialInterface* UnrealMaterial = UMaterial::GetDefaultMaterial(MD_Surface);
@@ -380,7 +380,7 @@ namespace MutableMeshPreviewUtils
 		 * @param OutSkeletalMesh - The skeletal mesh to set up.
 		 * @return True if the operation could be performed successfully, false if not.
 		 */
-		bool BuildSkeletalMeshRenderData(const mu::FMesh* InMutableMesh, USkeletalMesh* OutSkeletalMesh)
+		bool BuildSkeletalMeshRenderData(const UE::Mutable::Private::FMesh* InMutableMesh, USkeletalMesh* OutSkeletalMesh)
 		{
 			OutSkeletalMesh->SetHasVertexColors(false);
 
@@ -388,11 +388,11 @@ namespace MutableMeshPreviewUtils
 			const int32 NumBonesInBoneMap = !InMutableMesh->GetBoneMap().IsEmpty() ? InMutableMesh->GetBoneMap().Num() : InMutableMesh->GetBonePoseCount();
 			
 			// Fill the bonemap with zeros
-			TArray<mu::FBoneName> BoneMap;
+			TArray<UE::Mutable::Private::FBoneName> BoneMap;
 			BoneMap.SetNumZeroed(FMath::Max(NumBonesInBoneMap, 1));
 
-			TMap <mu::FBoneName, TPair<FName, uint16>> BoneInfoMap;
-			BoneInfoMap.Add(mu::FBoneName(0), {NAME_None, 0});
+			TMap <UE::Mutable::Private::FBoneName, TPair<FName, uint16>> BoneInfoMap;
+			BoneInfoMap.Add(UE::Mutable::Private::FBoneName(0), {NAME_None, 0});
 
 			FSkeletalMeshLODRenderData& LODResource = OutSkeletalMesh->GetResourceForRendering()->LODRenderData[0];
 
@@ -457,7 +457,7 @@ namespace MutableMeshPreviewUtils
 	 * Publicly accessible functions of MutableMeshPreviewUtils
 	 */
 
-	USkeletalMesh* GenerateSkeletalMeshFromMutableMesh(const mu::FMesh* InMutableMesh)
+	USkeletalMesh* GenerateSkeletalMeshFromMutableMesh(const UE::Mutable::Private::FMesh* InMutableMesh)
 	{
 		if (!InMutableMesh)
 		{
@@ -469,7 +469,7 @@ namespace MutableMeshPreviewUtils
 			return nullptr;
 		}
 
-		TSharedPtr<const mu::FMesh> MeshCopy;
+		TSharedPtr<const UE::Mutable::Private::FMesh> MeshCopy;
 
 		// Buffer structures checks and corrective operations
 		{
diff --git a/Engine/Plugins/Mutable/Source/CustomizableObjectEditor/Private/MuCOE/MutableMeshPreviewUtils.h b/Engine/Plugins/Mutable/Source/CustomizableObjectEditor/Private/MuCOE/MutableMeshPreviewUtils.h
--- a/Engine/Plugins/Mutable/Source/CustomizableObjectEditor/Private/MuCOE/MutableMeshPreviewUtils.h
+++ b/Engine/Plugins/Mutable/Source/CustomizableObjectEditor/Private/MuCOE/MutableMeshPreviewUtils.h
@@ -18,6 +18,6 @@ namespace MutableMeshPreviewUtils
 	 * @param InMutableMesh - The mutable mesh used as base for the generation of the USkeletalMesh
 	 * @return - A new Skeletal mesh based on the provided data. It may be a nullptr if the conversion fails.
 	 */
-	CUSTOMIZABLEOBJECTEDITOR_API USkeletalMesh* GenerateSkeletalMeshFromMutableMesh(const mu::FMesh*);
+	CUSTOMIZABLEOBJECTEDITOR_API USkeletalMesh* GenerateSkeletalMeshFromMutableMesh(const UE::Mutable::Private::FMesh*);
 
 }
diff --git a/Engine/Plugins/Mutable/Source/CustomizableObjectEditor/Private/MuCOE/Nodes/CustomizableObjectNodeComponentPassthroughMesh.cpp b/Engine/Plugins/Mutable/Source/CustomizableObjectEditor/Private/MuCOE/Nodes/CustomizableObjectNodeComponentPassthroughMesh.cpp
--- a/Engine/Plugins/Mutable/Source/CustomizableObjectEditor/Private/MuCOE/Nodes/CustomizableObjectNodeComponentPassthroughMesh.cpp
+++ b/Engine/Plugins/Mutable/Source/CustomizableObjectEditor/Private/MuCOE/Nodes/CustomizableObjectNodeComponentPassthroughMesh.cpp
@@ -91,7 +91,7 @@ void UCustomizableObjectNodeComponentPassthroughMesh::AllocateDefaultPins(UCusto
 	}
 	
 	// Support SkeletalMeshes only, for now.
-	USkeletalMesh* StrongSkeletalMesh = MutablePrivate::LoadObject<USkeletalMesh>(SkeletalMesh);
+	USkeletalMesh* StrongSkeletalMesh = UE::Mutable::Private::LoadObject<USkeletalMesh>(SkeletalMesh);
 	if (!StrongSkeletalMesh)
 	{
 		return;
@@ -161,7 +161,7 @@ void UCustomizableObjectNodeComponentPassthroughMesh::GetPinSection(const UEdGra
 
 UMaterialInterface* UCustomizableObjectNodeComponentPassthroughMesh::GetMaterialFor(const UEdGraphPin* Pin) const
 {
-	USkeletalMesh* StrongSkeletalMesh = MutablePrivate::LoadObject<USkeletalMesh>(SkeletalMesh);
+	USkeletalMesh* StrongSkeletalMesh = UE::Mutable::Private::LoadObject<USkeletalMesh>(SkeletalMesh);
 	if (StrongSkeletalMesh)
 	{
 		if (FSkeletalMaterial* SkeletalMaterial = GetSkeletalMaterialFor(*Pin))
@@ -203,7 +203,7 @@ UMaterialInterface* UCustomizableObjectNodeComponentPassthroughMesh::GetMaterial
 
 FSkeletalMaterial* UCustomizableObjectNodeComponentPassthroughMesh::GetSkeletalMaterialFor(const int32 LODIndex, const int32 SectionIndex, const FSkeletalMeshModel* ImportedModel) const
 {
-	USkeletalMesh* StrongSkeletalMesh = Cast<USkeletalMesh>(MutablePrivate::LoadObject(SkeletalMesh));
+	USkeletalMesh* StrongSkeletalMesh = Cast<USkeletalMesh>(UE::Mutable::Private::LoadObject(SkeletalMesh));
 	if (!StrongSkeletalMesh)
 	{
 		return nullptr;
diff --git a/Engine/Plugins/Mutable/Source/CustomizableObjectEditor/Private/MuCOE/Nodes/CustomizableObjectNodeExtensionDataSwitch.cpp b/Engine/Plugins/Mutable/Source/CustomizableObjectEditor/Private/MuCOE/Nodes/CustomizableObjectNodeExtensionDataSwitch.cpp
--- a/Engine/Plugins/Mutable/Source/CustomizableObjectEditor/Private/MuCOE/Nodes/CustomizableObjectNodeExtensionDataSwitch.cpp
+++ b/Engine/Plugins/Mutable/Source/CustomizableObjectEditor/Private/MuCOE/Nodes/CustomizableObjectNodeExtensionDataSwitch.cpp
@@ -26,16 +26,16 @@ FString UCustomizableObjectNodeExtensionDataSwitch::GetPinPrefix() const
 	return FString::Format(TEXT("{0} "), { GetOutputPinName() });
 }
 
-mu::Ptr<mu::NodeExtensionData> UCustomizableObjectNodeExtensionDataSwitch::GenerateMutableNode(FExtensionDataCompilerInterface& InCompilerInterface) const
+UE::Mutable::Private::Ptr<UE::Mutable::Private::NodeExtensionData> UCustomizableObjectNodeExtensionDataSwitch::GenerateMutableNode(FExtensionDataCompilerInterface& InCompilerInterface) const
 {
 	const int32 NumParameters = FollowInputPinArray(*SwitchParameter()).Num();
 
 	const UEdGraphPin* EnumPin = FollowInputPin(*SwitchParameter());
-	mu::Ptr<mu::NodeScalar> SwitchParam = EnumPin ? GenerateMutableSourceFloat(EnumPin, InCompilerInterface.GenerationContext) : nullptr;
+	UE::Mutable::Private::Ptr<UE::Mutable::Private::NodeScalar> SwitchParam = EnumPin ? GenerateMutableSourceFloat(EnumPin, InCompilerInterface.GenerationContext) : nullptr;
 
 	const int32 NumSwitchOptions = GetNumElements();
 
-	mu::Ptr<mu::NodeExtensionDataSwitch> SwitchNode = new mu::NodeExtensionDataSwitch;
+	UE::Mutable::Private::Ptr<UE::Mutable::Private::NodeExtensionDataSwitch> SwitchNode = new UE::Mutable::Private::NodeExtensionDataSwitch;
 	SwitchNode->SetParameter(SwitchParam);
 	SwitchNode->SetOptionCount(NumSwitchOptions);
 
@@ -47,7 +47,7 @@ mu::Ptr<mu::NodeExtensionData> UCustomizableObjectNodeExtensionDataSwitch::Gener
 			{
 				if (const ICustomizableObjectExtensionNode* ExtensionNode = Cast<ICustomizableObjectExtensionNode>(ConnectedPin->GetOwningNode()))
 				{
-					if (mu::Ptr<mu::NodeExtensionData> ExtensionMutableNode = ExtensionNode->GenerateMutableNode(InCompilerInterface))
+					if (UE::Mutable::Private::Ptr<UE::Mutable::Private::NodeExtensionData> ExtensionMutableNode = ExtensionNode->GenerateMutableNode(InCompilerInterface))
 					{
 						SwitchNode->SetOption(OptionIndex, ExtensionMutableNode);
 					}
diff --git a/Engine/Plugins/Mutable/Source/CustomizableObjectEditor/Private/MuCOE/Nodes/CustomizableObjectNodeExtensionDataSwitch.h b/Engine/Plugins/Mutable/Source/CustomizableObjectEditor/Private/MuCOE/Nodes/CustomizableObjectNodeExtensionDataSwitch.h
--- a/Engine/Plugins/Mutable/Source/CustomizableObjectEditor/Private/MuCOE/Nodes/CustomizableObjectNodeExtensionDataSwitch.h
+++ b/Engine/Plugins/Mutable/Source/CustomizableObjectEditor/Private/MuCOE/Nodes/CustomizableObjectNodeExtensionDataSwitch.h
@@ -28,7 +28,7 @@ public:
 	//~End UCustomizableObjectNodeSwitchBase interface
 
 	//~Begin ICustomizableObjectExtensionNode interface
-	UE_API mu::Ptr<mu::NodeExtensionData> GenerateMutableNode(class FExtensionDataCompilerInterface& InCompilerInterface) const override;
+	UE_API UE::Mutable::Private::Ptr<UE::Mutable::Private::NodeExtensionData> GenerateMutableNode(class FExtensionDataCompilerInterface& InCompilerInterface) const override;
 	//~End ICustomizableObjectExtensionNode interface
 };
 
diff --git a/Engine/Plugins/Mutable/Source/CustomizableObjectEditor/Private/MuCOE/Nodes/CustomizableObjectNodeExtensionDataVariation.cpp b/Engine/Plugins/Mutable/Source/CustomizableObjectEditor/Private/MuCOE/Nodes/CustomizableObjectNodeExtensionDataVariation.cpp
--- a/Engine/Plugins/Mutable/Source/CustomizableObjectEditor/Private/MuCOE/Nodes/CustomizableObjectNodeExtensionDataVariation.cpp
+++ b/Engine/Plugins/Mutable/Source/CustomizableObjectEditor/Private/MuCOE/Nodes/CustomizableObjectNodeExtensionDataVariation.cpp
@@ -56,15 +56,15 @@ void UCustomizableObjectNodeExtensionDataVariation::AllocateDefaultPins(UCustomi
 	DefaultVariation->bDefaultValueIsIgnored = true;
 }
 
-mu::Ptr<mu::NodeExtensionData> UCustomizableObjectNodeExtensionDataVariation::GenerateMutableNode(FExtensionDataCompilerInterface& InCompilerInterface) const
+UE::Mutable::Private::Ptr<UE::Mutable::Private::NodeExtensionData> UCustomizableObjectNodeExtensionDataVariation::GenerateMutableNode(FExtensionDataCompilerInterface& InCompilerInterface) const
 {
-	mu::Ptr<mu::NodeExtensionDataVariation> VariationNode = new mu::NodeExtensionDataVariation;
+	UE::Mutable::Private::Ptr<UE::Mutable::Private::NodeExtensionDataVariation> VariationNode = new UE::Mutable::Private::NodeExtensionDataVariation;
 
 	if (const UEdGraphPin* ConnectedPin = FollowInputPin(*GetDefaultPin()))
 	{
 		if (const ICustomizableObjectExtensionNode* ExtensionNode = Cast<ICustomizableObjectExtensionNode>(ConnectedPin->GetOwningNode()))
 		{
-			if (mu::Ptr<mu::NodeExtensionData> DefaultValueExtensionData = ExtensionNode->GenerateMutableNode(InCompilerInterface))
+			if (UE::Mutable::Private::Ptr<UE::Mutable::Private::NodeExtensionData> DefaultValueExtensionData = ExtensionNode->GenerateMutableNode(InCompilerInterface))
 			{
 				VariationNode->SetDefaultValue(DefaultValueExtensionData);
 			}
@@ -85,7 +85,7 @@ mu::Ptr<mu::NodeExtensionData> UCustomizableObjectNodeExtensionDataVariation::Ge
 			{
 				if (const ICustomizableObjectExtensionNode* ExtensionNode = Cast<ICustomizableObjectExtensionNode>(ConnectedPin->GetOwningNode()))
 				{
-					if (mu::Ptr<mu::NodeExtensionData> VariationExtensionData = ExtensionNode->GenerateMutableNode(InCompilerInterface))
+					if (UE::Mutable::Private::Ptr<UE::Mutable::Private::NodeExtensionData> VariationExtensionData = ExtensionNode->GenerateMutableNode(InCompilerInterface))
 					{
 						VariationNode->SetVariationValue(VariationIndex, VariationExtensionData);
 					}
diff --git a/Engine/Plugins/Mutable/Source/CustomizableObjectEditor/Private/MuCOE/Nodes/CustomizableObjectNodeExtensionDataVariation.h b/Engine/Plugins/Mutable/Source/CustomizableObjectEditor/Private/MuCOE/Nodes/CustomizableObjectNodeExtensionDataVariation.h
--- a/Engine/Plugins/Mutable/Source/CustomizableObjectEditor/Private/MuCOE/Nodes/CustomizableObjectNodeExtensionDataVariation.h
+++ b/Engine/Plugins/Mutable/Source/CustomizableObjectEditor/Private/MuCOE/Nodes/CustomizableObjectNodeExtensionDataVariation.h
@@ -42,7 +42,7 @@ public:
 	//~End UCustomizableObjectNode interface
 
 	//~Begin ICustomizableObjectExtensionNode interface
-	UE_API virtual mu::Ptr<mu::NodeExtensionData> GenerateMutableNode(FExtensionDataCompilerInterface& InCompilerInterface) const override;
+	UE_API virtual UE::Mutable::Private::Ptr<UE::Mutable::Private::NodeExtensionData> GenerateMutableNode(FExtensionDataCompilerInterface& InCompilerInterface) const override;
 	//~End ICustomizableObjectExtensionNode interface
 
 	virtual FName GetCategory() const PURE_VIRTUAL(GetCategory(), return FName();)
diff --git a/Engine/Plugins/Mutable/Source/CustomizableObjectEditor/Private/MuCOE/Nodes/CustomizableObjectNodeExternalPinDetails.cpp b/Engine/Plugins/Mutable/Source/CustomizableObjectEditor/Private/MuCOE/Nodes/CustomizableObjectNodeExternalPinDetails.cpp
--- a/Engine/Plugins/Mutable/Source/CustomizableObjectEditor/Private/MuCOE/Nodes/CustomizableObjectNodeExternalPinDetails.cpp
+++ b/Engine/Plugins/Mutable/Source/CustomizableObjectEditor/Private/MuCOE/Nodes/CustomizableObjectNodeExternalPinDetails.cpp
@@ -107,7 +107,7 @@ void FCustomizableObjectNodeExternalPinDetails::ParentObjectSelectionChanged(con
 {
 	const FScopedTransaction Transaction(LOCTEXT("ChangedImportPinCustomizableObjectTransaction", "Changed Import Pin Customizable Object"));
 	Node->Modify();
-	Node->ExternalObject = Cast<UCustomizableObject>(MutablePrivate::LoadObject(AssetData));
+	Node->ExternalObject = Cast<UCustomizableObject>(UE::Mutable::Private::LoadObject(AssetData));
 	Node->SetExternalObjectNodeId(FGuid());
 
 	if (DetailBuilderPtr)
diff --git a/Engine/Plugins/Mutable/Source/CustomizableObjectEditor/Private/MuCOE/Nodes/CustomizableObjectNodeGroupProjectorParameter.cpp b/Engine/Plugins/Mutable/Source/CustomizableObjectEditor/Private/MuCOE/Nodes/CustomizableObjectNodeGroupProjectorParameter.cpp
--- a/Engine/Plugins/Mutable/Source/CustomizableObjectEditor/Private/MuCOE/Nodes/CustomizableObjectNodeGroupProjectorParameter.cpp
+++ b/Engine/Plugins/Mutable/Source/CustomizableObjectEditor/Private/MuCOE/Nodes/CustomizableObjectNodeGroupProjectorParameter.cpp
@@ -87,7 +87,7 @@ TArray<FGroupProjectorParameterImage> UCustomizableObjectNodeGroupProjectorParam
 		uint8* RowData = RowIt.Value();
 		FString PropertyValue(TEXT(""));
 		PropertyTexturePath->ExportText_InContainer(0, PropertyValue, RowData, RowData, nullptr, PPF_None);
-		UTexture2D* Texture = MutablePrivate::LoadObject<UTexture2D>(nullptr, *PropertyValue);
+		UTexture2D* Texture = UE::Mutable::Private::LoadObject<UTexture2D>(nullptr, *PropertyValue);
 
 		if (Texture == nullptr)
 		{
diff --git a/Engine/Plugins/Mutable/Source/CustomizableObjectEditor/Private/MuCOE/Nodes/CustomizableObjectNodeMaterialBase.h b/Engine/Plugins/Mutable/Source/CustomizableObjectEditor/Private/MuCOE/Nodes/CustomizableObjectNodeMaterialBase.h
--- a/Engine/Plugins/Mutable/Source/CustomizableObjectEditor/Private/MuCOE/Nodes/CustomizableObjectNodeMaterialBase.h
+++ b/Engine/Plugins/Mutable/Source/CustomizableObjectEditor/Private/MuCOE/Nodes/CustomizableObjectNodeMaterialBase.h
@@ -44,8 +44,8 @@ inline uint32 GetTypeHash(const FNodeMaterialParameterId& Key)
 }
 
 
-/** Equivalent to mu::NodeSurface but with limitations. Currently only nodes that generate a mu::NodeSurfaceNew inherit this (NodeMaterial and NodeCopyMaterial).
- * Nodes that generate mu::NodeSurfaceEdit (NodeEditMaterial, NodeExtendMaterial), mu::NodeSurfaceSwitch (NodeSwitchMaterial)... are excluded.
+/** Equivalent to UE::Mutable::Private::NodeSurface but with limitations. Currently only nodes that generate a UE::Mutable::Private::NodeSurfaceNew inherit this (NodeMaterial and NodeCopyMaterial).
+ * Nodes that generate UE::Mutable::Private::NodeSurfaceEdit (NodeEditMaterial, NodeExtendMaterial), UE::Mutable::Private::NodeSurfaceSwitch (NodeSwitchMaterial)... are excluded.
  *
  * WARNING: All methods in this class should be PURE_VIRTUAL!
  */
diff --git a/Engine/Plugins/Mutable/Source/CustomizableObjectEditor/Private/MuCOE/Nodes/CustomizableObjectNodeMeshMorphDetails.cpp b/Engine/Plugins/Mutable/Source/CustomizableObjectEditor/Private/MuCOE/Nodes/CustomizableObjectNodeMeshMorphDetails.cpp
--- a/Engine/Plugins/Mutable/Source/CustomizableObjectEditor/Private/MuCOE/Nodes/CustomizableObjectNodeMeshMorphDetails.cpp
+++ b/Engine/Plugins/Mutable/Source/CustomizableObjectEditor/Private/MuCOE/Nodes/CustomizableObjectNodeMeshMorphDetails.cpp
@@ -56,7 +56,7 @@ void FCustomizableObjectNodeMeshMorphDetails::CustomizeDetails( IDetailLayoutBui
 			USkeletalMesh* SkeletalMesh = nullptr;
 			if (const ICustomizableObjectNodeMeshInterface* TypedNodeSkeletalMesh = Cast<ICustomizableObjectNodeMeshInterface>(BaseSourcePin->GetOwningNode()) )
 			{
-				SkeletalMesh = Cast<USkeletalMesh>(MutablePrivate::LoadObject(TypedNodeSkeletalMesh->GetMesh()));
+				SkeletalMesh = Cast<USkeletalMesh>(UE::Mutable::Private::LoadObject(TypedNodeSkeletalMesh->GetMesh()));
 			}
 			else if (const UCustomizableObjectNodeTable* TypedNodeTable = Cast<UCustomizableObjectNodeTable>(BaseSourcePin->GetOwningNode()))
 			{
diff --git a/Engine/Plugins/Mutable/Source/CustomizableObjectEditor/Private/MuCOE/Nodes/CustomizableObjectNodeMeshMorphStackApplication.cpp b/Engine/Plugins/Mutable/Source/CustomizableObjectEditor/Private/MuCOE/Nodes/CustomizableObjectNodeMeshMorphStackApplication.cpp
--- a/Engine/Plugins/Mutable/Source/CustomizableObjectEditor/Private/MuCOE/Nodes/CustomizableObjectNodeMeshMorphStackApplication.cpp
+++ b/Engine/Plugins/Mutable/Source/CustomizableObjectEditor/Private/MuCOE/Nodes/CustomizableObjectNodeMeshMorphStackApplication.cpp
@@ -111,7 +111,7 @@ TArray<FString> UCustomizableObjectNodeMeshMorphStackApplication::GetMorphList()
 
 		if (const ICustomizableObjectNodeMeshInterface* MeshNode = Cast<ICustomizableObjectNodeMeshInterface>(MeshNodePin->GetOwningNode()))
 		{
-			SkeletalMesh = Cast<USkeletalMesh>(MutablePrivate::LoadObject(MeshNode->GetMesh()));
+			SkeletalMesh = Cast<USkeletalMesh>(UE::Mutable::Private::LoadObject(MeshNode->GetMesh()));
 		}
 		else if (const UCustomizableObjectNodeTable* TableNode = Cast< UCustomizableObjectNodeTable >(MeshNodePin->GetOwningNode()))
 		{
diff --git a/Engine/Plugins/Mutable/Source/CustomizableObjectEditor/Private/MuCOE/Nodes/CustomizableObjectNodeMeshMorphStackDefinition.cpp b/Engine/Plugins/Mutable/Source/CustomizableObjectEditor/Private/MuCOE/Nodes/CustomizableObjectNodeMeshMorphStackDefinition.cpp
--- a/Engine/Plugins/Mutable/Source/CustomizableObjectEditor/Private/MuCOE/Nodes/CustomizableObjectNodeMeshMorphStackDefinition.cpp
+++ b/Engine/Plugins/Mutable/Source/CustomizableObjectEditor/Private/MuCOE/Nodes/CustomizableObjectNodeMeshMorphStackDefinition.cpp
@@ -101,7 +101,7 @@ void UCustomizableObjectNodeMeshMorphStackDefinition::UpdateMorphList()
 
 		if (const ICustomizableObjectNodeMeshInterface* MeshNode = Cast<ICustomizableObjectNodeMeshInterface>(MeshNodePin->GetOwningNode()))
 		{
-			SkeletalMesh = Cast<USkeletalMesh>(MutablePrivate::LoadObject(MeshNode->GetMesh()));
+			SkeletalMesh = Cast<USkeletalMesh>(UE::Mutable::Private::LoadObject(MeshNode->GetMesh()));
 		}
 		else if (const UCustomizableObjectNodeTable* TableNode = Cast<UCustomizableObjectNodeTable>(MeshNodePin->GetOwningNode()))
 		{
@@ -194,7 +194,7 @@ bool UCustomizableObjectNodeMeshMorphStackDefinition::IsNodeOutDatedAndNeedsRefr
 
 		if (const UCustomizableObjectNodeSkeletalMesh* SkeletalMeshNode = Cast< UCustomizableObjectNodeSkeletalMesh >(MeshNodePin->GetOwningNode()))
 		{
-			SkeletalMesh = MutablePrivate::LoadObject(SkeletalMeshNode->SkeletalMesh);
+			SkeletalMesh = UE::Mutable::Private::LoadObject(SkeletalMeshNode->SkeletalMesh);
 		}
 		else if (const UCustomizableObjectNodeTable* TableNode = Cast< UCustomizableObjectNodeTable >(MeshNodePin->GetOwningNode()))
 		{
diff --git a/Engine/Plugins/Mutable/Source/CustomizableObjectEditor/Private/MuCOE/Nodes/CustomizableObjectNodeMeshReshapeSelectionDetails.cpp b/Engine/Plugins/Mutable/Source/CustomizableObjectEditor/Private/MuCOE/Nodes/CustomizableObjectNodeMeshReshapeSelectionDetails.cpp
--- a/Engine/Plugins/Mutable/Source/CustomizableObjectEditor/Private/MuCOE/Nodes/CustomizableObjectNodeMeshReshapeSelectionDetails.cpp
+++ b/Engine/Plugins/Mutable/Source/CustomizableObjectEditor/Private/MuCOE/Nodes/CustomizableObjectNodeMeshReshapeSelectionDetails.cpp
@@ -90,7 +90,7 @@ USkeletalMesh* FMeshReshapeBonesReferenceCustomization::GetSkeletalMesh(TSharedP
 				{
 					if (const UCustomizableObjectNodeSkeletalMesh* SkeletalMeshNode = Cast<UCustomizableObjectNodeSkeletalMesh>(SourceMeshPin->GetOwningNode()))
 					{
-						return MutablePrivate::LoadObject(SkeletalMeshNode->SkeletalMesh);
+						return UE::Mutable::Private::LoadObject(SkeletalMeshNode->SkeletalMesh);
 					}
 					else if (UCustomizableObjectNodeTable* TableNode = Cast<UCustomizableObjectNodeTable>(SourceMeshPin->GetOwningNode()))
 					{
diff --git a/Engine/Plugins/Mutable/Source/CustomizableObjectEditor/Private/MuCOE/Nodes/CustomizableObjectNodeModifierMorphMeshSectionDetails.cpp b/Engine/Plugins/Mutable/Source/CustomizableObjectEditor/Private/MuCOE/Nodes/CustomizableObjectNodeModifierMorphMeshSectionDetails.cpp
--- a/Engine/Plugins/Mutable/Source/CustomizableObjectEditor/Private/MuCOE/Nodes/CustomizableObjectNodeModifierMorphMeshSectionDetails.cpp
+++ b/Engine/Plugins/Mutable/Source/CustomizableObjectEditor/Private/MuCOE/Nodes/CustomizableObjectNodeModifierMorphMeshSectionDetails.cpp
@@ -280,7 +280,7 @@ void FCustomizableObjectNodeModifierMorphMeshSectionDetails::AddMorphsFromNode(
 				UCustomizableObjectNodeSkeletalMesh* SkeletalNode = Cast<UCustomizableObjectNodeSkeletalMesh>(SourceMeshPin->GetOwningNode());
 				if (SkeletalNode)
 				{
-					SkeletalMesh = MutablePrivate::LoadObject(SkeletalNode->SkeletalMesh);
+					SkeletalMesh = UE::Mutable::Private::LoadObject(SkeletalNode->SkeletalMesh);
 				}
 			}
 		}
@@ -296,7 +296,7 @@ void FCustomizableObjectNodeModifierMorphMeshSectionDetails::AddMorphsFromNode(
 				UCustomizableObjectNodeSkeletalMesh* SkeletalNode = Cast<UCustomizableObjectNodeSkeletalMesh>(SourceMeshPin->GetOwningNode());
 				if (SkeletalNode)
 				{
-					SkeletalMesh = MutablePrivate::LoadObject(SkeletalNode->SkeletalMesh);
+					SkeletalMesh = UE::Mutable::Private::LoadObject(SkeletalNode->SkeletalMesh);
 				}
 			}
 		}
diff --git a/Engine/Plugins/Mutable/Source/CustomizableObjectEditor/Private/MuCOE/Nodes/CustomizableObjectNodeObjectDetails.cpp b/Engine/Plugins/Mutable/Source/CustomizableObjectEditor/Private/MuCOE/Nodes/CustomizableObjectNodeObjectDetails.cpp
--- a/Engine/Plugins/Mutable/Source/CustomizableObjectEditor/Private/MuCOE/Nodes/CustomizableObjectNodeObjectDetails.cpp
+++ b/Engine/Plugins/Mutable/Source/CustomizableObjectEditor/Private/MuCOE/Nodes/CustomizableObjectNodeObjectDetails.cpp
@@ -311,7 +311,7 @@ void FCustomizableObjectNodeObjectDetails::ParentObjectSelectionChanged(const FA
 {
 	if (BaseObjectNode.IsValid())
 	{
-		UCustomizableObject* Parent = Cast<UCustomizableObject>(MutablePrivate::LoadObject(AssetData));
+		UCustomizableObject* Parent = Cast<UCustomizableObject>(UE::Mutable::Private::LoadObject(AssetData));
 		BaseObjectNode->SetParentObject(Parent);
 
 		// If set the parent to nullt, invalidate also the reference GUID
diff --git a/Engine/Plugins/Mutable/Source/CustomizableObjectEditor/Private/MuCOE/Nodes/CustomizableObjectNodeSkeletalMesh.cpp b/Engine/Plugins/Mutable/Source/CustomizableObjectEditor/Private/MuCOE/Nodes/CustomizableObjectNodeSkeletalMesh.cpp
--- a/Engine/Plugins/Mutable/Source/CustomizableObjectEditor/Private/MuCOE/Nodes/CustomizableObjectNodeSkeletalMesh.cpp
+++ b/Engine/Plugins/Mutable/Source/CustomizableObjectEditor/Private/MuCOE/Nodes/CustomizableObjectNodeSkeletalMesh.cpp
@@ -36,7 +36,7 @@ static const TCHAR* SKELETAL_MESH_PIN_NAME = TEXT("Skeletal Mesh");
 
 void UCustomizableObjectNodeSkeletalMesh::LoadObjects()
 {
-	if (USkeletalMesh* LoadedMesh = MutablePrivate::LoadObject(SkeletalMesh))
+	if (USkeletalMesh* LoadedMesh = UE::Mutable::Private::LoadObject(SkeletalMesh))
 	{
 		ConditionalPostLoadReference(*LoadedMesh);
 
@@ -74,7 +74,7 @@ void UCustomizableObjectNodeSkeletalMesh::AllocateDefaultPins(UCustomizableObjec
 	const UEdGraphSchema_CustomizableObject* Schema = GetDefault<UEdGraphSchema_CustomizableObject>();
 
 	// Force the loading of the skeletal mesh so we can later access the internal WeakObjectPtr (may be null if no SKM is set)
-	MutablePrivate::LoadObject(SkeletalMesh);
+	UE::Mutable::Private::LoadObject(SkeletalMesh);
 	
 	if (!SkeletalMesh)
 	{
@@ -790,7 +790,7 @@ void SGraphNodeSkeletalMesh::Construct(const FArguments& InArgs, UEdGraphNode* I
 
 	// Thumbnail
 	AssetThumbnailPool = MakeShareable(new FAssetThumbnailPool(32));
-	USkeletalMesh* SkeletalMesh = MutablePrivate::LoadObject(NodeSkeletalMesh->SkeletalMesh);
+	USkeletalMesh* SkeletalMesh = UE::Mutable::Private::LoadObject(NodeSkeletalMesh->SkeletalMesh);
 
 	const FAssetRegistryModule& AssetRegistryModule = FModuleManager::LoadModuleChecked<FAssetRegistryModule>(AssetRegistryConstants::ModuleName);
 
diff --git a/Engine/Plugins/Mutable/Source/CustomizableObjectEditor/Private/MuCOE/Nodes/CustomizableObjectNodeTable.cpp b/Engine/Plugins/Mutable/Source/CustomizableObjectEditor/Private/MuCOE/Nodes/CustomizableObjectNodeTable.cpp
--- a/Engine/Plugins/Mutable/Source/CustomizableObjectEditor/Private/MuCOE/Nodes/CustomizableObjectNodeTable.cpp
+++ b/Engine/Plugins/Mutable/Source/CustomizableObjectEditor/Private/MuCOE/Nodes/CustomizableObjectNodeTable.cpp
@@ -229,13 +229,13 @@ void UCustomizableObjectNodeTable::BackwardsCompatibleFixup(int32 CustomizableOb
 
 	if (CustomizableObjectCustomVersion == FCustomizableObjectCustomVersion::SetDisplayNamePropertyAsPinNameOfTableNodes)
 	{
-		TObjectPtr<UScriptStruct> LoadedStructure = MutablePrivate::LoadObject(Structure);
+		TObjectPtr<UScriptStruct> LoadedStructure = UE::Mutable::Private::LoadObject(Structure);
 		if (LoadedStructure)
 		{
 			ConditionalPostLoadReference(*LoadedStructure);
 		}
 
-		TObjectPtr<UDataTable> LoadedTable = MutablePrivate::LoadObject(Table);
+		TObjectPtr<UDataTable> LoadedTable = UE::Mutable::Private::LoadObject(Table);
 		if (LoadedTable)
 		{
 			ConditionalPostLoadReference(*LoadedTable);
@@ -292,13 +292,13 @@ void UCustomizableObjectNodeTable::BackwardsCompatibleFixup(int32 CustomizableOb
 	{
 		if (bDisableCheckedRows_DEPRECATED)
 		{
-			TObjectPtr<UScriptStruct> LoadedStructure = MutablePrivate::LoadObject(Structure);
+			TObjectPtr<UScriptStruct> LoadedStructure = UE::Mutable::Private::LoadObject(Structure);
 			if (LoadedStructure)
 			{
 				ConditionalPostLoadReference(*LoadedStructure);
 			}
 
-			TObjectPtr<UDataTable> LoadedTable = MutablePrivate::LoadObject(Table);
+			TObjectPtr<UDataTable> LoadedTable = UE::Mutable::Private::LoadObject(Table);
 			if (LoadedTable)
 			{
 				ConditionalPostLoadReference(*LoadedTable);
@@ -338,7 +338,7 @@ void UCustomizableObjectNodeTable::PreEditChange(FProperty* PropertyAboutToChang
 {
 	Super::PreEditChange(PropertyAboutToChange);
 
-	TObjectPtr<UDataTable> LoadedTable = MutablePrivate::LoadObject(Table);
+	TObjectPtr<UDataTable> LoadedTable = UE::Mutable::Private::LoadObject(Table);
 
 	if (PropertyAboutToChange && PropertyAboutToChange->GetFName() == GET_MEMBER_NAME_CHECKED(UCustomizableObjectNodeTable, Table) &&
 		LoadedTable)
@@ -363,12 +363,12 @@ FText UCustomizableObjectNodeTable::GetNodeTitle(ENodeTitleType::Type TitleType)
 	}
 	else
 	{
-		if (TObjectPtr<UDataTable> LoadedTable = MutablePrivate::LoadObject(Table))
+		if (TObjectPtr<UDataTable> LoadedTable = UE::Mutable::Private::LoadObject(Table))
 		{
 			Args.Add(TEXT("TableName"), FText::FromString(LoadedTable->GetName()));
 			return FText::Format(LOCTEXT("TableNode_Title_DataTable", "{ParamName}\n{TableName} - Data Table"), Args);
 		}
-		else if (TObjectPtr<UScriptStruct> LoadedStructure = MutablePrivate::LoadObject(Structure))
+		else if (TObjectPtr<UScriptStruct> LoadedStructure = UE::Mutable::Private::LoadObject(Structure))
 		{
 			Args.Add(TEXT("StructureName"), FText::FromString(LoadedStructure->GetName()));
 			return FText::Format(LOCTEXT("TableNode_Title_ScriptedStruct", "{ParamName}\n{StructureName} - Script Struct"), Args);
@@ -416,7 +416,7 @@ void UCustomizableObjectNodeTable::PostEditChangeProperty(FPropertyChangedEvent&
 		}
 		else if (PropertyThatChanged->GetName() == GET_MEMBER_NAME_CHECKED(UCustomizableObjectNodeTable, TableDataGatheringMode))
 		{
-			if (TObjectPtr<UDataTable> LoadedTable = MutablePrivate::LoadObject(Table))
+			if (TObjectPtr<UDataTable> LoadedTable = UE::Mutable::Private::LoadObject(Table))
 			{
 				LoadedTable->OnDataTableChanged().Remove(OnTableChangedDelegateHandle);
 			}
@@ -509,7 +509,7 @@ void UCustomizableObjectNodeTable::AllocateDefaultPins(UCustomizableObjectNodeRe
 
 				if (CellData)
 				{
-					Object = MutablePrivate::LoadObject(SoftObjectProperty->GetPropertyValue(CellData));
+					Object = UE::Mutable::Private::LoadObject(SoftObjectProperty->GetPropertyValue(CellData));
 				}
 
 				if (Object)
@@ -790,7 +790,7 @@ bool UCustomizableObjectNodeTable::IsNodeOutDatedAndNeedsRefresh()
 
 				if (CellData)
 				{
-					Object = MutablePrivate::LoadObject(SoftObjectProperty->GetPropertyValue(CellData));
+					Object = UE::Mutable::Private::LoadObject(SoftObjectProperty->GetPropertyValue(CellData));
 				}
 
 				if (!Object)
@@ -1334,14 +1334,14 @@ const UScriptStruct* UCustomizableObjectNodeTable::GetTableNodeStruct() const
 
 	if (TableDataGatheringMode == ETableDataGatheringSource::ETDGM_AssetRegistry)
 	{
-		if (TObjectPtr<UScriptStruct> LoadedStructure = MutablePrivate::LoadObject(Structure))
+		if (TObjectPtr<UScriptStruct> LoadedStructure = UE::Mutable::Private::LoadObject(Structure))
 		{
 			TableStruct = LoadedStructure;
 		}
 	}
 	else
 	{
-		if (TObjectPtr<UDataTable> LoadedTable = MutablePrivate::LoadObject(Table))
+		if (TObjectPtr<UDataTable> LoadedTable = UE::Mutable::Private::LoadObject(Table))
 		{
 			TableStruct = LoadedTable->GetRowStruct();
 		}
@@ -1460,7 +1460,7 @@ TArray<FAssetData> UCustomizableObjectNodeTable::GetParentTables() const
 		Filter.ClassPaths.Add(FTopLevelAssetPath(UDataTable::StaticClass()));
 		Filter.bRecursivePaths = true;
 
-		if (TObjectPtr<UScriptStruct> LoadedStructure = MutablePrivate::LoadObject(Structure))
+		if (TObjectPtr<UScriptStruct> LoadedStructure = UE::Mutable::Private::LoadObject(Structure))
 		{
 			TArray<FName> ReferencedTables;
 			AssetRegistry.Get()->GetReferencers(LoadedStructure.GetPackage().GetFName(), ReferencedTables, UE::AssetRegistry::EDependencyCategory::Package, UE::AssetRegistry::EDependencyQuery::NoRequirements);
diff --git a/Engine/Plugins/Mutable/Source/CustomizableObjectEditor/Private/MuCOE/Nodes/CustomizableObjectNodeTable.h b/Engine/Plugins/Mutable/Source/CustomizableObjectEditor/Private/MuCOE/Nodes/CustomizableObjectNodeTable.h
--- a/Engine/Plugins/Mutable/Source/CustomizableObjectEditor/Private/MuCOE/Nodes/CustomizableObjectNodeTable.h
+++ b/Engine/Plugins/Mutable/Source/CustomizableObjectEditor/Private/MuCOE/Nodes/CustomizableObjectNodeTable.h
@@ -360,7 +360,7 @@ public:
 
 					if (CellData)
 					{
-						Object = MutablePrivate::LoadObject(SoftObjectProperty->GetPropertyValue(CellData));
+						Object = UE::Mutable::Private::LoadObject(SoftObjectProperty->GetPropertyValue(CellData));
 					}
 
 					if (Object)
diff --git a/Engine/Plugins/Mutable/Source/CustomizableObjectEditor/Private/MuCOE/SCustomizableObjectEditorAdvancedPreviewSettings.cpp b/Engine/Plugins/Mutable/Source/CustomizableObjectEditor/Private/MuCOE/SCustomizableObjectEditorAdvancedPreviewSettings.cpp
--- a/Engine/Plugins/Mutable/Source/CustomizableObjectEditor/Private/MuCOE/SCustomizableObjectEditorAdvancedPreviewSettings.cpp
+++ b/Engine/Plugins/Mutable/Source/CustomizableObjectEditor/Private/MuCOE/SCustomizableObjectEditorAdvancedPreviewSettings.cpp
@@ -482,7 +482,7 @@ void SCustomizableObjectEditorAdvancedPreviewSettings::LoadProfileEnvironment()
 		if (!DefaultSettings->Profiles[i].EnvironmentCubeMapPath.IsEmpty())
 		{
 			UObject* LoadedObject = nullptr;
-			LoadedObject = MutablePrivate::LoadObject<UObject>(nullptr, *DefaultSettings->Profiles[i].EnvironmentCubeMapPath);
+			LoadedObject = UE::Mutable::Private::LoadObject<UObject>(nullptr, *DefaultSettings->Profiles[i].EnvironmentCubeMapPath);
 			while (UObjectRedirector* Redirector = Cast<UObjectRedirector>(LoadedObject))
 			{
 				LoadedObject = Redirector->DestinationObject;
diff --git a/Engine/Plugins/Mutable/Source/CustomizableObjectEditor/Private/MuCOE/SCustomizableObjectEditorViewport.cpp b/Engine/Plugins/Mutable/Source/CustomizableObjectEditor/Private/MuCOE/SCustomizableObjectEditorViewport.cpp
--- a/Engine/Plugins/Mutable/Source/CustomizableObjectEditor/Private/MuCOE/SCustomizableObjectEditorViewport.cpp
+++ b/Engine/Plugins/Mutable/Source/CustomizableObjectEditor/Private/MuCOE/SCustomizableObjectEditorViewport.cpp
@@ -1019,7 +1019,7 @@ FReply SCustomizableObjectEditorViewportTabBody::OnDrop(const FGeometry& MyGeome
 		if (DragDropOp->GetAssets().Num())
 		{
 			// This cast also includes UPoseAsset assets.
-			UAnimationAsset* AnimationAsset = Cast<UAnimationAsset>(MutablePrivate::LoadObject(DragDropOp->GetAssets()[0]));
+			UAnimationAsset* AnimationAsset = Cast<UAnimationAsset>(UE::Mutable::Private::LoadObject(DragDropOp->GetAssets()[0]));
 			if (AnimationAsset)
 			{
 				FObjectEditorUtils::SetPropertyValue(WeakEditor.Pin()->GetCustomSettings(), GET_MEMBER_NAME_CHECKED(UCustomSettings, Animation), AnimationAsset);
diff --git a/Engine/Plugins/Mutable/Source/CustomizableObjectEditor/Private/MuCOE/SCustomizableObjectLayoutGrid.cpp b/Engine/Plugins/Mutable/Source/CustomizableObjectEditor/Private/MuCOE/SCustomizableObjectLayoutGrid.cpp
--- a/Engine/Plugins/Mutable/Source/CustomizableObjectEditor/Private/MuCOE/SCustomizableObjectLayoutGrid.cpp
+++ b/Engine/Plugins/Mutable/Source/CustomizableObjectEditor/Private/MuCOE/SCustomizableObjectLayoutGrid.cpp
@@ -556,7 +556,7 @@ void SCustomizableObjectLayoutGrid::CloseMenu()
 
 void SCustomizableObjectLayoutGrid::OnMaskAssetSelected(const FAssetData& AssetData)
 {
-	UTexture2D* Mask = Cast<UTexture2D>(MutablePrivate::LoadObject(AssetData));
+	UTexture2D* Mask = Cast<UTexture2D>(UE::Mutable::Private::LoadObject(AssetData));
 	if (SelectedBlocks.Num())
 	{
 		OnSetBlockMask.ExecuteIfBound(Mask);
diff --git a/Engine/Plugins/Mutable/Source/CustomizableObjectEditor/Private/MuCOE/SCustomizableObjectNodeSkeletalMeshRTMorphSelector.cpp b/Engine/Plugins/Mutable/Source/CustomizableObjectEditor/Private/MuCOE/SCustomizableObjectNodeSkeletalMeshRTMorphSelector.cpp
--- a/Engine/Plugins/Mutable/Source/CustomizableObjectEditor/Private/MuCOE/SCustomizableObjectNodeSkeletalMeshRTMorphSelector.cpp
+++ b/Engine/Plugins/Mutable/Source/CustomizableObjectEditor/Private/MuCOE/SCustomizableObjectNodeSkeletalMeshRTMorphSelector.cpp
@@ -189,7 +189,7 @@ void SCustomizableObjectNodeSkeletalMeshRTMorphSelector::GenerateMorphsInfoList(
 {
     Items.Reset();
 
-    const USkeletalMesh* SkeletalMesh = MutablePrivate::LoadObject(Node->SkeletalMesh);
+    const USkeletalMesh* SkeletalMesh = UE::Mutable::Private::LoadObject(Node->SkeletalMesh);
     if (!SkeletalMesh)
     {
         return;
