[Mutable] Change namespace name

diff --git a/Engine/Plugins/Mutable/Source/CustomizableObjectEditor/Private/MuCOE/SMutableCodeViewer.cpp b/Engine/Plugins/Mutable/Source/CustomizableObjectEditor/Private/MuCOE/SMutableCodeViewer.cpp
--- a/Engine/Plugins/Mutable/Source/CustomizableObjectEditor/Private/MuCOE/SMutableCodeViewer.cpp
+++ b/Engine/Plugins/Mutable/Source/CustomizableObjectEditor/Private/MuCOE/SMutableCodeViewer.cpp
@@ -44,8 +44,8 @@ class FExtender;
 class FReferenceCollector;
 class FUICommandList;
 class SWidget;
-namespace mu { struct FProjector; }
-namespace mu { struct FShape; }
+namespace UE::Mutable::Private { struct FProjector; }
+namespace UE::Mutable::Private { struct FShape; }
 struct FGeometry;
 struct FSlateBrush;
 
@@ -282,7 +282,7 @@ void SMutableCodeViewer::ClearSelectedTreeRow() const
 }
 
 
-void SMutableCodeViewer::Construct(const FArguments& InArgs, UCustomizableObject& InObject, UCustomizableObjectInstance& InInstance, const TSharedPtr<mu::FModel>& InMutableModel)
+void SMutableCodeViewer::Construct(const FArguments& InArgs, UCustomizableObject& InObject, UCustomizableObjectInstance& InInstance, const TSharedPtr<UE::Mutable::Private::FModel>& InMutableModel)
 {
 	// Min width allowed for the column. Needed to avoid having issues with the constants space being to small
 	// and then getting too tall on the y axis crashing the UI drawer.
@@ -821,7 +821,7 @@ void SMutableCodeViewer::CacheOperationsMatchingStringPattern()
 		// to start searching
 		SearchPayload.SetupRootBatch(RootNodeAddresses);
 
-		const mu::FProgram& Program = MutableModel->GetPrivate()->Program;
+		const UE::Mutable::Private::FProgram& Program = MutableModel->GetPrivate()->Program;
 		GetOperationsMatchingStringPattern(SearchString,bIsSearchStringRegularExpression,SearchPayload, Program);
 	
 		// Dump the located resources array onto the navigation array
@@ -845,7 +845,7 @@ void SMutableCodeViewer::CacheOperationsMatchingStringPattern()
 }
 
 
-void SMutableCodeViewer::GetOperationsMatchingStringPattern(const FString& InStringPattern,const bool bIsRegularExpression ,FElementsSearchCache& SearchPayload,const mu::FProgram& InProgram)
+void SMutableCodeViewer::GetOperationsMatchingStringPattern(const FString& InStringPattern,const bool bIsRegularExpression ,FElementsSearchCache& SearchPayload,const UE::Mutable::Private::FProgram& InProgram)
 {
 	// next batch of addresses to be explored 
 	TArray<FItemCacheKey> NextBatchAddressesData;
@@ -869,7 +869,7 @@ void SMutableCodeViewer::GetOperationsMatchingStringPattern(const FString& InStr
 		}
 		
 		// Get one of the previous run "children" and treat as a parent to get it's children and process them
-		const mu::OP::ADDRESS& ParentAddress = SearchPayload.BatchData[ParentIndex].Child;
+		const UE::Mutable::Private::OP::ADDRESS& ParentAddress = SearchPayload.BatchData[ParentIndex].Child;
 		
 		if (bMatchesPattern)
 		{
@@ -938,7 +938,7 @@ void SMutableCodeViewer::GenerateNavigationDropdownElements()
 		// Get the type as a string to be able to print it on the UI
 		const FText OperationTypeName = FText::FromString(ModelOperationTypeNames[OperationTypeIndex]);
 
-		const mu::EOpType RepresentedType = ModelOperationTypes[OperationTypeIndex].Key;
+		const UE::Mutable::Private::EOpType RepresentedType = ModelOperationTypes[OperationTypeIndex].Key;
 		const uint32 OperationTypeInstancesCount = ModelOperationTypes[OperationTypeIndex].Value;
 		
 		// Get the Color to be used on the text that will represent the operation on the dropdown
@@ -953,7 +953,7 @@ void SMutableCodeViewer::GenerateNavigationDropdownElements()
 	{
 		const FText EntryName = FText::FromString("NONE");
 		const FSlateColor EntryColor = ColorPerComputationalCost[StaticCast<uint8>(EOperationComputationalCost::Standard)];
-		NoneOperationEntry = MakeShared<FMutableOperationElement>(mu::EOpType::NONE,EntryName,0,EntryColor);
+		NoneOperationEntry = MakeShared<FMutableOperationElement>(UE::Mutable::Private::EOpType::NONE,EntryName,0,EntryColor);
 
 		// @warn While not visible this element must be part of the collection for the ComboBox to be able to work
 		// properly
@@ -964,7 +964,7 @@ void SMutableCodeViewer::GenerateNavigationDropdownElements()
 	{
 		const FText EntryName = FText::FromString("Selected Constant");	
 		const FSlateColor EntryColor = FSlateColor(FLinearColor(0.35f ,0.35f,1.0f,1));
-		ConstantBasedNavigationEntry = MakeShared<FMutableOperationElement>(mu::EOpType::NONE,EntryName,0,EntryColor);
+		ConstantBasedNavigationEntry = MakeShared<FMutableOperationElement>(UE::Mutable::Private::EOpType::NONE,EntryName,0,EntryColor);
 		
 		// @warn While not visible this element must be part of the collection for the ComboBox to be able to work
 		// properly
@@ -998,12 +998,12 @@ void SMutableCodeViewer::OnNavigationSelectedOperationChanged(
 	check (NewSelectedElement.IsValid());
 
 	// Cache the currently selected operation set on the UI by the user
-	const mu::EOpType NewOperationType = NewSelectedElement->OperationType;
+	const UE::Mutable::Private::EOpType NewOperationType = NewSelectedElement->OperationType;
 	OperationTypeToSearch = NewOperationType;
 	CurrentlySelectedOperationTypeElement = NewSelectedElement;
 	
 	// Only do the internal work if the type is one that makes sense searching
-	if (OperationTypeToSearch != mu::EOpType::NONE)
+	if (OperationTypeToSearch != UE::Mutable::Private::EOpType::NONE)
 	{
 		// Locate all operations on the mutable operations tree (not the visual one) that do share the same operation type
 		// as the one selected. This will fill the array with the elements we should be looking for during the navigation operation
@@ -1023,11 +1023,11 @@ void SMutableCodeViewer::OnNavigationSelectedOperationChanged(
 void SMutableCodeViewer::GenerateNavigationOpTypeStrings()
 {
 	// Grab only the names from the operation types located during the caching of operation types of the model
-	for (const TTuple<mu::EOpType, uint32>& LocatedOperationType : ModelOperationTypes)
+	for (const TTuple<UE::Mutable::Private::EOpType, uint32>& LocatedOperationType : ModelOperationTypes)
 	{
 		// Find the name of the Operation type
 		const uint16 OperationIndex = static_cast<uint16>(LocatedOperationType.Key);
-		const TCHAR* OpName = mu::s_opNames[OperationIndex];
+		const TCHAR* OpName = UE::Mutable::Private::s_opNames[OperationIndex];
 
 		// Remove trailing whitespaces adding noise and messing up concatenations with other strings
 		FString OperationNameString{OpName};
@@ -1045,7 +1045,7 @@ void SMutableCodeViewer::OnSelectedOperationTypeFromTree()
 	
 	const TSharedPtr<FMutableCodeTreeElement> ReferenceOperationElement = TreeView->GetSelectedItems()[0];
 	
-	const mu::EOpType OperationType =
+	const UE::Mutable::Private::EOpType OperationType =
 		MutableModel->GetPrivate()->Program.GetOpType(ReferenceOperationElement->MutableOperation);
 
 	// Find the operation type directly in our array of operation elements (from the drop down)
@@ -1090,7 +1090,7 @@ void SMutableCodeViewer::CacheAddressesOfOperationsOfType()
 	SearchPayload.SetupRootBatch(RootNodeAddresses);
 	
 	// Main update procedure run for the targeted state and the targeted parameter values
-	const mu::FProgram& Program = MutableModel->GetPrivate()->Program;
+	const UE::Mutable::Private::FProgram& Program = MutableModel->GetPrivate()->Program;
 	GetOperationsOfType(OperationTypeToSearch,SearchPayload, Program);
 	
 	if (!SearchPayload.FoundElements.IsEmpty())
@@ -1104,9 +1104,9 @@ void SMutableCodeViewer::CacheAddressesOfOperationsOfType()
 	}
 }
 
-void SMutableCodeViewer::GetOperationsOfType(const mu::EOpType& TargetOperationType,
+void SMutableCodeViewer::GetOperationsOfType(const UE::Mutable::Private::EOpType& TargetOperationType,
                                              FElementsSearchCache& InSearchPayload,
-                                             const mu::FProgram& InProgram)
+                                             const UE::Mutable::Private::FProgram& InProgram)
 {
 	// next batch of addresses to be explored 
 	TArray<FItemCacheKey> NextBatchAddressesData;
@@ -1114,7 +1114,7 @@ void SMutableCodeViewer::GetOperationsOfType(const mu::EOpType& TargetOperationT
 	for	(int32 ParentIndex = 0 ; ParentIndex < InSearchPayload.BatchData.Num(); ParentIndex++)
 	{
 		// Get one of the previous run "children" and treat as a parent to get it's children and process them
-		const mu::OP::ADDRESS& CurrentAddress = InSearchPayload.BatchData[ParentIndex].Child;
+		const UE::Mutable::Private::OP::ADDRESS& CurrentAddress = InSearchPayload.BatchData[ParentIndex].Child;
 		
 		// Cache if same data type and we share the same address (means this op is pointing at the provided resource)
 		// It will cache duplicated entries
@@ -1146,12 +1146,12 @@ void SMutableCodeViewer::CacheOperationTypesPresentOnModel()
 
 	// Initialize NodeOperationTypes with empty TPair for each possible mutable operation type
 	{
-		constexpr uint32 OperationTypesCount = static_cast<uint16>(mu::EOpType::COUNT);
+		constexpr uint32 OperationTypesCount = static_cast<uint16>(UE::Mutable::Private::EOpType::COUNT);
 		ModelOperationTypes.Empty(OperationTypesCount);
 		for (uint32 Index = 0; Index < OperationTypesCount ; Index++)
 		{
-			mu::EOpType TargetType = StaticCast<mu::EOpType>(Index);
-			ModelOperationTypes.Add(TPair<mu::EOpType,uint32>{TargetType,0});
+			UE::Mutable::Private::EOpType TargetType = StaticCast<UE::Mutable::Private::EOpType>(Index);
+			ModelOperationTypes.Add(TPair<UE::Mutable::Private::EOpType,uint32>{TargetType,0});
 		}
 	}
 	
@@ -1159,17 +1159,17 @@ void SMutableCodeViewer::CacheOperationTypesPresentOnModel()
 	// there are
 	{
 		// Get the types and the amount of instances of each unique operation on the mutable model
-		const mu::FProgram& Program = MutableModel->GetPrivate()->Program;
+		const UE::Mutable::Private::FProgram& Program = MutableModel->GetPrivate()->Program;
 		const uint32 ProgramAddressesCount = Program.OpAddress.Num();
 		
 		// Ensure first operation type is NONE since we are skipping it due to it having to have that type
-		check (Program.GetOpType(Program.OpAddress[0]) == mu::EOpType::NONE);
+		check (Program.GetOpType(Program.OpAddress[0]) == UE::Mutable::Private::EOpType::NONE);
 
 		// Iterate over the addresses of the program and count how many instances each type has.
 		for (uint32 ProgramAddressesIndex = 1 ; ProgramAddressesIndex < ProgramAddressesCount; ProgramAddressesIndex++)
 		{
 			// Locate what is the position (index) of the operation type of the address on our collection of types found until now
-			const mu::EOpType OperationType = Program.GetOpType(ProgramAddressesIndex);
+			const UE::Mutable::Private::EOpType OperationType = Program.GetOpType(ProgramAddressesIndex);
 			
 			// Increase the counter for this operation type
 			const uint16 TypeAsInteger = StaticCast<uint16>(OperationType);
@@ -1180,20 +1180,20 @@ void SMutableCodeViewer::CacheOperationTypesPresentOnModel()
 	// Remove all operation types that do have no operations present on the model
 	{
 		ModelOperationTypes.RemoveAll(
-		[](const TPair<mu::EOpType, uint32>& Current)
+		[](const TPair<UE::Mutable::Private::EOpType, uint32>& Current)
 			{
 				return Current.Value == 0;
 			});
 	}
 	
 	// Sort the contents of the array of mutable operation types alphabetically
-	ModelOperationTypes.StableSort([&](const TPair<mu::EOpType,uint32>& A, const TPair<mu::EOpType,uint32>& B)
+	ModelOperationTypes.StableSort([&](const TPair<UE::Mutable::Private::EOpType,uint32>& A, const TPair<UE::Mutable::Private::EOpType,uint32>& B)
 	{
 		// Find the name
 		FString AString;
 		{
 			const uint16 OperationIndex = static_cast<uint16>(A.Key);
-			const TCHAR* OpName = mu::s_opNames[OperationIndex];
+			const TCHAR* OpName = UE::Mutable::Private::s_opNames[OperationIndex];
 			AString = FString(OpName);
 		}
 		
@@ -1201,7 +1201,7 @@ void SMutableCodeViewer::CacheOperationTypesPresentOnModel()
 		FString BString;
 		{
 			const uint16 OperationIndex = static_cast<uint16>(B.Key);
-			const TCHAR* OpName = mu::s_opNames[OperationIndex];
+			const TCHAR* OpName = UE::Mutable::Private::s_opNames[OperationIndex];
 			BString = FString(OpName);
 		}
 		
@@ -1334,13 +1334,13 @@ void SMutableCodeViewer::GenerateAllTreeElements()
 	// 	be ignored due to how we handle them when expanding and contracting elements (OnExpansionChanged)
 
 	// Generate all root nodes
-	const mu::FModel::Private* ModelPrivate = MutableModel->GetPrivate();
+	const UE::Mutable::Private::FModel::Private* ModelPrivate = MutableModel->GetPrivate();
 	check(ModelPrivate);
-	const mu::FProgram& Program = ModelPrivate->Program;
+	const UE::Mutable::Private::FProgram& Program = ModelPrivate->Program;
 	const uint32 StateCount = Program.States.Num();
 	for ( uint32 StateIndex = 0; StateIndex < StateCount; ++StateIndex )
 	{
-		const mu::FProgram::FState& State = Program.States[StateIndex];
+		const UE::Mutable::Private::FProgram::FState& State = Program.States[StateIndex];
 		const FString Caption = FString::Printf( TEXT("state [%s]"), *State.Name );
 
 		const FSlateColor LabelColor = ColorPerComputationalCost[StaticCast<uint8>(GetOperationTypeComputationalCost(
@@ -1355,7 +1355,7 @@ void SMutableCodeViewer::GenerateAllTreeElements()
 		RootNodes.Add(RootNodeElement);
 		
 		// Add the element to the cache so we keep the indices straight.
-		constexpr mu::OP::ADDRESS CommonParent = 0;
+		constexpr UE::Mutable::Private::OP::ADDRESS CommonParent = 0;
 		const FItemCacheKey Key = { CommonParent, State.Root, StateIndex };
 		ItemCache.Add(Key, RootNodeElement);
 		
@@ -1371,11 +1371,11 @@ void SMutableCodeViewer::GenerateAllTreeElements()
 	
 }
 
-void SMutableCodeViewer::GenerateElementRecursive(const int32& InStateIndex, mu::OP::ADDRESS InParentAddress,  const mu::FProgram& InProgram)
+void SMutableCodeViewer::GenerateElementRecursive(const int32& InStateIndex, UE::Mutable::Private::OP::ADDRESS InParentAddress,  const UE::Mutable::Private::FProgram& InProgram)
 {
 	// This will be used to add operations
 	uint32 ChildIndex = 0;
-	auto AddOpFunc = [this, InParentAddress, &InProgram, &ChildIndex, &InStateIndex](mu::OP::ADDRESS ChildAddress, const FString& Caption)
+	auto AddOpFunc = [this, InParentAddress, &InProgram, &ChildIndex, &InStateIndex](UE::Mutable::Private::OP::ADDRESS ChildAddress, const FString& Caption)
 		{
 			{
 				const FItemCacheKey Key = { InParentAddress, ChildAddress, ChildIndex };
@@ -1416,49 +1416,49 @@ void SMutableCodeViewer::GenerateElementRecursive(const int32& InStateIndex, mu:
 	
 	// For some specific parent operation types we create more detailed subtrees.
 	bool bUseGeneric = false;
-	const mu::EOpType ParentOperationType = InProgram.GetOpType(InParentAddress);
+	const UE::Mutable::Private::EOpType ParentOperationType = InProgram.GetOpType(InParentAddress);
 	switch (ParentOperationType)
 	{
-	case mu::EOpType::IM_CONDITIONAL:
-	case mu::EOpType::LA_CONDITIONAL:
-	case mu::EOpType::ME_CONDITIONAL:
-	case mu::EOpType::CO_CONDITIONAL:
-	case mu::EOpType::SC_CONDITIONAL:
-	case mu::EOpType::NU_CONDITIONAL:
-	case mu::EOpType::IN_CONDITIONAL:
-	case mu::EOpType::ED_CONDITIONAL:
+	case UE::Mutable::Private::EOpType::IM_CONDITIONAL:
+	case UE::Mutable::Private::EOpType::LA_CONDITIONAL:
+	case UE::Mutable::Private::EOpType::ME_CONDITIONAL:
+	case UE::Mutable::Private::EOpType::CO_CONDITIONAL:
+	case UE::Mutable::Private::EOpType::SC_CONDITIONAL:
+	case UE::Mutable::Private::EOpType::NU_CONDITIONAL:
+	case UE::Mutable::Private::EOpType::IN_CONDITIONAL:
+	case UE::Mutable::Private::EOpType::ED_CONDITIONAL:
 	{
-		mu::OP::ConditionalArgs Args = InProgram.GetOpArgs<mu::OP::ConditionalArgs>(InParentAddress);
+		UE::Mutable::Private::OP::ConditionalArgs Args = InProgram.GetOpArgs<UE::Mutable::Private::OP::ConditionalArgs>(InParentAddress);
 		AddOpFunc(Args.condition, TEXT("cond "));
 		AddOpFunc(Args.yes, TEXT("true "));
 		AddOpFunc(Args.no, TEXT("false "));
 		break;
 	}
 
-	case mu::EOpType::IM_SWITCH:
-	case mu::EOpType::LA_SWITCH:
-	case mu::EOpType::ME_SWITCH:
-	case mu::EOpType::CO_SWITCH:
-	case mu::EOpType::SC_SWITCH:
-	case mu::EOpType::NU_SWITCH:
-	case mu::EOpType::IN_SWITCH:
-	case mu::EOpType::ED_SWITCH:
+	case UE::Mutable::Private::EOpType::IM_SWITCH:
+	case UE::Mutable::Private::EOpType::LA_SWITCH:
+	case UE::Mutable::Private::EOpType::ME_SWITCH:
+	case UE::Mutable::Private::EOpType::CO_SWITCH:
+	case UE::Mutable::Private::EOpType::SC_SWITCH:
+	case UE::Mutable::Private::EOpType::NU_SWITCH:
+	case UE::Mutable::Private::EOpType::IN_SWITCH:
+	case UE::Mutable::Private::EOpType::ED_SWITCH:
 	{
 		const uint8* OpData = InProgram.GetOpArgsPointer(InParentAddress);
 
-		mu::OP::ADDRESS VarAddress;
-		FMemory::Memcpy(&VarAddress, OpData, sizeof(mu::OP::ADDRESS));
-		OpData += sizeof(mu::OP::ADDRESS);
+		UE::Mutable::Private::OP::ADDRESS VarAddress;
+		FMemory::Memcpy(&VarAddress, OpData, sizeof(UE::Mutable::Private::OP::ADDRESS));
+		OpData += sizeof(UE::Mutable::Private::OP::ADDRESS);
 		AddOpFunc(VarAddress, TEXT("var "));
 
-		mu::OP::ADDRESS DefAddress;
-		FMemory::Memcpy(&DefAddress, OpData, sizeof(mu::OP::ADDRESS));
-		OpData += sizeof(mu::OP::ADDRESS);
+		UE::Mutable::Private::OP::ADDRESS DefAddress;
+		FMemory::Memcpy(&DefAddress, OpData, sizeof(UE::Mutable::Private::OP::ADDRESS));
+		OpData += sizeof(UE::Mutable::Private::OP::ADDRESS);
 		AddOpFunc(DefAddress, TEXT("def "));
 
-		mu::OP::FSwitchCaseDescriptor CaseDesc;
-		FMemory::Memcpy(&CaseDesc, OpData, sizeof(mu::OP::FSwitchCaseDescriptor));
-		OpData += sizeof(mu::OP::FSwitchCaseDescriptor);
+		UE::Mutable::Private::OP::FSwitchCaseDescriptor CaseDesc;
+		FMemory::Memcpy(&CaseDesc, OpData, sizeof(UE::Mutable::Private::OP::FSwitchCaseDescriptor));
+		OpData += sizeof(UE::Mutable::Private::OP::FSwitchCaseDescriptor);
 
 		if (!CaseDesc.bUseRanges)
 		{
@@ -1468,9 +1468,9 @@ void SMutableCodeViewer::GenerateElementRecursive(const int32& InStateIndex, mu:
 				FMemory::Memcpy(&Condition, OpData, sizeof(int32));
 				OpData += sizeof(int32);
 
-				mu::OP::ADDRESS At;
-				FMemory::Memcpy(&At, OpData, sizeof(mu::OP::ADDRESS));
-				OpData += sizeof(mu::OP::ADDRESS);
+				UE::Mutable::Private::OP::ADDRESS At;
+				FMemory::Memcpy(&At, OpData, sizeof(UE::Mutable::Private::OP::ADDRESS));
+				OpData += sizeof(UE::Mutable::Private::OP::ADDRESS);
 
 				FString Caption = FString::Printf(TEXT("case %d "), Condition);
 				AddOpFunc(At, Caption);
@@ -1488,9 +1488,9 @@ void SMutableCodeViewer::GenerateElementRecursive(const int32& InStateIndex, mu:
 				FMemory::Memcpy(&RangeSize, OpData, sizeof(uint32));
 				OpData += sizeof(uint32);
 
-				mu::OP::ADDRESS At;
-				FMemory::Memcpy(&At, OpData, sizeof(mu::OP::ADDRESS));
-				OpData += sizeof(mu::OP::ADDRESS);
+				UE::Mutable::Private::OP::ADDRESS At;
+				FMemory::Memcpy(&At, OpData, sizeof(UE::Mutable::Private::OP::ADDRESS));
+				OpData += sizeof(UE::Mutable::Private::OP::ADDRESS);
 
 				FString Caption = FString::Printf(TEXT("case %d to %d "), ConditionStart, ConditionStart+RangeSize-1);
 				AddOpFunc(At, Caption);
@@ -1500,9 +1500,9 @@ void SMutableCodeViewer::GenerateElementRecursive(const int32& InStateIndex, mu:
 		break;
 	}
 
-	case mu::EOpType::IM_SWIZZLE:
+	case UE::Mutable::Private::EOpType::IM_SWIZZLE:
 	{
-		mu::OP::ImageSwizzleArgs Args = InProgram.GetOpArgs<mu::OP::ImageSwizzleArgs>(InParentAddress);
+		UE::Mutable::Private::OP::ImageSwizzleArgs Args = InProgram.GetOpArgs<UE::Mutable::Private::OP::ImageSwizzleArgs>(InParentAddress);
 		for (int32 Channel = 0; Channel < 4; ++Channel)
 		{
 			FString Caption = FString::Printf(TEXT("%d is %d from "), Channel, Args.sourceChannels[Channel]);
@@ -1511,9 +1511,9 @@ void SMutableCodeViewer::GenerateElementRecursive(const int32& InStateIndex, mu:
 		break;
 	}
 
-	case mu::EOpType::CO_SWIZZLE:
+	case UE::Mutable::Private::EOpType::CO_SWIZZLE:
 	{
-		mu::OP::ColourSwizzleArgs Args = InProgram.GetOpArgs<mu::OP::ColourSwizzleArgs>(InParentAddress);
+		UE::Mutable::Private::OP::ColourSwizzleArgs Args = InProgram.GetOpArgs<UE::Mutable::Private::OP::ColourSwizzleArgs>(InParentAddress);
 		for (int32 Channel = 0; Channel < 4; ++Channel)
 		{
 			FString Caption = FString::Printf(TEXT("%d is %d from "), Channel, Args.sourceChannels[Channel]);
@@ -1522,27 +1522,27 @@ void SMutableCodeViewer::GenerateElementRecursive(const int32& InStateIndex, mu:
 		break;
 	}
 
-	case mu::EOpType::IM_LAYER:
+	case UE::Mutable::Private::EOpType::IM_LAYER:
 	{
-		mu::OP::ImageLayerArgs Args = InProgram.GetOpArgs<mu::OP::ImageLayerArgs>(InParentAddress);
+		UE::Mutable::Private::OP::ImageLayerArgs Args = InProgram.GetOpArgs<UE::Mutable::Private::OP::ImageLayerArgs>(InParentAddress);
 		AddOpFunc(Args.base, TEXT("base "));
 		AddOpFunc(Args.mask, TEXT("mask "));
 		AddOpFunc(Args.blended, TEXT("blended "));
 		break;
 	}
 
-	case mu::EOpType::IM_LAYERCOLOUR:
+	case UE::Mutable::Private::EOpType::IM_LAYERCOLOUR:
 	{
-		mu::OP::ImageLayerColourArgs Args = InProgram.GetOpArgs<mu::OP::ImageLayerColourArgs>(InParentAddress);
+		UE::Mutable::Private::OP::ImageLayerColourArgs Args = InProgram.GetOpArgs<UE::Mutable::Private::OP::ImageLayerColourArgs>(InParentAddress);
 		AddOpFunc(Args.base, TEXT("base "));
 		AddOpFunc(Args.mask, TEXT("mask "));
 		AddOpFunc(Args.colour, TEXT("colour "));
 		break;
 	}
 
-	case mu::EOpType::IM_MULTILAYER:
+	case UE::Mutable::Private::EOpType::IM_MULTILAYER:
 	{
-		mu::OP::ImageMultiLayerArgs Args = InProgram.GetOpArgs<mu::OP::ImageMultiLayerArgs>(InParentAddress);
+		UE::Mutable::Private::OP::ImageMultiLayerArgs Args = InProgram.GetOpArgs<UE::Mutable::Private::OP::ImageMultiLayerArgs>(InParentAddress);
 		AddOpFunc(Args.rangeSize, TEXT("range "));
 		AddOpFunc(Args.base, TEXT("base "));
 		AddOpFunc(Args.mask, TEXT("mask "));
@@ -1550,11 +1550,11 @@ void SMutableCodeViewer::GenerateElementRecursive(const int32& InStateIndex, mu:
 		break;
 	}
 
-	case mu::EOpType::ME_ADDMETADATA:
+	case UE::Mutable::Private::EOpType::ME_ADDMETADATA:
 	{
-		mu::OP::MeshAddMetadataArgs Args = InProgram.GetOpArgs<mu::OP::MeshAddMetadataArgs>(InParentAddress);
+		UE::Mutable::Private::OP::MeshAddMetadataArgs Args = InProgram.GetOpArgs<UE::Mutable::Private::OP::MeshAddMetadataArgs>(InParentAddress);
 
-		using OpEnumFlags = mu::OP::MeshAddMetadataArgs::EnumFlags;
+		using OpEnumFlags = UE::Mutable::Private::OP::MeshAddMetadataArgs::EnumFlags;
 
 		// Set count to 1 if is not a list.
 		int32 TagCount      = int32(!EnumHasAnyFlags(Args.Flags, OpEnumFlags::IsTagList));
@@ -1588,103 +1588,103 @@ void SMutableCodeViewer::GenerateElementRecursive(const int32& InStateIndex, mu:
 		break;
 	}
 
-	case mu::EOpType::ME_APPLYLAYOUT:
+	case UE::Mutable::Private::EOpType::ME_APPLYLAYOUT:
 	{
-		mu::OP::MeshApplyLayoutArgs Args = InProgram.GetOpArgs<mu::OP::MeshApplyLayoutArgs>(InParentAddress);
+		UE::Mutable::Private::OP::MeshApplyLayoutArgs Args = InProgram.GetOpArgs<UE::Mutable::Private::OP::MeshApplyLayoutArgs>(InParentAddress);
 		AddOpFunc(Args.Layout, TEXT("layout "));
 		AddOpFunc(Args.Mesh, TEXT("mesh "));
 		break;
 	}
 
-	case mu::EOpType::ME_PREPARELAYOUT:
+	case UE::Mutable::Private::EOpType::ME_PREPARELAYOUT:
 	{
-		mu::OP::MeshPrepareLayoutArgs Args = InProgram.GetOpArgs<mu::OP::MeshPrepareLayoutArgs>(InParentAddress);
+		UE::Mutable::Private::OP::MeshPrepareLayoutArgs Args = InProgram.GetOpArgs<UE::Mutable::Private::OP::MeshPrepareLayoutArgs>(InParentAddress);
 		AddOpFunc(Args.Layout, TEXT("layout "));
 		AddOpFunc(Args.Mesh, TEXT("mesh "));
 		break;
 	}
 
-	case mu::EOpType::ME_DIFFERENCE:
+	case UE::Mutable::Private::EOpType::ME_DIFFERENCE:
 	{
 		const uint8* data = InProgram.GetOpArgsPointer(InParentAddress);
 
-		mu::OP::ADDRESS BaseAt = 0;
-		FMemory::Memcpy(&BaseAt, data, sizeof(mu::OP::ADDRESS));
+		UE::Mutable::Private::OP::ADDRESS BaseAt = 0;
+		FMemory::Memcpy(&BaseAt, data, sizeof(UE::Mutable::Private::OP::ADDRESS));
 		AddOpFunc(BaseAt, TEXT("base "));
-		data += sizeof(mu::OP::ADDRESS);
+		data += sizeof(UE::Mutable::Private::OP::ADDRESS);
 
-		mu::OP::ADDRESS TargetAt = 0;
-		FMemory::Memcpy(&TargetAt, data, sizeof(mu::OP::ADDRESS));
+		UE::Mutable::Private::OP::ADDRESS TargetAt = 0;
+		FMemory::Memcpy(&TargetAt, data, sizeof(UE::Mutable::Private::OP::ADDRESS));
 		AddOpFunc(TargetAt, TEXT("target "));
-		data += sizeof(mu::OP::ADDRESS);
+		data += sizeof(UE::Mutable::Private::OP::ADDRESS);
 		
 		break;
 	}
 		
-	case mu::EOpType::ME_MORPH:
+	case UE::Mutable::Private::EOpType::ME_MORPH:
 	{
 		const uint8* Data = InProgram.GetOpArgsPointer(InParentAddress);
 
-		mu::OP::ADDRESS FactorAt = 0;
-		FMemory::Memcpy(&FactorAt, Data, sizeof(mu::OP::ADDRESS));
+		UE::Mutable::Private::OP::ADDRESS FactorAt = 0;
+		FMemory::Memcpy(&FactorAt, Data, sizeof(UE::Mutable::Private::OP::ADDRESS));
 		AddOpFunc(FactorAt, TEXT("factor "));
-		Data += sizeof(mu::OP::ADDRESS);
+		Data += sizeof(UE::Mutable::Private::OP::ADDRESS);
 		
-		mu::OP::ADDRESS BaseAt = 0;
-		FMemory::Memcpy(&BaseAt, Data, sizeof(mu::OP::ADDRESS));
+		UE::Mutable::Private::OP::ADDRESS BaseAt = 0;
+		FMemory::Memcpy(&BaseAt, Data, sizeof(UE::Mutable::Private::OP::ADDRESS));
 		AddOpFunc(BaseAt, TEXT("base "));
-		Data += sizeof(mu::OP::ADDRESS);
+		Data += sizeof(UE::Mutable::Private::OP::ADDRESS);
 
-		mu::OP::ADDRESS TargetAt = 0;
-		FMemory::Memcpy(&TargetAt, Data, sizeof(mu::OP::ADDRESS));
+		UE::Mutable::Private::OP::ADDRESS TargetAt = 0;
+		FMemory::Memcpy(&TargetAt, Data, sizeof(UE::Mutable::Private::OP::ADDRESS));
 		AddOpFunc(TargetAt, TEXT("target "));
-		Data += sizeof(mu::OP::ADDRESS);
+		Data += sizeof(UE::Mutable::Private::OP::ADDRESS);
 		
 		break;
 	}
 		
-	case mu::EOpType::SC_CURVE:
+	case UE::Mutable::Private::EOpType::SC_CURVE:
 	{
-		mu::OP::ScalarCurveArgs Args = InProgram.GetOpArgs<mu::OP::ScalarCurveArgs>(InParentAddress);
+		UE::Mutable::Private::OP::ScalarCurveArgs Args = InProgram.GetOpArgs<UE::Mutable::Private::OP::ScalarCurveArgs>(InParentAddress);
 		AddOpFunc(Args.time, TEXT("time "));
 		AddOpFunc(Args.curve, TEXT("curve "));
 		break;
 	}
 		
-	case mu::EOpType::CO_SAMPLEIMAGE:
+	case UE::Mutable::Private::EOpType::CO_SAMPLEIMAGE:
 	{
-		mu::OP::ColourSampleImageArgs Args = InProgram.GetOpArgs<mu::OP::ColourSampleImageArgs>(InParentAddress);
+		UE::Mutable::Private::OP::ColourSampleImageArgs Args = InProgram.GetOpArgs<UE::Mutable::Private::OP::ColourSampleImageArgs>(InParentAddress);
 		AddOpFunc(Args.Image, TEXT("image "));
 		AddOpFunc(Args.X, TEXT("x "));
 		AddOpFunc(Args.Y, TEXT("y "));
 		break;
 	}
 		
-	case mu::EOpType::IM_RESIZELIKE:
+	case UE::Mutable::Private::EOpType::IM_RESIZELIKE:
 	{
-		mu::OP::ImageResizeLikeArgs Args = InProgram.GetOpArgs<mu::OP::ImageResizeLikeArgs>(InParentAddress);
+		UE::Mutable::Private::OP::ImageResizeLikeArgs Args = InProgram.GetOpArgs<UE::Mutable::Private::OP::ImageResizeLikeArgs>(InParentAddress);
 		AddOpFunc(Args.Source, TEXT("src "));
 		AddOpFunc(Args.SizeSource, TEXT("sizeSrc "));
 		break;
 	}
 
-	case mu::EOpType::IN_ADDMESH:
-	case mu::EOpType::IN_ADDIMAGE:
-	case mu::EOpType::IN_ADDVECTOR:
-	case mu::EOpType::IN_ADDSCALAR:
-	case mu::EOpType::IN_ADDSTRING:
-	case mu::EOpType::IN_ADDCOMPONENT:
-	case mu::EOpType::IN_ADDSURFACE:
+	case UE::Mutable::Private::EOpType::IN_ADDMESH:
+	case UE::Mutable::Private::EOpType::IN_ADDIMAGE:
+	case UE::Mutable::Private::EOpType::IN_ADDVECTOR:
+	case UE::Mutable::Private::EOpType::IN_ADDSCALAR:
+	case UE::Mutable::Private::EOpType::IN_ADDSTRING:
+	case UE::Mutable::Private::EOpType::IN_ADDCOMPONENT:
+	case UE::Mutable::Private::EOpType::IN_ADDSURFACE:
 	{
-		mu::OP::InstanceAddArgs Args = InProgram.GetOpArgs<mu::OP::InstanceAddArgs>(InParentAddress);
+		UE::Mutable::Private::OP::InstanceAddArgs Args = InProgram.GetOpArgs<UE::Mutable::Private::OP::InstanceAddArgs>(InParentAddress);
 		AddOpFunc(Args.instance, TEXT("instance "));
 		AddOpFunc(Args.value, TEXT("value "));
 		break;
 	}
 		
-	case mu::EOpType::IM_COMPOSE:
+	case UE::Mutable::Private::EOpType::IM_COMPOSE:
 	{
-		mu::OP::ImageComposeArgs Args = InProgram.GetOpArgs<mu::OP::ImageComposeArgs>(InParentAddress);
+		UE::Mutable::Private::OP::ImageComposeArgs Args = InProgram.GetOpArgs<UE::Mutable::Private::OP::ImageComposeArgs>(InParentAddress);
 		AddOpFunc(Args.layout, TEXT("layout "));
 		AddOpFunc(Args.base, TEXT("base "));
 		AddOpFunc(Args.blockImage, TEXT("blockImage "));
@@ -1692,55 +1692,55 @@ void SMutableCodeViewer::GenerateElementRecursive(const int32& InStateIndex, mu:
 		break;
 	}
 		
-	case mu::EOpType::IM_INTERPOLATE:
+	case UE::Mutable::Private::EOpType::IM_INTERPOLATE:
 	{
-		mu::OP::ImageInterpolateArgs Args = InProgram.GetOpArgs<mu::OP::ImageInterpolateArgs>(InParentAddress);
+		UE::Mutable::Private::OP::ImageInterpolateArgs Args = InProgram.GetOpArgs<UE::Mutable::Private::OP::ImageInterpolateArgs>(InParentAddress);
 		AddOpFunc(Args.Factor, TEXT("factor "));
 			
 		int32 TargetIndex = 0;
-		for (const mu::OP::ADDRESS& Operation : Args.Targets)
+		for (const UE::Mutable::Private::OP::ADDRESS& Operation : Args.Targets)
 		{
 			AddOpFunc(Operation, FString::Printf(TEXT("target %i "), TargetIndex++));
 		}
 		break;
 	}
 		
-	case mu::EOpType::IM_SATURATE:
+	case UE::Mutable::Private::EOpType::IM_SATURATE:
 	{
-		mu::OP::ImageSaturateArgs Args = InProgram.GetOpArgs<mu::OP::ImageSaturateArgs>(InParentAddress);
+		UE::Mutable::Private::OP::ImageSaturateArgs Args = InProgram.GetOpArgs<UE::Mutable::Private::OP::ImageSaturateArgs>(InParentAddress);
 		AddOpFunc(Args.Base, TEXT("base "));
 		AddOpFunc(Args.Factor, TEXT("factor "));
 		break;
 	}
 		
-	case mu::EOpType::IM_COLOURMAP:
+	case UE::Mutable::Private::EOpType::IM_COLOURMAP:
 	{
-		mu::OP::ImageColourMapArgs Args = InProgram.GetOpArgs<mu::OP::ImageColourMapArgs>(InParentAddress);
+		UE::Mutable::Private::OP::ImageColourMapArgs Args = InProgram.GetOpArgs<UE::Mutable::Private::OP::ImageColourMapArgs>(InParentAddress);
 		AddOpFunc(Args.Base, TEXT("base "));
 		AddOpFunc(Args.Mask, TEXT("mask "));
 		AddOpFunc(Args.Map, TEXT("map "));
 		break;
 	}
 		
-	case mu::EOpType::IM_BINARISE:
+	case UE::Mutable::Private::EOpType::IM_BINARISE:
 	{
-		mu::OP::ImageBinariseArgs Args = InProgram.GetOpArgs<mu::OP::ImageBinariseArgs>(InParentAddress);
+		UE::Mutable::Private::OP::ImageBinariseArgs Args = InProgram.GetOpArgs<UE::Mutable::Private::OP::ImageBinariseArgs>(InParentAddress);
 		AddOpFunc(Args.Base, TEXT("base "));
 		AddOpFunc(Args.Threshold, TEXT("threshold "));
 		break;
 	}
 		
-	case mu::EOpType::IM_PATCH:
+	case UE::Mutable::Private::EOpType::IM_PATCH:
 	{
-		mu::OP::ImagePatchArgs Args = InProgram.GetOpArgs<mu::OP::ImagePatchArgs>(InParentAddress);
+		UE::Mutable::Private::OP::ImagePatchArgs Args = InProgram.GetOpArgs<UE::Mutable::Private::OP::ImagePatchArgs>(InParentAddress);
 		AddOpFunc(Args.base, TEXT("base "));
 		AddOpFunc(Args.patch, TEXT("patch "));
 		break;
 	}
 		
-	case mu::EOpType::IM_RASTERMESH:
+	case UE::Mutable::Private::EOpType::IM_RASTERMESH:
 	{
-		mu::OP::ImageRasterMeshArgs Args = InProgram.GetOpArgs<mu::OP::ImageRasterMeshArgs>(InParentAddress);
+		UE::Mutable::Private::OP::ImageRasterMeshArgs Args = InProgram.GetOpArgs<UE::Mutable::Private::OP::ImageRasterMeshArgs>(InParentAddress);
 		AddOpFunc(Args.mesh, TEXT("mesh "));
 		AddOpFunc(Args.image, TEXT("image "));
 		AddOpFunc(Args.mask, TEXT("mask "));
@@ -1749,25 +1749,25 @@ void SMutableCodeViewer::GenerateElementRecursive(const int32& InStateIndex, mu:
 		break;
 	}
 		
-	case mu::EOpType::IM_DISPLACE:
+	case UE::Mutable::Private::EOpType::IM_DISPLACE:
 	{
-		mu::OP::ImageDisplaceArgs Args = InProgram.GetOpArgs<mu::OP::ImageDisplaceArgs>(InParentAddress);
+		UE::Mutable::Private::OP::ImageDisplaceArgs Args = InProgram.GetOpArgs<UE::Mutable::Private::OP::ImageDisplaceArgs>(InParentAddress);
 		AddOpFunc(Args.Source, TEXT("src "));
 		AddOpFunc(Args.DisplacementMap, TEXT("displacementMap "));
 		break;
 	}
 		
-	case mu::EOpType::IM_NORMALCOMPOSITE:
+	case UE::Mutable::Private::EOpType::IM_NORMALCOMPOSITE:
 	{
-		mu::OP::ImageNormalCompositeArgs Args = InProgram.GetOpArgs<mu::OP::ImageNormalCompositeArgs>(InParentAddress);
+		UE::Mutable::Private::OP::ImageNormalCompositeArgs Args = InProgram.GetOpArgs<UE::Mutable::Private::OP::ImageNormalCompositeArgs>(InParentAddress);
 		AddOpFunc(Args.base, TEXT("base "));
 		AddOpFunc(Args.normal, TEXT("normal "));
 		break;
 	}
 		
-	case mu::EOpType::IM_TRANSFORM:
+	case UE::Mutable::Private::EOpType::IM_TRANSFORM:
 	{
-		mu::OP::ImageTransformArgs Args = InProgram.GetOpArgs<mu::OP::ImageTransformArgs>(InParentAddress);
+		UE::Mutable::Private::OP::ImageTransformArgs Args = InProgram.GetOpArgs<UE::Mutable::Private::OP::ImageTransformArgs>(InParentAddress);
 		AddOpFunc(Args.Base, TEXT("base "));
 		AddOpFunc(Args.OffsetX, TEXT("offsetX "));
 		AddOpFunc(Args.OffsetY, TEXT("offsetY "));
@@ -1777,11 +1777,11 @@ void SMutableCodeViewer::GenerateElementRecursive(const int32& InStateIndex, mu:
 		break;
 	}
 
-	case mu::EOpType::SC_MATERIAL_BREAK:
-	case mu::EOpType::CO_MATERIAL_BREAK:
-	case mu::EOpType::IM_MATERIAL_BREAK:
+	case UE::Mutable::Private::EOpType::SC_MATERIAL_BREAK:
+	case UE::Mutable::Private::EOpType::CO_MATERIAL_BREAK:
+	case UE::Mutable::Private::EOpType::IM_MATERIAL_BREAK:
 	{
-		mu::OP::MaterialBreakArgs Args = InProgram.GetOpArgs<mu::OP::MaterialBreakArgs>(InParentAddress);
+		UE::Mutable::Private::OP::MaterialBreakArgs Args = InProgram.GetOpArgs<UE::Mutable::Private::OP::MaterialBreakArgs>(InParentAddress);
 		AddOpFunc(Args.Material, TEXT("Material "));
 		break;
 	}
@@ -1801,7 +1801,7 @@ void SMutableCodeViewer::GenerateElementRecursive(const int32& InStateIndex, mu:
 	if (bUseGeneric)
 	{
 		// Find children of the provided element without adding any extra string to better identify what each of them represents
-		mu::ForEachReference(InProgram, InParentAddress, [this, &InProgram, AddOpFunc](mu::OP::ADDRESS ChildAddress)
+		UE::Mutable::Private::ForEachReference(InProgram, InParentAddress, [this, &InProgram, AddOpFunc](UE::Mutable::Private::OP::ADDRESS ChildAddress)
 			{
 				AddOpFunc(ChildAddress,TEXT(""));
 			});
@@ -1811,7 +1811,7 @@ void SMutableCodeViewer::GenerateElementRecursive(const int32& InStateIndex, mu:
 		// Validate in case there is a mismatch in the custom processing of children and the generic one, which would cause problems.
 		ChildIndex = 0;
 
-		auto ValidateOpFunc = [this, InParentAddress, ParentOperationType, &ChildIndex](mu::OP::ADDRESS ChildAddress)
+		auto ValidateOpFunc = [this, InParentAddress, ParentOperationType, &ChildIndex](UE::Mutable::Private::OP::ADDRESS ChildAddress)
 		{
 			const FItemCacheKey Key = { InParentAddress, ChildAddress, ChildIndex };
 			const TSharedPtr<FMutableCodeTreeElement>* CachedItem = ItemCache.Find(Key);
@@ -1823,14 +1823,14 @@ void SMutableCodeViewer::GenerateElementRecursive(const int32& InStateIndex, mu:
 			++ChildIndex;
 		};
 
-		mu::ForEachReference(InProgram, InParentAddress, [this, ValidateOpFunc](mu::OP::ADDRESS ChildAddress)
+		UE::Mutable::Private::ForEachReference(InProgram, InParentAddress, [this, ValidateOpFunc](UE::Mutable::Private::OP::ADDRESS ChildAddress)
 		{
 			ValidateOpFunc(ChildAddress);
 		});
 	}
 }
 
-SMutableCodeViewer::EOperationComputationalCost SMutableCodeViewer::GetOperationTypeComputationalCost(mu::EOpType OperationType) const
+SMutableCodeViewer::EOperationComputationalCost SMutableCodeViewer::GetOperationTypeComputationalCost(UE::Mutable::Private::EOpType OperationType) const
 {
 	if (VeryExpensiveOperationTypes.Contains(OperationType))
 	{
@@ -1885,13 +1885,13 @@ void SMutableCodeViewer::GetChildrenForInfo(TSharedPtr<FMutableCodeTreeElement>
 	}
 
 	check(MutableModel);
-	const mu::FProgram& Program = MutableModel->GetPrivate()->Program;
+	const UE::Mutable::Private::FProgram& Program = MutableModel->GetPrivate()->Program;
 
-	mu::OP::ADDRESS ParentAddress = InInfo->MutableOperation;
+	UE::Mutable::Private::OP::ADDRESS ParentAddress = InInfo->MutableOperation;
 
 	// Generic case for unnamed children traversal.
 	uint32 ChildIndex = 0;
-	mu::ForEachReference(Program, InInfo->MutableOperation, [this, ParentAddress, &ChildIndex, &OutChildren](mu::OP::ADDRESS ChildAddress)
+	UE::Mutable::Private::ForEachReference(Program, InInfo->MutableOperation, [this, ParentAddress, &ChildIndex, &OutChildren](UE::Mutable::Private::OP::ADDRESS ChildAddress)
 	{
 		{
 			const FItemCacheKey Key = { ParentAddress, ChildAddress, ChildIndex };
@@ -1919,7 +1919,7 @@ void SMutableCodeViewer::OnExpansionChanged(TSharedPtr<FMutableCodeTreeElement>
 	// If an element gets expanded then contract (if found) the other element that uses the same address
 	if (bInExpanded)
 	{
-		const mu::OP::ADDRESS MutableOperation = InItem->MutableOperation;
+		const UE::Mutable::Private::OP::ADDRESS MutableOperation = InItem->MutableOperation;
 		const TSharedPtr<FMutableCodeTreeElement>* PreviouslyExpandedElement = ExpandedElements.Find(MutableOperation);
 		if (PreviouslyExpandedElement)
 		{
@@ -1952,7 +1952,7 @@ void SMutableCodeViewer::OnExpansionChanged(TSharedPtr<FMutableCodeTreeElement>
 void SMutableCodeViewer::GetVisibleChildren(TSharedPtr<FMutableCodeTreeElement> InInfo, TSet<TSharedPtr<FMutableCodeTreeElement>>& OutChildren)
 {
 	check(MutableModel);
-	const mu::FProgram& MutableProgram = MutableModel->GetPrivate()->Program;
+	const UE::Mutable::Private::FProgram& MutableProgram = MutableModel->GetPrivate()->Program;
 	
 	TArray<TSharedPtr<FMutableCodeTreeElement>> ToSearchForChildren;
 	ToSearchForChildren.Add(InInfo);
@@ -1962,11 +1962,11 @@ void SMutableCodeViewer::GetVisibleChildren(TSharedPtr<FMutableCodeTreeElement>
 		const TSharedPtr<FMutableCodeTreeElement> ToCheck = ToSearchForChildren[0];
 		ToSearchForChildren.RemoveAt(0);
 		
-		const mu::OP::ADDRESS ParentAddress = ToCheck->MutableOperation;
+		const UE::Mutable::Private::OP::ADDRESS ParentAddress = ToCheck->MutableOperation;
 	
 		// Generic case for unnamed children traversal.
 		uint32 ChildIndex = 0;
-		mu::ForEachReference(MutableProgram, ParentAddress, [this, ParentAddress, &ChildIndex, &OutChildren, &ToSearchForChildren ](mu::OP::ADDRESS ChildAddress)
+		UE::Mutable::Private::ForEachReference(MutableProgram, ParentAddress, [this, ParentAddress, &ChildIndex, &OutChildren, &ToSearchForChildren ](UE::Mutable::Private::OP::ADDRESS ChildAddress)
 		{
 			{
 				const FItemCacheKey Key = { ParentAddress, ChildAddress, ChildIndex };
@@ -2030,19 +2030,19 @@ void SMutableCodeViewer::OnSelectionChanged(TSharedPtr<FMutableCodeTreeElement>
 	bIsPreviewPendingUpdate = true;
 
 	SelectedOperationAddress = SelectedNodes[0]->MutableOperation;
-	const mu::EOpType OperationType = MutableModel->GetPrivate()->Program.GetOpType(SelectedOperationAddress);
-	const mu::EDataType OperationDataType = mu::GetOpDataType(OperationType);
+	const UE::Mutable::Private::EOpType OperationType = MutableModel->GetPrivate()->Program.GetOpType(SelectedOperationAddress);
+	const UE::Mutable::Private::EDataType OperationDataType = UE::Mutable::Private::GetOpDataType(OperationType);
 	
 	switch (OperationDataType)
 	{
-	case mu::EDataType::Layout:
+	case UE::Mutable::Private::EDataType::Layout:
 	{
 		// Create or reuse the UI
 		PrepareLayoutViewer();
 		break;
 	}
 
-	case mu::EDataType::Image:
+	case UE::Mutable::Private::EDataType::Image:
 	{
 		// Create or reuse the UI
 		bSelectedOperationIsImage = true;
@@ -2050,21 +2050,21 @@ void SMutableCodeViewer::OnSelectionChanged(TSharedPtr<FMutableCodeTreeElement>
 		break;
 	}
 
-	case mu::EDataType::Mesh:
+	case UE::Mutable::Private::EDataType::Mesh:
 	{
 		// Create or reuse the UI
 		PrepareMeshViewer();
 		break;
 	}
 
-	case mu::EDataType::Instance:
+	case UE::Mutable::Private::EDataType::Instance:
 	{
 		// Create or reuse the UI
 		PrepareInstanceViewer();
 		break;
 	}
 
-	case mu::EDataType::Scalar:
+	case UE::Mutable::Private::EDataType::Scalar:
 	{
 		// Create or reuse the UI
 		if (!PreviewScalarViewer)
@@ -2076,14 +2076,14 @@ void SMutableCodeViewer::OnSelectionChanged(TSharedPtr<FMutableCodeTreeElement>
 		break;
 	}
 
-	case mu::EDataType::String:
+	case UE::Mutable::Private::EDataType::String:
 	{
 		// Create or reuse the UI
 		PrepareStringViewer();
 		break;
 	}
 
-	case mu::EDataType::Color:
+	case UE::Mutable::Private::EDataType::Color:
 	{
 		// Create or reuse the UI
 		if (!PreviewColorViewer)
@@ -2095,7 +2095,7 @@ void SMutableCodeViewer::OnSelectionChanged(TSharedPtr<FMutableCodeTreeElement>
 		break;
 	}
 
-	case mu::EDataType::Int:
+	case UE::Mutable::Private::EDataType::Int:
 	{
 		// Create or reuse the UI
 		if (!PreviewIntViewer)
@@ -2107,7 +2107,7 @@ void SMutableCodeViewer::OnSelectionChanged(TSharedPtr<FMutableCodeTreeElement>
 		break;
 	}
 
-	case mu::EDataType::Bool:
+	case UE::Mutable::Private::EDataType::Bool:
 	{
 		// Create or reuse the UI
 		if (!PreviewBoolViewer)
@@ -2119,7 +2119,7 @@ void SMutableCodeViewer::OnSelectionChanged(TSharedPtr<FMutableCodeTreeElement>
 		break;
 	}
 
-	case mu::EDataType::Projector:
+	case UE::Mutable::Private::EDataType::Projector:
 	{
 		// Create or reuse the UI
 		PrepareProjectorViewer();
@@ -2258,7 +2258,7 @@ void SMutableCodeViewer::ClearHighlightedItems()
 
 void SMutableCodeViewer::TreeExpandElements(TArray<TSharedPtr<FMutableCodeTreeElement>>& InElementsToExpand, 
 	bool bForceExpandDuplicates /*= false*/,
-	mu::EDataType FilteringDataType /*= mu::EDataType::None*/,
+	UE::Mutable::Private::EDataType FilteringDataType /*= UE::Mutable::Private::EDataType::None*/,
 	TSharedPtr<FProcessedOperationsBuffer> InExpandedOperationsBuffer /* = nullptr */)
 {
 	if (InElementsToExpand.IsEmpty())
@@ -2276,8 +2276,8 @@ void SMutableCodeViewer::TreeExpandElements(TArray<TSharedPtr<FMutableCodeTreeEl
 
 	// Load references to the arrays containing all the operations already worked on during another recursive call to this
 	// method
-	TArray<mu::OP::ADDRESS>& AlreadyExpandedOriginalOperations = InExpandedOperationsBuffer->ExpandedOriginalOperations;
-	TArray<mu::OP::ADDRESS>& AlreadyExpandedDuplicatedOperations = InExpandedOperationsBuffer->ExpandedDuplicatedOperations;
+	TArray<UE::Mutable::Private::OP::ADDRESS>& AlreadyExpandedOriginalOperations = InExpandedOperationsBuffer->ExpandedOriginalOperations;
+	TArray<UE::Mutable::Private::OP::ADDRESS>& AlreadyExpandedDuplicatedOperations = InExpandedOperationsBuffer->ExpandedDuplicatedOperations;
 	
 	// Array containing the children object found on Item.
 	TArray<TSharedPtr<FMutableCodeTreeElement>> Children;
@@ -2291,13 +2291,13 @@ void SMutableCodeViewer::TreeExpandElements(TArray<TSharedPtr<FMutableCodeTreeEl
 		check(Item);
 
 		// Identifier of the element. May be repeated if there are elements duplicating another element
-		const mu::OP::ADDRESS OperationAddress = Item->MutableOperation;
+		const UE::Mutable::Private::OP::ADDRESS OperationAddress = Item->MutableOperation;
 		
 		// Filter the elements being expanded if the user has defined a desired EDataType
-		if (FilteringDataType != mu::EDataType::None)
+		if (FilteringDataType != UE::Mutable::Private::EDataType::None)
 		{
-			const mu::EOpType OperationType = Item->MutableModel->GetPrivate()->Program.GetOpType(OperationAddress);
-			const mu::EDataType OperationDataType = mu::GetOpDataType(OperationType);
+			const UE::Mutable::Private::EOpType OperationType = Item->MutableModel->GetPrivate()->Program.GetOpType(OperationAddress);
+			const UE::Mutable::Private::EDataType OperationDataType = UE::Mutable::Private::GetOpDataType(OperationType);
 
 			// If it is not of the desired type then ignore it and continue to the next pending element
 			if (OperationDataType != FilteringDataType)
@@ -2426,7 +2426,7 @@ void SMutableCodeViewer::TreeExpandUnique()
 void SMutableCodeViewer::TreeExpandInstance()
 {
 	// Expand only the items that match the datatype provided
-	TreeExpandElements(RootNodes,false, mu::EDataType::Instance);
+	TreeExpandElements(RootNodes,false, UE::Mutable::Private::EDataType::Instance);
 }
 
 
@@ -2435,20 +2435,20 @@ void SMutableCodeViewer::CacheRootNodeAddresses()
 	check (MutableModel);
 	check (RootNodeAddresses.IsEmpty())
 
-	TArray<mu::OP::ADDRESS> FoundRootNodeAddresses;
+	TArray<UE::Mutable::Private::OP::ADDRESS> FoundRootNodeAddresses;
 	
-	const mu::FModel::Private* ModelPrivate = MutableModel->GetPrivate();
+	const UE::Mutable::Private::FModel::Private* ModelPrivate = MutableModel->GetPrivate();
 	const int32 StateCount = ModelPrivate->Program.States.Num();
 	for ( int32 StateIndex=0; StateIndex < StateCount; ++StateIndex )
 	{
-		const mu::FProgram::FState& State = ModelPrivate->Program.States[StateIndex];
+		const UE::Mutable::Private::FProgram::FState& State = ModelPrivate->Program.States[StateIndex];
 		FoundRootNodeAddresses.Add(State.Root);
 	}
 
 	RootNodeAddresses = MoveTemp(FoundRootNodeAddresses);
 }
 
- void SMutableCodeViewer::CacheAddressesRelatedWithConstantResource(const mu::EDataType ConstantDataType,
+ void SMutableCodeViewer::CacheAddressesRelatedWithConstantResource(const UE::Mutable::Private::EDataType ConstantDataType,
 	const int32 IndexOnConstantsArray)
 {
 	check(MutableModel);
@@ -2468,7 +2468,7 @@ void SMutableCodeViewer::CacheRootNodeAddresses()
 	SearchPayload.SetupRootBatch(RootNodeAddresses);
 	
 	// Main update procedure run for the targeted state and the targeted parameter values
-	const mu::FProgram& Program = MutableModel->GetPrivate()->Program; 
+	const UE::Mutable::Private::FProgram& Program = MutableModel->GetPrivate()->Program; 
 	GetOperationsReferencingConstantResource(ConstantDataType,IndexOnConstantsArray,SearchPayload, Program);
 	
 	// At this point we did get all the addresses of operations that do involve the usage of our resource
@@ -2494,10 +2494,10 @@ void SMutableCodeViewer::CacheRootNodeAddresses()
 
 
 void SMutableCodeViewer::GetOperationsReferencingConstantResource(
-	const mu::EDataType ConstantDataType,
+	const UE::Mutable::Private::EDataType ConstantDataType,
 	const int32 IndexOnConstantsArray,
 	FElementsSearchCache& InSearchPayload,
-	const mu::FProgram& InProgram)
+	const UE::Mutable::Private::FProgram& InProgram)
 {
 	// next batch of addresses to be explored 
 	TArray<FItemCacheKey> NextBatchAddressesData;
@@ -2505,7 +2505,7 @@ void SMutableCodeViewer::GetOperationsReferencingConstantResource(
 	for	(int32 ParentIndex = 0 ; ParentIndex < InSearchPayload.BatchData.Num(); ParentIndex++)
 	{
 		// Get one of the previous run "children" and treat as a parent to get it's children and process them
-		const mu::OP::ADDRESS& ParentAddress = InSearchPayload.BatchData[ParentIndex].Child;
+		const UE::Mutable::Private::OP::ADDRESS& ParentAddress = InSearchPayload.BatchData[ParentIndex].Child;
 		
 		// Cache if same data type and we share the same address (means this op is pointing at the provided resource)
 		// It will cache duplicated entries
@@ -2533,14 +2533,14 @@ void SMutableCodeViewer::GetOperationsReferencingConstantResource(
 }
 
 bool SMutableCodeViewer::IsConstantResourceUsedByOperation(const int32 IndexOnConstantsArray,
-	const mu::EDataType ConstantDataType, const mu::OP::ADDRESS OperationAddress, const mu::FProgram& InProgram) const
+	const UE::Mutable::Private::EDataType ConstantDataType, const UE::Mutable::Private::OP::ADDRESS OperationAddress, const UE::Mutable::Private::FProgram& InProgram) const
 {
 	// Cache the current operation type to know where to look and what to check
-	const mu::EOpType OperationType = InProgram.GetOpType(OperationAddress);
+	const UE::Mutable::Private::EOpType OperationType = InProgram.GetOpType(OperationAddress);
 
 	// Making usage of the operation data type is not valid since some operations while return one type do, in fact,
 	// contain data from other types (like the mesh constant for example that contains mesh, skeleton and physics asset)
-	// const mu::EDataType DataType = mu::GetOpDataType(OperationType);
+	// const UE::Mutable::Private::EDataType DataType = UE::Mutable::Private::GetOpDataType(OperationType);
 	// if (DataType != ConstantDataType)
 	// {
 	// 	return false;
@@ -2552,44 +2552,44 @@ bool SMutableCodeViewer::IsConstantResourceUsedByOperation(const int32 IndexOnCo
 	// Check if the operation data type is compatible with the type of resources we are providing
 	switch (ConstantDataType)
 	{
-		case mu::EDataType::String:
+		case UE::Mutable::Private::EDataType::String:
 			{
 				// TIP: To know if they represent a constant value check the code on code runner to see if they read from the constants array
-				if (OperationType == mu::EOpType::ST_CONSTANT)
+				if (OperationType == UE::Mutable::Private::EOpType::ST_CONSTANT)
 				{
-					bResourceLocated = IndexOnConstantsArray == InProgram.GetOpArgs<mu::OP::ResourceConstantArgs>(OperationAddress).value;
+					bResourceLocated = IndexOnConstantsArray == InProgram.GetOpArgs<UE::Mutable::Private::OP::ResourceConstantArgs>(OperationAddress).value;
 				}
-				else if (OperationType == mu::EOpType::IN_ADDSTRING)
+				else if (OperationType == UE::Mutable::Private::EOpType::IN_ADDSTRING)
 				{
-					bResourceLocated = IndexOnConstantsArray == InProgram.GetOpArgs<mu::OP::InstanceAddArgs>(OperationAddress).name;
+					bResourceLocated = IndexOnConstantsArray == InProgram.GetOpArgs<UE::Mutable::Private::OP::InstanceAddArgs>(OperationAddress).name;
 				}
-				else if (OperationType == mu::EOpType::IN_ADDMESH)
+				else if (OperationType == UE::Mutable::Private::EOpType::IN_ADDMESH)
 				{
-					bResourceLocated = IndexOnConstantsArray == InProgram.GetOpArgs<mu::OP::InstanceAddArgs>(OperationAddress).name;
+					bResourceLocated = IndexOnConstantsArray == InProgram.GetOpArgs<UE::Mutable::Private::OP::InstanceAddArgs>(OperationAddress).name;
 				}
-				else if(OperationType == mu::EOpType::IN_ADDIMAGE)
+				else if(OperationType == UE::Mutable::Private::EOpType::IN_ADDIMAGE)
 				{
-					bResourceLocated = IndexOnConstantsArray == InProgram.GetOpArgs<mu::OP::InstanceAddArgs>(OperationAddress).name;
+					bResourceLocated = IndexOnConstantsArray == InProgram.GetOpArgs<UE::Mutable::Private::OP::InstanceAddArgs>(OperationAddress).name;
 				}
-				else if (OperationType == mu::EOpType::IN_ADDVECTOR)
+				else if (OperationType == UE::Mutable::Private::EOpType::IN_ADDVECTOR)
 				{
-					bResourceLocated = IndexOnConstantsArray == InProgram.GetOpArgs<mu::OP::InstanceAddArgs>(OperationAddress).name;
+					bResourceLocated = IndexOnConstantsArray == InProgram.GetOpArgs<UE::Mutable::Private::OP::InstanceAddArgs>(OperationAddress).name;
 				}
-				else if (OperationType == mu::EOpType::IN_ADDSCALAR)
+				else if (OperationType == UE::Mutable::Private::EOpType::IN_ADDSCALAR)
 				{
-					bResourceLocated = IndexOnConstantsArray == InProgram.GetOpArgs<mu::OP::InstanceAddArgs>(OperationAddress).name;
+					bResourceLocated = IndexOnConstantsArray == InProgram.GetOpArgs<UE::Mutable::Private::OP::InstanceAddArgs>(OperationAddress).name;
 				}
-				else if (OperationType == mu::EOpType::IN_ADDCOMPONENT)
+				else if (OperationType == UE::Mutable::Private::EOpType::IN_ADDCOMPONENT)
 				{
-					bResourceLocated = IndexOnConstantsArray == InProgram.GetOpArgs<mu::OP::InstanceAddArgs>(OperationAddress).name;
+					bResourceLocated = IndexOnConstantsArray == InProgram.GetOpArgs<UE::Mutable::Private::OP::InstanceAddArgs>(OperationAddress).name;
 				}
-				else if (OperationType == mu::EOpType::IN_ADDSURFACE)
+				else if (OperationType == UE::Mutable::Private::EOpType::IN_ADDSURFACE)
 				{
-					bResourceLocated = IndexOnConstantsArray == InProgram.GetOpArgs<mu::OP::InstanceAddArgs>(OperationAddress).name;
+					bResourceLocated = IndexOnConstantsArray == InProgram.GetOpArgs<UE::Mutable::Private::OP::InstanceAddArgs>(OperationAddress).name;
 				}
-				else if (OperationType == mu::EOpType::IN_CONDITIONAL)
+				else if (OperationType == UE::Mutable::Private::EOpType::IN_CONDITIONAL)
 				{
-					mu::OP::ConditionalArgs Arguments = InProgram.GetOpArgs<mu::OP::ConditionalArgs>(OperationAddress);
+					UE::Mutable::Private::OP::ConditionalArgs Arguments = InProgram.GetOpArgs<UE::Mutable::Private::OP::ConditionalArgs>(OperationAddress);
 					
 					if (IndexOnConstantsArray == Arguments.condition)
 					{
@@ -2609,13 +2609,13 @@ bool SMutableCodeViewer::IsConstantResourceUsedByOperation(const int32 IndexOnCo
 						break;
 					}
 				}
-				else if (OperationType == mu::EOpType::IN_ADDEXTENSIONDATA)
+				else if (OperationType == UE::Mutable::Private::EOpType::IN_ADDEXTENSIONDATA)
 				{
-					bResourceLocated = IndexOnConstantsArray == InProgram.GetOpArgs<mu::OP::InstanceAddExtensionDataArgs>(OperationAddress).ExtensionDataName;
+					bResourceLocated = IndexOnConstantsArray == InProgram.GetOpArgs<UE::Mutable::Private::OP::InstanceAddExtensionDataArgs>(OperationAddress).ExtensionDataName;
 				}
-				else if (OperationType == mu::EOpType::ME_BINDSHAPE)
+				else if (OperationType == UE::Mutable::Private::EOpType::ME_BINDSHAPE)
 				{
-					mu::OP::MeshBindShapeArgs Arguments = InProgram.GetOpArgs<mu::OP::MeshBindShapeArgs>(OperationAddress);
+					UE::Mutable::Private::OP::MeshBindShapeArgs Arguments = InProgram.GetOpArgs<UE::Mutable::Private::OP::MeshBindShapeArgs>(OperationAddress);
 					const uint8_t* Data = InProgram.GetOpArgsPointer(OperationAddress);
 
 					// Bones are stored after the args
@@ -2638,13 +2638,13 @@ bool SMutableCodeViewer::IsConstantResourceUsedByOperation(const int32 IndexOnCo
 						Data += sizeof(int32);
 					}
 				}
-				else if (OperationType == mu::EOpType::ME_ADDMETADATA)
+				else if (OperationType == UE::Mutable::Private::EOpType::ME_ADDMETADATA)
 				{
 					 const uint8* OpData = InProgram.GetOpArgsPointer(OperationAddress);
 					
-					 mu::OP::ADDRESS SourceAddress;
-					 FMemory::Memcpy(&SourceAddress, OpData, sizeof(mu::OP::ADDRESS));
-					 OpData += sizeof(mu::OP::ADDRESS);
+					 UE::Mutable::Private::OP::ADDRESS SourceAddress;
+					 FMemory::Memcpy(&SourceAddress, OpData, sizeof(UE::Mutable::Private::OP::ADDRESS));
+					 OpData += sizeof(UE::Mutable::Private::OP::ADDRESS);
 					
 					 uint16 TagCount;
 					 FMemory::Memcpy(&TagCount, OpData, sizeof(uint16));
@@ -2666,22 +2666,22 @@ bool SMutableCodeViewer::IsConstantResourceUsedByOperation(const int32 IndexOnCo
 				break;
 			}
 
-		case mu::EDataType::Image:
+		case UE::Mutable::Private::EDataType::Image:
 			{
-				if (OperationType == mu::EOpType::IM_CONSTANT)
+				if (OperationType == UE::Mutable::Private::EOpType::IM_CONSTANT)
 				{
-					bResourceLocated = IndexOnConstantsArray == InProgram.GetOpArgs<mu::OP::ResourceConstantArgs>(OperationAddress).value;
+					bResourceLocated = IndexOnConstantsArray == InProgram.GetOpArgs<UE::Mutable::Private::OP::ResourceConstantArgs>(OperationAddress).value;
 				} 
 				
 				break;
 			}
 
-		case mu::EDataType::Mesh:
+		case UE::Mutable::Private::EDataType::Mesh:
 			{
-				if (OperationType == mu::EOpType::ME_CONSTANT)
+				if (OperationType == UE::Mutable::Private::EOpType::ME_CONSTANT)
 				{
-					mu::OP::MeshConstantArgs Args = InProgram.GetOpArgs<mu::OP::MeshConstantArgs>(OperationAddress);
-					mu::FConstantResourceIndex ResourceIndex = *reinterpret_cast<const mu::FConstantResourceIndex*>(&Args.Value);
+					UE::Mutable::Private::OP::MeshConstantArgs Args = InProgram.GetOpArgs<UE::Mutable::Private::OP::MeshConstantArgs>(OperationAddress);
+					UE::Mutable::Private::FConstantResourceIndex ResourceIndex = *reinterpret_cast<const UE::Mutable::Private::FConstantResourceIndex*>(&Args.Value);
 					
 					if (!ResourceIndex.Streamable)
 					{
@@ -2690,7 +2690,7 @@ bool SMutableCodeViewer::IsConstantResourceUsedByOperation(const int32 IndexOnCo
 					else
 					{
 						int32 DebuggerIndexOffset = InProgram.ConstantMeshesPermanent.Num();
-						for (const TPair<uint32, TSharedPtr<const mu::FMesh>>& Entry : InProgram.ConstantMeshesStreamed)
+						for (const TPair<uint32, TSharedPtr<const UE::Mutable::Private::FMesh>>& Entry : InProgram.ConstantMeshesStreamed)
 						{
 							if (Entry.Key == ResourceIndex.Index)
 							{
@@ -2704,50 +2704,50 @@ bool SMutableCodeViewer::IsConstantResourceUsedByOperation(const int32 IndexOnCo
 				break;
 			}
 
-		case mu::EDataType::Layout:
+		case UE::Mutable::Private::EDataType::Layout:
 			{
-				if (OperationType == mu::EOpType::LA_CONSTANT)
+				if (OperationType == UE::Mutable::Private::EOpType::LA_CONSTANT)
 				{
-					bResourceLocated = IndexOnConstantsArray == InProgram.GetOpArgs<mu::OP::ResourceConstantArgs>(OperationAddress).value;
+					bResourceLocated = IndexOnConstantsArray == InProgram.GetOpArgs<UE::Mutable::Private::OP::ResourceConstantArgs>(OperationAddress).value;
 				}
 				break;
 			}
 
-		case mu::EDataType::Projector:
+		case UE::Mutable::Private::EDataType::Projector:
 			{
-				if (OperationType == mu::EOpType::PR_CONSTANT)
+				if (OperationType == UE::Mutable::Private::EOpType::PR_CONSTANT)
 				{
-					bResourceLocated = IndexOnConstantsArray == InProgram.GetOpArgs<mu::OP::ResourceConstantArgs>(OperationAddress).value;
+					bResourceLocated = IndexOnConstantsArray == InProgram.GetOpArgs<UE::Mutable::Private::OP::ResourceConstantArgs>(OperationAddress).value;
 				}
 				break;
 			}
 
-		case mu::EDataType::Matrix:
+		case UE::Mutable::Private::EDataType::Matrix:
 			{
-				if (OperationType == mu::EOpType::ME_TRANSFORM)
+				if (OperationType == UE::Mutable::Private::EOpType::ME_TRANSFORM)
 				{
-					bResourceLocated = IndexOnConstantsArray == InProgram.GetOpArgs<mu::OP::MeshTransformArgs>(OperationAddress).matrix;
+					bResourceLocated = IndexOnConstantsArray == InProgram.GetOpArgs<UE::Mutable::Private::OP::MeshTransformArgs>(OperationAddress).matrix;
 				}
-				else if (OperationType == mu::EOpType::ME_TRANSFORMWITHMESH)
+				else if (OperationType == UE::Mutable::Private::EOpType::ME_TRANSFORMWITHMESH)
 				{
-					bResourceLocated = IndexOnConstantsArray == InProgram.GetOpArgs<mu::OP::MeshTransformWithinMeshArgs>(OperationAddress).matrix;
+					bResourceLocated = IndexOnConstantsArray == InProgram.GetOpArgs<UE::Mutable::Private::OP::MeshTransformWithinMeshArgs>(OperationAddress).matrix;
 				}
-				else if (OperationType == mu::EOpType::ME_TRANSFORMWITHBONE)
+				else if (OperationType == UE::Mutable::Private::EOpType::ME_TRANSFORMWITHBONE)
 				{
-					bResourceLocated = IndexOnConstantsArray == InProgram.GetOpArgs<mu::OP::MeshTransformWithBoneArgs>(OperationAddress).Matrix;
+					bResourceLocated = IndexOnConstantsArray == InProgram.GetOpArgs<UE::Mutable::Private::OP::MeshTransformWithBoneArgs>(OperationAddress).Matrix;
 				}
-				else if (OperationType == mu::EOpType::MA_CONSTANT)
+				else if (OperationType == UE::Mutable::Private::EOpType::MA_CONSTANT)
 				{
-					bResourceLocated = IndexOnConstantsArray == InProgram.GetOpArgs<mu::OP::MatrixConstantArgs>(OperationAddress).value;
+					bResourceLocated = IndexOnConstantsArray == InProgram.GetOpArgs<UE::Mutable::Private::OP::MatrixConstantArgs>(OperationAddress).value;
 				}
 				break;
 			}
 
-		case mu::EDataType::Shape:
+		case UE::Mutable::Private::EDataType::Shape:
 			{
-				if (OperationType == mu::EOpType::ME_CLIPMORPHPLANE)
+				if (OperationType == UE::Mutable::Private::EOpType::ME_CLIPMORPHPLANE)
 				{
-					const mu::OP::MeshClipMorphPlaneArgs Arguments =  InProgram.GetOpArgs<mu::OP::MeshClipMorphPlaneArgs>(OperationAddress);
+					const UE::Mutable::Private::OP::MeshClipMorphPlaneArgs Arguments =  InProgram.GetOpArgs<UE::Mutable::Private::OP::MeshClipMorphPlaneArgs>(OperationAddress);
 					
 					// Morph shape
 					bResourceLocated = IndexOnConstantsArray == Arguments.MorphShape;
@@ -2759,44 +2759,44 @@ bool SMutableCodeViewer::IsConstantResourceUsedByOperation(const int32 IndexOnCo
 					if (Arguments.VertexSelectionType == EClipVertexSelectionType::Shape)
 					{
 						// Selection Shape
-						bResourceLocated = IndexOnConstantsArray == InProgram.GetOpArgs<mu::OP::MeshClipMorphPlaneArgs>(OperationAddress).VertexSelectionShapeOrBone;
+						bResourceLocated = IndexOnConstantsArray == InProgram.GetOpArgs<UE::Mutable::Private::OP::MeshClipMorphPlaneArgs>(OperationAddress).VertexSelectionShapeOrBone;
 					}
 				}
 				break;
 			}
 		
-		case mu::EDataType::Curve:
+		case UE::Mutable::Private::EDataType::Curve:
 			{
-				if (OperationType == mu::EOpType::SC_CURVE)
+				if (OperationType == UE::Mutable::Private::EOpType::SC_CURVE)
 				{
-					bResourceLocated = IndexOnConstantsArray == InProgram.GetOpArgs<mu::OP::ScalarCurveArgs>(OperationAddress).curve;
+					bResourceLocated = IndexOnConstantsArray == InProgram.GetOpArgs<UE::Mutable::Private::OP::ScalarCurveArgs>(OperationAddress).curve;
 				}
 				break;
 			}
 
-		case mu::EDataType::Skeleton:
+		case UE::Mutable::Private::EDataType::Skeleton:
 			{
-				if (OperationType == mu::EOpType::ME_CONSTANT)
+				if (OperationType == UE::Mutable::Private::EOpType::ME_CONSTANT)
 				{
-					bResourceLocated = IndexOnConstantsArray == InProgram.GetOpArgs<mu::OP::MeshConstantArgs>(OperationAddress).Skeleton;
+					bResourceLocated = IndexOnConstantsArray == InProgram.GetOpArgs<UE::Mutable::Private::OP::MeshConstantArgs>(OperationAddress).Skeleton;
 				}
-				else if (OperationType == mu::EOpType::ME_SETSKELETON)
+				else if (OperationType == UE::Mutable::Private::EOpType::ME_SETSKELETON)
 				{
-					bResourceLocated = IndexOnConstantsArray == InProgram.GetOpArgs<mu::OP::MeshSetSkeletonArgs>(OperationAddress).Skeleton;
+					bResourceLocated = IndexOnConstantsArray == InProgram.GetOpArgs<UE::Mutable::Private::OP::MeshSetSkeletonArgs>(OperationAddress).Skeleton;
 				}
 				break;
 			}
-		case mu::EDataType::PhysicsAsset:
+		case UE::Mutable::Private::EDataType::PhysicsAsset:
 			{
-				if (OperationType == mu::EOpType::ME_CONSTANT)
+				if (OperationType == UE::Mutable::Private::EOpType::ME_CONSTANT)
 				{
-					bResourceLocated = IndexOnConstantsArray == InProgram.GetOpArgs<mu::OP::MeshConstantArgs>(OperationAddress).Value;
+					bResourceLocated = IndexOnConstantsArray == InProgram.GetOpArgs<UE::Mutable::Private::OP::MeshConstantArgs>(OperationAddress).Value;
 				}
 				break;
 			}
 
 		// Invalid types
-		case mu::EDataType::None:
+		case UE::Mutable::Private::EDataType::None:
 		default:
 			{
 				checkNoEntry();
@@ -2832,49 +2832,49 @@ void SMutableCodeViewer::Tick(const FGeometry& AllottedGeometry, const double In
 
 	bIsPreviewPendingUpdate = false;
 
-	const mu::EOpType OperationType = MutableModel->GetPrivate()->Program.GetOpType(SelectedOperationAddress);
-	const mu::EDataType OperationDataType = mu::GetOpDataType(OperationType);
+	const UE::Mutable::Private::EOpType OperationType = MutableModel->GetPrivate()->Program.GetOpType(SelectedOperationAddress);
+	const UE::Mutable::Private::EDataType OperationDataType = UE::Mutable::Private::GetOpDataType(OperationType);
 
-	mu::FSettings Settings;
-	const TSharedPtr<mu::FSystem> System = MakeShared<mu::FSystem>(Settings);
+	UE::Mutable::Private::FSettings Settings;
+	const TSharedPtr<UE::Mutable::Private::FSystem> System = MakeShared<UE::Mutable::Private::FSystem>(Settings);
 	
 	System->GetPrivate()->BeginBuild(MutableModel, Object->GetPrivate()->MeshIdRegistry, Object->GetPrivate()->ImageIdRegistry, Object->GetPrivate()->MaterialIdRegistry);
 
 	switch (OperationDataType)
 	{
-	case mu::EDataType::Layout:
+	case UE::Mutable::Private::EDataType::Layout:
 	{
 		check(PreviewLayoutViewer);
-		TSharedPtr<const mu::FLayout> MutableLayout = System->GetPrivate()->BuildLayout(ExternalResourceProvider, MutableModel, PreviewParameters.Get(), SelectedOperationAddress);
+		TSharedPtr<const UE::Mutable::Private::FLayout> MutableLayout = System->GetPrivate()->BuildLayout(ExternalResourceProvider, MutableModel, PreviewParameters.Get(), SelectedOperationAddress);
 		PreviewLayoutViewer->SetLayout(MutableLayout);
 		break;
 	}
 
-	case mu::EDataType::Image:
+	case UE::Mutable::Private::EDataType::Image:
 	{
 		check(PreviewImageViewer);
-		TSharedPtr<const mu::FImage> MutableImage = System->GetPrivate()->BuildImage(ExternalResourceProvider, MutableModel, PreviewParameters.Get(), SelectedOperationAddress, MipsToSkip, 0);
+		TSharedPtr<const UE::Mutable::Private::FImage> MutableImage = System->GetPrivate()->BuildImage(ExternalResourceProvider, MutableModel, PreviewParameters.Get(), SelectedOperationAddress, MipsToSkip, 0);
 		PreviewImageViewer->SetImage(MutableImage, 0);
 		break;
 	}
 
-	case mu::EDataType::Mesh:
+	case UE::Mutable::Private::EDataType::Mesh:
 	{
 		check(PreviewMeshViewer);
-		TSharedPtr<const mu::FMesh> MutableMesh = System->GetPrivate()->BuildMesh(ExternalResourceProvider, MutableModel, PreviewParameters.Get(), SelectedOperationAddress, mu::EMeshContentFlags::AllFlags);
+		TSharedPtr<const UE::Mutable::Private::FMesh> MutableMesh = System->GetPrivate()->BuildMesh(ExternalResourceProvider, MutableModel, PreviewParameters.Get(), SelectedOperationAddress, UE::Mutable::Private::EMeshContentFlags::AllFlags);
 		PreviewMeshViewer->SetMesh(MutableMesh);
 		break;
 	}
 
-	case mu::EDataType::Instance:
+	case UE::Mutable::Private::EDataType::Instance:
 	{
 		check(PreviewInstanceViewer);
-		TSharedPtr<const mu::FInstance> MutableInstance = System->GetPrivate()->BuildInstance(ExternalResourceProvider, MutableModel, PreviewParameters.Get(), SelectedOperationAddress);
+		TSharedPtr<const UE::Mutable::Private::FInstance> MutableInstance = System->GetPrivate()->BuildInstance(ExternalResourceProvider, MutableModel, PreviewParameters.Get(), SelectedOperationAddress);
 		PreviewInstanceViewer->SetInstance(MutableInstance, MutableModel, PreviewParameters, *System);
 		break;
 	}
 
-	case mu::EDataType::Bool:
+	case UE::Mutable::Private::EDataType::Bool:
 	{
 		check(PreviewBoolViewer);
 		const bool MutableBool = System->GetPrivate()->BuildBool(ExternalResourceProvider, MutableModel, PreviewParameters.Get(), SelectedOperationAddress);
@@ -2882,7 +2882,7 @@ void SMutableCodeViewer::Tick(const FGeometry& AllottedGeometry, const double In
 		break;
 	}
 
-	case mu::EDataType::Int:
+	case UE::Mutable::Private::EDataType::Int:
 	{
 		check(PreviewIntViewer);
 		const int32 MutableInt = System->GetPrivate()->BuildInt(ExternalResourceProvider, MutableModel, PreviewParameters.Get(), SelectedOperationAddress);
@@ -2890,7 +2890,7 @@ void SMutableCodeViewer::Tick(const FGeometry& AllottedGeometry, const double In
 		break;
 	}
 
-	case mu::EDataType::Scalar:
+	case UE::Mutable::Private::EDataType::Scalar:
 	{
 		check(PreviewScalarViewer);
 		const float MutableScalar = System->GetPrivate()->BuildScalar(ExternalResourceProvider, MutableModel, PreviewParameters.Get(), SelectedOperationAddress);
@@ -2898,16 +2898,16 @@ void SMutableCodeViewer::Tick(const FGeometry& AllottedGeometry, const double In
 		break;
 	}
 
-	case mu::EDataType::String:
+	case UE::Mutable::Private::EDataType::String:
 	{
 		check(PreviewStringViewer);
-		TSharedPtr<const mu::String> MutableString = System->GetPrivate()->BuildString(ExternalResourceProvider, MutableModel, PreviewParameters.Get(), SelectedOperationAddress);
+		TSharedPtr<const UE::Mutable::Private::String> MutableString = System->GetPrivate()->BuildString(ExternalResourceProvider, MutableModel, PreviewParameters.Get(), SelectedOperationAddress);
 		const FText MutableText = FText::FromString(FString(MutableString->GetValue()));
 		PreviewStringViewer->SetString(MutableText);
 		break;
 	}
 
-	case mu::EDataType::Color:
+	case UE::Mutable::Private::EDataType::Color:
 	{
 		check(PreviewColorViewer);
 		FVector4f Color = System->GetPrivate()->BuildColour(ExternalResourceProvider, MutableModel, PreviewParameters.Get(), SelectedOperationAddress);
@@ -2915,10 +2915,10 @@ void SMutableCodeViewer::Tick(const FGeometry& AllottedGeometry, const double In
 		break;
 	}
 
-	case mu::EDataType::Projector:
+	case UE::Mutable::Private::EDataType::Projector:
 	{
 		check (PreviewProjectorViewer);
-		mu::FProjector Projector = System->GetPrivate()->BuildProjector(ExternalResourceProvider, MutableModel, PreviewParameters.Get(), SelectedOperationAddress);
+		UE::Mutable::Private::FProjector Projector = System->GetPrivate()->BuildProjector(ExternalResourceProvider, MutableModel, PreviewParameters.Get(), SelectedOperationAddress);
 		PreviewProjectorViewer->SetProjector(Projector);
 	}
 	
@@ -3016,28 +3016,28 @@ void SMutableCodeViewer::PreviewMutableString(const FString& InString)
 }
 
 
-void SMutableCodeViewer::PreviewMutableImage(TSharedPtr<const mu::FImage> InImagePtr)
+void SMutableCodeViewer::PreviewMutableImage(TSharedPtr<const UE::Mutable::Private::FImage> InImagePtr)
 {
 	PrepareImageViewer();
 	PreviewImageViewer->SetImage(InImagePtr,0);
 }
 
 
-void SMutableCodeViewer::PreviewMutableMesh(TSharedPtr<const mu::FMesh> InMeshPtr)
+void SMutableCodeViewer::PreviewMutableMesh(TSharedPtr<const UE::Mutable::Private::FMesh> InMeshPtr)
 {
 	PrepareMeshViewer();
 	PreviewMeshViewer->SetMesh(InMeshPtr);
 }
 
 
-void SMutableCodeViewer::PreviewMutableLayout(TSharedPtr<const mu::FLayout> Layout)
+void SMutableCodeViewer::PreviewMutableLayout(TSharedPtr<const UE::Mutable::Private::FLayout> Layout)
 {
 	PrepareLayoutViewer();
 	PreviewLayoutViewer->SetLayout(Layout);
 }
 
 
-void SMutableCodeViewer::PreviewMutableProjector(const mu::FProjector* Projector)
+void SMutableCodeViewer::PreviewMutableProjector(const UE::Mutable::Private::FProjector* Projector)
 {
 	if (!Projector)
 	{
@@ -3050,7 +3050,7 @@ void SMutableCodeViewer::PreviewMutableProjector(const mu::FProjector* Projector
 }
 
 
-void SMutableCodeViewer::PreviewMutableSkeleton(TSharedPtr<const mu::FSkeleton> Skeleton)
+void SMutableCodeViewer::PreviewMutableSkeleton(TSharedPtr<const UE::Mutable::Private::FSkeleton> Skeleton)
 {
 	if (!PreviewSkeletonViewer)
 	{
@@ -3076,7 +3076,7 @@ void SMutableCodeViewer::PreviewMutableCurve(const FRichCurve& Curve)
 }
 
 // TODO: Implement physics viewer
-void SMutableCodeViewer::PreviewMutablePhysics(TSharedPtr<const mu::FPhysicsBody> Physics)
+void SMutableCodeViewer::PreviewMutablePhysics(TSharedPtr<const UE::Mutable::Private::FPhysicsBody> Physics)
 {
 	UE_LOG(LogMutable, Warning, TEXT("Previewer for Mutable Physics not yet implemented"))
 }
@@ -3088,13 +3088,13 @@ void SMutableCodeViewer::PreviewMutableMatrix(const FMatrix44f& Mat)
 }
 
 // TODO: Implement shape viewer
-void SMutableCodeViewer::PreviewMutableShape(const mu::FShape* Shape)
+void SMutableCodeViewer::PreviewMutableShape(const UE::Mutable::Private::FShape* Shape)
 {
 	UE_LOG(LogMutable, Warning, TEXT("Previewer for Mutable Shapes not yet implemented"))
 }
 
 
-FMutableCodeTreeElement::FMutableCodeTreeElement(int32 InIndexOnTree, const int32& InMutableStateIndex, const TSharedPtr<mu::FModel, ESPMode::ThreadSafe>& InModel, mu::OP::ADDRESS InOperation, const FString& InCaption, const FSlateColor InLabelColor, const TSharedPtr<FMutableCodeTreeElement>* InDuplicatedOf)
+FMutableCodeTreeElement::FMutableCodeTreeElement(int32 InIndexOnTree, const int32& InMutableStateIndex, const TSharedPtr<UE::Mutable::Private::FModel, ESPMode::ThreadSafe>& InModel, UE::Mutable::Private::OP::ADDRESS InOperation, const FString& InCaption, const FSlateColor InLabelColor, const TSharedPtr<FMutableCodeTreeElement>* InDuplicatedOf)
 {
 	MutableModel = InModel;
 	MutableOperation = InOperation;
@@ -3131,11 +3131,11 @@ void FMutableCodeTreeElement::SetElementCurrentState(const int32& InStateIndex)
 
 	// Check for an out of bounds value
 	check(MutableModel);
-	mu::FProgram& MutableProgram = MutableModel->GetPrivate()->Program;
+	UE::Mutable::Private::FProgram& MutableProgram = MutableModel->GetPrivate()->Program;
 	check(InStateIndex >= 0 && InStateIndex < MutableProgram.States.Num());
 
 	CurrentMutableStateIndex = InStateIndex;
-	const mu::FProgram::FState& CurrentState = MutableProgram.States[CurrentMutableStateIndex];
+	const UE::Mutable::Private::FProgram::FState& CurrentState = MutableProgram.States[CurrentMutableStateIndex];
 
 	// Check if it is a dynamic resource
 	for (auto& DynamicResource : CurrentState.m_dynamicResources)
@@ -3160,126 +3160,126 @@ void FMutableCodeTreeElement::SetElementCurrentState(const int32& InStateIndex)
 
 void FMutableCodeTreeElement::GenerateLabelText()
 {
-	const mu::FProgram& Program = MutableModel->GetPrivate()->Program;
-	const mu::EOpType OperationType = Program.GetOpType(MutableOperation);
-	FString OpName = mu::s_opNames[static_cast<int32>(OperationType)];
+	const UE::Mutable::Private::FProgram& Program = MutableModel->GetPrivate()->Program;
+	const UE::Mutable::Private::EOpType OperationType = Program.GetOpType(MutableOperation);
+	FString OpName = UE::Mutable::Private::s_opNames[static_cast<int32>(OperationType)];
 	OpName.TrimEndInline();
 
 	// See if the operation type accepts additional information in the label
 	switch (OperationType)
 	{
-	case mu::EOpType::BO_PARAMETER:
-	case mu::EOpType::NU_PARAMETER:
-	case mu::EOpType::SC_PARAMETER:
-	case mu::EOpType::CO_PARAMETER:
-	case mu::EOpType::PR_PARAMETER:
-	case mu::EOpType::IM_PARAMETER:
-	case mu::EOpType::ST_PARAMETER:
+	case UE::Mutable::Private::EOpType::BO_PARAMETER:
+	case UE::Mutable::Private::EOpType::NU_PARAMETER:
+	case UE::Mutable::Private::EOpType::SC_PARAMETER:
+	case UE::Mutable::Private::EOpType::CO_PARAMETER:
+	case UE::Mutable::Private::EOpType::PR_PARAMETER:
+	case UE::Mutable::Private::EOpType::IM_PARAMETER:
+	case UE::Mutable::Private::EOpType::ST_PARAMETER:
 	{
-		mu::OP::ParameterArgs Args = Program.GetOpArgs<mu::OP::ParameterArgs>(MutableOperation);
+		UE::Mutable::Private::OP::ParameterArgs Args = Program.GetOpArgs<UE::Mutable::Private::OP::ParameterArgs>(MutableOperation);
 		OpName += TEXT(" ");
 		OpName += Program.Parameters[int32(Args.variable)].Name;
 		break;
 	}
 
-	case mu::EOpType::ME_PARAMETER:
+	case UE::Mutable::Private::EOpType::ME_PARAMETER:
 	{
-		mu::OP::MeshParameterArgs Args = Program.GetOpArgs<mu::OP::MeshParameterArgs>(MutableOperation);
+		UE::Mutable::Private::OP::MeshParameterArgs Args = Program.GetOpArgs<UE::Mutable::Private::OP::MeshParameterArgs>(MutableOperation);
 		OpName += FString::Printf( TEXT(" LOD %d Section %d of "), Args.LOD, Args.Section );
 		OpName += Program.Parameters[int32(Args.variable)].Name;
 		break;
 	}
 
-	case mu::EOpType::IM_SWIZZLE:
+	case UE::Mutable::Private::EOpType::IM_SWIZZLE:
 	{
-		mu::OP::ImageSwizzleArgs Args = Program.GetOpArgs<mu::OP::ImageSwizzleArgs>(MutableOperation);
+		UE::Mutable::Private::OP::ImageSwizzleArgs Args = Program.GetOpArgs<UE::Mutable::Private::OP::ImageSwizzleArgs>(MutableOperation);
 		OpName += TEXT(" ");
-		OpName += StringCast<TCHAR>(mu::TypeInfo::s_imageFormatName[int32(Args.format)]).Get();
+		OpName += StringCast<TCHAR>(UE::Mutable::Private::TypeInfo::s_imageFormatName[int32(Args.format)]).Get();
 		break;
 	}
 
-	case mu::EOpType::IM_PIXELFORMAT:
+	case UE::Mutable::Private::EOpType::IM_PIXELFORMAT:
 	{
-		mu::OP::ImagePixelFormatArgs Args = Program.GetOpArgs<mu::OP::ImagePixelFormatArgs>(MutableOperation);
+		UE::Mutable::Private::OP::ImagePixelFormatArgs Args = Program.GetOpArgs<UE::Mutable::Private::OP::ImagePixelFormatArgs>(MutableOperation);
 		OpName += TEXT(" ");
-		OpName += StringCast<TCHAR>(mu::TypeInfo::s_imageFormatName[int32(Args.format)]).Get();
+		OpName += StringCast<TCHAR>(UE::Mutable::Private::TypeInfo::s_imageFormatName[int32(Args.format)]).Get();
 		OpName += TEXT(" or ");
-		OpName += StringCast<TCHAR>(mu::TypeInfo::s_imageFormatName[int32(Args.formatIfAlpha)]).Get();
+		OpName += StringCast<TCHAR>(UE::Mutable::Private::TypeInfo::s_imageFormatName[int32(Args.formatIfAlpha)]).Get();
 		break;
 	}
 
-	case mu::EOpType::IM_MIPMAP:
+	case UE::Mutable::Private::EOpType::IM_MIPMAP:
 	{
-		mu::OP::ImageMipmapArgs Args = Program.GetOpArgs<mu::OP::ImageMipmapArgs>(MutableOperation);
+		UE::Mutable::Private::OP::ImageMipmapArgs Args = Program.GetOpArgs<UE::Mutable::Private::OP::ImageMipmapArgs>(MutableOperation);
 		OpName += FString::Printf(TEXT(" levels: %d-%d tail: %d"), Args.levels, Args.blockLevels, int32(Args.onlyTail));
 		break;
 	}
 
-	case mu::EOpType::IM_RESIZE:
+	case UE::Mutable::Private::EOpType::IM_RESIZE:
 	{
-		mu::OP::ImageResizeArgs Args = Program.GetOpArgs<mu::OP::ImageResizeArgs>(MutableOperation);
+		UE::Mutable::Private::OP::ImageResizeArgs Args = Program.GetOpArgs<UE::Mutable::Private::OP::ImageResizeArgs>(MutableOperation);
 		OpName += FString::Printf(TEXT(" %d x %d"), int32(Args.Size[0]), int32(Args.Size[1]));
 		break;
 	}
 
-	case mu::EOpType::IM_RESIZEREL:
+	case UE::Mutable::Private::EOpType::IM_RESIZEREL:
 	{
-		mu::OP::ImageResizeRelArgs Args = Program.GetOpArgs<mu::OP::ImageResizeRelArgs>(MutableOperation);
+		UE::Mutable::Private::OP::ImageResizeRelArgs Args = Program.GetOpArgs<UE::Mutable::Private::OP::ImageResizeRelArgs>(MutableOperation);
 		OpName += FString::Printf(TEXT(" %.3f x %.3f"), Args.Factor[0], Args.Factor[1]);
 		break;
 	}
 
-	case mu::EOpType::IM_MULTILAYER:
+	case UE::Mutable::Private::EOpType::IM_MULTILAYER:
 	{
-		mu::OP::ImageMultiLayerArgs Args = Program.GetOpArgs<mu::OP::ImageMultiLayerArgs>(MutableOperation);
+		UE::Mutable::Private::OP::ImageMultiLayerArgs Args = Program.GetOpArgs<UE::Mutable::Private::OP::ImageMultiLayerArgs>(MutableOperation);
 		OpName += TEXT(" rgb: ");
-		OpName += mu::TypeInfo::s_blendModeName[int32(Args.blendType)];
+		OpName += UE::Mutable::Private::TypeInfo::s_blendModeName[int32(Args.blendType)];
 		OpName += TEXT(", a: ");
-		OpName += mu::TypeInfo::s_blendModeName[int32(Args.blendTypeAlpha)];
+		OpName += UE::Mutable::Private::TypeInfo::s_blendModeName[int32(Args.blendTypeAlpha)];
 		OpName += FString::Printf(TEXT(" a from %d "), Args.BlendAlphaSourceChannel);
 		OpName += FString::Printf(TEXT(" range-id: %d"), Args.rangeId);
 		OpName += FString::Printf(TEXT(" mask-from-alpha: %d"), int32(Args.bUseMaskFromBlended));
 		break;
 	}
 
-	case mu::EOpType::IM_LAYER:
+	case UE::Mutable::Private::EOpType::IM_LAYER:
 	{
-		mu::OP::ImageLayerArgs Args = Program.GetOpArgs<mu::OP::ImageLayerArgs>(MutableOperation);
+		UE::Mutable::Private::OP::ImageLayerArgs Args = Program.GetOpArgs<UE::Mutable::Private::OP::ImageLayerArgs>(MutableOperation);
 		OpName += TEXT(" rgb: ");
-		OpName += mu::TypeInfo::s_blendModeName[int32(Args.blendType)];
+		OpName += UE::Mutable::Private::TypeInfo::s_blendModeName[int32(Args.blendType)];
 		OpName += TEXT(", a: ");
-		OpName += mu::TypeInfo::s_blendModeName[int32(Args.blendTypeAlpha)];
+		OpName += UE::Mutable::Private::TypeInfo::s_blendModeName[int32(Args.blendTypeAlpha)];
 		OpName += FString::Printf(TEXT(" a from %d "), Args.BlendAlphaSourceChannel);
 		OpName += FString::Printf(TEXT(" flags %d"), Args.flags);
 		break;
 	}
 
-	case mu::EOpType::IM_LAYERCOLOUR:
+	case UE::Mutable::Private::EOpType::IM_LAYERCOLOUR:
 	{
-		mu::OP::ImageLayerColourArgs Args = Program.GetOpArgs<mu::OP::ImageLayerColourArgs>(MutableOperation);
+		UE::Mutable::Private::OP::ImageLayerColourArgs Args = Program.GetOpArgs<UE::Mutable::Private::OP::ImageLayerColourArgs>(MutableOperation);
 		OpName += TEXT(" rgb: ");
-		OpName += mu::TypeInfo::s_blendModeName[int32(Args.blendType)];
+		OpName += UE::Mutable::Private::TypeInfo::s_blendModeName[int32(Args.blendType)];
 		OpName += TEXT(" a: ");
-		OpName += mu::TypeInfo::s_blendModeName[int32(Args.blendTypeAlpha)];
+		OpName += UE::Mutable::Private::TypeInfo::s_blendModeName[int32(Args.blendTypeAlpha)];
 		OpName += TEXT(" a from ");
 		OpName += FString::Printf(TEXT(" a from %d "), Args.BlendAlphaSourceChannel);
 		OpName += FString::Printf(TEXT(" flags %d"), Args.flags);
 		break;
 	}
 
-	case mu::EOpType::IM_PLAINCOLOUR:
+	case UE::Mutable::Private::EOpType::IM_PLAINCOLOUR:
 	{
-		mu::OP::ImagePlainColorArgs Args = Program.GetOpArgs<mu::OP::ImagePlainColorArgs>(MutableOperation);
+		UE::Mutable::Private::OP::ImagePlainColorArgs Args = Program.GetOpArgs<UE::Mutable::Private::OP::ImagePlainColorArgs>(MutableOperation);
 		OpName += TEXT(" format: ");
-		OpName += StringCast<TCHAR>(mu::TypeInfo::s_imageFormatName[int32(Args.Format)]).Get();
+		OpName += StringCast<TCHAR>(UE::Mutable::Private::TypeInfo::s_imageFormatName[int32(Args.Format)]).Get();
 		OpName += FString::Printf(TEXT(" size %d x %d"), Args.Size[0], Args.Size[1]);
 		OpName += FString::Printf(TEXT(" mips %d"), Args.LODs);
 		break;
 	}
 
-	case mu::EOpType::IN_ADDIMAGE:
+	case UE::Mutable::Private::EOpType::IN_ADDIMAGE:
 	{
-		mu::OP::InstanceAddArgs Args = Program.GetOpArgs<mu::OP::InstanceAddArgs>(MutableOperation);
+		UE::Mutable::Private::OP::InstanceAddArgs Args = Program.GetOpArgs<UE::Mutable::Private::OP::InstanceAddArgs>(MutableOperation);
 		if (Program.ConstantStrings.IsValidIndex(Args.name))
 		{
 			OpName += TEXT(" name: ");
