[Mutable] Change namespace name

diff --git a/Engine/Plugins/Mutable/Source/CustomizableObjectEditor/Private/MuCOE/SMutableCodeViewer.h b/Engine/Plugins/Mutable/Source/CustomizableObjectEditor/Private/MuCOE/SMutableCodeViewer.h
--- a/Engine/Plugins/Mutable/Source/CustomizableObjectEditor/Private/MuCOE/SMutableCodeViewer.h
+++ b/Engine/Plugins/Mutable/Source/CustomizableObjectEditor/Private/MuCOE/SMutableCodeViewer.h
@@ -41,8 +41,8 @@ class SMutableSkeletonViewer;
 class SMutableStringViewer;
 class STextComboBox;
 class SWidget;
-namespace mu { struct FProjector; }
-namespace mu { struct FShape; }
+namespace UE::Mutable::Private { struct FProjector; }
+namespace UE::Mutable::Private { struct FShape; }
 struct FGeometry;
 
 /** This widget shows the internal Mutable Code for debugging purposes. 
@@ -58,7 +58,7 @@ public:
 		
 	SLATE_END_ARGS()
 
-	void Construct(const FArguments& InArgs, UCustomizableObject& InObject, UCustomizableObjectInstance& InInstance, const TSharedPtr<mu::FModel>& InMutableModel);
+	void Construct(const FArguments& InArgs, UCustomizableObject& InObject, UCustomizableObjectInstance& InInstance, const TSharedPtr<UE::Mutable::Private::FModel>& InMutableModel);
 
 	// SWidget interface
 	virtual void Tick(const FGeometry& AllottedGeometry, const double InCurrentTime, const float InDeltaTime) override;
@@ -71,7 +71,7 @@ private:
 	TStrongObjectPtr<UCustomizableObjectInstance> Instance;
 	
 	/** The Mutable Model that we are showing. */
-	TSharedPtr<mu::FModel> MutableModel;
+	TSharedPtr<UE::Mutable::Private::FModel> MutableModel;
 
 	TSharedPtr<FUnrealMutableResourceProvider> ExternalResourceProvider;
 	
@@ -80,10 +80,10 @@ private:
 	TArray<TSoftObjectPtr<const UStreamableRenderAsset>> ReferencedMeshes;
 
 	/** Selected model operation for preview. */
-	mu::OP::ADDRESS SelectedOperationAddress = 0;
+	UE::Mutable::Private::OP::ADDRESS SelectedOperationAddress = 0;
 
 	/** Mutable parameters used in the preview. */
-	TSharedPtr<mu::FParameters> PreviewParameters;
+	TSharedPtr<UE::Mutable::Private::FParameters> PreviewParameters;
 
 	/** Widget showing the parameters that affect the current preview. */
 	TSharedPtr<SMutableParametersWidget> ParametersWidget;
@@ -148,8 +148,8 @@ private:
 	*/
 	struct FItemCacheKey
 	{
-		mu::OP::ADDRESS Parent=0;
-		mu::OP::ADDRESS Child=0;
+		UE::Mutable::Private::OP::ADDRESS Parent=0;
+		UE::Mutable::Private::OP::ADDRESS Child=0;
 		uint32 ChildIndexInParent=0;
 
 		friend FORCEINLINE bool operator == (const FItemCacheKey& A, const FItemCacheKey& B)
@@ -176,13 +176,13 @@ private:
 	TMap< FItemCacheKey, TSharedPtr<FMutableCodeTreeElement>> ItemCache;
 
 	/** Main tree item for each op. An op can be represented with multiple tree nodes if it is reachable from different paths. */
-	TMap< mu::OP::ADDRESS, TSharedPtr<FMutableCodeTreeElement>> MainItemPerOp;
+	TMap< UE::Mutable::Private::OP::ADDRESS, TSharedPtr<FMutableCodeTreeElement>> MainItemPerOp;
 
 	/** List with all the elements related to nodes displayed on the tree. */
 	TArray<TSharedPtr<FMutableCodeTreeElement>> TreeElements;
 
 	/** Array with all the elements that have been manually expanded by the user */
-	TMap< mu::OP::ADDRESS, TSharedPtr<FMutableCodeTreeElement>> ExpandedElements;
+	TMap< UE::Mutable::Private::OP::ADDRESS, TSharedPtr<FMutableCodeTreeElement>> ExpandedElements;
 
 	/** Before any UI operation generate all the elements that may be navigable over the tree. No children of duplicated 
 	 * addresses will be generated.
@@ -190,11 +190,11 @@ private:
 	void GenerateAllTreeElements();
 
 	/** Support GenerateAllTreeElements by generating the elements recursively */
-	void GenerateElementRecursive( const int32& InStateIndex, mu::OP::ADDRESS InParentAddress, const mu::FProgram& InProgram);
+	void GenerateElementRecursive( const int32& InStateIndex, UE::Mutable::Private::OP::ADDRESS InParentAddress, const UE::Mutable::Private::FProgram& InProgram);
 
 	
 	/** The addresses of the root operations. Cached when this object gets loaded on the Construct method of this slate */
-	TArray<mu::OP::ADDRESS> RootNodeAddresses;
+	TArray<UE::Mutable::Private::OP::ADDRESS> RootNodeAddresses;
 	
 	/** Utility method that provides you with the addresses of the root nodes of the program */
 	void CacheRootNodeAddresses();
@@ -203,32 +203,32 @@ private:
 	bool bIsElementHighlighted = false;
 
 	/** Operation that is currently being highlighted */
-	mu::OP::ADDRESS HighlightedOperation{};
+	UE::Mutable::Private::OP::ADDRESS HighlightedOperation{};
 
 	/*
 	 * Operations computational cost reference collections
 	 */
 
 	/** Collection with all very expensive to run operation types */
-	const TArray<mu::EOpType> VeryExpensiveOperationTypes
+	const TArray<UE::Mutable::Private::EOpType> VeryExpensiveOperationTypes
 	{
-		mu::EOpType::ME_BINDSHAPE,
-		mu::EOpType::ME_MASKCLIPMESH,
-		mu::EOpType::ME_MASKCLIPUVMASK,
-		mu::EOpType::ME_FORMAT,
-		mu::EOpType::ME_MASKDIFF,
-		mu::EOpType::ME_DIFFERENCE,
-		mu::EOpType::ME_EXTRACTLAYOUTBLOCK,
-		mu::EOpType::IM_MAKEGROWMAP,
+		UE::Mutable::Private::EOpType::ME_BINDSHAPE,
+		UE::Mutable::Private::EOpType::ME_MASKCLIPMESH,
+		UE::Mutable::Private::EOpType::ME_MASKCLIPUVMASK,
+		UE::Mutable::Private::EOpType::ME_FORMAT,
+		UE::Mutable::Private::EOpType::ME_MASKDIFF,
+		UE::Mutable::Private::EOpType::ME_DIFFERENCE,
+		UE::Mutable::Private::EOpType::ME_EXTRACTLAYOUTBLOCK,
+		UE::Mutable::Private::EOpType::IM_MAKEGROWMAP,
 	};
 
 	/** Collection with all expensive to run operation types */
-	const TArray<mu::EOpType> ExpensiveOperationTypes
+	const TArray<UE::Mutable::Private::EOpType> ExpensiveOperationTypes
 	{
-		mu::EOpType::IM_PIXELFORMAT,
-		mu::EOpType::IM_SWIZZLE,
-		mu::EOpType::ME_PROJECT,
-		mu::EOpType::ME_TRANSFORMWITHMESH,
+		UE::Mutable::Private::EOpType::IM_PIXELFORMAT,
+		UE::Mutable::Private::EOpType::IM_SWIZZLE,
+		UE::Mutable::Private::EOpType::ME_PROJECT,
+		UE::Mutable::Private::EOpType::ME_TRANSFORMWITHMESH,
 	};
 	
 	/** Enum designed to be able to notify the row generation of the type of operation being generated */
@@ -253,7 +253,7 @@ private:
 	 * @param OperationType The operation type you want to know the computational cost of
 	 * @return The enum value representing the computational cost of the operation type provided.
 	 */
-	EOperationComputationalCost GetOperationTypeComputationalCost(mu::EOpType OperationType) const;
+	EOperationComputationalCost GetOperationTypeComputationalCost(UE::Mutable::Private::EOpType OperationType) const;
 	
 	
 	/*
@@ -278,11 +278,11 @@ private:
 	TSharedPtr<const FMutableOperationElement> CurrentlySelectedOperationTypeElement;
 	
 	/** Operation type we are using to search for tree nodes. Driven primarily by the UI */
-	mu::EOpType OperationTypeToSearch = mu::EOpType::NONE;
+	UE::Mutable::Private::EOpType OperationTypeToSearch = UE::Mutable::Private::EOpType::NONE;
 
 	/** Operation types present on the currently set mutable model.
 	 * The value of the TPair represents the amount of times operations using it have been found (not duplicates) */
-	TArray< TPair< mu::EOpType,uint32>> ModelOperationTypes;
+	TArray< TPair< UE::Mutable::Private::EOpType,uint32>> ModelOperationTypes;
 	
 	/** Array with all the names for each of the operations available on mutable. Used by the UI  */
 	TArray<FString> ModelOperationTypeNames;
@@ -388,7 +388,7 @@ private:
 	struct FElementsSearchCache
 	{
 		/** Set of addresses that have already been searched for relevant data */
-		TSet<mu::OP::ADDRESS> ProcessedAddresses;
+		TSet<UE::Mutable::Private::OP::ADDRESS> ProcessedAddresses;
 
 		/** Collection of elements that have been found during the search. They may be related by OP_Type or used constant resource*/
 		TArray<TSharedPtr<FMutableCodeTreeElement>> FoundElements;
@@ -405,7 +405,7 @@ private:
 		 * start of the search operation.
 		 * @param InRootNodeAddresses The addresses of the root operations of the operations tree
 		 */
-		void SetupRootBatch(const TArray<mu::OP::ADDRESS>& InRootNodeAddresses)
+		void SetupRootBatch(const TArray<UE::Mutable::Private::OP::ADDRESS>& InRootNodeAddresses)
 		{
 			check (InRootNodeAddresses.Num());
 			// This method should only be called once when no data is present on this cache
@@ -437,7 +437,7 @@ private:
 		* @param InItemCache Cache with all the elements of the tree. Here is where we search for the element based on
 		* the OpAddress and the parent and ChildIndexInParent provided thanks to IndexAsChildOfInputAddress
 		*/
-		void AddToFoundElements(const mu::OP::ADDRESS OpAddress,const int32 IndexAsChildOfInputAddress,
+		void AddToFoundElements(const UE::Mutable::Private::OP::ADDRESS OpAddress,const int32 IndexAsChildOfInputAddress,
 		                        const TMap<FItemCacheKey, TSharedPtr<FMutableCodeTreeElement>>& InItemCache)
 		{
 			FItemCacheKey Key;
@@ -468,8 +468,8 @@ private:
 		 * @note It will not add the children of any provided ParentAddress for the next batch if the parent address provided
 		 * has already been processed and therefore whose children have already been searched or prepared for searching.
 		 */
-		void CacheChildrenOfAddressIfNotProcessed(mu::OP::ADDRESS InParentAddress,
-		                                          const mu::FProgram& InProgram,
+		void CacheChildrenOfAddressIfNotProcessed(UE::Mutable::Private::OP::ADDRESS InParentAddress,
+		                                          const UE::Mutable::Private::FProgram& InProgram,
 		                                          TArray<FItemCacheKey>& OutFoundChildrenData)
 		{
 			if (!ProcessedAddresses.Contains(InParentAddress))
@@ -479,7 +479,7 @@ private:
 	
 				// Generic case for unnamed children traversal.
 				uint32 ChildIndex = 0;
-				mu::ForEachReference(InProgram, InParentAddress, [this, &InParentAddress, &ChildIndex,&OutFoundChildrenData]( mu::OP::ADDRESS ChildAddress)
+				UE::Mutable::Private::ForEachReference(InProgram, InParentAddress, [this, &InParentAddress, &ChildIndex,&OutFoundChildrenData]( UE::Mutable::Private::OP::ADDRESS ChildAddress)
 				{
 					// If the parent does have a child then process it
 					// Ignoring child address 0 since those can not have children
@@ -518,9 +518,9 @@ private:
 	 * the other. It also stores the found elements and other data.
 	 * @param InProgram The mutable program holding the data to be searched over
 	 */
-	void GetOperationsOfType(const mu::EOpType& TargetOperationType,
+	void GetOperationsOfType(const UE::Mutable::Private::EOpType& TargetOperationType,
 	                         FElementsSearchCache& InSearchPayload,
-	                         const mu::FProgram& InProgram);
+	                         const UE::Mutable::Private::FProgram& InProgram);
 
 	/*
 	 * Navigation : Navigation based on constant resource relation with addresses
@@ -534,11 +534,11 @@ public:
 	 * @note The output of this operation will be cached onto NavigationOPAddresses so we can navigate over them.
 	 * @param ConstantDataType The constant data type to locate operations referencing it: It is required for the
 	 * system to know what operations should be checked for the usage of the provided constant resource. For example, a
-	 * type of mu::EDataType::DT_Curve will try to locate operations that we know operate with constant curves.
+	 * type of UE::Mutable::Private::EDataType::DT_Curve will try to locate operations that we know operate with constant curves.
 	 * @param IndexOnConstantsArray The index of the constant on its constants array. If it is, for example, a constant mesh
-	 * this value represents the index of the constant inside constantMeshes array on mu::program.
+	 * this value represents the index of the constant inside constantMeshes array on UE::Mutable::Private::program.
 	 */
-	void CacheAddressesRelatedWithConstantResource(const mu::EDataType ConstantDataType, const int32 IndexOnConstantsArray);
+	void CacheAddressesRelatedWithConstantResource(const UE::Mutable::Private::EDataType ConstantDataType, const int32 IndexOnConstantsArray);
 
 private:
 	
@@ -551,10 +551,10 @@ private:
 	 * resource
 	 * @param InProgram The mutable program containing all the operations of the graph.
 	 */
-	void GetOperationsReferencingConstantResource(const mu::EDataType ConstantDataType,
+	void GetOperationsReferencingConstantResource(const UE::Mutable::Private::EDataType ConstantDataType,
 	                                              const int32 IndexOnConstantsArray,
 	                                              FElementsSearchCache& InSearchPayload,
-	                                              const mu::FProgram& InProgram);
+	                                              const UE::Mutable::Private::FProgram& InProgram);
 
 	/**
 	 * Provided an operation address this method tells us if that address is making use of our constant resource or not.
@@ -569,8 +569,8 @@ private:
 	 * @return True if the operation at OperationAddress does make use of the provided constant resource. False otherwise
 	 */
 	bool IsConstantResourceUsedByOperation(const int32 IndexOnConstantsArray,
-	                                       const mu::EDataType ConstantDataType, const mu::OP::ADDRESS OperationAddress,
-	                                       const mu::FProgram& InProgram) const;
+	                                       const UE::Mutable::Private::EDataType ConstantDataType, const UE::Mutable::Private::OP::ADDRESS OperationAddress,
+	                                       const UE::Mutable::Private::FProgram& InProgram) const;
 	
 
 	/*
@@ -588,7 +588,7 @@ private:
 	/** Last search configuration -> Prevents recurrent searches with the same search parameters */
 	FString LastSearchedString;
 	bool bWasLastSearchRegEx = false;
-	TSharedPtr<mu::FModel> LastSearchedModel;
+	TSharedPtr<UE::Mutable::Private::FModel> LastSearchedModel;
 
 	/** A collection with all the elements we have found to be related with the SearchString. */
 	TArray<TSharedPtr<FMutableCodeTreeElement>> NameBasedNavigationElements;
@@ -641,7 +641,7 @@ private:
 	 * @param SearchPayload Structure containing the operations to be used as the origin of the search. This object will also contain the results of the search.
 	 * @param InProgram MutableProgram used to allow the search of the children of the operations being currently processed.
 	 */
-	void GetOperationsMatchingStringPattern(const FString& InStringPattern,const bool bIsRegularExpression, FElementsSearchCache& SearchPayload, const mu::FProgram& InProgram);
+	void GetOperationsMatchingStringPattern(const FString& InStringPattern,const bool bIsRegularExpression, FElementsSearchCache& SearchPayload, const UE::Mutable::Private::FProgram& InProgram);
 
 	/** Provided a ItemCacheKey this method returns the string that represents it. It will not take in consideration duplicated entries since those are not part
 	 * of said ItemCache.
@@ -698,10 +698,10 @@ private:
 	struct FProcessedOperationsBuffer
 	{
 		/** Array with the operation addresses of all original (not duplicates) expanded operations. */
-		TArray<mu::OP::ADDRESS> ExpandedOriginalOperations;
+		TArray<UE::Mutable::Private::OP::ADDRESS> ExpandedOriginalOperations;
 
 		/** Array with the operation addresses of all duplicated expanded operations. */
-		TArray<mu::OP::ADDRESS> ExpandedDuplicatedOperations;
+		TArray<UE::Mutable::Private::OP::ADDRESS> ExpandedDuplicatedOperations;
 	};
 	
 	/** 
@@ -719,7 +719,7 @@ private:
 	*/
 	void TreeExpandElements(TArray<TSharedPtr<FMutableCodeTreeElement>>& InElementsToExpand,
 		bool bForceExpandDuplicates = false,
-		mu::EDataType FilteringDataType = mu::EDataType::None,
+		UE::Mutable::Private::EDataType FilteringDataType = UE::Mutable::Private::EDataType::None,
 		TSharedPtr<FProcessedOperationsBuffer> InExpandedOperationsBuffer = nullptr);
 	
 
@@ -791,15 +791,15 @@ private:
 	 */
 
 public:
-	void PreviewMutableImage (TSharedPtr<const mu::FImage> InImagePtr);
-	void PreviewMutableMesh (TSharedPtr<const mu::FMesh> InMeshPtr);
-	void PreviewMutableLayout(TSharedPtr<const mu::FLayout> Layout);
-	void PreviewMutableSkeleton(TSharedPtr<const mu::FSkeleton> Skeleton);
-	void PreviewMutablePhysics(TSharedPtr<const mu::FPhysicsBody> Physics);
+	void PreviewMutableImage (TSharedPtr<const UE::Mutable::Private::FImage> InImagePtr);
+	void PreviewMutableMesh (TSharedPtr<const UE::Mutable::Private::FMesh> InMeshPtr);
+	void PreviewMutableLayout(TSharedPtr<const UE::Mutable::Private::FLayout> Layout);
+	void PreviewMutableSkeleton(TSharedPtr<const UE::Mutable::Private::FSkeleton> Skeleton);
+	void PreviewMutablePhysics(TSharedPtr<const UE::Mutable::Private::FPhysicsBody> Physics);
 	void PreviewMutableString(const FString& InString);
-	void PreviewMutableProjector(const mu::FProjector* Projector);
+	void PreviewMutableProjector(const UE::Mutable::Private::FProjector* Projector);
 	void PreviewMutableMatrix(const FMatrix44f& Mat);
-	void PreviewMutableShape(const mu::FShape* Shape);
+	void PreviewMutableShape(const UE::Mutable::Private::FShape* Shape);
 	void PreviewMutableCurve(const FRichCurve& Curve);
 	
 private:
@@ -815,14 +815,14 @@ private:
 class FMutableOperationElement : public TSharedFromThis<FMutableOperationElement>
 {
 public:
-	FMutableOperationElement(mu::EOpType InOperationType,FText OperationTypeName,uint32 OperationTypeInstanceCount,FSlateColor OperationColor)
+	FMutableOperationElement(UE::Mutable::Private::EOpType InOperationType,FText OperationTypeName,uint32 OperationTypeInstanceCount,FSlateColor OperationColor)
 	{
 		OperationType = InOperationType;
 		OperationTextColor = OperationColor;
 
 		// Show the amount of instances of the operation type is found on model.
 		// @note This if block will not be triggered when working with an operation type that is not present on the model
-		// but added manually to the list of operation types (currently mu::OP_Type::None is an example of this)
+		// but added manually to the list of operation types (currently UE::Mutable::Private::OP_Type::None is an example of this)
 		FText OperationsCountText = FText::GetEmpty();
 		if (OperationTypeInstanceCount > 0)
 		{
@@ -832,14 +832,14 @@ public:
 	}
 
 public:
-	mu::EOpType OperationType;
+	UE::Mutable::Private::EOpType OperationType;
 	FText OperationTypeText;
 	FSlateColor OperationTextColor;
 
 	/** Get the type of visibility an slate representing this objects should have. */
 	EVisibility GetEntryVisibility() const
 	{
-		return OperationType == mu::EOpType::NONE ? EVisibility::Collapsed : EVisibility::Visible;
+		return OperationType == UE::Mutable::Private::EOpType::NONE ? EVisibility::Collapsed : EVisibility::Visible;
 	}
 };
 
@@ -849,7 +849,7 @@ public:
 class FMutableCodeTreeElement : public TSharedFromThis<FMutableCodeTreeElement>
 {
 public:
-	FMutableCodeTreeElement(int32 InIndexOnTree, const int32& InMutableStateIndex, const TSharedPtr<mu::FModel, ESPMode::ThreadSafe>& InModel, mu::OP::ADDRESS InOperation, const FString& InCaption, const FSlateColor InLabelColor, const TSharedPtr<FMutableCodeTreeElement>* InDuplicatedOf = nullptr);
+	FMutableCodeTreeElement(int32 InIndexOnTree, const int32& InMutableStateIndex, const TSharedPtr<UE::Mutable::Private::FModel, ESPMode::ThreadSafe>& InModel, UE::Mutable::Private::OP::ADDRESS InOperation, const FString& InCaption, const FSlateColor InLabelColor, const TSharedPtr<FMutableCodeTreeElement>* InDuplicatedOf = nullptr);
 	
 	void SetElementCurrentState(const int32& InStateIndex);
 
@@ -860,10 +860,10 @@ public:
 public:
 
 	/** Model this operation is part of */
-	TSharedPtr<mu::FModel, ESPMode::ThreadSafe> MutableModel;
+	TSharedPtr<UE::Mutable::Private::FModel, ESPMode::ThreadSafe> MutableModel;
 	
 	/** Mutable Graph Node represented in this tree row. */
-	mu::OP::ADDRESS MutableOperation;
+	UE::Mutable::Private::OP::ADDRESS MutableOperation;
 
 	/** If this tree element is a duplicated of another op, this is the op. */
 	TSharedPtr<FMutableCodeTreeElement> DuplicatedOf;
diff --git a/Engine/Plugins/Mutable/Source/CustomizableObjectEditor/Private/MuCOE/SMutableConstantsWidget.cpp b/Engine/Plugins/Mutable/Source/CustomizableObjectEditor/Private/MuCOE/SMutableConstantsWidget.cpp
--- a/Engine/Plugins/Mutable/Source/CustomizableObjectEditor/Private/MuCOE/SMutableConstantsWidget.cpp
+++ b/Engine/Plugins/Mutable/Source/CustomizableObjectEditor/Private/MuCOE/SMutableConstantsWidget.cpp
@@ -65,10 +65,10 @@ namespace
 	* @param InBufferSet : The bufferset whose channels we want to count.
 	* @return : the amount of channels found 
 	*/
-	int32 GetChannelCountOfBufferSet(const mu::FMeshBufferSet& InBufferSet)
+	int32 GetChannelCountOfBufferSet(const UE::Mutable::Private::FMeshBufferSet& InBufferSet)
 	{
 		int32 MeshBufferChannelsCount = 0;
-		for (const mu::FMeshBuffer& MeshBuffer : InBufferSet.Buffers)
+		for (const UE::Mutable::Private::FMeshBuffer& MeshBuffer : InBufferSet.Buffers)
 		{
 			MeshBufferChannelsCount += MeshBuffer.Channels.Num();
 		}
@@ -81,7 +81,7 @@ namespace
 	 * @param MeshPtr Pointer to the mesh whose channels we want to count
 	 * @return : the amount of channels found.
 	 */
-	int32 GetMeshChannelCount (TSharedPtr<const mu::FMesh> MeshPtr)
+	int32 GetMeshChannelCount (TSharedPtr<const UE::Mutable::Private::FMesh> MeshPtr)
 	{
 		check(MeshPtr)
 
@@ -135,7 +135,7 @@ public:
 		// The amount of mesh vertex buffer and image buffer channels
 		if (ColumnName == MeshConstantTitles::MeshChannels)
 		{
-			TSharedPtr<const mu::FMesh> Mesh = RowElement->MeshPtr;
+			TSharedPtr<const UE::Mutable::Private::FMesh> Mesh = RowElement->MeshPtr;
 			const int32 ChannelCount =  GetMeshChannelCount(Mesh);
 			const FText ChannelCountText = FText::AsNumber(ChannelCount);
 			return SNew(STextBlock).Text(ChannelCountText);
@@ -144,7 +144,7 @@ public:
 		// The amount of indices of the mesh
 		if (ColumnName == MeshConstantTitles::MeshIndices)
 		{
-			TSharedPtr<const mu::FMesh> Mesh = RowElement->MeshPtr;
+			TSharedPtr<const UE::Mutable::Private::FMesh> Mesh = RowElement->MeshPtr;
 			const FText IndexCount = FText::AsNumber(Mesh->GetIndexCount());
 			return SNew(STextBlock).Text(IndexCount);
 		}
@@ -152,7 +152,7 @@ public:
 		// The amount of vertices of the mesh
 		if (ColumnName == MeshConstantTitles::MeshVertices)
 		{
-			TSharedPtr<const mu::FMesh> Mesh = RowElement->MeshPtr;
+			TSharedPtr<const UE::Mutable::Private::FMesh> Mesh = RowElement->MeshPtr;
 			const FText VertexCount = FText::AsNumber(Mesh->GetVertexCount());
 			return SNew(STextBlock).Text(VertexCount);
 		}
@@ -160,7 +160,7 @@ public:
 		// Memory used by the mesh 
 		if (ColumnName == MeshConstantTitles::MeshMemory)
 		{
-			TSharedPtr<const mu::FMesh> Mesh = RowElement->MeshPtr;
+			TSharedPtr<const UE::Mutable::Private::FMesh> Mesh = RowElement->MeshPtr;
 			const FText SizeAsText = GenerateTextForSize(Mesh->GetDataSize());
 			return SNew(STextBlock).Text(SizeAsText);
 		}
@@ -244,7 +244,7 @@ public:
 	{
 		check(RowElement);
 
-		const mu::FImage* Image = RowElement->ImagePtr.Get();
+		const UE::Mutable::Private::FImage* Image = RowElement->ImagePtr.Get();
 		
 		// Index
 		if (ColumnName == ImageConstantTitles::ImageID)
@@ -274,10 +274,10 @@ public:
 		// Image format
 		if (ColumnName == ImageConstantTitles::ImageFormat)
 		{
-			const mu::EImageFormat ImageFormat = Image->GetFormat();
+			const UE::Mutable::Private::EImageFormat ImageFormat = Image->GetFormat();
 			const uint8 ImageFormatValue = static_cast<uint8>(ImageFormat);
 			const FText FormatAsText = FText::FromString(
-				FString(mu::TypeInfo::s_imageFormatName[ImageFormatValue]));
+				FString(UE::Mutable::Private::TypeInfo::s_imageFormatName[ImageFormatValue]));
 			
 			return SNew(STextBlock).Text(FormatAsText);
 		}
@@ -455,7 +455,7 @@ public:
 
 
 
-void SMutableConstantsWidget::Construct(const FArguments& InArgs,const mu::FProgram* InMutableProgramPtr,  TSharedPtr<SMutableCodeViewer> InMutableCodeViewerPtr)
+void SMutableConstantsWidget::Construct(const FArguments& InArgs,const UE::Mutable::Private::FProgram* InMutableProgramPtr,  TSharedPtr<SMutableCodeViewer> InMutableCodeViewerPtr)
 {
 	// A pointer to MutableCodeViewerPtr is required in order to be able to invoke the preview of our constants
 	check (InMutableCodeViewerPtr);
@@ -815,7 +815,7 @@ void SMutableConstantsWidget::Construct(const FArguments& InArgs,const mu::FProg
 
 
 
-void SMutableConstantsWidget::SetProgram(const mu::FProgram* InProgram)
+void SMutableConstantsWidget::SetProgram(const UE::Mutable::Private::FProgram* InProgram)
 {
 	// Make sure we do not process the setting of the same program object as the one already set
 	if (MutableProgramPtr == InProgram)
@@ -945,10 +945,10 @@ void SMutableConstantsWidget::OnImageTableSortRequested(EColumnSortPriority::Typ
 		if (ColumnID == ImageConstantTitles::ImageFormat)
 		{
 			const uint8 AImageFormatValue = static_cast<uint8>(A->ImagePtr->GetFormat());
-			const FString AImageFormatString = FString(mu::TypeInfo::s_imageFormatName[AImageFormatValue]);
+			const FString AImageFormatString = FString(UE::Mutable::Private::TypeInfo::s_imageFormatName[AImageFormatValue]);
 			
 			const uint8 BImageFormatValue = static_cast<uint8>(B->ImagePtr->GetFormat());
-			const FString BImageFormatString = FString(mu::TypeInfo::s_imageFormatName[BImageFormatValue]);
+			const FString BImageFormatString = FString(UE::Mutable::Private::TypeInfo::s_imageFormatName[BImageFormatValue]);
 			
 			if (ColumnSortMode == EColumnSortMode::Ascending)
 			{
@@ -1223,7 +1223,7 @@ void SMutableConstantsWidget::LoadConstantStrings()
 		ConstantStringElement->IndexOnSourceVector = StringAddressIndex;
 		
 		// Cache resource size
-		// in case we change the type of the contents of the mu::string we check its size as if it was a vector<>
+		// in case we change the type of the contents of the UE::Mutable::Private::string we check its size as if it was a vector<>
 		ConstantStringsAccumulatedSize += ConstantStringElement->MutableString.GetAllocatedSize();
 		
 		ConstantStringElements.Add(ConstantStringElement);
@@ -1255,7 +1255,7 @@ void SMutableConstantsWidget::LoadConstantImages()
 		TSharedPtr<FMutableConstantImageElement> ConstantImageElement = MakeShared<FMutableConstantImageElement>();
 
 		MutableProgramPtr->GetConstant(ImageIndex, ConstantImageElement->ImagePtr, 0,
-			[this](int32 x, int32 y, int32 m, mu::EImageFormat f, mu::EInitializationType i) { return MakeShared<mu::FImage>(x, y, m, f, i); });
+			[this](int32 x, int32 y, int32 m, UE::Mutable::Private::EImageFormat f, UE::Mutable::Private::EInitializationType i) { return MakeShared<UE::Mutable::Private::FImage>(x, y, m, f, i); });
 
 		ConstantImageElement->IndexOnSourceVector = ImageIndex;
 		
@@ -1272,10 +1272,10 @@ void SMutableConstantsWidget::LoadConstantMeshes()
 {
 	check (MutableProgramPtr);
 	
-	TArray<TSharedPtr<const mu::FMesh>> AllMeshes;
+	TArray<TSharedPtr<const UE::Mutable::Private::FMesh>> AllMeshes;
 	AllMeshes.Reserve(MutableProgramPtr->ConstantMeshesPermanent.Num() + MutableProgramPtr->ConstantMeshesStreamed.Num() );
 	AllMeshes.Append(MutableProgramPtr->ConstantMeshesPermanent);
-	for (const TPair<uint32, TSharedPtr<const mu::FMesh>>& Entry : MutableProgramPtr->ConstantMeshesStreamed)
+	for (const TPair<uint32, TSharedPtr<const UE::Mutable::Private::FMesh>>& Entry : MutableProgramPtr->ConstantMeshesStreamed)
 	{
 		AllMeshes.Add(Entry.Value);
 	}
@@ -1300,8 +1300,8 @@ void SMutableConstantsWidget::LoadConstantMeshes()
 		ConstantMeshElement->IndexOnSourceVector = MeshIndex;
 
 		// Actual core disk size would be:
-		//mu::FOutputMemoryStream Stream;
-		//mu::FOutputArchive Archive{ &Stream };
+		//UE::Mutable::Private::FOutputMemoryStream Stream;
+		//UE::Mutable::Private::FOutputArchive Archive{ &Stream };
 		//ConstantMeshElement->MeshPtr->Serialise(Archive);
 		//uint64 OtherSize = Stream.GetBufferSize();
 
@@ -1321,8 +1321,8 @@ void SMutableConstantsWidget::LoadConstantLayouts()
 	const int32 ConstantsCount = MutableProgramPtr->ConstantLayouts.Num();
 	ConstantLayoutElements.Empty(ConstantsCount);
 	
-	mu::FOutputMemoryStream Stream;
-	mu::FOutputArchive Archive{&Stream};
+	UE::Mutable::Private::FOutputMemoryStream Stream;
+	UE::Mutable::Private::FOutputArchive Archive{&Stream};
 	
 	for (int32 LayoutIndex = 0; LayoutIndex < ConstantsCount; LayoutIndex++)
 	{		
@@ -1346,8 +1346,8 @@ void SMutableConstantsWidget::LoadConstantSkeletons()
 	const int32 ConstantsCount = MutableProgramPtr->ConstantSkeletons.Num();
 	ConstantSkeletonElements.Empty(ConstantsCount);
 
-	mu::FOutputMemoryStream Stream;
-	mu::FOutputArchive Archive{&Stream};
+	UE::Mutable::Private::FOutputMemoryStream Stream;
+	UE::Mutable::Private::FOutputArchive Archive{&Stream};
 	
 	for (int32 SkeletonIndex = 0; SkeletonIndex < ConstantsCount; SkeletonIndex++)
 	{
@@ -1370,8 +1370,8 @@ void SMutableConstantsWidget::LoadConstantPhysics()
 	const int32 ConstantsCount = MutableProgramPtr->ConstantPhysicsBodies.Num();
 	ConstantPhysicsElements.Empty(ConstantsCount);
 
-	mu::FOutputMemoryStream Stream;
-	mu::FOutputArchive Archive{&Stream};
+	UE::Mutable::Private::FOutputMemoryStream Stream;
+	UE::Mutable::Private::FOutputArchive Archive{&Stream};
 	
 	for (int32 PhysicsIndex = 0; PhysicsIndex < ConstantsCount; PhysicsIndex++)
 	{
@@ -1395,8 +1395,8 @@ void SMutableConstantsWidget::LoadConstantProjectors()
 	const int32 ConstantsCount = MutableProgramPtr->ConstantProjectors.Num();
 	ConstantProjectorElements.Empty(ConstantsCount);
 
-	mu::FOutputMemoryStream Stream;
-	mu::FOutputArchive Archive{&Stream};
+	UE::Mutable::Private::FOutputMemoryStream Stream;
+	UE::Mutable::Private::FOutputArchive Archive{&Stream};
 	
 	for (int32 ProjectorIndex = 0; ProjectorIndex < ConstantsCount; ProjectorIndex++)
 	{
@@ -1419,8 +1419,8 @@ void SMutableConstantsWidget::LoadConstantMatrices()
 	const int32 ConstantsCount = MutableProgramPtr->ConstantMatrices.Num();
 	ConstantMatrixElements.Empty(ConstantsCount);
 
-	mu::FOutputMemoryStream Stream;
-	mu::FOutputArchive Archive{&Stream};
+	UE::Mutable::Private::FOutputMemoryStream Stream;
+	UE::Mutable::Private::FOutputArchive Archive{&Stream};
 	
 	for (int32 MatrixIndex = 0; MatrixIndex < ConstantsCount; MatrixIndex++)
 	{
@@ -1443,8 +1443,8 @@ void SMutableConstantsWidget::LoadConstantShapes()
 	const int32 ConstantsCount = MutableProgramPtr->ConstantShapes.Num();
 	ConstantShapeElements.Empty(ConstantsCount);
 
-	mu::FOutputMemoryStream Stream;
-	mu::FOutputArchive Archive{&Stream};	
+	UE::Mutable::Private::FOutputMemoryStream Stream;
+	UE::Mutable::Private::FOutputArchive Archive{&Stream};	
 	
 	for (int32 ShapeIndex = 0; ShapeIndex < ConstantsCount; ShapeIndex++)
 	{
@@ -1467,8 +1467,8 @@ void SMutableConstantsWidget::LoadConstantCurves()
 	const int32 ConstantsCount = MutableProgramPtr->ConstantCurves.Num();
 	ConstantCurveElements.Empty(ConstantsCount);
 	
-	mu::FOutputMemoryStream Stream;
-	mu::FOutputArchive Archive{&Stream};
+	UE::Mutable::Private::FOutputMemoryStream Stream;
+	UE::Mutable::Private::FOutputArchive Archive{&Stream};
 	
 	for (int32 CurveIndex = 0; CurveIndex < ConstantsCount; CurveIndex++)
 	{
@@ -1491,45 +1491,45 @@ void SMutableConstantsWidget::LoadConstantCurves()
 #pragma region Previewer invocation methods
 
 
-void SMutableConstantsWidget::ClearSelectedConstantItems(mu::EDataType ExceptionDataType /* = mu::EDataType::None */) const
+void SMutableConstantsWidget::ClearSelectedConstantItems(UE::Mutable::Private::EDataType ExceptionDataType /* = UE::Mutable::Private::EDataType::None */) const
 {
-	if (ExceptionDataType != mu::EDataType::Mesh)
+	if (ExceptionDataType != UE::Mutable::Private::EDataType::Mesh)
 	{
 		MeshListViewHandler->ClearSelection();
 	}
-	if (ExceptionDataType != mu::EDataType::String)
+	if (ExceptionDataType != UE::Mutable::Private::EDataType::String)
 	{
 		ConstantStringsSlate->ClearSelection();
 	}
-	if (ExceptionDataType != mu::EDataType::Layout)
+	if (ExceptionDataType != UE::Mutable::Private::EDataType::Layout)
 	{
 		ConstantLayoutsSlate->ClearSelection();
 	}
-	if (ExceptionDataType != mu::EDataType::Projector)
+	if (ExceptionDataType != UE::Mutable::Private::EDataType::Projector)
 	{
 		ConstantProjectorsSlate->ClearSelection();
 	}
-	if (ExceptionDataType != mu::EDataType::Matrix)
+	if (ExceptionDataType != UE::Mutable::Private::EDataType::Matrix)
 	{
 		ConstantMatricesSlate->ClearSelection();
 	}
-	if (ExceptionDataType != mu::EDataType::Shape)
+	if (ExceptionDataType != UE::Mutable::Private::EDataType::Shape)
 	{
 		ConstantShapesSlate->ClearSelection();
 	}
-	if (ExceptionDataType != mu::EDataType::Curve)
+	if (ExceptionDataType != UE::Mutable::Private::EDataType::Curve)
 	{
 		ConstantCurvesSlate->ClearSelection();
 	}
-	if (ExceptionDataType != mu::EDataType::Skeleton)
+	if (ExceptionDataType != UE::Mutable::Private::EDataType::Skeleton)
 	{
 		ConstantSkeletonsSlate->ClearSelection();
 	}
-	if (ExceptionDataType != mu::EDataType::PhysicsAsset)
+	if (ExceptionDataType != UE::Mutable::Private::EDataType::PhysicsAsset)
 	{
 		ConstantPhysicsSlate->ClearSelection();
 	}
-	if (ExceptionDataType != mu::EDataType::Image)
+	if (ExceptionDataType != UE::Mutable::Private::EDataType::Image)
 	{
 		ImageListViewHandler->ClearSelection();
 		//ConstantImagesListView->ClearSelection();
@@ -1544,7 +1544,7 @@ void SMutableConstantsWidget::OnSelectedStringChanged(
 {
 	if (MutableConstantStringElement)
 	{
-		static constexpr mu::EDataType SlateDataType = mu::EDataType::String;
+		static constexpr UE::Mutable::Private::EDataType SlateDataType = UE::Mutable::Private::EDataType::String;
 
 		// Clear the selected CodeViewer row and all other constant viewer slates that are not the type provided.
 		MutableCodeViewerPtr->ClearSelectedTreeRow();
@@ -1562,7 +1562,7 @@ void SMutableConstantsWidget::OnSelectedImageChanged(
 {
 	if (MutableConstantImageElement)
 	{
-		static constexpr mu::EDataType SlateDataType = mu::EDataType::Image;
+		static constexpr UE::Mutable::Private::EDataType SlateDataType = UE::Mutable::Private::EDataType::Image;
 
 		MutableCodeViewerPtr->ClearSelectedTreeRow();
 		ClearSelectedConstantItems(SlateDataType);
@@ -1578,7 +1578,7 @@ void SMutableConstantsWidget::OnSelectedMeshChanged(
 {
 	if (MutableConstantMeshElement)
 	{
-		static constexpr mu::EDataType SlateDataType = mu::EDataType::Mesh;
+		static constexpr UE::Mutable::Private::EDataType SlateDataType = UE::Mutable::Private::EDataType::Mesh;
 		
 		MutableCodeViewerPtr->ClearSelectedTreeRow();
 		ClearSelectedConstantItems(SlateDataType);
@@ -1594,7 +1594,7 @@ void SMutableConstantsWidget::OnSelectedLayoutChanged(
 {
 	if (MutableConstantLayoutElement)
 	{
-		static constexpr mu::EDataType SlateDataType = mu::EDataType::Layout;
+		static constexpr UE::Mutable::Private::EDataType SlateDataType = UE::Mutable::Private::EDataType::Layout;
 		
 		MutableCodeViewerPtr->ClearSelectedTreeRow();
 		ClearSelectedConstantItems(SlateDataType);
@@ -1610,7 +1610,7 @@ void SMutableConstantsWidget::OnSelectedProjectorChanged(
 {
 	if (MutableConstantProjectorElement)
 	{
-		static constexpr mu::EDataType SlateDataType = mu::EDataType::Projector;
+		static constexpr UE::Mutable::Private::EDataType SlateDataType = UE::Mutable::Private::EDataType::Projector;
 
 		MutableCodeViewerPtr->ClearSelectedTreeRow();
 		ClearSelectedConstantItems(SlateDataType);
@@ -1626,7 +1626,7 @@ void SMutableConstantsWidget::OnSelectedMatrixChanged(
 {
 	if (MutableConstantMatrixElement)
 	{
-		static constexpr mu::EDataType SlateDataType = mu::EDataType::Matrix;
+		static constexpr UE::Mutable::Private::EDataType SlateDataType = UE::Mutable::Private::EDataType::Matrix;
 
 		MutableCodeViewerPtr->ClearSelectedTreeRow();
 		ClearSelectedConstantItems(SlateDataType);
@@ -1642,7 +1642,7 @@ void SMutableConstantsWidget::OnSelectedShapeChanged(
 {
 	if (MutableConstantShapeElement)
 	{
-		static constexpr mu::EDataType SlateDataType = mu::EDataType::Shape;
+		static constexpr UE::Mutable::Private::EDataType SlateDataType = UE::Mutable::Private::EDataType::Shape;
 
 		MutableCodeViewerPtr->ClearSelectedTreeRow();
 		ClearSelectedConstantItems(SlateDataType);
@@ -1658,7 +1658,7 @@ void SMutableConstantsWidget::OnSelectedCurveChanged(
 {
 	if (MutableConstantCurveElement)
 	{
-		static constexpr mu::EDataType SlateDataType = mu::EDataType::Curve;
+		static constexpr UE::Mutable::Private::EDataType SlateDataType = UE::Mutable::Private::EDataType::Curve;
 
 		MutableCodeViewerPtr->ClearSelectedTreeRow();
 		ClearSelectedConstantItems(SlateDataType);
@@ -1674,7 +1674,7 @@ void SMutableConstantsWidget::OnSelectedSkeletonChanged(
 {
 	if (MutableConstantSkeletonElement)
 	{
-		static constexpr mu::EDataType SlateDataType = mu::EDataType::Skeleton;
+		static constexpr UE::Mutable::Private::EDataType SlateDataType = UE::Mutable::Private::EDataType::Skeleton;
 
 		MutableCodeViewerPtr->ClearSelectedTreeRow();
 		ClearSelectedConstantItems(SlateDataType);
@@ -1690,7 +1690,7 @@ void SMutableConstantsWidget::OnSelectedPhysicsChanged(
 {
 	if (MutableConstantPhysicsElement)
 	{
-		static constexpr mu::EDataType SlateDataType = mu::EDataType::PhysicsAsset;
+		static constexpr UE::Mutable::Private::EDataType SlateDataType = UE::Mutable::Private::EDataType::PhysicsAsset;
 
 		MutableCodeViewerPtr->ClearSelectedTreeRow();
 		ClearSelectedConstantItems(SlateDataType);
diff --git a/Engine/Plugins/Mutable/Source/CustomizableObjectEditor/Private/MuCOE/SMutableConstantsWidget.h b/Engine/Plugins/Mutable/Source/CustomizableObjectEditor/Private/MuCOE/SMutableConstantsWidget.h
--- a/Engine/Plugins/Mutable/Source/CustomizableObjectEditor/Private/MuCOE/SMutableConstantsWidget.h
+++ b/Engine/Plugins/Mutable/Source/CustomizableObjectEditor/Private/MuCOE/SMutableConstantsWidget.h
@@ -17,9 +17,9 @@ class SExpandableArea;
 class SMutableCodeViewer;
 class STableViewBase;
 class SWidget;
-namespace mu { struct FProgram; }
-namespace mu { struct FProjector; }
-namespace mu { struct FShape; }
+namespace UE::Mutable::Private { struct FProgram; }
+namespace UE::Mutable::Private { struct FProjector; }
+namespace UE::Mutable::Private { struct FShape; }
 
 /**
  * Base Structure to define the different elements used by the lists on this object
@@ -35,7 +35,7 @@ struct FMutableConstantElement
 */
 struct FMutableConstantImageElement : public FMutableConstantElement
 {
-	TSharedPtr<const mu::FImage> ImagePtr;
+	TSharedPtr<const UE::Mutable::Private::FImage> ImagePtr;
 };
 
 /**
@@ -43,7 +43,7 @@ struct FMutableConstantImageElement : public FMutableConstantElement
 */
 struct FMutableConstantMeshElement : public FMutableConstantElement
 {
-	TSharedPtr<const mu::FMesh> MeshPtr;
+	TSharedPtr<const UE::Mutable::Private::FMesh> MeshPtr;
 };
 
 /**
@@ -59,7 +59,7 @@ struct FMutableConstantStringElement : public FMutableConstantElement
 */
 struct FMutableConstantLayoutElement : public FMutableConstantElement
 {
-	TSharedPtr<const mu::FLayout> Layout;
+	TSharedPtr<const UE::Mutable::Private::FLayout> Layout;
 };
 
 /**
@@ -67,7 +67,7 @@ struct FMutableConstantLayoutElement : public FMutableConstantElement
 */
 struct FMutableConstantSkeletonElement : public FMutableConstantElement
 {
-	TSharedPtr<const mu::FSkeleton> Skeleton;
+	TSharedPtr<const UE::Mutable::Private::FSkeleton> Skeleton;
 };
 
 /**
@@ -75,7 +75,7 @@ struct FMutableConstantSkeletonElement : public FMutableConstantElement
 */
 struct FMutableConstantPhysicsElement : public FMutableConstantElement
 {
-	TSharedPtr<const mu::FPhysicsBody> Physics;
+	TSharedPtr<const UE::Mutable::Private::FPhysicsBody> Physics;
 };
 
 /**
@@ -83,7 +83,7 @@ struct FMutableConstantPhysicsElement : public FMutableConstantElement
 */
 struct FMutableConstantProjectorElement : public FMutableConstantElement
 {
-	const mu::FProjector* Projector = nullptr;
+	const UE::Mutable::Private::FProjector* Projector = nullptr;
 };
 
 /**
@@ -99,7 +99,7 @@ struct FMutableConstantMatrixElement : public FMutableConstantElement
 */
 struct FMutableConstantShapeElement : public FMutableConstantElement
 {
-	const mu::FShape* Shape = nullptr;
+	const UE::Mutable::Private::FShape* Shape = nullptr;
 };
 
 /**
@@ -123,22 +123,22 @@ public:
 
 	/** Builds the widget
 	 * @param InArgs - Arguments provided when generating this slate object
-	 * @param InMutableProgramPtr - Pointer to the mu::FProgram object that holds the constants data.
+	 * @param InMutableProgramPtr - Pointer to the UE::Mutable::Private::FProgram object that holds the constants data.
 	 * @param  InMutableCodeViewerPtr - Pointer to the MutableCodeViewer tasked with the previewing of the constant values
 	 */
-	void Construct(const FArguments& InArgs, const mu::FProgram*  InMutableProgramPtr ,  TSharedPtr<SMutableCodeViewer> InMutableCodeViewerPtr);
+	void Construct(const FArguments& InArgs, const UE::Mutable::Private::FProgram*  InMutableProgramPtr ,  TSharedPtr<SMutableCodeViewer> InMutableCodeViewerPtr);
 
 	/** It clears the selected element for all constant view slates (SListView and STileView) except the one that shows the data of the provided type
 	 * @param ExceptionDataType The type of data of the slate container that we want to skip when clearing the selected elements.
-	 * EX : We want to clear all selected elements but not the ones from the constant images. Provide as argument "mu::EDataType::Image"
+	 * EX : We want to clear all selected elements but not the ones from the constant images. Provide as argument "UE::Mutable::Private::EDataType::Image"
 	 * to clear all but the selected item/s of the Image type.
 	 */
-	void ClearSelectedConstantItems(mu::EDataType ExceptionDataType = mu::EDataType::None) const;
+	void ClearSelectedConstantItems(UE::Mutable::Private::EDataType ExceptionDataType = UE::Mutable::Private::EDataType::None) const;
 
 private:
 
 	/** Mutable object containing the constants data */
-	const mu::FProgram* MutableProgramPtr = nullptr;
+	const UE::Mutable::Private::FProgram* MutableProgramPtr = nullptr;
 
 	/** Slate capable of accessing the previewer object */
 	TSharedPtr<SMutableCodeViewer> MutableCodeViewerPtr = nullptr;
@@ -147,7 +147,7 @@ private:
 	 *Sets the back end for the operation of this widget. Each time this is done the ui backend gets updated
 	 * @param InProgram - Mutable program object holding all the constants data
 	 */
-	void SetProgram(const mu::FProgram* InProgram);
+	void SetProgram(const UE::Mutable::Private::FProgram* InProgram);
 
 	/*
 	 * Pointers to all slates showing the constants
@@ -183,7 +183,7 @@ private:
 	TArray<TSharedPtr<FMutableConstantPhysicsElement>> ConstantPhysicsElements;
 
 	/**
-	 * Load up all the elements with the data found on the mu::FProgram object onto TArrays after parsing the data found.
+	 * Load up all the elements with the data found on the UE::Mutable::Private::FProgram object onto TArrays after parsing the data found.
 	 */
 	void LoadConstantElements();
 	
diff --git a/Engine/Plugins/Mutable/Source/CustomizableObjectEditor/Private/MuCOE/SMutableCurveViewer.cpp b/Engine/Plugins/Mutable/Source/CustomizableObjectEditor/Private/MuCOE/SMutableCurveViewer.cpp
--- a/Engine/Plugins/Mutable/Source/CustomizableObjectEditor/Private/MuCOE/SMutableCurveViewer.cpp
+++ b/Engine/Plugins/Mutable/Source/CustomizableObjectEditor/Private/MuCOE/SMutableCurveViewer.cpp
@@ -447,10 +447,10 @@ void SMutableCurveViewer::SetCurve(const FRichCurve& InMutableCurve)
 {
 	this->MutableCurve = InMutableCurve;
 
-	// Load the Curve Editor with the data found on the mu::Curve
+	// Load the Curve Editor with the data found on the UE::Mutable::Private::Curve
 	SetupMutableCurveGraph();
 	
-	// Load the data required by the SListView to display the data on the mu::Curve
+	// Load the data required by the SListView to display the data on the UE::Mutable::Private::Curve
 	SetupMutableCurveListView();
 }
 
diff --git a/Engine/Plugins/Mutable/Source/CustomizableObjectEditor/Private/MuCOE/SMutableCurveViewer.h b/Engine/Plugins/Mutable/Source/CustomizableObjectEditor/Private/MuCOE/SMutableCurveViewer.h
--- a/Engine/Plugins/Mutable/Source/CustomizableObjectEditor/Private/MuCOE/SMutableCurveViewer.h
+++ b/Engine/Plugins/Mutable/Source/CustomizableObjectEditor/Private/MuCOE/SMutableCurveViewer.h
@@ -109,7 +109,7 @@ struct FMutableCurveElement
 
 
 /**
- * Object designed to be used as a way to inspect the data present on a Mutable Curve (mu::Curve)
+ * Object designed to be used as a way to inspect the data present on a Mutable Curve (UE::Mutable::Private::Curve)
  */
 class SMutableCurveViewer final : public SCompoundWidget
 {
diff --git a/Engine/Plugins/Mutable/Source/CustomizableObjectEditor/Private/MuCOE/SMutableGraphViewer.cpp b/Engine/Plugins/Mutable/Source/CustomizableObjectEditor/Private/MuCOE/SMutableGraphViewer.cpp
--- a/Engine/Plugins/Mutable/Source/CustomizableObjectEditor/Private/MuCOE/SMutableGraphViewer.cpp
+++ b/Engine/Plugins/Mutable/Source/CustomizableObjectEditor/Private/MuCOE/SMutableGraphViewer.cpp
@@ -221,7 +221,7 @@ static const char* MutableNodeNames[] =
 	"ColorMaterialBreak",
 };
 
-static_assert(UE_ARRAY_COUNT(MutableNodeNames) == SIZE_T(mu::Node::EType::Count));
+static_assert(UE_ARRAY_COUNT(MutableNodeNames) == SIZE_T(UE::Mutable::Private::Node::EType::Count));
 
 
 class SMutableGraphTreeRow : public STableRow<TSharedPtr<FMutableGraphTreeElement>>
@@ -235,7 +235,7 @@ public:
 		FText MainLabel = FText::GetEmpty();
 		if (RowItem->MutableNode)
 		{
-			mu::Node::EType MutableType = RowItem->MutableNode->GetType()->Type;
+			UE::Mutable::Private::Node::EType MutableType = RowItem->MutableNode->GetType()->Type;
 			
 			FString NodeName = MutableNodeNames[int32(MutableType)];
 
@@ -301,7 +301,7 @@ FString SMutableGraphViewer::GetReferencerName() const
 }
 
 
-void SMutableGraphViewer::Construct(const FArguments& InArgs, const mu::NodePtr& InRootNode)
+void SMutableGraphViewer::Construct(const FArguments& InArgs, const UE::Mutable::Private::NodePtr& InRootNode)
 {
 	ReferencedRuntimeTextures = InArgs._ReferencedRuntimeTextures;
 	ReferencedCompileTextures = InArgs._ReferencedCompileTextures;
@@ -380,10 +380,10 @@ void SMutableGraphViewer::GetChildrenForInfo(TSharedPtr<FMutableGraphTreeElement
 		return;
 	}
 
-	mu::Node* ParentNode = InInfo->MutableNode.get();
+	UE::Mutable::Private::Node* ParentNode = InInfo->MutableNode.get();
 	uint32 InputIndex = 0;
 
-	auto AddChildFunc = [this, ParentNode, &InputIndex, &OutChildren](mu::Node* ChildNode, const FString& Prefix)
+	auto AddChildFunc = [this, ParentNode, &InputIndex, &OutChildren](UE::Mutable::Private::Node* ChildNode, const FString& Prefix)
 	{
 		if (ChildNode)
 		{
@@ -416,9 +416,9 @@ void SMutableGraphViewer::GetChildrenForInfo(TSharedPtr<FMutableGraphTreeElement
 		++InputIndex;
 	};
 
-	if (ParentNode->GetType() == mu::NodeObjectNew::GetStaticType())
+	if (ParentNode->GetType() == UE::Mutable::Private::NodeObjectNew::GetStaticType())
 	{
-		mu::NodeObjectNew* ObjectNew = StaticCast<mu::NodeObjectNew*>(ParentNode);
+		UE::Mutable::Private::NodeObjectNew* ObjectNew = StaticCast<UE::Mutable::Private::NodeObjectNew*>(ParentNode);
 		for (int32 l = 0; l < ObjectNew->Components.Num(); ++l)
 		{
 			AddChildFunc(ObjectNew->Components[l].get(), TEXT("COMP") );
@@ -435,18 +435,18 @@ void SMutableGraphViewer::GetChildrenForInfo(TSharedPtr<FMutableGraphTreeElement
 		}
 	}
 
-	else if (ParentNode->GetType() == mu::NodeObjectGroup::GetStaticType())
+	else if (ParentNode->GetType() == UE::Mutable::Private::NodeObjectGroup::GetStaticType())
 	{
-		mu::NodeObjectGroup* ObjectGroup = StaticCast<mu::NodeObjectGroup*>(ParentNode);
+		UE::Mutable::Private::NodeObjectGroup* ObjectGroup = StaticCast<UE::Mutable::Private::NodeObjectGroup*>(ParentNode);
 		for (int32 l = 0; l < ObjectGroup->Children.Num(); ++l)
 		{
 			AddChildFunc(ObjectGroup->Children[l].get(), TEXT("CHILD"));
 		}
 	}
 
-	else if (ParentNode->GetType() == mu::NodeSurfaceNew::GetStaticType())
+	else if (ParentNode->GetType() == UE::Mutable::Private::NodeSurfaceNew::GetStaticType())
 	{
-		mu::NodeSurfaceNew* SurfaceNew = StaticCast<mu::NodeSurfaceNew*>(ParentNode);
+		UE::Mutable::Private::NodeSurfaceNew* SurfaceNew = StaticCast<UE::Mutable::Private::NodeSurfaceNew*>(ParentNode);
 		AddChildFunc(SurfaceNew->Mesh.get(), TEXT("MESH"));
 
 		for (int32 l = 0; l < SurfaceNew->Images.Num(); ++l)
@@ -470,9 +470,9 @@ void SMutableGraphViewer::GetChildrenForInfo(TSharedPtr<FMutableGraphTreeElement
 		}
 	}
 
-	else if (ParentNode->GetType() == mu::NodeModifierSurfaceEdit::GetStaticType())
+	else if (ParentNode->GetType() == UE::Mutable::Private::NodeModifierSurfaceEdit::GetStaticType())
 	{
-		mu::NodeModifierSurfaceEdit* SurfaceEdit = StaticCast<mu::NodeModifierSurfaceEdit*>(ParentNode);
+		UE::Mutable::Private::NodeModifierSurfaceEdit* SurfaceEdit = StaticCast<UE::Mutable::Private::NodeModifierSurfaceEdit*>(ParentNode);
 
 		AddChildFunc(SurfaceEdit->MorphFactor.get(), FString::Printf(TEXT("MORPH_FACTOR [%s]"), *SurfaceEdit->MeshMorph));
 
@@ -491,9 +491,9 @@ void SMutableGraphViewer::GetChildrenForInfo(TSharedPtr<FMutableGraphTreeElement
 		}
 	}
 
-	else if (ParentNode->GetType() == mu::NodeSurfaceSwitch::GetStaticType())
+	else if (ParentNode->GetType() == UE::Mutable::Private::NodeSurfaceSwitch::GetStaticType())
 	{
-		mu::NodeSurfaceSwitch* SurfaceSwitch = StaticCast<mu::NodeSurfaceSwitch*>(ParentNode);
+		UE::Mutable::Private::NodeSurfaceSwitch* SurfaceSwitch = StaticCast<UE::Mutable::Private::NodeSurfaceSwitch*>(ParentNode);
 		AddChildFunc(SurfaceSwitch->Parameter.get(), TEXT("PARAM"));
 		for (int32 l = 0; l < SurfaceSwitch->Options.Num(); ++l)
 		{
@@ -501,9 +501,9 @@ void SMutableGraphViewer::GetChildrenForInfo(TSharedPtr<FMutableGraphTreeElement
 		}
 	}
 
-	else if (ParentNode->GetType() == mu::NodeSurfaceVariation::GetStaticType())
+	else if (ParentNode->GetType() == UE::Mutable::Private::NodeSurfaceVariation::GetStaticType())
 	{
-		mu::NodeSurfaceVariation* SurfaceVar = StaticCast<mu::NodeSurfaceVariation*>(ParentNode);
+		UE::Mutable::Private::NodeSurfaceVariation* SurfaceVar = StaticCast<UE::Mutable::Private::NodeSurfaceVariation*>(ParentNode);
 		for (int32 l = 0; l < SurfaceVar->DefaultSurfaces.Num(); ++l)
 		{
 			AddChildFunc(SurfaceVar->DefaultSurfaces[l].get(), FString::Printf(TEXT("DEF SURF [%d]"), l));
@@ -515,7 +515,7 @@ void SMutableGraphViewer::GetChildrenForInfo(TSharedPtr<FMutableGraphTreeElement
 
 		for (int32 v = 0; v < SurfaceVar->Variations.Num(); ++v)
 		{
-			const mu::NodeSurfaceVariation::FVariation Var = SurfaceVar->Variations[v];
+			const UE::Mutable::Private::NodeSurfaceVariation::FVariation Var = SurfaceVar->Variations[v];
 			for (int32 l = 0; l < Var.Surfaces.Num(); ++l)
 			{
 				AddChildFunc(Var.Surfaces[l].get(), FString::Printf(TEXT("VAR [%s] SURF [%d]"), *Var.Tag, l));
@@ -527,9 +527,9 @@ void SMutableGraphViewer::GetChildrenForInfo(TSharedPtr<FMutableGraphTreeElement
 		}
 	}
 	
-	else if (ParentNode->GetType() == mu::NodeLOD::GetStaticType())
+	else if (ParentNode->GetType() == UE::Mutable::Private::NodeLOD::GetStaticType())
 	{
-		mu::NodeLOD* LodVar = StaticCast<mu::NodeLOD*>(ParentNode);
+		UE::Mutable::Private::NodeLOD* LodVar = StaticCast<UE::Mutable::Private::NodeLOD*>(ParentNode);
 
 		for (int32 SurfaceIndex = 0; SurfaceIndex < LodVar->Surfaces.Num(); SurfaceIndex++)
 		{
@@ -537,9 +537,9 @@ void SMutableGraphViewer::GetChildrenForInfo(TSharedPtr<FMutableGraphTreeElement
 		}
 	}
 	
-	else if (ParentNode->GetType() == mu::NodeComponentNew::GetStaticType())
+	else if (ParentNode->GetType() == UE::Mutable::Private::NodeComponentNew::GetStaticType())
 	{
-		mu::NodeComponentNew* ComponentVar = StaticCast<mu::NodeComponentNew*>(ParentNode);
+		UE::Mutable::Private::NodeComponentNew* ComponentVar = StaticCast<UE::Mutable::Private::NodeComponentNew*>(ParentNode);
 		for (int32 LODIndex = 0; LODIndex < ComponentVar->LODs.Num(); LODIndex++)
 		{
 			AddChildFunc(ComponentVar->LODs[LODIndex].get(), FString::Printf(TEXT("LOD [%d]"), LODIndex));
@@ -548,18 +548,18 @@ void SMutableGraphViewer::GetChildrenForInfo(TSharedPtr<FMutableGraphTreeElement
 		AddChildFunc(ComponentVar->OverlayMaterial.get(), TEXT("OVERLAY MATERIAL"));
 	}
 
-	else if (ParentNode->GetType() == mu::NodeComponentEdit::GetStaticType())
+	else if (ParentNode->GetType() == UE::Mutable::Private::NodeComponentEdit::GetStaticType())
 	{
-		mu::NodeComponentEdit* ComponentVar = StaticCast<mu::NodeComponentEdit*>(ParentNode);
+		UE::Mutable::Private::NodeComponentEdit* ComponentVar = StaticCast<UE::Mutable::Private::NodeComponentEdit*>(ParentNode);
 		for (int32 LODIndex = 0; LODIndex < ComponentVar->LODs.Num(); LODIndex++)
 		{
 			AddChildFunc(ComponentVar->LODs[LODIndex].get(), FString::Printf(TEXT("LOD [%d]"), LODIndex));
 		}
 		}
 
-	else if (ParentNode->GetType() == mu::NodeComponentSwitch::GetStaticType())
+	else if (ParentNode->GetType() == UE::Mutable::Private::NodeComponentSwitch::GetStaticType())
 	{
-		mu::NodeComponentSwitch* ComponentSwitch = StaticCast<mu::NodeComponentSwitch*>(ParentNode);
+		UE::Mutable::Private::NodeComponentSwitch* ComponentSwitch = StaticCast<UE::Mutable::Private::NodeComponentSwitch*>(ParentNode);
 		AddChildFunc(ComponentSwitch->Parameter.get(), TEXT("PARAM"));
 		for (int32 OptionIndex = 0; OptionIndex < ComponentSwitch->Options.Num(); OptionIndex++)
 		{
@@ -567,67 +567,67 @@ void SMutableGraphViewer::GetChildrenForInfo(TSharedPtr<FMutableGraphTreeElement
 		}
 	}
 
-	else if (ParentNode->GetType() == mu::NodeMeshConstant::GetStaticType())
+	else if (ParentNode->GetType() == UE::Mutable::Private::NodeMeshConstant::GetStaticType())
 	{
-		mu::NodeMeshConstant* MeshConstantVar = StaticCast<mu::NodeMeshConstant*>(ParentNode);
+		UE::Mutable::Private::NodeMeshConstant* MeshConstantVar = StaticCast<UE::Mutable::Private::NodeMeshConstant*>(ParentNode);
 		for (int32 LayoutIndex = 0; LayoutIndex < MeshConstantVar->Layouts.Num(); LayoutIndex++)
 		{
 			AddChildFunc(MeshConstantVar->Layouts[LayoutIndex].get(), FString::Printf(TEXT("LAYOUT [%d]"), LayoutIndex));
 		}
 	}
 
-	else if (ParentNode->GetType() == mu::NodeImageFormat::GetStaticType())
+	else if (ParentNode->GetType() == UE::Mutable::Private::NodeImageFormat::GetStaticType())
 	{
-		mu::NodeImageFormat* ImageFormatVar = StaticCast<mu::NodeImageFormat*>(ParentNode);
+		UE::Mutable::Private::NodeImageFormat* ImageFormatVar = StaticCast<UE::Mutable::Private::NodeImageFormat*>(ParentNode);
 		AddChildFunc(ImageFormatVar->Source.get(), FString::Printf(TEXT("SOURCE IMAGE")));
 	}
 
-	else if (ParentNode->GetType() == mu::NodeMeshFormat::GetStaticType())
+	else if (ParentNode->GetType() == UE::Mutable::Private::NodeMeshFormat::GetStaticType())
 	{
-		mu::NodeMeshFormat* MeshFormatVar = StaticCast<mu::NodeMeshFormat*>(ParentNode);
+		UE::Mutable::Private::NodeMeshFormat* MeshFormatVar = StaticCast<UE::Mutable::Private::NodeMeshFormat*>(ParentNode);
 		AddChildFunc(MeshFormatVar->Source.get(), FString::Printf(TEXT("SOURCE MESH")));
 	}
 
-	else if (ParentNode->GetType() == mu::NodeModifierMeshClipMorphPlane::GetStaticType())
+	else if (ParentNode->GetType() == UE::Mutable::Private::NodeModifierMeshClipMorphPlane::GetStaticType())
 	{
 		// Nothing to show
 	}
 
-	else if (ParentNode->GetType() == mu::NodeModifierMeshClipWithMesh::GetStaticType())
+	else if (ParentNode->GetType() == UE::Mutable::Private::NodeModifierMeshClipWithMesh::GetStaticType())
 	{
-		mu::NodeModifierMeshClipWithMesh* ModifierMeshClipWithMeshVar = StaticCast<mu::NodeModifierMeshClipWithMesh*>(ParentNode);
+		UE::Mutable::Private::NodeModifierMeshClipWithMesh* ModifierMeshClipWithMeshVar = StaticCast<UE::Mutable::Private::NodeModifierMeshClipWithMesh*>(ParentNode);
 		AddChildFunc(ModifierMeshClipWithMeshVar->ClipMesh.get(), FString::Printf(TEXT("CLIP MESH")));
 	}
 
-	else if (ParentNode->GetType() == mu::NodeModifierMeshClipDeform::GetStaticType())
+	else if (ParentNode->GetType() == UE::Mutable::Private::NodeModifierMeshClipDeform::GetStaticType())
 	{
-		mu::NodeModifierMeshClipDeform* ModifierMeshClipDeformVar = StaticCast<mu::NodeModifierMeshClipDeform*>(ParentNode);
+		UE::Mutable::Private::NodeModifierMeshClipDeform* ModifierMeshClipDeformVar = StaticCast<UE::Mutable::Private::NodeModifierMeshClipDeform*>(ParentNode);
 		AddChildFunc(ModifierMeshClipDeformVar->ClipMesh.get(), FString::Printf(TEXT("CLIP MESH")));
 	}
 
-	else if (ParentNode->GetType() == mu::NodeModifierMeshClipWithUVMask::GetStaticType())
+	else if (ParentNode->GetType() == UE::Mutable::Private::NodeModifierMeshClipWithUVMask::GetStaticType())
 	{
-		mu::NodeModifierMeshClipWithUVMask* ModifierMeshClipWithUVMaskVar = StaticCast<mu::NodeModifierMeshClipWithUVMask*>(ParentNode);
+		UE::Mutable::Private::NodeModifierMeshClipWithUVMask* ModifierMeshClipWithUVMaskVar = StaticCast<UE::Mutable::Private::NodeModifierMeshClipWithUVMask*>(ParentNode);
 		AddChildFunc(ModifierMeshClipWithUVMaskVar->ClipMask.get(), FString::Printf(TEXT("CLIP MASK")));
 		AddChildFunc(ModifierMeshClipWithUVMaskVar->ClipLayout.get(), FString::Printf(TEXT("CLIP LAYOUT")));
 	}
 
-	else if (ParentNode->GetType() == mu::NodeModifierMeshTransformInMesh::GetStaticType())
+	else if (ParentNode->GetType() == UE::Mutable::Private::NodeModifierMeshTransformInMesh::GetStaticType())
 	{
-		mu::NodeModifierMeshTransformInMesh* ModifierMeshTransformInMeshVar = StaticCast<mu::NodeModifierMeshTransformInMesh*>(ParentNode);
+		UE::Mutable::Private::NodeModifierMeshTransformInMesh* ModifierMeshTransformInMeshVar = StaticCast<UE::Mutable::Private::NodeModifierMeshTransformInMesh*>(ParentNode);
 		AddChildFunc(ModifierMeshTransformInMeshVar->BoundingMesh.get(), FString::Printf(TEXT("BOUNDING MESH")));
 		AddChildFunc(ModifierMeshTransformInMeshVar->MatrixNode.get(), FString::Printf(TEXT("MESH TRANSFORM")));
 	}
 
-	else if (ParentNode->GetType() == mu::NodeModifierMeshTransformWithBone::GetStaticType())
+	else if (ParentNode->GetType() == UE::Mutable::Private::NodeModifierMeshTransformWithBone::GetStaticType())
 	{
-		mu::NodeModifierMeshTransformWithBone* ModifierMeshTransformWithBoneVar = StaticCast<mu::NodeModifierMeshTransformWithBone*>(ParentNode);
+		UE::Mutable::Private::NodeModifierMeshTransformWithBone* ModifierMeshTransformWithBoneVar = StaticCast<UE::Mutable::Private::NodeModifierMeshTransformWithBone*>(ParentNode);
 		AddChildFunc(ModifierMeshTransformWithBoneVar->MatrixNode.get(), FString::Printf(TEXT("MESH TRANSFORM")));
 	}
 
-	else if (ParentNode->GetType() == mu::NodeImageSwitch::GetStaticType())
+	else if (ParentNode->GetType() == UE::Mutable::Private::NodeImageSwitch::GetStaticType())
 	{
-		mu::NodeImageSwitch* ImageSwitchVar = StaticCast<mu::NodeImageSwitch*>(ParentNode);
+		UE::Mutable::Private::NodeImageSwitch* ImageSwitchVar = StaticCast<UE::Mutable::Private::NodeImageSwitch*>(ParentNode);
 		AddChildFunc(ImageSwitchVar->Parameter.get(), FString::Printf(TEXT("PARAM")));
 		for (int32 OptionIndex = 0; OptionIndex < ImageSwitchVar->Options.Num(); OptionIndex++)
 		{
@@ -635,46 +635,46 @@ void SMutableGraphViewer::GetChildrenForInfo(TSharedPtr<FMutableGraphTreeElement
 		}
 	}
 
-	else if (ParentNode->GetType() == mu::NodeImageMipmap::GetStaticType())
+	else if (ParentNode->GetType() == UE::Mutable::Private::NodeImageMipmap::GetStaticType())
 	{
-		mu::NodeImageMipmap* ImageMipMapVar = StaticCast<mu::NodeImageMipmap*>(ParentNode);
+		UE::Mutable::Private::NodeImageMipmap* ImageMipMapVar = StaticCast<UE::Mutable::Private::NodeImageMipmap*>(ParentNode);
 		AddChildFunc(ImageMipMapVar->Source.get(), FString::Printf(TEXT("SOURCE")));
 		AddChildFunc(ImageMipMapVar->Factor.get(), FString::Printf(TEXT("FACTOR")));
 	}
 
-	else if (ParentNode->GetType() == mu::NodeImageLayer::GetStaticType())
+	else if (ParentNode->GetType() == UE::Mutable::Private::NodeImageLayer::GetStaticType())
 	{
-		mu::NodeImageLayer* ImageLayerVar = StaticCast<mu::NodeImageLayer*>(ParentNode);
+		UE::Mutable::Private::NodeImageLayer* ImageLayerVar = StaticCast<UE::Mutable::Private::NodeImageLayer*>(ParentNode);
 		AddChildFunc(ImageLayerVar->Base.get(), FString::Printf(TEXT("BASE")));
 		AddChildFunc(ImageLayerVar->Mask.get(), FString::Printf(TEXT("MASK")));
 		AddChildFunc(ImageLayerVar->Blended.get(), FString::Printf(TEXT("BLEND")));
 	}
 	
-	else if (ParentNode->GetType() == mu::NodeImageLayerColour::GetStaticType())
+	else if (ParentNode->GetType() == UE::Mutable::Private::NodeImageLayerColour::GetStaticType())
 	{
-		mu::NodeImageLayerColour* ImageLayerColourVar = StaticCast<mu::NodeImageLayerColour*>(ParentNode);
+		UE::Mutable::Private::NodeImageLayerColour* ImageLayerColourVar = StaticCast<UE::Mutable::Private::NodeImageLayerColour*>(ParentNode);
 		AddChildFunc(ImageLayerColourVar->Base.get(), FString::Printf(TEXT("BASE")));
 		AddChildFunc(ImageLayerColourVar->Mask.get(), FString::Printf(TEXT("MASK")));
 		AddChildFunc(ImageLayerColourVar->Colour.get(), FString::Printf(TEXT("COLOR")));
 	}
 
-	else if (ParentNode->GetType() == mu::NodeImageResize::GetStaticType())
+	else if (ParentNode->GetType() == UE::Mutable::Private::NodeImageResize::GetStaticType())
 	{
-		mu::NodeImageResize* ImageResizeVar = StaticCast<mu::NodeImageResize*>(ParentNode);
+		UE::Mutable::Private::NodeImageResize* ImageResizeVar = StaticCast<UE::Mutable::Private::NodeImageResize*>(ParentNode);
 		AddChildFunc(ImageResizeVar->Base.get(), FString::Printf(TEXT("BASE")));
 	}
 
-	else if (ParentNode->GetType() == mu::NodeMeshMorph::GetStaticType())
+	else if (ParentNode->GetType() == UE::Mutable::Private::NodeMeshMorph::GetStaticType())
 	{
-		mu::NodeMeshMorph* MeshMorphVar = StaticCast<mu::NodeMeshMorph*>(ParentNode);
+		UE::Mutable::Private::NodeMeshMorph* MeshMorphVar = StaticCast<UE::Mutable::Private::NodeMeshMorph*>(ParentNode);
 		AddChildFunc(MeshMorphVar->Base.get(), FString::Printf(TEXT("BASE")));
 		AddChildFunc(MeshMorphVar->Morph.get(), FString::Printf(TEXT("MORPH")));
 		AddChildFunc(MeshMorphVar->Factor.get(), FString::Printf(TEXT("FACTOR")));
 	}
 
-	else if (ParentNode->GetType() == mu::NodeImageProject::GetStaticType())
+	else if (ParentNode->GetType() == UE::Mutable::Private::NodeImageProject::GetStaticType())
 	{
-		mu::NodeImageProject* ImageProjectVar = StaticCast<mu::NodeImageProject*>(ParentNode);
+		UE::Mutable::Private::NodeImageProject* ImageProjectVar = StaticCast<UE::Mutable::Private::NodeImageProject*>(ParentNode);
 		AddChildFunc(ImageProjectVar->Projector.get(), FString::Printf(TEXT("PROJECTOR")));
 		AddChildFunc(ImageProjectVar->Mesh.get(), FString::Printf(TEXT("MESH")));
 		AddChildFunc(ImageProjectVar->Image.get(), FString::Printf(TEXT("IMAGE")));
@@ -683,43 +683,43 @@ void SMutableGraphViewer::GetChildrenForInfo(TSharedPtr<FMutableGraphTreeElement
 		AddChildFunc(ImageProjectVar->AngleFadeEnd.get(), FString::Printf(TEXT("FADE END ANGLE")));
 	}
 
-	else if (ParentNode->GetType() == mu::NodeImagePlainColour::GetStaticType())
+	else if (ParentNode->GetType() == UE::Mutable::Private::NodeImagePlainColour::GetStaticType())
 	{
-		mu::NodeImagePlainColour* ImagePlainColourVar = StaticCast<mu::NodeImagePlainColour*>(ParentNode);
+		UE::Mutable::Private::NodeImagePlainColour* ImagePlainColourVar = StaticCast<UE::Mutable::Private::NodeImagePlainColour*>(ParentNode);
 		AddChildFunc(ImagePlainColourVar->Colour.get(), FString::Printf(TEXT("COLOR")));
 	}
 
-	else if (ParentNode->GetType() == mu::NodeLayout::GetStaticType())
+	else if (ParentNode->GetType() == UE::Mutable::Private::NodeLayout::GetStaticType())
 	{
 		// Nothing to show
 	}
 
-	else if (ParentNode->GetType() == mu::NodeScalarEnumParameter::GetStaticType())
+	else if (ParentNode->GetType() == UE::Mutable::Private::NodeScalarEnumParameter::GetStaticType())
 	{
-		mu::NodeScalarEnumParameter* ScalarEnumParameterVar = StaticCast<mu::NodeScalarEnumParameter*>(ParentNode);
+		UE::Mutable::Private::NodeScalarEnumParameter* ScalarEnumParameterVar = StaticCast<UE::Mutable::Private::NodeScalarEnumParameter*>(ParentNode);
 		for (int32 RangeIndex = 0; RangeIndex < ScalarEnumParameterVar->Ranges.Num(); RangeIndex++)
 		{
 			AddChildFunc(ScalarEnumParameterVar->Ranges[RangeIndex].get(), FString::Printf(TEXT("RANGE [%d]"), RangeIndex));
 		}
 	}
 
-	else if (ParentNode->GetType() == mu::NodeMeshFragment::GetStaticType())
+	else if (ParentNode->GetType() == UE::Mutable::Private::NodeMeshFragment::GetStaticType())
 	{
-		mu::NodeMeshFragment* MeshFragmentVar = StaticCast<mu::NodeMeshFragment*>(ParentNode);
+		UE::Mutable::Private::NodeMeshFragment* MeshFragmentVar = StaticCast<UE::Mutable::Private::NodeMeshFragment*>(ParentNode);
 		AddChildFunc(MeshFragmentVar->SourceMesh.get(), FString::Printf(TEXT("MESH")));
 	}
 
-	else if (ParentNode->GetType() == mu::NodeColourSampleImage::GetStaticType())
+	else if (ParentNode->GetType() == UE::Mutable::Private::NodeColourSampleImage::GetStaticType())
 	{
-		mu::NodeColourSampleImage* ColorSampleImageVar = StaticCast<mu::NodeColourSampleImage*>(ParentNode);
+		UE::Mutable::Private::NodeColourSampleImage* ColorSampleImageVar = StaticCast<UE::Mutable::Private::NodeColourSampleImage*>(ParentNode);
 		AddChildFunc(ColorSampleImageVar->Image.get(), FString::Printf(TEXT("IMAGE")));
 		AddChildFunc(ColorSampleImageVar->X.get(), FString::Printf(TEXT("X")));
 		AddChildFunc(ColorSampleImageVar->Y.get(), FString::Printf(TEXT("Y")));
 	}
 
-	else if (ParentNode->GetType() == mu::NodeImageInterpolate::GetStaticType())
+	else if (ParentNode->GetType() == UE::Mutable::Private::NodeImageInterpolate::GetStaticType())
 	{
-		mu::NodeImageInterpolate* ImageInterpolateVar = StaticCast<mu::NodeImageInterpolate*>(ParentNode);
+		UE::Mutable::Private::NodeImageInterpolate* ImageInterpolateVar = StaticCast<UE::Mutable::Private::NodeImageInterpolate*>(ParentNode);
 		AddChildFunc(ImageInterpolateVar->Factor.get(), FString::Printf(TEXT("FACTOR")));
 		for (int32 TargetIndex = 0; TargetIndex < ImageInterpolateVar->Targets.Num(); TargetIndex++)
 		{
@@ -727,71 +727,71 @@ void SMutableGraphViewer::GetChildrenForInfo(TSharedPtr<FMutableGraphTreeElement
 		}
 	}
 	
-	else if (ParentNode->GetType() == mu::NodeScalarConstant::GetStaticType())
+	else if (ParentNode->GetType() == UE::Mutable::Private::NodeScalarConstant::GetStaticType())
 	{
 		// Nothing to show
 	}
 
-	else if (ParentNode->GetType() == mu::NodeScalarParameter::GetStaticType())
+	else if (ParentNode->GetType() == UE::Mutable::Private::NodeScalarParameter::GetStaticType())
 	{
-		mu::NodeScalarParameter* ScalarParameterVar = StaticCast<mu::NodeScalarParameter*>(ParentNode);
+		UE::Mutable::Private::NodeScalarParameter* ScalarParameterVar = StaticCast<UE::Mutable::Private::NodeScalarParameter*>(ParentNode);
 		for (int32 RangeIndex = 0; RangeIndex < ScalarParameterVar->Ranges.Num(); RangeIndex++)
 		{
 			AddChildFunc(ScalarParameterVar->Ranges[RangeIndex].get(), FString::Printf(TEXT("RANGE [%d]"), RangeIndex));
 		}
 	}
 
-	else if (ParentNode->GetType() == mu::NodeColourParameter::GetStaticType())
+	else if (ParentNode->GetType() == UE::Mutable::Private::NodeColourParameter::GetStaticType())
 	{
-		mu::NodeColourParameter* ColorParameterVar = StaticCast<mu::NodeColourParameter*>(ParentNode);
+		UE::Mutable::Private::NodeColourParameter* ColorParameterVar = StaticCast<UE::Mutable::Private::NodeColourParameter*>(ParentNode);
 		for (int32 RangeIndex = 0; RangeIndex < ColorParameterVar->Ranges.Num(); RangeIndex++)
 		{
 			AddChildFunc(ColorParameterVar->Ranges[RangeIndex].get(), FString::Printf(TEXT("RANGE [%d]"), RangeIndex));
 		}
 	}
 
-	else if (ParentNode->GetType() == mu::NodeColourConstant::GetStaticType())
+	else if (ParentNode->GetType() == UE::Mutable::Private::NodeColourConstant::GetStaticType())
 	{
 		// Nothing to show
 	}
 
 
-	else if (ParentNode->GetType() == mu::NodeImageConstant::GetStaticType())
+	else if (ParentNode->GetType() == UE::Mutable::Private::NodeImageConstant::GetStaticType())
 	{
 		// Nothing to show
 	}
 
 	
-	else if (ParentNode->GetType() == mu::NodeScalarCurve::GetStaticType())
+	else if (ParentNode->GetType() == UE::Mutable::Private::NodeScalarCurve::GetStaticType())
 	{
-		mu::NodeScalarCurve* ScalarCurveVar = StaticCast<mu::NodeScalarCurve*>(ParentNode);
+		UE::Mutable::Private::NodeScalarCurve* ScalarCurveVar = StaticCast<UE::Mutable::Private::NodeScalarCurve*>(ParentNode);
 		AddChildFunc(ScalarCurveVar->CurveSampleValue.get(), FString::Printf(TEXT("INPUT")));
 	}
 
-	else if (ParentNode->GetType() == mu::NodeMeshMakeMorph::GetStaticType())
+	else if (ParentNode->GetType() == UE::Mutable::Private::NodeMeshMakeMorph::GetStaticType())
 	{
-		mu::NodeMeshMakeMorph* MeshMakeMorphVar = StaticCast<mu::NodeMeshMakeMorph*>(ParentNode);
+		UE::Mutable::Private::NodeMeshMakeMorph* MeshMakeMorphVar = StaticCast<UE::Mutable::Private::NodeMeshMakeMorph*>(ParentNode);
 		AddChildFunc(MeshMakeMorphVar->Base.get(), FString::Printf(TEXT("BASE")));
 		AddChildFunc(MeshMakeMorphVar->Target.get(), FString::Printf(TEXT("TARGET")));
 	}
 	
-	else if (ParentNode->GetType() == mu::NodeProjectorParameter::GetStaticType())
+	else if (ParentNode->GetType() == UE::Mutable::Private::NodeProjectorParameter::GetStaticType())
 	{
-		mu::NodeProjectorParameter* ProjectorParameterVar = StaticCast<mu::NodeProjectorParameter*>(ParentNode);
+		UE::Mutable::Private::NodeProjectorParameter* ProjectorParameterVar = StaticCast<UE::Mutable::Private::NodeProjectorParameter*>(ParentNode);
 		for (int32 RangeIndex = 0; RangeIndex < ProjectorParameterVar->Ranges.Num(); RangeIndex++)
 		{
 			AddChildFunc(ProjectorParameterVar->Ranges[RangeIndex].get(), FString::Printf(TEXT("RANGE [%d]"), RangeIndex));
 		}
 	}
 
-	else if (ParentNode->GetType() == mu::NodeProjectorConstant::GetStaticType())
+	else if (ParentNode->GetType() == UE::Mutable::Private::NodeProjectorConstant::GetStaticType())
 	{
 		// Nothing to show
 	}
 
-	else if (ParentNode->GetType() == mu::NodeColourSwitch::GetStaticType())
+	else if (ParentNode->GetType() == UE::Mutable::Private::NodeColourSwitch::GetStaticType())
 	{
-		mu::NodeColourSwitch* ColorSwitchVar = StaticCast<mu::NodeColourSwitch*>(ParentNode);
+		UE::Mutable::Private::NodeColourSwitch* ColorSwitchVar = StaticCast<UE::Mutable::Private::NodeColourSwitch*>(ParentNode);
 		AddChildFunc(ColorSwitchVar->Parameter.get(), TEXT("PARAM"));
 		for (int32 OptionIndex = 0; OptionIndex < ColorSwitchVar->Options.Num(); ++OptionIndex)
 		{
@@ -799,52 +799,52 @@ void SMutableGraphViewer::GetChildrenForInfo(TSharedPtr<FMutableGraphTreeElement
 		}
 	}
 
-	else if (ParentNode->GetType() == mu::NodeImageSwizzle::GetStaticType())
+	else if (ParentNode->GetType() == UE::Mutable::Private::NodeImageSwizzle::GetStaticType())
 	{
-		mu::NodeImageSwizzle* ImageSwizzleVar = StaticCast<mu::NodeImageSwizzle*>(ParentNode);
+		UE::Mutable::Private::NodeImageSwizzle* ImageSwizzleVar = StaticCast<UE::Mutable::Private::NodeImageSwizzle*>(ParentNode);
 		for (int32 SourceIndex = 0; SourceIndex < ImageSwizzleVar->Sources.Num(); ++SourceIndex)
 		{
 			AddChildFunc(ImageSwizzleVar->Sources[SourceIndex].get(), FString::Printf(TEXT("SOURCE [%d]"), SourceIndex));
 		}
 	}
 	
-	else if (ParentNode->GetType() == mu::NodeImageInvert::GetStaticType())
+	else if (ParentNode->GetType() == UE::Mutable::Private::NodeImageInvert::GetStaticType())
 	{
-		mu::NodeImageInvert* ImageInvertVar = StaticCast<mu::NodeImageInvert*>(ParentNode);
+		UE::Mutable::Private::NodeImageInvert* ImageInvertVar = StaticCast<UE::Mutable::Private::NodeImageInvert*>(ParentNode);
 		AddChildFunc(ImageInvertVar->Base.get(), TEXT("BASE"));
 	}
 
-	else if (ParentNode->GetType() == mu::NodeImageMultiLayer::GetStaticType())
+	else if (ParentNode->GetType() == UE::Mutable::Private::NodeImageMultiLayer::GetStaticType())
 	{
-		mu::NodeImageMultiLayer* ImageMultilayerVar = StaticCast<mu::NodeImageMultiLayer*>(ParentNode);
+		UE::Mutable::Private::NodeImageMultiLayer* ImageMultilayerVar = StaticCast<UE::Mutable::Private::NodeImageMultiLayer*>(ParentNode);
 		AddChildFunc(ImageMultilayerVar->Base.get(), FString::Printf(TEXT("BASE")));
 		AddChildFunc(ImageMultilayerVar->Mask.get(), FString::Printf(TEXT("MASK")));
 		AddChildFunc(ImageMultilayerVar->Blended.get(), FString::Printf(TEXT("BLEND")));
 		AddChildFunc(ImageMultilayerVar->Range.get(), FString::Printf(TEXT("RANGE")));
 	}
 	
-	else if (ParentNode->GetType() == mu::NodeImageTable::GetStaticType())
+	else if (ParentNode->GetType() == UE::Mutable::Private::NodeImageTable::GetStaticType())
 	{
 		// No nodes to show
 	}
 
-	else if (ParentNode->GetType() == mu::NodeMeshTable::GetStaticType())
+	else if (ParentNode->GetType() == UE::Mutable::Private::NodeMeshTable::GetStaticType())
 	{
-		mu::NodeMeshTable* MeshTableVar = StaticCast<mu::NodeMeshTable*>(ParentNode);
+		UE::Mutable::Private::NodeMeshTable* MeshTableVar = StaticCast<UE::Mutable::Private::NodeMeshTable*>(ParentNode);
 		for (int32 LayoutIndex = 0; LayoutIndex < MeshTableVar->Layouts.Num(); ++LayoutIndex)
 		{
 			AddChildFunc(MeshTableVar->Layouts[LayoutIndex].get(), FString::Printf(TEXT("LAYOUT [%d]"), LayoutIndex));
 		}
 	}
 
-	else if (ParentNode->GetType() == mu::NodeScalarTable::GetStaticType())
+	else if (ParentNode->GetType() == UE::Mutable::Private::NodeScalarTable::GetStaticType())
 	{
 		// Nothing to show
 	}
 
-	else if (ParentNode->GetType() == mu::NodeScalarSwitch::GetStaticType())
+	else if (ParentNode->GetType() == UE::Mutable::Private::NodeScalarSwitch::GetStaticType())
 	{
-		mu::NodeScalarSwitch* ScalarSwitchVar = StaticCast<mu::NodeScalarSwitch*>(ParentNode);
+		UE::Mutable::Private::NodeScalarSwitch* ScalarSwitchVar = StaticCast<UE::Mutable::Private::NodeScalarSwitch*>(ParentNode);
 		AddChildFunc(ScalarSwitchVar->Parameter.get(), TEXT("PARAM"));
 		for (int32 OptionIndex = 0; OptionIndex < ScalarSwitchVar->Options.Num(); ++OptionIndex)
 		{
@@ -852,9 +852,9 @@ void SMutableGraphViewer::GetChildrenForInfo(TSharedPtr<FMutableGraphTreeElement
 		}
 	}
 	
-	else if (ParentNode->GetType() == mu::NodeColourFromScalars::GetStaticType())
+	else if (ParentNode->GetType() == UE::Mutable::Private::NodeColourFromScalars::GetStaticType())
 	{
-		mu::NodeColourFromScalars* ScalarTableVar = StaticCast<mu::NodeColourFromScalars*>(ParentNode);
+		UE::Mutable::Private::NodeColourFromScalars* ScalarTableVar = StaticCast<UE::Mutable::Private::NodeColourFromScalars*>(ParentNode);
 		AddChildFunc(ScalarTableVar->X.get(), TEXT("X"));
 		AddChildFunc(ScalarTableVar->Y.get(), TEXT("Y"));
 		AddChildFunc(ScalarTableVar->Z.get(), TEXT("Z"));
diff --git a/Engine/Plugins/Mutable/Source/CustomizableObjectEditor/Private/MuCOE/SMutableGraphViewer.h b/Engine/Plugins/Mutable/Source/CustomizableObjectEditor/Private/MuCOE/SMutableGraphViewer.h
--- a/Engine/Plugins/Mutable/Source/CustomizableObjectEditor/Private/MuCOE/SMutableGraphViewer.h
+++ b/Engine/Plugins/Mutable/Source/CustomizableObjectEditor/Private/MuCOE/SMutableGraphViewer.h
@@ -39,7 +39,7 @@ public:
 
 	SLATE_END_ARGS()
 
-	void Construct(const FArguments& InArgs, const mu::NodePtr& InRootNode);
+	void Construct(const FArguments& InArgs, const UE::Mutable::Private::NodePtr& InRootNode);
 
 	// FGCObject interface
 	virtual void AddReferencedObjects(FReferenceCollector& Collector) override;
@@ -47,7 +47,7 @@ public:
 
 private:
 	/** The Mutable Graph to show, represented by its root. */
-	mu::NodePtr RootNode;
+	UE::Mutable::Private::NodePtr RootNode;
 
 	/** Array of external referenced textures in MutableModel, indexed by id. */
 	TArray<TSoftObjectPtr<const UTexture>> ReferencedRuntimeTextures;
@@ -67,8 +67,8 @@ private:
 	*/
 	struct FItemCacheKey
 	{
-		const mu::Node* Parent = nullptr;
-		const mu::Node* Child = nullptr;
+		const UE::Mutable::Private::Node* Parent = nullptr;
+		const UE::Mutable::Private::Node* Child = nullptr;
 		uint32 ChildIndexInParent = 0;
 
 		friend FORCEINLINE bool operator == (const FItemCacheKey& A, const FItemCacheKey& B)
@@ -84,7 +84,7 @@ private:
 	TMap< FItemCacheKey, TSharedPtr<FMutableGraphTreeElement>> ItemCache;
 
 	/** Main tree item for each node. A graph node can be representes with multiple tree ndoes if it is reachable from different paths. */
-	TMap< const mu::Node*, TSharedPtr<FMutableGraphTreeElement>> MainItemPerNode;
+	TMap< const UE::Mutable::Private::Node*, TSharedPtr<FMutableGraphTreeElement>> MainItemPerNode;
 
 	/** */
 	void RebuildTree();
@@ -103,7 +103,7 @@ private:
 class FMutableGraphTreeElement : public TSharedFromThis<FMutableGraphTreeElement>
 {
 public:
-	FMutableGraphTreeElement(const mu::NodePtr& InNode, TSharedPtr<FMutableGraphTreeElement>* InDuplicatedOf=nullptr, const FString& InPrefix=FString() )
+	FMutableGraphTreeElement(const UE::Mutable::Private::NodePtr& InNode, TSharedPtr<FMutableGraphTreeElement>* InDuplicatedOf=nullptr, const FString& InPrefix=FString() )
 	{
 		MutableNode = InNode;
 		Prefix = InPrefix;
@@ -116,7 +116,7 @@ public:
 public:
 
 	/** Mutable Graph Node represented in this tree row. */
-	mu::NodePtr MutableNode;
+	UE::Mutable::Private::NodePtr MutableNode;
 
 	/** If this tree element is a duplicated of another node, this is the node. */
 	TSharedPtr<FMutableGraphTreeElement> DuplicatedOf;
diff --git a/Engine/Plugins/Mutable/Source/CustomizableObjectEditor/Private/MuCOE/SMutableImageViewer.cpp b/Engine/Plugins/Mutable/Source/CustomizableObjectEditor/Private/MuCOE/SMutableImageViewer.cpp
--- a/Engine/Plugins/Mutable/Source/CustomizableObjectEditor/Private/MuCOE/SMutableImageViewer.cpp
+++ b/Engine/Plugins/Mutable/Source/CustomizableObjectEditor/Private/MuCOE/SMutableImageViewer.cpp
@@ -267,16 +267,16 @@ void SMutableImageViewer::Tick(const FGeometry& AllottedGeometry, const double I
 		{
 			// If the image format is not supported in the editor platform (for instance, it is a mobile-specific format)
 			// convert the mutable image to a generic format so that it can be previewed.
-			TSharedPtr<const mu::FImage> ImageToConvert = MutableImage;
+			TSharedPtr<const UE::Mutable::Private::FImage> ImageToConvert = MutableImage;
 
 			// \TODO: This test is a bit weak
-			if (MutableImage->GetFormat() >= mu::EImageFormat::ASTC_4x4_RGB_LDR
+			if (MutableImage->GetFormat() >= UE::Mutable::Private::EImageFormat::ASTC_4x4_RGB_LDR
 				&&
-				MutableImage->GetFormat() <= mu::EImageFormat::ASTC_10x10_RG_LDR)
+				MutableImage->GetFormat() <= UE::Mutable::Private::EImageFormat::ASTC_10x10_RG_LDR)
 			{
-				mu::FImageOperator ImOp = mu::FImageOperator::GetDefault(mu::FImageOperator::FImagePixelFormatFunc());
+				UE::Mutable::Private::FImageOperator ImOp = UE::Mutable::Private::FImageOperator::GetDefault(UE::Mutable::Private::FImageOperator::FImagePixelFormatFunc());
 				int32 Quality = 4;
-				ImageToConvert = ImOp.ImagePixelFormat(Quality, MutableImage.Get(), mu::GetUncompressedFormat(MutableImage->GetFormat()) );
+				ImageToConvert = ImOp.ImagePixelFormat(Quality, MutableImage.Get(), UE::Mutable::Private::GetUncompressedFormat(MutableImage->GetFormat()) );
 			}
 
 			// Should we extract a single channel?
@@ -305,7 +305,7 @@ void SMutableImageViewer::Tick(const FGeometry& AllottedGeometry, const double I
 }
 
 
-void SMutableImageViewer::SetImage(const TSharedPtr<const mu::FImage>& InMutableImage, int32 CurrentLOD)
+void SMutableImageViewer::SetImage(const TSharedPtr<const UE::Mutable::Private::FImage>& InMutableImage, int32 CurrentLOD)
 {
 	MutableImage = InMutableImage;
 	CurrentVisibleLOD = CurrentLOD;
@@ -330,9 +330,9 @@ FText SMutableImageViewer::GetImageDescriptionLabel() const
 		return FText::FromString(TEXT("No image."));
 	}
 
-	mu::EImageFormat Format = MutableImage->GetFormat();
+	UE::Mutable::Private::EImageFormat Format = MutableImage->GetFormat();
 	FString Label = FString::Printf(TEXT("%s - %d x %d - %d LODs"), 
-		StringCast<TCHAR>(mu::TypeInfo::s_imageFormatName[(size_t)Format]).Get(),
+		StringCast<TCHAR>(UE::Mutable::Private::TypeInfo::s_imageFormatName[(size_t)Format]).Get(),
 		MutableImage->GetSizeX(), 
 		MutableImage->GetSizeY(),
 		MutableImage->GetLODCount() );
diff --git a/Engine/Plugins/Mutable/Source/CustomizableObjectEditor/Private/MuCOE/SMutableImageViewer.h b/Engine/Plugins/Mutable/Source/CustomizableObjectEditor/Private/MuCOE/SMutableImageViewer.h
--- a/Engine/Plugins/Mutable/Source/CustomizableObjectEditor/Private/MuCOE/SMutableImageViewer.h
+++ b/Engine/Plugins/Mutable/Source/CustomizableObjectEditor/Private/MuCOE/SMutableImageViewer.h
@@ -78,7 +78,7 @@ public:
 
 	SLATE_BEGIN_ARGS(SMutableImageViewer) {}
 		SLATE_ATTRIBUTE(FIntPoint, GridSize)
-		SLATE_ARGUMENT(TSharedPtr<const mu::FImage>, Image)
+		SLATE_ARGUMENT(TSharedPtr<const UE::Mutable::Private::FImage>, Image)
 	SLATE_END_ARGS()
 
 	void Construct(const FArguments& InArgs);
@@ -89,7 +89,7 @@ public:
 	// Own interface
 
 	/** Set the image to show in the widget. */
-	void SetImage(const TSharedPtr<const mu::FImage>& Image, int32 VisibleLOD );
+	void SetImage(const TSharedPtr<const UE::Mutable::Private::FImage>& Image, int32 VisibleLOD );
 
 private:
 
@@ -100,7 +100,7 @@ private:
 	TAttribute<FIntPoint> GridSize;
 
 	/** */
-	TSharedPtr<const mu::FImage> MutableImage;
+	TSharedPtr<const UE::Mutable::Private::FImage> MutableImage;
 
 	/** Mipmap shown in the widget. 0 is the biggest one. */
 	int32 CurrentVisibleLOD=0;
diff --git a/Engine/Plugins/Mutable/Source/CustomizableObjectEditor/Private/MuCOE/SMutableInstanceViewer.cpp b/Engine/Plugins/Mutable/Source/CustomizableObjectEditor/Private/MuCOE/SMutableInstanceViewer.cpp
--- a/Engine/Plugins/Mutable/Source/CustomizableObjectEditor/Private/MuCOE/SMutableInstanceViewer.cpp
+++ b/Engine/Plugins/Mutable/Source/CustomizableObjectEditor/Private/MuCOE/SMutableInstanceViewer.cpp
@@ -280,7 +280,7 @@ void SMutableInstanceViewer::Construct(const FArguments& InArgs, const TSharedPt
 }
 
 
-void SMutableInstanceViewer::SetInstance(const TSharedPtr<const mu::FInstance>& InInstance, const TSharedPtr<const mu::FModel>& Model, const TSharedPtr<const mu::FParameters>& Parameters, const mu::FSystem& System)
+void SMutableInstanceViewer::SetInstance(const TSharedPtr<const UE::Mutable::Private::FInstance>& InInstance, const TSharedPtr<const UE::Mutable::Private::FModel>& Model, const TSharedPtr<const UE::Mutable::Private::FParameters>& Parameters, const UE::Mutable::Private::FSystem& System)
 {
 	if (InInstance != MutableInstance)
 	{
@@ -309,19 +309,19 @@ void SMutableInstanceViewer::SetInstance(const TSharedPtr<const mu::FInstance>&
 					LODArray->Add(NewLOD);
 
 					NewLOD->LODIndex = FText::AsNumber(LODIndex);
-					mu::FMeshId MeshId = MutableInstance->GetMeshId(ComponentIndex, LODIndex);
+					UE::Mutable::Private::FMeshId MeshId = MutableInstance->GetMeshId(ComponentIndex, LODIndex);
 					NewLOD->MeshId = FText::FromString(MeshId.ToString());
 
 					// Generate the mesh to get some data
 					if (MeshId)
 					{
-						mu::OP::ADDRESS MeshRootAddress = MeshId.GetKey()->Address;
-						TSharedPtr<const mu::FMesh> MutableMesh = System.GetPrivate()->BuildMesh(
+						UE::Mutable::Private::OP::ADDRESS MeshRootAddress = MeshId.GetKey()->Address;
+						TSharedPtr<const UE::Mutable::Private::FMesh> MutableMesh = System.GetPrivate()->BuildMesh(
 								ExternalResourceProvider,
 								Model, 
 								Parameters.Get(), 
 								MeshRootAddress, 
-								mu::EMeshContentFlags::AllFlags);
+								UE::Mutable::Private::EMeshContentFlags::AllFlags);
 
 						if (MutableMesh && MutableMesh->IsReference())
 						{
diff --git a/Engine/Plugins/Mutable/Source/CustomizableObjectEditor/Private/MuCOE/SMutableInstanceViewer.h b/Engine/Plugins/Mutable/Source/CustomizableObjectEditor/Private/MuCOE/SMutableInstanceViewer.h
--- a/Engine/Plugins/Mutable/Source/CustomizableObjectEditor/Private/MuCOE/SMutableInstanceViewer.h
+++ b/Engine/Plugins/Mutable/Source/CustomizableObjectEditor/Private/MuCOE/SMutableInstanceViewer.h
@@ -60,7 +60,7 @@ public:
 	void Construct(const FArguments& InArgs, const TSharedPtr<FUnrealMutableResourceProvider>& InExternalResourceProvider);
 
 	/** Set the Mutable Instance to be used for this widget */
-	void SetInstance(const TSharedPtr<const mu::FInstance>&, const TSharedPtr<const mu::FModel>&, const TSharedPtr<const mu::FParameters>&, const mu::FSystem&);
+	void SetInstance(const TSharedPtr<const UE::Mutable::Private::FInstance>&, const TSharedPtr<const UE::Mutable::Private::FModel>&, const TSharedPtr<const UE::Mutable::Private::FParameters>&, const UE::Mutable::Private::FSystem&);
 	
 	/** */
 	TSharedRef<SWidget> GenerateLODsListView(const TSharedPtr<TArray<TSharedPtr<FMutableInstanceViewerLODElement>>>& InBufferChannelElements);
@@ -70,7 +70,7 @@ private:
 	TSharedPtr<FUnrealMutableResourceProvider> ExternalResourceProvider;
 	
 	/** Data backend for the widget. It represents the Instance that is being "displayed" */
-	TSharedPtr<const mu::FInstance> MutableInstance = nullptr;
+	TSharedPtr<const UE::Mutable::Private::FInstance> MutableInstance = nullptr;
 
 	/** Splitter used to separate the two sides of the slate (tables and viewport) */
 	TSharedPtr<SSplitter> SpaceSplitter;
diff --git a/Engine/Plugins/Mutable/Source/CustomizableObjectEditor/Private/MuCOE/SMutableLayoutViewer.cpp b/Engine/Plugins/Mutable/Source/CustomizableObjectEditor/Private/MuCOE/SMutableLayoutViewer.cpp
--- a/Engine/Plugins/Mutable/Source/CustomizableObjectEditor/Private/MuCOE/SMutableLayoutViewer.cpp
+++ b/Engine/Plugins/Mutable/Source/CustomizableObjectEditor/Private/MuCOE/SMutableLayoutViewer.cpp
@@ -120,7 +120,7 @@ void SMutableLayoutViewer::Tick(const FGeometry& AllottedGeometry, const double
 }
 
 
-void SMutableLayoutViewer::SetLayout(const TSharedPtr<const mu::FLayout>& InMutableLayout)
+void SMutableLayoutViewer::SetLayout(const TSharedPtr<const UE::Mutable::Private::FLayout>& InMutableLayout)
 {
 	MutableLayout = InMutableLayout;
 	bIsPendingUpdate = true;
diff --git a/Engine/Plugins/Mutable/Source/CustomizableObjectEditor/Private/MuCOE/SMutableLayoutViewer.h b/Engine/Plugins/Mutable/Source/CustomizableObjectEditor/Private/MuCOE/SMutableLayoutViewer.h
--- a/Engine/Plugins/Mutable/Source/CustomizableObjectEditor/Private/MuCOE/SMutableLayoutViewer.h
+++ b/Engine/Plugins/Mutable/Source/CustomizableObjectEditor/Private/MuCOE/SMutableLayoutViewer.h
@@ -17,7 +17,7 @@ public:
 
 	SLATE_BEGIN_ARGS(SMutableLayoutViewer) {}
 		SLATE_ATTRIBUTE(FIntPoint, GridSize)
-		SLATE_ARGUMENT(TSharedPtr<const mu::FLayout>, Layout)
+		SLATE_ARGUMENT(TSharedPtr<const UE::Mutable::Private::FLayout>, Layout)
 	SLATE_END_ARGS()
 
 	void Construct(const FArguments& InArgs);
@@ -28,7 +28,7 @@ public:
 	// Own interface
 
 	/** Set the image to show in the widget. */
-	void SetLayout(const TSharedPtr<const mu::FLayout>& Image );
+	void SetLayout(const TSharedPtr<const UE::Mutable::Private::FLayout>& Image );
 
 private:
 
@@ -36,7 +36,7 @@ private:
 	TSharedPtr<SCustomizableObjectLayoutGrid> LayoutViewer;
 
 	/** Mutable Layout object being previewed */
-	TSharedPtr<const mu::FLayout> MutableLayout;
+	TSharedPtr<const UE::Mutable::Private::FLayout> MutableLayout;
 
 	/** Is true, the image or the visible LOD have changed and we need to update. */
 	bool bIsPendingUpdate = false;
diff --git a/Engine/Plugins/Mutable/Source/CustomizableObjectEditor/Private/MuCOE/SMutableMeshViewer.cpp b/Engine/Plugins/Mutable/Source/CustomizableObjectEditor/Private/MuCOE/SMutableMeshViewer.cpp
--- a/Engine/Plugins/Mutable/Source/CustomizableObjectEditor/Private/MuCOE/SMutableMeshViewer.cpp
+++ b/Engine/Plugins/Mutable/Source/CustomizableObjectEditor/Private/MuCOE/SMutableMeshViewer.cpp
@@ -200,7 +200,7 @@ void SMutableMeshViewer::Construct(const FArguments& InArgs)
 }
 
 
-void SMutableMeshViewer::SetMesh(const TSharedPtr<const mu::FMesh>& InMesh)
+void SMutableMeshViewer::SetMesh(const TSharedPtr<const UE::Mutable::Private::FMesh>& InMesh)
 {
 	if (InMesh != MutableMesh)
 	{
@@ -219,7 +219,7 @@ void SMutableMeshViewer::SetMesh(const TSharedPtr<const mu::FMesh>& InMesh)
 			OnMeshChanged();
 
 			// Skeleton slate viewer update
-			if (TSharedPtr<const mu::FSkeleton> MutableSkeleton = MutableMesh->GetSkeleton())
+			if (TSharedPtr<const UE::Mutable::Private::FSkeleton> MutableSkeleton = MutableMesh->GetSkeleton())
 			{
 				MutableSkeletonViewer->SetSkeleton(MutableSkeleton);
 
@@ -564,8 +564,8 @@ TArray<FVector2f> SMutableMeshViewer::GetUVs() const
 		// TODO
 		int32 UVIndex = 0;
 
-		mu::UntypedMeshBufferIteratorConst ItIndices(MutableMesh->IndexBuffers, mu::EMeshBufferSemantic::VertexIndex, 0);
-		mu::UntypedMeshBufferIteratorConst It(MutableMesh->VertexBuffers, mu::EMeshBufferSemantic::TexCoords, UVIndex);
+		UE::Mutable::Private::UntypedMeshBufferIteratorConst ItIndices(MutableMesh->IndexBuffers, UE::Mutable::Private::EMeshBufferSemantic::VertexIndex, 0);
+		UE::Mutable::Private::UntypedMeshBufferIteratorConst It(MutableMesh->VertexBuffers, UE::Mutable::Private::EMeshBufferSemantic::TexCoords, UVIndex);
 		for (int32 FaceIndex=0; FaceIndex<NumFaces; ++FaceIndex)
 		{
 			uint32 V0 = (ItIndices + FaceIndex * 3 + 0).GetAsUINT32();
@@ -681,7 +681,7 @@ TSharedRef<SWidget> SMutableMeshViewer::GenerateBufferChannelsListView(
 
 void SMutableMeshViewer::OnMeshChanged()
 {
-	// Cache the data accessible from the mu::FMesh to be later used by the UI
+	// Cache the data accessible from the UE::Mutable::Private::FMesh to be later used by the UI
 	FillTargetBufferSetDataArray(MutableMesh->GetVertexBuffers(), VertexBuffers, VertexBuffersSlateView);
 	FillTargetBufferSetDataArray(MutableMesh->GetIndexBuffers(), IndexBuffers, IndexBuffersSlateView);
 
@@ -691,7 +691,7 @@ void SMutableMeshViewer::OnMeshChanged()
 }
 
 
-void SMutableMeshViewer::FillTargetBufferSetDataArray(const mu::FMeshBufferSet& InMuMeshBufferSet,
+void SMutableMeshViewer::FillTargetBufferSetDataArray(const UE::Mutable::Private::FMeshBufferSet& InMuMeshBufferSet,
                                                       TArray<TSharedPtr<FBufferElement>>& OutBuffersDataArray,
                                                       TSharedPtr<SListView<TSharedPtr<FBufferElement>>>& InHostListView)
 {
@@ -728,8 +728,8 @@ void SMutableMeshViewer::FillTargetBufferSetDataArray(const mu::FMeshBufferSet&
 		{
 			TSharedPtr<FBufferChannelElement> NewChannelDefinition = MakeShareable(new FBufferChannelElement());
 
-			mu::EMeshBufferSemantic BufferChannelSemantic = mu::EMeshBufferSemantic::None;
-			mu::EMeshBufferFormat BufferFormat = mu::EMeshBufferFormat::None;
+			UE::Mutable::Private::EMeshBufferSemantic BufferChannelSemantic = UE::Mutable::Private::EMeshBufferSemantic::None;
+			UE::Mutable::Private::EMeshBufferFormat BufferFormat = UE::Mutable::Private::EMeshBufferFormat::None;
 			int32 BufferComponentCount = 0;
 			int32 SemanticIndex = -1;
 
@@ -749,12 +749,12 @@ void SMutableMeshViewer::FillTargetBufferSetDataArray(const mu::FMeshBufferSet&
 			const int32 SemanticTypeIndex = static_cast<int32>(BufferChannelSemantic);
 			const int32 BufferFormatIndex = static_cast<int32>(BufferFormat);
 
-			// Using mu::TypeInfo find what is the name of the buffer semantic and the buffer format
+			// Using UE::Mutable::Private::TypeInfo find what is the name of the buffer semantic and the buffer format
 			NewChannelDefinition->SemanticIndex = FText::AsNumber(SemanticIndex);
 			NewChannelDefinition->BufferSemantic = FText::FromString(
-				*FString(mu::TypeInfo::s_meshBufferSemanticName[SemanticTypeIndex]));
+				*FString(UE::Mutable::Private::TypeInfo::s_meshBufferSemanticName[SemanticTypeIndex]));
 			NewChannelDefinition->BufferFormat = FText::FromString(
-				*FString(mu::TypeInfo::s_meshBufferFormatName[BufferFormatIndex]));
+				*FString(UE::Mutable::Private::TypeInfo::s_meshBufferFormatName[BufferFormatIndex]));
 			NewChannelDefinition->BufferComponentCount = FText::FromString(*FString::FromInt(BufferComponentCount));
 
 			ChannelsArray->Add(NewChannelDefinition);
diff --git a/Engine/Plugins/Mutable/Source/CustomizableObjectEditor/Private/MuCOE/SMutableMeshViewer.h b/Engine/Plugins/Mutable/Source/CustomizableObjectEditor/Private/MuCOE/SMutableMeshViewer.h
--- a/Engine/Plugins/Mutable/Source/CustomizableObjectEditor/Private/MuCOE/SMutableMeshViewer.h
+++ b/Engine/Plugins/Mutable/Source/CustomizableObjectEditor/Private/MuCOE/SMutableMeshViewer.h
@@ -18,7 +18,7 @@ class SSplitter;
 class STableViewBase;
 class SWidget;
 class USkeletalMesh;
-namespace mu { class FMeshBufferSet; }
+namespace UE::Mutable::Private { class FMeshBufferSet; }
 struct FAssetData;
 
 
@@ -47,14 +47,14 @@ class SMutableMeshViewer : public SCompoundWidget
 {
 public:
 	SLATE_BEGIN_ARGS(SMutableMeshViewer) {}
-		SLATE_ARGUMENT_DEFAULT(TSharedPtr<const mu::FMesh>, Mesh){nullptr};
+		SLATE_ARGUMENT_DEFAULT(TSharedPtr<const UE::Mutable::Private::FMesh>, Mesh){nullptr};
 	SLATE_END_ARGS()
 
 	/** Builds the widget */
 	void Construct(const FArguments& InArgs);
 
 	/** Set the Mutable Mesh to be used for this widget */
-	void SetMesh(const TSharedPtr<const mu::FMesh>& InMesh);
+	void SetMesh(const TSharedPtr<const UE::Mutable::Private::FMesh>& InMesh);
 
 	
 	/** Generate a new SListView for all the ChannelElements provided.
@@ -74,7 +74,7 @@ private:
 	TSharedRef<ITableRow> GenerateTagRow(TSharedPtr<FString> InItem, const TSharedRef<STableViewBase>& OwnerTable);
 
 	/** Data backend for the widget. It represents the mesh that is being "displayed" */
-	TSharedPtr<const mu::FMesh> MutableMesh = nullptr;
+	TSharedPtr<const UE::Mutable::Private::FMesh> MutableMesh = nullptr;
 
 	/** Splitter used to separate the two sides of the slate (tables and viewport) */
 	TSharedPtr<SSplitter> SpaceSplitter;
@@ -134,7 +134,7 @@ private:
 	* @param OutBuffersDataArray - TArray witch will contain the parsed data from the provided mutable buffer set converted onto FBufferElements
 	* @param InHostListView - The slate list using the previously referenced array as data backend. 
 	*/
-	void FillTargetBufferSetDataArray(const mu::FMeshBufferSet& InMuMeshBufferSetPtr,
+	void FillTargetBufferSetDataArray(const UE::Mutable::Private::FMeshBufferSet& InMuMeshBufferSetPtr,
 	                                  TArray<TSharedPtr<FBufferElement>>& OutBuffersDataArray,
 	                                  TSharedPtr<SListView<TSharedPtr<FBufferElement>>>& InHostListView);
 
diff --git a/Engine/Plugins/Mutable/Source/CustomizableObjectEditor/Private/MuCOE/SMutableMeshViewport.cpp b/Engine/Plugins/Mutable/Source/CustomizableObjectEditor/Private/MuCOE/SMutableMeshViewport.cpp
--- a/Engine/Plugins/Mutable/Source/CustomizableObjectEditor/Private/MuCOE/SMutableMeshViewport.cpp
+++ b/Engine/Plugins/Mutable/Source/CustomizableObjectEditor/Private/MuCOE/SMutableMeshViewport.cpp
@@ -33,7 +33,7 @@ void SMutableMeshViewport::Tick(const FGeometry& AllottedGeometry, const double
 	SEditorViewport::Tick(AllottedGeometry, InCurrentTime, InDeltaTime);
 }
 
-void SMutableMeshViewport::SetMesh(const TSharedPtr<const mu::FMesh>& InMesh)
+void SMutableMeshViewport::SetMesh(const TSharedPtr<const UE::Mutable::Private::FMesh>& InMesh)
 {
 	if (InMesh == MutableMesh)
 	{
diff --git a/Engine/Plugins/Mutable/Source/CustomizableObjectEditor/Private/MuCOE/SMutableMeshViewport.h b/Engine/Plugins/Mutable/Source/CustomizableObjectEditor/Private/MuCOE/SMutableMeshViewport.h
--- a/Engine/Plugins/Mutable/Source/CustomizableObjectEditor/Private/MuCOE/SMutableMeshViewport.h
+++ b/Engine/Plugins/Mutable/Source/CustomizableObjectEditor/Private/MuCOE/SMutableMeshViewport.h
@@ -22,7 +22,7 @@ struct FGeometry;
 class SMutableMeshViewport : public SEditorViewport , public FGCObject
 {
 	SLATE_BEGIN_ARGS(SMutableMeshViewport) {}
-    	SLATE_ARGUMENT_DEFAULT(TSharedPtr<const mu::FMesh>, Mesh) { nullptr };
+    	SLATE_ARGUMENT_DEFAULT(TSharedPtr<const UE::Mutable::Private::FMesh>, Mesh) { nullptr };
     SLATE_END_ARGS()
     
 public:
@@ -32,7 +32,7 @@ public:
     /** Set the Mutable Mesh to be used for this widget
      * @param InMesh - The mutable mesh to be displayed
      */
-    void SetMesh(const TSharedPtr<const mu::FMesh>& InMesh);
+    void SetMesh(const TSharedPtr<const UE::Mutable::Private::FMesh>& InMesh);
 
 	// FGCObject interface
 	virtual void AddReferencedObjects(FReferenceCollector& Collector) override;
@@ -52,7 +52,7 @@ protected:
 private:
 
 	/** The mutable mesh being displayed */
-	TSharedPtr<const mu::FMesh> MutableMesh = nullptr;
+	TSharedPtr<const UE::Mutable::Private::FMesh> MutableMesh = nullptr;
 
 	/*
 	 *	UI  
diff --git a/Engine/Plugins/Mutable/Source/CustomizableObjectEditor/Private/MuCOE/SMutableProjectorViewer.cpp b/Engine/Plugins/Mutable/Source/CustomizableObjectEditor/Private/MuCOE/SMutableProjectorViewer.cpp
--- a/Engine/Plugins/Mutable/Source/CustomizableObjectEditor/Private/MuCOE/SMutableProjectorViewer.cpp
+++ b/Engine/Plugins/Mutable/Source/CustomizableObjectEditor/Private/MuCOE/SMutableProjectorViewer.cpp
@@ -172,7 +172,7 @@ void SMutableProjectorViewer::Construct(const FArguments& InArgs)
 }
 
 
-void SMutableProjectorViewer::SetProjector(const mu::FProjector& InMutableProjector)
+void SMutableProjectorViewer::SetProjector(const UE::Mutable::Private::FProjector& InMutableProjector)
 {
 	MutableProjector = InMutableProjector;
 }
@@ -229,19 +229,19 @@ FText SMutableProjectorViewer::GetProjectorAngleAsText() const
 }
 
 
-mu::EProjectorType SMutableProjectorViewer::GetProjectorType() const
+UE::Mutable::Private::EProjectorType SMutableProjectorViewer::GetProjectorType() const
 {
-	const mu::EProjectorType ProjectorType = MutableProjector.type;
+	const UE::Mutable::Private::EProjectorType ProjectorType = MutableProjector.type;
 	return ProjectorType;
 }
 
 
 FText SMutableProjectorViewer::GetProjectorTypeAsText() const
 {
-	const mu::EProjectorType ProjectorType =  GetProjectorType();
+	const UE::Mutable::Private::EProjectorType ProjectorType =  GetProjectorType();
 	const uint32 TypeIndex = static_cast<uint32> (ProjectorType);
 
-	return FText::FromString(* FString(mu::TypeInfo::s_projectorTypeName[TypeIndex]));
+	return FText::FromString(* FString(UE::Mutable::Private::TypeInfo::s_projectorTypeName[TypeIndex]));
 }
 
 #undef LOCTEXT_NAMESPACE
