[Mutable] Change namespace name

diff --git a/Engine/Plugins/Mutable/Source/CustomizableObjectEditor/Private/MuCOE/SMutableProjectorViewer.h b/Engine/Plugins/Mutable/Source/CustomizableObjectEditor/Private/MuCOE/SMutableProjectorViewer.h
--- a/Engine/Plugins/Mutable/Source/CustomizableObjectEditor/Private/MuCOE/SMutableProjectorViewer.h
+++ b/Engine/Plugins/Mutable/Source/CustomizableObjectEditor/Private/MuCOE/SMutableProjectorViewer.h
@@ -5,13 +5,13 @@
 #include "MuR/ParametersPrivate.h"
 #include "Widgets/SCompoundWidget.h"
 
-namespace mu { enum class EProjectorType : uint32; }
+namespace UE::Mutable::Private { enum class EProjectorType : uint32; }
 
 
 class SMutableProjectorViewer final : public SCompoundWidget
 {
 	SLATE_BEGIN_ARGS(SMutableProjectorViewer) {}
-		SLATE_ARGUMENT(mu::FProjector,MutableProjector);
+		SLATE_ARGUMENT(UE::Mutable::Private::FProjector,MutableProjector);
 	SLATE_END_ARGS()
 
 public:
@@ -26,12 +26,12 @@ public:
 	 * Sets the projector object whose data we want to inspect
 	 * @param InMutableProjector - Projector we desire to inspect. It is the origin of the data shown on the UI
 	 */
-	void SetProjector(const mu::FProjector& InMutableProjector);
+	void SetProjector(const UE::Mutable::Private::FProjector& InMutableProjector);
 
 private:
 
 	/** The mutable projector witch data is being exposed on the UI */
-	mu::FProjector MutableProjector;
+	UE::Mutable::Private::FProjector MutableProjector;
 	
 	/** Get the float corresponding to the position of the projector on the targeted component index
 	 * @param VectorComponentIndex - Determines if you get the X, Y or Z value from the vector (0,1 and 2)
@@ -72,7 +72,7 @@ private:
 	 * It accesses the set projector and returns the type of it.
 	 * @return - The projector type.
 	 */
-	mu::EProjectorType GetProjectorType() const;
+	UE::Mutable::Private::EProjectorType GetProjectorType() const;
 	
 	/**
 	 * Provides the UI with the mutable type in a format it can display on an STextBlock
diff --git a/Engine/Plugins/Mutable/Source/CustomizableObjectEditor/Private/MuCOE/SMutableSkeletonViewer.cpp b/Engine/Plugins/Mutable/Source/CustomizableObjectEditor/Private/MuCOE/SMutableSkeletonViewer.cpp
--- a/Engine/Plugins/Mutable/Source/CustomizableObjectEditor/Private/MuCOE/SMutableSkeletonViewer.cpp
+++ b/Engine/Plugins/Mutable/Source/CustomizableObjectEditor/Private/MuCOE/SMutableSkeletonViewer.cpp
@@ -10,7 +10,7 @@ class ITableRow;
 
 #define LOCTEXT_NAMESPACE "MutableSkeletonViewer"
 
-/** Row on the tree representing the bone structure of the mu::FMesh. It represents the UI side of the BoneDefinition */
+/** Row on the tree representing the bone structure of the UE::Mutable::Private::FMesh. It represents the UI side of the BoneDefinition */
 class SMutableMeshBoneTreeRow : public STableRow<TSharedPtr<FBoneDefinition>>
 {
 public:
@@ -91,7 +91,7 @@ void SMutableSkeletonViewer::Construct(const FArguments& InArgs)
 	SetSkeleton(InArgs._Skeleton);
 }
 
-void SMutableSkeletonViewer::SetSkeleton(TSharedPtr<const mu::FSkeleton> InSkeleton)
+void SMutableSkeletonViewer::SetSkeleton(TSharedPtr<const UE::Mutable::Private::FSkeleton> InSkeleton)
 {
 	if (InSkeleton != Skeleton)
 	{
diff --git a/Engine/Plugins/Mutable/Source/CustomizableObjectEditor/Private/MuCOE/SMutableSkeletonViewer.h b/Engine/Plugins/Mutable/Source/CustomizableObjectEditor/Private/MuCOE/SMutableSkeletonViewer.h
--- a/Engine/Plugins/Mutable/Source/CustomizableObjectEditor/Private/MuCOE/SMutableSkeletonViewer.h
+++ b/Engine/Plugins/Mutable/Source/CustomizableObjectEditor/Private/MuCOE/SMutableSkeletonViewer.h
@@ -24,7 +24,7 @@ struct FBoneDefinition
 class SMutableSkeletonViewer : public SCompoundWidget
 {
 	SLATE_BEGIN_ARGS(SMutableSkeletonViewer) {}
-		SLATE_ARGUMENT_DEFAULT(TSharedPtr<const mu::FSkeleton>, Skeleton) { nullptr };
+		SLATE_ARGUMENT_DEFAULT(TSharedPtr<const UE::Mutable::Private::FSkeleton>, Skeleton) { nullptr };
 	SLATE_END_ARGS()
 
 public:
@@ -32,12 +32,12 @@ public:
 	void Construct(const FArguments& InArgs);
 
 	/** Sets the mutable skeleton object to be inspected */
-	void SetSkeleton(TSharedPtr<const mu::FSkeleton> InSkeleton);
+	void SetSkeleton(TSharedPtr<const UE::Mutable::Private::FSkeleton> InSkeleton);
 	
 private:
 
 	/** The mutable skeleton whose data is being inspected */
-	TSharedPtr<const mu::FSkeleton> Skeleton;
+	TSharedPtr<const UE::Mutable::Private::FSkeleton> Skeleton;
 	
 	/** Bones data: Bone definitions used to generate the bone tree structure */
 	TArray<TSharedPtr<FBoneDefinition>> RootBones;
diff --git a/Engine/Plugins/Mutable/Source/CustomizableObjectEditor/Private/MuCOE/SStandAloneAssetPicker.cpp b/Engine/Plugins/Mutable/Source/CustomizableObjectEditor/Private/MuCOE/SStandAloneAssetPicker.cpp
--- a/Engine/Plugins/Mutable/Source/CustomizableObjectEditor/Private/MuCOE/SStandAloneAssetPicker.cpp
+++ b/Engine/Plugins/Mutable/Source/CustomizableObjectEditor/Private/MuCOE/SStandAloneAssetPicker.cpp
@@ -154,7 +154,7 @@ void SStandAloneAssetPicker::OnAssetSelectedFromPicker(const struct FAssetData&
 {
 	AssetPickerAnchor->SetIsOpen(false);
 
-	CurrentAsset = MutablePrivate::LoadObject(AssetData);
+	CurrentAsset = UE::Mutable::Private::LoadObject(AssetData);
 
 	OnAssetSelected.ExecuteIfBound(CurrentAsset.Get());
 
@@ -168,7 +168,7 @@ void SStandAloneAssetPicker::OnAssetEnterPressedFromPicker(const TArray<struct F
 	CurrentAsset = nullptr;
 	if (AssetData.Num() > 0)
 	{
-		CurrentAsset = MutablePrivate::LoadObject(AssetData[0]);
+		CurrentAsset = UE::Mutable::Private::LoadObject(AssetData[0]);
 	}
 
 	OnAssetSelected.ExecuteIfBound(CurrentAsset.Get());
diff --git a/Engine/Plugins/Mutable/Source/CustomizableObjectEditor/Private/MuCOE/UnrealBakeHelpers.cpp b/Engine/Plugins/Mutable/Source/CustomizableObjectEditor/Private/MuCOE/UnrealBakeHelpers.cpp
--- a/Engine/Plugins/Mutable/Source/CustomizableObjectEditor/Private/MuCOE/UnrealBakeHelpers.cpp
+++ b/Engine/Plugins/Mutable/Source/CustomizableObjectEditor/Private/MuCOE/UnrealBakeHelpers.cpp
@@ -204,8 +204,8 @@ UTexture2D* FUnrealBakeHelpers::BakeHelper_CreateAssetTexture(UTexture2D* Source
 	}
 
 	// Create a mutable image from the platform data.
-	mu::EImageFormat MutableFormat = UnrealToMutablePixelFormat(SourceTexture->GetPlatformData()->PixelFormat, SourceTexture->HasAlphaChannel());
-	TSharedPtr<mu::FImage> PlatformImage = MakeShared<mu::FImage>(SourceTextureSizeX, SourceTextureSizeY, 1, MutableFormat, mu::EInitializationType::NotInitialized);
+	UE::Mutable::Private::EImageFormat MutableFormat = UnrealToMutablePixelFormat(SourceTexture->GetPlatformData()->PixelFormat, SourceTexture->HasAlphaChannel());
+	TSharedPtr<UE::Mutable::Private::FImage> PlatformImage = MakeShared<UE::Mutable::Private::FImage>(SourceTextureSizeX, SourceTextureSizeY, 1, MutableFormat, UE::Mutable::Private::EInitializationType::NotInitialized);
 
 	constexpr int32 MipIndex = 0;
 	const uint8* SourceData = reinterpret_cast<const uint8*>(SourceTexture->GetPlatformData()->Mips[MipIndex].BulkData.LockReadOnly());
@@ -217,23 +217,23 @@ UTexture2D* FUnrealBakeHelpers::BakeHelper_CreateAssetTexture(UTexture2D* Source
 	SourceTexture->GetPlatformData()->Mips[MipIndex].BulkData.Unlock();
 
 	// Reformat the mutable image
-	mu::EImageFormat UncompressedMutableFormat = mu::EImageFormat::RGBA_UByte;
+	UE::Mutable::Private::EImageFormat UncompressedMutableFormat = UE::Mutable::Private::EImageFormat::RGBA_UByte;
 	switch (SourceTexture->GetPlatformData()->PixelFormat)
 	{
 	case PF_G8:
 	case PF_L8:
 	case PF_A8:
 	case PF_BC4:
-		UncompressedMutableFormat = mu::EImageFormat::L_UByte;
+		UncompressedMutableFormat = UE::Mutable::Private::EImageFormat::L_UByte;
 		break;
 
 	default:
 		break;
 	}
 
-	mu::FImageOperator ImOp = mu::FImageOperator::GetDefault(mu::FImageOperator::FImagePixelFormatFunc());
+	UE::Mutable::Private::FImageOperator ImOp = UE::Mutable::Private::FImageOperator::GetDefault(UE::Mutable::Private::FImageOperator::FImagePixelFormatFunc());
 	int32 Quality = 4; // Doesn't matter for decompression.
-	TSharedPtr<mu::FImage> UncompressedImage = ImOp.ImagePixelFormat(Quality, PlatformImage.Get(), UncompressedMutableFormat);
+	TSharedPtr<UE::Mutable::Private::FImage> UncompressedImage = ImOp.ImagePixelFormat(Quality, PlatformImage.Get(), UncompressedMutableFormat);
 
 	// Copy the decompressed data to the texture source data
 	int32 SourceDataSize = DupTex->Source.CalcMipSize(MipIndex);
diff --git a/Engine/Plugins/Mutable/Source/CustomizableObjectEditor/Private/MuCOE/Widgets/SCustomizableObjectMacroLibraryList.cpp b/Engine/Plugins/Mutable/Source/CustomizableObjectEditor/Private/MuCOE/Widgets/SCustomizableObjectMacroLibraryList.cpp
--- a/Engine/Plugins/Mutable/Source/CustomizableObjectEditor/Private/MuCOE/Widgets/SCustomizableObjectMacroLibraryList.cpp
+++ b/Engine/Plugins/Mutable/Source/CustomizableObjectEditor/Private/MuCOE/Widgets/SCustomizableObjectMacroLibraryList.cpp
@@ -12,7 +12,7 @@
 
 #define LOCTEXT_NAMESPACE "SCustomizableObjectMacroLibraryList"
 
-/** Row on the tree representing the bone structure of the mu::FMesh. It represents the UI side of the BoneDefinition */
+/** Row on the tree representing the bone structure of the UE::Mutable::Private::FMesh. It represents the UI side of the BoneDefinition */
 class SMacroLibraryTreeRow : public STableRow<TSharedPtr<FMacroTreeEntry>>
 {
 public:
diff --git a/Engine/Plugins/Mutable/Source/CustomizableObjectEditor/Public/MuCOE/ExtensionDataCompilerInterface.h b/Engine/Plugins/Mutable/Source/CustomizableObjectEditor/Public/MuCOE/ExtensionDataCompilerInterface.h
--- a/Engine/Plugins/Mutable/Source/CustomizableObjectEditor/Public/MuCOE/ExtensionDataCompilerInterface.h
+++ b/Engine/Plugins/Mutable/Source/CustomizableObjectEditor/Public/MuCOE/ExtensionDataCompilerInterface.h
@@ -31,7 +31,7 @@ public:
 	 * TSoftObjectPtr, because the container will be moved to a different package during cooking
 	 * and your paths won't be automatically updated.
 	 */
-	UE_API TSharedPtr<const mu::FExtensionData> MakeStreamedExtensionData(FInstancedStruct&& Data);
+	UE_API TSharedPtr<const UE::Mutable::Private::FExtensionData> MakeStreamedExtensionData(FInstancedStruct&& Data);
 
 	/**
 	 * Register a new Extension Data compile-time constant that will always be loaded.
@@ -42,12 +42,12 @@ public:
 	 * Any UObjects referenced by the provided instanced struct that aren't in an asset package
 	 * should be created with the return value of GetOuterForAlwaysLoadedObjects as their Outer.
 	 *
-	 * Move the constant data into this function and use the resulting mu::FExtensionData in the
-	 * mu::Node graph, e.g. set it as the value of a mu::NodeExtensionDataConstant.
+	 * Move the constant data into this function and use the resulting UE::Mutable::Private::FExtensionData in the
+	 * UE::Mutable::Private::Node graph, e.g. set it as the value of a UE::Mutable::Private::NodeExtensionDataConstant.
 	 * 
 	 * As with MakeStreamedExtensionData, the provided instanced struct will be visible to the GC.
 	 */
-	UE_API TSharedPtr<const mu::FExtensionData> MakeAlwaysLoadedExtensionData(FInstancedStruct&& Data);
+	UE_API TSharedPtr<const UE::Mutable::Private::FExtensionData> MakeAlwaysLoadedExtensionData(FInstancedStruct&& Data);
 
 	/** The Outer to use for objects owned by always-loaded Extension Data constants. */
 	UE_API const UObject* GetOuterForAlwaysLoadedObjects();
diff --git a/Engine/Plugins/Mutable/Source/CustomizableObjectEditor/Public/MuCOE/ICustomizableObjectExtensionNode.h b/Engine/Plugins/Mutable/Source/CustomizableObjectEditor/Public/MuCOE/ICustomizableObjectExtensionNode.h
--- a/Engine/Plugins/Mutable/Source/CustomizableObjectEditor/Public/MuCOE/ICustomizableObjectExtensionNode.h
+++ b/Engine/Plugins/Mutable/Source/CustomizableObjectEditor/Public/MuCOE/ICustomizableObjectExtensionNode.h
@@ -21,6 +21,6 @@ class ICustomizableObjectExtensionNode
 	GENERATED_BODY()
 
 public:
-	/** Generate a mu::Node that produces FExtensionData */
-	virtual mu::Ptr<mu::NodeExtensionData> GenerateMutableNode(FExtensionDataCompilerInterface& CompilerInterface) const = 0;
+	/** Generate a UE::Mutable::Private::Node that produces FExtensionData */
+	virtual UE::Mutable::Private::Ptr<UE::Mutable::Private::NodeExtensionData> GenerateMutableNode(FExtensionDataCompilerInterface& CompilerInterface) const = 0;
 };
diff --git a/Engine/Plugins/Mutable/Source/MutableRuntime/Internal/MuR/CodeVisitor.h b/Engine/Plugins/Mutable/Source/MutableRuntime/Internal/MuR/CodeVisitor.h
--- a/Engine/Plugins/Mutable/Source/MutableRuntime/Internal/MuR/CodeVisitor.h
+++ b/Engine/Plugins/Mutable/Source/MutableRuntime/Internal/MuR/CodeVisitor.h
@@ -13,7 +13,7 @@
 #include "MuR/System.h"
 #include "MuR/SystemPrivate.h"
 
-namespace mu
+namespace UE::Mutable::Private
 {
 
     /** Decide what operations are an "add resource" since they are handled differently sometimes. */
diff --git a/Engine/Plugins/Mutable/Source/MutableRuntime/Internal/MuR/ImageDataStorage.h b/Engine/Plugins/Mutable/Source/MutableRuntime/Internal/MuR/ImageDataStorage.h
--- a/Engine/Plugins/Mutable/Source/MutableRuntime/Internal/MuR/ImageDataStorage.h
+++ b/Engine/Plugins/Mutable/Source/MutableRuntime/Internal/MuR/ImageDataStorage.h
@@ -13,7 +13,7 @@
 
 #define UE_API MUTABLERUNTIME_API
 
-namespace mu::MemoryCounters
+namespace UE::Mutable::Private::MemoryCounters
 {
 	struct FImageMemoryCounter
 	{
@@ -22,12 +22,12 @@ namespace mu::MemoryCounters
 }
 
 
-namespace mu
+namespace UE::Mutable::Private
 {
 	using FImageArray = TArray<uint8, FDefaultMemoryTrackingAllocator<MemoryCounters::FImageMemoryCounter>>;
 }
 
-namespace mu::FImageDataStorageInternal
+namespace UE::Mutable::Private::FImageDataStorageInternal
 {
 	inline int32 ComputeNumLODsForSize(FImageSize Size)
 	{
@@ -35,7 +35,7 @@ namespace mu::FImageDataStorageInternal
 	}
 }
 
-namespace mu
+namespace UE::Mutable::Private
 {
 class FImageDataStorage
 {
diff --git a/Engine/Plugins/Mutable/Source/MutableRuntime/Internal/MuR/MemoryCounters.h b/Engine/Plugins/Mutable/Source/MutableRuntime/Internal/MuR/MemoryCounters.h
--- a/Engine/Plugins/Mutable/Source/MutableRuntime/Internal/MuR/MemoryCounters.h
+++ b/Engine/Plugins/Mutable/Source/MutableRuntime/Internal/MuR/MemoryCounters.h
@@ -5,7 +5,7 @@
 #include "MuR/MemoryTrackingAllocationPolicy.h"
 
 
-namespace mu::MemoryCounters
+namespace UE::Mutable::Private::MemoryCounters
 {
 	struct FInternalMemoryCounter
 	{
@@ -14,7 +14,7 @@ namespace mu::MemoryCounters
 }
 
 
-namespace mu
+namespace UE::Mutable::Private
 {
 	using AllocType = FDefaultMemoryTrackingAllocator<MemoryCounters::FInternalMemoryCounter>;
 
diff --git a/Engine/Plugins/Mutable/Source/MutableRuntime/Internal/MuR/MemoryTrackingAllocationPolicy.h b/Engine/Plugins/Mutable/Source/MutableRuntime/Internal/MuR/MemoryTrackingAllocationPolicy.h
--- a/Engine/Plugins/Mutable/Source/MutableRuntime/Internal/MuR/MemoryTrackingAllocationPolicy.h
+++ b/Engine/Plugins/Mutable/Source/MutableRuntime/Internal/MuR/MemoryTrackingAllocationPolicy.h
@@ -9,7 +9,7 @@
 #include <atomic>
 #include <type_traits>
 
-namespace mu
+namespace UE::Mutable::Private
 {
 	
 	
@@ -180,7 +180,7 @@ namespace mu
 }
 
 template<typename BaseAlloc, typename Counter>
-struct TAllocatorTraits<mu::TMemoryTrackingAllocatorWrapper<BaseAlloc, Counter>> : public TAllocatorTraits<BaseAlloc>
+struct TAllocatorTraits<UE::Mutable::Private::TMemoryTrackingAllocatorWrapper<BaseAlloc, Counter>> : public TAllocatorTraits<BaseAlloc>
 {
 	enum { SupportsMoveFromOtherAllocator = false };
 };
diff --git a/Engine/Plugins/Mutable/Source/MutableRuntime/Internal/MuR/ModelPrivate.h b/Engine/Plugins/Mutable/Source/MutableRuntime/Internal/MuR/ModelPrivate.h
--- a/Engine/Plugins/Mutable/Source/MutableRuntime/Internal/MuR/ModelPrivate.h
+++ b/Engine/Plugins/Mutable/Source/MutableRuntime/Internal/MuR/ModelPrivate.h
@@ -18,7 +18,7 @@
 #define MUTABLE_MAX_RUNTIME_PARAMETERS_PER_STATE	65
 #define MUTABLE_GROW_BORDER_VALUE					2
 
-namespace mu
+namespace UE::Mutable::Private
 {
 
 	/** Used to debug and log. */
@@ -1116,7 +1116,7 @@ namespace mu
     {
     public:
 
-		mu::FProgram Program;
+		UE::Mutable::Private::FProgram Program;
 
     	UE_API void UnloadRoms();
 
diff --git a/Engine/Plugins/Mutable/Source/MutableRuntime/Internal/MuR/ParallelExecutionUtils.h b/Engine/Plugins/Mutable/Source/MutableRuntime/Internal/MuR/ParallelExecutionUtils.h
--- a/Engine/Plugins/Mutable/Source/MutableRuntime/Internal/MuR/ParallelExecutionUtils.h
+++ b/Engine/Plugins/Mutable/Source/MutableRuntime/Internal/MuR/ParallelExecutionUtils.h
@@ -4,7 +4,7 @@
 
 #include "Templates/Function.h"
 
-namespace mu::ParallelExecutionUtils
+namespace UE::Mutable::Private::ParallelExecutionUtils
 {
     void InvokeBatchParallelFor(int32 Num, TFunctionRef<void(int32)> Body);
     void InvokeBatchParallelForSingleThread(int32 Num, TFunctionRef<void(int32)> Body);
diff --git a/Engine/Plugins/Mutable/Source/MutableRuntime/Internal/MuR/SystemPrivate.h b/Engine/Plugins/Mutable/Source/MutableRuntime/Internal/MuR/SystemPrivate.h
--- a/Engine/Plugins/Mutable/Source/MutableRuntime/Internal/MuR/SystemPrivate.h
+++ b/Engine/Plugins/Mutable/Source/MutableRuntime/Internal/MuR/SystemPrivate.h
@@ -28,7 +28,7 @@
 #endif
 
 
-namespace mu
+namespace UE::Mutable::Private
 {
 	class FExtensionDataStreamer;
 
@@ -1062,7 +1062,7 @@ namespace mu
 			}
 			check(Data->DataTypeIndex != 0);
 
-			mu::UpdateLLMStats();
+			UE::Mutable::Private::UpdateLLMStats();
 		}
 
 		void SetMesh(FCacheAddress At, TSharedPtr<const FMesh> Value)
@@ -1086,7 +1086,7 @@ namespace mu
 			}
 			check(Data->DataTypeIndex != 0);
 
-			mu::UpdateLLMStats();
+			UE::Mutable::Private::UpdateLLMStats();
 		}
 
 		void SetLayout(FCacheAddress at, TSharedPtr<const FLayout> v)
@@ -1108,7 +1108,7 @@ namespace mu
 			}
 			check(Data->DataTypeIndex != 0);
 
-			mu::UpdateLLMStats();
+			UE::Mutable::Private::UpdateLLMStats();
 		}
 
 		void SetString(FCacheAddress at, TSharedPtr<const String> v)
diff --git a/Engine/Plugins/Mutable/Source/MutableRuntime/Private/MuR/CodeRunner.cpp b/Engine/Plugins/Mutable/Source/MutableRuntime/Private/MuR/CodeRunner.cpp
--- a/Engine/Plugins/Mutable/Source/MutableRuntime/Private/MuR/CodeRunner.cpp
+++ b/Engine/Plugins/Mutable/Source/MutableRuntime/Private/MuR/CodeRunner.cpp
@@ -107,7 +107,7 @@ static FAutoConsoleVariableRef CVarUseImageTransformVectorImpl (
 
 #define AddOp(...) Invoke([&](){ AddOp(__VA_ARGS__); });
 
-namespace mu::Private
+namespace UE::Mutable::Private::Private
 {
 	FString DumpItemScheduledCallstack(const FScheduledOp& Item)
 	{
@@ -130,7 +130,7 @@ namespace mu::Private
 
 #endif
 
-namespace mu::MemoryCounters
+namespace UE::Mutable::Private::MemoryCounters
 {
 	std::atomic<SSIZE_T>& FStreamingMemoryCounter::Get()
 	{
@@ -139,7 +139,7 @@ namespace mu::MemoryCounters
 	}
 }
 
-namespace mu
+namespace UE::Mutable::Private
 {
 
 	TSharedRef<CodeRunner> CodeRunner::Create(
@@ -547,7 +547,7 @@ namespace mu
                 }
                 else
                 {
-					Result = mu::CloneOrTakeOver<FInstance>(Base);
+					Result = UE::Mutable::Private::CloneOrTakeOver<FInstance>(Base);
                 }
 
                 if ( Args.value )
@@ -592,7 +592,7 @@ namespace mu
                 }
                 else
                 {
-                    Result = mu::CloneOrTakeOver<FInstance>(Base);
+                    Result = UE::Mutable::Private::CloneOrTakeOver<FInstance>(Base);
                 }
 
                 if ( Args.value )
@@ -637,7 +637,7 @@ namespace mu
                 }
                 else
                 {
-					Result = mu::CloneOrTakeOver<FInstance>(Base);
+					Result = UE::Mutable::Private::CloneOrTakeOver<FInstance>(Base);
 				}
 
                 if ( Args.value )
@@ -682,7 +682,7 @@ namespace mu
                 }
                 else
                 {
-					Result = mu::CloneOrTakeOver<FInstance>(Base);
+					Result = UE::Mutable::Private::CloneOrTakeOver<FInstance>(Base);
 				}
 
                 if ( Args.value )
@@ -728,7 +728,7 @@ namespace mu
 				TSharedPtr<FInstance> Result;
 				if (Base)
 				{
-					Result = mu::CloneOrTakeOver<FInstance>(Base);
+					Result = UE::Mutable::Private::CloneOrTakeOver<FInstance>(Base);
 				}
 				else
 				{
@@ -882,7 +882,7 @@ namespace mu
 					TSharedPtr<FInstance> Instance;
 					if (Args.Instance)
 					{
-						Instance = mu::CloneOrTakeOver<FInstance>(LoadInstance(FCacheAddress(Args.Instance, Item)));
+						Instance = UE::Mutable::Private::CloneOrTakeOver<FInstance>(LoadInstance(FCacheAddress(Args.Instance, Item)));
 					}
 					else
 					{
@@ -950,7 +950,7 @@ namespace mu
                 }
                 else
                 {
-					Result = mu::CloneOrTakeOver<FInstance>(Base);
+					Result = UE::Mutable::Private::CloneOrTakeOver<FInstance>(Base);
 				}
 
                 if ( Args.value )
@@ -991,7 +991,7 @@ namespace mu
                 }
                 else
                 {
-					Result = mu::CloneOrTakeOver<FInstance>(Base);
+					Result = UE::Mutable::Private::CloneOrTakeOver<FInstance>(Base);
 				}
 
                 if ( Args.value )
@@ -1040,7 +1040,7 @@ namespace mu
         				}
         				else
         				{
-        					Result = mu::CloneOrTakeOver<FInstance>(Base);
+        					Result = UE::Mutable::Private::CloneOrTakeOver<FInstance>(Base);
         				}
 
         				if (Args.Material)
@@ -1078,7 +1078,7 @@ namespace mu
         				TSharedPtr<FInstance> Result;
         				if (Base)
         				{
-        					Result = mu::CloneOrTakeOver<FInstance>(Base);
+        					Result = UE::Mutable::Private::CloneOrTakeOver<FInstance>(Base);
         				}
         				else
         				{
@@ -1168,7 +1168,7 @@ namespace mu
 			OP::ResourceConstantArgs Args = Program.GetOpArgs<OP::ResourceConstantArgs>(Item.At);
 
 			// Assume the ROM has been loaded previously
-			TSharedPtr<const mu::FExtensionData> SourceConst;
+			TSharedPtr<const UE::Mutable::Private::FExtensionData> SourceConst;
 			Program.GetExtensionDataConstant(Args.value, SourceConst);
 
 			check(SourceConst);
@@ -5598,7 +5598,7 @@ namespace mu
 	
 	TSharedPtr<FRangeIndex> CodeRunner::BuildCurrentOpRangeIndex(const FScheduledOp& Item, int32 ParameterIndex)
 	{
-		return mu::BuildCurrentOpRangeIndex(Item, *Params, *Model.Get(), GetMemory(), ParameterIndex);
+		return UE::Mutable::Private::BuildCurrentOpRangeIndex(Item, *Params, *Model.Get(), GetMemory(), ParameterIndex);
 	}
 
     //---------------------------------------------------------------------------------------------
diff --git a/Engine/Plugins/Mutable/Source/MutableRuntime/Private/MuR/CodeRunner.h b/Engine/Plugins/Mutable/Source/MutableRuntime/Private/MuR/CodeRunner.h
--- a/Engine/Plugins/Mutable/Source/MutableRuntime/Private/MuR/CodeRunner.h
+++ b/Engine/Plugins/Mutable/Source/MutableRuntime/Private/MuR/CodeRunner.h
@@ -21,7 +21,7 @@
 #include "Templates/SharedPointer.h"
 #include "Templates/Tuple.h"
 
-namespace mu::MemoryCounters
+namespace UE::Mutable::Private::MemoryCounters
 {
 	struct FStreamingMemoryCounter
 	{
@@ -29,7 +29,7 @@ namespace mu::MemoryCounters
 	};
 }
 
-namespace  mu
+namespace  UE::Mutable::Private
 {
 	class FModel;
 	class FParameters;
@@ -156,7 +156,7 @@ namespace  mu
 
 		/** Load an external image asynchronously, returns an event to wait for complition and a cleanup function that must be called once the event has completed. */
 		TTuple<UE::Tasks::FTask, TFunction<void()>> LoadExternalImageAsync(FExternalResourceId Id, uint8 MipmapsToSkip, TFunction<void(TSharedPtr<FImage>)>& ResultCallback);
- 	    mu::FExtendedImageDesc GetExternalImageDesc(UTexture* Id);
+ 	    UE::Mutable::Private::FExtendedImageDesc GetExternalImageDesc(UTexture* Id);
 
 		/** Load an external mesh asynchronously, returns an event to wait for complition and a cleanup function that must be called once the event has completed. */
 		TTuple<UE::Tasks::FTask, TFunction<void()>> LoadExternalMeshAsync(FExternalResourceId Id, int32 LODIndex, int32 SectionIndex, TFunction<void(TSharedPtr<FMesh>)>& ResultCallback);
@@ -732,7 +732,7 @@ namespace  mu
 }
 
 #if MUTABLE_DEBUG_CODERUNNER_TASK_SCHEDULE_CALLSTACK
-namespace mu::Private
+namespace UE::Mutable::Private::Private
 {
 	FString DumpItemScheduledCallstack(const FScheduledOp& Item);
 }
diff --git a/Engine/Plugins/Mutable/Source/MutableRuntime/Private/MuR/CodeRunnerBegin.cpp b/Engine/Plugins/Mutable/Source/MutableRuntime/Private/MuR/CodeRunnerBegin.cpp
--- a/Engine/Plugins/Mutable/Source/MutableRuntime/Private/MuR/CodeRunnerBegin.cpp
+++ b/Engine/Plugins/Mutable/Source/MutableRuntime/Private/MuR/CodeRunnerBegin.cpp
@@ -8,7 +8,7 @@
 #include "MuR/SystemPrivate.h"
 
 
-namespace mu
+namespace UE::Mutable::Private
 {
 	CodeRunnerBegin::CodeRunnerBegin(FSystem& InSystem, const TSharedPtr<const FModel>& InModel, const TSharedPtr<const FParameters>& InParams, uint32 InLodMask) :
 		Model(InModel),
@@ -723,7 +723,7 @@ namespace mu
 							TSharedPtr<FInstance> Instance;
 							if (InInstance)
 							{
-								Instance = mu::CloneOrTakeOver<FInstance>(InInstance);
+								Instance = UE::Mutable::Private::CloneOrTakeOver<FInstance>(InInstance);
 							}
 							else
 							{
@@ -766,7 +766,7 @@ namespace mu
 							TSharedPtr<FInstance> Instance;
 							if (InInstance)
 							{
-								Instance = mu::CloneOrTakeOver<FInstance>(InInstance);
+								Instance = UE::Mutable::Private::CloneOrTakeOver<FInstance>(InInstance);
 							}
 							else
 							{
@@ -818,7 +818,7 @@ namespace mu
 							TSharedPtr<FInstance> Instance;
 							if (InInstance)
 							{
-								Instance = mu::CloneOrTakeOver<FInstance>(InInstance);
+								Instance = UE::Mutable::Private::CloneOrTakeOver<FInstance>(InInstance);
 							}
 							else
 							{
@@ -869,7 +869,7 @@ namespace mu
 							TSharedPtr<FInstance> Instance;
 							if (InInstance)
 							{
-								Instance = mu::CloneOrTakeOver<FInstance>(InInstance);
+								Instance = UE::Mutable::Private::CloneOrTakeOver<FInstance>(InInstance);
 							}
 							else
 							{
@@ -919,7 +919,7 @@ namespace mu
 							TSharedPtr<FInstance> Instance;
 							if (InInstance)
 							{
-								Instance = mu::CloneOrTakeOver<FInstance>(InInstance);
+								Instance = UE::Mutable::Private::CloneOrTakeOver<FInstance>(InInstance);
 							}
 							else
 							{
@@ -968,7 +968,7 @@ namespace mu
 							TSharedPtr<FInstance> Instance;
 							if (InInstance)
 							{
-								Instance = mu::CloneOrTakeOver<FInstance>(InInstance);
+								Instance = UE::Mutable::Private::CloneOrTakeOver<FInstance>(InInstance);
 							}
 							else
 							{
@@ -1026,7 +1026,7 @@ namespace mu
 							TSharedPtr<FInstance> Instance;
 							if (InInstance)
 							{
-								Instance = mu::CloneOrTakeOver<FInstance>(InInstance);
+								Instance = UE::Mutable::Private::CloneOrTakeOver<FInstance>(InInstance);
 							}
 							else
 							{
@@ -1069,7 +1069,7 @@ namespace mu
 		            	TSharedPtr<FInstance> Instance;
 		            	if (InInstance)
 		            	{
-		            		Instance = mu::CloneOrTakeOver<FInstance>(InInstance);
+		            		Instance = UE::Mutable::Private::CloneOrTakeOver<FInstance>(InInstance);
 		            	}
 		            	else
 		            	{
@@ -1221,7 +1221,7 @@ namespace mu
 							TSharedPtr<FInstance> Instance;
 							if (InInstance)
 							{
-								Instance = mu::CloneOrTakeOver<FInstance>(InInstance);
+								Instance = UE::Mutable::Private::CloneOrTakeOver<FInstance>(InInstance);
 							}
 							else
 							{
diff --git a/Engine/Plugins/Mutable/Source/MutableRuntime/Private/MuR/CodeRunnerTasked.cpp b/Engine/Plugins/Mutable/Source/MutableRuntime/Private/MuR/CodeRunnerTasked.cpp
--- a/Engine/Plugins/Mutable/Source/MutableRuntime/Private/MuR/CodeRunnerTasked.cpp
+++ b/Engine/Plugins/Mutable/Source/MutableRuntime/Private/MuR/CodeRunnerTasked.cpp
@@ -55,7 +55,7 @@ static TAutoConsoleVariable<bool> CVarCodeRunnerForceInline(
 
 DECLARE_CYCLE_STAT(TEXT("MutableCoreTask"), STAT_MutableCoreTask, STATGROUP_Game);
 
-namespace mu
+namespace UE::Mutable::Private
 {
 
 	TRACE_DECLARE_INT_COUNTER(MutableRuntime_OpenTask, TEXT("MutableRuntime/OpenTask"));
@@ -2811,4 +2811,4 @@ namespace mu
 		return Issued;
 	}
 
-} // namespace mu
+} // namespace UE::Mutable::Private
diff --git a/Engine/Plugins/Mutable/Source/MutableRuntime/Private/MuR/Image.cpp b/Engine/Plugins/Mutable/Source/MutableRuntime/Private/MuR/Image.cpp
--- a/Engine/Plugins/Mutable/Source/MutableRuntime/Private/MuR/Image.cpp
+++ b/Engine/Plugins/Mutable/Source/MutableRuntime/Private/MuR/Image.cpp
@@ -10,7 +10,7 @@
 #include "MuR/MutableTrace.h"
 #include "ImageCore.h"
 
-namespace mu
+namespace UE::Mutable::Private
 {
 
 	FImage::FImage()
@@ -773,7 +773,7 @@ namespace mu
   //       GetNonBlackRect_Reference( debugRect );
   //       if (!(rect==debugRect))
   //       {
-  //           mu::Halt();
+  //           UE::Mutable::Private::Halt();
   //       }
     }
 
diff --git a/Engine/Plugins/Mutable/Source/MutableRuntime/Private/MuR/ImageDataStorage.cpp b/Engine/Plugins/Mutable/Source/MutableRuntime/Private/MuR/ImageDataStorage.cpp
--- a/Engine/Plugins/Mutable/Source/MutableRuntime/Private/MuR/ImageDataStorage.cpp
+++ b/Engine/Plugins/Mutable/Source/MutableRuntime/Private/MuR/ImageDataStorage.cpp
@@ -3,7 +3,7 @@
 #include "MuR/ImageDataStorage.h"
 #include "MuR/SerialisationPrivate.h"
 
-namespace mu::FImageDataStorageInternal
+namespace UE::Mutable::Private::FImageDataStorageInternal
 {
 	FORCEINLINE FImageSize ComputeLODSize(FImageSize BaseSize, int32 LOD)
 	{
@@ -52,7 +52,7 @@ namespace mu::FImageDataStorageInternal
 	}
 }
 
-namespace mu::MemoryCounters
+namespace UE::Mutable::Private::MemoryCounters
 {
 	std::atomic<SSIZE_T>& FImageMemoryCounter::Get()
 	{
@@ -61,7 +61,7 @@ namespace mu::MemoryCounters
 	}
 }
 
-namespace mu
+namespace UE::Mutable::Private
 {
 	FImageDataStorage::FImageDataStorage()
 	{
diff --git a/Engine/Plugins/Mutable/Source/MutableRuntime/Private/MuR/Layout.cpp b/Engine/Plugins/Mutable/Source/MutableRuntime/Private/MuR/Layout.cpp
--- a/Engine/Plugins/Mutable/Source/MutableRuntime/Private/MuR/Layout.cpp
+++ b/Engine/Plugins/Mutable/Source/MutableRuntime/Private/MuR/Layout.cpp
@@ -10,7 +10,7 @@
 #include "MuR/SerialisationPrivate.h"
 
 
-namespace mu {
+namespace UE::Mutable::Private {
 
 	MUTABLE_IMPLEMENT_POD_SERIALISABLE(FLayoutBlock);
 	MUTABLE_IMPLEMENT_POD_VECTOR_SERIALISABLE(FLayoutBlock);
diff --git a/Engine/Plugins/Mutable/Source/MutableRuntime/Private/MuR/MemoryTrackingUtils.cpp b/Engine/Plugins/Mutable/Source/MutableRuntime/Private/MuR/MemoryTrackingUtils.cpp
--- a/Engine/Plugins/Mutable/Source/MutableRuntime/Private/MuR/MemoryTrackingUtils.cpp
+++ b/Engine/Plugins/Mutable/Source/MutableRuntime/Private/MuR/MemoryTrackingUtils.cpp
@@ -9,7 +9,7 @@
 #include "HAL/LowLevelMemTracker.h"
 #include "HAL/IConsoleManager.h"
 
-namespace mu::Private
+namespace UE::Mutable::Private::Private
 {
 	// Call LLM system tick (we do this to simulate a frame since the LLM system is not entirelly designed to run over a program)
 	FORCEINLINE void UpdateLLMStats()
@@ -26,7 +26,7 @@ namespace mu::Private
 	}
 }
 
-namespace mu
+namespace UE::Mutable::Private
 {
 
 	struct FMemoryTrackingConsoleFlags
@@ -204,8 +204,8 @@ namespace mu
 				UE_LOG(LogMutableCore, 
 					   Display, 
 					   TEXT("Mutable Memory: Current %2.f MiB, Peak %2.f MiB"), 
-					   double(mu::FGlobalMemoryCounter::GetCounter()) / double(1024*1024),
-					   double(mu::FGlobalMemoryCounter::GetPeak())    / double(1024*1024)
+					   double(UE::Mutable::Private::FGlobalMemoryCounter::GetCounter()) / double(1024*1024),
+					   double(UE::Mutable::Private::FGlobalMemoryCounter::GetPeak())    / double(1024*1024)
 				);
 			})
 		);
@@ -230,8 +230,8 @@ namespace mu
 				UE_LOG(LogMutableCore, 
 					   Display, 
 					   TEXT("Mutable Memory: AbsoluteCurrent %2.f MiB, AbsolutePeak %2.f MiB"), 
-					   double(mu::FGlobalMemoryCounter::GetAbsoluteCounter()) / double(1024*1024),
-					   double(mu::FGlobalMemoryCounter::GetAbsolutePeak())    / double(1024*1024)
+					   double(UE::Mutable::Private::FGlobalMemoryCounter::GetAbsoluteCounter()) / double(1024*1024),
+					   double(UE::Mutable::Private::FGlobalMemoryCounter::GetAbsolutePeak())    / double(1024*1024)
 				);
 			})
 		);
diff --git a/Engine/Plugins/Mutable/Source/MutableRuntime/Private/MuR/Mesh.cpp b/Engine/Plugins/Mutable/Source/MutableRuntime/Private/MuR/Mesh.cpp
--- a/Engine/Plugins/Mutable/Source/MutableRuntime/Private/MuR/Mesh.cpp
+++ b/Engine/Plugins/Mutable/Source/MutableRuntime/Private/MuR/Mesh.cpp
@@ -13,7 +13,7 @@
 #include "Spatial/PointHashGrid3.h"
 #include "Engine/SkeletalMesh.h"
 
-namespace mu
+namespace UE::Mutable::Private
 {
 
 MUTABLE_IMPLEMENT_ENUM_SERIALISABLE(EBoneUsageFlags);
@@ -570,20 +570,20 @@ int32 FMesh::FindBonePose(const FBoneName& BoneId) const
 }
 
 
-void mu::FMesh::SetBonePoseCount(int32 count)
+void UE::Mutable::Private::FMesh::SetBonePoseCount(int32 count)
 {
 	LLM_SCOPE_BYNAME(TEXT("MutableRuntime"));
 	BonePoses.SetNum(count);
 }
 
 
-int32 mu::FMesh::GetBonePoseCount() const
+int32 UE::Mutable::Private::FMesh::GetBonePoseCount() const
 {
 	return BonePoses.Num();
 }
 
 
-void mu::FMesh::SetBonePose(int32 Index, const FBoneName& BoneId, FTransform3f Transform, EBoneUsageFlags BoneUsageFlags)
+void UE::Mutable::Private::FMesh::SetBonePose(int32 Index, const FBoneName& BoneId, FTransform3f Transform, EBoneUsageFlags BoneUsageFlags)
 {
 	check(BonePoses.IsValidIndex(Index));
 	if (BonePoses.IsValidIndex(Index))
@@ -600,7 +600,7 @@ const FBoneName& FMesh::GetBonePoseId(int32 Index) const
 }
 
 
-void mu::FMesh::GetBonePoseTransform(int32 BoneIndex, FTransform3f& Transform) const
+void UE::Mutable::Private::FMesh::GetBonePoseTransform(int32 BoneIndex, FTransform3f& Transform) const
 {
 	check(BoneIndex >= 0 && BoneIndex < BonePoses.Num());
 	Transform = BoneIndex > INDEX_NONE ? BonePoses[BoneIndex].BoneTransform : FTransform3f::Identity;
@@ -1583,7 +1583,7 @@ void GetUVIsland(TArray<FTriangleInfo>& InTriangles,
 }
 
 
-void MeshCreateCollapsedVertexMap(const mu::FMesh* Mesh, TArray<int32>& CollapsedVertices)
+void MeshCreateCollapsedVertexMap(const UE::Mutable::Private::FMesh* Mesh, TArray<int32>& CollapsedVertices)
 {
 	MUTABLE_CPUPROFILER_SCOPE(LayoutUV_CreateCollapsedVertexMap);
 
@@ -1596,7 +1596,7 @@ void MeshCreateCollapsedVertexMap(const mu::FMesh* Mesh, TArray<int32>& Collapse
 	TArray<FVector3f> Vertices;
 	Vertices.SetNumUninitialized(NumVertices);
 
-	mu::UntypedMeshBufferIteratorConst ItPosition = mu::UntypedMeshBufferIteratorConst(Mesh->GetVertexBuffers(), mu::EMeshBufferSemantic::Position);
+	UE::Mutable::Private::UntypedMeshBufferIteratorConst ItPosition = UE::Mutable::Private::UntypedMeshBufferIteratorConst(Mesh->GetVertexBuffers(), UE::Mutable::Private::EMeshBufferSemantic::Position);
 
 	FVector3f* VertexData = Vertices.GetData();
 
diff --git a/Engine/Plugins/Mutable/Source/MutableRuntime/Private/MuR/MeshBufferSet.cpp b/Engine/Plugins/Mutable/Source/MutableRuntime/Private/MuR/MeshBufferSet.cpp
--- a/Engine/Plugins/Mutable/Source/MutableRuntime/Private/MuR/MeshBufferSet.cpp
+++ b/Engine/Plugins/Mutable/Source/MutableRuntime/Private/MuR/MeshBufferSet.cpp
@@ -11,7 +11,7 @@
 #include "MuR/SerialisationPrivate.h"
 
 
-namespace mu::MemoryCounters
+namespace UE::Mutable::Private::MemoryCounters
 {
 	std::atomic<SSIZE_T>& FMeshMemoryCounter::Get()
 	{
@@ -20,7 +20,7 @@ namespace mu::MemoryCounters
 	}
 }
 
-namespace mu
+namespace UE::Mutable::Private
 {
 	MUTABLE_IMPLEMENT_POD_SERIALISABLE(FMeshBufferChannel);
 	MUTABLE_IMPLEMENT_POD_VECTOR_SERIALISABLE(FMeshBufferChannel);
diff --git a/Engine/Plugins/Mutable/Source/MutableRuntime/Private/MuR/Model.cpp b/Engine/Plugins/Mutable/Source/MutableRuntime/Private/MuR/Model.cpp
--- a/Engine/Plugins/Mutable/Source/MutableRuntime/Private/MuR/Model.cpp
+++ b/Engine/Plugins/Mutable/Source/MutableRuntime/Private/MuR/Model.cpp
@@ -24,7 +24,7 @@ class USkeletalMesh;
 class UMaterialInterface;
 
 
-namespace mu
+namespace UE::Mutable::Private
 {
 	MUTABLE_IMPLEMENT_POD_SERIALISABLE(FRomDataRuntime);
 	MUTABLE_IMPLEMENT_POD_SERIALISABLE(FRomDataCompile);
@@ -146,7 +146,7 @@ namespace mu
     {
 		LLM_SCOPE_BYNAME(TEXT("MutableRuntime"));
 
-		mu::FProgram& Program = p->m_pD->Program;
+		UE::Mutable::Private::FProgram& Program = p->m_pD->Program;
 
 		FOutputMemoryStream MemStream(16 * 1024 * 1024);
 
diff --git a/Engine/Plugins/Mutable/Source/MutableRuntime/Private/MuR/OpImageMipmap.cpp b/Engine/Plugins/Mutable/Source/MutableRuntime/Private/MuR/OpImageMipmap.cpp
--- a/Engine/Plugins/Mutable/Source/MutableRuntime/Private/MuR/OpImageMipmap.cpp
+++ b/Engine/Plugins/Mutable/Source/MutableRuntime/Private/MuR/OpImageMipmap.cpp
@@ -16,7 +16,7 @@ static FAutoConsoleVariableRef CVarEnableCompressedMipGenerationMemoryOptimizati
 	ECVF_Default);
 }
 
-namespace mu
+namespace UE::Mutable::Private
 {
 
 
@@ -537,10 +537,10 @@ namespace OpImageMipmapInternal
         check(DestImage->GetLODCount() == LevelCount);
 
 		EImageFormat DestFormat = DestImage->GetFormat();
-		if (mu::IsCompressedFormat(DestFormat))
+		if (UE::Mutable::Private::IsCompressedFormat(DestFormat))
 		{
 			// Is it a block format?
-			if (mu::GetImageFormatData(DestFormat).PixelsPerBlockX > 1)
+			if (UE::Mutable::Private::GetImageFormatData(DestFormat).PixelsPerBlockX > 1)
 			{
 				if (!bEnableCompressedMipGenerationMemoryOptimizations)
 				{
@@ -652,8 +652,8 @@ namespace OpImageMipmapInternal
 		}
 
 		EImageFormat BaseFormat = BaseImage->GetFormat();
-		const bool bIsBlockCompressedFormat = mu::IsBlockCompressedFormat(BaseFormat);
-		const bool bIsCompressedFormat = mu::IsCompressedFormat(BaseFormat);
+		const bool bIsBlockCompressedFormat = UE::Mutable::Private::IsBlockCompressedFormat(BaseFormat);
+		const bool bIsCompressedFormat = UE::Mutable::Private::IsCompressedFormat(BaseFormat);
 		
 		check(!bIsBlockCompressedFormat || bIsCompressedFormat);
 
diff --git a/Engine/Plugins/Mutable/Source/MutableRuntime/Private/MuR/OpImageProject.cpp b/Engine/Plugins/Mutable/Source/MutableRuntime/Private/MuR/OpImageProject.cpp
--- a/Engine/Plugins/Mutable/Source/MutableRuntime/Private/MuR/OpImageProject.cpp
+++ b/Engine/Plugins/Mutable/Source/MutableRuntime/Private/MuR/OpImageProject.cpp
@@ -29,7 +29,7 @@
 #include "MuR/Mesh.h"
 
 
-namespace mu
+namespace UE::Mutable::Private
 {
 
 //-------------------------------------------------------------------------------------------------
@@ -938,7 +938,7 @@ void ImageRasterProjectedPlanar(const FMesh* pMesh, FImage* pImage,
 	UE::Math::TIntVector2<uint16> CropMin, UE::Math::TIntVector2<uint16> UncroppedSize,
 	FScratchImageProject* Scratch, bool bUseVectorImplementation)
 {
-	using namespace mu::Private;
+	using namespace UE::Mutable::Private::Private;
 
     check(!pMask || pMask->GetSizeX() == pImage->GetSizeX());
     check(!pMask || pMask->GetSizeY() == pImage->GetSizeY());
@@ -1153,7 +1153,7 @@ void ImageRasterProjectedWrapping( const FMesh* pMesh, FImage* pImage,
 	FScratchImageProject* Scratch, bool bUseVectorImplementation)
 {
 
-	using namespace mu::Private;
+	using namespace UE::Mutable::Private::Private;
     
 	check(!pMask || pMask->GetSizeX() == pImage->GetSizeX());
     check(!pMask || pMask->GetSizeY() == pImage->GetSizeY());
@@ -1367,7 +1367,7 @@ void ImageRasterProjectedCylindrical( const FMesh* pMesh, FImage* pImage,
 	FScratchImageProject* Scratch, bool bUseVectorImplementation)
 {
 
-	using namespace mu::Private;
+	using namespace UE::Mutable::Private::Private;
 
     check(!pMask || pMask->GetSizeX() == pImage->GetSizeX());
     check(!pMask || pMask->GetSizeY() == pImage->GetSizeY());
@@ -1723,7 +1723,7 @@ inline void MeshCreateCollapsedVertexMap( const FMesh* pMesh,
 	CollapsedVerticesMap.Reserve(NumVertices);
 
 	// Get Vertices
-	mu::UntypedMeshBufferIteratorConst ItPosition = mu::UntypedMeshBufferIteratorConst(pMesh->GetVertexBuffers(), mu::EMeshBufferSemantic::Position);
+	UE::Mutable::Private::UntypedMeshBufferIteratorConst ItPosition = UE::Mutable::Private::UntypedMeshBufferIteratorConst(pMesh->GetVertexBuffers(), UE::Mutable::Private::EMeshBufferSemantic::Position);
 	FVector3f* VertexData = OutVertices.GetData();
 	
 	for (int32 VertexIndex = 0; VertexIndex < NumVertices; ++VertexIndex)
diff --git a/Engine/Plugins/Mutable/Source/MutableRuntime/Private/MuR/OpImageResize.cpp b/Engine/Plugins/Mutable/Source/MutableRuntime/Private/MuR/OpImageResize.cpp
--- a/Engine/Plugins/Mutable/Source/MutableRuntime/Private/MuR/OpImageResize.cpp
+++ b/Engine/Plugins/Mutable/Source/MutableRuntime/Private/MuR/OpImageResize.cpp
@@ -27,7 +27,7 @@ static FAutoConsoleVariableRef CVarEnableVectorImplementationForSmallResizes (
 
 }
 
-namespace mu
+namespace UE::Mutable::Private
 {
 	/** */
 	template< int32 NC >
@@ -699,7 +699,7 @@ namespace mu
 		default:
 			// Case not implemented
 			check(false);
-			//mu::Halt();
+			//UE::Mutable::Private::Halt();
 		}
 
 	}
@@ -1151,8 +1151,8 @@ namespace
 			BaseSize.X >= DestSize.X && BaseSize.Y >= DestSize.Y && 
 			BaseSize.X < DestSize.X*2 && BaseSize.Y < DestSize.Y*2; 
 
-		const bool bIsUncompressedFormat = !mu::IsCompressedFormat(InBase->GetFormat());  
-		const bool bIsBlockCompressedFormat = mu::IsBlockCompressedFormat(InBase->GetFormat());
+		const bool bIsUncompressedFormat = !UE::Mutable::Private::IsCompressedFormat(InBase->GetFormat());  
+		const bool bIsBlockCompressedFormat = UE::Mutable::Private::IsBlockCompressedFormat(InBase->GetFormat());
 
 		if (bIsDownSizeByLessThanTwo && (bIsUncompressedFormat || bIsBlockCompressedFormat) && bEnableImageResizeMemoryOptimizations)
 		{
diff --git a/Engine/Plugins/Mutable/Source/MutableRuntime/Private/MuR/OpImageSwizzle.cpp b/Engine/Plugins/Mutable/Source/MutableRuntime/Private/MuR/OpImageSwizzle.cpp
--- a/Engine/Plugins/Mutable/Source/MutableRuntime/Private/MuR/OpImageSwizzle.cpp
+++ b/Engine/Plugins/Mutable/Source/MutableRuntime/Private/MuR/OpImageSwizzle.cpp
@@ -4,7 +4,7 @@
 #include "MuR/Operations.h"
 #include "MuR/ParallelExecutionUtils.h"
 
-namespace mu
+namespace UE::Mutable::Private
 {
 	void ImageSwizzle(FImage* Result, const TSharedPtr<const FImage> Sources[], const uint8 Channels[])
 	{
@@ -218,7 +218,7 @@ namespace mu
 
 		TSharedPtr<FImage> Dest = CreateImage(ResultSize.X, ResultSize.Y, ResultNumLODs, Format, EInitializationType::Black);
 
-		mu::ImageSwizzle(Dest.Get(), Sources, Channels);
+		UE::Mutable::Private::ImageSwizzle(Dest.Get(), Sources, Channels);
 
 		return Dest;
 	}
diff --git a/Engine/Plugins/Mutable/Source/MutableRuntime/Private/MuR/OpMeshClipWithMesh.cpp b/Engine/Plugins/Mutable/Source/MutableRuntime/Private/MuR/OpMeshClipWithMesh.cpp
--- a/Engine/Plugins/Mutable/Source/MutableRuntime/Private/MuR/OpMeshClipWithMesh.cpp
+++ b/Engine/Plugins/Mutable/Source/MutableRuntime/Private/MuR/OpMeshClipWithMesh.cpp
@@ -19,7 +19,7 @@
 #include "Spatial/PointHashGrid3.h"
 
 
-namespace mu { namespace 
+namespace UE::Mutable::Private { namespace 
 {
 
     //---------------------------------------------------------------------------------------------
@@ -41,7 +41,7 @@ namespace mu { namespace
 		OutCollapsedVertices.Init(INDEX_NONE, NumVertices);
 
 		// Get Vertices
-		mu::UntypedMeshBufferIteratorConst ItPosition = mu::UntypedMeshBufferIteratorConst(pMesh->GetVertexBuffers(), mu::EMeshBufferSemantic::Position);
+		UE::Mutable::Private::UntypedMeshBufferIteratorConst ItPosition = UE::Mutable::Private::UntypedMeshBufferIteratorConst(pMesh->GetVertexBuffers(), UE::Mutable::Private::EMeshBufferSemantic::Position);
 		FVector3f* VertexData = OutVertices.GetData();
 
 		for (int32 VertexIndex = 0; VertexIndex < NumVertices; ++VertexIndex)
@@ -324,7 +324,7 @@ namespace mu { namespace
 		const FMeshBufferSet& MBSPriv = Base->GetVertexBuffers();
 		for (int32 b = 0; b < MBSPriv.Buffers.Num(); ++b)
 		{
-			const TArray<mu::FMeshBufferChannel>& Channels = MBSPriv.Buffers[b].Channels;
+			const TArray<UE::Mutable::Private::FMeshBufferChannel>& Channels = MBSPriv.Buffers[b].Channels;
 
 			for (int32 c = 0; c < Channels.Num(); ++c)
 			{
@@ -375,7 +375,7 @@ namespace mu { namespace
 		const FMeshBufferSet& MBSPriv = Base->GetVertexBuffers();
 		for (int32 b = 0; b < MBSPriv.Buffers.Num(); ++b)
 		{
-			const TArray<mu::FMeshBufferChannel>& Channels = MBSPriv.Buffers[b].Channels;
+			const TArray<UE::Mutable::Private::FMeshBufferChannel>& Channels = MBSPriv.Buffers[b].Channels;
 
 			for (int32 c = 0; c < Channels.Num(); ++c)
 			{
@@ -483,10 +483,10 @@ namespace mu { namespace
 		}
 	}
 
-}} // namespace mu::anonymous
+}} // namespace UE::Mutable::Private::anonymous
 
 
-namespace mu
+namespace UE::Mutable::Private
 {
 
 	bool IsMeshClosed(const FMesh* Mesh)
diff --git a/Engine/Plugins/Mutable/Source/MutableRuntime/Private/MuR/OpMeshMorph.cpp b/Engine/Plugins/Mutable/Source/MutableRuntime/Private/MuR/OpMeshMorph.cpp
--- a/Engine/Plugins/Mutable/Source/MutableRuntime/Private/MuR/OpMeshMorph.cpp
+++ b/Engine/Plugins/Mutable/Source/MutableRuntime/Private/MuR/OpMeshMorph.cpp
@@ -12,7 +12,7 @@
 #include "Animation/MorphTarget.h"
 #include "Rendering/MorphTargetVertexCodec.h"
 
-void mu::MeshMorph(FMesh* Mesh, const FName& MorphName, const float Factor)
+void UE::Mutable::Private::MeshMorph(FMesh* Mesh, const FName& MorphName, const float Factor)
 {
 	MUTABLE_CPUPROFILER_SCOPE(MeshMorph_Parameter);
 
@@ -212,7 +212,7 @@ void mu::MeshMorph(FMesh* Mesh, const FName& MorphName, const float Factor)
 	}
 }
 
-void mu::MeshMorph(FMesh* BaseMesh, const FMesh* MaxMesh, const float Factor)
+void UE::Mutable::Private::MeshMorph(FMesh* BaseMesh, const FMesh* MaxMesh, const float Factor)
 {
     MUTABLE_CPUPROFILER_SCOPE(MeshMorph_Compiled);
 
diff --git a/Engine/Plugins/Mutable/Source/MutableRuntime/Private/MuR/OpMeshMorph.h b/Engine/Plugins/Mutable/Source/MutableRuntime/Private/MuR/OpMeshMorph.h
--- a/Engine/Plugins/Mutable/Source/MutableRuntime/Private/MuR/OpMeshMorph.h
+++ b/Engine/Plugins/Mutable/Source/MutableRuntime/Private/MuR/OpMeshMorph.h
@@ -4,13 +4,13 @@
 
 class FName;
 
-namespace mu
+namespace UE::Mutable::Private
 {
 	class FMesh;
 }
 
 
-namespace mu
+namespace UE::Mutable::Private
 {
 	/** Morph Parameter Meshes. Mesh contains the deltas in morph buffer. Deltas are compressed using the engine format. */
 	void MeshMorph(FMesh* Mesh, const FName& MorphName, float Factor);
diff --git a/Engine/Plugins/Mutable/Source/MutableRuntime/Private/MuR/OpMeshOptimizeSkinning.h b/Engine/Plugins/Mutable/Source/MutableRuntime/Private/MuR/OpMeshOptimizeSkinning.h
--- a/Engine/Plugins/Mutable/Source/MutableRuntime/Private/MuR/OpMeshOptimizeSkinning.h
+++ b/Engine/Plugins/Mutable/Source/MutableRuntime/Private/MuR/OpMeshOptimizeSkinning.h
@@ -8,7 +8,7 @@
 #include "MuR/Platform.h"
 
 
-namespace mu
+namespace UE::Mutable::Private
 {
 
 	//---------------------------------------------------------------------------------------------
@@ -42,7 +42,7 @@ namespace mu
 		bool bRequiresFormatChange = false;
 
 		// Desired format if MaxBoneMapIndex <= MAX_uint8
-		const EMeshBufferFormat DesiredBoneIndexFormat = mu::EMeshBufferFormat::UInt8;
+		const EMeshBufferFormat DesiredBoneIndexFormat = UE::Mutable::Private::EMeshBufferFormat::UInt8;
 		
 		// Iterate all vertex buffers and check if the BoneIndices buffers have the desired format
 		const FMeshBufferSet& InMeshVertexBuffers = InMesh->GetVertexBuffers();
@@ -81,7 +81,7 @@ namespace mu
 		constexpr EMeshCopyFlags CopyFlags = ~EMeshCopyFlags::WithVertexBuffers;
 		Result->CopyFrom(*InMesh, CopyFlags);
 
-		mu::FMeshBufferSet& VertexBuffers = Result->GetVertexBuffers();
+		UE::Mutable::Private::FMeshBufferSet& VertexBuffers = Result->GetVertexBuffers();
 
 
 		VertexBuffers.Flags = InMeshVertexBuffers.Flags;
@@ -102,7 +102,7 @@ namespace mu
 			TArray<int32> Components;
 			TArray<int32> Offsets;
 
-			EMeshBufferFormat SourceBoneIndexFormat = mu::EMeshBufferFormat::None;
+			EMeshBufferFormat SourceBoneIndexFormat = UE::Mutable::Private::EMeshBufferFormat::None;
 
 			// Offset accumulator
 			int32 AuxOffset = 0;
@@ -116,7 +116,7 @@ namespace mu
 				if (Channel.Semantic == EMeshBufferSemantic::BoneIndices)
 				{
 					SourceBoneIndexFormat = Channel.Format;
-					Format = mu::EMeshBufferFormat::UInt8;
+					Format = UE::Mutable::Private::EMeshBufferFormat::UInt8;
 					BoneIndexChannelIndex = ChannelIndex;
 				}
 
@@ -133,7 +133,7 @@ namespace mu
 			}
 
 			// Copy buffers
-			if (BoneIndexChannelIndex != INDEX_NONE && SourceBoneIndexFormat == mu::EMeshBufferFormat::UInt16)
+			if (BoneIndexChannelIndex != INDEX_NONE && SourceBoneIndexFormat == UE::Mutable::Private::EMeshBufferFormat::UInt16)
 			{
 				// This buffer has a BoneIndex channel that needs to be manually fixed. 
 				VertexBuffers.SetBufferCount(BufferIndex + 1);
@@ -146,7 +146,7 @@ namespace mu
 					const uint8* SourceData = (const uint8*)SourceBuffer.Data.GetData();
 					const int32 SourceBoneIndexSize = GetMeshFormatData(SourceBoneIndexFormat).SizeInBytes;
 
-					const int32 BoneIndexSize = GetMeshFormatData(mu::EMeshBufferFormat::UInt8).SizeInBytes;
+					const int32 BoneIndexSize = GetMeshFormatData(UE::Mutable::Private::EMeshBufferFormat::UInt8).SizeInBytes;
 					const int32 BoneIndexComponentCount = Components[BoneIndexChannelIndex];
 					const int32 BoneIndexChannelOffset = Offsets[BoneIndexChannelIndex];
 
@@ -192,8 +192,8 @@ namespace mu
 			}
 			else
 			{
-				// SourceBoneIndexFormat must be mu::EMeshBufferFormat::UInt8 or none if the buffer doesn't have BoneIndices
-				check(SourceBoneIndexFormat == mu::EMeshBufferFormat::None || SourceBoneIndexFormat == mu::EMeshBufferFormat::UInt8);
+				// SourceBoneIndexFormat must be UE::Mutable::Private::EMeshBufferFormat::UInt8 or none if the buffer doesn't have BoneIndices
+				check(SourceBoneIndexFormat == UE::Mutable::Private::EMeshBufferFormat::None || SourceBoneIndexFormat == UE::Mutable::Private::EMeshBufferFormat::UInt8);
 
 				// Add buffers that don't require a fix up
 				VertexBuffers.AddBuffer(InMeshVertexBuffers, BufferIndex);
diff --git a/Engine/Plugins/Mutable/Source/MutableRuntime/Private/MuR/OpMeshPrepareLayout.cpp b/Engine/Plugins/Mutable/Source/MutableRuntime/Private/MuR/OpMeshPrepareLayout.cpp
--- a/Engine/Plugins/Mutable/Source/MutableRuntime/Private/MuR/OpMeshPrepareLayout.cpp
+++ b/Engine/Plugins/Mutable/Source/MutableRuntime/Private/MuR/OpMeshPrepareLayout.cpp
@@ -11,7 +11,7 @@
 #include "MuR/MutableRuntimeModule.h"
 
 
-namespace mu
+namespace UE::Mutable::Private
 {
 
 	template<typename T>
@@ -164,7 +164,7 @@ namespace mu
 			}
 		}
 
-		const bool bIsOverlayLayout = Layout->GetLayoutPackingStrategy() == mu::EPackStrategy::Overlay;
+		const bool bIsOverlayLayout = Layout->GetLayoutPackingStrategy() == UE::Mutable::Private::EPackStrategy::Overlay;
 		if (bNormalizeUVs && !bIsOverlayLayout)
 		{
 			MUTABLE_CPUPROFILER_SCOPE(NormalizeUVs);
@@ -231,7 +231,7 @@ namespace mu
 					// First: Assign the vertices to masked blocks in order
 					for (int32 CandidateBlockIndex = 0; CandidateBlockIndex < NumBlocks; ++CandidateBlockIndex)
 					{
-						const mu::FImage* Mask = nullptr;
+						const UE::Mutable::Private::FImage* Mask = nullptr;
 						int32 MaskIndex = Layout->Blocks[CandidateBlockIndex].MaskIndex;
 						if (Layout->Masks.IsValidIndex(MaskIndex))
 						{
diff --git a/Engine/Plugins/Mutable/Source/MutableRuntime/Private/MuR/OpMeshTransformWithBone.cpp b/Engine/Plugins/Mutable/Source/MutableRuntime/Private/MuR/OpMeshTransformWithBone.cpp
--- a/Engine/Plugins/Mutable/Source/MutableRuntime/Private/MuR/OpMeshTransformWithBone.cpp
+++ b/Engine/Plugins/Mutable/Source/MutableRuntime/Private/MuR/OpMeshTransformWithBone.cpp
@@ -10,7 +10,7 @@
 #include "MuR/ParallelExecutionUtils.h"
 #include "PackedNormal.h"
 
-namespace mu
+namespace UE::Mutable::Private
 {
 
 
@@ -82,7 +82,7 @@ void MeshTransformWithBoneInline(FMesh* Mesh, const FMatrix44f& Transform, const
 	const uint16 MinWeightThreshold = (BoneWeightsIterBegin.GetFormat() == EMeshBufferFormat::NUInt8 ? MAX_uint8 : MAX_uint16) * InThresholdFactor;
 
 	// Classify which vertices in the SourceMesh are completely bounded by the BoundingMesh geometry.
-	// If no BoundingMesh is provided, this defaults to act exactly like mu::MeshTransform
+	// If no BoundingMesh is provided, this defaults to act exactly like UE::Mutable::Private::MeshTransform
 	TBitArray<> VertexInBoneHierarchy;
 	VertexInBoneHierarchy.Init(false, NumVertices);
 	{
diff --git a/Engine/Plugins/Mutable/Source/MutableRuntime/Private/MuR/OpMeshTransformWithMesh.cpp b/Engine/Plugins/Mutable/Source/MutableRuntime/Private/MuR/OpMeshTransformWithMesh.cpp
--- a/Engine/Plugins/Mutable/Source/MutableRuntime/Private/MuR/OpMeshTransformWithMesh.cpp
+++ b/Engine/Plugins/Mutable/Source/MutableRuntime/Private/MuR/OpMeshTransformWithMesh.cpp
@@ -10,7 +10,7 @@
 #include "Async/ParallelFor.h"
 #include "PackedNormal.h"
 
-namespace mu
+namespace UE::Mutable::Private
 {
 void MeshTransformWithMesh(FMesh* Result, const FMesh* SourceMesh, const FMesh* BoundingMesh, const FMatrix44f& Transform, bool& bOutSuccess)
 {
@@ -27,7 +27,7 @@ void MeshTransformWithMesh(FMesh* Result, const FMesh* SourceMesh, const FMesh*
 	Result->CopyFrom(*SourceMesh);
 
 	// Classify which vertices in the SourceMesh are completely bounded by the BoundingMesh geometry.
-	// If no BoundingMesh is provided, this defaults to act exactly like mu::MeshTransform
+	// If no BoundingMesh is provided, this defaults to act exactly like UE::Mutable::Private::MeshTransform
 	TBitArray<> VertexInBoundaryMesh;  
 	if (BoundingMesh)
 	{
diff --git a/Engine/Plugins/Mutable/Source/MutableRuntime/Private/MuR/Operations.cpp b/Engine/Plugins/Mutable/Source/MutableRuntime/Private/MuR/Operations.cpp
--- a/Engine/Plugins/Mutable/Source/MutableRuntime/Private/MuR/Operations.cpp
+++ b/Engine/Plugins/Mutable/Source/MutableRuntime/Private/MuR/Operations.cpp
@@ -7,7 +7,7 @@
 #include "MuR/ModelPrivate.h"
 
 
-namespace mu
+namespace UE::Mutable::Private
 {
     MUTABLE_IMPLEMENT_POD_SERIALISABLE(OP);
     MUTABLE_IMPLEMENT_POD_VECTOR_SERIALISABLE(OP);
@@ -670,7 +670,7 @@ namespace mu
         case EOpType::ME_REMOVEMASK:
         {
             const uint8* data = program.GetOpArgsPointer(at);
-            mu::OP::ADDRESS source;
+            UE::Mutable::Private::OP::ADDRESS source;
             FMemory::Memcpy( &source, data, sizeof(OP::ADDRESS) ); data+=sizeof(OP::ADDRESS);
             f(source);
 
@@ -682,11 +682,11 @@ namespace mu
 			FMemory::Memcpy( &removes, data, sizeof(uint16) ); data+=sizeof(uint16);
             for (uint16 r=0; r<removes; ++r)
             {
-                mu::OP::ADDRESS condition;
+                UE::Mutable::Private::OP::ADDRESS condition;
 				FMemory::Memcpy( &condition, data, sizeof(OP::ADDRESS) ); data+=sizeof(OP::ADDRESS);
                 f(condition);
 
-                mu::OP::ADDRESS mask;
+                UE::Mutable::Private::OP::ADDRESS mask;
 				FMemory::Memcpy( &mask, data, sizeof(OP::ADDRESS) ); data+=sizeof(OP::ADDRESS);
                 f(mask);
             }
@@ -718,7 +718,7 @@ namespace mu
         case EOpType::ME_EXTRACTLAYOUTBLOCK:
         {
             const uint8_t* data = program.GetOpArgsPointer(at);
-            mu::OP::ADDRESS source;
+            UE::Mutable::Private::OP::ADDRESS source;
 			FMemory::Memcpy( &source, data, sizeof(OP::ADDRESS) );
             f(source);
             break;
diff --git a/Engine/Plugins/Mutable/Source/MutableRuntime/Private/MuR/ParallelExecutionUtils.cpp b/Engine/Plugins/Mutable/Source/MutableRuntime/Private/MuR/ParallelExecutionUtils.cpp
--- a/Engine/Plugins/Mutable/Source/MutableRuntime/Private/MuR/ParallelExecutionUtils.cpp
+++ b/Engine/Plugins/Mutable/Source/MutableRuntime/Private/MuR/ParallelExecutionUtils.cpp
@@ -3,7 +3,7 @@
 #include <MuR/ParallelExecutionUtils.h>
 #include <Async/ParallelFor.h>
 
-namespace mu::ParallelExecutionUtils
+namespace UE::Mutable::Private::ParallelExecutionUtils
 {
 	void InvokeBatchParallelFor(int32 Num, TFunctionRef<void(int32)> Body)
 	{
diff --git a/Engine/Plugins/Mutable/Source/MutableRuntime/Private/MuR/ResourceID.cpp b/Engine/Plugins/Mutable/Source/MutableRuntime/Private/MuR/ResourceID.cpp
--- a/Engine/Plugins/Mutable/Source/MutableRuntime/Private/MuR/ResourceID.cpp
+++ b/Engine/Plugins/Mutable/Source/MutableRuntime/Private/MuR/ResourceID.cpp
@@ -5,37 +5,37 @@
 #include "Hash/xxhash.h"
 
 
-uint32 mu::GetTypeHash(const FGeneratedMeshKey& Key)
+uint32 UE::Mutable::Private::GetTypeHash(const FGeneratedMeshKey& Key)
 {
 	return HashCombineFast(Key.Address, static_cast<uint32>(FXxHash64::HashBuffer(Key.ParameterValuesBlob.GetData(), Key.ParameterValuesBlob.Num()).Hash));
 }
 
 
-uint32 mu::GetTypeHashPersistent(const FGeneratedMeshKey& Key)
+uint32 UE::Mutable::Private::GetTypeHashPersistent(const FGeneratedMeshKey& Key)
 {
 	return HashCombine(Key.Address, static_cast<uint32>(FXxHash64::HashBuffer(Key.ParameterValuesBlob.GetData(), Key.ParameterValuesBlob.Num()).Hash));
 }
 
 
-uint32 mu::GetTypeHash(const FGeneratedImageKey& Key)
+uint32 UE::Mutable::Private::GetTypeHash(const FGeneratedImageKey& Key)
 {
 	return HashCombineFast(Key.Address, static_cast<uint32>(FXxHash64::HashBuffer(Key.ParameterValuesBlob.GetData(), Key.ParameterValuesBlob.Num()).Hash));
 }
 
 
-uint32 mu::GetTypeHashPersistent(const FGeneratedImageKey& Key)
+uint32 UE::Mutable::Private::GetTypeHashPersistent(const FGeneratedImageKey& Key)
 {
 	return HashCombine(Key.Address, static_cast<uint32>(FXxHash64::HashBuffer(Key.ParameterValuesBlob.GetData(), Key.ParameterValuesBlob.Num()).Hash));
 }
 
 
-uint32 mu::GetTypeHash(const FGeneratedMaterialKey& Key)
+uint32 UE::Mutable::Private::GetTypeHash(const FGeneratedMaterialKey& Key)
 {
 	return HashCombineFast(Key.Address, static_cast<uint32>(FXxHash64::HashBuffer(Key.ParameterValuesBlob.GetData(), Key.ParameterValuesBlob.Num()).Hash));
 }
 
 
-uint32 mu::GetTypeHashPersistent(const FGeneratedMaterialKey& Key)
+uint32 UE::Mutable::Private::GetTypeHashPersistent(const FGeneratedMaterialKey& Key)
 {
 	return HashCombine(Key.Address, static_cast<uint32>(FXxHash64::HashBuffer(Key.ParameterValuesBlob.GetData(), Key.ParameterValuesBlob.Num()).Hash));
 }
\ No newline at end of file
diff --git a/Engine/Plugins/Mutable/Source/MutableRuntime/Private/MuR/System.cpp b/Engine/Plugins/Mutable/Source/MutableRuntime/Private/MuR/System.cpp
--- a/Engine/Plugins/Mutable/Source/MutableRuntime/Private/MuR/System.cpp
+++ b/Engine/Plugins/Mutable/Source/MutableRuntime/Private/MuR/System.cpp
@@ -55,7 +55,7 @@ static FAutoConsoleVariableRef CVarEnableDetailedMemoryBudgetExceededLogging (
 	ECVF_Default);
 }
 
-namespace mu::MemoryCounters
+namespace UE::Mutable::Private::MemoryCounters
 {
 	std::atomic<SSIZE_T>& FInternalMemoryCounter::Get()
 	{
@@ -64,7 +64,7 @@ namespace mu::MemoryCounters
 	}
 }
 
-namespace mu
+namespace UE::Mutable::Private
 {
 	MUTABLE_IMPLEMENT_ENUM_SERIALISABLE(ETextureCompressionStrategy);
 
@@ -368,7 +368,7 @@ namespace mu
 
     TSharedRef<FImage> BlackImage()
     {
-    	return MakeShared<mu::FImage>(16, 16, 1, EImageFormat::RGBA_UByte, EInitializationType::Black);
+    	return MakeShared<UE::Mutable::Private::FImage>(16, 16, 1, EImageFormat::RGBA_UByte, EInitializationType::Black);
     }
 	
 
@@ -422,7 +422,7 @@ namespace mu
 		
 		m_pD->WorkingMemoryManager.BeginRunnerThread();
 		
-		mu::EOpType OpType = LiveInstance->Model->GetPrivate()->Program.GetOpType(RootAddress);
+		UE::Mutable::Private::EOpType OpType = LiveInstance->Model->GetPrivate()->Program.GetOpType(RootAddress);
 		if (GetOpDataType(OpType) != EDataType::Image)
 		{
 			m_pD->WorkingMemoryManager.EndRunnerThread();
@@ -459,7 +459,7 @@ namespace mu
 
 					if (!Result)
 					{
-						Result = MakeShared<mu::FImage>(16, 16, 1, EImageFormat::RGBA_UByte, EInitializationType::Black);
+						Result = MakeShared<UE::Mutable::Private::FImage>(16, 16, 1, EImageFormat::RGBA_UByte, EInitializationType::Black);
 					}
 
 					SystemPrivate->WorkingMemoryManager.EndRunnerThread();
@@ -500,8 +500,8 @@ namespace mu
 
 		OP::ADDRESS RootAddress = ImageId.GetKey()->Address;
 
-		const mu::FModel* Model = LiveInstance->Model.Get();
-		const mu::FProgram& Program = Model->GetPrivate()->Program;
+		const UE::Mutable::Private::FModel* Model = LiveInstance->Model.Get();
+		const UE::Mutable::Private::FProgram& Program = Model->GetPrivate()->Program;
 
 		// TODO: It should be possible to reuse this data if cleared in the correct places only, together with HeapImageDesc.
 		int32 VarValue = CVarClearImageDescCache.GetValueOnAnyThread();
@@ -510,7 +510,7 @@ namespace mu
 			m_pD->WorkingMemoryManager.CurrentInstanceCache->ClearDescCache();
 		}
 
-		mu::EOpType OpType = Program.GetOpType(RootAddress);
+		UE::Mutable::Private::EOpType OpType = Program.GetOpType(RootAddress);
 		if (GetOpDataType(OpType) == EDataType::Image)
 		{
 			// GetImageDesc may call normal execution paths where meshes are computed.
@@ -562,8 +562,8 @@ namespace mu
 
 		OP::ADDRESS RootAddress = ImageId.GetKey()->Address;
 
-		const mu::FModel* Model = LiveInstance->Model.Get();
-		const mu::FProgram& Program = Model->GetPrivate()->Program;
+		const UE::Mutable::Private::FModel* Model = LiveInstance->Model.Get();
+		const UE::Mutable::Private::FProgram& Program = Model->GetPrivate()->Program;
 
 		// TODO: It should be possible to reuse this data if cleared in the correct places only, together with HeapImageDesc.
 		int32 VarValue = CVarClearImageDescCache.GetValueOnAnyThread();
@@ -573,7 +573,7 @@ namespace mu
 		}
 
 		m_pD->WorkingMemoryManager.BeginRunnerThread();
-		mu::EOpType OpType = Program.GetOpType(RootAddress);
+		UE::Mutable::Private::EOpType OpType = Program.GetOpType(RootAddress);
 		if (GetOpDataType(OpType) != EDataType::Image)
 		{
 			m_pD->WorkingMemoryManager.EndRunnerThread();
@@ -670,7 +670,7 @@ namespace mu
 		
 		OP::ADDRESS RootAddress = MeshId.GetKey()->Address;
 			
-		mu::EOpType OpType = LiveInstance->Model->GetPrivate()->Program.GetOpType(RootAddress);
+		UE::Mutable::Private::EOpType OpType = LiveInstance->Model->GetPrivate()->Program.GetOpType(RootAddress);
 		if (GetOpDataType(OpType) != EDataType::Mesh)
 		{
 			m_pD->WorkingMemoryManager.EndRunnerThread();
@@ -736,7 +736,7 @@ namespace mu
 		
 		OP::ADDRESS RootAddress = MaterialId.GetKey()->Address;
 			
-		mu::EOpType OpType = LiveInstance->Model->GetPrivate()->Program.GetOpType(RootAddress);
+		UE::Mutable::Private::EOpType OpType = LiveInstance->Model->GetPrivate()->Program.GetOpType(RootAddress);
 		if (GetOpDataType(OpType) != EDataType::Material)
 		{
 			m_pD->WorkingMemoryManager.EndRunnerThread();
@@ -1132,7 +1132,7 @@ namespace mu
 
 		FVector4f Result(0,0,0,1);
 
-		mu::EOpType OpType = pModel->GetPrivate()->Program.GetOpType(at);
+		UE::Mutable::Private::EOpType OpType = pModel->GetPrivate()->Program.GetOpType(at);
 		if (GetOpDataType(OpType) == EDataType::Color)
 		{
 			RunCode(ExternalResourceProvider, pModel, Params, at);
@@ -1173,7 +1173,7 @@ namespace mu
 
 		TSharedPtr<const FImage> Result;
 
-		mu::EOpType OpType = pModel->GetPrivate()->Program.GetOpType(at);
+		UE::Mutable::Private::EOpType OpType = pModel->GetPrivate()->Program.GetOpType(at);
 		if (GetOpDataType(OpType) == EDataType::Image)
 		{
 			RunCode(ExternalResourceProvider, pModel, Params, at, FSystem::AllLODs, uint8(MipsToSkip), InImageLOD);
@@ -1195,7 +1195,7 @@ namespace mu
 
 		TSharedPtr<const FMesh> Result;
 
-		mu::EOpType OpType = Model->GetPrivate()->Program.GetOpType(RootAddress);
+		UE::Mutable::Private::EOpType OpType = Model->GetPrivate()->Program.GetOpType(RootAddress);
 
 		if (GetOpDataType(OpType) == EDataType::Mesh)
 		{
@@ -1219,7 +1219,7 @@ namespace mu
 
 		TSharedPtr<const FInstance> Result;
 
-		mu::EOpType OpType = pModel->GetPrivate()->Program.GetOpType(at);
+		UE::Mutable::Private::EOpType OpType = pModel->GetPrivate()->Program.GetOpType(at);
 		if (GetOpDataType(OpType) == EDataType::Instance)
 		{
 			RunCode(ExternalResourceProvider, pModel, Params, at);
@@ -1244,7 +1244,7 @@ namespace mu
 
 		if (pModel->GetPrivate()->Program.States[0].Root)
 		{
-			mu::EOpType OpType = pModel->GetPrivate()->Program.GetOpType(at);
+			UE::Mutable::Private::EOpType OpType = pModel->GetPrivate()->Program.GetOpType(at);
 			if (GetOpDataType(OpType) == EDataType::Layout)
 			{
 				RunCode(ExternalResourceProvider, pModel, Params, at);
@@ -1270,7 +1270,7 @@ namespace mu
 
 		if (pModel->GetPrivate()->Program.States[0].Root)
 		{
-			mu::EOpType OpType = pModel->GetPrivate()->Program.GetOpType(at);
+			UE::Mutable::Private::EOpType OpType = pModel->GetPrivate()->Program.GetOpType(at);
 			if (GetOpDataType(OpType) == EDataType::String)
 			{
 				RunCode(ExternalResourceProvider, pModel, Params, at);
@@ -1525,7 +1525,7 @@ namespace mu
 				TSharedPtr<const FModel> CacheModel = ModelCache.Model.Pin();
 				if (CacheModel)
 				{
-					mu::FProgram& Program = CacheModel->GetPrivate()->Program;
+					UE::Mutable::Private::FProgram& Program = CacheModel->GetPrivate()->Program;
 					check(ModelCache.RomWeights.Num() == Program.Roms.Num());
 	
 					check(Program.LoadedMemTrackedRoms.GetMaxIndex() <= Program.Roms.Num());
diff --git a/Engine/Plugins/Mutable/Source/MutableRuntime/Public/MuR/MeshBufferSet.h b/Engine/Plugins/Mutable/Source/MutableRuntime/Public/MuR/MeshBufferSet.h
--- a/Engine/Plugins/Mutable/Source/MutableRuntime/Public/MuR/MeshBufferSet.h
+++ b/Engine/Plugins/Mutable/Source/MutableRuntime/Public/MuR/MeshBufferSet.h
@@ -10,7 +10,7 @@
 
 #define UE_API MUTABLERUNTIME_API
 
-namespace mu::MemoryCounters
+namespace UE::Mutable::Private::MemoryCounters
 {
 	struct FMeshMemoryCounter
 	{
@@ -18,7 +18,7 @@ namespace mu::MemoryCounters
 	};
 }
 
-namespace mu
+namespace UE::Mutable::Private
 {
 	/** Supported formats for the elements in mesh buffers. **/
 	enum class EMeshBufferFormat : uint32
@@ -173,8 +173,8 @@ namespace mu
 		TMemoryTrackedArray<uint8> Data;
 		uint32 ElementSize = 0;
 
-		void Serialise(mu::FOutputArchive& Arch) const;
-		inline void Unserialise(mu::FInputArchive& Arch);
+		void Serialise(UE::Mutable::Private::FOutputArchive& Arch) const;
+		inline void Unserialise(UE::Mutable::Private::FInputArchive& Arch);
 
 		inline bool operator==(const FMeshBuffer& Other) const
 		{
diff --git a/Engine/Plugins/Mutable/Source/MutableRuntime/Public/MuR/Serialisation.h b/Engine/Plugins/Mutable/Source/MutableRuntime/Public/MuR/Serialisation.h
--- a/Engine/Plugins/Mutable/Source/MutableRuntime/Public/MuR/Serialisation.h
+++ b/Engine/Plugins/Mutable/Source/MutableRuntime/Public/MuR/Serialisation.h
@@ -22,25 +22,25 @@ class USkeletalMesh;
 
 #define UE_API MUTABLERUNTIME_API
 
-namespace mu
+namespace UE::Mutable::Private
 {    
     class FImage;
 	class FModel;
 	enum class EDataType : uint8;
 
 #define MUTABLE_DEFINE_POD_SERIALISABLE(Type)									\
-	void DLLEXPORT operator<<(mu::FOutputArchive& Arch, const Type& T);			\
-	void DLLEXPORT operator>>(mu::FInputArchive& Arch, Type& T);
+	void DLLEXPORT operator<<(UE::Mutable::Private::FOutputArchive& Arch, const Type& T);			\
+	void DLLEXPORT operator>>(UE::Mutable::Private::FInputArchive& Arch, Type& T);
 
 #define MUTABLE_DEFINE_POD_VECTOR_SERIALISABLE(Type)							\
 	template<typename Alloc>													\
-	void DLLEXPORT operator<<(mu::FOutputArchive& Arch, const TArray<Type, Alloc>& V);	\
+	void DLLEXPORT operator<<(UE::Mutable::Private::FOutputArchive& Arch, const TArray<Type, Alloc>& V);	\
 	template<typename Alloc>													\
-	void DLLEXPORT operator>>(mu::FInputArchive& Arch, TArray<Type, Alloc>& V);
+	void DLLEXPORT operator>>(UE::Mutable::Private::FInputArchive& Arch, TArray<Type, Alloc>& V);
 
 #define MUTABLE_DEFINE_ENUM_SERIALISABLE(Type)									\
-    void DLLEXPORT operator<<(mu::FOutputArchive& Arch, const Type& T);			\
-    void DLLEXPORT operator>>(mu::FInputArchive& Arch, Type& T);
+    void DLLEXPORT operator<<(UE::Mutable::Private::FOutputArchive& Arch, const Type& T);			\
+    void DLLEXPORT operator>>(UE::Mutable::Private::FInputArchive& Arch, Type& T);
 	
 	
     /** This repesents a handle to Mutable resource. */
diff --git a/Engine/Plugins/Mutable/Source/MutableTools/Internal/MuT/AST.h b/Engine/Plugins/Mutable/Source/MutableTools/Internal/MuT/AST.h
--- a/Engine/Plugins/Mutable/Source/MutableTools/Internal/MuT/AST.h
+++ b/Engine/Plugins/Mutable/Source/MutableTools/Internal/MuT/AST.h
@@ -31,9 +31,9 @@ namespace std
 {
 
   template<typename T>
-  struct hash<mu::Ptr<T>>
+  struct hash<UE::Mutable::Private::Ptr<T>>
   {
-    uint64 operator()(const mu::Ptr<T>& k) const
+    uint64 operator()(const UE::Mutable::Private::Ptr<T>& k) const
     {
       return hash<const void*>()(k.get());
     }
@@ -51,7 +51,7 @@ namespace
 }
 
 
-namespace mu
+namespace UE::Mutable::Private
 {
 	class ASTOp;
 	class ASTOpFixed;
@@ -258,7 +258,7 @@ namespace mu
 
 			static FORCEINLINE uint32 GetKeyHash(KeyInitType Key)
 			{
-				const mu::FMesh* Data = Key.Get();
+				const UE::Mutable::Private::FMesh* Data = Key.Get();
 				return HashCombineFast(
 					::GetTypeHash(Data->VertexBuffers.GetElementCount()),
 					::GetTypeHash(Data->IndexBuffers.GetElementCount())
@@ -278,7 +278,7 @@ namespace mu
 
 			static FORCEINLINE uint32 GetKeyHash(KeyInitType Key)
 			{
-				const mu::FImage* Data = Key.Get();
+				const UE::Mutable::Private::FImage* Data = Key.Get();
 				uint32 Hash = HashCombineFast(::GetTypeHash(Data->GetFormat()), GetTypeHash(Data->GetSize()));
 
 				TArrayView<const uint8> DataView = Data->DataStorage.GetLOD(0);
diff --git a/Engine/Plugins/Mutable/Source/MutableTools/Private/MuT/AST.cpp b/Engine/Plugins/Mutable/Source/MutableTools/Private/MuT/AST.cpp
--- a/Engine/Plugins/Mutable/Source/MutableTools/Private/MuT/AST.cpp
+++ b/Engine/Plugins/Mutable/Source/MutableTools/Private/MuT/AST.cpp
@@ -18,10 +18,10 @@
 #include <atomic>
 
 
-std::atomic<uint32> mu::ASTOp::s_lastTraverseIndex(1);
+std::atomic<uint32> UE::Mutable::Private::ASTOp::s_lastTraverseIndex(1);
 
 
-namespace mu
+namespace UE::Mutable::Private
 {
 
 //-------------------------------------------------------------------------------------------------
@@ -212,7 +212,7 @@ int32 ASTOp::CountNodes( const TArray<Ptr<ASTOp>>& Roots )
 
 
 //-------------------------------------------------------------------------------------------------
-mu::Ptr<ASTOp> ASTOp::DeepClone( const Ptr<ASTOp>& root )
+UE::Mutable::Private::Ptr<ASTOp> ASTOp::DeepClone( const Ptr<ASTOp>& root )
 {
     MUTABLE_CPUPROFILER_SCOPE(AST_DeepClone);
 
@@ -315,7 +315,7 @@ void ASTOp::LogHistogram( ASTOpList& roots )
 	//			UE_LOG(LogMutableCore, Log, TEXT("    base is type %s:"), s_opNames[(int)Base->GetOpType()]);
 	//			if (Constant)
 	//			{
-	//				int Format = int( ((mu::FImage*)Constant->GetValue().get())->GetFormat() );
+	//				int Format = int( ((UE::Mutable::Private::FImage*)Constant->GetValue().get())->GetFormat() );
 	//				UE_LOG(LogMutableCore, Log, TEXT("    constant image format is %d:"), Format);
 	//			}
 	//			else
@@ -963,4 +963,4 @@ void ASTOp::LinkRange(FProgram& program,
 	}
 }
 
-} // namespace mu
+} // namespace UE::Mutable::Private
diff --git a/Engine/Plugins/Mutable/Source/MutableTools/Private/MuT/ASTOpAddExtensionData.cpp b/Engine/Plugins/Mutable/Source/MutableTools/Private/MuT/ASTOpAddExtensionData.cpp
--- a/Engine/Plugins/Mutable/Source/MutableTools/Private/MuT/ASTOpAddExtensionData.cpp
+++ b/Engine/Plugins/Mutable/Source/MutableTools/Private/MuT/ASTOpAddExtensionData.cpp
@@ -7,7 +7,7 @@
 #include "MuR/RefCounted.h"
 
 
-namespace mu
+namespace UE::Mutable::Private
 {
 
 
@@ -50,7 +50,7 @@ uint64 ASTOpAddExtensionData::Hash() const
 	return Result;
 }
 
-mu::Ptr<ASTOp> ASTOpAddExtensionData::Clone(MapChildFuncRef MapChild) const
+UE::Mutable::Private::Ptr<ASTOp> ASTOpAddExtensionData::Clone(MapChildFuncRef MapChild) const
 {
 	Ptr<ASTOpAddExtensionData> NewInstance = new ASTOpAddExtensionData();
 
diff --git a/Engine/Plugins/Mutable/Source/MutableTools/Private/MuT/ASTOpAddLOD.cpp b/Engine/Plugins/Mutable/Source/MutableTools/Private/MuT/ASTOpAddLOD.cpp
--- a/Engine/Plugins/Mutable/Source/MutableTools/Private/MuT/ASTOpAddLOD.cpp
+++ b/Engine/Plugins/Mutable/Source/MutableTools/Private/MuT/ASTOpAddLOD.cpp
@@ -7,7 +7,7 @@
 #include "MuR/RefCounted.h"
 
 
-namespace mu
+namespace UE::Mutable::Private
 {
 
 
@@ -48,7 +48,7 @@ namespace mu
 	}
 
 
-	mu::Ptr<ASTOp> ASTOpAddLOD::Clone(MapChildFuncRef mapChild) const
+	UE::Mutable::Private::Ptr<ASTOp> ASTOpAddLOD::Clone(MapChildFuncRef mapChild) const
 	{
 		Ptr<ASTOpAddLOD> n = new ASTOpAddLOD();
 		for (auto& c : lods)
diff --git a/Engine/Plugins/Mutable/Source/MutableTools/Private/MuT/ASTOpAddMaterial.cpp b/Engine/Plugins/Mutable/Source/MutableTools/Private/MuT/ASTOpAddMaterial.cpp
--- a/Engine/Plugins/Mutable/Source/MutableTools/Private/MuT/ASTOpAddMaterial.cpp
+++ b/Engine/Plugins/Mutable/Source/MutableTools/Private/MuT/ASTOpAddMaterial.cpp
@@ -7,7 +7,7 @@
 #include "MuR/RefCounted.h"
 #include "MuT/CodeOptimiser.h"
 
-namespace mu
+namespace UE::Mutable::Private
 {
 
 
@@ -49,7 +49,7 @@ uint64 ASTOpAddMaterial::Hash() const
 	return Result;
 }
 
-mu::Ptr<ASTOp> ASTOpAddMaterial::Clone(MapChildFuncRef MapChild) const
+UE::Mutable::Private::Ptr<ASTOp> ASTOpAddMaterial::Clone(MapChildFuncRef MapChild) const
 {
 	Ptr<ASTOpAddMaterial> NewInstance = new ASTOpAddMaterial();
 
diff --git a/Engine/Plugins/Mutable/Source/MutableTools/Private/MuT/ASTOpBoolEqualIntConst.cpp b/Engine/Plugins/Mutable/Source/MutableTools/Private/MuT/ASTOpBoolEqualIntConst.cpp
--- a/Engine/Plugins/Mutable/Source/MutableTools/Private/MuT/ASTOpBoolEqualIntConst.cpp
+++ b/Engine/Plugins/Mutable/Source/MutableTools/Private/MuT/ASTOpBoolEqualIntConst.cpp
@@ -9,7 +9,7 @@
 #include "MuR/Types.h"
 
 
-namespace mu
+namespace UE::Mutable::Private
 {
 
 	ASTOpBoolEqualIntConst::ASTOpBoolEqualIntConst()
@@ -36,7 +36,7 @@ namespace mu
 	}
 
 
-	mu::Ptr<ASTOp> ASTOpBoolEqualIntConst::Clone(MapChildFuncRef mapChild) const
+	UE::Mutable::Private::Ptr<ASTOp> ASTOpBoolEqualIntConst::Clone(MapChildFuncRef mapChild) const
 	{
 		Ptr<ASTOpBoolEqualIntConst> n = new ASTOpBoolEqualIntConst();
 		n->Value = mapChild(Value.child());
diff --git a/Engine/Plugins/Mutable/Source/MutableTools/Private/MuT/ASTOpColorFromScalars.cpp b/Engine/Plugins/Mutable/Source/MutableTools/Private/MuT/ASTOpColorFromScalars.cpp
--- a/Engine/Plugins/Mutable/Source/MutableTools/Private/MuT/ASTOpColorFromScalars.cpp
+++ b/Engine/Plugins/Mutable/Source/MutableTools/Private/MuT/ASTOpColorFromScalars.cpp
@@ -23,7 +23,7 @@
 #include "MuT/ASTOpImageMipmap.h"
 #include "MuT/ASTOpSwitch.h"
 
-namespace mu
+namespace UE::Mutable::Private
 {
 
 	ASTOpColorFromScalars::ASTOpColorFromScalars()
@@ -67,9 +67,9 @@ namespace mu
 	}
 
 
-	mu::Ptr<ASTOp> ASTOpColorFromScalars::Clone(MapChildFuncRef mapChild) const
+	UE::Mutable::Private::Ptr<ASTOp> ASTOpColorFromScalars::Clone(MapChildFuncRef mapChild) const
 	{
-		mu::Ptr<ASTOpColorFromScalars> n = new ASTOpColorFromScalars();
+		UE::Mutable::Private::Ptr<ASTOpColorFromScalars> n = new ASTOpColorFromScalars();
 		for (int32 i = 0; i < MUTABLE_OP_MAX_SWIZZLE_CHANNELS; ++i)
 		{
 			n->V[i] = mapChild(V[i].child());
diff --git a/Engine/Plugins/Mutable/Source/MutableTools/Private/MuT/ASTOpColorSwizzle.cpp b/Engine/Plugins/Mutable/Source/MutableTools/Private/MuT/ASTOpColorSwizzle.cpp
--- a/Engine/Plugins/Mutable/Source/MutableTools/Private/MuT/ASTOpColorSwizzle.cpp
+++ b/Engine/Plugins/Mutable/Source/MutableTools/Private/MuT/ASTOpColorSwizzle.cpp
@@ -24,7 +24,7 @@
 #include "MuT/ASTOpColorFromScalars.h"
 #include "MuT/ASTOpSwitch.h"
 
-namespace mu
+namespace UE::Mutable::Private
 {
 
 	ASTOpColorSwizzle::ASTOpColorSwizzle()
@@ -72,9 +72,9 @@ namespace mu
 	}
 
 
-	mu::Ptr<ASTOp> ASTOpColorSwizzle::Clone(MapChildFuncRef mapChild) const
+	UE::Mutable::Private::Ptr<ASTOp> ASTOpColorSwizzle::Clone(MapChildFuncRef mapChild) const
 	{
-		mu::Ptr<ASTOpColorSwizzle> n = new ASTOpColorSwizzle();
+		UE::Mutable::Private::Ptr<ASTOpColorSwizzle> n = new ASTOpColorSwizzle();
 		for (int32 i = 0; i < MUTABLE_OP_MAX_SWIZZLE_CHANNELS; ++i)
 		{
 			n->Sources[i] = mapChild(Sources[i].child());
@@ -115,7 +115,7 @@ namespace mu
 	}
 
 
-	mu::Ptr<ASTOp> ASTOpColorSwizzle::OptimiseSemantic(const FModelOptimizationOptions&, int32 Pass) const
+	UE::Mutable::Private::Ptr<ASTOp> ASTOpColorSwizzle::OptimiseSemantic(const FModelOptimizationOptions&, int32 Pass) const
 	{
 		Ptr<ASTOp> at;
 
@@ -139,7 +139,7 @@ namespace mu
 				{
 					if (!NewSwizzle)
 					{
-						NewSwizzle = mu::Clone<ASTOpColorSwizzle>(this);
+						NewSwizzle = UE::Mutable::Private::Clone<ASTOpColorSwizzle>(this);
 					}
 					const ASTOpColorSwizzle* TypedCandidate = static_cast<const ASTOpColorSwizzle*>(Candidate.get());
 					int32 CandidateChannel = SourceChannels[ChannelIndex];
diff --git a/Engine/Plugins/Mutable/Source/MutableTools/Private/MuT/ASTOpColorToSRGB.cpp b/Engine/Plugins/Mutable/Source/MutableTools/Private/MuT/ASTOpColorToSRGB.cpp
--- a/Engine/Plugins/Mutable/Source/MutableTools/Private/MuT/ASTOpColorToSRGB.cpp
+++ b/Engine/Plugins/Mutable/Source/MutableTools/Private/MuT/ASTOpColorToSRGB.cpp
@@ -7,7 +7,7 @@
 #include "MuR/RefCounted.h"
 #include "MuR/Types.h"
 
-namespace mu
+namespace UE::Mutable::Private
 {
 
 	ASTOpColorToSRGB::ASTOpColorToSRGB()
@@ -41,9 +41,9 @@ namespace mu
 	}
 
 
-	mu::Ptr<ASTOp> ASTOpColorToSRGB::Clone(MapChildFuncRef MapChild) const
+	UE::Mutable::Private::Ptr<ASTOp> ASTOpColorToSRGB::Clone(MapChildFuncRef MapChild) const
 	{
-		mu::Ptr<ASTOpColorToSRGB> New = new ASTOpColorToSRGB();
+		UE::Mutable::Private::Ptr<ASTOpColorToSRGB> New = new ASTOpColorToSRGB();
 		New->Color = MapChild(Color.child());
 		return New;
 	}
diff --git a/Engine/Plugins/Mutable/Source/MutableTools/Private/MuT/ASTOpConditional.cpp b/Engine/Plugins/Mutable/Source/MutableTools/Private/MuT/ASTOpConditional.cpp
--- a/Engine/Plugins/Mutable/Source/MutableTools/Private/MuT/ASTOpConditional.cpp
+++ b/Engine/Plugins/Mutable/Source/MutableTools/Private/MuT/ASTOpConditional.cpp
@@ -12,7 +12,7 @@
 #include "MuT/ASTOpConstantBool.h"
 
 
-namespace mu
+namespace UE::Mutable::Private
 {
 
 
@@ -49,7 +49,7 @@ namespace mu
 
 
 	//-------------------------------------------------------------------------------------------------
-	mu::Ptr<ASTOp> ASTOpConditional::Clone(MapChildFuncRef mapChild) const
+	UE::Mutable::Private::Ptr<ASTOp> ASTOpConditional::Clone(MapChildFuncRef mapChild) const
 	{
 		Ptr<ASTOpConditional> n = new ASTOpConditional();
 		n->type = type;
@@ -281,7 +281,7 @@ namespace mu
 
 
 	//-------------------------------------------------------------------------------------------------
-	mu::Ptr<ASTOp> ASTOpConditional::OptimiseSemantic(const FModelOptimizationOptions&, int32 Pass) const
+	UE::Mutable::Private::Ptr<ASTOp> ASTOpConditional::OptimiseSemantic(const FModelOptimizationOptions&, int32 Pass) const
 	{
 		if (!condition)
 		{
@@ -317,7 +317,7 @@ namespace mu
 					||
 					*condition.child() == *typedYes->condition.child())
 				{
-					auto op = mu::Clone<ASTOpConditional>(this);
+					auto op = UE::Mutable::Private::Clone<ASTOpConditional>(this);
 					op->yes = typedYes->yes.child();
 					return op;
 				}
@@ -331,7 +331,7 @@ namespace mu
 					||
 					*condition.child() == *typedNo->condition.child())
 				{
-					auto op = mu::Clone<ASTOpConditional>(this);
+					auto op = UE::Mutable::Private::Clone<ASTOpConditional>(this);
 					op->no = typedNo->no.child();
 					return op;
 				}
@@ -343,7 +343,7 @@ namespace mu
 
 
 	//-------------------------------------------------------------------------------------------------
-	mu::Ptr<ImageSizeExpression> ASTOpConditional::GetImageSizeExpression() const
+	UE::Mutable::Private::Ptr<ImageSizeExpression> ASTOpConditional::GetImageSizeExpression() const
 	{
 		Ptr<ImageSizeExpression> pRes = new ImageSizeExpression;
 		pRes->type = ImageSizeExpression::ISET_CONDITIONAL;
diff --git a/Engine/Plugins/Mutable/Source/MutableTools/Private/MuT/ASTOpConstantBool.cpp b/Engine/Plugins/Mutable/Source/MutableTools/Private/MuT/ASTOpConstantBool.cpp
--- a/Engine/Plugins/Mutable/Source/MutableTools/Private/MuT/ASTOpConstantBool.cpp
+++ b/Engine/Plugins/Mutable/Source/MutableTools/Private/MuT/ASTOpConstantBool.cpp
@@ -8,7 +8,7 @@
 #include "MuR/Types.h"
 
 
-namespace mu
+namespace UE::Mutable::Private
 {
 
 	ASTOpConstantBool::ASTOpConstantBool(bool InValue)
@@ -41,7 +41,7 @@ namespace mu
 	}
 
 
-	mu::Ptr<ASTOp> ASTOpConstantBool::Clone(MapChildFuncRef) const
+	UE::Mutable::Private::Ptr<ASTOp> ASTOpConstantBool::Clone(MapChildFuncRef) const
 	{
 		Ptr<ASTOpConstantBool> New = new ASTOpConstantBool();
 		New->bValue = bValue;
diff --git a/Engine/Plugins/Mutable/Source/MutableTools/Private/MuT/ASTOpConstantColor.cpp b/Engine/Plugins/Mutable/Source/MutableTools/Private/MuT/ASTOpConstantColor.cpp
--- a/Engine/Plugins/Mutable/Source/MutableTools/Private/MuT/ASTOpConstantColor.cpp
+++ b/Engine/Plugins/Mutable/Source/MutableTools/Private/MuT/ASTOpConstantColor.cpp
@@ -8,7 +8,7 @@
 #include "MuR/Types.h"
 
 
-namespace mu
+namespace UE::Mutable::Private
 {
 
 	void ASTOpConstantColor::ForEachChild(const TFunctionRef<void(ASTChild&)>)
@@ -46,7 +46,7 @@ namespace mu
 	}
 
 
-	mu::Ptr<ASTOp> ASTOpConstantColor::Clone(MapChildFuncRef) const
+	UE::Mutable::Private::Ptr<ASTOp> ASTOpConstantColor::Clone(MapChildFuncRef) const
 	{
 		Ptr<ASTOpConstantColor> New = new ASTOpConstantColor();
 		New->Value = Value;
diff --git a/Engine/Plugins/Mutable/Source/MutableTools/Private/MuT/ASTOpConstantExtensionData.cpp b/Engine/Plugins/Mutable/Source/MutableTools/Private/MuT/ASTOpConstantExtensionData.cpp
--- a/Engine/Plugins/Mutable/Source/MutableTools/Private/MuT/ASTOpConstantExtensionData.cpp
+++ b/Engine/Plugins/Mutable/Source/MutableTools/Private/MuT/ASTOpConstantExtensionData.cpp
@@ -6,7 +6,7 @@
 #include "MuR/ModelPrivate.h"
 
 
-namespace mu
+namespace UE::Mutable::Private
 {
 
 
@@ -24,7 +24,7 @@ bool ASTOpConstantExtensionData::IsEqual(const ASTOp& OtherUntyped) const
 
 
 //-------------------------------------------------------------------------------------------------
-mu::Ptr<ASTOp> ASTOpConstantExtensionData::Clone(MapChildFuncRef MapChild) const
+UE::Mutable::Private::Ptr<ASTOp> ASTOpConstantExtensionData::Clone(MapChildFuncRef MapChild) const
 {
 	Ptr<ASTOpConstantExtensionData> Result = new ASTOpConstantExtensionData();
 
diff --git a/Engine/Plugins/Mutable/Source/MutableTools/Private/MuT/ASTOpConstantInt.cpp b/Engine/Plugins/Mutable/Source/MutableTools/Private/MuT/ASTOpConstantInt.cpp
--- a/Engine/Plugins/Mutable/Source/MutableTools/Private/MuT/ASTOpConstantInt.cpp
+++ b/Engine/Plugins/Mutable/Source/MutableTools/Private/MuT/ASTOpConstantInt.cpp
@@ -8,7 +8,7 @@
 #include "MuR/Types.h"
 
 
-namespace mu
+namespace UE::Mutable::Private
 {
 
 	ASTOpConstantInt::ASTOpConstantInt(int32 InValue)
@@ -41,7 +41,7 @@ namespace mu
 	}
 
 
-	mu::Ptr<ASTOp> ASTOpConstantInt::Clone(MapChildFuncRef) const
+	UE::Mutable::Private::Ptr<ASTOp> ASTOpConstantInt::Clone(MapChildFuncRef) const
 	{
 		Ptr<ASTOpConstantInt> New = new ASTOpConstantInt();
 		New->Value = Value;
diff --git a/Engine/Plugins/Mutable/Source/MutableTools/Private/MuT/ASTOpConstantMatrix.cpp b/Engine/Plugins/Mutable/Source/MutableTools/Private/MuT/ASTOpConstantMatrix.cpp
--- a/Engine/Plugins/Mutable/Source/MutableTools/Private/MuT/ASTOpConstantMatrix.cpp
+++ b/Engine/Plugins/Mutable/Source/MutableTools/Private/MuT/ASTOpConstantMatrix.cpp
@@ -2,18 +2,18 @@
 
 #include "MuT/ASTOpConstantMatrix.h"
 
-mu::ASTOpConstantMatrix::ASTOpConstantMatrix(FMatrix44f InitValue) : value(InitValue)
+UE::Mutable::Private::ASTOpConstantMatrix::ASTOpConstantMatrix(FMatrix44f InitValue) : value(InitValue)
 {
 }
 
-uint64 mu::ASTOpConstantMatrix::Hash() const
+uint64 UE::Mutable::Private::ASTOpConstantMatrix::Hash() const
 {
 	uint64 res = std::hash<uint64>()(uint64(EOpType::MA_CONSTANT));
 	hash_combine(res, value.ComputeHash());
 	return res;
 }
 
-bool mu::ASTOpConstantMatrix::IsEqual(const ASTOp& otherUntyped) const
+bool UE::Mutable::Private::ASTOpConstantMatrix::IsEqual(const ASTOp& otherUntyped) const
 {
 	if (otherUntyped.GetOpType() == GetOpType())
 	{
@@ -24,14 +24,14 @@ bool mu::ASTOpConstantMatrix::IsEqual(const ASTOp& otherUntyped) const
 	return false;
 }
 
-mu::Ptr<mu::ASTOp> mu::ASTOpConstantMatrix::Clone(MapChildFuncRef mapChild) const
+UE::Mutable::Private::Ptr<UE::Mutable::Private::ASTOp> UE::Mutable::Private::ASTOpConstantMatrix::Clone(MapChildFuncRef mapChild) const
 {
 	Ptr<ASTOpConstantMatrix> n = new ASTOpConstantMatrix();
 	n->value = value;
 	return n;
 }
 
-void mu::ASTOpConstantMatrix::Link(FProgram& program, FLinkerOptions* Options)
+void UE::Mutable::Private::ASTOpConstantMatrix::Link(FProgram& program, FLinkerOptions* Options)
 {
 	if (!linkedAddress)
 	{
diff --git a/Engine/Plugins/Mutable/Source/MutableTools/Private/MuT/ASTOpConstantProjector.cpp b/Engine/Plugins/Mutable/Source/MutableTools/Private/MuT/ASTOpConstantProjector.cpp
--- a/Engine/Plugins/Mutable/Source/MutableTools/Private/MuT/ASTOpConstantProjector.cpp
+++ b/Engine/Plugins/Mutable/Source/MutableTools/Private/MuT/ASTOpConstantProjector.cpp
@@ -7,7 +7,7 @@
 #include "MuR/RefCounted.h"
 
 
-namespace mu
+namespace UE::Mutable::Private
 {
 
 
@@ -38,7 +38,7 @@ namespace mu
 	}
 
 
-	mu::Ptr<ASTOp> ASTOpConstantProjector::Clone(MapChildFuncRef) const
+	UE::Mutable::Private::Ptr<ASTOp> ASTOpConstantProjector::Clone(MapChildFuncRef) const
 	{
 		Ptr<ASTOpConstantProjector> n = new ASTOpConstantProjector();
 		n->value = value;
diff --git a/Engine/Plugins/Mutable/Source/MutableTools/Private/MuT/ASTOpConstantResource.cpp b/Engine/Plugins/Mutable/Source/MutableTools/Private/MuT/ASTOpConstantResource.cpp
--- a/Engine/Plugins/Mutable/Source/MutableTools/Private/MuT/ASTOpConstantResource.cpp
+++ b/Engine/Plugins/Mutable/Source/MutableTools/Private/MuT/ASTOpConstantResource.cpp
@@ -21,7 +21,7 @@
 
 #include <inttypes.h> // Required for 64-bit printf macros
 
-namespace mu
+namespace UE::Mutable::Private
 {
 
 	/** Proxy class for a temporary resource while compiling. 
@@ -258,7 +258,7 @@ namespace mu
 
 
 	//-------------------------------------------------------------------------------------------------
-	mu::Ptr<ASTOp> ASTOpConstantResource::Clone(MapChildFuncRef MapChild) const
+	UE::Mutable::Private::Ptr<ASTOp> ASTOpConstantResource::Clone(MapChildFuncRef MapChild) const
 	{
 		Ptr<ASTOpConstantResource> n = new ASTOpConstantResource();
 		n->Type = Type;
@@ -993,7 +993,7 @@ namespace mu
 	}
 
 
-	mu::Ptr<ImageSizeExpression> ASTOpConstantResource::GetImageSizeExpression() const
+	UE::Mutable::Private::Ptr<ImageSizeExpression> ASTOpConstantResource::GetImageSizeExpression() const
 	{
 		if (Type==EOpType::IM_CONSTANT)
 		{
diff --git a/Engine/Plugins/Mutable/Source/MutableTools/Private/MuT/ASTOpConstantScalar.cpp b/Engine/Plugins/Mutable/Source/MutableTools/Private/MuT/ASTOpConstantScalar.cpp
--- a/Engine/Plugins/Mutable/Source/MutableTools/Private/MuT/ASTOpConstantScalar.cpp
+++ b/Engine/Plugins/Mutable/Source/MutableTools/Private/MuT/ASTOpConstantScalar.cpp
@@ -8,7 +8,7 @@
 #include "MuR/Types.h"
 
 
-namespace mu
+namespace UE::Mutable::Private
 {
 
 	ASTOpConstantScalar::ASTOpConstantScalar(float InValue)
@@ -41,7 +41,7 @@ namespace mu
 	}
 
 
-	mu::Ptr<ASTOp> ASTOpConstantScalar::Clone(MapChildFuncRef) const
+	UE::Mutable::Private::Ptr<ASTOp> ASTOpConstantScalar::Clone(MapChildFuncRef) const
 	{
 		Ptr<ASTOpConstantScalar> New = new ASTOpConstantScalar();
 		New->Value = Value;
diff --git a/Engine/Plugins/Mutable/Source/MutableTools/Private/MuT/ASTOpConstantString.cpp b/Engine/Plugins/Mutable/Source/MutableTools/Private/MuT/ASTOpConstantString.cpp
--- a/Engine/Plugins/Mutable/Source/MutableTools/Private/MuT/ASTOpConstantString.cpp
+++ b/Engine/Plugins/Mutable/Source/MutableTools/Private/MuT/ASTOpConstantString.cpp
@@ -8,7 +8,7 @@
 #include "MuR/Types.h"
 
 
-namespace mu
+namespace UE::Mutable::Private
 {
 
 
@@ -38,7 +38,7 @@ namespace mu
 	}
 
 
-	mu::Ptr<ASTOp> ASTOpConstantString::Clone(MapChildFuncRef) const
+	UE::Mutable::Private::Ptr<ASTOp> ASTOpConstantString::Clone(MapChildFuncRef) const
 	{
 		Ptr<ASTOpConstantString> n = new ASTOpConstantString();
 		n->value = value;
diff --git a/Engine/Plugins/Mutable/Source/MutableTools/Private/MuT/ASTOpImageCompose.cpp b/Engine/Plugins/Mutable/Source/MutableTools/Private/MuT/ASTOpImageCompose.cpp
--- a/Engine/Plugins/Mutable/Source/MutableTools/Private/MuT/ASTOpImageCompose.cpp
+++ b/Engine/Plugins/Mutable/Source/MutableTools/Private/MuT/ASTOpImageCompose.cpp
@@ -17,7 +17,7 @@
 #include "MuT/ASTOpImageResize.h"
 
 
-namespace mu
+namespace UE::Mutable::Private
 {
 
 
@@ -71,9 +71,9 @@ uint64 ASTOpImageCompose::Hash() const
 
 
 //-------------------------------------------------------------------------------------------------
-mu::Ptr<ASTOp> ASTOpImageCompose::Clone(MapChildFuncRef mapChild) const
+UE::Mutable::Private::Ptr<ASTOp> ASTOpImageCompose::Clone(MapChildFuncRef mapChild) const
 {
-	mu::Ptr<ASTOpImageCompose> n = new ASTOpImageCompose();
+	UE::Mutable::Private::Ptr<ASTOpImageCompose> n = new ASTOpImageCompose();
     n->Layout = mapChild(Layout.child());
 	n->Base = mapChild(Base.child());
 	n->BlockImage = mapChild(BlockImage.child());
@@ -159,7 +159,7 @@ FImageDesc ASTOpImageCompose::GetImageDesc( bool returnBestOption, FGetImageDesc
 }
 
 
-mu::Ptr<ImageSizeExpression> ASTOpImageCompose::GetImageSizeExpression() const
+UE::Mutable::Private::Ptr<ImageSizeExpression> ASTOpImageCompose::GetImageSizeExpression() const
 {
     if ( Base )
     {
@@ -227,13 +227,13 @@ void ASTOpImageCompose::GetLayoutBlockSize(int32* pBlockX, int32* pBlockY)
 }
 
 
-mu::Ptr<ASTOp> ASTOpImageCompose::OptimiseSemantic(const FModelOptimizationOptions& options, int32 Pass) const
+UE::Mutable::Private::Ptr<ASTOp> ASTOpImageCompose::OptimiseSemantic(const FModelOptimizationOptions& options, int32 Pass) const
 {
-	mu::Ptr<ASTOp> at;
+	UE::Mutable::Private::Ptr<ASTOp> at;
 
-	mu::Ptr<ASTOp> baseAt = Base.child();
-	mu::Ptr<ASTOp> blockAt = BlockImage.child();
-	mu::Ptr<ASTOp> layoutAt = Layout.child();
+	UE::Mutable::Private::Ptr<ASTOp> baseAt = Base.child();
+	UE::Mutable::Private::Ptr<ASTOp> blockAt = BlockImage.child();
+	UE::Mutable::Private::Ptr<ASTOp> layoutAt = Layout.child();
 	if (layoutAt
 		&&
 		layoutAt->GetOpType() == EOpType::LA_CONSTANT
@@ -243,7 +243,7 @@ mu::Ptr<ASTOp> ASTOpImageCompose::OptimiseSemantic(const FModelOptimizationOptio
 		blockAt)
 	{
 		const ASTOpConstantResource* typedLayout = static_cast<const ASTOpConstantResource*>(layoutAt.get());
-		const mu::FLayout* pLayout = static_cast<const mu::FLayout*>(typedLayout->GetValue().Get());
+		const UE::Mutable::Private::FLayout* pLayout = static_cast<const UE::Mutable::Private::FLayout*>(typedLayout->GetValue().Get());
 
 		// Constant single-block full layout?
 		if (pLayout->GetBlockCount() == 1
@@ -257,8 +257,8 @@ mu::Ptr<ASTOp> ASTOpImageCompose::OptimiseSemantic(const FModelOptimizationOptio
 		{
 			// We could only take the block, but we must make sure it will have the format
 			// and size of the base.
-			mu::FImageDesc BaseDesc = baseAt->GetImageDesc(true);
-			mu::FImageDesc BlockDesc = blockAt->GetImageDesc(true);
+			UE::Mutable::Private::FImageDesc BaseDesc = baseAt->GetImageDesc(true);
+			UE::Mutable::Private::FImageDesc BlockDesc = blockAt->GetImageDesc(true);
 
 			at = blockAt;
 
@@ -267,7 +267,7 @@ mu::Ptr<ASTOp> ASTOpImageCompose::OptimiseSemantic(const FModelOptimizationOptio
 				&&
 				BaseDesc.m_format != EImageFormat::None)
 			{
-				mu::Ptr<ASTOpImagePixelFormat> Reformat = new ASTOpImagePixelFormat;
+				UE::Mutable::Private::Ptr<ASTOpImagePixelFormat> Reformat = new ASTOpImagePixelFormat;
 				Reformat->Format = BaseDesc.m_format;
 				Reformat->FormatIfAlpha = EImageFormat::None;
 				Reformat->Source = at;
@@ -280,7 +280,7 @@ mu::Ptr<ASTOp> ASTOpImageCompose::OptimiseSemantic(const FModelOptimizationOptio
 				&&
 				BaseDesc.m_size[1] != 0)
 			{
-				mu::Ptr<ASTOpImageResize> Resize = new ASTOpImageResize;
+				UE::Mutable::Private::Ptr<ASTOpImageResize> Resize = new ASTOpImageResize;
 				Resize->Size[0] = BaseDesc.m_size[0];
 				Resize->Size[1] = BaseDesc.m_size[1];
 				Resize->Source = at;
diff --git a/Engine/Plugins/Mutable/Source/MutableTools/Private/MuT/ASTOpImageCrop.cpp b/Engine/Plugins/Mutable/Source/MutableTools/Private/MuT/ASTOpImageCrop.cpp
--- a/Engine/Plugins/Mutable/Source/MutableTools/Private/MuT/ASTOpImageCrop.cpp
+++ b/Engine/Plugins/Mutable/Source/MutableTools/Private/MuT/ASTOpImageCrop.cpp
@@ -18,7 +18,7 @@
 #include "MuT/ASTOpImageDisplace.h"
 
 
-namespace mu
+namespace UE::Mutable::Private
 {
 
 	ASTOpImageCrop::ASTOpImageCrop()
@@ -188,7 +188,7 @@ namespace mu
 		{
 		case EOpType::IM_PLAINCOLOUR:
 		{
-			Ptr<ASTOpImagePlainColor> NewOp = mu::Clone<ASTOpImagePlainColor>(SourceAt.get());
+			Ptr<ASTOpImagePlainColor> NewOp = UE::Mutable::Private::Clone<ASTOpImagePlainColor>(SourceAt.get());
 			NewOp->Size[0] = Size[0];
 			NewOp->Size[1] = Size[1];
 			NewOp->LODs = 1; // TODO
@@ -269,7 +269,7 @@ namespace mu
 		case EOpType::IM_CONDITIONAL:
 		{
 			// We move down the two paths
-			Ptr<ASTOpConditional> newOp = mu::Clone<ASTOpConditional>(at);
+			Ptr<ASTOpConditional> newOp = UE::Mutable::Private::Clone<ASTOpConditional>(at);
 			newOp->yes = Visit(newOp->yes.child(), currentCropOp);
 			newOp->no = Visit(newOp->no.child(), currentCropOp);
 			newAt = newOp;
@@ -279,7 +279,7 @@ namespace mu
 		case EOpType::IM_SWITCH:
 		{
 			// We move down all the paths
-			Ptr<ASTOpSwitch> newOp = mu::Clone<ASTOpSwitch>(at);
+			Ptr<ASTOpSwitch> newOp = UE::Mutable::Private::Clone<ASTOpSwitch>(at);
 			newOp->Default = Visit(newOp->Default.child(), currentCropOp);
 			for (auto& c : newOp->Cases)
 			{
@@ -291,7 +291,7 @@ namespace mu
 
 		case EOpType::IM_PIXELFORMAT:
 		{
-			Ptr<ASTOpImagePixelFormat> nop = mu::Clone<ASTOpImagePixelFormat>(at);
+			Ptr<ASTOpImagePixelFormat> nop = UE::Mutable::Private::Clone<ASTOpImagePixelFormat>(at);
 			nop->Source = Visit(nop->Source.child(), currentCropOp);
 			newAt = nop;
 			break;
@@ -325,13 +325,13 @@ namespace mu
 			{
 				// Crop the base with the full crop, and the patch with the intersected part,
 				// adapting the patch origin
-				Ptr<ASTOpImagePatch> newOp = mu::Clone<ASTOpImagePatch>(at);
+				Ptr<ASTOpImagePatch> newOp = UE::Mutable::Private::Clone<ASTOpImagePatch>(at);
 				newOp->base = Visit(newOp->base.child(), currentCropOp);
 
 				box<FIntVector2> ibox = patchBox.Intersect2i(cropBox);
 				check(ibox.size[0] > 0 && ibox.size[1] > 0);
 
-				Ptr<ASTOpImageCrop> patchCropOp = mu::Clone<ASTOpImageCrop>(currentCropOp);
+				Ptr<ASTOpImageCrop> patchCropOp = UE::Mutable::Private::Clone<ASTOpImageCrop>(currentCropOp);
 				patchCropOp->Min[0] = ibox.min[0] - patchBox.min[0];
 				patchCropOp->Min[1] = ibox.min[1] - patchBox.min[1];
 				patchCropOp->Size[0] = ibox.size[0];
@@ -388,7 +388,7 @@ namespace mu
 			else
 			{
 				// combine into one crop
-				Ptr<ASTOpImageCrop> newCropOp = mu::Clone<ASTOpImageCrop>(currentCropOp);
+				Ptr<ASTOpImageCrop> newCropOp = UE::Mutable::Private::Clone<ASTOpImageCrop>(currentCropOp);
 				newCropOp->Min[0] = ibox.min[0];
 				newCropOp->Min[1] = ibox.min[1];
 				newCropOp->Size[0] = ibox.size[0];
@@ -403,7 +403,7 @@ namespace mu
 		case EOpType::IM_LAYER:
 		{
 			// We move the op down the arguments
-			Ptr<ASTOpImageLayer> nop = mu::Clone<ASTOpImageLayer>(at);
+			Ptr<ASTOpImageLayer> nop = UE::Mutable::Private::Clone<ASTOpImageLayer>(at);
 
 			Ptr<ASTOp> aOp = nop->base.child();
 			nop->base = Visit(aOp, currentCropOp);
@@ -421,7 +421,7 @@ namespace mu
 		case EOpType::IM_LAYERCOLOUR:
 		{
 			// We move the op down the arguments
-			Ptr<ASTOpImageLayerColor> nop = mu::Clone<ASTOpImageLayerColor>(at);
+			Ptr<ASTOpImageLayerColor> nop = UE::Mutable::Private::Clone<ASTOpImageLayerColor>(at);
 
 			Ptr<ASTOp> aOp = nop->base.child();
 			nop->base = Visit(aOp, currentCropOp);
@@ -436,7 +436,7 @@ namespace mu
 		case EOpType::IM_DISPLACE:
 		{
 			// We move the op down the arguments
-			Ptr<ASTOpImageDisplace> nop = mu::Clone<ASTOpImageDisplace>(at);
+			Ptr<ASTOpImageDisplace> nop = UE::Mutable::Private::Clone<ASTOpImageDisplace>(at);
 
 			Ptr<ASTOp> aOp = nop->Source.child();
 			nop->Source = Visit(aOp, currentCropOp);
@@ -452,7 +452,7 @@ namespace mu
 		{
 			// We add cropping data to the raster mesh if it doesn't have any
 			// \TODO: Is is possible to hit 2 crops on a raster mesh? Combine the crop.
-			Ptr<ASTOpImageRasterMesh> nop = mu::Clone<ASTOpImageRasterMesh>(at);
+			Ptr<ASTOpImageRasterMesh> nop = UE::Mutable::Private::Clone<ASTOpImageRasterMesh>(at);
 
 			bool bRasterHasCrop = nop->UncroppedSizeX != 0;
 			if (!bRasterHasCrop)
@@ -478,7 +478,7 @@ namespace mu
 		case EOpType::IM_INTERPOLATE:
 		{
 			// Move the op  down all the paths
-			Ptr<ASTOpImageInterpolate> NewOp = mu::Clone<ASTOpImageInterpolate>(at);
+			Ptr<ASTOpImageInterpolate> NewOp = UE::Mutable::Private::Clone<ASTOpImageInterpolate>(at);
 
 			for (int32 v = 0; v < MUTABLE_OP_MAX_INTERPOLATE_COUNT; ++v)
 			{
@@ -498,7 +498,7 @@ namespace mu
 		// end on line, replace with crop
 		if (at == newAt && at != InitialSource && !skipSinking)
 		{
-			Ptr<ASTOpImageCrop> newOp = mu::Clone<ASTOpImageCrop>(currentCropOp);
+			Ptr<ASTOpImageCrop> newOp = UE::Mutable::Private::Clone<ASTOpImageCrop>(currentCropOp);
 			check(newOp->GetOpType() == EOpType::IM_CROP);
 
 			newOp->Source = at;
