[Mutable] Change namespace name

diff --git a/Engine/Plugins/Mutable/Source/MutableTools/Private/MuT/ASTOpImageDisplace.cpp b/Engine/Plugins/Mutable/Source/MutableTools/Private/MuT/ASTOpImageDisplace.cpp
--- a/Engine/Plugins/Mutable/Source/MutableTools/Private/MuT/ASTOpImageDisplace.cpp
+++ b/Engine/Plugins/Mutable/Source/MutableTools/Private/MuT/ASTOpImageDisplace.cpp
@@ -21,7 +21,7 @@
 #include "MuT/ASTOpImageResize.h"
 
 
-namespace mu
+namespace UE::Mutable::Private
 {
 
 	ASTOpImageDisplace::ASTOpImageDisplace()
@@ -178,14 +178,14 @@ namespace mu
 
 				if (typedSource->condition == typedDisplacementMap->condition)
 				{
-					Ptr<ASTOpConditional> nop = mu::Clone<ASTOpConditional>(sourceAt);
+					Ptr<ASTOpConditional> nop = UE::Mutable::Private::Clone<ASTOpConditional>(sourceAt);
 
-					Ptr<ASTOpImageDisplace> aOp = mu::Clone<ASTOpImageDisplace>(this);
+					Ptr<ASTOpImageDisplace> aOp = UE::Mutable::Private::Clone<ASTOpImageDisplace>(this);
 					aOp->Source = typedSource->yes.child();
 					aOp->DisplacementMap = typedDisplacementMap->yes.child();
 					nop->yes = aOp;
 
-					Ptr<ASTOpImageDisplace> bOp = mu::Clone<ASTOpImageDisplace>(this);
+					Ptr<ASTOpImageDisplace> bOp = UE::Mutable::Private::Clone<ASTOpImageDisplace>(this);
 					bOp->Source = typedSource->no.child();
 					bOp->DisplacementMap = typedDisplacementMap->no.child();
 					nop->no = bOp;
@@ -206,11 +206,11 @@ namespace mu
 				if (typedSource->IsCompatibleWith(typedDisplacementMap))
 				{
 					// Move the format down all the paths
-					auto nop = mu::Clone<ASTOpSwitch>(sourceAt);
+					auto nop = UE::Mutable::Private::Clone<ASTOpSwitch>(sourceAt);
 
 					if (nop->Default)
 					{
-						Ptr<ASTOpImageDisplace> defOp = mu::Clone<ASTOpImageDisplace>(this);
+						Ptr<ASTOpImageDisplace> defOp = UE::Mutable::Private::Clone<ASTOpImageDisplace>(this);
 						defOp->Source = typedSource->Default.child();
 						defOp->DisplacementMap = typedDisplacementMap->Default.child();
 						nop->Default = defOp;
@@ -221,7 +221,7 @@ namespace mu
 					{
 						if (nop->Cases[v].Branch)
 						{
-							Ptr<ASTOpImageDisplace> bOp = mu::Clone<ASTOpImageDisplace>(this);
+							Ptr<ASTOpImageDisplace> bOp = UE::Mutable::Private::Clone<ASTOpImageDisplace>(this);
 							bOp->Source = typedSource->Cases[v].Branch.child();
 							bOp->DisplacementMap = typedDisplacementMap->FindBranch(typedSource->Cases[v].Condition);
 							nop->Cases[v].Branch = bOp;
@@ -236,11 +236,11 @@ namespace mu
 			if (OriginalAt == at)
 			{
 				// Make a project for every path
-				Ptr<ASTOpSwitch> nop = mu::Clone<ASTOpSwitch>(sourceAt.get());
+				Ptr<ASTOpSwitch> nop = UE::Mutable::Private::Clone<ASTOpSwitch>(sourceAt.get());
 
 				if (nop->Default)
 				{
-					Ptr<ASTOpImageDisplace> defOp = mu::Clone<ASTOpImageDisplace>(this);
+					Ptr<ASTOpImageDisplace> defOp = UE::Mutable::Private::Clone<ASTOpImageDisplace>(this);
 					defOp->Source = nop->Default.child();
 					nop->Default = defOp;
 				}
@@ -250,7 +250,7 @@ namespace mu
 				{
 					if (nop->Cases[o].Branch)
 					{
-						Ptr<ASTOpImageDisplace> bOp = mu::Clone<ASTOpImageDisplace>(this);
+						Ptr<ASTOpImageDisplace> bOp = UE::Mutable::Private::Clone<ASTOpImageDisplace>(this);
 						bOp->Source = nop->Cases[o].Branch.child();
 						nop->Cases[o].Branch = bOp;
 					}
diff --git a/Engine/Plugins/Mutable/Source/MutableTools/Private/MuT/ASTOpImageInterpolate.cpp b/Engine/Plugins/Mutable/Source/MutableTools/Private/MuT/ASTOpImageInterpolate.cpp
--- a/Engine/Plugins/Mutable/Source/MutableTools/Private/MuT/ASTOpImageInterpolate.cpp
+++ b/Engine/Plugins/Mutable/Source/MutableTools/Private/MuT/ASTOpImageInterpolate.cpp
@@ -23,7 +23,7 @@
 #include "MuT/ASTOpImageMipmap.h"
 #include "MuT/ASTOpSwitch.h"
 
-namespace mu
+namespace UE::Mutable::Private
 {
 
 	ASTOpImageInterpolate::ASTOpImageInterpolate()
@@ -70,9 +70,9 @@ namespace mu
 	}
 
 
-	mu::Ptr<ASTOp> ASTOpImageInterpolate::Clone(MapChildFuncRef MapChild) const
+	UE::Mutable::Private::Ptr<ASTOp> ASTOpImageInterpolate::Clone(MapChildFuncRef MapChild) const
 	{
-		mu::Ptr<ASTOpImageInterpolate> n = new ASTOpImageInterpolate();
+		UE::Mutable::Private::Ptr<ASTOpImageInterpolate> n = new ASTOpImageInterpolate();
 		n->Factor = MapChild(Factor.child());
 		for (int32 i = 0; i < MUTABLE_OP_MAX_INTERPOLATE_COUNT; ++i)
 		{
@@ -164,9 +164,9 @@ namespace mu
 	}
 
 
-	mu::Ptr<ImageSizeExpression> ASTOpImageInterpolate::GetImageSizeExpression() const
+	UE::Mutable::Private::Ptr<ImageSizeExpression> ASTOpImageInterpolate::GetImageSizeExpression() const
 	{
-		mu::Ptr<ImageSizeExpression> pRes;
+		UE::Mutable::Private::Ptr<ImageSizeExpression> pRes;
 
 		if (Targets[0].child())
 		{
diff --git a/Engine/Plugins/Mutable/Source/MutableTools/Private/MuT/ASTOpImageLayer.cpp b/Engine/Plugins/Mutable/Source/MutableTools/Private/MuT/ASTOpImageLayer.cpp
--- a/Engine/Plugins/Mutable/Source/MutableTools/Private/MuT/ASTOpImageLayer.cpp
+++ b/Engine/Plugins/Mutable/Source/MutableTools/Private/MuT/ASTOpImageLayer.cpp
@@ -17,7 +17,7 @@
 #include "HAL/PlatformMath.h"
 
 
-namespace mu
+namespace UE::Mutable::Private
 {
 
 	ASTOpImageLayer::ASTOpImageLayer()
@@ -62,7 +62,7 @@ namespace mu
 	}
 
 
-	mu::Ptr<ASTOp> ASTOpImageLayer::Clone(MapChildFuncRef mapChild) const
+	UE::Mutable::Private::Ptr<ASTOp> ASTOpImageLayer::Clone(MapChildFuncRef mapChild) const
 	{
 		Ptr<ASTOpImageLayer> n = new ASTOpImageLayer();
 		n->base = mapChild(base.child());
@@ -155,7 +155,7 @@ namespace mu
 	}
 
 
-	mu::Ptr<ImageSizeExpression> ASTOpImageLayer::GetImageSizeExpression() const
+	UE::Mutable::Private::Ptr<ImageSizeExpression> ASTOpImageLayer::GetImageSizeExpression() const
 	{
 		if (base)
 		{
@@ -213,7 +213,7 @@ namespace mu
 				else if (FMath::IsNearlyEqual(colour[0], 1, UE_SMALL_NUMBER))
 				{
 					// If the mask is white, we can remove it
-					Ptr<ASTOpImageLayer> NewOp = mu::Clone<ASTOpImageLayer>(this);
+					Ptr<ASTOpImageLayer> NewOp = UE::Mutable::Private::Clone<ASTOpImageLayer>(this);
 					NewOp->mask = nullptr;
 					at = NewOp;
 				}
@@ -367,7 +367,7 @@ namespace mu
 					)
 				{
 					// we can do something good here
-					Ptr<ASTOpImageLayer> NewLayer = mu::Clone<ASTOpImageLayer>(this);
+					Ptr<ASTOpImageLayer> NewLayer = UE::Mutable::Private::Clone<ASTOpImageLayer>(this);
 					NewLayer->mask = nullptr;
 					NewLayer->Flags |= OP::ImageLayerArgs::F_USE_MASK_FROM_BLENDED;
 					at = NewLayer;
@@ -426,7 +426,7 @@ namespace mu
 					&&
 					TypedBlend->SourceChannels[0]==3)
 				{
-					Ptr<ASTOpImageLayer> nop = mu::Clone<ASTOpImageLayer>(this);
+					Ptr<ASTOpImageLayer> nop = UE::Mutable::Private::Clone<ASTOpImageLayer>(this);
 					nop->Flags = Flags | OP::ImageLayerArgs::FLAGS::F_BLENDED_RGB_FROM_ALPHA;
 					nop->blend = TypedBlend->Sources[0].child();
 					auto Desc = nop->blend->GetImageDesc(true);
@@ -500,7 +500,7 @@ namespace mu
 					cropBase->Size[0] = maskUsage.size[0];
 					cropBase->Size[1] = maskUsage.size[1];
 
-					Ptr<ASTOpImageLayer> newLayer = mu::Clone<ASTOpImageLayer>(this);
+					Ptr<ASTOpImageLayer> newLayer = UE::Mutable::Private::Clone<ASTOpImageLayer>(this);
 					newLayer->base = cropBase;
 					newLayer->blend = cropBlended;
 					newLayer->mask = cropMask;
@@ -571,11 +571,11 @@ namespace mu
 						bMaskIsCompatibleSwitch = MaskSwitch->IsCompatibleWith(BlendSwitch);
 					}
 
-					Ptr<ASTOpSwitch> NewSwitch = mu::Clone<ASTOpSwitch>(BlendSwitch);
+					Ptr<ASTOpSwitch> NewSwitch = UE::Mutable::Private::Clone<ASTOpSwitch>(BlendSwitch);
 
 					if (NewSwitch->Default)
 					{
-						Ptr<ASTOpImageLayer> defOp = mu::Clone<ASTOpImageLayer>(this);
+						Ptr<ASTOpImageLayer> defOp = UE::Mutable::Private::Clone<ASTOpImageLayer>(this);
 						defOp->blend = BlendSwitch->Default.child();
 						if (bMaskIsCompatibleSwitch)
 						{
@@ -588,7 +588,7 @@ namespace mu
 					{
 						if (NewSwitch->Cases[v].Branch)
 						{
-							Ptr<ASTOpImageLayer> BranchOp = mu::Clone<ASTOpImageLayer>(this);
+							Ptr<ASTOpImageLayer> BranchOp = UE::Mutable::Private::Clone<ASTOpImageLayer>(this);
 							BranchOp->blend = BlendSwitch->Cases[v].Branch.child();
 							if (bMaskIsCompatibleSwitch)
 							{
@@ -638,11 +638,11 @@ namespace mu
 
 				if (bWorthSinking)
 				{
-					Ptr<ASTOpSwitch> NewSwitch = mu::Clone<ASTOpSwitch>(MaskSwitch);
+					Ptr<ASTOpSwitch> NewSwitch = UE::Mutable::Private::Clone<ASTOpSwitch>(MaskSwitch);
 
 					if (NewSwitch->Default)
 					{
-						Ptr<ASTOpImageLayer> defOp = mu::Clone<ASTOpImageLayer>(this);
+						Ptr<ASTOpImageLayer> defOp = UE::Mutable::Private::Clone<ASTOpImageLayer>(this);
 						defOp->mask = MaskSwitch->Default.child();
 						NewSwitch->Default = defOp;
 					}
@@ -651,7 +651,7 @@ namespace mu
 					{
 						if (NewSwitch->Cases[v].Branch)
 						{
-							Ptr<ASTOpImageLayer> BranchOp = mu::Clone<ASTOpImageLayer>(this);
+							Ptr<ASTOpImageLayer> BranchOp = UE::Mutable::Private::Clone<ASTOpImageLayer>(this);
 							BranchOp->mask = MaskSwitch->Cases[v].Branch.child();
 							NewSwitch->Cases[v].Branch = BranchOp;
 						}
diff --git a/Engine/Plugins/Mutable/Source/MutableTools/Private/MuT/ASTOpImageLayerColor.cpp b/Engine/Plugins/Mutable/Source/MutableTools/Private/MuT/ASTOpImageLayerColor.cpp
--- a/Engine/Plugins/Mutable/Source/MutableTools/Private/MuT/ASTOpImageLayerColor.cpp
+++ b/Engine/Plugins/Mutable/Source/MutableTools/Private/MuT/ASTOpImageLayerColor.cpp
@@ -14,7 +14,7 @@
 #include "MuR/Types.h"
 
 
-namespace mu
+namespace UE::Mutable::Private
 {
 
 	ASTOpImageLayerColor::ASTOpImageLayerColor()
@@ -63,7 +63,7 @@ namespace mu
 
 
 	//-------------------------------------------------------------------------------------------------
-	mu::Ptr<ASTOp> ASTOpImageLayerColor::Clone(MapChildFuncRef mapChild) const
+	UE::Mutable::Private::Ptr<ASTOp> ASTOpImageLayerColor::Clone(MapChildFuncRef mapChild) const
 	{
 		Ptr<ASTOpImageLayerColor> n = new ASTOpImageLayerColor();
 		n->base = mapChild(base.child());
@@ -162,7 +162,7 @@ namespace mu
 
 
 	//-------------------------------------------------------------------------------------------------
-	mu::Ptr<ImageSizeExpression> ASTOpImageLayerColor::GetImageSizeExpression() const
+	UE::Mutable::Private::Ptr<ImageSizeExpression> ASTOpImageLayerColor::GetImageSizeExpression() const
 	{
 		if (base)
 		{
@@ -193,7 +193,7 @@ namespace mu
 				else if (FMath::IsNearlyEqual(colour[0], 1, UE_SMALL_NUMBER))
 				{
 					// If the mask is white, we can remove it
-					Ptr<ASTOpImageLayerColor> NewOp = mu::Clone<ASTOpImageLayerColor>(this);
+					Ptr<ASTOpImageLayerColor> NewOp = UE::Mutable::Private::Clone<ASTOpImageLayerColor>(this);
 					NewOp->mask = nullptr;
 					at = NewOp;
 				}
@@ -306,7 +306,7 @@ namespace mu
 
 				if (bAreAllSameAlpha)
 				{
-					Ptr<ASTOpImageLayerColor> NewOp = mu::Clone<ASTOpImageLayerColor>(this);
+					Ptr<ASTOpImageLayerColor> NewOp = UE::Mutable::Private::Clone<ASTOpImageLayerColor>(this);
 					NewOp->Flags |= OP::ImageLayerArgs::FLAGS::F_BASE_RGB_FROM_ALPHA;
 					NewOp->base = Source;
 					at = NewOp;
@@ -368,11 +368,11 @@ namespace mu
 			}
 
 			// Move the layer operation down base paths
-			Ptr<ASTOpSwitch> nop = mu::Clone<ASTOpSwitch>(baseSwitch);
+			Ptr<ASTOpSwitch> nop = UE::Mutable::Private::Clone<ASTOpSwitch>(baseSwitch);
 
 			if (nop->Default)
 			{
-				Ptr<ASTOpImageLayerColor> defOp = mu::Clone<ASTOpImageLayerColor>(this);
+				Ptr<ASTOpImageLayerColor> defOp = UE::Mutable::Private::Clone<ASTOpImageLayerColor>(this);
 				defOp->base = nop->Default.child();
 				nop->Default = defOp;
 			}
@@ -381,7 +381,7 @@ namespace mu
 			{
 				if (nop->Cases[v].Branch)
 				{
-					Ptr<ASTOpImageLayerColor> bOp = mu::Clone<ASTOpImageLayerColor>(this);
+					Ptr<ASTOpImageLayerColor> bOp = UE::Mutable::Private::Clone<ASTOpImageLayerColor>(this);
 					bOp->base = nop->Cases[v].Branch.child();
 					nop->Cases[v].Branch = bOp;
 				}
@@ -400,10 +400,10 @@ namespace mu
 				break;
 			}
 
-			Ptr<ASTOpImageDisplace> NewDisplace = mu::Clone<ASTOpImageDisplace>(baseAt);
+			Ptr<ASTOpImageDisplace> NewDisplace = UE::Mutable::Private::Clone<ASTOpImageDisplace>(baseAt);
 
 			Ptr<ASTOp> SourceOp = NewDisplace->Source.child();
-			Ptr<ASTOpImageLayerColor> NewSource = mu::Clone<ASTOpImageLayerColor>(this);
+			Ptr<ASTOpImageLayerColor> NewSource = UE::Mutable::Private::Clone<ASTOpImageLayerColor>(this);
 			NewSource->base = SourceOp;
 			NewDisplace->Source = NewSource;
 
@@ -419,10 +419,10 @@ namespace mu
 				break;
 			}
 
-			Ptr<ASTOpImageRasterMesh> NewRaster = mu::Clone<ASTOpImageRasterMesh>(baseAt);
+			Ptr<ASTOpImageRasterMesh> NewRaster = UE::Mutable::Private::Clone<ASTOpImageRasterMesh>(baseAt);
 
 			Ptr<ASTOp> sourceOp = NewRaster->image.child();
-			Ptr<ASTOpImageLayerColor> NewSource = mu::Clone<ASTOpImageLayerColor>(this);
+			Ptr<ASTOpImageLayerColor> NewSource = UE::Mutable::Private::Clone<ASTOpImageLayerColor>(this);
 			NewSource->base = sourceOp;
 			NewRaster->image = NewSource;
 
diff --git a/Engine/Plugins/Mutable/Source/MutableTools/Private/MuT/ASTOpImageMakeGrowMap.cpp b/Engine/Plugins/Mutable/Source/MutableTools/Private/MuT/ASTOpImageMakeGrowMap.cpp
--- a/Engine/Plugins/Mutable/Source/MutableTools/Private/MuT/ASTOpImageMakeGrowMap.cpp
+++ b/Engine/Plugins/Mutable/Source/MutableTools/Private/MuT/ASTOpImageMakeGrowMap.cpp
@@ -18,7 +18,7 @@
 #include "MuT/ASTOpSwitch.h"
 
 
-namespace mu
+namespace UE::Mutable::Private
 {
 
 	ASTOpImageMakeGrowMap::ASTOpImageMakeGrowMap()
@@ -54,9 +54,9 @@ namespace mu
 	}
 
 
-	mu::Ptr<ASTOp> ASTOpImageMakeGrowMap::Clone(MapChildFuncRef mapChild) const
+	UE::Mutable::Private::Ptr<ASTOp> ASTOpImageMakeGrowMap::Clone(MapChildFuncRef mapChild) const
 	{
-		mu::Ptr<ASTOpImageMakeGrowMap> n = new ASTOpImageMakeGrowMap();
+		UE::Mutable::Private::Ptr<ASTOpImageMakeGrowMap> n = new ASTOpImageMakeGrowMap();
 		n->Mask = mapChild(Mask.child());
 		n->Border = Border;
 		return n;
@@ -89,9 +89,9 @@ namespace mu
 	}
 
 
-	mu::Ptr<ASTOp> ASTOpImageMakeGrowMap::OptimiseSink(const FModelOptimizationOptions& options, FOptimizeSinkContext& context) const
+	UE::Mutable::Private::Ptr<ASTOp> ASTOpImageMakeGrowMap::OptimiseSink(const FModelOptimizationOptions& options, FOptimizeSinkContext& context) const
 	{
-		mu::Ptr<ASTOp> at;
+		UE::Mutable::Private::Ptr<ASTOp> at;
 
 		switch (Mask.child()->GetOpType())
 		{
@@ -99,13 +99,13 @@ namespace mu
 		case EOpType::IM_CONDITIONAL:
 		{
 			// We move the format down the two paths
-			Ptr<ASTOpConditional> nop = mu::Clone<ASTOpConditional>(Mask.child());
+			Ptr<ASTOpConditional> nop = UE::Mutable::Private::Clone<ASTOpConditional>(Mask.child());
 
-			Ptr<ASTOpImageMakeGrowMap> aOp = mu::Clone<ASTOpImageMakeGrowMap>(this);
+			Ptr<ASTOpImageMakeGrowMap> aOp = UE::Mutable::Private::Clone<ASTOpImageMakeGrowMap>(this);
 			aOp->Mask = nop->yes.child();
 			nop->yes = aOp;
 
-			Ptr<ASTOpImageMakeGrowMap> bOp = mu::Clone<ASTOpImageMakeGrowMap>(this);
+			Ptr<ASTOpImageMakeGrowMap> bOp = UE::Mutable::Private::Clone<ASTOpImageMakeGrowMap>(this);
 			bOp->Mask = nop->no.child();
 			nop->no = bOp;
 
@@ -116,11 +116,11 @@ namespace mu
 		case EOpType::IM_SWITCH:
 		{
 			// Move the format down all the paths
-			Ptr<ASTOpSwitch> nop = mu::Clone<ASTOpSwitch>(Mask.child());
+			Ptr<ASTOpSwitch> nop = UE::Mutable::Private::Clone<ASTOpSwitch>(Mask.child());
 
 			if (nop->Default)
 			{
-				Ptr<ASTOpImageMakeGrowMap> defOp = mu::Clone<ASTOpImageMakeGrowMap>(this);
+				Ptr<ASTOpImageMakeGrowMap> defOp = UE::Mutable::Private::Clone<ASTOpImageMakeGrowMap>(this);
 				defOp->Mask = nop->Default.child();
 				nop->Default = defOp;
 			}
@@ -130,7 +130,7 @@ namespace mu
 			{
 				if (nop->Cases[v].Branch)
 				{
-					Ptr<ASTOpImageMakeGrowMap> bOp = mu::Clone<ASTOpImageMakeGrowMap>(this);
+					Ptr<ASTOpImageMakeGrowMap> bOp = UE::Mutable::Private::Clone<ASTOpImageMakeGrowMap>(this);
 					bOp->Mask = nop->Cases[v].Branch.child();
 					nop->Cases[v].Branch = bOp;
 				}
@@ -203,9 +203,9 @@ namespace mu
 	}
 
 
-	mu::Ptr<ImageSizeExpression> ASTOpImageMakeGrowMap::GetImageSizeExpression() const
+	UE::Mutable::Private::Ptr<ImageSizeExpression> ASTOpImageMakeGrowMap::GetImageSizeExpression() const
 	{
-		mu::Ptr<ImageSizeExpression> pRes;
+		UE::Mutable::Private::Ptr<ImageSizeExpression> pRes;
 
 		if (Mask.child())
 		{
diff --git a/Engine/Plugins/Mutable/Source/MutableTools/Private/MuT/ASTOpImageMipmap.cpp b/Engine/Plugins/Mutable/Source/MutableTools/Private/MuT/ASTOpImageMipmap.cpp
--- a/Engine/Plugins/Mutable/Source/MutableTools/Private/MuT/ASTOpImageMipmap.cpp
+++ b/Engine/Plugins/Mutable/Source/MutableTools/Private/MuT/ASTOpImageMipmap.cpp
@@ -20,7 +20,7 @@
 #include "MuT/ASTOpSwitch.h"
 
 
-namespace mu
+namespace UE::Mutable::Private
 {
 
 	ASTOpImageMipmap::ASTOpImageMipmap()
@@ -61,9 +61,9 @@ namespace mu
 	}
 
 
-	mu::Ptr<ASTOp> ASTOpImageMipmap::Clone(MapChildFuncRef mapChild) const
+	UE::Mutable::Private::Ptr<ASTOp> ASTOpImageMipmap::Clone(MapChildFuncRef mapChild) const
 	{
-		mu::Ptr<ASTOpImageMipmap> n = new ASTOpImageMipmap();
+		UE::Mutable::Private::Ptr<ASTOpImageMipmap> n = new ASTOpImageMipmap();
 		n->Source = mapChild(Source.child());
 		n->Levels = Levels;
 		n->BlockLevels = BlockLevels;
@@ -106,9 +106,9 @@ namespace mu
 	}
 
 
-	mu::Ptr<ASTOp> ASTOpImageMipmap::OptimiseSemantic(const FModelOptimizationOptions&, int32 Pass) const
+	UE::Mutable::Private::Ptr<ASTOp> ASTOpImageMipmap::OptimiseSemantic(const FModelOptimizationOptions&, int32 Pass) const
 	{
-		mu::Ptr<ASTOp> at;
+		UE::Mutable::Private::Ptr<ASTOp> at;
 
 		// \TODO: This seems to fail with the bandit test model.
 		//if (Source.child())
@@ -127,18 +127,18 @@ namespace mu
 	}
 
 
-	mu::Ptr<ASTOp> ASTOpImageMipmap::OptimiseSink(const FModelOptimizationOptions&, FOptimizeSinkContext& context) const
+	UE::Mutable::Private::Ptr<ASTOp> ASTOpImageMipmap::OptimiseSink(const FModelOptimizationOptions&, FOptimizeSinkContext& context) const
 	{
-		mu::Ptr<ASTOp> at;
+		UE::Mutable::Private::Ptr<ASTOp> at;
 
-		mu::Ptr<ASTOp> sourceAt = Source.child();
+		UE::Mutable::Private::Ptr<ASTOp> sourceAt = Source.child();
 		switch (sourceAt->GetOpType())
 		{
 
 		case EOpType::IM_BLANKLAYOUT:
 		{
 			// Set the mipmap generation on the blank layout operation
-			Ptr<ASTOpImageBlankLayout> NewOp = mu::Clone<ASTOpImageBlankLayout>(sourceAt);
+			Ptr<ASTOpImageBlankLayout> NewOp = UE::Mutable::Private::Clone<ASTOpImageBlankLayout>(sourceAt);
 			NewOp->GenerateMipmaps = 1;    // true
 			NewOp->MipmapCount = Levels;
 			at = NewOp;
@@ -148,7 +148,7 @@ namespace mu
 		case EOpType::IM_PLAINCOLOUR:
 		{
 			// Set the mipmap generation on the plaincolour operation
-			Ptr<ASTOpImagePlainColor> mop = mu::Clone<ASTOpImagePlainColor>(sourceAt);
+			Ptr<ASTOpImagePlainColor> mop = UE::Mutable::Private::Clone<ASTOpImagePlainColor>(sourceAt);
 			mop->LODs = Levels;
 			at = mop;
 			break;
@@ -158,8 +158,8 @@ namespace mu
 		{
 			// Swap unless the mipmap operation builds only the tail or is compressed.
 			// Otherwise, we could fall in a loop of swapping mipmaps and pixelformats.
-			mu::Ptr<ASTOpImageMipmap> mop = mu::Clone<ASTOpImageMipmap>(this);
-			mu::Ptr<ASTOpImagePixelFormat> fop = mu::Clone<ASTOpImagePixelFormat>(sourceAt);
+			UE::Mutable::Private::Ptr<ASTOpImageMipmap> mop = UE::Mutable::Private::Clone<ASTOpImageMipmap>(this);
+			UE::Mutable::Private::Ptr<ASTOpImagePixelFormat> fop = UE::Mutable::Private::Clone<ASTOpImagePixelFormat>(sourceAt);
 			bool isCompressedFormat = IsCompressedFormat(fop->Format);
 			if (isCompressedFormat && !mop->bOnlyTail)
 			{
@@ -255,9 +255,9 @@ namespace mu
 	}
 
 
-	mu::Ptr<ImageSizeExpression> ASTOpImageMipmap::GetImageSizeExpression() const
+	UE::Mutable::Private::Ptr<ImageSizeExpression> ASTOpImageMipmap::GetImageSizeExpression() const
 	{
-		mu::Ptr<ImageSizeExpression> pRes;
+		UE::Mutable::Private::Ptr<ImageSizeExpression> pRes;
 
 		if (Source.child())
 		{
@@ -286,7 +286,7 @@ namespace mu
 	//---------------------------------------------------------------------------------------------
 	//---------------------------------------------------------------------------------------------
 	//---------------------------------------------------------------------------------------------
-	mu::Ptr<ASTOp> Sink_ImageMipmapAST::Apply(const ASTOpImageMipmap* InRoot)
+	UE::Mutable::Private::Ptr<ASTOp> Sink_ImageMipmapAST::Apply(const ASTOpImageMipmap* InRoot)
 	{
 		Root = InRoot;
 		OldToNew.Empty();
@@ -296,7 +296,7 @@ namespace mu
 			return nullptr;
 		}
 
-		mu::Ptr<ASTOp> newSource;
+		UE::Mutable::Private::Ptr<ASTOp> newSource;
 		InitialSource = Root->Source.child();
 
 		// Before sinking, see if it is worth splitting into miptail and mip.
@@ -304,14 +304,14 @@ namespace mu
 		if (!Root->bPreventSplitTail)
 		{
 			// the block mipmaps can be done before composition is done.
-			mu::Ptr<ASTOpImageMipmap> newMip = mu::Clone<ASTOpImageMipmap>(Root);
+			UE::Mutable::Private::Ptr<ASTOpImageMipmap> newMip = UE::Mutable::Private::Clone<ASTOpImageMipmap>(Root);
 			newMip->Levels = Root->BlockLevels;
 			newMip->BlockLevels = Root->BlockLevels;
 			newMip->bOnlyTail = false;
 			newMip->bPreventSplitTail = true;
 
 			// the smallest mipmaps after the composition is done.
-			mu::Ptr<ASTOpImageMipmap> topMipOp = mu::Clone<ASTOpImageMipmap>(Root);
+			UE::Mutable::Private::Ptr<ASTOpImageMipmap> topMipOp = UE::Mutable::Private::Clone<ASTOpImageMipmap>(Root);
 			topMipOp->bOnlyTail = true;
 			topMipOp->bPreventSplitTail = true;
 
@@ -341,7 +341,7 @@ namespace mu
 	//---------------------------------------------------------------------------------------------
 	//---------------------------------------------------------------------------------------------
 	//---------------------------------------------------------------------------------------------
-	mu::Ptr<ASTOp> Sink_ImageMipmapAST::Visit(mu::Ptr<ASTOp> at, const ASTOpImageMipmap* currentMipmapOp)
+	UE::Mutable::Private::Ptr<ASTOp> Sink_ImageMipmapAST::Visit(UE::Mutable::Private::Ptr<ASTOp> at, const ASTOpImageMipmap* currentMipmapOp)
 	{
 		if (!at) return nullptr;
 
@@ -352,14 +352,14 @@ namespace mu
 			return *Cached;
 		}
 
-		mu::Ptr<ASTOp> newAt = at;
+		UE::Mutable::Private::Ptr<ASTOp> newAt = at;
 		switch (at->GetOpType())
 		{
 
 		case EOpType::IM_CONDITIONAL:
 		{
 			// We move the op down the two paths
-			Ptr<ASTOpConditional> newOp = mu::Clone<ASTOpConditional>(at);
+			Ptr<ASTOpConditional> newOp = UE::Mutable::Private::Clone<ASTOpConditional>(at);
 			newOp->yes = Visit(newOp->yes.child(), currentMipmapOp);
 			newOp->no = Visit(newOp->no.child(), currentMipmapOp);
 			newAt = newOp;
@@ -369,7 +369,7 @@ namespace mu
 		case EOpType::IM_SWITCH:
 		{
 			// We move the op down all the paths
-			Ptr<ASTOpSwitch> newOp = mu::Clone<ASTOpSwitch>(at);
+			Ptr<ASTOpSwitch> newOp = UE::Mutable::Private::Clone<ASTOpSwitch>(at);
 			newOp->Default = Visit(newOp->Default.child(), currentMipmapOp);
 			for (ASTOpSwitch::FCase& c : newOp->Cases)
 			{
@@ -389,7 +389,7 @@ namespace mu
 				!typedAt->Mask
 				)
 			{
-				Ptr<ASTOpImageCompose> newOp = mu::Clone<ASTOpImageCompose>(at);
+				Ptr<ASTOpImageCompose> newOp = UE::Mutable::Private::Clone<ASTOpImageCompose>(at);
 
 				Ptr<ASTOp> baseOp = newOp->Base.child();
 				newOp->Base = Visit(baseOp, currentMipmapOp);
@@ -438,12 +438,12 @@ namespace mu
 
 				if (currentMipmapOp->Levels!= PatchBlockLevels || currentMipmapOp->BlockLevels!= PatchBlockLevels)
 				{
-					mu::Ptr<ASTOpImageMipmap> newMip = mu::Clone<ASTOpImageMipmap>(currentMipmapOp);
+					UE::Mutable::Private::Ptr<ASTOpImageMipmap> newMip = UE::Mutable::Private::Clone<ASTOpImageMipmap>(currentMipmapOp);
 					newMip->Levels = PatchBlockLevels;
 					newMip->BlockLevels = PatchBlockLevels;
 					newMip->bOnlyTail = false;
 
-					Ptr<ASTOpImagePatch> newOp = mu::Clone<ASTOpImagePatch>(at);
+					Ptr<ASTOpImagePatch> newOp = UE::Mutable::Private::Clone<ASTOpImagePatch>(at);
 					newOp->base = Visit(newOp->base.child(), newMip.get());
 					newOp->patch = Visit(newOp->patch.child(), newMip.get());
 					newAt = newOp;
@@ -455,7 +455,7 @@ namespace mu
 						)
 					{
 						// We need to add a mipmap on top to finish the mipmapping
-						mu::Ptr<ASTOpImageMipmap> topMipOp = mu::Clone<ASTOpImageMipmap>(currentMipmapOp);
+						UE::Mutable::Private::Ptr<ASTOpImageMipmap> topMipOp = UE::Mutable::Private::Clone<ASTOpImageMipmap>(currentMipmapOp);
 						topMipOp->Source = newOp;
 						topMipOp->bOnlyTail = true;
 						newAt = topMipOp;
@@ -464,7 +464,7 @@ namespace mu
 				else
 				{
 					// The patch supports the same amount of mips that we are currently sinking.
-					Ptr<ASTOpImagePatch> newOp = mu::Clone<ASTOpImagePatch>(at);
+					Ptr<ASTOpImagePatch> newOp = UE::Mutable::Private::Clone<ASTOpImagePatch>(at);
 					newOp->base = Visit(newOp->base.child(), currentMipmapOp);
 					newOp->patch = Visit(newOp->patch.child(), currentMipmapOp);
 					newAt = newOp;
@@ -481,7 +481,7 @@ namespace mu
 		// end on line, replace with mipmap
 		if (at == newAt && at != InitialSource)
 		{
-			mu::Ptr<ASTOpImageMipmap> newOp = mu::Clone<ASTOpImageMipmap>(currentMipmapOp);
+			UE::Mutable::Private::Ptr<ASTOpImageMipmap> newOp = UE::Mutable::Private::Clone<ASTOpImageMipmap>(currentMipmapOp);
 			check(newOp->GetOpType() == EOpType::IM_MIPMAP);
 
 			newOp->Source = at;
diff --git a/Engine/Plugins/Mutable/Source/MutableTools/Private/MuT/ASTOpImageMipmap.h b/Engine/Plugins/Mutable/Source/MutableTools/Private/MuT/ASTOpImageMipmap.h
--- a/Engine/Plugins/Mutable/Source/MutableTools/Private/MuT/ASTOpImageMipmap.h
+++ b/Engine/Plugins/Mutable/Source/MutableTools/Private/MuT/ASTOpImageMipmap.h
@@ -9,7 +9,7 @@
 #include "MuT/AST.h"
 
 
-namespace mu
+namespace UE::Mutable::Private
 {
 	struct FProgram;
 
@@ -47,7 +47,7 @@ namespace mu
 		virtual Ptr<ASTOp> Clone(MapChildFuncRef mapChild) const override;
 		virtual void ForEachChild(const TFunctionRef<void(ASTChild&)>) override;
 		virtual void Link(FProgram& program, FLinkerOptions* Options) override;
-		virtual mu::Ptr<ASTOp> OptimiseSemantic(const FModelOptimizationOptions& options, int32 Pass) const override;
+		virtual UE::Mutable::Private::Ptr<ASTOp> OptimiseSemantic(const FModelOptimizationOptions& options, int32 Pass) const override;
 		virtual Ptr<ASTOp> OptimiseSink(const FModelOptimizationOptions& options, FOptimizeSinkContext& context) const override;
 		virtual FImageDesc GetImageDesc(bool returnBestOption, FGetImageDescContext* context) const override;
 		virtual void GetLayoutBlockSize(int* pBlockX, int* pBlockY) override;
diff --git a/Engine/Plugins/Mutable/Source/MutableTools/Private/MuT/ASTOpImageMultiLayer.cpp b/Engine/Plugins/Mutable/Source/MutableTools/Private/MuT/ASTOpImageMultiLayer.cpp
--- a/Engine/Plugins/Mutable/Source/MutableTools/Private/MuT/ASTOpImageMultiLayer.cpp
+++ b/Engine/Plugins/Mutable/Source/MutableTools/Private/MuT/ASTOpImageMultiLayer.cpp
@@ -9,7 +9,7 @@
 #include "MuR/Types.h"
 
 
-namespace mu
+namespace UE::Mutable::Private
 {
 
 
@@ -61,7 +61,7 @@ namespace mu
 
 
 	//-------------------------------------------------------------------------------------------------
-	mu::Ptr<ASTOp> ASTOpImageMultiLayer::Clone(MapChildFuncRef mapChild) const
+	UE::Mutable::Private::Ptr<ASTOp> ASTOpImageMultiLayer::Clone(MapChildFuncRef mapChild) const
 	{
 		Ptr<ASTOpImageMultiLayer> n = new ASTOpImageMultiLayer();
 		n->base = mapChild(base.child());
@@ -167,7 +167,7 @@ namespace mu
 
 
 	//-------------------------------------------------------------------------------------------------
-	mu::Ptr<ImageSizeExpression> ASTOpImageMultiLayer::GetImageSizeExpression() const
+	UE::Mutable::Private::Ptr<ImageSizeExpression> ASTOpImageMultiLayer::GetImageSizeExpression() const
 	{
 		if (base)
 		{
diff --git a/Engine/Plugins/Mutable/Source/MutableTools/Private/MuT/ASTOpImageNormalComposite.cpp b/Engine/Plugins/Mutable/Source/MutableTools/Private/MuT/ASTOpImageNormalComposite.cpp
--- a/Engine/Plugins/Mutable/Source/MutableTools/Private/MuT/ASTOpImageNormalComposite.cpp
+++ b/Engine/Plugins/Mutable/Source/MutableTools/Private/MuT/ASTOpImageNormalComposite.cpp
@@ -10,7 +10,7 @@
 #include "MuR/Types.h"
 
 
-namespace mu
+namespace UE::Mutable::Private
 {
 
 
@@ -54,7 +54,7 @@ namespace mu
 	}
 
 
-	mu::Ptr<ASTOp> ASTOpImageNormalComposite::Clone(MapChildFuncRef mapChild) const
+	UE::Mutable::Private::Ptr<ASTOp> ASTOpImageNormalComposite::Clone(MapChildFuncRef mapChild) const
 	{
 		Ptr<ASTOpImageNormalComposite> N = new ASTOpImageNormalComposite();
 		N->Base = mapChild(Base.child());
@@ -144,7 +144,7 @@ namespace mu
 		}
 	}
 
-	mu::Ptr<ImageSizeExpression> ASTOpImageNormalComposite::GetImageSizeExpression() const
+	UE::Mutable::Private::Ptr<ImageSizeExpression> ASTOpImageNormalComposite::GetImageSizeExpression() const
 	{
 		if (Base)
 		{
diff --git a/Engine/Plugins/Mutable/Source/MutableTools/Private/MuT/ASTOpImagePatch.cpp b/Engine/Plugins/Mutable/Source/MutableTools/Private/MuT/ASTOpImagePatch.cpp
--- a/Engine/Plugins/Mutable/Source/MutableTools/Private/MuT/ASTOpImagePatch.cpp
+++ b/Engine/Plugins/Mutable/Source/MutableTools/Private/MuT/ASTOpImagePatch.cpp
@@ -8,7 +8,7 @@
 #include "MuR/Types.h"
 
 
-namespace mu
+namespace UE::Mutable::Private
 {
 
 	ASTOpImagePatch::ASTOpImagePatch()
@@ -51,7 +51,7 @@ namespace mu
 
 
 	//-------------------------------------------------------------------------------------------------
-	mu::Ptr<ASTOp> ASTOpImagePatch::Clone(MapChildFuncRef mapChild) const
+	UE::Mutable::Private::Ptr<ASTOp> ASTOpImagePatch::Clone(MapChildFuncRef mapChild) const
 	{
 		Ptr<ASTOpImagePatch> n = new ASTOpImagePatch();
 		n->base = mapChild(base.child());
@@ -131,7 +131,7 @@ namespace mu
 
 
 	//-------------------------------------------------------------------------------------------------
-	mu::Ptr<ImageSizeExpression> ASTOpImagePatch::GetImageSizeExpression() const
+	UE::Mutable::Private::Ptr<ImageSizeExpression> ASTOpImagePatch::GetImageSizeExpression() const
 	{
 		if (base)
 		{
diff --git a/Engine/Plugins/Mutable/Source/MutableTools/Private/MuT/ASTOpImagePixelFormat.cpp b/Engine/Plugins/Mutable/Source/MutableTools/Private/MuT/ASTOpImagePixelFormat.cpp
--- a/Engine/Plugins/Mutable/Source/MutableTools/Private/MuT/ASTOpImagePixelFormat.cpp
+++ b/Engine/Plugins/Mutable/Source/MutableTools/Private/MuT/ASTOpImagePixelFormat.cpp
@@ -25,7 +25,7 @@
 #include "MuT/ASTOpSwitch.h"
 #include "MuT/CompilerPrivate.h"
 
-namespace mu
+namespace UE::Mutable::Private
 {
 
 	template <class SCALAR> class vec4;
@@ -65,9 +65,9 @@ namespace mu
 	}
 
 
-	mu::Ptr<ASTOp> ASTOpImagePixelFormat::Clone(MapChildFuncRef mapChild) const
+	UE::Mutable::Private::Ptr<ASTOp> ASTOpImagePixelFormat::Clone(MapChildFuncRef mapChild) const
 	{
-		mu::Ptr<ASTOpImagePixelFormat> n = new ASTOpImagePixelFormat();
+		UE::Mutable::Private::Ptr<ASTOpImagePixelFormat> n = new ASTOpImagePixelFormat();
 		n->Source = mapChild(Source.child());
 		n->Format = Format;
 		n->FormatIfAlpha = FormatIfAlpha;
@@ -105,7 +105,7 @@ namespace mu
 
 	Ptr<ASTOp> ASTOpImagePixelFormat::OptimiseSemantic(const FModelOptimizationOptions& options, int32 Pass) const
 	{
-		mu::Ptr<ASTOp> NewOp;
+		UE::Mutable::Private::Ptr<ASTOp> NewOp;
 
 		// Skip this operation if the source op format is already the one we want.
 		if (Source)
@@ -139,7 +139,7 @@ namespace mu
 		{
 			// Keep only the top pixel format
 			const ASTOpImagePixelFormat* typedSource = static_cast<const ASTOpImagePixelFormat*>(sourceAt.get());
-			mu::Ptr<ASTOpImagePixelFormat> formatOp = mu::Clone<ASTOpImagePixelFormat>(this);
+			UE::Mutable::Private::Ptr<ASTOpImagePixelFormat> formatOp = UE::Mutable::Private::Clone<ASTOpImagePixelFormat>(this);
 			formatOp->Source = typedSource->Source.child();
 			at = formatOp;
 			break;
@@ -150,9 +150,9 @@ namespace mu
 			// This op doesn't support compressed formats
 			if (!bIsCompressedFormat)
 			{
-				mu::Ptr<ASTOpImageDisplace> newOp = mu::Clone<ASTOpImageDisplace>(sourceAt);
+				UE::Mutable::Private::Ptr<ASTOpImageDisplace> newOp = UE::Mutable::Private::Clone<ASTOpImageDisplace>(sourceAt);
 
-				mu::Ptr<ASTOpImagePixelFormat> fop = mu::Clone<ASTOpImagePixelFormat>(this);
+				UE::Mutable::Private::Ptr<ASTOpImagePixelFormat> fop = UE::Mutable::Private::Clone<ASTOpImagePixelFormat>(this);
 				fop->Source = newOp->Source.child();
 				newOp->Source = fop;
 
@@ -166,9 +166,9 @@ namespace mu
 			// This op doesn't support compressed formats
 			if (!bIsCompressedFormat)
 			{
-				mu::Ptr<ASTOpImageInvert> newOp = mu::Clone<ASTOpImageInvert>(sourceAt);
+				UE::Mutable::Private::Ptr<ASTOpImageInvert> newOp = UE::Mutable::Private::Clone<ASTOpImageInvert>(sourceAt);
 
-				mu::Ptr<ASTOpImagePixelFormat> fop = mu::Clone<ASTOpImagePixelFormat>(this);
+				UE::Mutable::Private::Ptr<ASTOpImagePixelFormat> fop = UE::Mutable::Private::Clone<ASTOpImagePixelFormat>(this);
 				fop->Source = newOp->Base.child();
 				newOp->Base = fop;
 
@@ -184,9 +184,9 @@ namespace mu
 				&&
 				static_cast<const ASTOpImageRasterMesh*>(sourceAt.get())->image)
 			{
-				Ptr<ASTOpImageRasterMesh> newOp = mu::Clone<ASTOpImageRasterMesh>(sourceAt);
+				Ptr<ASTOpImageRasterMesh> newOp = UE::Mutable::Private::Clone<ASTOpImageRasterMesh>(sourceAt);
 
-				mu::Ptr<ASTOpImagePixelFormat> fop = mu::Clone<ASTOpImagePixelFormat>(this);
+				UE::Mutable::Private::Ptr<ASTOpImagePixelFormat> fop = UE::Mutable::Private::Clone<ASTOpImagePixelFormat>(this);
 				fop->Source = newOp->image.child();
 				newOp->image = fop;
 
@@ -198,7 +198,7 @@ namespace mu
 		case EOpType::IM_BLANKLAYOUT:
 		{
 			// Just make sure the layout format is the right one and forget the op
-			Ptr<ASTOpImageBlankLayout> NewOp = mu::Clone<ASTOpImageBlankLayout>(sourceAt);
+			Ptr<ASTOpImageBlankLayout> NewOp = UE::Mutable::Private::Clone<ASTOpImageBlankLayout>(sourceAt);
 
 			EImageFormat layoutFormat = NewOp->Format;
 			if (FormatIfAlpha != EImageFormat::None
@@ -216,7 +216,7 @@ namespace mu
 		case EOpType::IM_PLAINCOLOUR:
 		{
 			// Just make sure the format is the right one and forget the op
-			Ptr<ASTOpImagePlainColor> NewOp = mu::Clone<ASTOpImagePlainColor>(sourceAt);
+			Ptr<ASTOpImagePlainColor> NewOp = UE::Mutable::Private::Clone<ASTOpImagePlainColor>(sourceAt);
 
 			EImageFormat LayoutFormat = NewOp->Format;
 			if (FormatIfAlpha != EImageFormat::None
@@ -314,9 +314,9 @@ namespace mu
 	}
 
 
-	mu::Ptr<ImageSizeExpression> ASTOpImagePixelFormat::GetImageSizeExpression() const
+	UE::Mutable::Private::Ptr<ImageSizeExpression> ASTOpImagePixelFormat::GetImageSizeExpression() const
 	{
-		mu::Ptr<ImageSizeExpression> pRes;
+		UE::Mutable::Private::Ptr<ImageSizeExpression> pRes;
 
 		if (Source.child())
 		{
@@ -345,7 +345,7 @@ namespace mu
 	//---------------------------------------------------------------------------------------------
 	//---------------------------------------------------------------------------------------------
 	//---------------------------------------------------------------------------------------------
-	mu::Ptr<ASTOp> Sink_ImagePixelFormatAST::Apply(const ASTOpImagePixelFormat* root)
+	UE::Mutable::Private::Ptr<ASTOp> Sink_ImagePixelFormatAST::Apply(const ASTOpImagePixelFormat* root)
 	{
 		Root = root;
 		OldToNew.Reset();
@@ -353,7 +353,7 @@ namespace mu
 		check(root->GetOpType() == EOpType::IM_PIXELFORMAT);
 
 		InitialSource = Root->Source.child();
-		mu::Ptr<ASTOp> newSource = Visit(InitialSource, Root);
+		UE::Mutable::Private::Ptr<ASTOp> newSource = Visit(InitialSource, Root);
 
 		Root = nullptr;
 
@@ -368,7 +368,7 @@ namespace mu
 
 
 	//---------------------------------------------------------------------------------------------
-	mu::Ptr<ASTOp> Sink_ImagePixelFormatAST::Visit(mu::Ptr<ASTOp> at, const ASTOpImagePixelFormat* currentFormatOp)
+	UE::Mutable::Private::Ptr<ASTOp> Sink_ImagePixelFormatAST::Visit(UE::Mutable::Private::Ptr<ASTOp> at, const ASTOpImagePixelFormat* currentFormatOp)
 	{
 		if (!at) return nullptr;
 
@@ -383,14 +383,14 @@ namespace mu
 			return *Cached;
 		}
 
-		mu::Ptr<ASTOp> newAt = at;
+		UE::Mutable::Private::Ptr<ASTOp> newAt = at;
 		switch (at->GetOpType())
 		{
 
 		case EOpType::IM_CONDITIONAL:
 		{
 			// We move the op down the two paths
-			Ptr<ASTOpConditional> newOp = mu::Clone<ASTOpConditional>(at);
+			Ptr<ASTOpConditional> newOp = UE::Mutable::Private::Clone<ASTOpConditional>(at);
 			newOp->yes = Visit(newOp->yes.child(), currentFormatOp);
 			newOp->no = Visit(newOp->no.child(), currentFormatOp);
 			newAt = newOp;
@@ -400,7 +400,7 @@ namespace mu
 		case EOpType::IM_SWITCH:
 		{
 			// We move the op down all the paths
-			Ptr<ASTOpSwitch> newOp = mu::Clone<ASTOpSwitch>(at);
+			Ptr<ASTOpSwitch> newOp = UE::Mutable::Private::Clone<ASTOpSwitch>(at);
 			newOp->Default = Visit(newOp->Default.child(), currentFormatOp);
 			for (ASTOpSwitch::FCase& c : newOp->Cases)
 			{
@@ -417,7 +417,7 @@ namespace mu
 				// Since blocks can be resized at runtime anyway, push the format down and rely on reformatting on the fly if necessary.
 
 				// We move the format down the two paths
-				Ptr<ASTOpImageCompose> newOp = mu::Clone<ASTOpImageCompose>(at);
+				Ptr<ASTOpImageCompose> newOp = UE::Mutable::Private::Clone<ASTOpImageCompose>(at);
 
 				// TODO: We have to make sure we don't end up with two different formats if
 				// there is an formatIfAlpha
@@ -439,7 +439,7 @@ namespace mu
 			if (isBlockFormat)
 			{
 				// We move the format down the two paths
-				Ptr<ASTOpImagePatch> newOp = mu::Clone<ASTOpImagePatch>(at);
+				Ptr<ASTOpImagePatch> newOp = UE::Mutable::Private::Clone<ASTOpImagePatch>(at);
 
 				newOp->base = Visit(newOp->base.child(), currentFormatOp);
 				newOp->patch = Visit(newOp->patch.child(), currentFormatOp);
@@ -458,7 +458,7 @@ namespace mu
 			// generate the tail. To avoid optimization loop.
 			if (!bIsCompressedFormat || typedSource->bOnlyTail)
 			{
-				Ptr<ASTOpImageMipmap> newOp = mu::Clone<ASTOpImageMipmap>(typedSource);
+				Ptr<ASTOpImageMipmap> newOp = UE::Mutable::Private::Clone<ASTOpImageMipmap>(typedSource);
 
 				Ptr<ASTOp> baseOp = newOp->Source.child();
 				newOp->Source = Visit(baseOp, currentFormatOp);
@@ -474,7 +474,7 @@ namespace mu
 			if (!bIsCompressedFormat)
 			{
 				// Move the format down all the paths
-				Ptr<ASTOpImageInterpolate> newOp = mu::Clone<ASTOpImageInterpolate>(at);
+				Ptr<ASTOpImageInterpolate> newOp = UE::Mutable::Private::Clone<ASTOpImageInterpolate>(at);
 
 				for (int v = 0; v < MUTABLE_OP_MAX_INTERPOLATE_COUNT; ++v)
 				{
@@ -493,7 +493,7 @@ namespace mu
 			if (GetOpToolsDesc(at->GetOpType()).bSupportedBasePixelFormats[(size_t)format])
 			{
 				// We move the format down the two paths
-				Ptr<ASTOpImageLayer> nop = mu::Clone<ASTOpImageLayer>(at);
+				Ptr<ASTOpImageLayer> nop = UE::Mutable::Private::Clone<ASTOpImageLayer>(at);
 
 				nop->base = Visit(nop->base.child(), currentFormatOp);
 				nop->blend = Visit(nop->blend.child(), currentFormatOp);
@@ -508,7 +508,7 @@ namespace mu
 			if (GetOpToolsDesc(at->GetOpType()).bSupportedBasePixelFormats[(size_t)format])
 			{
 				// We move the format down the base
-				Ptr<ASTOpImageLayerColor> nop = mu::Clone<ASTOpImageLayerColor>(at);
+				Ptr<ASTOpImageLayerColor> nop = UE::Mutable::Private::Clone<ASTOpImageLayerColor>(at);
 
 				nop->base = Visit(nop->base.child(), currentFormatOp);
 
@@ -525,7 +525,7 @@ namespace mu
 		// end on tree branch, replace with format
 		if (at == newAt && at != InitialSource)
 		{
-			mu::Ptr<ASTOpImagePixelFormat> newOp = mu::Clone<ASTOpImagePixelFormat>(currentFormatOp);
+			UE::Mutable::Private::Ptr<ASTOpImagePixelFormat> newOp = UE::Mutable::Private::Clone<ASTOpImagePixelFormat>(currentFormatOp);
 			check(newOp->GetOpType() == EOpType::IM_PIXELFORMAT);
 
 			newOp->Source = at;
diff --git a/Engine/Plugins/Mutable/Source/MutableTools/Private/MuT/ASTOpImageRasterMesh.cpp b/Engine/Plugins/Mutable/Source/MutableTools/Private/MuT/ASTOpImageRasterMesh.cpp
--- a/Engine/Plugins/Mutable/Source/MutableTools/Private/MuT/ASTOpImageRasterMesh.cpp
+++ b/Engine/Plugins/Mutable/Source/MutableTools/Private/MuT/ASTOpImageRasterMesh.cpp
@@ -16,7 +16,7 @@
 #include "HAL/PlatformMath.h"
 
 
-namespace mu
+namespace UE::Mutable::Private
 {
 
 
@@ -94,7 +94,7 @@ namespace mu
 
 
 	//-------------------------------------------------------------------------------------------------
-	mu::Ptr<ASTOp> ASTOpImageRasterMesh::Clone(MapChildFuncRef mapChild) const
+	UE::Mutable::Private::Ptr<ASTOp> ASTOpImageRasterMesh::Clone(MapChildFuncRef mapChild) const
 	{
 		Ptr<ASTOpImageRasterMesh> n = new ASTOpImageRasterMesh();
 		n->mesh = mapChild(mesh.child());
@@ -254,7 +254,7 @@ namespace mu
 			{
 				// We remove the project from the raster children
 				const ASTOpMeshProject* MeshProjectOp = static_cast<const ASTOpMeshProject*>(sourceAt.get());
-				Ptr<ASTOpImageRasterMesh> NewOp = mu::Clone<ASTOpImageRasterMesh>(this);
+				Ptr<ASTOpImageRasterMesh> NewOp = UE::Mutable::Private::Clone<ASTOpImageRasterMesh>(this);
 				NewOp->mesh = MeshProjectOp->Mesh.child();
 				at = NewOp;
 			}
@@ -265,7 +265,7 @@ namespace mu
 		{
 			// TODO: should be sink only if no imageAt?
 			const ASTOpMeshMorph* typedSource = static_cast<const ASTOpMeshMorph*>(sourceAt.get());
-			Ptr<ASTOpImageRasterMesh> rasterOp = mu::Clone<ASTOpImageRasterMesh>(this);
+			Ptr<ASTOpImageRasterMesh> rasterOp = UE::Mutable::Private::Clone<ASTOpImageRasterMesh>(this);
 			rasterOp->mesh = typedSource->Base.child();
 			at = rasterOp;
 			break;
@@ -275,7 +275,7 @@ namespace mu
 		{
 			// Ignore metadata
 			const ASTOpMeshAddMetadata* typedSource = static_cast<const ASTOpMeshAddMetadata*>(sourceAt.get());
-			Ptr<ASTOpImageRasterMesh> rasterOp = mu::Clone<ASTOpImageRasterMesh>(this);
+			Ptr<ASTOpImageRasterMesh> rasterOp = UE::Mutable::Private::Clone<ASTOpImageRasterMesh>(this);
 			rasterOp->mesh = typedSource->Source.child();
 			at = rasterOp;
 			break;
@@ -283,14 +283,14 @@ namespace mu
 
 		case EOpType::ME_CONDITIONAL:
 		{
-			auto nop = mu::Clone<ASTOpConditional>(sourceAt.get());
+			auto nop = UE::Mutable::Private::Clone<ASTOpConditional>(sourceAt.get());
 			nop->type = EOpType::IM_CONDITIONAL;
 
-			Ptr<ASTOpImageRasterMesh> aOp = mu::Clone<ASTOpImageRasterMesh>(this);
+			Ptr<ASTOpImageRasterMesh> aOp = UE::Mutable::Private::Clone<ASTOpImageRasterMesh>(this);
 			aOp->mesh = nop->yes.child();
 			nop->yes = aOp;
 
-			Ptr<ASTOpImageRasterMesh> bOp = mu::Clone<ASTOpImageRasterMesh>(this);
+			Ptr<ASTOpImageRasterMesh> bOp = UE::Mutable::Private::Clone<ASTOpImageRasterMesh>(this);
 			bOp->mesh = nop->no.child();
 			nop->no = bOp;
 
@@ -301,12 +301,12 @@ namespace mu
 		case EOpType::ME_SWITCH:
 		{
 			// Make an image for every path
-			auto nop = mu::Clone<ASTOpSwitch>(sourceAt.get());
+			auto nop = UE::Mutable::Private::Clone<ASTOpSwitch>(sourceAt.get());
 			nop->Type = EOpType::IM_SWITCH;
 
 			if (nop->Default)
 			{
-				Ptr<ASTOpImageRasterMesh> defOp = mu::Clone<ASTOpImageRasterMesh>(this);
+				Ptr<ASTOpImageRasterMesh> defOp = UE::Mutable::Private::Clone<ASTOpImageRasterMesh>(this);
 				defOp->mesh = nop->Default.child();
 				nop->Default = defOp;
 			}
@@ -316,7 +316,7 @@ namespace mu
 			{
 				if (nop->Cases[o].Branch)
 				{
-					Ptr<ASTOpImageRasterMesh> bOp = mu::Clone<ASTOpImageRasterMesh>(this);
+					Ptr<ASTOpImageRasterMesh> bOp = UE::Mutable::Private::Clone<ASTOpImageRasterMesh>(this);
 					bOp->mesh = nop->Cases[o].Branch.child();
 					nop->Cases[o].Branch = bOp;
 				}
@@ -340,14 +340,14 @@ namespace mu
 				// TODO: Implement for image conditionals.
 				//case EOpType::ME_CONDITIONAL:
 				//{
-				//	auto nop = mu::Clone<ASTOpConditional>(sourceAt.get());
+				//	auto nop = UE::Mutable::Private::Clone<ASTOpConditional>(sourceAt.get());
 				//	nop->type = EOpType::IM_CONDITIONAL;
 
-				//	Ptr<ASTOpFixed> aOp = mu::Clone<ASTOpFixed>(this);
+				//	Ptr<ASTOpFixed> aOp = UE::Mutable::Private::Clone<ASTOpFixed>(this);
 				//	aOp->SetChild(aOp->op.Args.ImageRasterMesh.mesh, nop->yes);
 				//	nop->yes = aOp;
 
-				//	Ptr<ASTOpFixed> bOp = mu::Clone<ASTOpFixed>(this);
+				//	Ptr<ASTOpFixed> bOp = UE::Mutable::Private::Clone<ASTOpFixed>(this);
 				//	bOp->SetChild(bOp->op.Args.ImageRasterMesh.mesh, nop->no);
 				//	nop->no = bOp;
 
@@ -360,11 +360,11 @@ namespace mu
 				// TODO: Do this only if the projector is constant?
 
 				// Make a project for every path
-				auto nop = mu::Clone<ASTOpSwitch>(imageAt.get());
+				auto nop = UE::Mutable::Private::Clone<ASTOpSwitch>(imageAt.get());
 
 				if (nop->Default)
 				{
-					Ptr<ASTOpImageRasterMesh> defOp = mu::Clone<ASTOpImageRasterMesh>(this);
+					Ptr<ASTOpImageRasterMesh> defOp = UE::Mutable::Private::Clone<ASTOpImageRasterMesh>(this);
 					defOp->image = nop->Default.child();
 					nop->Default = defOp;
 				}
@@ -374,7 +374,7 @@ namespace mu
 				{
 					if (nop->Cases[o].Branch)
 					{
-						Ptr<ASTOpImageRasterMesh> bOp = mu::Clone<ASTOpImageRasterMesh>(this);
+						Ptr<ASTOpImageRasterMesh> bOp = UE::Mutable::Private::Clone<ASTOpImageRasterMesh>(this);
 						bOp->image = nop->Cases[o].Branch.child();
 						nop->Cases[o].Branch = bOp;
 					}
diff --git a/Engine/Plugins/Mutable/Source/MutableTools/Private/MuT/ASTOpImageResize.cpp b/Engine/Plugins/Mutable/Source/MutableTools/Private/MuT/ASTOpImageResize.cpp
--- a/Engine/Plugins/Mutable/Source/MutableTools/Private/MuT/ASTOpImageResize.cpp
+++ b/Engine/Plugins/Mutable/Source/MutableTools/Private/MuT/ASTOpImageResize.cpp
@@ -24,7 +24,7 @@
 #include "MuT/ASTOpImageInvert.h"
 
 
-namespace mu
+namespace UE::Mutable::Private
 {
 
 	ASTOpImageResize::ASTOpImageResize()
@@ -217,7 +217,7 @@ namespace mu
 		{
 		case EOpType::IM_PLAINCOLOUR:
 		{
-			Ptr<ASTOpImagePlainColor> NewOp = mu::Clone<ASTOpImagePlainColor>(SourceAt.get());
+			Ptr<ASTOpImagePlainColor> NewOp = UE::Mutable::Private::Clone<ASTOpImagePlainColor>(SourceAt.get());
 			NewOp->Size[0] = Size[0];
 			NewOp->Size[1] = Size[1];
 			NewOp->LODs = 1; // TODO
@@ -248,7 +248,7 @@ namespace mu
 			// Keep top resize
 			Ptr<const ASTOpImageResize> SourceOp = static_cast<const ASTOpImageResize*>(SourceAt.get());
 
-			Ptr<ASTOpImageResize> NewOp = mu::Clone<ASTOpImageResize>(this);
+			Ptr<ASTOpImageResize> NewOp = UE::Mutable::Private::Clone<ASTOpImageResize>(this);
 			NewOp->Source = SourceOp->Source.child();
 
 			Result = NewOp;
@@ -258,7 +258,7 @@ namespace mu
 		case EOpType::IM_PLAINCOLOUR:
 		{
 			// Set the size in the children and remove resize
-			Ptr<ASTOpImagePlainColor> SourceOp = mu::Clone<ASTOpImagePlainColor>(SourceAt.get());
+			Ptr<ASTOpImagePlainColor> SourceOp = UE::Mutable::Private::Clone<ASTOpImagePlainColor>(SourceAt.get());
 			SourceOp->Size[0] = Size[0];
 			SourceOp->Size[1] = Size[1];
 			SourceOp->LODs = 1; // TODO
@@ -269,7 +269,7 @@ namespace mu
 		case EOpType::IM_TRANSFORM:
 		{
 			// Set the size in the children and remove resize
-			Ptr<ASTOpImageTransform> SourceOp = mu::Clone<ASTOpImageTransform>(SourceAt.get());
+			Ptr<ASTOpImageTransform> SourceOp = UE::Mutable::Private::Clone<ASTOpImageTransform>(SourceAt.get());
 			SourceOp->SizeX = Size[0];
 			SourceOp->SizeY = Size[1];
 			Result = SourceOp;
@@ -279,13 +279,13 @@ namespace mu
 		case EOpType::IM_CONDITIONAL:
 		{
 			// We move the resize down the two paths
-			Ptr<ASTOpConditional> NewOp = mu::Clone<ASTOpConditional>(SourceAt);
+			Ptr<ASTOpConditional> NewOp = UE::Mutable::Private::Clone<ASTOpConditional>(SourceAt);
 
-			Ptr<ASTOpImageResize> AOp = mu::Clone<ASTOpImageResize>(this);
+			Ptr<ASTOpImageResize> AOp = UE::Mutable::Private::Clone<ASTOpImageResize>(this);
 			AOp->Source = NewOp->yes.child();
 			NewOp->yes = AOp;
 
-			Ptr<ASTOpImageResize> BOp = mu::Clone<ASTOpImageResize>(this);
+			Ptr<ASTOpImageResize> BOp = UE::Mutable::Private::Clone<ASTOpImageResize>(this);
 			BOp->Source = NewOp->no.child();
 			NewOp->no = BOp;
 
@@ -296,11 +296,11 @@ namespace mu
 		case EOpType::IM_SWITCH:
 		{
 			// Move the resize down all the paths
-			Ptr<ASTOpSwitch> NewOp = mu::Clone<ASTOpSwitch>(SourceAt);
+			Ptr<ASTOpSwitch> NewOp = UE::Mutable::Private::Clone<ASTOpSwitch>(SourceAt);
 
 			if (NewOp->Default)
 			{
-				Ptr<ASTOpImageResize> DefOp = mu::Clone<ASTOpImageResize>(this);
+				Ptr<ASTOpImageResize> DefOp = UE::Mutable::Private::Clone<ASTOpImageResize>(this);
 				DefOp->Source = NewOp->Default.child();
 				NewOp->Default = DefOp;
 			}
@@ -309,7 +309,7 @@ namespace mu
 			{
 				if (Case.Branch)
 				{
-					Ptr<ASTOpImageResize> CaseOp = mu::Clone<ASTOpImageResize>(this);
+					Ptr<ASTOpImageResize> CaseOp = UE::Mutable::Private::Clone<ASTOpImageResize>(this);
 					CaseOp->Source = Case.Branch.child();
 					Case.Branch = CaseOp;
 				}
@@ -321,13 +321,13 @@ namespace mu
 
 		case EOpType::IM_SWIZZLE:
 		{
-			Ptr<ASTOpImageSwizzle> NewOp = mu::Clone<ASTOpImageSwizzle>(SourceAt);
+			Ptr<ASTOpImageSwizzle> NewOp = UE::Mutable::Private::Clone<ASTOpImageSwizzle>(SourceAt);
 			for (int32 ChannelIndex = 0; ChannelIndex < MUTABLE_OP_MAX_SWIZZLE_CHANNELS; ++ChannelIndex)
 			{
 				Ptr<ASTOp> OldChannelOp = NewOp->Sources[ChannelIndex].child();
 				if (OldChannelOp)
 				{
-					Ptr<ASTOpImageResize> ChannelResize = mu::Clone<ASTOpImageResize>(this);
+					Ptr<ASTOpImageResize> ChannelResize = UE::Mutable::Private::Clone<ASTOpImageResize>(this);
 					ChannelResize->Source = OldChannelOp;
 					NewOp->Sources[ChannelIndex] = ChannelResize;
 				}
@@ -338,19 +338,19 @@ namespace mu
 
 		case EOpType::IM_COMPOSE:
 		{
-			Ptr<ASTOpImageCompose> NewOp = mu::Clone<ASTOpImageCompose>(SourceAt);
+			Ptr<ASTOpImageCompose> NewOp = UE::Mutable::Private::Clone<ASTOpImageCompose>(SourceAt);
 
-			Ptr<ASTOpImageResize> BaseOp = mu::Clone<ASTOpImageResize>(this);
+			Ptr<ASTOpImageResize> BaseOp = UE::Mutable::Private::Clone<ASTOpImageResize>(this);
 			BaseOp->Source = NewOp->Base.child();
 			NewOp->Base = BaseOp;
 
-			Ptr<ASTOpImageResize> BlockOp = mu::Clone<ASTOpImageResize>(this);
+			Ptr<ASTOpImageResize> BlockOp = UE::Mutable::Private::Clone<ASTOpImageResize>(this);
 			BlockOp->Source = NewOp->BlockImage.child();
 			NewOp->BlockImage = BlockOp;
 
 			if (NewOp->Mask)
 			{
-				Ptr<ASTOpImageResize> MaskOp = mu::Clone<ASTOpImageResize>(this);
+				Ptr<ASTOpImageResize> MaskOp = UE::Mutable::Private::Clone<ASTOpImageResize>(this);
 				MaskOp->Source = NewOp->Mask.child();
 				NewOp->Mask = MaskOp;
 			}
@@ -361,7 +361,7 @@ namespace mu
 
 		case EOpType::IM_RASTERMESH:
 		{
-			Ptr<ASTOpImageRasterMesh> NewOp = mu::Clone<ASTOpImageRasterMesh>(SourceAt);
+			Ptr<ASTOpImageRasterMesh> NewOp = UE::Mutable::Private::Clone<ASTOpImageRasterMesh>(SourceAt);
 
 			//if ( NewOp->op.args.ImageRasterMesh.sizeX != op.args.ImageResize.size[0]
 			//     ||
@@ -372,7 +372,7 @@ namespace mu
 
 				if (NewOp->mask)
 				{
-					Ptr<ASTOpImageResize> MaskOp = mu::Clone<ASTOpImageResize>(this);
+					Ptr<ASTOpImageResize> MaskOp = UE::Mutable::Private::Clone<ASTOpImageResize>(this);
 					MaskOp->Source = NewOp->mask.child();
 					NewOp->mask = MaskOp;
 				}
@@ -381,7 +381,7 @@ namespace mu
 				// This should only be scaled with relative resizes, which come from LOD biases, etc.
 				//if (NewOp->op.args.ImageRasterMesh.image)
 				//{
-				//	Ptr<ASTOpImageResize> mop = mu::Clone<ASTOpImageResize>(this);
+				//	Ptr<ASTOpImageResize> mop = UE::Mutable::Private::Clone<ASTOpImageResize>(this);
 				//	mop->SetChild(mop->op.args.ImageResize.source, NewOp->children[NewOp->op.args.ImageRasterMesh.image].child());
 				//	NewOp->SetChild(NewOp->op.args.ImageRasterMesh.image, mop);
 				//}
@@ -393,14 +393,14 @@ namespace mu
 
 		case EOpType::IM_INTERPOLATE:
 		{
-			Ptr<ASTOpImageInterpolate> NewOp = mu::Clone<ASTOpImageInterpolate>(SourceAt);
+			Ptr<ASTOpImageInterpolate> NewOp = UE::Mutable::Private::Clone<ASTOpImageInterpolate>(SourceAt);
 
 			for (int32 TargetIndex = 0; TargetIndex < MUTABLE_OP_MAX_INTERPOLATE_COUNT; ++TargetIndex)
 			{
 				Ptr<ASTOp> TargetAt = NewOp->Targets[TargetIndex].child();
 				if (TargetAt)
 				{
-					Ptr<ASTOpImageResize> SourceOp = mu::Clone<ASTOpImageResize>(this);
+					Ptr<ASTOpImageResize> SourceOp = UE::Mutable::Private::Clone<ASTOpImageResize>(this);
 					SourceOp->Source = TargetAt;
 					NewOp->Targets[TargetIndex] = SourceOp;
 				}
@@ -412,10 +412,10 @@ namespace mu
 
 		case EOpType::IM_INVERT:
 		{
-			Ptr<ASTOpImageInvert> NewOp = mu::Clone<ASTOpImageInvert>(SourceAt);
+			Ptr<ASTOpImageInvert> NewOp = UE::Mutable::Private::Clone<ASTOpImageInvert>(SourceAt);
 			Ptr<ASTOp> BaseAt = NewOp->Base.child();
 
-			Ptr<ASTOpImageResize> NewBase = mu::Clone<ASTOpImageResize>(this);
+			Ptr<ASTOpImageResize> NewBase = UE::Mutable::Private::Clone<ASTOpImageResize>(this);
 			NewBase->Source = BaseAt;
 
 			NewOp->Base = NewBase;
@@ -433,12 +433,12 @@ namespace mu
 			FImageDesc PixelFormatSourceDesc = SourceTyped->Source->GetImageDesc();
 			if (PixelFormatSourceDesc.m_format != EImageFormat::None
 				&&
-				!mu::IsCompressedFormat(PixelFormatSourceDesc.m_format))
+				!UE::Mutable::Private::IsCompressedFormat(PixelFormatSourceDesc.m_format))
 			{
-				Ptr<ASTOpImagePixelFormat> NewOp = mu::Clone<ASTOpImagePixelFormat>(SourceAt);
+				Ptr<ASTOpImagePixelFormat> NewOp = UE::Mutable::Private::Clone<ASTOpImagePixelFormat>(SourceAt);
 				Ptr<ASTOp> BaseAt = NewOp->Source.child();
 
-				Ptr<ASTOpImageResize> NewBase = mu::Clone<ASTOpImageResize>(this);
+				Ptr<ASTOpImageResize> NewBase = UE::Mutable::Private::Clone<ASTOpImageResize>(this);
 				NewBase->Source = BaseAt;
 
 				NewOp->Source = NewBase;
@@ -450,20 +450,20 @@ namespace mu
 
 		case EOpType::IM_LAYER:
 		{
-			Ptr<ASTOpImageLayer> NewOp = mu::Clone<ASTOpImageLayer>(SourceAt);
+			Ptr<ASTOpImageLayer> NewOp = UE::Mutable::Private::Clone<ASTOpImageLayer>(SourceAt);
 
-			Ptr<ASTOpImageResize> BaseOp = mu::Clone<ASTOpImageResize>(this);
+			Ptr<ASTOpImageResize> BaseOp = UE::Mutable::Private::Clone<ASTOpImageResize>(this);
 			BaseOp->Source = NewOp->base.child();
 			NewOp->base = BaseOp;
 
-			Ptr<ASTOpImageResize> blendOp = mu::Clone<ASTOpImageResize>(this);
+			Ptr<ASTOpImageResize> blendOp = UE::Mutable::Private::Clone<ASTOpImageResize>(this);
 			blendOp->Source = NewOp->blend.child();
 			NewOp->blend = blendOp;
 
 			Ptr<ASTOp> MaskAt = NewOp->mask.child();
 			if (MaskAt)
 			{
-				Ptr<ASTOpImageResize> ResizedMaskOp = mu::Clone<ASTOpImageResize>(this);
+				Ptr<ASTOpImageResize> ResizedMaskOp = UE::Mutable::Private::Clone<ASTOpImageResize>(this);
 				ResizedMaskOp->Source = MaskAt;
 				NewOp->mask = ResizedMaskOp;
 			}
@@ -474,16 +474,16 @@ namespace mu
 
 		case EOpType::IM_LAYERCOLOUR:
 		{
-			Ptr<ASTOpImageLayerColor> NewOp = mu::Clone<ASTOpImageLayerColor>(SourceAt);
+			Ptr<ASTOpImageLayerColor> NewOp = UE::Mutable::Private::Clone<ASTOpImageLayerColor>(SourceAt);
 
-			Ptr<ASTOpImageResize> BaseOp = mu::Clone<ASTOpImageResize>(this);
+			Ptr<ASTOpImageResize> BaseOp = UE::Mutable::Private::Clone<ASTOpImageResize>(this);
 			BaseOp->Source = NewOp->base.child();
 			NewOp->base = BaseOp;
 
 			Ptr<ASTOp> MaskAt = NewOp->mask.child();
 			if (MaskAt)
 			{
-				Ptr<ASTOpImageResize> ResizedMaskOp = mu::Clone<ASTOpImageResize>(this);
+				Ptr<ASTOpImageResize> ResizedMaskOp = UE::Mutable::Private::Clone<ASTOpImageResize>(this);
 				ResizedMaskOp->Source = MaskAt;
 				NewOp->mask = ResizedMaskOp;
 			}
@@ -501,17 +501,17 @@ namespace mu
 			Ptr<ASTOp> OriginalDisplacementMapOp = SourceTyped->DisplacementMap.child();
 			if (OriginalDisplacementMapOp->GetOpType() == EOpType::IM_MAKEGROWMAP)
 			{
-				Ptr<ASTOpImageDisplace> NewOp = mu::Clone<ASTOpImageDisplace>(SourceAt);
+				Ptr<ASTOpImageDisplace> NewOp = UE::Mutable::Private::Clone<ASTOpImageDisplace>(SourceAt);
 
-				Ptr<ASTOpImageResize> BaseOp = mu::Clone<ASTOpImageResize>(this);
+				Ptr<ASTOpImageResize> BaseOp = UE::Mutable::Private::Clone<ASTOpImageResize>(this);
 				BaseOp->Source = NewOp->Source.child();
 				NewOp->Source = BaseOp;
 
 				// Clone the map op and replace its children
-				Ptr<ASTOpImageMakeGrowMap> MapOp = mu::Clone<ASTOpImageMakeGrowMap>(OriginalDisplacementMapOp);
+				Ptr<ASTOpImageMakeGrowMap> MapOp = UE::Mutable::Private::Clone<ASTOpImageMakeGrowMap>(OriginalDisplacementMapOp);
 				NewOp->DisplacementMap = MapOp;
 
-				Ptr<ASTOpImageResize> MapSourceOp = mu::Clone<ASTOpImageResize>(this);
+				Ptr<ASTOpImageResize> MapSourceOp = UE::Mutable::Private::Clone<ASTOpImageResize>(this);
 				MapSourceOp->Source = MapOp->Mask.child();
 				MapOp->Mask = MapSourceOp;
 
diff --git a/Engine/Plugins/Mutable/Source/MutableTools/Private/MuT/ASTOpImageResizeRel.cpp b/Engine/Plugins/Mutable/Source/MutableTools/Private/MuT/ASTOpImageResizeRel.cpp
--- a/Engine/Plugins/Mutable/Source/MutableTools/Private/MuT/ASTOpImageResizeRel.cpp
+++ b/Engine/Plugins/Mutable/Source/MutableTools/Private/MuT/ASTOpImageResizeRel.cpp
@@ -29,7 +29,7 @@
 #include "MuT/ASTOpImageInvert.h"
 
 
-namespace mu
+namespace UE::Mutable::Private
 {
 
 	ASTOpImageResizeRel::ASTOpImageResizeRel()
@@ -238,7 +238,7 @@ namespace mu
 		{
 		case EOpType::IM_RESIZE:
 		{
-			Ptr<ASTOpImageResize> NewOp = mu::Clone<ASTOpImageResize>(SourceAt.get());
+			Ptr<ASTOpImageResize> NewOp = UE::Mutable::Private::Clone<ASTOpImageResize>(SourceAt.get());
 			NewOp->Size[0] = int16(NewOp->Size[0] * Factor[0]);
 			NewOp->Size[1] = int16(NewOp->Size[1] * Factor[1]);
 
@@ -304,7 +304,7 @@ namespace mu
 			case EOpType::IM_CONDITIONAL:
 			{
 				// We move the mask creation down the two paths
-				auto NewOp = mu::Clone<ASTOpConditional>(at);
+				auto NewOp = UE::Mutable::Private::Clone<ASTOpConditional>(at);
 				NewOp->yes = Visit(NewOp->yes.child(), CurrentSinkingOp);
 				NewOp->no = Visit(NewOp->no.child(), CurrentSinkingOp);
 				Result = NewOp;
@@ -314,7 +314,7 @@ namespace mu
 			case EOpType::IM_PIXELFORMAT:
 			{
 				// \todo: only if shrinking?
-				Ptr<ASTOpImagePixelFormat> NewOp = mu::Clone<ASTOpImagePixelFormat>(at);
+				Ptr<ASTOpImagePixelFormat> NewOp = UE::Mutable::Private::Clone<ASTOpImagePixelFormat>(at);
 				NewOp->Source = Visit(NewOp->Source.child(), CurrentSinkingOp);
 				Result = NewOp;
 				break;
@@ -323,7 +323,7 @@ namespace mu
 			case EOpType::IM_SWITCH:
 			{
 				// We move the mask creation down all the paths
-				Ptr<ASTOpSwitch> NewOp = mu::Clone<ASTOpSwitch>(at);
+				Ptr<ASTOpSwitch> NewOp = UE::Mutable::Private::Clone<ASTOpSwitch>(at);
 				NewOp->Default = Visit(NewOp->Default.child(), CurrentSinkingOp);
 				for (ASTOpSwitch::FCase& c : NewOp->Cases)
 				{
@@ -335,7 +335,7 @@ namespace mu
 
 			case EOpType::IM_SWIZZLE:
 			{
-				Ptr<ASTOpImageSwizzle> NewOp = mu::Clone<ASTOpImageSwizzle>(at);
+				Ptr<ASTOpImageSwizzle> NewOp = UE::Mutable::Private::Clone<ASTOpImageSwizzle>(at);
 				for (int32 s = 0; s < MUTABLE_OP_MAX_SWIZZLE_CHANNELS; ++s)
 				{
 					Ptr<ASTOp> channelOp = NewOp->Sources[s].child();
@@ -376,7 +376,7 @@ namespace mu
 
 				if (acceptable)
 				{
-					Ptr<ASTOpImageCompose> NewOp = mu::Clone<ASTOpImageCompose>(at);
+					Ptr<ASTOpImageCompose> NewOp = UE::Mutable::Private::Clone<ASTOpImageCompose>(at);
 
 					Ptr<ASTOp> baseOp = NewOp->Base.child();
 					NewOp->Base = Visit(baseOp, CurrentSinkingOp);
@@ -392,7 +392,7 @@ namespace mu
 
 			case EOpType::IM_PATCH:
 			{
-				Ptr<ASTOpImagePatch> NewOp = mu::Clone<ASTOpImagePatch>(at);
+				Ptr<ASTOpImagePatch> NewOp = UE::Mutable::Private::Clone<ASTOpImagePatch>(at);
 
 				NewOp->base = Visit(NewOp->base.child(), CurrentSinkingOp);
 				NewOp->patch = Visit(NewOp->patch.child(), CurrentSinkingOp);
@@ -408,7 +408,7 @@ namespace mu
 
 			case EOpType::IM_MIPMAP:
 			{
-				Ptr<ASTOpImageMipmap> NewOp = mu::Clone<ASTOpImageMipmap>(at.get());
+				Ptr<ASTOpImageMipmap> NewOp = UE::Mutable::Private::Clone<ASTOpImageMipmap>(at.get());
 				Ptr<ASTOp> baseOp = NewOp->Source.child();
 				NewOp->Source = Visit(baseOp, CurrentSinkingOp);
 				Result = NewOp;
@@ -417,7 +417,7 @@ namespace mu
 
 			case EOpType::IM_INTERPOLATE:
 			{
-				Ptr<ASTOpImageInterpolate> NewOp = mu::Clone<ASTOpImageInterpolate>(at);
+				Ptr<ASTOpImageInterpolate> NewOp = UE::Mutable::Private::Clone<ASTOpImageInterpolate>(at);
 
 				for (int32 v = 0; v < MUTABLE_OP_MAX_INTERPOLATE_COUNT; ++v)
 				{
@@ -432,7 +432,7 @@ namespace mu
 
 			case EOpType::IM_MULTILAYER:
 			{
-				Ptr<ASTOpImageMultiLayer> nop = mu::Clone<ASTOpImageMultiLayer>(at);
+				Ptr<ASTOpImageMultiLayer> nop = UE::Mutable::Private::Clone<ASTOpImageMultiLayer>(at);
 				nop->base = Visit(nop->base.child(), CurrentSinkingOp);
 				nop->mask = Visit(nop->mask.child(), CurrentSinkingOp);
 				nop->blend = Visit(nop->blend.child(), CurrentSinkingOp);
@@ -442,7 +442,7 @@ namespace mu
 
 			case EOpType::IM_LAYER:
 			{
-				Ptr<ASTOpImageLayer> nop = mu::Clone<ASTOpImageLayer>(at);
+				Ptr<ASTOpImageLayer> nop = UE::Mutable::Private::Clone<ASTOpImageLayer>(at);
 
 				Ptr<ASTOp> aOp = nop->base.child();
 				nop->base = Visit(aOp, CurrentSinkingOp);
@@ -459,7 +459,7 @@ namespace mu
 
 			case EOpType::IM_LAYERCOLOUR:
 			{
-				Ptr<ASTOpImageLayerColor> nop = mu::Clone<ASTOpImageLayerColor>(at);
+				Ptr<ASTOpImageLayerColor> nop = UE::Mutable::Private::Clone<ASTOpImageLayerColor>(at);
 
 				Ptr<ASTOp> aOp = nop->base.child();
 				nop->base = Visit(aOp, CurrentSinkingOp);
@@ -473,7 +473,7 @@ namespace mu
 
 			case EOpType::IM_DISPLACE:
 			{
-				Ptr<ASTOpImageDisplace> nop = mu::Clone<ASTOpImageDisplace>(at);
+				Ptr<ASTOpImageDisplace> nop = UE::Mutable::Private::Clone<ASTOpImageDisplace>(at);
 
 				Ptr<ASTOp> sourceOp = nop->Source.child();
 				nop->Source = Visit(sourceOp, CurrentSinkingOp);
@@ -496,7 +496,7 @@ namespace mu
 
 			case EOpType::IM_MAKEGROWMAP:
 			{
-				Ptr<ASTOpImageMakeGrowMap> nop = mu::Clone<ASTOpImageMakeGrowMap>(at);
+				Ptr<ASTOpImageMakeGrowMap> nop = UE::Mutable::Private::Clone<ASTOpImageMakeGrowMap>(at);
 				Ptr<ASTOp> maskOp = nop->Mask.child();
 				nop->Mask = Visit(maskOp, CurrentSinkingOp);
 				Result = nop;
@@ -505,7 +505,7 @@ namespace mu
 
 			case EOpType::IM_INVERT:
 			{
-				Ptr<ASTOpImageInvert> nop = mu::Clone<ASTOpImageInvert>(at);
+				Ptr<ASTOpImageInvert> nop = UE::Mutable::Private::Clone<ASTOpImageInvert>(at);
 				Ptr<ASTOp> maskOp = nop->Base.child();
 				nop->Base = Visit(maskOp, CurrentSinkingOp);
 				Result = nop;
@@ -514,7 +514,7 @@ namespace mu
 
 			case EOpType::IM_SATURATE:
 			{
-				Ptr<ASTOpImageSaturate> NewOp = mu::Clone<ASTOpImageSaturate>(at);
+				Ptr<ASTOpImageSaturate> NewOp = UE::Mutable::Private::Clone<ASTOpImageSaturate>(at);
 				NewOp->Base = Visit(NewOp->Base.child(), CurrentSinkingOp);
 				Result = NewOp;
 				break;
@@ -526,7 +526,7 @@ namespace mu
 				const ASTOpImageTransform* TypedAt = static_cast<const ASTOpImageTransform*>(at.get());
 				if (TypedAt->SizeX == 0 && TypedAt->SizeY == 0)
 				{
-					Ptr<ASTOpImageTransform> NewOp = mu::Clone<ASTOpImageTransform>(at);
+					Ptr<ASTOpImageTransform> NewOp = UE::Mutable::Private::Clone<ASTOpImageTransform>(at);
 					Ptr<ASTOp> MaskOp = NewOp->Base.child();
 					NewOp->Base = Visit(MaskOp, CurrentSinkingOp);
 					Result = NewOp;
@@ -537,7 +537,7 @@ namespace mu
 
 			case EOpType::IM_RASTERMESH:
 			{
-				Ptr<ASTOpImageRasterMesh> nop = mu::Clone<ASTOpImageRasterMesh>(at);
+				Ptr<ASTOpImageRasterMesh> nop = UE::Mutable::Private::Clone<ASTOpImageRasterMesh>(at);
 				Ptr<ASTOp> maskOp = nop->mask.child();
 				nop->mask = Visit(maskOp, CurrentSinkingOp);
 
@@ -560,7 +560,7 @@ namespace mu
 			// end on line, replace with sinking op
 			if (at == Result && at != InitialSource)
 			{
-				Ptr<ASTOpImageResizeRel> NewOp = mu::Clone<ASTOpImageResizeRel>(CurrentSinkingOp);
+				Ptr<ASTOpImageResizeRel> NewOp = UE::Mutable::Private::Clone<ASTOpImageResizeRel>(CurrentSinkingOp);
 				NewOp->Source = at;
 
 				Result = NewOp;
@@ -586,7 +586,7 @@ namespace mu
 
 		case EOpType::IM_BLANKLAYOUT:
 		{
-			Ptr<ASTOpImageBlankLayout> NewOp = mu::Clone<ASTOpImageBlankLayout>(SourceAt);
+			Ptr<ASTOpImageBlankLayout> NewOp = UE::Mutable::Private::Clone<ASTOpImageBlankLayout>(SourceAt);
 			NewOp->BlockSize[0] = uint16(NewOp->BlockSize[0] * Factor[0] + 0.5f);
 			NewOp->BlockSize[1] = uint16(NewOp->BlockSize[1] * Factor[1] + 0.5f);
 			Result = NewOp;
@@ -595,7 +595,7 @@ namespace mu
 
 		case EOpType::IM_PLAINCOLOUR:
 		{
-			Ptr<ASTOpImagePlainColor> NewOp = mu::Clone<ASTOpImagePlainColor>(SourceAt);
+			Ptr<ASTOpImagePlainColor> NewOp = UE::Mutable::Private::Clone<ASTOpImagePlainColor>(SourceAt);
 			NewOp->Size[0] = FMath::CeilToInt(NewOp->Size[0] * Factor[0]);
 			NewOp->Size[1] = FMath::CeilToInt(NewOp->Size[1] * Factor[1]);
 			NewOp->LODs = 1; // TODO
@@ -610,7 +610,7 @@ namespace mu
 			if (typedAt->SizeX != 0 && typedAt->SizeY != 0)
 			{
 				// Set the size in the children and remove resize
-				Ptr<ASTOpImageTransform> sourceOp = mu::Clone<ASTOpImageTransform>(SourceAt.get());
+				Ptr<ASTOpImageTransform> sourceOp = UE::Mutable::Private::Clone<ASTOpImageTransform>(SourceAt.get());
 				sourceOp->SizeX = FMath::CeilToInt32(sourceOp->SizeX * Factor[0]);
 				sourceOp->SizeY = FMath::CeilToInt32(sourceOp->SizeY * Factor[1]);
 				Result = sourceOp;
diff --git a/Engine/Plugins/Mutable/Source/MutableTools/Private/MuT/ASTOpImageSwizzle.cpp b/Engine/Plugins/Mutable/Source/MutableTools/Private/MuT/ASTOpImageSwizzle.cpp
--- a/Engine/Plugins/Mutable/Source/MutableTools/Private/MuT/ASTOpImageSwizzle.cpp
+++ b/Engine/Plugins/Mutable/Source/MutableTools/Private/MuT/ASTOpImageSwizzle.cpp
@@ -29,7 +29,7 @@
 #include "MuT/ASTOpColorSwizzle.h"
 #include "MuT/ASTOpSwitch.h"
 
-namespace mu
+namespace UE::Mutable::Private
 {
 
 	ASTOpImageSwizzle::ASTOpImageSwizzle()
@@ -78,9 +78,9 @@ namespace mu
 	}
 
 
-	mu::Ptr<ASTOp> ASTOpImageSwizzle::Clone(MapChildFuncRef mapChild) const
+	UE::Mutable::Private::Ptr<ASTOp> ASTOpImageSwizzle::Clone(MapChildFuncRef mapChild) const
 	{
-		mu::Ptr<ASTOpImageSwizzle> n = new ASTOpImageSwizzle();
+		UE::Mutable::Private::Ptr<ASTOpImageSwizzle> n = new ASTOpImageSwizzle();
 		for (int32 i = 0; i < MUTABLE_OP_MAX_SWIZZLE_CHANNELS; ++i)
 		{
 			n->Sources[i] = mapChild(Sources[i].child());
@@ -123,7 +123,7 @@ namespace mu
 	}
 
 
-	mu::Ptr<ASTOp> ASTOpImageSwizzle::OptimiseSemantic(const FModelOptimizationOptions&, int32 Pass) const
+	UE::Mutable::Private::Ptr<ASTOp> ASTOpImageSwizzle::OptimiseSemantic(const FModelOptimizationOptions&, int32 Pass) const
 	{
 		Ptr<ASTOpImageSwizzle> sat;
 
@@ -142,7 +142,7 @@ namespace mu
 			{
 				if (!sat)
 				{
-					sat = mu::Clone<ASTOpImageSwizzle>(this);
+					sat = UE::Mutable::Private::Clone<ASTOpImageSwizzle>(this);
 				}
 				const ASTOpImageSwizzle* typedCandidate = static_cast<const ASTOpImageSwizzle*>(candidate.get());
 				int candidateChannel = SourceChannels[c];
@@ -167,7 +167,7 @@ namespace mu
 					{
 						if (!sat)
 						{
-							sat = mu::Clone<ASTOpImageSwizzle>(this);
+							sat = UE::Mutable::Private::Clone<ASTOpImageSwizzle>(this);
 						}
 						sat->Sources[c] = formatSource;
 					}
@@ -203,7 +203,7 @@ namespace mu
 		}
 	}
 
-	mu::Ptr<ASTOp> ASTOpImageSwizzle::OptimiseSink(const FModelOptimizationOptions& options, FOptimizeSinkContext& context) const
+	UE::Mutable::Private::Ptr<ASTOp> ASTOpImageSwizzle::OptimiseSink(const FModelOptimizationOptions& options, FOptimizeSinkContext& context) const
 	{
 		MUTABLE_CPUPROFILER_SCOPE(OptimiseSwizzleAST);
 
@@ -301,19 +301,19 @@ namespace mu
 				{
 					// We can combine the 2 multilayers into the composite blend+lighten mode
 
-					Ptr<ASTOpImageSwizzle> NewBase = mu::Clone<ASTOpImageSwizzle>(this);
+					Ptr<ASTOpImageSwizzle> NewBase = UE::Mutable::Private::Clone<ASTOpImageSwizzle>(this);
 					NewBase->Sources[0] = ColorMultiLayer->base.child();
 					NewBase->Sources[1] = ColorMultiLayer->base.child();
 					NewBase->Sources[2] = ColorMultiLayer->base.child();
 					NewBase->Sources[3] = AlphaMultiLayer->base.child();
 
-					Ptr<ASTOpImageSwizzle> NewBlended = mu::Clone<ASTOpImageSwizzle>(this);
+					Ptr<ASTOpImageSwizzle> NewBlended = UE::Mutable::Private::Clone<ASTOpImageSwizzle>(this);
 					NewBlended->Sources[0] = ColorMultiLayer->blend.child();
 					NewBlended->Sources[1] = ColorMultiLayer->blend.child();
 					NewBlended->Sources[2] = ColorMultiLayer->blend.child();
 					NewBlended->Sources[3] = AlphaMultiLayer->blend.child();
 
-					Ptr<ASTOpImageMultiLayer> NewMultiLayer = mu::Clone<ASTOpImageMultiLayer>(ColorMultiLayer);
+					Ptr<ASTOpImageMultiLayer> NewMultiLayer = UE::Mutable::Private::Clone<ASTOpImageMultiLayer>(ColorMultiLayer);
 					NewMultiLayer->blendTypeAlpha = AlphaMultiLayer->blendType;
 					NewMultiLayer->BlendAlphaSourceChannel = 3;
 					NewMultiLayer->base = NewBase;
@@ -358,19 +358,19 @@ namespace mu
 				{
 					// We can combine the 2 image_layers into the composite blend+lighten mode
 
-					Ptr<ASTOpImageSwizzle> NewBase = mu::Clone<ASTOpImageSwizzle>(this);
+					Ptr<ASTOpImageSwizzle> NewBase = UE::Mutable::Private::Clone<ASTOpImageSwizzle>(this);
 					NewBase->Sources[0] = ColorLayer->base.child();
 					NewBase->Sources[1] = ColorLayer->base.child();
 					NewBase->Sources[2] = Sources[2] ? ColorLayer->base.child() : nullptr;
 					NewBase->Sources[3] = AlphaLayer->base.child();
 
-					Ptr<ASTOpImageSwizzle> NewBlended = mu::Clone<ASTOpImageSwizzle>(this);
+					Ptr<ASTOpImageSwizzle> NewBlended = UE::Mutable::Private::Clone<ASTOpImageSwizzle>(this);
 					NewBlended->Sources[0] = ColorLayer->blend.child();
 					NewBlended->Sources[1] = ColorLayer->blend.child();
 					NewBlended->Sources[2] = Sources[2] ? ColorLayer->blend.child() : nullptr;
 					NewBlended->Sources[3] = AlphaLayer->blend.child();
 
-					Ptr<ASTOpImageLayer> NewLayer = mu::Clone<ASTOpImageLayer>(ColorLayer);
+					Ptr<ASTOpImageLayer> NewLayer = UE::Mutable::Private::Clone<ASTOpImageLayer>(ColorLayer);
 					NewLayer->blendTypeAlpha = AlphaLayer->blendType;
 					NewLayer->BlendAlphaSourceChannel = 3;
 					NewLayer->base = NewBase;
@@ -415,11 +415,11 @@ namespace mu
 				if (bAreAllSwitchesCompatible)
 				{
 					// Move the swizzle down all the paths
-					Ptr<ASTOpSwitch> nop = mu::Clone<ASTOpSwitch>(channelSourceAt);
+					Ptr<ASTOpSwitch> nop = UE::Mutable::Private::Clone<ASTOpSwitch>(channelSourceAt);
 
 					if (nop->Default)
 					{
-						Ptr<ASTOpImageSwizzle> defOp = mu::Clone<ASTOpImageSwizzle>(this);
+						Ptr<ASTOpImageSwizzle> defOp = UE::Mutable::Private::Clone<ASTOpImageSwizzle>(this);
 						for (int c = 0; c < MUTABLE_OP_MAX_SWIZZLE_CHANNELS; ++c)
 						{
 							const ASTOpSwitch* ChannelSwitch = static_cast<const ASTOpSwitch*>(Sources[c].child().get());
@@ -435,7 +435,7 @@ namespace mu
 					{
 						if (nop->Cases[v].Branch)
 						{
-							Ptr<ASTOpImageSwizzle> branchOp = mu::Clone<ASTOpImageSwizzle>(this);
+							Ptr<ASTOpImageSwizzle> branchOp = UE::Mutable::Private::Clone<ASTOpImageSwizzle>(this);
 							for (int c = 0; c < MUTABLE_OP_MAX_SWIZZLE_CHANNELS; ++c)
 							{
 								const ASTOpSwitch* ChannelSwitch = static_cast<const ASTOpSwitch*>(Sources[c].child().get());
@@ -476,9 +476,9 @@ namespace mu
 				if (bAreAllDisplacesCompatible)
 				{
 					// Move the swizzle down all the paths
-					Ptr<ASTOpImageDisplace> NewDisplace = mu::Clone<ASTOpImageDisplace>(FirstDisplace);
+					Ptr<ASTOpImageDisplace> NewDisplace = UE::Mutable::Private::Clone<ASTOpImageDisplace>(FirstDisplace);
 
-					Ptr<ASTOpImageSwizzle> SourceOp = mu::Clone<ASTOpImageSwizzle>(this);
+					Ptr<ASTOpImageSwizzle> SourceOp = UE::Mutable::Private::Clone<ASTOpImageSwizzle>(this);
 					for (int c = 0; c < MUTABLE_OP_MAX_SWIZZLE_CHANNELS; ++c)
 					{
 						const ASTOpImageDisplace* ChannelDisplace = static_cast<const ASTOpImageDisplace*>(Sources[c].child().get());
@@ -536,9 +536,9 @@ namespace mu
 				if (bAreAllRasterMeshesCompatible)
 				{
 					// Move the swizzle down all the paths
-					Ptr<ASTOpImageRasterMesh> NewRaster = mu::Clone<ASTOpImageRasterMesh>(FirstRasterMesh);
+					Ptr<ASTOpImageRasterMesh> NewRaster = UE::Mutable::Private::Clone<ASTOpImageRasterMesh>(FirstRasterMesh);
 
-					Ptr<ASTOpImageSwizzle> NewSwizzle = mu::Clone<ASTOpImageSwizzle>(this);
+					Ptr<ASTOpImageSwizzle> NewSwizzle = UE::Mutable::Private::Clone<ASTOpImageSwizzle>(this);
 					for (int c = 0; c < MUTABLE_OP_MAX_SWIZZLE_CHANNELS; ++c)
 					{
 						const ASTOpImageRasterMesh* ChannelRaster = static_cast<const ASTOpImageRasterMesh*>(Sources[c].child().get());
@@ -606,9 +606,9 @@ namespace mu
 				if (bAreAllTransformsCompatible)
 				{
 					// Move the swizzle down all the paths
-					Ptr<ASTOpImageTransform> NewTransform = mu::Clone<ASTOpImageTransform>(FirstTransform);
+					Ptr<ASTOpImageTransform> NewTransform = UE::Mutable::Private::Clone<ASTOpImageTransform>(FirstTransform);
 
-					Ptr<ASTOpImageSwizzle> NewSwizzle = mu::Clone<ASTOpImageSwizzle>(this);
+					Ptr<ASTOpImageSwizzle> NewSwizzle = UE::Mutable::Private::Clone<ASTOpImageSwizzle>(this);
 					for (int c = 0; c < MUTABLE_OP_MAX_SWIZZLE_CHANNELS; ++c)
 					{
 						const ASTOpImageTransform* ChannelTransform = static_cast<const ASTOpImageTransform*>(Sources[c].child().get());
@@ -652,9 +652,9 @@ namespace mu
 			//	if (bAreAllResizesCompatible)
 			//	{
 			//		// Move the swizzle down all the paths
-			//		Ptr<ASTOpFixed> NewResize = mu::Clone<ASTOpFixed>(FirstResize);
+			//		Ptr<ASTOpFixed> NewResize = UE::Mutable::Private::Clone<ASTOpFixed>(FirstResize);
 
-			//		Ptr<ASTOpImageSwizzle> NewSwizzle = mu::Clone<ASTOpImageSwizzle>(this);
+			//		Ptr<ASTOpImageSwizzle> NewSwizzle = UE::Mutable::Private::Clone<ASTOpImageSwizzle>(this);
 			//		for (int c = 0; c < MUTABLE_OP_MAX_SWIZZLE_CHANNELS; ++c)
 			//		{
 			//			const ASTOpFixed* ChannelResize = static_cast<const ASTOpFixed*>(Sources[c].child().get());
@@ -696,7 +696,7 @@ namespace mu
 				if (bAreAllFormatsCompatible)
 				{
 					// Move the swizzle down all the paths
-					Ptr<ASTOpImagePixelFormat> NewFormat = mu::Clone<ASTOpImagePixelFormat>(FirstFormat);
+					Ptr<ASTOpImagePixelFormat> NewFormat = UE::Mutable::Private::Clone<ASTOpImagePixelFormat>(FirstFormat);
 					NewFormat->Format = Format;
 					at = NewFormat;
 				}
@@ -705,7 +705,7 @@ namespace mu
 			// Swizzle down plaincolours.
 			if (!at && sourceType == EOpType::IM_PLAINCOLOUR)
 			{
-				Ptr<ASTOpImagePlainColor> NewPlain = mu::Clone<ASTOpImagePlainColor>(channelSourceAt);
+				Ptr<ASTOpImagePlainColor> NewPlain = UE::Mutable::Private::Clone<ASTOpImagePlainColor>(channelSourceAt);
 				Ptr<ASTOpColorSwizzle> NewSwizzle = new ASTOpColorSwizzle;
 				for (int c = 0; c < MUTABLE_OP_MAX_SWIZZLE_CHANNELS; ++c)
 				{
@@ -734,9 +734,9 @@ namespace mu
 		//	)
 		//{
 		//	// Move the swizzle down all the paths
-		//	Ptr<ASTOpImageLayerColor> NewLayerColour = mu::Clone<ASTOpImageLayerColor>(Sources[3].child());
+		//	Ptr<ASTOpImageLayerColor> NewLayerColour = UE::Mutable::Private::Clone<ASTOpImageLayerColor>(Sources[3].child());
 
-		//	Ptr<ASTOpImageSwizzle> NewSwizzle = mu::Clone<ASTOpImageSwizzle>(this);
+		//	Ptr<ASTOpImageSwizzle> NewSwizzle = UE::Mutable::Private::Clone<ASTOpImageSwizzle>(this);
 		//	NewSwizzle->Sources[3] = NewLayerColour->base.child();
 
 		//	NewLayerColour->blendTypeAlpha = NewLayerColour->blendType;
@@ -759,9 +759,9 @@ namespace mu
 		//	Sources[3] && Sources[3]->GetOpType() == EOpType::IM_LAYER)
 		//{
 		//	// Move the swizzle down all the paths
-		//	Ptr<ASTOpImageLayer> NewLayer = mu::Clone<ASTOpImageLayer>(Sources[3].child());
+		//	Ptr<ASTOpImageLayer> NewLayer = UE::Mutable::Private::Clone<ASTOpImageLayer>(Sources[3].child());
 
-		//	Ptr<ASTOpImageSwizzle> NewSwizzle = mu::Clone<ASTOpImageSwizzle>(this);
+		//	Ptr<ASTOpImageSwizzle> NewSwizzle = UE::Mutable::Private::Clone<ASTOpImageSwizzle>(this);
 		//	NewSwizzle->Sources[3] = NewLayer->base.child();
 
 		//	NewLayer->blendTypeAlpha = NewLayer->blendType;
@@ -787,10 +787,10 @@ namespace mu
 			!(Sources[3]==Sources[0]) )
 		{
 			// Move the swizzle down all the rgb path
-			Ptr<ASTOpImageLayerColor> NewLayerColour = mu::Clone<ASTOpImageLayerColor>(Sources[0].child());
+			Ptr<ASTOpImageLayerColor> NewLayerColour = UE::Mutable::Private::Clone<ASTOpImageLayerColor>(Sources[0].child());
 			check(NewLayerColour);
 
-			Ptr<ASTOpImageSwizzle> NewSwizzle = mu::Clone<ASTOpImageSwizzle>(this);
+			Ptr<ASTOpImageSwizzle> NewSwizzle = UE::Mutable::Private::Clone<ASTOpImageSwizzle>(this);
 			NewSwizzle->Sources[0] = NewLayerColour->base.child();
 			NewSwizzle->Sources[1] = Sources[1] ? NewLayerColour->base.child() : nullptr;
 			NewSwizzle->Sources[2] = Sources[2] ? NewLayerColour->base.child() : nullptr;
@@ -814,7 +814,7 @@ namespace mu
 					// Remove the saturate for this channel
 					if (!NewSwizzle)
 					{
-						NewSwizzle = mu::Clone<ASTOpImageSwizzle>(this);
+						NewSwizzle = UE::Mutable::Private::Clone<ASTOpImageSwizzle>(this);
 					}
 
 					const ASTOpImageSaturate* OldSaturate = static_cast<const ASTOpImageSaturate*>(Sources[Channel].child().get());
@@ -841,10 +841,10 @@ namespace mu
 			)
 		{
 			// Move the swizzle down 
-			Ptr<ASTOpImageSaturate> NewSaturate = mu::Clone<ASTOpImageSaturate>(Sources[0].child());
+			Ptr<ASTOpImageSaturate> NewSaturate = UE::Mutable::Private::Clone<ASTOpImageSaturate>(Sources[0].child());
 			check(NewSaturate);
 
-			Ptr<ASTOpImageSwizzle> NewSwizzle = mu::Clone<ASTOpImageSwizzle>(this);
+			Ptr<ASTOpImageSwizzle> NewSwizzle = UE::Mutable::Private::Clone<ASTOpImageSwizzle>(this);
 			Ptr<ASTOp> OldSaturateBase = NewSaturate->Base.child();
 			NewSwizzle->Sources[0] = OldSaturateBase;
 			NewSwizzle->Sources[1] = OldSaturateBase;
@@ -938,7 +938,7 @@ namespace mu
 				NewBase->Source = Sources[0].child();
 				NewBase->Format = Format;
 
-				Ptr<ASTOpImageLayer> NewLayer = mu::Clone<ASTOpImageLayer>(OldLayer);
+				Ptr<ASTOpImageLayer> NewLayer = UE::Mutable::Private::Clone<ASTOpImageLayer>(OldLayer);
 				NewLayer->blend = OldLayer->base.child();
 				NewLayer->base = NewBase;
 				NewLayer->blendTypeAlpha = NewLayer->blendType;
@@ -992,11 +992,11 @@ namespace mu
 						NewBlendSwizzle->Sources[0] = OldLayer->blend.child();
 						NewBlendSwizzle->SourceChannels[0] = SourceChannels[3];
 
-						Ptr<ASTOpImageLayer> NewLayer = mu::Clone<ASTOpImageLayer>(OldLayer);
+						Ptr<ASTOpImageLayer> NewLayer = UE::Mutable::Private::Clone<ASTOpImageLayer>(OldLayer);
 						NewLayer->base = NewBaseSwizzle;
 						NewLayer->blend = NewBlendSwizzle;
 
-						Ptr<ASTOpImageSwizzle> NewSwizzle = mu::Clone<ASTOpImageSwizzle>(this);
+						Ptr<ASTOpImageSwizzle> NewSwizzle = UE::Mutable::Private::Clone<ASTOpImageSwizzle>(this);
 						NewSwizzle->Sources[3] = NewLayer;
 
 						at = NewSwizzle;
@@ -1076,9 +1076,9 @@ namespace mu
 	}
 
 
-	mu::Ptr<ImageSizeExpression> ASTOpImageSwizzle::GetImageSizeExpression() const
+	UE::Mutable::Private::Ptr<ImageSizeExpression> ASTOpImageSwizzle::GetImageSizeExpression() const
 	{
-		mu::Ptr<ImageSizeExpression> pRes;
+		UE::Mutable::Private::Ptr<ImageSizeExpression> pRes;
 
 		if (Sources[0].child())
 		{
@@ -1131,7 +1131,7 @@ namespace mu
 	//---------------------------------------------------------------------------------------------
 	//---------------------------------------------------------------------------------------------
 	//---------------------------------------------------------------------------------------------
-	mu::Ptr<ASTOp> Sink_ImageSwizzleAST::Apply(const ASTOpImageSwizzle* InRoot)
+	UE::Mutable::Private::Ptr<ASTOp> Sink_ImageSwizzleAST::Apply(const ASTOpImageSwizzle* InRoot)
 	{
 		Root = InRoot;
 		OldToNew.Reset();
@@ -1163,7 +1163,7 @@ namespace mu
 		}
 
 		InitialSource = Source;
-		mu::Ptr<ASTOp> NewSource = Visit(InitialSource, Root);
+		UE::Mutable::Private::Ptr<ASTOp> NewSource = Visit(InitialSource, Root);
 
 		Root = nullptr;
 
@@ -1178,7 +1178,7 @@ namespace mu
 
 
 	//---------------------------------------------------------------------------------------------
-	mu::Ptr<ASTOp> Sink_ImageSwizzleAST::Visit(mu::Ptr<ASTOp> at, const ASTOpImageSwizzle* CurrentSwizzleOp)
+	UE::Mutable::Private::Ptr<ASTOp> Sink_ImageSwizzleAST::Visit(UE::Mutable::Private::Ptr<ASTOp> at, const ASTOpImageSwizzle* CurrentSwizzleOp)
 	{
 		if (!at) return nullptr;
 
@@ -1189,14 +1189,14 @@ namespace mu
 			return *Cached;
 		}
 
-		mu::Ptr<ASTOp> newAt = at;
+		UE::Mutable::Private::Ptr<ASTOp> newAt = at;
 		switch (at->GetOpType())
 		{
 
 		case EOpType::IM_CONDITIONAL:
 		{
 			// We move the op down the two paths
-			auto newOp = mu::Clone<ASTOpConditional>(at);
+			auto newOp = UE::Mutable::Private::Clone<ASTOpConditional>(at);
 			newOp->yes = Visit(newOp->yes.child(), CurrentSwizzleOp);
 			newOp->no = Visit(newOp->no.child(), CurrentSwizzleOp);
 			newAt = newOp;
@@ -1206,7 +1206,7 @@ namespace mu
 		case EOpType::IM_SWITCH:
 		{
 			// We move the op down all the paths
-			Ptr<ASTOpSwitch> newOp = mu::Clone<ASTOpSwitch>(at);
+			Ptr<ASTOpSwitch> newOp = UE::Mutable::Private::Clone<ASTOpSwitch>(at);
 			newOp->Default = Visit(newOp->Default.child(), CurrentSwizzleOp);
 			for (ASTOpSwitch::FCase& c : newOp->Cases)
 			{
@@ -1218,7 +1218,7 @@ namespace mu
 
 		case EOpType::IM_COMPOSE:
 		{
-			Ptr<ASTOpImageCompose> newOp = mu::Clone<ASTOpImageCompose>(at);
+			Ptr<ASTOpImageCompose> newOp = UE::Mutable::Private::Clone<ASTOpImageCompose>(at);
 			newOp->Base = Visit(newOp->Base.child(), CurrentSwizzleOp);
 			newOp->BlockImage = Visit(newOp->BlockImage.child(), CurrentSwizzleOp);
 			newAt = newOp;
@@ -1227,7 +1227,7 @@ namespace mu
 
 		case EOpType::IM_PATCH:
 		{
-			Ptr<ASTOpImagePatch> newOp = mu::Clone<ASTOpImagePatch>(at);
+			Ptr<ASTOpImagePatch> newOp = UE::Mutable::Private::Clone<ASTOpImagePatch>(at);
 			newOp->base = Visit(newOp->base.child(), CurrentSwizzleOp);
 			newOp->patch = Visit(newOp->patch.child(), CurrentSwizzleOp);
 			newAt = newOp;
@@ -1236,7 +1236,7 @@ namespace mu
 
 		case EOpType::IM_MIPMAP:
 		{
-			Ptr<ASTOpImageMipmap> newOp = mu::Clone<ASTOpImageMipmap>(at);
+			Ptr<ASTOpImageMipmap> newOp = UE::Mutable::Private::Clone<ASTOpImageMipmap>(at);
 			newOp->Source = Visit(newOp->Source.child(), CurrentSwizzleOp);
 			newAt = newOp;
 			break;
@@ -1244,7 +1244,7 @@ namespace mu
 
 		case EOpType::IM_INTERPOLATE:
 		{
-			Ptr<ASTOpImageInterpolate> NewOp = mu::Clone<ASTOpImageInterpolate>(at);
+			Ptr<ASTOpImageInterpolate> NewOp = UE::Mutable::Private::Clone<ASTOpImageInterpolate>(at);
 
 			for (int32 v = 0; v < MUTABLE_OP_MAX_INTERPOLATE_COUNT; ++v)
 			{
@@ -1257,7 +1257,7 @@ namespace mu
 
 		case EOpType::IM_LAYER:
 		{
-			Ptr<ASTOpImageLayer> nop = mu::Clone<ASTOpImageLayer>(at);
+			Ptr<ASTOpImageLayer> nop = UE::Mutable::Private::Clone<ASTOpImageLayer>(at);
 			nop->base = Visit(nop->base.child(), CurrentSwizzleOp);
 			nop->blend = Visit(nop->blend.child(), CurrentSwizzleOp);
 			newAt = nop;
@@ -1266,7 +1266,7 @@ namespace mu
 
 		case EOpType::IM_LAYERCOLOUR:
 		{
-			Ptr<ASTOpImageLayerColor> nop = mu::Clone<ASTOpImageLayerColor>(at);
+			Ptr<ASTOpImageLayerColor> nop = UE::Mutable::Private::Clone<ASTOpImageLayerColor>(at);
 			nop->base = Visit(nop->base.child(), CurrentSwizzleOp);
 
 			// We need to swizzle the colour too
@@ -1284,7 +1284,7 @@ namespace mu
 
 		case EOpType::IM_DISPLACE:
 		{
-			Ptr<ASTOpImageDisplace> NewOp = mu::Clone<ASTOpImageDisplace>(at);
+			Ptr<ASTOpImageDisplace> NewOp = UE::Mutable::Private::Clone<ASTOpImageDisplace>(at);
 
 			Ptr<ASTOp> Child = NewOp->Source.child();
 			Ptr<ASTOp> NewSource = Visit(Child, CurrentSwizzleOp);
@@ -1296,7 +1296,7 @@ namespace mu
 
 		case EOpType::IM_INVERT:
 		{
-			Ptr<ASTOpImageInvert> NewOp = mu::Clone<ASTOpImageInvert>(at);
+			Ptr<ASTOpImageInvert> NewOp = UE::Mutable::Private::Clone<ASTOpImageInvert>(at);
 
 			Ptr<ASTOp> Child = NewOp->Base.child();
 			Ptr<ASTOp> NewSource = Visit(Child, CurrentSwizzleOp);
@@ -1308,7 +1308,7 @@ namespace mu
 
 		case EOpType::IM_RASTERMESH:
 		{
-			Ptr<ASTOpImageRasterMesh> NewOp = mu::Clone<ASTOpImageRasterMesh>(at);
+			Ptr<ASTOpImageRasterMesh> NewOp = UE::Mutable::Private::Clone<ASTOpImageRasterMesh>(at);
 
 			NewOp->image = Visit(NewOp->image.child(), CurrentSwizzleOp);
 
@@ -1331,7 +1331,7 @@ namespace mu
 
 		case EOpType::IM_TRANSFORM:
 		{
-			Ptr<ASTOpImageTransform> nop = mu::Clone<ASTOpImageTransform>(at);
+			Ptr<ASTOpImageTransform> nop = UE::Mutable::Private::Clone<ASTOpImageTransform>(at);
 			nop->Base = Visit(nop->Base.child(), CurrentSwizzleOp);
 			newAt = nop;
 			break;
@@ -1352,7 +1352,7 @@ namespace mu
 
 			if (bSameChannelOrder)
 			{
-				Ptr<ASTOpImagePixelFormat> NewOp = mu::Clone<ASTOpImagePixelFormat>(at);
+				Ptr<ASTOpImagePixelFormat> NewOp = UE::Mutable::Private::Clone<ASTOpImagePixelFormat>(at);
 				NewOp->Format = CurrentSwizzleOp->Format;
 				newAt = NewOp;
 			}
@@ -1364,7 +1364,7 @@ namespace mu
 			// We can remove the swizzle entirely.
 			// It is not 100% equivalent, because blank layouts are initialized with 0,0,0,1 so the result could be
 			// different, but those pixels shouldn't be used anyway.
-			newAt = mu::Clone<ASTOp>(at);
+			newAt = UE::Mutable::Private::Clone<ASTOp>(at);
 			break;
 		}
 
@@ -1375,7 +1375,7 @@ namespace mu
 		// end on tree branch, replace with format
 		if (at == newAt && at != InitialSource)
 		{
-			mu::Ptr<ASTOpImageSwizzle> NewOp = mu::Clone<ASTOpImageSwizzle>(CurrentSwizzleOp);
+			UE::Mutable::Private::Ptr<ASTOpImageSwizzle> NewOp = UE::Mutable::Private::Clone<ASTOpImageSwizzle>(CurrentSwizzleOp);
 			check(NewOp->GetOpType() == EOpType::IM_SWIZZLE);
 
 			ReplaceAllSources(NewOp, at);
diff --git a/Engine/Plugins/Mutable/Source/MutableTools/Private/MuT/ASTOpImageTransform.cpp b/Engine/Plugins/Mutable/Source/MutableTools/Private/MuT/ASTOpImageTransform.cpp
--- a/Engine/Plugins/Mutable/Source/MutableTools/Private/MuT/ASTOpImageTransform.cpp
+++ b/Engine/Plugins/Mutable/Source/MutableTools/Private/MuT/ASTOpImageTransform.cpp
@@ -9,7 +9,7 @@
 #include "MuR/Types.h"
 
 
-namespace mu
+namespace UE::Mutable::Private
 {
 
 	//-------------------------------------------------------------------------------------------------
@@ -77,7 +77,7 @@ namespace mu
 
 
 	//-------------------------------------------------------------------------------------------------
-	mu::Ptr<ASTOp> ASTOpImageTransform::Clone(MapChildFuncRef MapChild) const
+	UE::Mutable::Private::Ptr<ASTOp> ASTOpImageTransform::Clone(MapChildFuncRef MapChild) const
 	{
 		Ptr<ASTOpImageTransform> NewOp = new ASTOpImageTransform();
 		NewOp->Base     = MapChild(Base.child());
@@ -195,7 +195,7 @@ namespace mu
 
 
 	//-------------------------------------------------------------------------------------------------
-	mu::Ptr<ImageSizeExpression> ASTOpImageTransform::GetImageSizeExpression() const
+	UE::Mutable::Private::Ptr<ImageSizeExpression> ASTOpImageTransform::GetImageSizeExpression() const
 	{
 		if (Base)
 		{
diff --git a/Engine/Plugins/Mutable/Source/MutableTools/Private/MuT/ASTOpInstanceAdd.cpp b/Engine/Plugins/Mutable/Source/MutableTools/Private/MuT/ASTOpInstanceAdd.cpp
--- a/Engine/Plugins/Mutable/Source/MutableTools/Private/MuT/ASTOpInstanceAdd.cpp
+++ b/Engine/Plugins/Mutable/Source/MutableTools/Private/MuT/ASTOpInstanceAdd.cpp
@@ -11,7 +11,7 @@
 #include "Misc/AssertionMacros.h"
 
 
-namespace mu
+namespace UE::Mutable::Private
 {
 
 	ASTOpInstanceAdd::ASTOpInstanceAdd()
@@ -45,7 +45,7 @@ namespace mu
 	}
 
 
-	mu::Ptr<ASTOp> ASTOpInstanceAdd::Clone(MapChildFuncRef mapChild) const
+	UE::Mutable::Private::Ptr<ASTOp> ASTOpInstanceAdd::Clone(MapChildFuncRef mapChild) const
 	{
 		Ptr<ASTOpInstanceAdd> n = new ASTOpInstanceAdd();
 		n->type = type;
@@ -201,15 +201,15 @@ namespace mu
 
 				if (bCanSink)
 				{
-					Ptr<ASTOpConditional> NewOp = mu::Clone<ASTOpConditional>(ValueAt);
+					Ptr<ASTOpConditional> NewOp = UE::Mutable::Private::Clone<ASTOpConditional>(ValueAt);
 					NewOp->type = EOpType::IN_CONDITIONAL;
 
-					Ptr<ASTOpInstanceAdd> TrueOp = mu::Clone<ASTOpInstanceAdd>(this);
+					Ptr<ASTOpInstanceAdd> TrueOp = UE::Mutable::Private::Clone<ASTOpInstanceAdd>(this);
 					TrueOp->instance = TypedSource ? TypedSource->yes.child() : nullptr; 
 					TrueOp->value = TypedValue->yes.child();
 					NewOp->yes = TrueOp;
 
-					Ptr<ASTOpInstanceAdd> FalseOp = mu::Clone<ASTOpInstanceAdd>(this);
+					Ptr<ASTOpInstanceAdd> FalseOp = UE::Mutable::Private::Clone<ASTOpInstanceAdd>(this);
 					FalseOp->instance = TypedSource ? TypedSource->no.child() : nullptr; 
 					FalseOp->value = TypedValue->no.child();
 					NewOp->no = FalseOp;
diff --git a/Engine/Plugins/Mutable/Source/MutableTools/Private/MuT/ASTOpLayoutFromMesh.cpp b/Engine/Plugins/Mutable/Source/MutableTools/Private/MuT/ASTOpLayoutFromMesh.cpp
--- a/Engine/Plugins/Mutable/Source/MutableTools/Private/MuT/ASTOpLayoutFromMesh.cpp
+++ b/Engine/Plugins/Mutable/Source/MutableTools/Private/MuT/ASTOpLayoutFromMesh.cpp
@@ -25,7 +25,7 @@
 #include "MuT/ASTOpLayoutMerge.h"
 
 
-namespace mu
+namespace UE::Mutable::Private
 {
 
 
@@ -62,9 +62,9 @@ namespace mu
 	}
 
 
-	mu::Ptr<ASTOp> ASTOpLayoutFromMesh::Clone(MapChildFuncRef mapChild) const
+	UE::Mutable::Private::Ptr<ASTOp> ASTOpLayoutFromMesh::Clone(MapChildFuncRef mapChild) const
 	{
-		mu::Ptr<ASTOpLayoutFromMesh> n = new ASTOpLayoutFromMesh();
+		UE::Mutable::Private::Ptr<ASTOpLayoutFromMesh> n = new ASTOpLayoutFromMesh();
 		n->Mesh = mapChild(Mesh.child());
 		n->LayoutIndex = LayoutIndex;
 		return n;
@@ -114,7 +114,7 @@ namespace mu
 		public:
 
 			// \TODO This is recursive and may cause stack overflows in big models.
-			mu::Ptr<ASTOp> Apply(const ASTOpLayoutFromMesh* InRoot)
+			UE::Mutable::Private::Ptr<ASTOp> Apply(const ASTOpLayoutFromMesh* InRoot)
 			{
 				InitialRoot = InRoot;
 				check(InitialRoot);
@@ -122,7 +122,7 @@ namespace mu
 				OldToNew.Empty();
 
 				InitialSource = InitialRoot->Mesh.child();
-				mu::Ptr<ASTOp> NewSource = Visit(InitialSource, InitialRoot);
+				UE::Mutable::Private::Ptr<ASTOp> NewSource = Visit(InitialSource, InitialRoot);
 
 				// If there is any change, it is the new root.
 				if (NewSource != InitialSource)
@@ -140,7 +140,7 @@ namespace mu
 			Ptr<ASTOp> InitialSource;
 			TMap<FSinkerOldToNewKey, Ptr<ASTOp>> OldToNew;
 
-			mu::Ptr<ASTOp> Visit(const mu::Ptr<ASTOp>& at, const ASTOpLayoutFromMesh* CurrentSinkingOp)
+			UE::Mutable::Private::Ptr<ASTOp> Visit(const UE::Mutable::Private::Ptr<ASTOp>& at, const ASTOpLayoutFromMesh* CurrentSinkingOp)
 			{
 				if (!at) return nullptr;
 
@@ -151,7 +151,7 @@ namespace mu
 					return *Cached;
 				}
 
-				mu::Ptr<ASTOp> newAt = at;
+				UE::Mutable::Private::Ptr<ASTOp> newAt = at;
 				switch (at->GetOpType())
 				{
 
@@ -197,7 +197,7 @@ namespace mu
 
 				case EOpType::ME_CONDITIONAL:
 				{
-					Ptr<ASTOpConditional> NewConditional = mu::Clone<ASTOpConditional>(at);
+					Ptr<ASTOpConditional> NewConditional = UE::Mutable::Private::Clone<ASTOpConditional>(at);
 					NewConditional->type = EOpType::LA_CONDITIONAL;
 					NewConditional->yes = Visit(NewConditional->yes.child(), CurrentSinkingOp);
 					NewConditional->no = Visit(NewConditional->no.child(), CurrentSinkingOp);
@@ -207,7 +207,7 @@ namespace mu
 
 				case EOpType::ME_SWITCH:
 				{
-					Ptr<ASTOpSwitch> NewOp = mu::Clone<ASTOpSwitch>(at);
+					Ptr<ASTOpSwitch> NewOp = UE::Mutable::Private::Clone<ASTOpSwitch>(at);
 					NewOp->Type = EOpType::LA_SWITCH;
 					NewOp->Default = Visit(NewOp->Default.child(), CurrentSinkingOp);
 					for (ASTOpSwitch::FCase& c : NewOp->Cases)
@@ -232,7 +232,7 @@ namespace mu
 				default:
 					if (at != InitialSource)
 					{
-						mu::Ptr<ASTOpLayoutFromMesh> NewOp = mu::Clone<ASTOpLayoutFromMesh>(CurrentSinkingOp);
+						UE::Mutable::Private::Ptr<ASTOpLayoutFromMesh> NewOp = UE::Mutable::Private::Clone<ASTOpLayoutFromMesh>(CurrentSinkingOp);
 						NewOp->Mesh = at;
 						newAt = NewOp;
 					}
@@ -250,12 +250,12 @@ namespace mu
 	}
 
 
-	mu::Ptr<ASTOp> ASTOpLayoutFromMesh::OptimiseSink(const FModelOptimizationOptions& options, FOptimizeSinkContext& context) const
+	UE::Mutable::Private::Ptr<ASTOp> ASTOpLayoutFromMesh::OptimiseSink(const FModelOptimizationOptions& options, FOptimizeSinkContext& context) const
 	{
 		MUTABLE_CPUPROFILER_SCOPE(ASTOpLayoutFromMesh_Sink);
 
 		FSinkLayoutFromMesh Sinker;
-		mu::Ptr<ASTOp> at = Sinker.Apply(this);
+		UE::Mutable::Private::Ptr<ASTOp> at = Sinker.Apply(this);
 
 		return at;
 	}
diff --git a/Engine/Plugins/Mutable/Source/MutableTools/Private/MuT/ASTOpLayoutMerge.cpp b/Engine/Plugins/Mutable/Source/MutableTools/Private/MuT/ASTOpLayoutMerge.cpp
--- a/Engine/Plugins/Mutable/Source/MutableTools/Private/MuT/ASTOpLayoutMerge.cpp
+++ b/Engine/Plugins/Mutable/Source/MutableTools/Private/MuT/ASTOpLayoutMerge.cpp
@@ -19,7 +19,7 @@
 #include "MuT/ASTOpSwitch.h"
 
 
-namespace mu
+namespace UE::Mutable::Private
 {
 
 
@@ -57,9 +57,9 @@ namespace mu
 	}
 
 
-	mu::Ptr<ASTOp> ASTOpLayoutMerge::Clone(MapChildFuncRef mapChild) const
+	UE::Mutable::Private::Ptr<ASTOp> ASTOpLayoutMerge::Clone(MapChildFuncRef mapChild) const
 	{
-		mu::Ptr<ASTOpLayoutMerge> n = new ASTOpLayoutMerge();
+		UE::Mutable::Private::Ptr<ASTOpLayoutMerge> n = new ASTOpLayoutMerge();
 		n->Base = mapChild(Base.child());
 		n->Added = mapChild(Added.child());
 		return n;
diff --git a/Engine/Plugins/Mutable/Source/MutableTools/Private/MuT/ASTOpLayoutPack.cpp b/Engine/Plugins/Mutable/Source/MutableTools/Private/MuT/ASTOpLayoutPack.cpp
--- a/Engine/Plugins/Mutable/Source/MutableTools/Private/MuT/ASTOpLayoutPack.cpp
+++ b/Engine/Plugins/Mutable/Source/MutableTools/Private/MuT/ASTOpLayoutPack.cpp
@@ -19,7 +19,7 @@
 #include "MuT/ASTOpSwitch.h"
 
 
-namespace mu
+namespace UE::Mutable::Private
 {
 
 	ASTOpLayoutPack::ASTOpLayoutPack()
@@ -53,9 +53,9 @@ namespace mu
 	}
 
 
-	mu::Ptr<ASTOp> ASTOpLayoutPack::Clone(MapChildFuncRef mapChild) const
+	UE::Mutable::Private::Ptr<ASTOp> ASTOpLayoutPack::Clone(MapChildFuncRef mapChild) const
 	{
-		mu::Ptr<ASTOpLayoutPack> n = new ASTOpLayoutPack();
+		UE::Mutable::Private::Ptr<ASTOpLayoutPack> n = new ASTOpLayoutPack();
 		n->Source = mapChild(Source.child());
 		return n;
 	}
diff --git a/Engine/Plugins/Mutable/Source/MutableTools/Private/MuT/ASTOpLayoutRemoveBlocks.cpp b/Engine/Plugins/Mutable/Source/MutableTools/Private/MuT/ASTOpLayoutRemoveBlocks.cpp
--- a/Engine/Plugins/Mutable/Source/MutableTools/Private/MuT/ASTOpLayoutRemoveBlocks.cpp
+++ b/Engine/Plugins/Mutable/Source/MutableTools/Private/MuT/ASTOpLayoutRemoveBlocks.cpp
@@ -19,7 +19,7 @@
 #include "MuT/ASTOpSwitch.h"
 
 
-namespace mu
+namespace UE::Mutable::Private
 {
 
 	ASTOpLayoutRemoveBlocks::ASTOpLayoutRemoveBlocks()
@@ -55,9 +55,9 @@ namespace mu
 	}
 
 
-	mu::Ptr<ASTOp> ASTOpLayoutRemoveBlocks::Clone(MapChildFuncRef mapChild) const
+	UE::Mutable::Private::Ptr<ASTOp> ASTOpLayoutRemoveBlocks::Clone(MapChildFuncRef mapChild) const
 	{
-		mu::Ptr<ASTOpLayoutRemoveBlocks> n = new ASTOpLayoutRemoveBlocks();
+		UE::Mutable::Private::Ptr<ASTOpLayoutRemoveBlocks> n = new ASTOpLayoutRemoveBlocks();
 		n->Source = mapChild(Source.child());
 		n->ReferenceLayout = mapChild(ReferenceLayout.child());
 		return n;
diff --git a/Engine/Plugins/Mutable/Source/MutableTools/Private/MuT/ASTOpMaterialBreak.cpp b/Engine/Plugins/Mutable/Source/MutableTools/Private/MuT/ASTOpMaterialBreak.cpp
--- a/Engine/Plugins/Mutable/Source/MutableTools/Private/MuT/ASTOpMaterialBreak.cpp
+++ b/Engine/Plugins/Mutable/Source/MutableTools/Private/MuT/ASTOpMaterialBreak.cpp
@@ -11,7 +11,7 @@
 #include "MuR/Types.h"
 
 
-namespace mu
+namespace UE::Mutable::Private
 {
 	ASTOpMaterialBreak::ASTOpMaterialBreak()
 		: Material(this)
@@ -44,7 +44,7 @@ namespace mu
 	}
 
 
-	mu::Ptr<ASTOp> ASTOpMaterialBreak::Clone(MapChildFuncRef MapChild) const
+	UE::Mutable::Private::Ptr<ASTOp> ASTOpMaterialBreak::Clone(MapChildFuncRef MapChild) const
 	{
 		Ptr<ASTOpMaterialBreak> Result = new ASTOpMaterialBreak();
 		Result->Type = Type;
diff --git a/Engine/Plugins/Mutable/Source/MutableTools/Private/MuT/ASTOpMaterialBreakImageParameter.cpp b/Engine/Plugins/Mutable/Source/MutableTools/Private/MuT/ASTOpMaterialBreakImageParameter.cpp
--- a/Engine/Plugins/Mutable/Source/MutableTools/Private/MuT/ASTOpMaterialBreakImageParameter.cpp
+++ b/Engine/Plugins/Mutable/Source/MutableTools/Private/MuT/ASTOpMaterialBreakImageParameter.cpp
@@ -5,7 +5,7 @@
 #include "MuR/Types.h"
 
 
-namespace mu
+namespace UE::Mutable::Private
 {
 	ASTOpImageFromMaterialParameter::ASTOpImageFromMaterialParameter()
 	{
@@ -35,7 +35,7 @@ namespace mu
 	}
 
 
-	mu::Ptr<ASTOp> ASTOpImageFromMaterialParameter::Clone(MapChildFuncRef MapChild) const
+	UE::Mutable::Private::Ptr<ASTOp> ASTOpImageFromMaterialParameter::Clone(MapChildFuncRef MapChild) const
 	{
 		Ptr<ASTOpImageFromMaterialParameter> Result = new ASTOpImageFromMaterialParameter();
 		Result->ParameterName = ParameterName;
diff --git a/Engine/Plugins/Mutable/Source/MutableTools/Private/MuT/ASTOpMeshAddMetadata.cpp b/Engine/Plugins/Mutable/Source/MutableTools/Private/MuT/ASTOpMeshAddMetadata.cpp
--- a/Engine/Plugins/Mutable/Source/MutableTools/Private/MuT/ASTOpMeshAddMetadata.cpp
+++ b/Engine/Plugins/Mutable/Source/MutableTools/Private/MuT/ASTOpMeshAddMetadata.cpp
@@ -9,7 +9,7 @@
 #include "MuR/Types.h"
 
 
-namespace mu
+namespace UE::Mutable::Private
 {
 
 	ASTOpMeshAddMetadata::ASTOpMeshAddMetadata()
@@ -41,7 +41,7 @@ namespace mu
 	}
 
 
-	mu::Ptr<ASTOp> ASTOpMeshAddMetadata::Clone(MapChildFuncRef MapChild) const
+	UE::Mutable::Private::Ptr<ASTOp> ASTOpMeshAddMetadata::Clone(MapChildFuncRef MapChild) const
 	{
 		Ptr<ASTOpMeshAddMetadata> NewOp = new ASTOpMeshAddMetadata();
 		NewOp->Source = MapChild(Source.child());
