[Mutable] Change namespace name

diff --git a/Engine/Plugins/Mutable/Source/MutableTools/Private/MuT/ASTOpMeshApplyLayout.cpp b/Engine/Plugins/Mutable/Source/MutableTools/Private/MuT/ASTOpMeshApplyLayout.cpp
--- a/Engine/Plugins/Mutable/Source/MutableTools/Private/MuT/ASTOpMeshApplyLayout.cpp
+++ b/Engine/Plugins/Mutable/Source/MutableTools/Private/MuT/ASTOpMeshApplyLayout.cpp
@@ -9,7 +9,7 @@
 #include "MuR/RefCounted.h"
 #include "MuR/Types.h"
 
-namespace mu
+namespace UE::Mutable::Private
 {
 
 	ASTOpMeshApplyLayout::ASTOpMeshApplyLayout()
@@ -105,10 +105,10 @@ namespace mu
 		{
 			case EOpType::ME_ADDMETADATA:
 			{
-				Ptr<ASTOpMeshAddMetadata> New = mu::Clone<ASTOpMeshAddMetadata>(MeshAt);
+				Ptr<ASTOpMeshAddMetadata> New = UE::Mutable::Private::Clone<ASTOpMeshAddMetadata>(MeshAt);
 				if (New->Source)
 				{
-					Ptr<ASTOpMeshApplyLayout> NewApplyPose = mu::Clone<ASTOpMeshApplyLayout>(this);
+					Ptr<ASTOpMeshApplyLayout> NewApplyPose = UE::Mutable::Private::Clone<ASTOpMeshApplyLayout>(this);
 					NewApplyPose->Mesh = New->Source.child();
 					New->Source = NewApplyPose;
 				}
@@ -122,11 +122,11 @@ namespace mu
 				if (bLayoutIsConstant)
 				{
 					// Move the apply down all the paths
-					Ptr<ASTOpSwitch> NewSwitch = mu::Clone<ASTOpSwitch>(MeshAt);
+					Ptr<ASTOpSwitch> NewSwitch = UE::Mutable::Private::Clone<ASTOpSwitch>(MeshAt);
 
 					if (NewSwitch->Default)
 					{
-						Ptr<ASTOpMeshApplyLayout> defOp = mu::Clone<ASTOpMeshApplyLayout>(this);
+						Ptr<ASTOpMeshApplyLayout> defOp = UE::Mutable::Private::Clone<ASTOpMeshApplyLayout>(this);
 						defOp->Mesh = NewSwitch->Default.child();
 						NewSwitch->Default = defOp;
 					}
@@ -136,7 +136,7 @@ namespace mu
 					{
 						if (NewSwitch->Cases[CaseIndex].Branch)
 						{
-							Ptr<ASTOpMeshApplyLayout> bOp = mu::Clone<ASTOpMeshApplyLayout>(this);
+							Ptr<ASTOpMeshApplyLayout> bOp = UE::Mutable::Private::Clone<ASTOpMeshApplyLayout>(this);
 							bOp->Mesh = NewSwitch->Cases[CaseIndex].Branch.child();
 							NewSwitch->Cases[CaseIndex].Branch = bOp;
 						}
diff --git a/Engine/Plugins/Mutable/Source/MutableTools/Private/MuT/ASTOpMeshApplyPose.cpp b/Engine/Plugins/Mutable/Source/MutableTools/Private/MuT/ASTOpMeshApplyPose.cpp
--- a/Engine/Plugins/Mutable/Source/MutableTools/Private/MuT/ASTOpMeshApplyPose.cpp
+++ b/Engine/Plugins/Mutable/Source/MutableTools/Private/MuT/ASTOpMeshApplyPose.cpp
@@ -8,7 +8,7 @@
 #include "MuR/RefCounted.h"
 #include "MuR/Types.h"
 
-namespace mu
+namespace UE::Mutable::Private
 {
 
 	ASTOpMeshApplyPose::ASTOpMeshApplyPose()
@@ -98,10 +98,10 @@ namespace mu
 		{
 			case EOpType::ME_ADDMETADATA:
 			{
-				Ptr<ASTOpMeshAddMetadata> New = mu::Clone<ASTOpMeshAddMetadata>(MeshAt);
+				Ptr<ASTOpMeshAddMetadata> New = UE::Mutable::Private::Clone<ASTOpMeshAddMetadata>(MeshAt);
 				if (New->Source)
 				{
-					Ptr<ASTOpMeshApplyPose> NewApplyPose = mu::Clone<ASTOpMeshApplyPose>(this);
+					Ptr<ASTOpMeshApplyPose> NewApplyPose = UE::Mutable::Private::Clone<ASTOpMeshApplyPose>(this);
 					NewApplyPose->Base = New->Source.child();
 					New->Source = NewApplyPose;
 				}
diff --git a/Engine/Plugins/Mutable/Source/MutableTools/Private/MuT/ASTOpMeshApplyShape.cpp b/Engine/Plugins/Mutable/Source/MutableTools/Private/MuT/ASTOpMeshApplyShape.cpp
--- a/Engine/Plugins/Mutable/Source/MutableTools/Private/MuT/ASTOpMeshApplyShape.cpp
+++ b/Engine/Plugins/Mutable/Source/MutableTools/Private/MuT/ASTOpMeshApplyShape.cpp
@@ -8,7 +8,7 @@
 #include "MuR/Types.h"
 
 
-namespace mu
+namespace UE::Mutable::Private
 {
 
 
@@ -64,7 +64,7 @@ namespace mu
 	}
 
 
-	mu::Ptr<ASTOp> ASTOpMeshApplyShape::Clone(MapChildFuncRef MapChild) const
+	UE::Mutable::Private::Ptr<ASTOp> ASTOpMeshApplyShape::Clone(MapChildFuncRef MapChild) const
 	{
 		Ptr<ASTOpMeshApplyShape> NewOp = new ASTOpMeshApplyShape();
 		NewOp->Mesh = MapChild(Mesh.child());
diff --git a/Engine/Plugins/Mutable/Source/MutableTools/Private/MuT/ASTOpMeshBindShape.cpp b/Engine/Plugins/Mutable/Source/MutableTools/Private/MuT/ASTOpMeshBindShape.cpp
--- a/Engine/Plugins/Mutable/Source/MutableTools/Private/MuT/ASTOpMeshBindShape.cpp
+++ b/Engine/Plugins/Mutable/Source/MutableTools/Private/MuT/ASTOpMeshBindShape.cpp
@@ -14,7 +14,7 @@
 #include "MuT/ASTOpMeshAddMetadata.h"
 
 
-namespace mu
+namespace UE::Mutable::Private
 {
 
 
@@ -96,7 +96,7 @@ namespace mu
 	}
 
 
-	mu::Ptr<ASTOp> ASTOpMeshBindShape::Clone(MapChildFuncRef mapChild) const
+	UE::Mutable::Private::Ptr<ASTOp> ASTOpMeshBindShape::Clone(MapChildFuncRef mapChild) const
 	{
 		Ptr<ASTOpMeshBindShape> NewOp = new ASTOpMeshBindShape();
 		NewOp->Mesh = mapChild(Mesh.child());
@@ -199,7 +199,7 @@ namespace mu
 
 
 	//-------------------------------------------------------------------------------------------------
-	mu::Ptr<ASTOp> ASTOpMeshBindShape::OptimiseSink(const FModelOptimizationOptions&, FOptimizeSinkContext&) const
+	UE::Mutable::Private::Ptr<ASTOp> ASTOpMeshBindShape::OptimiseSink(const FModelOptimizationOptions&, FOptimizeSinkContext&) const
 	{
 		Ptr<ASTOp> NewOp;
 
@@ -236,11 +236,11 @@ namespace mu
 				}
 
 				// Move the operation down all the paths
-				Ptr<ASTOpSwitch> NewSwitch = mu::Clone<ASTOpSwitch>(MeshAt);
+				Ptr<ASTOpSwitch> NewSwitch = UE::Mutable::Private::Clone<ASTOpSwitch>(MeshAt);
 
 				if (NewSwitch->Default)
 				{
-					Ptr<ASTOpMeshBindShape> NewBind = mu::Clone<ASTOpMeshBindShape>(this);
+					Ptr<ASTOpMeshBindShape> NewBind = UE::Mutable::Private::Clone<ASTOpMeshBindShape>(this);
 					NewBind->Mesh = MeshSwitch->Default.child();
 					NewBind->Shape = ShapeSwitch->Default.child();
 					NewSwitch->Default = NewBind;
@@ -250,7 +250,7 @@ namespace mu
 				{
 					if (NewSwitch->Cases[v].Branch)
 					{
-						Ptr<ASTOpMeshBindShape> NewBind = mu::Clone<ASTOpMeshBindShape>(this);
+						Ptr<ASTOpMeshBindShape> NewBind = UE::Mutable::Private::Clone<ASTOpMeshBindShape>(this);
 						NewBind->Mesh = MeshSwitch->Cases[v].Branch.child();
 						NewBind->Shape = ShapeSwitch->FindBranch(MeshSwitch->Cases[v].Condition);
 						NewSwitch->Cases[v].Branch = NewBind;
@@ -272,11 +272,11 @@ namespace mu
 					break;
 				}
 
-				Ptr<ASTOpConditional> NewConditional = mu::Clone<ASTOpConditional>(MeshAt);
+				Ptr<ASTOpConditional> NewConditional = UE::Mutable::Private::Clone<ASTOpConditional>(MeshAt);
 
 				if (NewConditional->yes)
 				{
-					Ptr<ASTOpMeshBindShape> NewBind = mu::Clone<ASTOpMeshBindShape>(this);
+					Ptr<ASTOpMeshBindShape> NewBind = UE::Mutable::Private::Clone<ASTOpMeshBindShape>(this);
 					NewBind->Mesh = MeshConditional->yes.child();
 					NewBind->Shape = ShapeConditional->yes.child();
 					NewConditional->yes = NewBind;
@@ -284,7 +284,7 @@ namespace mu
 
 				if (NewConditional->no)
 				{
-					Ptr<ASTOpMeshBindShape> NewBind = mu::Clone<ASTOpMeshBindShape>(this);
+					Ptr<ASTOpMeshBindShape> NewBind = UE::Mutable::Private::Clone<ASTOpMeshBindShape>(this);
 					NewBind->Mesh = MeshConditional->no.child();
 					NewBind->Shape = ShapeConditional->no.child();
 					NewConditional->no = NewBind;
@@ -312,11 +312,11 @@ namespace mu
 			case EOpType::ME_SWITCH:
 			{
 				// Move the operation down all the paths
-				Ptr<ASTOpSwitch> NewSwitch = mu::Clone<ASTOpSwitch>(MeshAt);
+				Ptr<ASTOpSwitch> NewSwitch = UE::Mutable::Private::Clone<ASTOpSwitch>(MeshAt);
 
 				if (NewSwitch->Default)
 				{
-					Ptr<ASTOpMeshBindShape> NewBind = mu::Clone<ASTOpMeshBindShape>(this);
+					Ptr<ASTOpMeshBindShape> NewBind = UE::Mutable::Private::Clone<ASTOpMeshBindShape>(this);
 					NewBind->Mesh = NewSwitch->Default.child();
 					NewSwitch->Default = NewBind;
 				}
@@ -325,7 +325,7 @@ namespace mu
 				{
 					if (NewSwitch->Cases[v].Branch)
 					{
-						Ptr<ASTOpMeshBindShape> NewBind = mu::Clone<ASTOpMeshBindShape>(this);
+						Ptr<ASTOpMeshBindShape> NewBind = UE::Mutable::Private::Clone<ASTOpMeshBindShape>(this);
 						NewBind->Mesh = NewSwitch->Cases[v].Branch.child();
 						NewSwitch->Cases[v].Branch = NewBind;
 					}
@@ -338,18 +338,18 @@ namespace mu
 			case EOpType::ME_CONDITIONAL:
 			{
 				// Move the operation down all the paths
-				Ptr<ASTOpConditional> NewConditional = mu::Clone<ASTOpConditional>(MeshAt);
+				Ptr<ASTOpConditional> NewConditional = UE::Mutable::Private::Clone<ASTOpConditional>(MeshAt);
 
 				if (NewConditional->yes)
 				{
-					Ptr<ASTOpMeshBindShape> NewBind = mu::Clone<ASTOpMeshBindShape>(this);
+					Ptr<ASTOpMeshBindShape> NewBind = UE::Mutable::Private::Clone<ASTOpMeshBindShape>(this);
 					NewBind->Mesh = NewConditional->yes.child();
 					NewConditional->yes = NewBind;
 				}
 
 				if (NewConditional->no)
 				{
-					Ptr<ASTOpMeshBindShape> NewBind = mu::Clone<ASTOpMeshBindShape>(this);
+					Ptr<ASTOpMeshBindShape> NewBind = UE::Mutable::Private::Clone<ASTOpMeshBindShape>(this);
 					NewBind->Mesh = NewConditional->no.child();
 					NewConditional->no = NewBind;
 				}
@@ -363,10 +363,10 @@ namespace mu
 				// We bind something that could have a part removed: we can reorder to bind the entire mesh
 				// and apply remove later at runtime.
 
-				Ptr<ASTOpMeshRemoveMask> NewRemove = mu::Clone<ASTOpMeshRemoveMask>(MeshAt);
+				Ptr<ASTOpMeshRemoveMask> NewRemove = UE::Mutable::Private::Clone<ASTOpMeshRemoveMask>(MeshAt);
 				if (NewRemove->source)
 				{
-					Ptr<ASTOpMeshBindShape> NewBind = mu::Clone<ASTOpMeshBindShape>(this);
+					Ptr<ASTOpMeshBindShape> NewBind = UE::Mutable::Private::Clone<ASTOpMeshBindShape>(this);
 					NewBind->Mesh = NewRemove->source.child();
 					NewRemove->source = NewBind;
 				}
@@ -377,10 +377,10 @@ namespace mu
 
 			case EOpType::ME_ADDMETADATA:
 			{
-				Ptr<ASTOpMeshAddMetadata> New = mu::Clone<ASTOpMeshAddMetadata>(MeshAt);
+				Ptr<ASTOpMeshAddMetadata> New = UE::Mutable::Private::Clone<ASTOpMeshAddMetadata>(MeshAt);
 				if (New->Source)
 				{
-					Ptr<ASTOpMeshBindShape> NewBind = mu::Clone<ASTOpMeshBindShape>(this);
+					Ptr<ASTOpMeshBindShape> NewBind = UE::Mutable::Private::Clone<ASTOpMeshBindShape>(this);
 					NewBind->Mesh = New->Source.child();
 					New->Source = NewBind;
 				}
@@ -391,10 +391,10 @@ namespace mu
 
 			case EOpType::ME_PREPARELAYOUT:
 			{
-				Ptr<ASTOpMeshPrepareLayout> New = mu::Clone<ASTOpMeshPrepareLayout>(MeshAt);
+				Ptr<ASTOpMeshPrepareLayout> New = UE::Mutable::Private::Clone<ASTOpMeshPrepareLayout>(MeshAt);
 				if (New->Mesh)
 				{
-					Ptr<ASTOpMeshBindShape> NewBind = mu::Clone<ASTOpMeshBindShape>(this);
+					Ptr<ASTOpMeshBindShape> NewBind = UE::Mutable::Private::Clone<ASTOpMeshBindShape>(this);
 					NewBind->Mesh = New->Mesh.child();
 					New->Mesh = NewBind;
 				}
@@ -419,11 +419,11 @@ namespace mu
 			case EOpType::ME_SWITCH:
 			{
 				// Move the operation down all the paths
-				Ptr<ASTOpSwitch> NewSwitch = mu::Clone<ASTOpSwitch>(ShapeAt);
+				Ptr<ASTOpSwitch> NewSwitch = UE::Mutable::Private::Clone<ASTOpSwitch>(ShapeAt);
 
 				if (NewSwitch->Default)
 				{
-					Ptr<ASTOpMeshBindShape> NewBind = mu::Clone<ASTOpMeshBindShape>(this);
+					Ptr<ASTOpMeshBindShape> NewBind = UE::Mutable::Private::Clone<ASTOpMeshBindShape>(this);
 					NewBind->Shape = NewSwitch->Default.child();
 					NewSwitch->Default = NewBind;
 				}
@@ -432,7 +432,7 @@ namespace mu
 				{
 					if (NewSwitch->Cases[v].Branch)
 					{
-						Ptr<ASTOpMeshBindShape> NewBind = mu::Clone<ASTOpMeshBindShape>(this);
+						Ptr<ASTOpMeshBindShape> NewBind = UE::Mutable::Private::Clone<ASTOpMeshBindShape>(this);
 						NewBind->Shape = NewSwitch->Cases[v].Branch.child();
 						NewSwitch->Cases[v].Branch = NewBind;
 					}
@@ -445,18 +445,18 @@ namespace mu
 			case EOpType::ME_CONDITIONAL:
 			{
 				// Move the operation down all the paths
-				Ptr<ASTOpConditional> NewConditional = mu::Clone<ASTOpConditional>(ShapeAt);
+				Ptr<ASTOpConditional> NewConditional = UE::Mutable::Private::Clone<ASTOpConditional>(ShapeAt);
 
 				if (NewConditional->yes)
 				{
-					Ptr<ASTOpMeshBindShape> NewBind = mu::Clone<ASTOpMeshBindShape>(this);
+					Ptr<ASTOpMeshBindShape> NewBind = UE::Mutable::Private::Clone<ASTOpMeshBindShape>(this);
 					NewBind->Shape = NewConditional->yes.child();
 					NewConditional->yes = NewBind;
 				}
 
 				if (NewConditional->no)
 				{
-					Ptr<ASTOpMeshBindShape> NewBind = mu::Clone<ASTOpMeshBindShape>(this);
+					Ptr<ASTOpMeshBindShape> NewBind = UE::Mutable::Private::Clone<ASTOpMeshBindShape>(this);
 					NewBind->Shape = NewConditional->no.child();
 					NewConditional->no = NewBind;
 				}
@@ -468,7 +468,7 @@ namespace mu
 			case EOpType::ME_ADDMETADATA:
 			{
 				// Ignore the tags in the shape
-				Ptr<ASTOpMeshBindShape> NewBind = mu::Clone<ASTOpMeshBindShape>(this);
+				Ptr<ASTOpMeshBindShape> NewBind = UE::Mutable::Private::Clone<ASTOpMeshBindShape>(this);
 				const ASTOpMeshAddMetadata* New = static_cast<const ASTOpMeshAddMetadata*>(ShapeAt.get());
 				NewBind->Shape = New->Source.child();
 				NewOp = NewBind;
diff --git a/Engine/Plugins/Mutable/Source/MutableTools/Private/MuT/ASTOpMeshClipDeform.cpp b/Engine/Plugins/Mutable/Source/MutableTools/Private/MuT/ASTOpMeshClipDeform.cpp
--- a/Engine/Plugins/Mutable/Source/MutableTools/Private/MuT/ASTOpMeshClipDeform.cpp
+++ b/Engine/Plugins/Mutable/Source/MutableTools/Private/MuT/ASTOpMeshClipDeform.cpp
@@ -8,7 +8,7 @@
 #include "MuR/Types.h"
 
 
-namespace mu
+namespace UE::Mutable::Private
 {
 
 
@@ -46,7 +46,7 @@ namespace mu
 	}
 
 
-	mu::Ptr<ASTOp> ASTOpMeshClipDeform::Clone(MapChildFuncRef mapChild) const
+	UE::Mutable::Private::Ptr<ASTOp> ASTOpMeshClipDeform::Clone(MapChildFuncRef mapChild) const
 	{
 		Ptr<ASTOpMeshClipDeform> n = new ASTOpMeshClipDeform();
 		n->Mesh = mapChild(Mesh.child());
diff --git a/Engine/Plugins/Mutable/Source/MutableTools/Private/MuT/ASTOpMeshClipMorphPlane.cpp b/Engine/Plugins/Mutable/Source/MutableTools/Private/MuT/ASTOpMeshClipMorphPlane.cpp
--- a/Engine/Plugins/Mutable/Source/MutableTools/Private/MuT/ASTOpMeshClipMorphPlane.cpp
+++ b/Engine/Plugins/Mutable/Source/MutableTools/Private/MuT/ASTOpMeshClipMorphPlane.cpp
@@ -9,7 +9,7 @@
 #include "HAL/PlatformMath.h"
 
 
-namespace mu
+namespace UE::Mutable::Private
 {
 
 
@@ -58,7 +58,7 @@ namespace mu
 	}
 
 
-	mu::Ptr<ASTOp> ASTOpMeshClipMorphPlane::Clone(MapChildFuncRef mapChild) const
+	UE::Mutable::Private::Ptr<ASTOp> ASTOpMeshClipMorphPlane::Clone(MapChildFuncRef mapChild) const
 	{
 		Ptr<ASTOpMeshClipMorphPlane> n = new ASTOpMeshClipMorphPlane();
 		n->Source = mapChild(Source.child());
@@ -118,7 +118,7 @@ namespace mu
 	}
 
 
-	mu::Ptr<ASTOp> ASTOpMeshClipMorphPlane::OptimiseSink(const FModelOptimizationOptions&, FOptimizeSinkContext&) const
+	UE::Mutable::Private::Ptr<ASTOp> ASTOpMeshClipMorphPlane::OptimiseSink(const FModelOptimizationOptions&, FOptimizeSinkContext&) const
 	{
 		Ptr<ASTOp> NewOp;
 
@@ -135,11 +135,11 @@ namespace mu
 
 		case EOpType::ME_ADDMETADATA:
 		{
-			Ptr<ASTOpMeshAddMetadata> NewAddMetadata = mu::Clone<ASTOpMeshAddMetadata>(Source.child());
+			Ptr<ASTOpMeshAddMetadata> NewAddMetadata = UE::Mutable::Private::Clone<ASTOpMeshAddMetadata>(Source.child());
 
 			if (NewAddMetadata->Source)
 			{
-				Ptr<ASTOpMeshClipMorphPlane> New = mu::Clone<ASTOpMeshClipMorphPlane>(this);
+				Ptr<ASTOpMeshClipMorphPlane> New = UE::Mutable::Private::Clone<ASTOpMeshClipMorphPlane>(this);
 				New->Source = NewAddMetadata->Source.child();
 				NewAddMetadata->Source = New;
 			}
diff --git a/Engine/Plugins/Mutable/Source/MutableTools/Private/MuT/ASTOpMeshClipMorphPlane.h b/Engine/Plugins/Mutable/Source/MutableTools/Private/MuT/ASTOpMeshClipMorphPlane.h
--- a/Engine/Plugins/Mutable/Source/MutableTools/Private/MuT/ASTOpMeshClipMorphPlane.h
+++ b/Engine/Plugins/Mutable/Source/MutableTools/Private/MuT/ASTOpMeshClipMorphPlane.h
@@ -8,7 +8,7 @@
 #include "MuR/Ptr.h"
 #include "MuT/AST.h"
 
-namespace mu
+namespace UE::Mutable::Private
 {
 	struct FProgram;
 
@@ -40,7 +40,7 @@ namespace mu
 		virtual Ptr<ASTOp> Clone(MapChildFuncRef mapChild) const override;
 		virtual void ForEachChild(const TFunctionRef<void(ASTChild&)>) override;
 		virtual void Link(FProgram& program, FLinkerOptions* Options) override;
-		virtual mu::Ptr<ASTOp> OptimiseSink(const FModelOptimizationOptions&, FOptimizeSinkContext&) const override;
+		virtual UE::Mutable::Private::Ptr<ASTOp> OptimiseSink(const FModelOptimizationOptions&, FOptimizeSinkContext&) const override;
 		virtual FSourceDataDescriptor GetSourceDataDescriptor(FGetSourceDataDescriptorContext*) const override;
 	};
 
diff --git a/Engine/Plugins/Mutable/Source/MutableTools/Private/MuT/ASTOpMeshClipWithMesh.cpp b/Engine/Plugins/Mutable/Source/MutableTools/Private/MuT/ASTOpMeshClipWithMesh.cpp
--- a/Engine/Plugins/Mutable/Source/MutableTools/Private/MuT/ASTOpMeshClipWithMesh.cpp
+++ b/Engine/Plugins/Mutable/Source/MutableTools/Private/MuT/ASTOpMeshClipWithMesh.cpp
@@ -8,7 +8,7 @@
 #include "MuR/RefCounted.h"
 #include "MuR/Types.h"
 
-namespace mu
+namespace UE::Mutable::Private
 {
 
 	ASTOpMeshClipWithMesh::ASTOpMeshClipWithMesh()
@@ -98,10 +98,10 @@ namespace mu
 		{
 			case EOpType::ME_ADDMETADATA:
 			{
-				Ptr<ASTOpMeshAddMetadata> New = mu::Clone<ASTOpMeshAddMetadata>(MeshAt);
+				Ptr<ASTOpMeshAddMetadata> New = UE::Mutable::Private::Clone<ASTOpMeshAddMetadata>(MeshAt);
 				if (New->Source)
 				{
-					Ptr<ASTOpMeshClipWithMesh> NewApplyPose = mu::Clone<ASTOpMeshClipWithMesh>(this);
+					Ptr<ASTOpMeshClipWithMesh> NewApplyPose = UE::Mutable::Private::Clone<ASTOpMeshClipWithMesh>(this);
 					NewApplyPose->Source = New->Source.child();
 					New->Source = NewApplyPose;
 				}
diff --git a/Engine/Plugins/Mutable/Source/MutableTools/Private/MuT/ASTOpMeshDifference.cpp b/Engine/Plugins/Mutable/Source/MutableTools/Private/MuT/ASTOpMeshDifference.cpp
--- a/Engine/Plugins/Mutable/Source/MutableTools/Private/MuT/ASTOpMeshDifference.cpp
+++ b/Engine/Plugins/Mutable/Source/MutableTools/Private/MuT/ASTOpMeshDifference.cpp
@@ -12,7 +12,7 @@
 #include "MuT/ASTOpMeshAddMetadata.h"
 
 
-namespace mu
+namespace UE::Mutable::Private
 {
 
 	//---------------------------------------------------------------------------------------------
@@ -46,7 +46,7 @@ namespace mu
 
 
 	//---------------------------------------------------------------------------------------------
-	mu::Ptr<ASTOp> ASTOpMeshDifference::Clone(MapChildFuncRef mapChild) const
+	UE::Mutable::Private::Ptr<ASTOp> ASTOpMeshDifference::Clone(MapChildFuncRef mapChild) const
 	{
 		Ptr<ASTOpMeshDifference> n = new ASTOpMeshDifference();
 		n->Base = mapChild(Base.child());
@@ -104,7 +104,7 @@ namespace mu
 
 
 	//-------------------------------------------------------------------------------------------------
-	mu::Ptr<ASTOp> ASTOpMeshDifference::OptimiseSink(const FModelOptimizationOptions&, FOptimizeSinkContext&) const
+	UE::Mutable::Private::Ptr<ASTOp> ASTOpMeshDifference::OptimiseSink(const FModelOptimizationOptions&, FOptimizeSinkContext&) const
 	{
 		Ptr<ASTOp> NewOp;
 
@@ -141,11 +141,11 @@ namespace mu
 				}
 
 				// Move the operation down all the paths
-				Ptr<ASTOpSwitch> NewSwitch = mu::Clone<ASTOpSwitch>(BaseAt);
+				Ptr<ASTOpSwitch> NewSwitch = UE::Mutable::Private::Clone<ASTOpSwitch>(BaseAt);
 
 				if (NewSwitch->Default)
 				{
-					Ptr<ASTOpMeshDifference> NewDiff = mu::Clone<ASTOpMeshDifference>(this);
+					Ptr<ASTOpMeshDifference> NewDiff = UE::Mutable::Private::Clone<ASTOpMeshDifference>(this);
 					NewDiff->Base = BaseSwitch->Default.child();
 					NewDiff->Target = TargetSwitch->Default.child();
 					NewSwitch->Default = NewDiff;
@@ -155,7 +155,7 @@ namespace mu
 				{
 					if (NewSwitch->Cases[v].Branch)
 					{
-						Ptr<ASTOpMeshDifference> NewDiff = mu::Clone<ASTOpMeshDifference>(this);
+						Ptr<ASTOpMeshDifference> NewDiff = UE::Mutable::Private::Clone<ASTOpMeshDifference>(this);
 						NewDiff->Base = BaseSwitch->Cases[v].Branch.child();
 						NewDiff->Target = TargetSwitch->FindBranch(BaseSwitch->Cases[v].Condition);
 						NewSwitch->Cases[v].Branch = NewDiff;
@@ -177,11 +177,11 @@ namespace mu
 					break;
 				}
 
-				Ptr<ASTOpConditional> NewConditional = mu::Clone<ASTOpConditional>(BaseAt);
+				Ptr<ASTOpConditional> NewConditional = UE::Mutable::Private::Clone<ASTOpConditional>(BaseAt);
 
 				if (NewConditional->yes)
 				{
-					Ptr<ASTOpMeshDifference> NewDiff = mu::Clone<ASTOpMeshDifference>(this);
+					Ptr<ASTOpMeshDifference> NewDiff = UE::Mutable::Private::Clone<ASTOpMeshDifference>(this);
 					NewDiff->Base = BaseConditional->yes.child();
 					NewDiff->Target = TargetConditional->yes.child();
 					NewConditional->yes = NewDiff;
@@ -189,7 +189,7 @@ namespace mu
 
 				if (NewConditional->no)
 				{
-					Ptr<ASTOpMeshDifference> NewDiff = mu::Clone<ASTOpMeshDifference>(this);
+					Ptr<ASTOpMeshDifference> NewDiff = UE::Mutable::Private::Clone<ASTOpMeshDifference>(this);
 					NewDiff->Base = BaseConditional->no.child();
 					NewDiff->Target = TargetConditional->no.child();
 					NewConditional->no = NewDiff;
@@ -217,11 +217,11 @@ namespace mu
 			case EOpType::ME_SWITCH:
 			{
 				// Move the operation down all the paths
-				Ptr<ASTOpSwitch> NewSwitch = mu::Clone<ASTOpSwitch>(BaseAt);
+				Ptr<ASTOpSwitch> NewSwitch = UE::Mutable::Private::Clone<ASTOpSwitch>(BaseAt);
 
 				if (NewSwitch->Default)
 				{
-					Ptr<ASTOpMeshDifference> NewDiff = mu::Clone<ASTOpMeshDifference>(this);
+					Ptr<ASTOpMeshDifference> NewDiff = UE::Mutable::Private::Clone<ASTOpMeshDifference>(this);
 					NewDiff->Base = NewSwitch->Default.child();
 					NewSwitch->Default = NewDiff;
 				}
@@ -230,7 +230,7 @@ namespace mu
 				{
 					if (NewSwitch->Cases[v].Branch)
 					{
-						Ptr<ASTOpMeshDifference> NewDiff = mu::Clone<ASTOpMeshDifference>(this);
+						Ptr<ASTOpMeshDifference> NewDiff = UE::Mutable::Private::Clone<ASTOpMeshDifference>(this);
 						NewDiff->Base = NewSwitch->Cases[v].Branch.child();
 						NewSwitch->Cases[v].Branch = NewDiff;
 					}
@@ -243,18 +243,18 @@ namespace mu
 			case EOpType::ME_CONDITIONAL:
 			{
 				// Move the operation down all the paths
-				Ptr<ASTOpConditional> NewConditional = mu::Clone<ASTOpConditional>(BaseAt);
+				Ptr<ASTOpConditional> NewConditional = UE::Mutable::Private::Clone<ASTOpConditional>(BaseAt);
 
 				if (NewConditional->yes)
 				{
-					Ptr<ASTOpMeshDifference> NewDiff = mu::Clone<ASTOpMeshDifference>(this);
+					Ptr<ASTOpMeshDifference> NewDiff = UE::Mutable::Private::Clone<ASTOpMeshDifference>(this);
 					NewDiff->Base = NewConditional->yes.child();
 					NewConditional->yes = NewDiff;
 				}
 
 				if (NewConditional->no)
 				{
-					Ptr<ASTOpMeshDifference> NewDiff = mu::Clone<ASTOpMeshDifference>(this);
+					Ptr<ASTOpMeshDifference> NewDiff = UE::Mutable::Private::Clone<ASTOpMeshDifference>(this);
 					NewDiff->Base = NewConditional->no.child();
 					NewConditional->no = NewDiff;
 				}
@@ -265,11 +265,11 @@ namespace mu
 
 			case EOpType::ME_ADDMETADATA:
 			{
-				Ptr<ASTOpMeshAddMetadata> NewAdd = mu::Clone<ASTOpMeshAddMetadata>(BaseAt);
+				Ptr<ASTOpMeshAddMetadata> NewAdd = UE::Mutable::Private::Clone<ASTOpMeshAddMetadata>(BaseAt);
 
 				if (NewAdd->Source)
 				{
-					Ptr<ASTOpMeshDifference> NewDiff = mu::Clone<ASTOpMeshDifference>(this);
+					Ptr<ASTOpMeshDifference> NewDiff = UE::Mutable::Private::Clone<ASTOpMeshDifference>(this);
 					NewDiff->Base = NewAdd->Source.child();
 					NewAdd->Source = NewDiff;
 				}
@@ -294,11 +294,11 @@ namespace mu
 			case EOpType::ME_SWITCH:
 			{
 				// Move the operation down all the paths
-				Ptr<ASTOpSwitch> NewSwitch = mu::Clone<ASTOpSwitch>(TargetAt);
+				Ptr<ASTOpSwitch> NewSwitch = UE::Mutable::Private::Clone<ASTOpSwitch>(TargetAt);
 
 				if (NewSwitch->Default)
 				{
-					Ptr<ASTOpMeshDifference> NewDiff = mu::Clone<ASTOpMeshDifference>(this);
+					Ptr<ASTOpMeshDifference> NewDiff = UE::Mutable::Private::Clone<ASTOpMeshDifference>(this);
 					NewDiff->Target = NewSwitch->Default.child();
 					NewSwitch->Default = NewDiff;
 				}
@@ -307,7 +307,7 @@ namespace mu
 				{
 					if (NewSwitch->Cases[v].Branch)
 					{
-						Ptr<ASTOpMeshDifference> NewDiff = mu::Clone<ASTOpMeshDifference>(this);
+						Ptr<ASTOpMeshDifference> NewDiff = UE::Mutable::Private::Clone<ASTOpMeshDifference>(this);
 						NewDiff->Target = NewSwitch->Cases[v].Branch.child();
 						NewSwitch->Cases[v].Branch = NewDiff;
 					}
@@ -320,18 +320,18 @@ namespace mu
 			case EOpType::ME_CONDITIONAL:
 			{
 				// Move the operation down all the paths
-				Ptr<ASTOpConditional> NewConditional = mu::Clone<ASTOpConditional>(TargetAt);
+				Ptr<ASTOpConditional> NewConditional = UE::Mutable::Private::Clone<ASTOpConditional>(TargetAt);
 
 				if (NewConditional->yes)
 				{
-					Ptr<ASTOpMeshDifference> NewDiff = mu::Clone<ASTOpMeshDifference>(this);
+					Ptr<ASTOpMeshDifference> NewDiff = UE::Mutable::Private::Clone<ASTOpMeshDifference>(this);
 					NewDiff->Target = NewConditional->yes.child();
 					NewConditional->yes = NewDiff;
 				}
 
 				if (NewConditional->no)
 				{
-					Ptr<ASTOpMeshDifference> NewDiff = mu::Clone<ASTOpMeshDifference>(this);
+					Ptr<ASTOpMeshDifference> NewDiff = UE::Mutable::Private::Clone<ASTOpMeshDifference>(this);
 					NewDiff->Target = NewConditional->no.child();
 					NewConditional->no = NewDiff;
 				}
@@ -345,7 +345,7 @@ namespace mu
 				// Ignore tags in this branch
 				const ASTOpMeshAddMetadata* Add = static_cast<const ASTOpMeshAddMetadata*>(TargetAt.get());
 
-				Ptr<ASTOpMeshDifference> NewDiff = mu::Clone<ASTOpMeshDifference>(this);
+				Ptr<ASTOpMeshDifference> NewDiff = UE::Mutable::Private::Clone<ASTOpMeshDifference>(this);
 				NewDiff->Target = Add->Source.child();
 				NewOp = NewDiff;
 				break;
diff --git a/Engine/Plugins/Mutable/Source/MutableTools/Private/MuT/ASTOpMeshExtractLayoutBlocks.cpp b/Engine/Plugins/Mutable/Source/MutableTools/Private/MuT/ASTOpMeshExtractLayoutBlocks.cpp
--- a/Engine/Plugins/Mutable/Source/MutableTools/Private/MuT/ASTOpMeshExtractLayoutBlocks.cpp
+++ b/Engine/Plugins/Mutable/Source/MutableTools/Private/MuT/ASTOpMeshExtractLayoutBlocks.cpp
@@ -19,7 +19,7 @@
 #include "MuR/Types.h"
 #include "Misc/AssertionMacros.h"
 
-namespace mu
+namespace UE::Mutable::Private
 {
 
 	//---------------------------------------------------------------------------------------------
@@ -49,7 +49,7 @@ namespace mu
 	}
 
 
-	mu::Ptr<ASTOp> ASTOpMeshExtractLayoutBlocks::Clone(MapChildFuncRef mapChild) const
+	UE::Mutable::Private::Ptr<ASTOp> ASTOpMeshExtractLayoutBlocks::Clone(MapChildFuncRef mapChild) const
 	{
 		Ptr<ASTOpMeshExtractLayoutBlocks> n = new ASTOpMeshExtractLayoutBlocks();
 		n->Source = mapChild(Source.child());
@@ -101,9 +101,9 @@ namespace mu
 	}
 
 	 
-	mu::Ptr<ASTOp> ASTOpMeshExtractLayoutBlocks::OptimiseSink(const FModelOptimizationOptions&, FOptimizeSinkContext& Context) const
+	UE::Mutable::Private::Ptr<ASTOp> ASTOpMeshExtractLayoutBlocks::OptimiseSink(const FModelOptimizationOptions&, FOptimizeSinkContext& Context) const
 	{
-		mu::Ptr<ASTOp> NewOp = Context.MeshExtractLayoutBlocksSinker.Apply(this);
+		UE::Mutable::Private::Ptr<ASTOp> NewOp = Context.MeshExtractLayoutBlocksSinker.Apply(this);
 		return NewOp;
 	}
 
@@ -120,14 +120,14 @@ namespace mu
 
 
 	//---------------------------------------------------------------------------------------------
-	mu::Ptr<ASTOp> Sink_MeshExtractLayoutBlocksAST::Apply(const ASTOpMeshExtractLayoutBlocks* root)
+	UE::Mutable::Private::Ptr<ASTOp> Sink_MeshExtractLayoutBlocksAST::Apply(const ASTOpMeshExtractLayoutBlocks* root)
 	{
 		Root = root;
 
 		OldToNew.Reset();
 
 		InitialSource = Root->Source.child();
-		mu::Ptr<ASTOp> newSource = Visit(InitialSource, Root);
+		UE::Mutable::Private::Ptr<ASTOp> newSource = Visit(InitialSource, Root);
 
 		// If there is any change, it is the new root.
 		if (newSource != InitialSource)
@@ -138,7 +138,7 @@ namespace mu
 		return nullptr;
 	}
 
-	mu::Ptr<ASTOp> Sink_MeshExtractLayoutBlocksAST::Visit(const mu::Ptr<ASTOp>& at, const ASTOpMeshExtractLayoutBlocks* currentSinkOp)
+	UE::Mutable::Private::Ptr<ASTOp> Sink_MeshExtractLayoutBlocksAST::Visit(const UE::Mutable::Private::Ptr<ASTOp>& at, const ASTOpMeshExtractLayoutBlocks* currentSinkOp)
 	{
 		if (!at) return nullptr;
 
@@ -149,13 +149,13 @@ namespace mu
 			return *Cached;
 		}
 
-		mu::Ptr<ASTOp> newAt = at;
+		UE::Mutable::Private::Ptr<ASTOp> newAt = at;
 		switch (at->GetOpType())
 		{
 
 		case EOpType::ME_APPLYLAYOUT:
 		{
-			Ptr<ASTOpMeshApplyLayout> NewOp = mu::Clone<ASTOpMeshApplyLayout>(at);
+			Ptr<ASTOpMeshApplyLayout> NewOp = UE::Mutable::Private::Clone<ASTOpMeshApplyLayout>(at);
 			NewOp->Mesh = Visit(NewOp->Mesh.child(), currentSinkOp);
 			newAt = NewOp;
 			break;
@@ -197,7 +197,7 @@ namespace mu
 				if (at != InitialSource)
 				{
 					// We didn't optimize it, so create the end operation as in the switch "default" case.
-					mu::Ptr<ASTOpMeshExtractLayoutBlocks> newOp = mu::Clone<ASTOpMeshExtractLayoutBlocks>(currentSinkOp);
+					UE::Mutable::Private::Ptr<ASTOpMeshExtractLayoutBlocks> newOp = UE::Mutable::Private::Clone<ASTOpMeshExtractLayoutBlocks>(currentSinkOp);
 					newOp->Source = at;
 					newAt = newOp;
 				}
@@ -207,7 +207,7 @@ namespace mu
 
 		case EOpType::ME_SETSKELETON:
 		{
-			Ptr<ASTOpMeshSetSkeleton> NewOp = mu::Clone<ASTOpMeshSetSkeleton>(at);
+			Ptr<ASTOpMeshSetSkeleton> NewOp = UE::Mutable::Private::Clone<ASTOpMeshSetSkeleton>(at);
 			NewOp->Source = Visit( NewOp->Source.child(), currentSinkOp);
 			newAt = NewOp;
 			break;
@@ -215,7 +215,7 @@ namespace mu
 
 		case EOpType::ME_ADDMETADATA:
 		{
-			Ptr<ASTOpMeshAddMetadata> NewOp = mu::Clone<ASTOpMeshAddMetadata>(at);
+			Ptr<ASTOpMeshAddMetadata> NewOp = UE::Mutable::Private::Clone<ASTOpMeshAddMetadata>(at);
 			NewOp->Source = Visit(NewOp->Source.child(), currentSinkOp);
 			newAt = NewOp;
 			break;
@@ -223,7 +223,7 @@ namespace mu
 
 		case EOpType::ME_CLIPMORPHPLANE:
 		{
-			Ptr<ASTOpMeshClipMorphPlane> NewOp = mu::Clone<ASTOpMeshClipMorphPlane>(at);
+			Ptr<ASTOpMeshClipMorphPlane> NewOp = UE::Mutable::Private::Clone<ASTOpMeshClipMorphPlane>(at);
 			NewOp->Source = Visit(NewOp->Source.child(), currentSinkOp);
 			newAt = NewOp;
 			break;
@@ -231,7 +231,7 @@ namespace mu
 
 		case EOpType::ME_MORPH:
 		{
-			Ptr<ASTOpMeshMorph> NewOp = mu::Clone<ASTOpMeshMorph>(at);
+			Ptr<ASTOpMeshMorph> NewOp = UE::Mutable::Private::Clone<ASTOpMeshMorph>(at);
 			NewOp->Base = Visit(NewOp->Base.child(), currentSinkOp);
 			NewOp->Target = Visit(NewOp->Target.child(), currentSinkOp);
 			newAt = NewOp;
@@ -240,7 +240,7 @@ namespace mu
 
 		case EOpType::ME_MERGE:
 		{
-			Ptr<ASTOpMeshMerge> NewOp = mu::Clone<ASTOpMeshMerge>(at);
+			Ptr<ASTOpMeshMerge> NewOp = UE::Mutable::Private::Clone<ASTOpMeshMerge>(at);
 			NewOp->Base = Visit(NewOp->Base.child(), currentSinkOp);
 			NewOp->Added = Visit(NewOp->Added.child(), currentSinkOp);
 			newAt = NewOp;
@@ -249,7 +249,7 @@ namespace mu
 
 		case EOpType::ME_APPLYPOSE:
 		{
-			Ptr<ASTOpMeshApplyPose> NewOp = mu::Clone<ASTOpMeshApplyPose>(at);
+			Ptr<ASTOpMeshApplyPose> NewOp = UE::Mutable::Private::Clone<ASTOpMeshApplyPose>(at);
 			NewOp->Base = Visit(NewOp->Base.child(), currentSinkOp);
 			newAt = NewOp;
 			break;
@@ -258,7 +258,7 @@ namespace mu
 		case EOpType::ME_REMOVEMASK:
 		{
 			// TODO: Make mask smaller?
-			Ptr<ASTOpMeshRemoveMask> newOp = mu::Clone<ASTOpMeshRemoveMask>(at);
+			Ptr<ASTOpMeshRemoveMask> newOp = UE::Mutable::Private::Clone<ASTOpMeshRemoveMask>(at);
 			newOp->source = Visit(newOp->source.child(), currentSinkOp);
 			newAt = newOp;
 			break;
@@ -266,7 +266,7 @@ namespace mu
 
 		case EOpType::ME_CONDITIONAL:
 		{
-			Ptr<ASTOpConditional> newOp = mu::Clone<ASTOpConditional>(at);
+			Ptr<ASTOpConditional> newOp = UE::Mutable::Private::Clone<ASTOpConditional>(at);
 			newOp->yes = Visit(newOp->yes.child(), currentSinkOp);
 			newOp->no = Visit(newOp->no.child(), currentSinkOp);
 			newAt = newOp;
@@ -275,7 +275,7 @@ namespace mu
 
 		case EOpType::ME_SWITCH:
 		{
-			Ptr<ASTOpSwitch> newOp = mu::Clone<ASTOpSwitch>(at);
+			Ptr<ASTOpSwitch> newOp = UE::Mutable::Private::Clone<ASTOpSwitch>(at);
 			newOp->Default = Visit(newOp->Default.child(), currentSinkOp);
 			for (ASTOpSwitch::FCase& c : newOp->Cases)
 			{
@@ -289,7 +289,7 @@ namespace mu
 		default:
 			if (at != InitialSource)
 			{
-				mu::Ptr<ASTOpMeshExtractLayoutBlocks> newOp = mu::Clone<ASTOpMeshExtractLayoutBlocks>(currentSinkOp);
+				UE::Mutable::Private::Ptr<ASTOpMeshExtractLayoutBlocks> newOp = UE::Mutable::Private::Clone<ASTOpMeshExtractLayoutBlocks>(currentSinkOp);
 				newOp->Source = at;
 				newAt = newOp;
 			}
diff --git a/Engine/Plugins/Mutable/Source/MutableTools/Private/MuT/ASTOpMeshExtractLayoutBlocks.h b/Engine/Plugins/Mutable/Source/MutableTools/Private/MuT/ASTOpMeshExtractLayoutBlocks.h
--- a/Engine/Plugins/Mutable/Source/MutableTools/Private/MuT/ASTOpMeshExtractLayoutBlocks.h
+++ b/Engine/Plugins/Mutable/Source/MutableTools/Private/MuT/ASTOpMeshExtractLayoutBlocks.h
@@ -8,7 +8,7 @@
 #include "MuT/AST.h"
 
 
-namespace mu
+namespace UE::Mutable::Private
 {
 	struct FProgram;
 
@@ -43,7 +43,7 @@ namespace mu
 		Ptr<ASTOp> Clone(MapChildFuncRef mapChild) const override;
 		void Assert() override;
 		void Link(FProgram& program, FLinkerOptions* Options) override;
-		mu::Ptr<ASTOp> OptimiseSink(const FModelOptimizationOptions&, FOptimizeSinkContext&) const;
+		UE::Mutable::Private::Ptr<ASTOp> OptimiseSink(const FModelOptimizationOptions&, FOptimizeSinkContext&) const;
 		virtual FSourceDataDescriptor GetSourceDataDescriptor(FGetSourceDataDescriptorContext*) const override;
 
 	};
diff --git a/Engine/Plugins/Mutable/Source/MutableTools/Private/MuT/ASTOpMeshFormat.cpp b/Engine/Plugins/Mutable/Source/MutableTools/Private/MuT/ASTOpMeshFormat.cpp
--- a/Engine/Plugins/Mutable/Source/MutableTools/Private/MuT/ASTOpMeshFormat.cpp
+++ b/Engine/Plugins/Mutable/Source/MutableTools/Private/MuT/ASTOpMeshFormat.cpp
@@ -27,7 +27,7 @@
 
 #include "GPUSkinPublicDefs.h"
 
-namespace mu
+namespace UE::Mutable::Private
 {
 
 
@@ -65,9 +65,9 @@ uint64 ASTOpMeshFormat::Hash() const
 }
 
 
-mu::Ptr<ASTOp> ASTOpMeshFormat::Clone(MapChildFuncRef mapChild) const
+UE::Mutable::Private::Ptr<ASTOp> ASTOpMeshFormat::Clone(MapChildFuncRef mapChild) const
 {
-	mu::Ptr<ASTOpMeshFormat> n = new ASTOpMeshFormat();
+	UE::Mutable::Private::Ptr<ASTOpMeshFormat> n = new ASTOpMeshFormat();
     n->Source = mapChild(Source.child());
     n->Format = mapChild(Format.child());
 	n->Flags = Flags;
@@ -109,9 +109,9 @@ void ASTOpMeshFormat::Link( FProgram& program, FLinkerOptions* )
 }
 
 
-mu::Ptr<ASTOp> ASTOpMeshFormat::OptimiseSink(const FModelOptimizationOptions& options, FOptimizeSinkContext& context) const
+UE::Mutable::Private::Ptr<ASTOp> ASTOpMeshFormat::OptimiseSink(const FModelOptimizationOptions& options, FOptimizeSinkContext& context) const
 {
-	mu::Ptr<ASTOp> at = context.MeshFormatSinker.Apply(this);
+	UE::Mutable::Private::Ptr<ASTOp> at = context.MeshFormatSinker.Apply(this);
 	return at;
 }
 
@@ -130,14 +130,14 @@ FSourceDataDescriptor ASTOpMeshFormat::GetSourceDataDescriptor(FGetSourceDataDes
 //---------------------------------------------------------------------------------------------
 //---------------------------------------------------------------------------------------------
 //---------------------------------------------------------------------------------------------
-mu::Ptr<ASTOp> Sink_MeshFormatAST::Apply(const ASTOpMeshFormat* root)
+UE::Mutable::Private::Ptr<ASTOp> Sink_MeshFormatAST::Apply(const ASTOpMeshFormat* root)
 {
 	Root = root;
 
 	OldToNew.Reset();
 
 	InitialSource = Root->Source.child();
-	mu::Ptr<ASTOp> newSource = Visit(InitialSource, Root);
+	UE::Mutable::Private::Ptr<ASTOp> newSource = Visit(InitialSource, Root);
 
 	// If there is any change, it is the new root.
 	if (newSource != InitialSource)
@@ -151,7 +151,7 @@ mu::Ptr<ASTOp> Sink_MeshFormatAST::Apply(const ASTOpMeshFormat* root)
 
 namespace
 {
-	TSharedPtr<const FMesh> FindBaseMeshConstant(mu::Ptr<ASTOp> at)
+	TSharedPtr<const FMesh> FindBaseMeshConstant(UE::Mutable::Private::Ptr<ASTOp> at)
 	{
 		TSharedPtr<const FMesh> res;
 
@@ -276,7 +276,7 @@ namespace
 }
 
 
-mu::Ptr<ASTOp> Sink_MeshFormatAST::Visit(const mu::Ptr<ASTOp>& at, const ASTOpMeshFormat* currentFormatOp)
+UE::Mutable::Private::Ptr<ASTOp> Sink_MeshFormatAST::Visit(const UE::Mutable::Private::Ptr<ASTOp>& at, const ASTOpMeshFormat* currentFormatOp)
 {
 	if (!at)
 	{
@@ -290,13 +290,13 @@ mu::Ptr<ASTOp> Sink_MeshFormatAST::Visit(const mu::Ptr<ASTOp>& at, const ASTOpMe
 		return *Cached;
 	}
 
-	mu::Ptr<ASTOp> newAt = at;
+	UE::Mutable::Private::Ptr<ASTOp> newAt = at;
 	switch (at->GetOpType())
 	{
 
 	case EOpType::ME_APPLYLAYOUT:
 	{
-		Ptr<ASTOpMeshApplyLayout> NewOp = mu::Clone<ASTOpMeshApplyLayout>(at);
+		Ptr<ASTOpMeshApplyLayout> NewOp = UE::Mutable::Private::Clone<ASTOpMeshApplyLayout>(at);
 		NewOp->Mesh = Visit(NewOp->Mesh.child(), currentFormatOp);
 		newAt = NewOp;
 		break;
@@ -304,7 +304,7 @@ mu::Ptr<ASTOp> Sink_MeshFormatAST::Visit(const mu::Ptr<ASTOp>& at, const ASTOpMe
 
 	case EOpType::ME_SETSKELETON:
 	{
-		Ptr<ASTOpMeshSetSkeleton> NewOp = mu::Clone<ASTOpMeshSetSkeleton>(at);
+		Ptr<ASTOpMeshSetSkeleton> NewOp = UE::Mutable::Private::Clone<ASTOpMeshSetSkeleton>(at);
 		NewOp->Source = Visit(NewOp->Source.child(), currentFormatOp);
 		newAt = NewOp;
 		break;
@@ -312,7 +312,7 @@ mu::Ptr<ASTOp> Sink_MeshFormatAST::Visit(const mu::Ptr<ASTOp>& at, const ASTOpMe
 
 	case EOpType::ME_ADDMETADATA:
 	{
-		Ptr<ASTOpMeshAddMetadata> newOp = mu::Clone<ASTOpMeshAddMetadata>(at);
+		Ptr<ASTOpMeshAddMetadata> newOp = UE::Mutable::Private::Clone<ASTOpMeshAddMetadata>(at);
 		newOp->Source = Visit(newOp->Source.child(), currentFormatOp);
 		newAt = newOp;
 		break;
@@ -320,7 +320,7 @@ mu::Ptr<ASTOp> Sink_MeshFormatAST::Visit(const mu::Ptr<ASTOp>& at, const ASTOpMe
 
 	case EOpType::ME_CLIPMORPHPLANE:
 	{
-		Ptr<ASTOpMeshClipMorphPlane> newOp = mu::Clone<ASTOpMeshClipMorphPlane>(at);
+		Ptr<ASTOpMeshClipMorphPlane> newOp = UE::Mutable::Private::Clone<ASTOpMeshClipMorphPlane>(at);
 		newOp->Source = Visit(newOp->Source.child(), currentFormatOp);
 		newAt = newOp;
 		break;
@@ -328,7 +328,7 @@ mu::Ptr<ASTOp> Sink_MeshFormatAST::Visit(const mu::Ptr<ASTOp>& at, const ASTOpMe
 
 	case EOpType::ME_TRANSFORMWITHMESH:
 	{
-		Ptr<ASTOpMeshTransformWithBoundingMesh> NewOp = mu::Clone<ASTOpMeshTransformWithBoundingMesh>(at);
+		Ptr<ASTOpMeshTransformWithBoundingMesh> NewOp = UE::Mutable::Private::Clone<ASTOpMeshTransformWithBoundingMesh>(at);
 		NewOp->source = Visit(NewOp->source.child(), currentFormatOp);
 
 		// Don't transform the bounding mesh: it should be optimized with a different specific format elsewhere (TODO).
@@ -340,7 +340,7 @@ mu::Ptr<ASTOp> Sink_MeshFormatAST::Visit(const mu::Ptr<ASTOp>& at, const ASTOpMe
 
 	case EOpType::ME_TRANSFORMWITHBONE:
 	{
-		Ptr<ASTOpMeshTransformWithBone> NewOp = mu::Clone<ASTOpMeshTransformWithBone>(at);
+		Ptr<ASTOpMeshTransformWithBone> NewOp = UE::Mutable::Private::Clone<ASTOpMeshTransformWithBone>(at);
 		NewOp->SourceMesh = Visit(NewOp->SourceMesh.child(), currentFormatOp);
 		newAt = NewOp;
 		break;
@@ -349,21 +349,21 @@ mu::Ptr<ASTOp> Sink_MeshFormatAST::Visit(const mu::Ptr<ASTOp>& at, const ASTOpMe
 	case EOpType::ME_MORPH:
 	{
 		// Move the format down the base of the morph
-		Ptr<ASTOpMeshMorph> NewOp = mu::Clone<ASTOpMeshMorph>(at);
+		Ptr<ASTOpMeshMorph> NewOp = UE::Mutable::Private::Clone<ASTOpMeshMorph>(at);
 		NewOp->Base = Visit(NewOp->Base.child(), currentFormatOp);
 
 		// Reformat the morph targets to match the new format.
 		TSharedPtr<const FMesh> pTargetFormat = FindBaseMeshConstant(currentFormatOp->Format.child());
 		TSharedPtr<const FMesh> pTargetMorphFormat = MakeMorphTargetFormat(pTargetFormat);
 
-		mu::Ptr<ASTOpConstantResource> NewFormatConstant = new ASTOpConstantResource();
+		UE::Mutable::Private::Ptr<ASTOpConstantResource> NewFormatConstant = new ASTOpConstantResource();
 		NewFormatConstant->Type = EOpType::ME_CONSTANT;
 		NewFormatConstant->SetValue(pTargetMorphFormat, nullptr);
 		NewFormatConstant->SourceDataDescriptor = at->GetSourceDataDescriptor();
 
 		if (NewOp->Target)
 		{
-			mu::Ptr<ASTOpMeshFormat> newFormat = mu::Clone<ASTOpMeshFormat>(currentFormatOp);
+			UE::Mutable::Private::Ptr<ASTOpMeshFormat> newFormat = UE::Mutable::Private::Clone<ASTOpMeshFormat>(currentFormatOp);
 			newFormat->Flags = OP::MeshFormatArgs::Vertex | OP::MeshFormatArgs::IgnoreMissing;
 			newFormat->Format = NewFormatConstant;
 
@@ -376,7 +376,7 @@ mu::Ptr<ASTOp> Sink_MeshFormatAST::Visit(const mu::Ptr<ASTOp>& at, const ASTOpMe
 
 	case EOpType::ME_MERGE:
 	{
-		Ptr<ASTOpMeshMerge> NewOp = mu::Clone<ASTOpMeshMerge>(at);
+		Ptr<ASTOpMeshMerge> NewOp = UE::Mutable::Private::Clone<ASTOpMeshMerge>(at);
 		NewOp->Base = Visit(NewOp->Base.child(), currentFormatOp);
 		NewOp->Added = Visit(NewOp->Added.child(), currentFormatOp);
 		newAt = NewOp;
@@ -388,10 +388,10 @@ mu::Ptr<ASTOp> Sink_MeshFormatAST::Visit(const mu::Ptr<ASTOp>& at, const ASTOpMe
 		TSharedPtr<const FMesh> TargetFormat = FindBaseMeshConstant(currentFormatOp->Format.child());
         TargetFormat = EnsureFormatHasSkinningBuffers(TargetFormat);
 
-		Ptr<ASTOpMeshApplyPose> NewOp = mu::Clone<ASTOpMeshApplyPose>(at);
-        mu::Ptr<ASTOpMeshFormat> NewFormat = mu::Clone<ASTOpMeshFormat>(currentFormatOp);
+		Ptr<ASTOpMeshApplyPose> NewOp = UE::Mutable::Private::Clone<ASTOpMeshApplyPose>(at);
+        UE::Mutable::Private::Ptr<ASTOpMeshFormat> NewFormat = UE::Mutable::Private::Clone<ASTOpMeshFormat>(currentFormatOp);
         
-		mu::Ptr<ASTOpConstantResource> NewFormatConstant = new ASTOpConstantResource();
+		UE::Mutable::Private::Ptr<ASTOpConstantResource> NewFormatConstant = new ASTOpConstantResource();
 		NewFormatConstant->Type = EOpType::ME_CONSTANT;
 		NewFormatConstant->SetValue(TargetFormat, nullptr);
 		NewFormatConstant->SourceDataDescriptor = at->GetSourceDataDescriptor();
@@ -411,7 +411,7 @@ mu::Ptr<ASTOp> Sink_MeshFormatAST::Visit(const mu::Ptr<ASTOp>& at, const ASTOpMe
 
 	case EOpType::ME_REMOVEMASK:
 	{
-		Ptr<ASTOpMeshRemoveMask> newOp = mu::Clone<ASTOpMeshRemoveMask>(at);
+		Ptr<ASTOpMeshRemoveMask> newOp = UE::Mutable::Private::Clone<ASTOpMeshRemoveMask>(at);
 		newOp->source = Visit(newOp->source.child(), currentFormatOp);
 		newAt = newOp;
 		break;
@@ -419,7 +419,7 @@ mu::Ptr<ASTOp> Sink_MeshFormatAST::Visit(const mu::Ptr<ASTOp>& at, const ASTOpMe
 
 	case EOpType::ME_CONDITIONAL:
 	{
-		Ptr<ASTOpConditional> newOp = mu::Clone<ASTOpConditional>(at);
+		Ptr<ASTOpConditional> newOp = UE::Mutable::Private::Clone<ASTOpConditional>(at);
 		newOp->yes = Visit(newOp->yes.child(), currentFormatOp);
 		newOp->no = Visit(newOp->no.child(), currentFormatOp);
 		newAt = newOp;
@@ -428,7 +428,7 @@ mu::Ptr<ASTOp> Sink_MeshFormatAST::Visit(const mu::Ptr<ASTOp>& at, const ASTOpMe
 
 	case EOpType::ME_SWITCH:
 	{
-		Ptr<ASTOpSwitch> newOp = mu::Clone<ASTOpSwitch>(at);
+		Ptr<ASTOpSwitch> newOp = UE::Mutable::Private::Clone<ASTOpSwitch>(at);
 		newOp->Default = Visit(newOp->Default.child(), currentFormatOp);
 		for (ASTOpSwitch::FCase& c : newOp->Cases)
 		{
@@ -451,7 +451,7 @@ mu::Ptr<ASTOp> Sink_MeshFormatAST::Visit(const mu::Ptr<ASTOp>& at, const ASTOpMe
 	default:
 		if (at != InitialSource)
 		{
-			mu::Ptr<ASTOpMeshFormat> newOp = mu::Clone<ASTOpMeshFormat>(currentFormatOp);
+			UE::Mutable::Private::Ptr<ASTOpMeshFormat> newOp = UE::Mutable::Private::Clone<ASTOpMeshFormat>(currentFormatOp);
 			newOp->Source = at;
 			newAt = newOp;
 		}
diff --git a/Engine/Plugins/Mutable/Source/MutableTools/Private/MuT/ASTOpMeshMaskClipMesh.cpp b/Engine/Plugins/Mutable/Source/MutableTools/Private/MuT/ASTOpMeshMaskClipMesh.cpp
--- a/Engine/Plugins/Mutable/Source/MutableTools/Private/MuT/ASTOpMeshMaskClipMesh.cpp
+++ b/Engine/Plugins/Mutable/Source/MutableTools/Private/MuT/ASTOpMeshMaskClipMesh.cpp
@@ -13,7 +13,7 @@
 #include "MuT/ASTOpSwitch.h"
 
 
-namespace mu
+namespace UE::Mutable::Private
 {
 
 
@@ -57,7 +57,7 @@ namespace mu
 
 
 	//-------------------------------------------------------------------------------------------------
-	mu::Ptr<ASTOp> ASTOpMeshMaskClipMesh::Clone(MapChildFuncRef mapChild) const
+	UE::Mutable::Private::Ptr<ASTOp> ASTOpMeshMaskClipMesh::Clone(MapChildFuncRef mapChild) const
 	{
 		Ptr<ASTOpMeshMaskClipMesh> n = new ASTOpMeshMaskClipMesh();
 		n->source = mapChild(source.child());
@@ -106,13 +106,13 @@ namespace mu
 
 			// \TODO This is recursive and may cause stack overflows in big models.
 
-			mu::Ptr<ASTOp> Apply(const ASTOpMeshMaskClipMesh* root)
+			UE::Mutable::Private::Ptr<ASTOp> Apply(const ASTOpMeshMaskClipMesh* root)
 			{
 				Root = root;
 				m_oldToNew.Empty();
 
 				InitialSource = root->source.child();
-				mu::Ptr<ASTOp> NewSource = Visit(InitialSource);
+				UE::Mutable::Private::Ptr<ASTOp> NewSource = Visit(InitialSource);
 
 				// If there is any change, it is the new root.
 				if (NewSource != InitialSource)
@@ -126,11 +126,11 @@ namespace mu
 		protected:
 
 			const ASTOpMeshMaskClipMesh* Root;
-			mu::Ptr<ASTOp> InitialSource;
-			TMap<mu::Ptr<ASTOp>, mu::Ptr<ASTOp>> m_oldToNew;
-			TArray<mu::Ptr<ASTOp>> m_newOps;
+			UE::Mutable::Private::Ptr<ASTOp> InitialSource;
+			TMap<UE::Mutable::Private::Ptr<ASTOp>, UE::Mutable::Private::Ptr<ASTOp>> m_oldToNew;
+			TArray<UE::Mutable::Private::Ptr<ASTOp>> m_newOps;
 
-			mu::Ptr<ASTOp> Visit(const mu::Ptr<ASTOp>& at)
+			UE::Mutable::Private::Ptr<ASTOp> Visit(const UE::Mutable::Private::Ptr<ASTOp>& at)
 			{
 				if (!at) return nullptr;
 
@@ -141,13 +141,13 @@ namespace mu
 				}
 
 				// Already visited?
-				mu::Ptr<ASTOp>* cacheIt = m_oldToNew.Find(at);
+				UE::Mutable::Private::Ptr<ASTOp>* cacheIt = m_oldToNew.Find(at);
 				if (cacheIt)
 				{
 					return *cacheIt;
 				}
 
-				mu::Ptr<ASTOp> newAt = at;
+				UE::Mutable::Private::Ptr<ASTOp> newAt = at;
 				switch (at->GetOpType())
 				{
 
@@ -181,7 +181,7 @@ namespace mu
 
 				case EOpType::ME_ADDMETADATA:
 				{
-					Ptr<ASTOpMeshAddMetadata> newOp = mu::Clone<ASTOpMeshAddMetadata>(at);
+					Ptr<ASTOpMeshAddMetadata> newOp = UE::Mutable::Private::Clone<ASTOpMeshAddMetadata>(at);
 					newOp->Source = Visit(newOp->Source.child());
 					newAt = newOp;
 					break;
@@ -193,7 +193,7 @@ namespace mu
 					// It always needs to be a clone because otherwise we could be modifying an
 					// instruction that shouldn't if the parent was a ME_REMOVEMASK above and we
 					// skipped the cloning for the parent.
-					Ptr<ASTOpConditional> newOp = mu::Clone<ASTOpConditional>(at);
+					Ptr<ASTOpConditional> newOp = UE::Mutable::Private::Clone<ASTOpConditional>(at);
 					newOp->yes = Visit(newOp->yes.child());
 					newOp->no = Visit(newOp->no.child());
 					newAt = newOp;
@@ -203,7 +203,7 @@ namespace mu
 				case EOpType::ME_SWITCH:
 				{
 					// We move the mask creation down all the paths
-					Ptr<ASTOpSwitch> newOp = mu::Clone<ASTOpSwitch>(at);
+					Ptr<ASTOpSwitch> newOp = UE::Mutable::Private::Clone<ASTOpSwitch>(at);
 					newOp->Default = Visit(newOp->Default.child());
 					for (ASTOpSwitch::FCase& c : newOp->Cases)
 					{
@@ -229,7 +229,7 @@ namespace mu
 					//
 					if (at != InitialSource)
 					{
-						Ptr<ASTOpMeshMaskClipMesh> newOp = mu::Clone<ASTOpMeshMaskClipMesh>(Root);
+						Ptr<ASTOpMeshMaskClipMesh> newOp = UE::Mutable::Private::Clone<ASTOpMeshMaskClipMesh>(Root);
 						newOp->source = at;
 						newAt = newOp;
 					}
@@ -251,13 +251,13 @@ namespace mu
 
 			// \TODO This is recursive and may cause stack overflows in big models.
 
-			mu::Ptr<ASTOp> Apply(const ASTOpMeshMaskClipMesh* root)
+			UE::Mutable::Private::Ptr<ASTOp> Apply(const ASTOpMeshMaskClipMesh* root)
 			{
 				Root = root;
 				m_oldToNew.Empty();
 
 				m_initialClip = root->clip.child();
-				mu::Ptr<ASTOp> NewClip = Visit(m_initialClip);
+				UE::Mutable::Private::Ptr<ASTOp> NewClip = Visit(m_initialClip);
 
 				// If there is any change, it is the new root.
 				if (NewClip != m_initialClip)
@@ -271,11 +271,11 @@ namespace mu
 		protected:
 
 			const ASTOpMeshMaskClipMesh* Root;
-			mu::Ptr<ASTOp> m_initialClip;
-			TMap<mu::Ptr<ASTOp>, mu::Ptr<ASTOp>> m_oldToNew;
-			TArray<mu::Ptr<ASTOp>> m_newOps;
+			UE::Mutable::Private::Ptr<ASTOp> m_initialClip;
+			TMap<UE::Mutable::Private::Ptr<ASTOp>, UE::Mutable::Private::Ptr<ASTOp>> m_oldToNew;
+			TArray<UE::Mutable::Private::Ptr<ASTOp>> m_newOps;
 
-			mu::Ptr<ASTOp> Visit(const mu::Ptr<ASTOp>& at)
+			UE::Mutable::Private::Ptr<ASTOp> Visit(const UE::Mutable::Private::Ptr<ASTOp>& at)
 			{
 				if (!at) return nullptr;
 
@@ -286,13 +286,13 @@ namespace mu
 				}
 
 				// Already visited?
-				mu::Ptr<ASTOp>* cacheIt = m_oldToNew.Find(at);
+				UE::Mutable::Private::Ptr<ASTOp>* cacheIt = m_oldToNew.Find(at);
 				if (cacheIt)
 				{
 					return *cacheIt;
 				}
 
-				mu::Ptr<ASTOp> newAt = at;
+				UE::Mutable::Private::Ptr<ASTOp> newAt = at;
 				switch (at->GetOpType())
 				{
 
@@ -302,7 +302,7 @@ namespace mu
 					// It always needs to be a clone because otherwise we could be modifying an
 					// instruction that shouldn't if the parent was a ME_REMOVEMASK above and we
 					// skipped the cloning for the parent.
-					Ptr<ASTOpConditional> newOp = mu::Clone<ASTOpConditional>(at);
+					Ptr<ASTOpConditional> newOp = UE::Mutable::Private::Clone<ASTOpConditional>(at);
 					newOp->yes = Visit(newOp->yes.child());
 					newOp->no = Visit(newOp->no.child());
 					newAt = newOp;
@@ -312,7 +312,7 @@ namespace mu
 				case EOpType::ME_SWITCH:
 				{
 					// We move the mask creation down all the paths
-					Ptr<ASTOpSwitch> newOp = mu::Clone<ASTOpSwitch>(at);
+					Ptr<ASTOpSwitch> newOp = UE::Mutable::Private::Clone<ASTOpSwitch>(at);
 					newOp->Default = Visit(newOp->Default.child());
 					for (ASTOpSwitch::FCase& c : newOp->Cases)
 					{
@@ -335,7 +335,7 @@ namespace mu
 					//
 					if (at != m_initialClip)
 					{
-						Ptr<ASTOpMeshMaskClipMesh> newOp = mu::Clone<ASTOpMeshMaskClipMesh>(Root);
+						Ptr<ASTOpMeshMaskClipMesh> newOp = UE::Mutable::Private::Clone<ASTOpMeshMaskClipMesh>(Root);
 						newOp->clip = at;
 						newAt = newOp;
 					}
@@ -354,13 +354,13 @@ namespace mu
 
 
 	//-------------------------------------------------------------------------------------------------
-	mu::Ptr<ASTOp> ASTOpMeshMaskClipMesh::OptimiseSink(const FModelOptimizationOptions&, FOptimizeSinkContext&) const
+	UE::Mutable::Private::Ptr<ASTOp> ASTOpMeshMaskClipMesh::OptimiseSink(const FModelOptimizationOptions&, FOptimizeSinkContext&) const
 	{
 		// \TODO: Add logic state to the sinkers to avoid explosion with switches in both branches and similar cases.
 
 		// This will sink the operation down the source
 		Sink_MeshMaskClipMeshSource SinkerSource;
-		mu::Ptr<ASTOp> at = SinkerSource.Apply(this);
+		UE::Mutable::Private::Ptr<ASTOp> at = SinkerSource.Apply(this);
 
 		// If we didn't sink it.
 		if (!at || at == this)
diff --git a/Engine/Plugins/Mutable/Source/MutableTools/Private/MuT/ASTOpMeshMaskClipUVMask.cpp b/Engine/Plugins/Mutable/Source/MutableTools/Private/MuT/ASTOpMeshMaskClipUVMask.cpp
--- a/Engine/Plugins/Mutable/Source/MutableTools/Private/MuT/ASTOpMeshMaskClipUVMask.cpp
+++ b/Engine/Plugins/Mutable/Source/MutableTools/Private/MuT/ASTOpMeshMaskClipUVMask.cpp
@@ -13,7 +13,7 @@
 #include "MuT/ASTOpSwitch.h"
 
 
-namespace mu
+namespace UE::Mutable::Private
 {
 
 
@@ -65,7 +65,7 @@ namespace mu
 
 
 	//-------------------------------------------------------------------------------------------------
-	mu::Ptr<ASTOp> ASTOpMeshMaskClipUVMask::Clone(MapChildFuncRef mapChild) const
+	UE::Mutable::Private::Ptr<ASTOp> ASTOpMeshMaskClipUVMask::Clone(MapChildFuncRef mapChild) const
 	{
 		Ptr<ASTOpMeshMaskClipUVMask> n = new ASTOpMeshMaskClipUVMask();
 		n->Source = mapChild(Source.child());
@@ -121,13 +121,13 @@ namespace mu
 
 			// \TODO This is recursive and may cause stack overflows in big models.
 
-			mu::Ptr<ASTOp> Apply(const ASTOpMeshMaskClipUVMask* root)
+			UE::Mutable::Private::Ptr<ASTOp> Apply(const ASTOpMeshMaskClipUVMask* root)
 			{
 				Root = root;
 				m_oldToNew.Empty();
 
 				InitialSource = root->Source.child();
-				mu::Ptr<ASTOp> NewSource = Visit(InitialSource);
+				UE::Mutable::Private::Ptr<ASTOp> NewSource = Visit(InitialSource);
 
 				// If there is any change, it is the new root.
 				if (NewSource != InitialSource)
@@ -141,11 +141,11 @@ namespace mu
 		protected:
 
 			const ASTOpMeshMaskClipUVMask* Root;
-			mu::Ptr<ASTOp> InitialSource;
-			TMap<mu::Ptr<ASTOp>, mu::Ptr<ASTOp>> m_oldToNew;
-			TArray<mu::Ptr<ASTOp>> m_newOps;
+			UE::Mutable::Private::Ptr<ASTOp> InitialSource;
+			TMap<UE::Mutable::Private::Ptr<ASTOp>, UE::Mutable::Private::Ptr<ASTOp>> m_oldToNew;
+			TArray<UE::Mutable::Private::Ptr<ASTOp>> m_newOps;
 
-			mu::Ptr<ASTOp> Visit(const mu::Ptr<ASTOp>& at)
+			UE::Mutable::Private::Ptr<ASTOp> Visit(const UE::Mutable::Private::Ptr<ASTOp>& at)
 			{
 				if (!at) return nullptr;
 
@@ -156,13 +156,13 @@ namespace mu
 				}
 
 				// Already visited?
-				mu::Ptr<ASTOp>* cacheIt = m_oldToNew.Find(at);
+				UE::Mutable::Private::Ptr<ASTOp>* cacheIt = m_oldToNew.Find(at);
 				if (cacheIt)
 				{
 					return *cacheIt;
 				}
 
-				mu::Ptr<ASTOp> newAt = at;
+				UE::Mutable::Private::Ptr<ASTOp> newAt = at;
 				switch (at->GetOpType())
 				{
 
@@ -188,7 +188,7 @@ namespace mu
 
 				case EOpType::ME_ADDMETADATA:
 				{
-					Ptr<ASTOpMeshAddMetadata> newOp = mu::Clone<ASTOpMeshAddMetadata>(at);
+					Ptr<ASTOpMeshAddMetadata> newOp = UE::Mutable::Private::Clone<ASTOpMeshAddMetadata>(at);
 					newOp->Source = Visit(newOp->Source.child());
 					newAt = newOp;
 					break;
@@ -208,7 +208,7 @@ namespace mu
 					// It always needs to be a clone because otherwise we could be modifying an
 					// instruction that shouldn't if the parent was a ME_REMOVEMASK above and we
 					// skipped the cloning for the parent.
-					Ptr<ASTOpConditional> newOp = mu::Clone<ASTOpConditional>(at);
+					Ptr<ASTOpConditional> newOp = UE::Mutable::Private::Clone<ASTOpConditional>(at);
 					newOp->yes = Visit(newOp->yes.child());
 					newOp->no = Visit(newOp->no.child());
 					newAt = newOp;
@@ -218,7 +218,7 @@ namespace mu
 				case EOpType::ME_SWITCH:
 				{
 					// We move the mask creation down all the paths
-					Ptr<ASTOpSwitch> newOp = mu::Clone<ASTOpSwitch>(at);
+					Ptr<ASTOpSwitch> newOp = UE::Mutable::Private::Clone<ASTOpSwitch>(at);
 					newOp->Default = Visit(newOp->Default.child());
 					for (ASTOpSwitch::FCase& c : newOp->Cases)
 					{
@@ -244,7 +244,7 @@ namespace mu
 					//
 					if (at != InitialSource)
 					{
-						Ptr<ASTOpMeshMaskClipUVMask> newOp = mu::Clone<ASTOpMeshMaskClipUVMask>(Root);
+						Ptr<ASTOpMeshMaskClipUVMask> newOp = UE::Mutable::Private::Clone<ASTOpMeshMaskClipUVMask>(Root);
 						newOp->Source = at;
 						newAt = newOp;
 					}
@@ -266,13 +266,13 @@ namespace mu
 
 			// \TODO This is recursive and may cause stack overflows in big models.
 
-			mu::Ptr<ASTOp> Apply(const ASTOpMeshMaskClipUVMask* root)
+			UE::Mutable::Private::Ptr<ASTOp> Apply(const ASTOpMeshMaskClipUVMask* root)
 			{
 				Root = root;
 				m_oldToNew.Empty();
 
 				m_initialClip = root->MaskImage.child();
-				mu::Ptr<ASTOp> NewClip = Visit(m_initialClip);
+				UE::Mutable::Private::Ptr<ASTOp> NewClip = Visit(m_initialClip);
 
 				// If there is any change, it is the new root.
 				if (NewClip != m_initialClip)
@@ -286,11 +286,11 @@ namespace mu
 		protected:
 
 			const ASTOpMeshMaskClipUVMask* Root;
-			mu::Ptr<ASTOp> m_initialClip;
-			TMap<mu::Ptr<ASTOp>, mu::Ptr<ASTOp>> m_oldToNew;
-			TArray<mu::Ptr<ASTOp>> m_newOps;
+			UE::Mutable::Private::Ptr<ASTOp> m_initialClip;
+			TMap<UE::Mutable::Private::Ptr<ASTOp>, UE::Mutable::Private::Ptr<ASTOp>> m_oldToNew;
+			TArray<UE::Mutable::Private::Ptr<ASTOp>> m_newOps;
 
-			mu::Ptr<ASTOp> Visit(const mu::Ptr<ASTOp>& at)
+			UE::Mutable::Private::Ptr<ASTOp> Visit(const UE::Mutable::Private::Ptr<ASTOp>& at)
 			{
 				if (!at) return nullptr;
 
@@ -301,13 +301,13 @@ namespace mu
 				}
 
 				// Already visited?
-				mu::Ptr<ASTOp>* cacheIt = m_oldToNew.Find(at);
+				UE::Mutable::Private::Ptr<ASTOp>* cacheIt = m_oldToNew.Find(at);
 				if (cacheIt)
 				{
 					return *cacheIt;
 				}
 
-				mu::Ptr<ASTOp> newAt = at;
+				UE::Mutable::Private::Ptr<ASTOp> newAt = at;
 				switch (at->GetOpType())
 				{
 
@@ -317,7 +317,7 @@ namespace mu
 					// It always needs to be a clone because otherwise we could be modifying an
 					// instruction that shouldn't if the parent was a ME_REMOVEMASK above and we
 					// skipped the cloning for the parent.
-					Ptr<ASTOpConditional> newOp = mu::Clone<ASTOpConditional>(at);
+					Ptr<ASTOpConditional> newOp = UE::Mutable::Private::Clone<ASTOpConditional>(at);
 					newOp->yes = Visit(newOp->yes.child());
 					newOp->no = Visit(newOp->no.child());
 					newAt = newOp;
@@ -327,7 +327,7 @@ namespace mu
 				case EOpType::ME_SWITCH:
 				{
 					// We move the mask creation down all the paths
-					Ptr<ASTOpSwitch> newOp = mu::Clone<ASTOpSwitch>(at);
+					Ptr<ASTOpSwitch> newOp = UE::Mutable::Private::Clone<ASTOpSwitch>(at);
 					newOp->Default = Visit(newOp->Default.child());
 					for (ASTOpSwitch::FCase& c : newOp->Cases)
 					{
@@ -343,7 +343,7 @@ namespace mu
 					//
 					if (at != m_initialClip)
 					{
-						Ptr<ASTOpMeshMaskClipUVMask> newOp = mu::Clone<ASTOpMeshMaskClipUVMask>(Root);
+						Ptr<ASTOpMeshMaskClipUVMask> newOp = UE::Mutable::Private::Clone<ASTOpMeshMaskClipUVMask>(Root);
 						newOp->MaskImage = at;
 						newAt = newOp;
 					}
@@ -362,13 +362,13 @@ namespace mu
 
 
 	//-------------------------------------------------------------------------------------------------
-	mu::Ptr<ASTOp> ASTOpMeshMaskClipUVMask::OptimiseSink(const FModelOptimizationOptions&, FOptimizeSinkContext&) const
+	UE::Mutable::Private::Ptr<ASTOp> ASTOpMeshMaskClipUVMask::OptimiseSink(const FModelOptimizationOptions&, FOptimizeSinkContext&) const
 	{
 		// \TODO: Add logic state to the sinkers to avoid explosion with switches in both branches and similar cases.
 
 		// This will sink the operation down the Source
 		Sink_MeshMaskClipUVMaskSource SinkerSource;
-		mu::Ptr<ASTOp> at = SinkerSource.Apply(this);
+		UE::Mutable::Private::Ptr<ASTOp> at = SinkerSource.Apply(this);
 
 		// If we didn't sink it.
 		if (!at || at == this)
diff --git a/Engine/Plugins/Mutable/Source/MutableTools/Private/MuT/ASTOpMeshMaskDiff.cpp b/Engine/Plugins/Mutable/Source/MutableTools/Private/MuT/ASTOpMeshMaskDiff.cpp
--- a/Engine/Plugins/Mutable/Source/MutableTools/Private/MuT/ASTOpMeshMaskDiff.cpp
+++ b/Engine/Plugins/Mutable/Source/MutableTools/Private/MuT/ASTOpMeshMaskDiff.cpp
@@ -9,7 +9,7 @@
 #include "MuR/Types.h"
 #include "HAL/PlatformMath.h"
 
-namespace mu
+namespace UE::Mutable::Private
 {
 
 	ASTOpMeshMaskDiff::ASTOpMeshMaskDiff()
@@ -81,14 +81,14 @@ namespace mu
 	}
 
 
-	mu::Ptr<ASTOp> ASTOpMeshMaskDiff::OptimiseSemantic(const FModelOptimizationOptions&, int32 Pass) const
+	UE::Mutable::Private::Ptr<ASTOp> ASTOpMeshMaskDiff::OptimiseSemantic(const FModelOptimizationOptions&, int32 Pass) const
 	{
 		if (!Fragment)
 		{
 			return nullptr;
 		}
 
-		mu::Ptr<ASTOp> ResultOp;
+		UE::Mutable::Private::Ptr<ASTOp> ResultOp;
 
 		EOpType FragmentType = Fragment->GetOpType();
 		switch (FragmentType)
@@ -99,7 +99,7 @@ namespace mu
 			// Tags in the fragment can be ignored.
 			const ASTOpMeshAddMetadata* Add = static_cast<const ASTOpMeshAddMetadata*>(Fragment.child().get());
 
-			Ptr<ASTOpMeshMaskDiff> NewAt = mu::Clone<ASTOpMeshMaskDiff>(this);
+			Ptr<ASTOpMeshMaskDiff> NewAt = UE::Mutable::Private::Clone<ASTOpMeshMaskDiff>(this);
 			NewAt->Fragment = Add->Source.child();
 			ResultOp = NewAt;
 			break;
@@ -110,7 +110,7 @@ namespace mu
 			// Ignore the prepare op in the source: it doesn't contribute to the diff generation.
 			const ASTOpMeshPrepareLayout* Prepare = static_cast<const ASTOpMeshPrepareLayout*>(Source.child().get());
 
-			Ptr<ASTOpMeshMaskDiff> NewAt = mu::Clone<ASTOpMeshMaskDiff>(this);
+			Ptr<ASTOpMeshMaskDiff> NewAt = UE::Mutable::Private::Clone<ASTOpMeshMaskDiff>(this);
 			NewAt->Source = Prepare->Mesh.child();
 			ResultOp = NewAt;
 			break;
@@ -119,11 +119,11 @@ namespace mu
 		case EOpType::ME_SWITCH:
 		{
 			// Move the mask diff down all the paths
-			Ptr<ASTOpSwitch> NewOp = mu::Clone<ASTOpSwitch>(Fragment.child().get());
+			Ptr<ASTOpSwitch> NewOp = UE::Mutable::Private::Clone<ASTOpSwitch>(Fragment.child().get());
 
 			if (NewOp->Default)
 			{
-				Ptr<ASTOpMeshMaskDiff> defOp = mu::Clone<ASTOpMeshMaskDiff>(this);
+				Ptr<ASTOpMeshMaskDiff> defOp = UE::Mutable::Private::Clone<ASTOpMeshMaskDiff>(this);
 				defOp->Fragment = NewOp->Default.child();
 				NewOp->Default = defOp;
 			}
@@ -133,7 +133,7 @@ namespace mu
 			{
 				if (NewOp->Cases[v].Branch)
 				{
-					Ptr<ASTOpMeshMaskDiff> bOp = mu::Clone<ASTOpMeshMaskDiff>(this);
+					Ptr<ASTOpMeshMaskDiff> bOp = UE::Mutable::Private::Clone<ASTOpMeshMaskDiff>(this);
 					bOp->Fragment = NewOp->Cases[v].Branch.child();
 					NewOp->Cases[v].Branch = bOp;
 				}
@@ -159,7 +159,7 @@ namespace mu
 				// Tags in the source can be ignored for mask generation.
 				const ASTOpMeshAddMetadata* Add = static_cast<const ASTOpMeshAddMetadata*>(Source.child().get());
 
-				Ptr<ASTOpMeshMaskDiff> NewAt = mu::Clone<ASTOpMeshMaskDiff>(this);
+				Ptr<ASTOpMeshMaskDiff> NewAt = UE::Mutable::Private::Clone<ASTOpMeshMaskDiff>(this);
 				NewAt->Source = Add->Source.child();
 				ResultOp = NewAt;
 				break;
@@ -170,7 +170,7 @@ namespace mu
 				// Ignore the prepare op in the source: it doesn't contribute to the diff generation.
 				const ASTOpMeshPrepareLayout* Prepare = static_cast<const ASTOpMeshPrepareLayout*>(Source.child().get());
 
-				Ptr<ASTOpMeshMaskDiff> NewAt = mu::Clone<ASTOpMeshMaskDiff>(this);
+				Ptr<ASTOpMeshMaskDiff> NewAt = UE::Mutable::Private::Clone<ASTOpMeshMaskDiff>(this);
 				NewAt->Source = Prepare->Mesh.child();
 				ResultOp = NewAt;
 				break;
diff --git a/Engine/Plugins/Mutable/Source/MutableTools/Private/MuT/ASTOpMeshMorph.cpp b/Engine/Plugins/Mutable/Source/MutableTools/Private/MuT/ASTOpMeshMorph.cpp
--- a/Engine/Plugins/Mutable/Source/MutableTools/Private/MuT/ASTOpMeshMorph.cpp
+++ b/Engine/Plugins/Mutable/Source/MutableTools/Private/MuT/ASTOpMeshMorph.cpp
@@ -9,7 +9,7 @@
 #include "HAL/PlatformMath.h"
 
 
-namespace mu
+namespace UE::Mutable::Private
 {
 
 	//---------------------------------------------------------------------------------------------
@@ -40,7 +40,7 @@ namespace mu
 
 
 	//---------------------------------------------------------------------------------------------
-	mu::Ptr<ASTOp> ASTOpMeshMorph::Clone(MapChildFuncRef mapChild) const
+	UE::Mutable::Private::Ptr<ASTOp> ASTOpMeshMorph::Clone(MapChildFuncRef mapChild) const
 	{
 		Ptr<ASTOpMeshMorph> n = new ASTOpMeshMorph();
 		n->Name = Name;
@@ -98,7 +98,7 @@ namespace mu
 
 
 	//-------------------------------------------------------------------------------------------------
-	mu::Ptr<ASTOp> ASTOpMeshMorph::OptimiseSink(const FModelOptimizationOptions&, FOptimizeSinkContext&) const
+	UE::Mutable::Private::Ptr<ASTOp> ASTOpMeshMorph::OptimiseSink(const FModelOptimizationOptions&, FOptimizeSinkContext&) const
 	{
 		Ptr<ASTOp> NewOp;
 
@@ -115,11 +115,11 @@ namespace mu
 		case EOpType::ME_ADDMETADATA:
 		{
 			// Add the base tags after morphing
-			Ptr<ASTOpMeshAddMetadata> NewAddMetadata = mu::Clone<ASTOpMeshAddMetadata>(Base.child());
+			Ptr<ASTOpMeshAddMetadata> NewAddMetadata = UE::Mutable::Private::Clone<ASTOpMeshAddMetadata>(Base.child());
 
 			if (NewAddMetadata->Source)
 			{
-				Ptr<ASTOpMeshMorph> New = mu::Clone<ASTOpMeshMorph>(this);
+				Ptr<ASTOpMeshMorph> New = UE::Mutable::Private::Clone<ASTOpMeshMorph>(this);
 				New->Base = NewAddMetadata->Source.child();
 				NewAddMetadata->Source = New;
 			}
@@ -146,7 +146,7 @@ namespace mu
 				// Ignore the morph target tags
 				const ASTOpMeshAddMetadata* AddMetadata = static_cast<const ASTOpMeshAddMetadata*>(Target.child().get());
 
-				Ptr<ASTOpMeshMorph> New = mu::Clone<ASTOpMeshMorph>(this);
+				Ptr<ASTOpMeshMorph> New = UE::Mutable::Private::Clone<ASTOpMeshMorph>(this);
 				New->Target = AddMetadata->Source.child();
 				NewOp = New;
 
diff --git a/Engine/Plugins/Mutable/Source/MutableTools/Private/MuT/ASTOpMeshMorphReshape.cpp b/Engine/Plugins/Mutable/Source/MutableTools/Private/MuT/ASTOpMeshMorphReshape.cpp
--- a/Engine/Plugins/Mutable/Source/MutableTools/Private/MuT/ASTOpMeshMorphReshape.cpp
+++ b/Engine/Plugins/Mutable/Source/MutableTools/Private/MuT/ASTOpMeshMorphReshape.cpp
@@ -8,7 +8,7 @@
 #include "MuR/Types.h"
 
 
-namespace mu
+namespace UE::Mutable::Private
 {
 
 
@@ -47,7 +47,7 @@ namespace mu
 	}
 
 
-	mu::Ptr<ASTOp> ASTOpMeshMorphReshape::Clone(MapChildFuncRef mapChild) const
+	UE::Mutable::Private::Ptr<ASTOp> ASTOpMeshMorphReshape::Clone(MapChildFuncRef mapChild) const
 	{
 		Ptr<ASTOpMeshMorphReshape> n = new ASTOpMeshMorphReshape();
 		n->Morph = mapChild(Morph.child());
diff --git a/Engine/Plugins/Mutable/Source/MutableTools/Private/MuT/ASTOpMeshOptimizeSkinning.cpp b/Engine/Plugins/Mutable/Source/MutableTools/Private/MuT/ASTOpMeshOptimizeSkinning.cpp
--- a/Engine/Plugins/Mutable/Source/MutableTools/Private/MuT/ASTOpMeshOptimizeSkinning.cpp
+++ b/Engine/Plugins/Mutable/Source/MutableTools/Private/MuT/ASTOpMeshOptimizeSkinning.cpp
@@ -13,7 +13,7 @@
 #include "HAL/PlatformMath.h"
 
 
-namespace mu
+namespace UE::Mutable::Private
 {
 
 	ASTOpMeshOptimizeSkinning::ASTOpMeshOptimizeSkinning()
@@ -47,7 +47,7 @@ namespace mu
 	}
 
 
-	mu::Ptr<ASTOp> ASTOpMeshOptimizeSkinning::Clone(MapChildFuncRef MapChild) const
+	UE::Mutable::Private::Ptr<ASTOp> ASTOpMeshOptimizeSkinning::Clone(MapChildFuncRef MapChild) const
 	{
 		Ptr<ASTOpMeshOptimizeSkinning> n = new ASTOpMeshOptimizeSkinning();
 		n->Source = MapChild(Source.child());
@@ -82,7 +82,7 @@ namespace mu
 	}
 
 
-	mu::Ptr<ASTOp> ASTOpMeshOptimizeSkinning::OptimiseSink(const FModelOptimizationOptions&, FOptimizeSinkContext& Context) const
+	UE::Mutable::Private::Ptr<ASTOp> ASTOpMeshOptimizeSkinning::OptimiseSink(const FModelOptimizationOptions&, FOptimizeSinkContext& Context) const
 	{
 		return Context.MeshOptimizeSkinningSinker.Apply(this);
 	}
@@ -104,14 +104,14 @@ namespace mu
 	//---------------------------------------------------------------------------------------------
 	//---------------------------------------------------------------------------------------------
 	//---------------------------------------------------------------------------------------------
-	mu::Ptr<ASTOp> Sink_MeshOptimizeSkinningAST::Apply(const ASTOpMeshOptimizeSkinning* InRoot)
+	UE::Mutable::Private::Ptr<ASTOp> Sink_MeshOptimizeSkinningAST::Apply(const ASTOpMeshOptimizeSkinning* InRoot)
 	{
 		Root = InRoot;
 
 		OldToNew.Reset();
 
 		InitialSource = Root->Source.child();
-		mu::Ptr<ASTOp> NewSource = Visit(InitialSource, Root);
+		UE::Mutable::Private::Ptr<ASTOp> NewSource = Visit(InitialSource, Root);
 
 		// If there is any change, it is the new root.
 		if (NewSource != InitialSource)
@@ -123,7 +123,7 @@ namespace mu
 	}
 
 
-	mu::Ptr<ASTOp> Sink_MeshOptimizeSkinningAST::Visit(const mu::Ptr<ASTOp>& at, const ASTOpMeshOptimizeSkinning* CurrentSinkOp)
+	UE::Mutable::Private::Ptr<ASTOp> Sink_MeshOptimizeSkinningAST::Visit(const UE::Mutable::Private::Ptr<ASTOp>& at, const ASTOpMeshOptimizeSkinning* CurrentSinkOp)
 	{
 		if (!at)
 		{
@@ -140,13 +140,13 @@ namespace mu
 			return *Cached;
 		}
 
-		mu::Ptr<ASTOp> NewAt = at;
+		UE::Mutable::Private::Ptr<ASTOp> NewAt = at;
 		switch (at->GetOpType())
 		{
 
 		case EOpType::ME_APPLYLAYOUT:
 		{
-			Ptr<ASTOpMeshApplyLayout> NewOp = mu::Clone<ASTOpMeshApplyLayout>(at);
+			Ptr<ASTOpMeshApplyLayout> NewOp = UE::Mutable::Private::Clone<ASTOpMeshApplyLayout>(at);
 			NewOp->Mesh = Visit(NewOp->Mesh.child(), CurrentSinkOp);
 			NewAt = NewOp;
 			break;
@@ -154,7 +154,7 @@ namespace mu
 
 		case EOpType::ME_ADDMETADATA:
 		{
-			Ptr<ASTOpMeshAddMetadata> NewOp = mu::Clone<ASTOpMeshAddMetadata>(at);
+			Ptr<ASTOpMeshAddMetadata> NewOp = UE::Mutable::Private::Clone<ASTOpMeshAddMetadata>(at);
 			NewOp->Source = Visit(NewOp->Source.child(), CurrentSinkOp);
 			NewAt = NewOp;
 			break;
@@ -162,7 +162,7 @@ namespace mu
 
 		case EOpType::ME_CONDITIONAL:
 		{
-			Ptr<ASTOpConditional> NewOp = mu::Clone<ASTOpConditional>(at);
+			Ptr<ASTOpConditional> NewOp = UE::Mutable::Private::Clone<ASTOpConditional>(at);
 			NewOp->yes = Visit(NewOp->yes.child(), CurrentSinkOp);
 			NewOp->no = Visit(NewOp->no.child(), CurrentSinkOp);
 			NewAt = NewOp;
@@ -171,7 +171,7 @@ namespace mu
 
 		case EOpType::ME_SWITCH:
 		{
-			Ptr<ASTOpSwitch> NewOp = mu::Clone<ASTOpSwitch>(at);
+			Ptr<ASTOpSwitch> NewOp = UE::Mutable::Private::Clone<ASTOpSwitch>(at);
 			NewOp->Default = Visit(NewOp->Default.child(), CurrentSinkOp);
 			for (ASTOpSwitch::FCase& c : NewOp->Cases)
 			{
@@ -185,7 +185,7 @@ namespace mu
 		default:
 			if (at != InitialSource)
 			{
-				mu::Ptr<ASTOpMeshOptimizeSkinning> NewOp = mu::Clone<ASTOpMeshOptimizeSkinning>(CurrentSinkOp);
+				UE::Mutable::Private::Ptr<ASTOpMeshOptimizeSkinning> NewOp = UE::Mutable::Private::Clone<ASTOpMeshOptimizeSkinning>(CurrentSinkOp);
 				NewOp->Source = at;
 				NewAt = NewOp;
 			}
diff --git a/Engine/Plugins/Mutable/Source/MutableTools/Private/MuT/ASTOpMeshProject.cpp b/Engine/Plugins/Mutable/Source/MutableTools/Private/MuT/ASTOpMeshProject.cpp
--- a/Engine/Plugins/Mutable/Source/MutableTools/Private/MuT/ASTOpMeshProject.cpp
+++ b/Engine/Plugins/Mutable/Source/MutableTools/Private/MuT/ASTOpMeshProject.cpp
@@ -10,7 +10,7 @@
 #include "MuR/RefCounted.h"
 #include "MuR/Types.h"
 
-namespace mu
+namespace UE::Mutable::Private
 {
 
 	ASTOpMeshProject::ASTOpMeshProject()
@@ -104,13 +104,13 @@ namespace mu
 			if (ProjectorAt->GetOpType() == EOpType::PR_CONSTANT)
 			{
 				// We move the project down the two paths
-				Ptr<ASTOpConditional> nop = mu::Clone<ASTOpConditional>(SourceAt);
+				Ptr<ASTOpConditional> nop = UE::Mutable::Private::Clone<ASTOpConditional>(SourceAt);
 
-				Ptr<ASTOpMeshProject> aOp = mu::Clone<ASTOpMeshProject>(this);
+				Ptr<ASTOpMeshProject> aOp = UE::Mutable::Private::Clone<ASTOpMeshProject>(this);
 				aOp->Mesh = nop->yes.child();
 				nop->yes = aOp;
 
-				Ptr<ASTOpMeshProject> bOp = mu::Clone<ASTOpMeshProject>(this);
+				Ptr<ASTOpMeshProject> bOp = UE::Mutable::Private::Clone<ASTOpMeshProject>(this);
 				bOp->Mesh = nop->no.child();
 				nop->no = bOp;
 
@@ -124,11 +124,11 @@ namespace mu
 			if (ProjectorAt->GetOpType() == EOpType::PR_CONSTANT)
 			{
 				// Move the format down all the paths
-				Ptr<ASTOpSwitch> nop = mu::Clone<ASTOpSwitch>(SourceAt);
+				Ptr<ASTOpSwitch> nop = UE::Mutable::Private::Clone<ASTOpSwitch>(SourceAt);
 
 				if (nop->Default)
 				{
-					Ptr<ASTOpMeshProject> defOp = mu::Clone<ASTOpMeshProject>(this);
+					Ptr<ASTOpMeshProject> defOp = UE::Mutable::Private::Clone<ASTOpMeshProject>(this);
 					defOp->Mesh = nop->Default.child();
 					nop->Default = defOp;
 				}
@@ -138,7 +138,7 @@ namespace mu
 				{
 					if (nop->Cases[CaseIndex].Branch)
 					{
-						Ptr<ASTOpMeshProject> bOp = mu::Clone<ASTOpMeshProject>(this);
+						Ptr<ASTOpMeshProject> bOp = UE::Mutable::Private::Clone<ASTOpMeshProject>(this);
 						bOp->Mesh = nop->Cases[CaseIndex].Branch.child();
 						nop->Cases[CaseIndex].Branch = bOp;
 					}
@@ -152,11 +152,11 @@ namespace mu
 		case EOpType::ME_ADDMETADATA:
 		{
 			// Add the tags after layout
-			Ptr<ASTOpMeshAddMetadata> NewAddMetadata = mu::Clone<ASTOpMeshAddMetadata>(SourceAt);
+			Ptr<ASTOpMeshAddMetadata> NewAddMetadata = UE::Mutable::Private::Clone<ASTOpMeshAddMetadata>(SourceAt);
 
 			if (NewAddMetadata->Source)
 			{
-				Ptr<ASTOpMeshProject> NewAt = mu::Clone<ASTOpMeshProject>(this);
+				Ptr<ASTOpMeshProject> NewAt = UE::Mutable::Private::Clone<ASTOpMeshProject>(this);
 				NewAt->Mesh = NewAddMetadata->Source.child();
 				NewAddMetadata->Source = NewAt;
 			}
diff --git a/Engine/Plugins/Mutable/Source/MutableTools/Private/MuT/ASTOpMeshRemoveMask.cpp b/Engine/Plugins/Mutable/Source/MutableTools/Private/MuT/ASTOpMeshRemoveMask.cpp
--- a/Engine/Plugins/Mutable/Source/MutableTools/Private/MuT/ASTOpMeshRemoveMask.cpp
+++ b/Engine/Plugins/Mutable/Source/MutableTools/Private/MuT/ASTOpMeshRemoveMask.cpp
@@ -10,7 +10,7 @@
 #include "MuR/Types.h"
 
 
-namespace mu
+namespace UE::Mutable::Private
 {
 
 
@@ -49,7 +49,7 @@ namespace mu
 
 
 	//---------------------------------------------------------------------------------------------
-	mu::Ptr<ASTOp> ASTOpMeshRemoveMask::Clone(MapChildFuncRef mapChild) const
+	UE::Mutable::Private::Ptr<ASTOp> ASTOpMeshRemoveMask::Clone(MapChildFuncRef mapChild) const
 	{
 		Ptr<ASTOpMeshRemoveMask> n = new ASTOpMeshRemoveMask();
 		n->source = mapChild(source.child());
@@ -127,13 +127,13 @@ namespace mu
 
 			// \TODO This is recursive and may cause stack overflows in big models.
 
-			mu::Ptr<ASTOp> Apply(const ASTOpMeshRemoveMask* root)
+			UE::Mutable::Private::Ptr<ASTOp> Apply(const ASTOpMeshRemoveMask* root)
 			{
 				Root = root;
 				m_oldToNew.Empty();
 
 				InitialSource = root->source.child();
-				mu::Ptr<ASTOp> newSource = Visit(InitialSource);
+				UE::Mutable::Private::Ptr<ASTOp> newSource = Visit(InitialSource);
 
 				// If there is any change, it is the new root.
 				if (newSource != InitialSource)
@@ -147,11 +147,11 @@ namespace mu
 		protected:
 
 			const ASTOpMeshRemoveMask* Root;
-			mu::Ptr<ASTOp> InitialSource;
-			TMap<mu::Ptr<ASTOp>, mu::Ptr<ASTOp>> m_oldToNew;
-			TArray<mu::Ptr<ASTOp>> m_newOps;
+			UE::Mutable::Private::Ptr<ASTOp> InitialSource;
+			TMap<UE::Mutable::Private::Ptr<ASTOp>, UE::Mutable::Private::Ptr<ASTOp>> m_oldToNew;
+			TArray<UE::Mutable::Private::Ptr<ASTOp>> m_newOps;
 
-			mu::Ptr<ASTOp> Visit(const mu::Ptr<ASTOp>& at)
+			UE::Mutable::Private::Ptr<ASTOp> Visit(const UE::Mutable::Private::Ptr<ASTOp>& at)
 			{
 				if (!at) return nullptr;
 
@@ -168,13 +168,13 @@ namespace mu
 					return *cacheIt;
 				}
 
-				mu::Ptr<ASTOp> newAt = at;
+				UE::Mutable::Private::Ptr<ASTOp> newAt = at;
 				switch (at->GetOpType())
 				{
 
 				case EOpType::ME_MORPH:
 				{
-					mu::Ptr<ASTOpMeshMorph> newOp = mu::Clone<ASTOpMeshMorph>(at);
+					UE::Mutable::Private::Ptr<ASTOpMeshMorph> newOp = UE::Mutable::Private::Clone<ASTOpMeshMorph>(at);
 					newOp->Base = Visit(newOp->Base.child());
 					newAt = newOp;
 					break;
@@ -182,7 +182,7 @@ namespace mu
 
 				case EOpType::ME_ADDMETADATA:
 				{
-					mu::Ptr<ASTOpMeshAddMetadata> newOp = mu::Clone<ASTOpMeshAddMetadata>(at);
+					UE::Mutable::Private::Ptr<ASTOpMeshAddMetadata> newOp = UE::Mutable::Private::Clone<ASTOpMeshAddMetadata>(at);
 					newOp->Source = Visit(newOp->Source.child());
 					newAt = newOp;
 					break;
@@ -191,7 +191,7 @@ namespace mu
 				// disabled to avoid code explosion (or bug?) TODO
 //            case EOpType::ME_CONDITIONAL:
 //            {
-//                Ptr<ASTOpConditional> newOp = mu::Clone<ASTOpConditional>(at);
+//                Ptr<ASTOpConditional> newOp = UE::Mutable::Private::Clone<ASTOpConditional>(at);
 //                newOp->yes = Visit(newOp->yes.child());
 //                newOp->no = Visit(newOp->no.child());
 //                newAt = newOp;
@@ -200,7 +200,7 @@ namespace mu
 
 //            case EOpType::ME_SWITCH:
 //            {
-//                auto newOp = mu::Clone<ASTOpSwitch>(at);
+//                auto newOp = UE::Mutable::Private::Clone<ASTOpSwitch>(at);
 //                newOp->def = Visit(newOp->def.child());
 //                for( auto& c:newOp->cases )
 //                {
@@ -215,7 +215,7 @@ namespace mu
 					//
 					if (at != InitialSource)
 					{
-						Ptr<ASTOpMeshRemoveMask> newOp = mu::Clone<ASTOpMeshRemoveMask>(Root);
+						Ptr<ASTOpMeshRemoveMask> newOp = UE::Mutable::Private::Clone<ASTOpMeshRemoveMask>(Root);
 						newOp->source = at;
 						newAt = newOp;
 					}
@@ -233,10 +233,10 @@ namespace mu
 
 
 	//-------------------------------------------------------------------------------------------------
-	mu::Ptr<ASTOp> ASTOpMeshRemoveMask::OptimiseSink(const FModelOptimizationOptions&, FOptimizeSinkContext&) const
+	UE::Mutable::Private::Ptr<ASTOp> ASTOpMeshRemoveMask::OptimiseSink(const FModelOptimizationOptions&, FOptimizeSinkContext&) const
 	{
 		Sink_MeshRemoveMaskAST sinker;
-		mu::Ptr<ASTOp> at = sinker.Apply(this);
+		UE::Mutable::Private::Ptr<ASTOp> at = sinker.Apply(this);
 
 		// If not optimized already, see if we can optimize the "remove" branches
 		if (!at)
@@ -258,7 +258,7 @@ namespace mu
 					// It can be ignored.
 					if (!NewOp)
 					{
-						NewOp = mu::Clone<ASTOpMeshRemoveMask>(this);
+						NewOp = UE::Mutable::Private::Clone<ASTOpMeshRemoveMask>(this);
 					}
 
 					const ASTOpMeshAddMetadata* Add = static_cast<const ASTOpMeshAddMetadata*>(removes[RemoveIndex].Value.child().get());
diff --git a/Engine/Plugins/Mutable/Source/MutableTools/Private/MuT/ASTOpMeshTransform.cpp b/Engine/Plugins/Mutable/Source/MutableTools/Private/MuT/ASTOpMeshTransform.cpp
--- a/Engine/Plugins/Mutable/Source/MutableTools/Private/MuT/ASTOpMeshTransform.cpp
+++ b/Engine/Plugins/Mutable/Source/MutableTools/Private/MuT/ASTOpMeshTransform.cpp
@@ -8,7 +8,7 @@
 #include "MuR/Types.h"
 
 
-namespace mu
+namespace UE::Mutable::Private
 {
 
 	ASTOpMeshTransform::ASTOpMeshTransform()
@@ -44,7 +44,7 @@ namespace mu
 	}
 
 
-	mu::Ptr<ASTOp> ASTOpMeshTransform::Clone(MapChildFuncRef mapChild) const
+	UE::Mutable::Private::Ptr<ASTOp> ASTOpMeshTransform::Clone(MapChildFuncRef mapChild) const
 	{
 		Ptr<ASTOpMeshTransform> n = new ASTOpMeshTransform();
 		n->matrix = matrix;
diff --git a/Engine/Plugins/Mutable/Source/MutableTools/Private/MuT/ASTOpMeshTransformWithBone.cpp b/Engine/Plugins/Mutable/Source/MutableTools/Private/MuT/ASTOpMeshTransformWithBone.cpp
--- a/Engine/Plugins/Mutable/Source/MutableTools/Private/MuT/ASTOpMeshTransformWithBone.cpp
+++ b/Engine/Plugins/Mutable/Source/MutableTools/Private/MuT/ASTOpMeshTransformWithBone.cpp
@@ -5,19 +5,19 @@
 #include "ASTOpImageTransform.h"
 
 
-mu::ASTOpMeshTransformWithBone::ASTOpMeshTransformWithBone()
+UE::Mutable::Private::ASTOpMeshTransformWithBone::ASTOpMeshTransformWithBone()
 	: SourceMesh(this)
 	, Matrix(this)
 {
 }
 
-mu::ASTOpMeshTransformWithBone::~ASTOpMeshTransformWithBone()
+UE::Mutable::Private::ASTOpMeshTransformWithBone::~ASTOpMeshTransformWithBone()
 {
 	// Explicit call needed to avoid recursive destruction
 	ASTOp::RemoveChildren();
 }
 
-uint64 mu::ASTOpMeshTransformWithBone::Hash() const
+uint64 UE::Mutable::Private::ASTOpMeshTransformWithBone::Hash() const
 {
 	uint64 res = std::hash<EOpType>()(GetOpType());
 	hash_combine(res, SourceMesh.child());
@@ -28,7 +28,7 @@ uint64 mu::ASTOpMeshTransformWithBone::Hash() const
 	return res;
 }
 
-bool mu::ASTOpMeshTransformWithBone::IsEqual(const ASTOp& OtherUntyped) const
+bool UE::Mutable::Private::ASTOpMeshTransformWithBone::IsEqual(const ASTOp& OtherUntyped) const
 {
 	if (GetOpType() == OtherUntyped.GetOpType())
 	{
@@ -41,7 +41,7 @@ bool mu::ASTOpMeshTransformWithBone::IsEqual(const ASTOp& OtherUntyped) const
 	return false;
 }
 
-mu::Ptr<mu::ASTOp> mu::ASTOpMeshTransformWithBone::Clone(MapChildFuncRef MapChild) const
+UE::Mutable::Private::Ptr<UE::Mutable::Private::ASTOp> UE::Mutable::Private::ASTOpMeshTransformWithBone::Clone(MapChildFuncRef MapChild) const
 {
 	Ptr<ASTOpMeshTransformWithBone> NewASTOp = new ASTOpMeshTransformWithBone();
 	NewASTOp->SourceMesh = MapChild(SourceMesh.child());
@@ -51,13 +51,13 @@ mu::Ptr<mu::ASTOp> mu::ASTOpMeshTransformWithBone::Clone(MapChildFuncRef MapChil
 	return NewASTOp;
 }
 
-void mu::ASTOpMeshTransformWithBone::ForEachChild(const TFunctionRef<void(ASTChild&)> Function)
+void UE::Mutable::Private::ASTOpMeshTransformWithBone::ForEachChild(const TFunctionRef<void(ASTChild&)> Function)
 {
 	Function(SourceMesh);
 	Function(Matrix);
 }
 
-void mu::ASTOpMeshTransformWithBone::Link(FProgram& Program, FLinkerOptions* Options)
+void UE::Mutable::Private::ASTOpMeshTransformWithBone::Link(FProgram& Program, FLinkerOptions* Options)
 {
 	if (!linkedAddress)
 	{
@@ -83,7 +83,7 @@ void mu::ASTOpMeshTransformWithBone::Link(FProgram& Program, FLinkerOptions* Opt
 	}
 }
 
-mu::FSourceDataDescriptor mu::ASTOpMeshTransformWithBone::GetSourceDataDescriptor(FGetSourceDataDescriptorContext* Context) const
+UE::Mutable::Private::FSourceDataDescriptor UE::Mutable::Private::ASTOpMeshTransformWithBone::GetSourceDataDescriptor(FGetSourceDataDescriptorContext* Context) const
 {
 	if (SourceMesh)
 	{
diff --git a/Engine/Plugins/Mutable/Source/MutableTools/Private/MuT/ASTOpMeshTransformWithBoundingMesh.cpp b/Engine/Plugins/Mutable/Source/MutableTools/Private/MuT/ASTOpMeshTransformWithBoundingMesh.cpp
--- a/Engine/Plugins/Mutable/Source/MutableTools/Private/MuT/ASTOpMeshTransformWithBoundingMesh.cpp
+++ b/Engine/Plugins/Mutable/Source/MutableTools/Private/MuT/ASTOpMeshTransformWithBoundingMesh.cpp
@@ -5,20 +5,20 @@
 #include "ASTOpImageTransform.h"
 
 
-mu::ASTOpMeshTransformWithBoundingMesh::ASTOpMeshTransformWithBoundingMesh()
+UE::Mutable::Private::ASTOpMeshTransformWithBoundingMesh::ASTOpMeshTransformWithBoundingMesh()
 	: source(this)
 	, boundingMesh(this)
 	, matrix(this)
 {
 }
 
-mu::ASTOpMeshTransformWithBoundingMesh::~ASTOpMeshTransformWithBoundingMesh()
+UE::Mutable::Private::ASTOpMeshTransformWithBoundingMesh::~ASTOpMeshTransformWithBoundingMesh()
 {
 	// Explicit call needed to avoid recursive destruction
 	ASTOp::RemoveChildren();
 }
 
-uint64 mu::ASTOpMeshTransformWithBoundingMesh::Hash() const
+uint64 UE::Mutable::Private::ASTOpMeshTransformWithBoundingMesh::Hash() const
 {
 	uint64 res = std::hash<EOpType>()(GetOpType());
 	hash_combine(res, source.child());
@@ -28,7 +28,7 @@ uint64 mu::ASTOpMeshTransformWithBoundingMesh::Hash() const
 	return res;
 }
 
-bool mu::ASTOpMeshTransformWithBoundingMesh::IsEqual(const ASTOp& otherUntyped) const
+bool UE::Mutable::Private::ASTOpMeshTransformWithBoundingMesh::IsEqual(const ASTOp& otherUntyped) const
 {
 	if (GetOpType() == otherUntyped.GetOpType())
 	{
@@ -38,7 +38,7 @@ bool mu::ASTOpMeshTransformWithBoundingMesh::IsEqual(const ASTOp& otherUntyped)
 	return false;
 }
 
-mu::Ptr<mu::ASTOp> mu::ASTOpMeshTransformWithBoundingMesh::Clone(MapChildFuncRef mapChild) const
+UE::Mutable::Private::Ptr<UE::Mutable::Private::ASTOp> UE::Mutable::Private::ASTOpMeshTransformWithBoundingMesh::Clone(MapChildFuncRef mapChild) const
 {
 	Ptr<ASTOpMeshTransformWithBoundingMesh> n = new ASTOpMeshTransformWithBoundingMesh();
 	n->source = mapChild(source.child());
@@ -47,14 +47,14 @@ mu::Ptr<mu::ASTOp> mu::ASTOpMeshTransformWithBoundingMesh::Clone(MapChildFuncRef
 	return n;
 }
 
-void mu::ASTOpMeshTransformWithBoundingMesh::ForEachChild(const TFunctionRef<void(ASTChild&)> f)
+void UE::Mutable::Private::ASTOpMeshTransformWithBoundingMesh::ForEachChild(const TFunctionRef<void(ASTChild&)> f)
 {
 	f(source);
 	f(boundingMesh);
 	f(matrix);
 }
 
-void mu::ASTOpMeshTransformWithBoundingMesh::Link(FProgram& program, FLinkerOptions* Options)
+void UE::Mutable::Private::ASTOpMeshTransformWithBoundingMesh::Link(FProgram& program, FLinkerOptions* Options)
 {
 	if (!linkedAddress)
 	{
@@ -79,7 +79,7 @@ void mu::ASTOpMeshTransformWithBoundingMesh::Link(FProgram& program, FLinkerOpti
 	}
 }
 
-mu::FSourceDataDescriptor mu::ASTOpMeshTransformWithBoundingMesh::GetSourceDataDescriptor(FGetSourceDataDescriptorContext* Context) const
+UE::Mutable::Private::FSourceDataDescriptor UE::Mutable::Private::ASTOpMeshTransformWithBoundingMesh::GetSourceDataDescriptor(FGetSourceDataDescriptorContext* Context) const
 {
 	if (source)
 	{
diff --git a/Engine/Plugins/Mutable/Source/MutableTools/Private/MuT/ASTOpParameter.cpp b/Engine/Plugins/Mutable/Source/MutableTools/Private/MuT/ASTOpParameter.cpp
--- a/Engine/Plugins/Mutable/Source/MutableTools/Private/MuT/ASTOpParameter.cpp
+++ b/Engine/Plugins/Mutable/Source/MutableTools/Private/MuT/ASTOpParameter.cpp
@@ -11,7 +11,7 @@
 #include "MuR/Types.h"
 
 
-namespace mu
+namespace UE::Mutable::Private
 {
 
 
@@ -55,7 +55,7 @@ namespace mu
 	}
 
 
-	mu::Ptr<ASTOp> ASTOpParameter::Clone(MapChildFuncRef MapChild) const
+	UE::Mutable::Private::Ptr<ASTOp> ASTOpParameter::Clone(MapChildFuncRef MapChild) const
 	{
 		Ptr<ASTOpParameter> n = new ASTOpParameter();
 		n->Type = Type;
diff --git a/Engine/Plugins/Mutable/Source/MutableTools/Private/MuT/ASTOpReferenceResource.cpp b/Engine/Plugins/Mutable/Source/MutableTools/Private/MuT/ASTOpReferenceResource.cpp
--- a/Engine/Plugins/Mutable/Source/MutableTools/Private/MuT/ASTOpReferenceResource.cpp
+++ b/Engine/Plugins/Mutable/Source/MutableTools/Private/MuT/ASTOpReferenceResource.cpp
@@ -9,7 +9,7 @@
 #include "MuR/Types.h"
 
 
-namespace mu
+namespace UE::Mutable::Private
 {
 
 	void ASTOpReferenceResource::ForEachChild(const TFunctionRef<void(ASTChild&)>)
@@ -30,7 +30,7 @@ namespace mu
 
 
 	//-------------------------------------------------------------------------------------------------
-	mu::Ptr<ASTOp> ASTOpReferenceResource::Clone(MapChildFuncRef) const
+	UE::Mutable::Private::Ptr<ASTOp> ASTOpReferenceResource::Clone(MapChildFuncRef) const
 	{
 		Ptr<ASTOpReferenceResource> n = new ASTOpReferenceResource();
 		n->Type = Type;
@@ -122,7 +122,7 @@ namespace mu
 	}
 
 
-	mu::Ptr<ImageSizeExpression> ASTOpReferenceResource::GetImageSizeExpression() const
+	UE::Mutable::Private::Ptr<ImageSizeExpression> ASTOpReferenceResource::GetImageSizeExpression() const
 	{
 		Ptr<ImageSizeExpression> pRes = new ImageSizeExpression;
 		pRes->type = ImageSizeExpression::ISET_UNKNOWN;
diff --git a/Engine/Plugins/Mutable/Source/MutableTools/Private/MuT/ASTOpScalarCurve.cpp b/Engine/Plugins/Mutable/Source/MutableTools/Private/MuT/ASTOpScalarCurve.cpp
--- a/Engine/Plugins/Mutable/Source/MutableTools/Private/MuT/ASTOpScalarCurve.cpp
+++ b/Engine/Plugins/Mutable/Source/MutableTools/Private/MuT/ASTOpScalarCurve.cpp
@@ -8,7 +8,7 @@
 #include "MuR/Types.h"
 
 
-namespace mu
+namespace UE::Mutable::Private
 {
 
 
@@ -50,7 +50,7 @@ namespace mu
 	}
 
 
-	mu::Ptr<ASTOp> ASTOpScalarCurve::Clone(MapChildFuncRef mapChild) const
+	UE::Mutable::Private::Ptr<ASTOp> ASTOpScalarCurve::Clone(MapChildFuncRef mapChild) const
 	{
 		Ptr<ASTOpScalarCurve> n = new ASTOpScalarCurve();
 		n->Curve = Curve;
diff --git a/Engine/Plugins/Mutable/Source/MutableTools/Private/MuT/ASTOpSwitch.cpp b/Engine/Plugins/Mutable/Source/MutableTools/Private/MuT/ASTOpSwitch.cpp
--- a/Engine/Plugins/Mutable/Source/MutableTools/Private/MuT/ASTOpSwitch.cpp
+++ b/Engine/Plugins/Mutable/Source/MutableTools/Private/MuT/ASTOpSwitch.cpp
@@ -13,7 +13,7 @@
 #include "MuR/RefCounted.h"
 #include "MuR/Types.h"
 
-namespace mu
+namespace UE::Mutable::Private
 {
 
 	ASTOpSwitch::ASTOpSwitch()
@@ -46,7 +46,7 @@ namespace mu
 	}
 
 
-	mu::Ptr<ASTOp> ASTOpSwitch::Clone(MapChildFuncRef MapChild) const
+	UE::Mutable::Private::Ptr<ASTOp> ASTOpSwitch::Clone(MapChildFuncRef MapChild) const
 	{
 		Ptr<ASTOpSwitch> New = new ASTOpSwitch();
 
@@ -98,7 +98,7 @@ namespace mu
 	}
 
 
-	mu::Ptr<ASTOp> ASTOpSwitch::GetFirstValidValue()
+	UE::Mutable::Private::Ptr<ASTOp> ASTOpSwitch::GetFirstValidValue()
 	{
 		for (const FCase& Case : Cases)
 		{
@@ -151,7 +151,7 @@ namespace mu
 	}
 
 
-	mu::Ptr<ASTOp> ASTOpSwitch::FindBranch(int32 Condition) const
+	UE::Mutable::Private::Ptr<ASTOp> ASTOpSwitch::FindBranch(int32 Condition) const
 	{
 		for (const FCase& Case : Cases)
 		{
@@ -483,7 +483,7 @@ namespace mu
 	}
 
 
-	mu::Ptr<ASTOp> ASTOpSwitch::OptimiseSemantic(const FModelOptimizationOptions&, int32 Pass) const
+	UE::Mutable::Private::Ptr<ASTOp> ASTOpSwitch::OptimiseSemantic(const FModelOptimizationOptions&, int32 Pass) const
 	{
 		// Constant Condition?
 		if (Variable->GetOpType() == EOpType::NU_CONSTANT)
@@ -751,7 +751,7 @@ namespace mu
 					New->SkeletonIds = SkeletonIds;
 
 					{
-						Ptr<ASTOpSwitch> NewSwitch = mu::Clone<ASTOpSwitch>(this);
+						Ptr<ASTOpSwitch> NewSwitch = UE::Mutable::Private::Clone<ASTOpSwitch>(this);
 
 						// Replace all branches removing the "add tags" operation.
 						if (Default)
@@ -793,7 +793,7 @@ namespace mu
 	}
 
 
-	mu::Ptr<ImageSizeExpression> ASTOpSwitch::GetImageSizeExpression() const
+	UE::Mutable::Private::Ptr<ImageSizeExpression> ASTOpSwitch::GetImageSizeExpression() const
 	{
 		Ptr<ImageSizeExpression> pRes = new ImageSizeExpression;
 
diff --git a/Engine/Plugins/Mutable/Source/MutableTools/Private/MuT/CodeGenerator.cpp b/Engine/Plugins/Mutable/Source/MutableTools/Private/MuT/CodeGenerator.cpp
--- a/Engine/Plugins/Mutable/Source/MutableTools/Private/MuT/CodeGenerator.cpp
+++ b/Engine/Plugins/Mutable/Source/MutableTools/Private/MuT/CodeGenerator.cpp
@@ -98,7 +98,7 @@
 #include "Trace/Detail/Channel.h"
 
 
-namespace mu
+namespace UE::Mutable::Private
 {
 
 	CodeGenerator::CodeGenerator(CompilerOptions::Private* Options, TFunction<void()>& InWaitCallback)
@@ -868,7 +868,7 @@ namespace mu
 				FInt32Point BlockSize = BlockRectInPixels.Size();
 				if (!PatchMask)
 				{
-					PatchMask = MakeShared<FImage>(BlockSize[0], BlockSize[1], 1, mu::EImageFormat::L_UByte, mu::EInitializationType::Black);
+					PatchMask = MakeShared<FImage>(BlockSize[0], BlockSize[1], 1, UE::Mutable::Private::EImageFormat::L_UByte, UE::Mutable::Private::EInitializationType::Black);
 				}
 
 				uint8* Pixels = PatchMask->GetMipData(0);
@@ -1171,7 +1171,7 @@ namespace mu
 						continue;
 					}
 
-					if (SurfaceReferenceLayouts[LayoutIndex].Layout->GetLayoutPackingStrategy() == mu::EPackStrategy::Overlay)
+					if (SurfaceReferenceLayouts[LayoutIndex].Layout->GetLayoutPackingStrategy() == UE::Mutable::Private::EPackStrategy::Overlay)
 					{
 						continue;
 					}
@@ -2671,7 +2671,7 @@ namespace mu
 					GenerateImage(ClipOptions, ClipMaskResult, TypedClipNode->ClipMask);
 
 					// It could be IF_L_UBIT, but since this should be optimized out at compile time, leave the most cpu efficient.
-					MaskImage = GenerateImageFormat(ClipMaskResult.op, mu::EImageFormat::L_UByte);
+					MaskImage = GenerateImageFormat(ClipMaskResult.op, UE::Mutable::Private::EImageFormat::L_UByte);
 				}
 				else if (TypedClipNode->ClipLayout)
 				{
diff --git a/Engine/Plugins/Mutable/Source/MutableTools/Private/MuT/CodeGenerator.h b/Engine/Plugins/Mutable/Source/MutableTools/Private/MuT/CodeGenerator.h
--- a/Engine/Plugins/Mutable/Source/MutableTools/Private/MuT/CodeGenerator.h
+++ b/Engine/Plugins/Mutable/Source/MutableTools/Private/MuT/CodeGenerator.h
@@ -41,7 +41,7 @@
 #include "Tasks/Pipe.h"
 
 
-namespace mu
+namespace UE::Mutable::Private
 {
 	class ASTOpParameter;
 	class FLayout;
@@ -118,7 +118,7 @@ namespace mu
 	struct FObjectState;
 	struct FProgram;
 
-    /** Converts mu::Node graphs into ASTOp graphs. */
+    /** Converts UE::Mutable::Private::Node graphs into ASTOp graphs. */
     class CodeGenerator
     {
 		
diff --git a/Engine/Plugins/Mutable/Source/MutableTools/Private/MuT/CodeGenerator_Color.cpp b/Engine/Plugins/Mutable/Source/MutableTools/Private/MuT/CodeGenerator_Color.cpp
--- a/Engine/Plugins/Mutable/Source/MutableTools/Private/MuT/CodeGenerator_Color.cpp
+++ b/Engine/Plugins/Mutable/Source/MutableTools/Private/MuT/CodeGenerator_Color.cpp
@@ -47,7 +47,7 @@
 #include "MuT/TablePrivate.h"
 
 
-namespace mu
+namespace UE::Mutable::Private
 {
 
 
@@ -568,7 +568,7 @@ namespace mu
 
 
 	//---------------------------------------------------------------------------------------------
-	mu::Ptr<ASTOp> CodeGenerator::GenerateMissingColourCode(const TCHAR* StrWhere, const void* ErrorContext)
+	UE::Mutable::Private::Ptr<ASTOp> CodeGenerator::GenerateMissingColourCode(const TCHAR* StrWhere, const void* ErrorContext)
 	{
 		// Log a warning
 		FString Msg = FString::Printf(TEXT("Required connection not found: %s"), StrWhere);
diff --git a/Engine/Plugins/Mutable/Source/MutableTools/Private/MuT/CodeGenerator_FirstPass.cpp b/Engine/Plugins/Mutable/Source/MutableTools/Private/MuT/CodeGenerator_FirstPass.cpp
--- a/Engine/Plugins/Mutable/Source/MutableTools/Private/MuT/CodeGenerator_FirstPass.cpp
+++ b/Engine/Plugins/Mutable/Source/MutableTools/Private/MuT/CodeGenerator_FirstPass.cpp
@@ -31,7 +31,7 @@
 #include "MuT/NodeSurfaceSwitch.h"
 #include "MuT/NodeScalarEnumParameter.h"
 
-namespace mu
+namespace UE::Mutable::Private
 {
 
 	//---------------------------------------------------------------------------------------------
@@ -653,7 +653,7 @@ namespace mu
         // Parse the child objects
 		for ( int32 t=0; t< Node->Children.Num(); ++t )
         {
-            if ( mu::Ptr<const NodeObject> ChildNode = Node->Children[t] )
+            if ( UE::Mutable::Private::Ptr<const NodeObject> ChildNode = Node->Children[t] )
             {
                 // Overwrite the implicit condition
                 Ptr<ASTOp> paramOp = 0;
diff --git a/Engine/Plugins/Mutable/Source/MutableTools/Private/MuT/CodeGenerator_Mesh.cpp b/Engine/Plugins/Mutable/Source/MutableTools/Private/MuT/CodeGenerator_Mesh.cpp
--- a/Engine/Plugins/Mutable/Source/MutableTools/Private/MuT/CodeGenerator_Mesh.cpp
+++ b/Engine/Plugins/Mutable/Source/MutableTools/Private/MuT/CodeGenerator_Mesh.cpp
@@ -62,7 +62,7 @@
 #include "MuT/TablePrivate.h"
 
 
-namespace mu
+namespace UE::Mutable::Private
 {
 	class Node;
 
@@ -468,7 +468,7 @@ namespace mu
 			[this] (const NodeMeshTable& node, int32 colIndex, int32 row, FErrorLog* pErrorLog) -> Ptr<NodeMesh>
 			{
 				const FTableValue& Cell = node.Table->GetPrivate()->Rows[row].Values[colIndex];
-				TSharedPtr<mu::FMesh> pMesh = Cell.Mesh;
+				TSharedPtr<UE::Mutable::Private::FMesh> pMesh = Cell.Mesh;
 				
 				if (!pMesh)
 				{
@@ -718,7 +718,7 @@ namespace mu
 		// Generate the chain of ops preparing the mesh for the layouts
 		for (int32 LayoutIndex = 0; LayoutIndex < OutResult.GeneratedLayouts.Num(); ++LayoutIndex)
 		{
-			TSharedPtr<const mu::FLayout> Layout = OutResult.GeneratedLayouts[LayoutIndex].Layout;
+			TSharedPtr<const UE::Mutable::Private::FLayout> Layout = OutResult.GeneratedLayouts[LayoutIndex].Layout;
 
 			Ptr<ASTOpConstantResource> LayoutOp = new ASTOpConstantResource;
 			LayoutOp->Type = EOpType::LA_CONSTANT;
diff --git a/Engine/Plugins/Mutable/Source/MutableTools/Private/MuT/CodeGenerator_Scalar.cpp b/Engine/Plugins/Mutable/Source/MutableTools/Private/MuT/CodeGenerator_Scalar.cpp
--- a/Engine/Plugins/Mutable/Source/MutableTools/Private/MuT/CodeGenerator_Scalar.cpp
+++ b/Engine/Plugins/Mutable/Source/MutableTools/Private/MuT/CodeGenerator_Scalar.cpp
@@ -39,7 +39,7 @@
 #include "MuT/TablePrivate.h"
 
 
-namespace mu
+namespace UE::Mutable::Private
 {
 
 	void CodeGenerator::GenerateScalar(FScalarGenerationResult& result, const FGenericGenerationOptions& Options, const Ptr<const NodeScalar>& Untyped)
@@ -448,7 +448,7 @@ namespace mu
 	}
 
 
-	mu::Ptr<ASTOp> CodeGenerator::GenerateMissingScalarCode(const TCHAR* strWhere, float value, const void* errorContext)
+	UE::Mutable::Private::Ptr<ASTOp> CodeGenerator::GenerateMissingScalarCode(const TCHAR* strWhere, float value, const void* errorContext)
 	{
 		// Log a warning
 		FString Msg = FString::Printf(TEXT("Required connection not found: %s"), strWhere );
diff --git a/Engine/Plugins/Mutable/Source/MutableTools/Private/MuT/CodeGenerator_SecondPass.cpp b/Engine/Plugins/Mutable/Source/MutableTools/Private/MuT/CodeGenerator_SecondPass.cpp
--- a/Engine/Plugins/Mutable/Source/MutableTools/Private/MuT/CodeGenerator_SecondPass.cpp
+++ b/Engine/Plugins/Mutable/Source/MutableTools/Private/MuT/CodeGenerator_SecondPass.cpp
@@ -12,7 +12,7 @@
 #include "MuT/ErrorLog.h"
 #include "MuT/CodeGenerator_FirstPass.h"
 
-namespace mu
+namespace UE::Mutable::Private
 {
 
 
@@ -29,7 +29,7 @@ namespace mu
 
 
 	//---------------------------------------------------------------------------------------------
-	mu::Ptr<ASTOp> SecondPassGenerator::GenerateTagCondition(int32 tagIndex,
+	UE::Mutable::Private::Ptr<ASTOp> SecondPassGenerator::GenerateTagCondition(int32 tagIndex,
 		const TSet<int32>& posSurf,
 		const TSet<int32>& negSurf,
 		const TSet<int32>& posTag,
@@ -147,7 +147,7 @@ namespace mu
 	}
 
 	//---------------------------------------------------------------------------------------------
-	mu::Ptr<ASTOp> SecondPassGenerator::GenerateDataCodition(int32 Index,
+	UE::Mutable::Private::Ptr<ASTOp> SecondPassGenerator::GenerateDataCodition(int32 Index,
 		const TArray<FString>& PositiveTags,
 		const TArray<FString>& NegativeTags,
 		const TSet<int32>& posSurf,
