[Mutable] Change namespace name

diff --git a/Engine/Plugins/Mutable/Source/MutableTools/Private/MuT/CodeOptimiser.cpp b/Engine/Plugins/Mutable/Source/MutableTools/Private/MuT/CodeOptimiser.cpp
--- a/Engine/Plugins/Mutable/Source/MutableTools/Private/MuT/CodeOptimiser.cpp
+++ b/Engine/Plugins/Mutable/Source/MutableTools/Private/MuT/CodeOptimiser.cpp
@@ -28,7 +28,7 @@
 #include "Tasks/Task.h"
 #include "Math/NumericLimits.h"
 
-namespace mu
+namespace UE::Mutable::Private
 {
 
 	namespace
@@ -427,7 +427,7 @@ namespace mu
 
 				if (MeshBuild)
 				{
-					mu::Ptr<ASTOpConstantResource> ConstantOp = new ASTOpConstantResource();
+					UE::Mutable::Private::Ptr<ASTOpConstantResource> ConstantOp = new ASTOpConstantResource();
 					ConstantOp->SourceDataDescriptor = SourceDataDescriptor;
 					ConstantOp->Type = EOpType::ME_CONSTANT;
 					ConstantOp->SetValue(MeshBuild, DiskCacheContext);
@@ -444,7 +444,7 @@ namespace mu
 
 				if (pImage)
 				{
-					mu::Ptr<ASTOpConstantResource> ConstantOp = new ASTOpConstantResource();
+					UE::Mutable::Private::Ptr<ASTOpConstantResource> ConstantOp = new ASTOpConstantResource();
 					ConstantOp->SourceDataDescriptor = SourceDataDescriptor;
 					ConstantOp->Type = EOpType::IM_CONSTANT;
 					ConstantOp->SetValue( pImage, DiskCacheContext );
@@ -461,7 +461,7 @@ namespace mu
 
 				if (pLayout)
 				{
-					mu::Ptr<ASTOpConstantResource> constantOp = new ASTOpConstantResource();
+					UE::Mutable::Private::Ptr<ASTOpConstantResource> constantOp = new ASTOpConstantResource();
 					constantOp->Type = EOpType::LA_CONSTANT;
 					constantOp->SetValue( pLayout, DiskCacheContext);
 					Result = constantOp;
@@ -486,7 +486,7 @@ namespace mu
 				ResultColor = System->GetPrivate()->BuildColour(nullptr, Model, LocalParams.Get(), at);
 
 				{
-					mu::Ptr<ASTOpConstantColor> ConstantOp = new ASTOpConstantColor();
+					UE::Mutable::Private::Ptr<ASTOpConstantColor> ConstantOp = new ASTOpConstantColor();
 					ConstantOp->Value = ResultColor;
 					Result = ConstantOp;
 				}
@@ -1003,12 +1003,12 @@ namespace mu
 		}
 
 		// List of meshes that require a skeleton
-		TArray<mu::Ptr<ASTOpConstantResource>> MeshesRequiringSkeleton;
+		TArray<UE::Mutable::Private::Ptr<ASTOpConstantResource>> MeshesRequiringSkeleton;
 
 	private:
 
 		// Visitor_TopDown_Unique_Const<uint8_t> interface
-		bool Visit( const mu::Ptr<ASTOp>& node ) override
+		bool Visit( const UE::Mutable::Private::Ptr<ASTOp>& node ) override
 		{
 			// \todo: refine to avoid instruction branches with irrelevant skeletons.
 
@@ -1019,7 +1019,7 @@ namespace mu
 
 			case EOpType::ME_CONSTANT:
 			{
-				mu::Ptr<ASTOpConstantResource> typedOp = static_cast<ASTOpConstantResource*>(node.get());
+				UE::Mutable::Private::Ptr<ASTOpConstantResource> typedOp = static_cast<ASTOpConstantResource*>(node.get());
 
 				if (currentProtected)
 				{
@@ -1090,12 +1090,12 @@ namespace mu
 	// for this op.
 	struct FAddMeshSkeleton
 	{
-		mu::Ptr<ASTOp> AddMeshOp;
-		TArray<mu::Ptr<ASTOpConstantResource>> ContributingMeshes;
+		UE::Mutable::Private::Ptr<ASTOp> AddMeshOp;
+		TArray<UE::Mutable::Private::Ptr<ASTOpConstantResource>> ContributingMeshes;
 		TSharedPtr<FSkeleton> FinalSkeleton;
 
-		FAddMeshSkeleton( const mu::Ptr<ASTOp>& InAddMeshOp,
-						  TArray<mu::Ptr<ASTOpConstantResource>>& InContributingMeshes,
+		FAddMeshSkeleton( const UE::Mutable::Private::Ptr<ASTOp>& InAddMeshOp,
+						  TArray<UE::Mutable::Private::Ptr<ASTOpConstantResource>>& InContributingMeshes,
 						  const TSharedPtr<FSkeleton>& InFinalSkeleton )
 		{
 			AddMeshOp = InAddMeshOp;
@@ -1105,7 +1105,7 @@ namespace mu
 	};
 
 
-	void SkeletonCleanerAST( TArray<mu::Ptr<ASTOp>>& roots, const FModelOptimizationOptions& options )
+	void SkeletonCleanerAST( TArray<UE::Mutable::Private::Ptr<ASTOp>>& roots, const FModelOptimizationOptions& options )
 	{
 		// This collects all the meshes that require a skeleton because they are used in operations
 		// that require it.
@@ -1113,7 +1113,7 @@ namespace mu
 
 		TArray<FAddMeshSkeleton> replacementsFound;
 
-		ASTOp::Traverse_TopDown_Unique_Imprecise( roots, [&](mu::Ptr<ASTOp>& at )
+		ASTOp::Traverse_TopDown_Unique_Imprecise( roots, [&](UE::Mutable::Private::Ptr<ASTOp>& at )
 		{
 			// Only recurse instance construction ops.
 			bool processChildren = GetOpDataType(at->GetOpType())== EDataType::Instance;
@@ -1121,20 +1121,20 @@ namespace mu
 			if ( at->GetOpType() == EOpType::IN_ADDMESH )
 			{
 				ASTOpInstanceAdd* typedNode = static_cast<ASTOpInstanceAdd*>(at.get());
-				mu::Ptr<ASTOp> meshRoot = typedNode->value.child();
+				UE::Mutable::Private::Ptr<ASTOp> meshRoot = typedNode->value.child();
 
 				if (meshRoot)
 				{
 					// Gather constant meshes contributing to the final mesh
-					TArray<mu::Ptr<ASTOpConstantResource>> subtreeMeshes;
-					TArray<mu::Ptr<ASTOp>> tempRoots;
+					TArray<UE::Mutable::Private::Ptr<ASTOpConstantResource>> subtreeMeshes;
+					TArray<UE::Mutable::Private::Ptr<ASTOp>> tempRoots;
 					tempRoots.Add(meshRoot);
-					ASTOp::Traverse_TopDown_Unique_Imprecise( tempRoots, [&](mu::Ptr<ASTOp>& lat )
+					ASTOp::Traverse_TopDown_Unique_Imprecise( tempRoots, [&](UE::Mutable::Private::Ptr<ASTOp>& lat )
 					{
 						// \todo: refine to avoid instruction branches with irrelevant skeletons.
 						if ( lat->GetOpType() == EOpType::ME_CONSTANT )
 						{
-							mu::Ptr<ASTOpConstantResource> typedOp = static_cast<ASTOpConstantResource*>(lat.get());
+							UE::Mutable::Private::Ptr<ASTOpConstantResource> typedOp = static_cast<ASTOpConstantResource*>(lat.get());
 							if ( subtreeMeshes.Find(typedOp)
 								 ==
 								 INDEX_NONE )
@@ -1166,7 +1166,7 @@ namespace mu
 
 
 		// Iterate all meshes again
-		ASTOp::Traverse_TopDown_Unique_Imprecise( roots, [&](mu::Ptr<ASTOp>& at )
+		ASTOp::Traverse_TopDown_Unique_Imprecise( roots, [&](UE::Mutable::Private::Ptr<ASTOp>& at )
 		{
 			if (at->GetOpType()==EOpType::ME_CONSTANT)
 			{
@@ -1186,7 +1186,7 @@ namespace mu
 						if (bOutSuccess)
 						{
 							NewMesh->CheckIntegrity();
-							mu::Ptr<ASTOpConstantResource> newOp = new ASTOpConstantResource();
+							UE::Mutable::Private::Ptr<ASTOpConstantResource> newOp = new ASTOpConstantResource();
 							newOp->Type = EOpType::ME_CONSTANT;
 							newOp->SetValue(NewMesh, options.DiskCacheContext);
 							newOp->SourceDataDescriptor = at->GetSourceDataDescriptor();
diff --git a/Engine/Plugins/Mutable/Source/MutableTools/Private/MuT/CodeOptimiser_Logic.cpp b/Engine/Plugins/Mutable/Source/MutableTools/Private/MuT/CodeOptimiser_Logic.cpp
--- a/Engine/Plugins/Mutable/Source/MutableTools/Private/MuT/CodeOptimiser_Logic.cpp
+++ b/Engine/Plugins/Mutable/Source/MutableTools/Private/MuT/CodeOptimiser_Logic.cpp
@@ -14,7 +14,7 @@
 #include "MuT/CodeOptimiser.h"
 
 
-namespace mu
+namespace UE::Mutable::Private
 {
 
 	bool LocalLogicOptimiserAST(ASTOpList& roots)
@@ -78,7 +78,7 @@ namespace mu
                                     else
                                     {
                                         // Other parents may not impose the same condition that allows the optimisation.
-                                        Ptr<ASTOpInstanceAdd> newAddSurface = mu::Clone<ASTOpInstanceAdd>(addSurface);
+                                        Ptr<ASTOpInstanceAdd> newAddSurface = UE::Mutable::Private::Clone<ASTOpInstanceAdd>(addSurface);
                                         newAddSurface->instance = bottomConditional->no.child();
                                         topConditional->yes = newAddSurface;
 										addSurface = newAddSurface.get();
@@ -134,7 +134,7 @@ namespace mu
                                     {
                                         // Other parents may not impose the same condition that allows
                                         // the optimisation.
-                                        Ptr<ASTOpMeshMerge> NewMerge = mu::Clone<ASTOpMeshMerge>(Merge);
+                                        Ptr<ASTOpMeshMerge> NewMerge = UE::Mutable::Private::Clone<ASTOpMeshMerge>(Merge);
 										NewMerge->Base = bottomConditional->no.child();
                                         topConditional->yes = NewMerge;
                                         Merge = NewMerge;
diff --git a/Engine/Plugins/Mutable/Source/MutableTools/Private/MuT/CodeOptimiser_Parameters.cpp b/Engine/Plugins/Mutable/Source/MutableTools/Private/MuT/CodeOptimiser_Parameters.cpp
--- a/Engine/Plugins/Mutable/Source/MutableTools/Private/MuT/CodeOptimiser_Parameters.cpp
+++ b/Engine/Plugins/Mutable/Source/MutableTools/Private/MuT/CodeOptimiser_Parameters.cpp
@@ -38,7 +38,7 @@
 #include <array>
 
 
-namespace mu
+namespace UE::Mutable::Private
 {
 
     RuntimeParameterVisitorAST::RuntimeParameterVisitorAST(const FStateCompilationData* InState)
@@ -395,27 +395,27 @@ namespace mu
                              )
                             {
                                 // Move the conditional down
-                                Ptr<ASTOpImageCompose> compOp = mu::Clone<ASTOpImageCompose>(typedYes);
+                                Ptr<ASTOpImageCompose> compOp = UE::Mutable::Private::Clone<ASTOpImageCompose>(typedYes);
 
-                                Ptr<ASTOpConditional> baseCond = mu::Clone<ASTOpConditional>(TypedOp);
+                                Ptr<ASTOpConditional> baseCond = UE::Mutable::Private::Clone<ASTOpConditional>(TypedOp);
                                 baseCond->yes = typedYes->Base.child();
                                 baseCond->no = typedNo->Base.child();
                                 compOp->Base = baseCond;
 
-                                Ptr<ASTOpConditional> blockCond = mu::Clone<ASTOpConditional>(TypedOp);
+                                Ptr<ASTOpConditional> blockCond = UE::Mutable::Private::Clone<ASTOpConditional>(TypedOp);
                                 blockCond->yes = typedYes->BlockImage.child();
                                 blockCond->no = typedNo->BlockImage.child();
                                 compOp->BlockImage = blockCond;
 
                                 if (typedYes->Mask)
                                 {
-                                    Ptr<ASTOpConditional> maskCond = mu::Clone<ASTOpConditional>(TypedOp);
+                                    Ptr<ASTOpConditional> maskCond = UE::Mutable::Private::Clone<ASTOpConditional>(TypedOp);
                                     maskCond->yes = typedYes->Mask.child();
                                     maskCond->no = typedNo->Mask.child();
                                     compOp->Mask = maskCond;
                                 }
 
-                                Ptr<ASTOpConditional> layCond = mu::Clone<ASTOpConditional>(TypedOp);
+                                Ptr<ASTOpConditional> layCond = UE::Mutable::Private::Clone<ASTOpConditional>(TypedOp);
                                 layCond->type = EOpType::LA_CONDITIONAL;
                                 layCond->yes = typedYes->Layout.child();
                                 layCond->no = typedNo->Layout.child();
@@ -458,17 +458,17 @@ namespace mu
 							ResizeOp->Source = plainOp;
 							ResizeOp->SizeSource = typedYes->base.child();
 
-                            Ptr<ASTOpConditional> maskOp = mu::Clone<ASTOpConditional>(TypedOp);
+                            Ptr<ASTOpConditional> maskOp = UE::Mutable::Private::Clone<ASTOpConditional>(TypedOp);
                             maskOp->no = ResizeOp;
 
                             // If there is no mask (because it is optional), we need to make a
                             // white plain image
                             maskOp->yes = EnsureValidMask( typedYes->mask.child(), typedYes->base.child() );
 
-							Ptr<ASTOpConditional> baseOp = mu::Clone<ASTOpConditional>(TypedOp);
+							Ptr<ASTOpConditional> baseOp = UE::Mutable::Private::Clone<ASTOpConditional>(TypedOp);
                             baseOp->yes = typedYes->base.child();
 
-                            Ptr<ASTOpImageLayerColor> softOp = mu::Clone<ASTOpImageLayerColor>(typedYes);
+                            Ptr<ASTOpImageLayerColor> softOp = UE::Mutable::Private::Clone<ASTOpImageLayerColor>(typedYes);
                             softOp->base = baseOp;
                             softOp->mask = maskOp;
 
@@ -559,17 +559,17 @@ namespace mu
 							ResizeOp->Source = plainOp;
 							ResizeOp->SizeSource = typedNo->base.child();
 
-                            Ptr<ASTOpConditional> maskOp = mu::Clone<ASTOpConditional>(TypedOp);
+                            Ptr<ASTOpConditional> maskOp = UE::Mutable::Private::Clone<ASTOpConditional>(TypedOp);
                             maskOp->no = ResizeOp;
 
                             // If there is no mask (because it is optional), we need to make a
                             // white plain image
                             maskOp->no = EnsureValidMask( typedNo->mask.child(), typedNo->base.child() );
 
-                            Ptr<ASTOpConditional> baseOp = mu::Clone<ASTOpConditional>(TypedOp);
+                            Ptr<ASTOpConditional> baseOp = UE::Mutable::Private::Clone<ASTOpConditional>(TypedOp);
                             baseOp->no = typedNo->base.child();
 
-                            Ptr<ASTOpImageLayerColor> softOp = mu::Clone<ASTOpImageLayerColor>(typedNo);
+                            Ptr<ASTOpImageLayerColor> softOp = UE::Mutable::Private::Clone<ASTOpImageLayerColor>(typedNo);
                             softOp->base = baseOp;
                             softOp->mask = maskOp;
 
@@ -872,17 +872,17 @@ namespace mu
 							Ptr<ASTOp> newBaseMask = EnsureValidMask(baseMask, baseImage);
 							Ptr<ASTOp> newBlockMask = EnsureValidMask(blockMask, blockImage);
 
-                            maskOp = mu::Clone<ASTOpImageCompose>(TypedOp);
+                            maskOp = UE::Mutable::Private::Clone<ASTOpImageCompose>(TypedOp);
                             maskOp->Base = newBaseMask;
                             maskOp->BlockImage = newBlockMask;
                         }
 
                         // The base is composition of the bases of both layer effect
-						Ptr<ASTOpImageCompose> baseOp = mu::Clone<ASTOpImageCompose>(TypedOp);
+						Ptr<ASTOpImageCompose> baseOp = UE::Mutable::Private::Clone<ASTOpImageCompose>(TypedOp);
                         baseOp->Base = baseImage;
                         baseOp->BlockImage = blockImage;
 
-                        Ptr<ASTOpImageLayerColor> nop = mu::Clone<ASTOpImageLayerColor>(blockAt);
+                        Ptr<ASTOpImageLayerColor> nop = UE::Mutable::Private::Clone<ASTOpImageLayerColor>(blockAt);
                         nop->mask = maskOp;
                         nop->base = baseOp;
 
@@ -914,22 +914,22 @@ namespace mu
 							Ptr<ASTOp> newBaseMask = EnsureValidMask(baseMask, baseImage);
 							Ptr<ASTOp> newBlockMask = EnsureValidMask(blockMask, blockImage);
 
-                            maskOp = mu::Clone<ASTOpImageCompose>(TypedOp);
+                            maskOp = UE::Mutable::Private::Clone<ASTOpImageCompose>(TypedOp);
                             maskOp->Base = newBaseMask;
                             maskOp->BlockImage = newBlockMask;
                         }
 
                         // The base is composition of the bases of both layer effect
-						Ptr<ASTOpImageCompose> baseOp = mu::Clone<ASTOpImageCompose>(TypedOp);
+						Ptr<ASTOpImageCompose> baseOp = UE::Mutable::Private::Clone<ASTOpImageCompose>(TypedOp);
                         baseOp->Base = baseImage;
                         baseOp->BlockImage = blockImage;
 
                         // The base is composition of the bases of both layer effect
-						Ptr<ASTOpImageCompose> blendedOp = mu::Clone<ASTOpImageCompose>(TypedOp);
+						Ptr<ASTOpImageCompose> blendedOp = UE::Mutable::Private::Clone<ASTOpImageCompose>(TypedOp);
                         blendedOp->Base = baseBlended;
                         blendedOp->BlockImage = blockBlended;
 
-                        Ptr<ASTOpImageLayer> nop = mu::Clone<ASTOpImageLayer>(blockAt);
+                        Ptr<ASTOpImageLayer> nop = UE::Mutable::Private::Clone<ASTOpImageLayer>(blockAt);
                         nop->mask = maskOp;
                         nop->base = baseOp;
                         nop->blend = blendedOp;
@@ -960,10 +960,10 @@ namespace mu
                     optimised = true;
 
                     // Swap
-					Ptr<ASTOpImageCompose> childCompose = mu::Clone<ASTOpImageCompose>(At);
+					Ptr<ASTOpImageCompose> childCompose = UE::Mutable::Private::Clone<ASTOpImageCompose>(At);
                     childCompose->Base = typedBaseAt->Base.child();
 
-					Ptr<ASTOpImageCompose> parentCompose = mu::Clone<ASTOpImageCompose>(baseAt);
+					Ptr<ASTOpImageCompose> parentCompose = UE::Mutable::Private::Clone<ASTOpImageCompose>(baseAt);
                     parentCompose->Base = childCompose;
 
                     At = parentCompose;
@@ -999,7 +999,7 @@ namespace mu
                     Ptr<ASTOpImageCompose> maskOp;
                     if (blockMask || baseType!=EOpType::IM_BLANKLAYOUT)
                     {
-                        maskOp = mu::Clone<ASTOpImageCompose>(At);
+                        maskOp = UE::Mutable::Private::Clone<ASTOpImageCompose>(At);
 						Ptr<ASTOp> newMaskBlock = EnsureValidMask(blockMask, blockImage);
                         maskOp->BlockImage = newMaskBlock;
 
@@ -1021,10 +1021,10 @@ namespace mu
                     }
 
                     // The base is composition of the layer base on the compose base
-                    Ptr<ASTOpImageCompose> baseOp = mu::Clone<ASTOpImageCompose>(At);
+                    Ptr<ASTOpImageCompose> baseOp = UE::Mutable::Private::Clone<ASTOpImageCompose>(At);
                     baseOp->BlockImage = typedBlockAt->base.child();
 
-                    Ptr<ASTOpImageLayerColor> nop = mu::Clone<ASTOpImageLayerColor>(blockAt);
+                    Ptr<ASTOpImageLayerColor> nop = UE::Mutable::Private::Clone<ASTOpImageLayerColor>(blockAt);
                     nop->mask = maskOp;
                     nop->base = baseOp;
 
@@ -1049,7 +1049,7 @@ namespace mu
                     Ptr<ASTOpImageCompose> maskOp;
                     if (blockMask || baseType != EOpType::IM_BLANKLAYOUT)
                     {
-                        maskOp = mu::Clone<ASTOpImageCompose>(At);
+                        maskOp = UE::Mutable::Private::Clone<ASTOpImageCompose>(At);
 						Ptr<ASTOp> newMaskBlock = EnsureValidMask(blockMask, blockImage);
                         maskOp->BlockImage = newMaskBlock;
 
@@ -1072,7 +1072,7 @@ namespace mu
                     }
 
                     // The blended is a compose of the blended image on a blank image
-					Ptr<ASTOpImageCompose> blendedOp = mu::Clone<ASTOpImageCompose>(At);
+					Ptr<ASTOpImageCompose> blendedOp = UE::Mutable::Private::Clone<ASTOpImageCompose>(At);
                     {
                         blendedOp->BlockImage = blockBlended;
 
@@ -1095,10 +1095,10 @@ namespace mu
                     }
 
                     // The base is composition of the softlight base on the compose base
-                    Ptr<ASTOpImageCompose> baseOp = mu::Clone<ASTOpImageCompose>(At);
+                    Ptr<ASTOpImageCompose> baseOp = UE::Mutable::Private::Clone<ASTOpImageCompose>(At);
                     baseOp->BlockImage = typedBlockAt->base.child();
 
-                    Ptr<ASTOpImageLayer> nop = mu::Clone<ASTOpImageLayer>(blockAt);
+                    Ptr<ASTOpImageLayer> nop = UE::Mutable::Private::Clone<ASTOpImageLayer>(blockAt);
                     nop->base = baseOp;
                     nop->mask = maskOp;
                     nop->blend = blendedOp;
@@ -1153,7 +1153,7 @@ namespace mu
 
                     const ASTOpImageLayerColor* typedBaseAt = static_cast<const ASTOpImageLayerColor*>(baseAt.get());
 
-					Ptr<ASTOpImageCompose> maskOp = mu::Clone<ASTOpImageCompose>(At);
+					Ptr<ASTOpImageCompose> maskOp = UE::Mutable::Private::Clone<ASTOpImageCompose>(At);
                     {
                         Ptr<ASTOpConstantColor> blackOp = new ASTOpConstantColor;
                         blackOp->Value = FVector4f(0,0,0,1);
@@ -1176,10 +1176,10 @@ namespace mu
                     }
 
                     // The base is composition of the softlight base on the compose base
-					Ptr<ASTOpImageCompose> baseOp = mu::Clone<ASTOpImageCompose>(At);
+					Ptr<ASTOpImageCompose> baseOp = UE::Mutable::Private::Clone<ASTOpImageCompose>(At);
                     baseOp->Base = typedBaseAt->base.child();
 
-                    Ptr<ASTOpImageLayerColor> nop = mu::Clone<ASTOpImageLayerColor>(baseAt);
+                    Ptr<ASTOpImageLayerColor> nop = UE::Mutable::Private::Clone<ASTOpImageLayerColor>(baseAt);
                     nop->base = baseOp;
                     nop->mask = maskOp;
 
@@ -1194,7 +1194,7 @@ namespace mu
 
                     const ASTOpImageLayer* typedBaseAt = static_cast<const ASTOpImageLayer*>(baseAt.get());
 
-					Ptr<ASTOpImageCompose> maskOp = mu::Clone<ASTOpImageCompose>(At);
+					Ptr<ASTOpImageCompose> maskOp = UE::Mutable::Private::Clone<ASTOpImageCompose>(At);
                     {
                         Ptr<ASTOpConstantColor> blackOp = new ASTOpConstantColor;
 						blackOp->Value = FVector4f(0, 0, 0, 1);
@@ -1217,10 +1217,10 @@ namespace mu
                     }
 
                     // The base is composition of the effect base on the compose base
-                    Ptr<ASTOpImageCompose> baseOp = mu::Clone<ASTOpImageCompose>(At);
+                    Ptr<ASTOpImageCompose> baseOp = UE::Mutable::Private::Clone<ASTOpImageCompose>(At);
                     baseOp->Base = typedBaseAt->base.child();
 
-                    Ptr<ASTOpImageLayer> nop = mu::Clone<ASTOpImageLayer>(baseAt);
+                    Ptr<ASTOpImageLayer> nop = UE::Mutable::Private::Clone<ASTOpImageLayer>(baseAt);
                     nop->base = baseOp;
                     nop->mask = maskOp;
 
@@ -1376,16 +1376,16 @@ namespace mu
                 {
                     bModified = true;
 
-					Ptr<ASTOpImageLayerColor> top = mu::Clone<ASTOpImageLayerColor>(sourceOp);
+					Ptr<ASTOpImageLayerColor> top = UE::Mutable::Private::Clone<ASTOpImageLayerColor>(sourceOp);
 
-					Ptr<ASTOpImageMipmap> baseOp = mu::Clone<ASTOpImageMipmap>(At);
+					Ptr<ASTOpImageMipmap> baseOp = UE::Mutable::Private::Clone<ASTOpImageMipmap>(At);
                     baseOp->Source = typedSource->base.child();
                     top->base = baseOp;
 
 					Ptr<ASTOp> sourceMaskOp = typedSource->mask.child();
                     if (sourceMaskOp)
                     {
-						Ptr<ASTOpImageMipmap> maskOp = mu::Clone<ASTOpImageMipmap>(At);
+						Ptr<ASTOpImageMipmap> maskOp = UE::Mutable::Private::Clone<ASTOpImageMipmap>(At);
                         maskOp->Source = sourceMaskOp;
                         top->mask = maskOp;
                     }
@@ -1881,7 +1881,7 @@ namespace mu
 
 					if (isCompressedFormat)
 					{
-						Ptr<ASTOpInstanceAdd> NewOp = mu::Clone<ASTOpInstanceAdd>(At);
+						Ptr<ASTOpInstanceAdd> NewOp = UE::Mutable::Private::Clone<ASTOpInstanceAdd>(At);
 
 						// Add a new format operation to uncompress the image
 						Ptr<ASTOpImagePixelFormat> fop = new ASTOpImagePixelFormat;
@@ -1959,7 +1959,7 @@ namespace mu
         	// Search for the first LOD that has a valid mesh.
         	const int32 FirstLOD = TypedOp->lods.IndexOfByPredicate([](const ASTChild& Element)
         	{
-        		TArray<mu::Ptr<ASTOp>> Roots;
+        		TArray<UE::Mutable::Private::Ptr<ASTOp>> Roots;
         		Roots.Add(Element.child());
         		
         		FFindMesh SearchMesh(Roots);
@@ -1971,7 +1971,7 @@ namespace mu
         	
             if (TypedOp->lods.Num() > NumLODs)
             {
-                Ptr<ASTOpAddLOD> NewOp = mu::Clone<ASTOpAddLOD>(At);
+                Ptr<ASTOpAddLOD> NewOp = UE::Mutable::Private::Clone<ASTOpAddLOD>(At);
                 while (NewOp->lods.Num() > NumLODs)
                 {
                     NewOp->lods.Pop();
diff --git a/Engine/Plugins/Mutable/Source/MutableTools/Private/MuT/Compiler.cpp b/Engine/Plugins/Mutable/Source/MutableTools/Private/MuT/Compiler.cpp
--- a/Engine/Plugins/Mutable/Source/MutableTools/Private/MuT/Compiler.cpp
+++ b/Engine/Plugins/Mutable/Source/MutableTools/Private/MuT/Compiler.cpp
@@ -35,7 +35,7 @@
 
 #include <inttypes.h>	// Required for 64-bit printf macros
 
-namespace mu
+namespace UE::Mutable::Private
 {
 
     const char* CompilerOptions::GetTextureLayoutStrategyName( CompilerOptions::TextureLayoutStrategy s )
@@ -557,7 +557,7 @@ namespace mu
 			}
 		}
 
-        mu::Private::RemoveDisabledAddInstanceOpsValues(m_pD->Options, States);
+        UE::Mutable::Private::Private::RemoveDisabledAddInstanceOpsValues(m_pD->Options, States);
 
         // Optimize the generated code
         {
@@ -737,7 +737,7 @@ namespace mu
 
 		uint64 EmbeddedDataBytesLimit = Options->GetPrivate()->EmbeddedDataBytesLimit;
 
-		mu::FProgram& Program = Model->GetPrivate()->Program;
+		UE::Mutable::Private::FProgram& Program = Model->GetPrivate()->Program;
 
 		// These are used for logging only.
 		int32 NumRoms = 0;
diff --git a/Engine/Plugins/Mutable/Source/MutableTools/Private/MuT/DataPacker.h b/Engine/Plugins/Mutable/Source/MutableTools/Private/MuT/DataPacker.h
--- a/Engine/Plugins/Mutable/Source/MutableTools/Private/MuT/DataPacker.h
+++ b/Engine/Plugins/Mutable/Source/MutableTools/Private/MuT/DataPacker.h
@@ -5,9 +5,9 @@
 #include "MuR/Operations.h"
 #include "MuT/AST.h"
 
-namespace mu { struct FProgram; }
+namespace UE::Mutable::Private { struct FProgram; }
 
-namespace mu
+namespace UE::Mutable::Private
 {
 	class CompilerOptions;
 
diff --git a/Engine/Plugins/Mutable/Source/MutableTools/Private/MuT/UnrealPixelFormatOverride.cpp b/Engine/Plugins/Mutable/Source/MutableTools/Private/MuT/UnrealPixelFormatOverride.cpp
--- a/Engine/Plugins/Mutable/Source/MutableTools/Private/MuT/UnrealPixelFormatOverride.cpp
+++ b/Engine/Plugins/Mutable/Source/MutableTools/Private/MuT/UnrealPixelFormatOverride.cpp
@@ -18,36 +18,36 @@
 
 static ITextureFormatManagerModule* STextureFormatManager = nullptr;
 
-static FName SPrefixedMutableTextureFormatNameTable[(uint32)mu::EImageFormat::Count];
+static FName SPrefixedMutableTextureFormatNameTable[(uint32)UE::Mutable::Private::EImageFormat::Count];
 static bool SbPrefixedMutableTextureFormatNameTableInitialized = false;
 
 namespace
 {
 
-FName GetMutableFormatTextureFormatName(mu::EImageFormat MutableFormat)
+FName GetMutableFormatTextureFormatName(UE::Mutable::Private::EImageFormat MutableFormat)
 {
 	switch (MutableFormat)
 	{
-		case mu::EImageFormat::BC1                 : return FName(TEXT("DXT1"));
-		case mu::EImageFormat::BC2                 : return FName(TEXT("DXT3"));
-		case mu::EImageFormat::BC3                 : return FName(TEXT("DXT5"));
-		case mu::EImageFormat::BC4                 : return FName(TEXT("BC4"));
-		case mu::EImageFormat::BC5                 : return FName(TEXT("BC5"));
-		case mu::EImageFormat::ASTC_4x4_RGB_LDR    : return FName(TEXT("ASTC_RGBA_HQ"));
-    	case mu::EImageFormat::ASTC_4x4_RGBA_LDR   : return FName(TEXT("ASTC_RGBA_HQ"));
-    	case mu::EImageFormat::ASTC_4x4_RG_LDR     : return FName(TEXT("ASTC_RGB"));
-    	case mu::EImageFormat::ASTC_8x8_RGB_LDR    : return FName(TEXT("ASTC_RGBA"));   
-    	case mu::EImageFormat::ASTC_8x8_RGBA_LDR   : return FName(TEXT("ASTC_RGBA"));  
-    	case mu::EImageFormat::ASTC_8x8_RG_LDR     : return FName(TEXT("ASTC_NormalLA"));
-    	case mu::EImageFormat::ASTC_12x12_RGB_LDR  : return FName(TEXT("ASTC_RGBA"));
-    	case mu::EImageFormat::ASTC_12x12_RGBA_LDR : return FName(TEXT("ASTC_RGBA"));	 
-    	case mu::EImageFormat::ASTC_12x12_RG_LDR   : return FName(TEXT("ASTC_NormalRG"));
-    	case mu::EImageFormat::ASTC_6x6_RGB_LDR    : return FName(TEXT("ASTC_RGBA"));
-    	case mu::EImageFormat::ASTC_6x6_RGBA_LDR   : return FName(TEXT("ASTC_RGBA"));   
-    	case mu::EImageFormat::ASTC_6x6_RG_LDR     : return FName(TEXT("ASTC_NormalRG"));
-    	case mu::EImageFormat::ASTC_10x10_RGB_LDR  : return FName(TEXT("ASTC_RGBA"));
-    	case mu::EImageFormat::ASTC_10x10_RGBA_LDR : return FName(TEXT("ASTC_RGBA"));  
-    	case mu::EImageFormat::ASTC_10x10_RG_LDR   : return FName(TEXT("ASTC_NormalRG"));
+		case UE::Mutable::Private::EImageFormat::BC1                 : return FName(TEXT("DXT1"));
+		case UE::Mutable::Private::EImageFormat::BC2                 : return FName(TEXT("DXT3"));
+		case UE::Mutable::Private::EImageFormat::BC3                 : return FName(TEXT("DXT5"));
+		case UE::Mutable::Private::EImageFormat::BC4                 : return FName(TEXT("BC4"));
+		case UE::Mutable::Private::EImageFormat::BC5                 : return FName(TEXT("BC5"));
+		case UE::Mutable::Private::EImageFormat::ASTC_4x4_RGB_LDR    : return FName(TEXT("ASTC_RGBA_HQ"));
+    	case UE::Mutable::Private::EImageFormat::ASTC_4x4_RGBA_LDR   : return FName(TEXT("ASTC_RGBA_HQ"));
+    	case UE::Mutable::Private::EImageFormat::ASTC_4x4_RG_LDR     : return FName(TEXT("ASTC_RGB"));
+    	case UE::Mutable::Private::EImageFormat::ASTC_8x8_RGB_LDR    : return FName(TEXT("ASTC_RGBA"));   
+    	case UE::Mutable::Private::EImageFormat::ASTC_8x8_RGBA_LDR   : return FName(TEXT("ASTC_RGBA"));  
+    	case UE::Mutable::Private::EImageFormat::ASTC_8x8_RG_LDR     : return FName(TEXT("ASTC_NormalLA"));
+    	case UE::Mutable::Private::EImageFormat::ASTC_12x12_RGB_LDR  : return FName(TEXT("ASTC_RGBA"));
+    	case UE::Mutable::Private::EImageFormat::ASTC_12x12_RGBA_LDR : return FName(TEXT("ASTC_RGBA"));	 
+    	case UE::Mutable::Private::EImageFormat::ASTC_12x12_RG_LDR   : return FName(TEXT("ASTC_NormalRG"));
+    	case UE::Mutable::Private::EImageFormat::ASTC_6x6_RGB_LDR    : return FName(TEXT("ASTC_RGBA"));
+    	case UE::Mutable::Private::EImageFormat::ASTC_6x6_RGBA_LDR   : return FName(TEXT("ASTC_RGBA"));   
+    	case UE::Mutable::Private::EImageFormat::ASTC_6x6_RG_LDR     : return FName(TEXT("ASTC_NormalRG"));
+    	case UE::Mutable::Private::EImageFormat::ASTC_10x10_RGB_LDR  : return FName(TEXT("ASTC_RGBA"));
+    	case UE::Mutable::Private::EImageFormat::ASTC_10x10_RGBA_LDR : return FName(TEXT("ASTC_RGBA"));  
+    	case UE::Mutable::Private::EImageFormat::ASTC_10x10_RG_LDR   : return FName(TEXT("ASTC_NormalRG"));
 
 		default: return NAME_None;
 	}
@@ -70,9 +70,9 @@ void PrepareUnrealCompression()
 		const ITargetPlatformSettings* TargetPlatformSettings = GetTargetPlatformManagerRef().GetRunningTargetPlatform()->GetTargetPlatformSettings();
 		check(TargetPlatformSettings);
 
-		for (uint32 FormatEnumIndex = 0; FormatEnumIndex < (uint32)mu::EImageFormat::Count; ++FormatEnumIndex)
+		for (uint32 FormatEnumIndex = 0; FormatEnumIndex < (uint32)UE::Mutable::Private::EImageFormat::Count; ++FormatEnumIndex)
 		{
-			FName TextureFormatName = GetMutableFormatTextureFormatName(mu::EImageFormat(FormatEnumIndex));
+			FName TextureFormatName = GetMutableFormatTextureFormatName(UE::Mutable::Private::EImageFormat(FormatEnumIndex));
 			SPrefixedMutableTextureFormatNameTable[FormatEnumIndex] = TextureFormatName;
 				
 			// Based on ConditionalGetPrefixedFormat in Texture.cpp
@@ -109,7 +109,7 @@ void PrepareUnrealCompression()
 	}
 }
 
-void FillBuildSettingsFromMutableFormat(FTextureBuildSettings& Settings, bool& bOutHasAlpha, mu::EImageFormat Format)
+void FillBuildSettingsFromMutableFormat(FTextureBuildSettings& Settings, bool& bOutHasAlpha, UE::Mutable::Private::EImageFormat Format)
 {
 	Settings.MipGenSettings = TMGS_NoMipmaps;
 
@@ -119,81 +119,81 @@ void FillBuildSettingsFromMutableFormat(FTextureBuildSettings& Settings, bool& b
 
 	switch (Format)
 	{
-	case mu::EImageFormat::ASTC_4x4_RGBA_LDR:
+	case UE::Mutable::Private::EImageFormat::ASTC_4x4_RGBA_LDR:
 		Settings.CompressionQuality = 4; // See GetQualityFormat in TextureFormatASTC.cpp
 		bOutHasAlpha = true;
 		break;
-	case mu::EImageFormat::ASTC_6x6_RGBA_LDR:
+	case UE::Mutable::Private::EImageFormat::ASTC_6x6_RGBA_LDR:
 		Settings.CompressionQuality = 3; // See GetQualityFormat in TextureFormatASTC.cpp
 		bOutHasAlpha = true;
 		break;
-	case mu::EImageFormat::ASTC_8x8_RGBA_LDR:
+	case UE::Mutable::Private::EImageFormat::ASTC_8x8_RGBA_LDR:
 		Settings.CompressionQuality = 2; // See GetQualityFormat in TextureFormatASTC.cpp
 		bOutHasAlpha = true;
 		break;
-	case mu::EImageFormat::ASTC_10x10_RGBA_LDR:
+	case UE::Mutable::Private::EImageFormat::ASTC_10x10_RGBA_LDR:
 		Settings.CompressionQuality = 1; // See GetQualityFormat in TextureFormatASTC.cpp
 		bOutHasAlpha = true;
 		break;
-	case mu::EImageFormat::ASTC_12x12_RGBA_LDR:
+	case UE::Mutable::Private::EImageFormat::ASTC_12x12_RGBA_LDR:
 		Settings.CompressionQuality = 0; // See GetQualityFormat in TextureFormatASTC.cpp
 		bOutHasAlpha = true;
 		break;
 
-	case mu::EImageFormat::ASTC_4x4_RGB_LDR:
+	case UE::Mutable::Private::EImageFormat::ASTC_4x4_RGB_LDR:
 		Settings.CompressionQuality = 4; // See GetQualityFormat in TextureFormatASTC.cpp
 		bOutHasAlpha = false;
 		break;
-	case mu::EImageFormat::ASTC_6x6_RGB_LDR:
+	case UE::Mutable::Private::EImageFormat::ASTC_6x6_RGB_LDR:
 		Settings.CompressionQuality = 3; // See GetQualityFormat in TextureFormatASTC.cpp
 		bOutHasAlpha = false;
 		break;
-	case mu::EImageFormat::ASTC_8x8_RGB_LDR:
+	case UE::Mutable::Private::EImageFormat::ASTC_8x8_RGB_LDR:
 		Settings.CompressionQuality = 2; // See GetQualityFormat in TextureFormatASTC.cpp
 		bOutHasAlpha = false;
 		break;
-	case mu::EImageFormat::ASTC_10x10_RGB_LDR:
+	case UE::Mutable::Private::EImageFormat::ASTC_10x10_RGB_LDR:
 		Settings.CompressionQuality = 1; // See GetQualityFormat in TextureFormatASTC.cpp
 		bOutHasAlpha = false;
 		break;
-	case mu::EImageFormat::ASTC_12x12_RGB_LDR:
+	case UE::Mutable::Private::EImageFormat::ASTC_12x12_RGB_LDR:
 		Settings.CompressionQuality = 0; // See GetQualityFormat in TextureFormatASTC.cpp
 		bOutHasAlpha = false;
 		break;
 
-	case mu::EImageFormat::ASTC_4x4_RG_LDR:
+	case UE::Mutable::Private::EImageFormat::ASTC_4x4_RG_LDR:
 		Settings.CompressionQuality = 4; // See GetQualityFormat in TextureFormatASTC.cpp
 		bOutHasAlpha = false;
 		break;
-	case mu::EImageFormat::ASTC_6x6_RG_LDR:
+	case UE::Mutable::Private::EImageFormat::ASTC_6x6_RG_LDR:
 		Settings.CompressionQuality = 3; // See GetQualityFormat in TextureFormatASTC.cpp
 		bOutHasAlpha = false;
 		break;
-	case mu::EImageFormat::ASTC_8x8_RG_LDR:
+	case UE::Mutable::Private::EImageFormat::ASTC_8x8_RG_LDR:
 		Settings.CompressionQuality = 2; // See GetQualityFormat in TextureFormatASTC.cpp
 		bOutHasAlpha = false;
 		break;
-	case mu::EImageFormat::ASTC_10x10_RG_LDR:
+	case UE::Mutable::Private::EImageFormat::ASTC_10x10_RG_LDR:
 		Settings.CompressionQuality = 1; // See GetQualityFormat in TextureFormatASTC.cpp
 		bOutHasAlpha = false;
 		break;
-	case mu::EImageFormat::ASTC_12x12_RG_LDR:
+	case UE::Mutable::Private::EImageFormat::ASTC_12x12_RG_LDR:
 		Settings.CompressionQuality = 0; // See GetQualityFormat in TextureFormatASTC.cpp
 		bOutHasAlpha = false;
 		break;
-	case mu::EImageFormat::BC1:
+	case UE::Mutable::Private::EImageFormat::BC1:
 		bOutHasAlpha = false;
 		break;
-	case mu::EImageFormat::BC2:
+	case UE::Mutable::Private::EImageFormat::BC2:
 		bOutHasAlpha = true;
 		break;
-	case mu::EImageFormat::BC3:
+	case UE::Mutable::Private::EImageFormat::BC3:
 		bOutHasAlpha = true;
 		break;
-	case mu::EImageFormat::BC4:
+	case UE::Mutable::Private::EImageFormat::BC4:
 		bOutHasAlpha = false;
 		break;
-	case mu::EImageFormat::BC5:
+	case UE::Mutable::Private::EImageFormat::BC5:
 		bOutHasAlpha = true;
 		break;
 	default:
@@ -204,27 +204,27 @@ void FillBuildSettingsFromMutableFormat(FTextureBuildSettings& Settings, bool& b
 }
 
 
-void MutableToImageCore(const mu::FImage* InMutable, FImage& CoreImage, int32 LOD, bool bSwizzleRGBHack)
+void MutableToImageCore(const UE::Mutable::Private::FImage* InMutable, FImage& CoreImage, int32 LOD, bool bSwizzleRGBHack)
 {
 	MUTABLE_CPUPROFILER_SCOPE(MutableToImageCore);
 
-	TSharedPtr<mu::FImage> TempMutable;
+	TSharedPtr<UE::Mutable::Private::FImage> TempMutable;
 
 	ERawImageFormat::Type CoreImageFormat;
 	switch (InMutable->GetFormat())
 	{
-	case mu::EImageFormat::BGRA_UByte: 
+	case UE::Mutable::Private::EImageFormat::BGRA_UByte: 
 		CoreImageFormat = ERawImageFormat::BGRA8; 
 		break;
 		
-	//case mu::EImageFormat::L_UByte: CoreImageFormat = ERawImageFormat::G8; break;
+	//case UE::Mutable::Private::EImageFormat::L_UByte: CoreImageFormat = ERawImageFormat::G8; break;
 
 	default:
 	{
 		// Unsupported format: force conversion
-		mu::FImageOperator ImOp = mu::FImageOperator::GetDefault(mu::FImageOperator::FImagePixelFormatFunc());
+		UE::Mutable::Private::FImageOperator ImOp = UE::Mutable::Private::FImageOperator::GetDefault(UE::Mutable::Private::FImageOperator::FImagePixelFormatFunc());
 
-		TempMutable = ImOp.ImagePixelFormat(4, InMutable, mu::EImageFormat::BGRA_UByte, LOD);
+		TempMutable = ImOp.ImagePixelFormat(4, InMutable, UE::Mutable::Private::EImageFormat::BGRA_UByte, LOD);
 		InMutable = TempMutable.Get();
 
 		if (bSwizzleRGBHack)
@@ -254,7 +254,7 @@ void MutableToImageCore(const mu::FImage* InMutable, FImage& CoreImage, int32 LO
 }
 
 
-bool ImageCoreToMutable(const FCompressedImage2D& Compressed, mu::FImage* Mutable, int32 LOD)
+bool ImageCoreToMutable(const FCompressedImage2D& Compressed, UE::Mutable::Private::FImage* Mutable, int32 LOD)
 {
 	TArrayView<uint8> MutableView = Mutable->DataStorage.GetLOD(LOD);
 
@@ -275,56 +275,56 @@ bool ImageCoreToMutable(const FCompressedImage2D& Compressed, mu::FImage* Mutabl
 }
 
 
-mu::EImageFormat UnrealToMutablePixelFormat(EPixelFormat PlatformFormat, bool bHasAlpha)
+UE::Mutable::Private::EImageFormat UnrealToMutablePixelFormat(EPixelFormat PlatformFormat, bool bHasAlpha)
 {
 	switch (PlatformFormat)
 	{
-	case PF_ASTC_4x4: return bHasAlpha ? mu::EImageFormat::ASTC_4x4_RGBA_LDR : mu::EImageFormat::ASTC_4x4_RGB_LDR;
-	case PF_ASTC_6x6: return bHasAlpha ? mu::EImageFormat::ASTC_6x6_RGBA_LDR : mu::EImageFormat::ASTC_6x6_RGB_LDR;
-	case PF_ASTC_8x8: return bHasAlpha ? mu::EImageFormat::ASTC_8x8_RGBA_LDR : mu::EImageFormat::ASTC_8x8_RGB_LDR;
-	case PF_ASTC_10x10: return bHasAlpha ? mu::EImageFormat::ASTC_10x10_RGBA_LDR : mu::EImageFormat::ASTC_10x10_RGB_LDR;
-	case PF_ASTC_12x12: return bHasAlpha ? mu::EImageFormat::ASTC_12x12_RGBA_LDR : mu::EImageFormat::ASTC_12x12_RGB_LDR;
-	case PF_ASTC_4x4_NORM_RG: return mu::EImageFormat::ASTC_4x4_RG_LDR;
-	case PF_ASTC_6x6_NORM_RG: return mu::EImageFormat::ASTC_6x6_RG_LDR;
-	case PF_ASTC_8x8_NORM_RG: return mu::EImageFormat::ASTC_8x8_RG_LDR;
-	case PF_ASTC_10x10_NORM_RG: return mu::EImageFormat::ASTC_10x10_RG_LDR;
-	case PF_ASTC_12x12_NORM_RG: return mu::EImageFormat::ASTC_12x12_RG_LDR;
-	case PF_DXT1: return mu::EImageFormat::BC1;
-	case PF_DXT3: return mu::EImageFormat::BC2;
-	case PF_DXT5: return mu::EImageFormat::BC3;
-	case PF_BC4: return mu::EImageFormat::BC4;
-	case PF_BC5: return mu::EImageFormat::BC5;
-	case PF_G8: return mu::EImageFormat::L_UByte;
-	case PF_L8: return mu::EImageFormat::L_UByte;
-	case PF_A8: return mu::EImageFormat::L_UByte;
-	case PF_R8G8B8A8: return mu::EImageFormat::RGBA_UByte;
-	case PF_A8R8G8B8: return mu::EImageFormat::RGBA_UByte;
-	case PF_B8G8R8A8: return mu::EImageFormat::BGRA_UByte;
+	case PF_ASTC_4x4: return bHasAlpha ? UE::Mutable::Private::EImageFormat::ASTC_4x4_RGBA_LDR : UE::Mutable::Private::EImageFormat::ASTC_4x4_RGB_LDR;
+	case PF_ASTC_6x6: return bHasAlpha ? UE::Mutable::Private::EImageFormat::ASTC_6x6_RGBA_LDR : UE::Mutable::Private::EImageFormat::ASTC_6x6_RGB_LDR;
+	case PF_ASTC_8x8: return bHasAlpha ? UE::Mutable::Private::EImageFormat::ASTC_8x8_RGBA_LDR : UE::Mutable::Private::EImageFormat::ASTC_8x8_RGB_LDR;
+	case PF_ASTC_10x10: return bHasAlpha ? UE::Mutable::Private::EImageFormat::ASTC_10x10_RGBA_LDR : UE::Mutable::Private::EImageFormat::ASTC_10x10_RGB_LDR;
+	case PF_ASTC_12x12: return bHasAlpha ? UE::Mutable::Private::EImageFormat::ASTC_12x12_RGBA_LDR : UE::Mutable::Private::EImageFormat::ASTC_12x12_RGB_LDR;
+	case PF_ASTC_4x4_NORM_RG: return UE::Mutable::Private::EImageFormat::ASTC_4x4_RG_LDR;
+	case PF_ASTC_6x6_NORM_RG: return UE::Mutable::Private::EImageFormat::ASTC_6x6_RG_LDR;
+	case PF_ASTC_8x8_NORM_RG: return UE::Mutable::Private::EImageFormat::ASTC_8x8_RG_LDR;
+	case PF_ASTC_10x10_NORM_RG: return UE::Mutable::Private::EImageFormat::ASTC_10x10_RG_LDR;
+	case PF_ASTC_12x12_NORM_RG: return UE::Mutable::Private::EImageFormat::ASTC_12x12_RG_LDR;
+	case PF_DXT1: return UE::Mutable::Private::EImageFormat::BC1;
+	case PF_DXT3: return UE::Mutable::Private::EImageFormat::BC2;
+	case PF_DXT5: return UE::Mutable::Private::EImageFormat::BC3;
+	case PF_BC4: return UE::Mutable::Private::EImageFormat::BC4;
+	case PF_BC5: return UE::Mutable::Private::EImageFormat::BC5;
+	case PF_G8: return UE::Mutable::Private::EImageFormat::L_UByte;
+	case PF_L8: return UE::Mutable::Private::EImageFormat::L_UByte;
+	case PF_A8: return UE::Mutable::Private::EImageFormat::L_UByte;
+	case PF_R8G8B8A8: return UE::Mutable::Private::EImageFormat::RGBA_UByte;
+	case PF_A8R8G8B8: return UE::Mutable::Private::EImageFormat::RGBA_UByte;
+	case PF_B8G8R8A8: return UE::Mutable::Private::EImageFormat::BGRA_UByte;
 	default:
-		return mu::EImageFormat::None;
+		return UE::Mutable::Private::EImageFormat::None;
 	}
 }
 
 
-mu::EImageFormat QualityAndPerformanceFix(mu::EImageFormat Format)
+UE::Mutable::Private::EImageFormat QualityAndPerformanceFix(UE::Mutable::Private::EImageFormat Format)
 {
 	switch (Format)
 	{
-	case mu::EImageFormat::ASTC_8x8_RGB_LDR:	return mu::EImageFormat::ASTC_4x4_RGB_LDR; 
-	case mu::EImageFormat::ASTC_8x8_RGBA_LDR:	return mu::EImageFormat::ASTC_4x4_RGBA_LDR;
-	case mu::EImageFormat::ASTC_8x8_RG_LDR:		return mu::EImageFormat::ASTC_4x4_RG_LDR;
-	case mu::EImageFormat::ASTC_12x12_RGB_LDR:	return mu::EImageFormat::ASTC_4x4_RGB_LDR; 
-	case mu::EImageFormat::ASTC_12x12_RGBA_LDR:	return mu::EImageFormat::ASTC_4x4_RGBA_LDR;
-	case mu::EImageFormat::ASTC_12x12_RG_LDR:	return mu::EImageFormat::ASTC_4x4_RG_LDR;
-	case mu::EImageFormat::ASTC_6x6_RGB_LDR:	return mu::EImageFormat::ASTC_4x4_RGB_LDR; 
-	case mu::EImageFormat::ASTC_6x6_RGBA_LDR:	return mu::EImageFormat::ASTC_4x4_RGBA_LDR;
-	case mu::EImageFormat::ASTC_6x6_RG_LDR:		return mu::EImageFormat::ASTC_4x4_RG_LDR;
-	case mu::EImageFormat::ASTC_10x10_RGB_LDR:	return mu::EImageFormat::ASTC_4x4_RGB_LDR; 
-	case mu::EImageFormat::ASTC_10x10_RGBA_LDR:	return mu::EImageFormat::ASTC_4x4_RGBA_LDR;
-	case mu::EImageFormat::ASTC_10x10_RG_LDR:	return mu::EImageFormat::ASTC_4x4_RG_LDR;
+	case UE::Mutable::Private::EImageFormat::ASTC_8x8_RGB_LDR:	return UE::Mutable::Private::EImageFormat::ASTC_4x4_RGB_LDR; 
+	case UE::Mutable::Private::EImageFormat::ASTC_8x8_RGBA_LDR:	return UE::Mutable::Private::EImageFormat::ASTC_4x4_RGBA_LDR;
+	case UE::Mutable::Private::EImageFormat::ASTC_8x8_RG_LDR:		return UE::Mutable::Private::EImageFormat::ASTC_4x4_RG_LDR;
+	case UE::Mutable::Private::EImageFormat::ASTC_12x12_RGB_LDR:	return UE::Mutable::Private::EImageFormat::ASTC_4x4_RGB_LDR; 
+	case UE::Mutable::Private::EImageFormat::ASTC_12x12_RGBA_LDR:	return UE::Mutable::Private::EImageFormat::ASTC_4x4_RGBA_LDR;
+	case UE::Mutable::Private::EImageFormat::ASTC_12x12_RG_LDR:	return UE::Mutable::Private::EImageFormat::ASTC_4x4_RG_LDR;
+	case UE::Mutable::Private::EImageFormat::ASTC_6x6_RGB_LDR:	return UE::Mutable::Private::EImageFormat::ASTC_4x4_RGB_LDR; 
+	case UE::Mutable::Private::EImageFormat::ASTC_6x6_RGBA_LDR:	return UE::Mutable::Private::EImageFormat::ASTC_4x4_RGBA_LDR;
+	case UE::Mutable::Private::EImageFormat::ASTC_6x6_RG_LDR:		return UE::Mutable::Private::EImageFormat::ASTC_4x4_RG_LDR;
+	case UE::Mutable::Private::EImageFormat::ASTC_10x10_RGB_LDR:	return UE::Mutable::Private::EImageFormat::ASTC_4x4_RGB_LDR; 
+	case UE::Mutable::Private::EImageFormat::ASTC_10x10_RGBA_LDR:	return UE::Mutable::Private::EImageFormat::ASTC_4x4_RGBA_LDR;
+	case UE::Mutable::Private::EImageFormat::ASTC_10x10_RG_LDR:	return UE::Mutable::Private::EImageFormat::ASTC_4x4_RG_LDR;
 
 	// This is more of a performance fix.
-	case mu::EImageFormat::BGRA_UByte:			return mu::EImageFormat::RGBA_UByte;
+	case UE::Mutable::Private::EImageFormat::BGRA_UByte:			return UE::Mutable::Private::EImageFormat::RGBA_UByte;
 
 	default:
 		break;
@@ -334,7 +334,7 @@ mu::EImageFormat QualityAndPerformanceFix(mu::EImageFormat Format)
 }
 
 
-void UnrealPixelFormatFunc(bool& bOutSuccess, int32 Quality, mu::FImage* Target, const mu::FImage* Source, int32 OnlyLOD)
+void UnrealPixelFormatFunc(bool& bOutSuccess, int32 Quality, UE::Mutable::Private::FImage* Target, const UE::Mutable::Private::FImage* Source, int32 OnlyLOD)
 {
 	// If this fails, PrepareUnrealCompression wasn't called before.
 	check(STextureFormatManager);
@@ -365,11 +365,11 @@ void UnrealPixelFormatFunc(bool& bOutSuccess, int32 Quality, mu::FImage* Target,
 
 	// This seems to be necessary because of a probable double swizzling that happens during conversions.
 	bool bSwizzleRGBHack =
-		Target->GetFormat() == mu::EImageFormat::ASTC_4x4_RG_LDR ||
-		Target->GetFormat() == mu::EImageFormat::ASTC_6x6_RG_LDR ||
-		Target->GetFormat() == mu::EImageFormat::ASTC_8x8_RG_LDR ||
-		Target->GetFormat() == mu::EImageFormat::ASTC_10x10_RG_LDR ||
-		Target->GetFormat() == mu::EImageFormat::ASTC_12x12_RG_LDR;
+		Target->GetFormat() == UE::Mutable::Private::EImageFormat::ASTC_4x4_RG_LDR ||
+		Target->GetFormat() == UE::Mutable::Private::EImageFormat::ASTC_6x6_RG_LDR ||
+		Target->GetFormat() == UE::Mutable::Private::EImageFormat::ASTC_8x8_RG_LDR ||
+		Target->GetFormat() == UE::Mutable::Private::EImageFormat::ASTC_10x10_RG_LDR ||
+		Target->GetFormat() == UE::Mutable::Private::EImageFormat::ASTC_12x12_RG_LDR;
 
 	for (int32 LOD = FirstLOD; bOutSuccess && (LOD < LODCount); ++LOD)
 	{
@@ -391,7 +391,7 @@ void UnrealPixelFormatFunc(bool& bOutSuccess, int32 Quality, mu::FImage* Target,
 }
 
 
-void DebugImageDump(const mu::FImage* Image, const FString& FileName)
+void DebugImageDump(const UE::Mutable::Private::FImage* Image, const FString& FileName)
 {
 	int32 LOD = 0;
 
@@ -430,8 +430,8 @@ void DebugImageDump(const mu::FImage* Image, const FString& FileName)
 	Header.dim_y[2] = (SizeY >> 16) & 0xff;
 	Header.dim_z[0] = 1;
 
-	Header.block_x = mu::GetImageFormatData(Image->GetFormat()).PixelsPerBlockX;
-	Header.block_y = mu::GetImageFormatData(Image->GetFormat()).PixelsPerBlockY;
+	Header.block_x = UE::Mutable::Private::GetImageFormatData(Image->GetFormat()).PixelsPerBlockX;
+	Header.block_y = UE::Mutable::Private::GetImageFormatData(Image->GetFormat()).PixelsPerBlockY;
 	Header.block_z = 1;
 
 	IPlatformFile& PlatformFile = FPlatformFileManager::Get().GetPlatformFile();
diff --git a/Engine/Plugins/Mutable/Source/MutableTools/Public/MuT/NodeColourVariation.h b/Engine/Plugins/Mutable/Source/MutableTools/Public/MuT/NodeColourVariation.h
--- a/Engine/Plugins/Mutable/Source/MutableTools/Public/MuT/NodeColourVariation.h
+++ b/Engine/Plugins/Mutable/Source/MutableTools/Public/MuT/NodeColourVariation.h
@@ -10,7 +10,7 @@
 #define UE_API MUTABLETOOLS_API
 
 
-namespace mu
+namespace UE::Mutable::Private
 {
 
 	/** Select different color subgraphs based on active tags. */
@@ -46,6 +46,6 @@ namespace mu
 	};
 
 
-} // namespace mu
+} // namespace UE::Mutable::Private
 
 #undef UE_API
diff --git a/Engine/Plugins/Mutable/Source/MutableTools/Public/MuT/TypeInfo.h b/Engine/Plugins/Mutable/Source/MutableTools/Public/MuT/TypeInfo.h
--- a/Engine/Plugins/Mutable/Source/MutableTools/Public/MuT/TypeInfo.h
+++ b/Engine/Plugins/Mutable/Source/MutableTools/Public/MuT/TypeInfo.h
@@ -13,7 +13,7 @@
 // This header file provides additional information about data types defined in the runtime, that
 // can be useful only for the tools. For instance, it provides strings for some enumeration values.
 
-namespace mu
+namespace UE::Mutable::Private
 {
 
 	class TypeInfo
@@ -28,7 +28,7 @@ namespace mu
 
 		static UE_API const char* s_blendModeName[uint32(EBlendType::_BT_COUNT) ];
 
-		static UE_API const char* s_projectorTypeName [ static_cast<uint32>(mu::EProjectorType::Count) ];
+		static UE_API const char* s_projectorTypeName [ static_cast<uint32>(UE::Mutable::Private::EProjectorType::Count) ];
 
 	};
 
diff --git a/Engine/Plugins/Mutable/Source/MutableTools/Public/MuT/UnrealPixelFormatOverride.h b/Engine/Plugins/Mutable/Source/MutableTools/Public/MuT/UnrealPixelFormatOverride.h
--- a/Engine/Plugins/Mutable/Source/MutableTools/Public/MuT/UnrealPixelFormatOverride.h
+++ b/Engine/Plugins/Mutable/Source/MutableTools/Public/MuT/UnrealPixelFormatOverride.h
@@ -14,25 +14,25 @@
 MUTABLETOOLS_API extern void PrepareUnrealCompression();
 
 /** */
-void FillBuildSettingsFromMutableFormat(FTextureBuildSettings& Settings, bool& bOutHasAlpha, mu::EImageFormat Format);
+void FillBuildSettingsFromMutableFormat(FTextureBuildSettings& Settings, bool& bOutHasAlpha, UE::Mutable::Private::EImageFormat Format);
 
 /** */
-void MutableToImageCore(const mu::FImage* InMutable, FImage& CoreImage, int32 LOD, bool bSwizzleRGBHack);
+void MutableToImageCore(const UE::Mutable::Private::FImage* InMutable, FImage& CoreImage, int32 LOD, bool bSwizzleRGBHack);
 
 /** */
-bool ImageCoreToMutable(const FCompressedImage2D& Compressed, mu::FImage* Mutable, int32 LOD);
+bool ImageCoreToMutable(const FCompressedImage2D& Compressed, UE::Mutable::Private::FImage* Mutable, int32 LOD);
 
 /** Convert an Unreal platform pixel format to Mutable. */
-MUTABLETOOLS_API extern mu::EImageFormat UnrealToMutablePixelFormat(EPixelFormat PlatformFormat, bool bHasAlpha);
+MUTABLETOOLS_API extern UE::Mutable::Private::EImageFormat UnrealToMutablePixelFormat(EPixelFormat PlatformFormat, bool bHasAlpha);
 
 /** Function that remaps some formats to the ones that provide more acceptable quality in Mutable. */
-MUTABLETOOLS_API extern mu::EImageFormat QualityAndPerformanceFix(mu::EImageFormat Format);
+MUTABLETOOLS_API extern UE::Mutable::Private::EImageFormat QualityAndPerformanceFix(UE::Mutable::Private::EImageFormat Format);
 
 /** Try to convert an image using unreal's compression. If the format is not supported, it will return false in bOutSuccess.
 * Can be called from any thread.
 * PrepareUnrealCompression needs to be called from the game thread once prior to using this function. 
 */
-MUTABLETOOLS_API extern void UnrealPixelFormatFunc(bool& bOutSuccess, int32 Quality, mu::FImage* Target, const mu::FImage* Source, int32 OnlyLOD);
+MUTABLETOOLS_API extern void UnrealPixelFormatFunc(bool& bOutSuccess, int32 Quality, UE::Mutable::Private::FImage* Target, const UE::Mutable::Private::FImage* Source, int32 OnlyLOD);
 
 /** Debug export the given image to a temp file in the saved folder. */
-MUTABLETOOLS_API void DebugImageDump(const mu::FImage*, const FString& FilePath);
+MUTABLETOOLS_API void DebugImageDump(const UE::Mutable::Private::FImage*, const FString& FilePath);
diff --git a/Engine/Plugins/Mutable/Source/MutableValidation/Private/MuV/AssetValidator_ReferencedCustomizableObjects.cpp b/Engine/Plugins/Mutable/Source/MutableValidation/Private/MuV/AssetValidator_ReferencedCustomizableObjects.cpp
--- a/Engine/Plugins/Mutable/Source/MutableValidation/Private/MuV/AssetValidator_ReferencedCustomizableObjects.cpp
+++ b/Engine/Plugins/Mutable/Source/MutableValidation/Private/MuV/AssetValidator_ReferencedCustomizableObjects.cpp
@@ -126,7 +126,7 @@ TSet< UCustomizableObject*> UAssetValidator_ReferencedCustomizableObjects::FindC
         	// We have found a referenced CustomizableObject
             if (ReferencerAssetData.GetClass() == UCustomizableObject::StaticClass())
             {
-				UObject* ReferencedAsset = MutablePrivate::LoadObject(ReferencerAssetData);
+				UObject* ReferencedAsset = UE::Mutable::Private::LoadObject(ReferencerAssetData);
             	if (ReferencedAsset)
             	{
             		UCustomizableObject* CastedCustomizableObject = Cast<UCustomizableObject>(ReferencedAsset);
diff --git a/Engine/Plugins/Mutable/Source/MutableValidation/Private/MuV/COIBakingTestCommandlet.cpp b/Engine/Plugins/Mutable/Source/MutableValidation/Private/MuV/COIBakingTestCommandlet.cpp
--- a/Engine/Plugins/Mutable/Source/MutableValidation/Private/MuV/COIBakingTestCommandlet.cpp
+++ b/Engine/Plugins/Mutable/Source/MutableValidation/Private/MuV/COIBakingTestCommandlet.cpp
@@ -81,7 +81,7 @@ int32 UCOIBakingTestCommandlet::Main(const FString& Params)
 		}
 	
 		// Load the resource
-		UObject* FoundObject = MutablePrivate::LoadObject(FSoftObjectPath(CustomizableObjectInstanceAssetPath));
+		UObject* FoundObject = UE::Mutable::Private::LoadObject(FSoftObjectPath(CustomizableObjectInstanceAssetPath));
 		if (!FoundObject)
 		{
 			UE_LOG(LogMutable,Error,TEXT("Failed to retrieve UObject from path %s"),*CustomizableObjectInstanceAssetPath);
diff --git a/Engine/Plugins/Mutable/Source/MutableValidation/Private/MuV/COIBulkUpdateTestCommandlet.cpp b/Engine/Plugins/Mutable/Source/MutableValidation/Private/MuV/COIBulkUpdateTestCommandlet.cpp
--- a/Engine/Plugins/Mutable/Source/MutableValidation/Private/MuV/COIBulkUpdateTestCommandlet.cpp
+++ b/Engine/Plugins/Mutable/Source/MutableValidation/Private/MuV/COIBulkUpdateTestCommandlet.cpp
@@ -66,7 +66,7 @@ int32 UCOIBulkUpdateTestCommandlet::Main(const FString& Params)
 	TMap<UCustomizableObject*, TArray<TStrongObjectPtr<UCustomizableObjectInstance>>>  MutableResources;
 	for (const FAssetData& Data : FoundAssetData)
 	{
-		UObject* LoadedAsset =  MutablePrivate::LoadObject(Data);
+		UObject* LoadedAsset =  UE::Mutable::Private::LoadObject(Data);
 		check (LoadedAsset);
 		UCustomizableObjectInstance* LoadedInstance =  Cast<UCustomizableObjectInstance>(LoadedAsset);
 		check (LoadedInstance)
diff --git a/Engine/Plugins/Mutable/Source/MutableValidation/Private/MuV/CustomizableObjectBulkValidationCommandlet.cpp b/Engine/Plugins/Mutable/Source/MutableValidation/Private/MuV/CustomizableObjectBulkValidationCommandlet.cpp
--- a/Engine/Plugins/Mutable/Source/MutableValidation/Private/MuV/CustomizableObjectBulkValidationCommandlet.cpp
+++ b/Engine/Plugins/Mutable/Source/MutableValidation/Private/MuV/CustomizableObjectBulkValidationCommandlet.cpp
@@ -111,7 +111,7 @@ int32 UCustomizableObjectBulkValidationCommandlet::Main(const FString& Params)
 		LLM_SCOPE_BYNAME(TEXT("CustomizableObjectBulkValidationCommandlet/COTest"));
 
 		const FAssetData& CustomizableObjectAssetData = FoundAssetData[AssetIndex];
-		TObjectPtr<UObject> FoundObject = MutablePrivate::LoadObject(CustomizableObjectAssetData);
+		TObjectPtr<UObject> FoundObject = UE::Mutable::Private::LoadObject(CustomizableObjectAssetData);
 		if (!FoundObject)
 		{
 			UE_LOG(LogMutable, Error, TEXT("Failed to load the asset with path : %s ."), *CustomizableObjectAssetData.GetSoftObjectPath().ToString())
diff --git a/Engine/Plugins/Mutable/Source/MutableValidation/Private/MuV/CustomizableObjectCompilationUtility.cpp b/Engine/Plugins/Mutable/Source/MutableValidation/Private/MuV/CustomizableObjectCompilationUtility.cpp
--- a/Engine/Plugins/Mutable/Source/MutableValidation/Private/MuV/CustomizableObjectCompilationUtility.cpp
+++ b/Engine/Plugins/Mutable/Source/MutableValidation/Private/MuV/CustomizableObjectCompilationUtility.cpp
@@ -102,8 +102,8 @@ bool FCustomizableObjectCompilationUtility::CompileCustomizableObject(UCustomiza
 #endif
 		
 		// Get the memory being used by mutable before the compilation
-		const int64 CompilationStartMutableBytes = mu::FGlobalMemoryCounter::GetAbsoluteCounter();
-		mu::FGlobalMemoryCounter::Zero();
+		const int64 CompilationStartMutableBytes = UE::Mutable::Private::FGlobalMemoryCounter::GetAbsoluteCounter();
+		UE::Mutable::Private::FGlobalMemoryCounter::Zero();
 		
 		UE_LOG(LogMutable,Display,TEXT("Compiling Customizable Object..."));
 		
@@ -160,7 +160,7 @@ bool FCustomizableObjectCompilationUtility::CompileCustomizableObject(UCustomiza
 		}
 		
 		// Get the peak mutable memory used during the compilation operation
-		const int64 CompilationEndPeakMutableBytes = mu::FGlobalMemoryCounter::GetPeak();
+		const int64 CompilationEndPeakMutableBytes = UE::Mutable::Private::FGlobalMemoryCounter::GetPeak();
 		const int64 CompilationEndRealMutablePeakBytes = CompilationStartMutableBytes + CompilationEndPeakMutableBytes;
 		UE_LOG(LogMutable, Display, TEXT("Pre-Compilation Mutable memory usage : %lld"), CompilationStartMutableBytes);
 		UE_LOG(LogMutable, Display, TEXT("Compilation Mutable peak memory usage : %lld"), CompilationEndPeakMutableBytes);
diff --git a/Engine/Plugins/Mutable/Source/MutableValidation/Private/MuV/CustomizableObjectValidationCommandlet.cpp b/Engine/Plugins/Mutable/Source/MutableValidation/Private/MuV/CustomizableObjectValidationCommandlet.cpp
--- a/Engine/Plugins/Mutable/Source/MutableValidation/Private/MuV/CustomizableObjectValidationCommandlet.cpp
+++ b/Engine/Plugins/Mutable/Source/MutableValidation/Private/MuV/CustomizableObjectValidationCommandlet.cpp
@@ -43,7 +43,7 @@ int32 UCustomizableObjectValidationCommandlet::Main(const FString& Params)
 		}
 	
 		// Load the resource
-		UObject* FoundObject = MutablePrivate::LoadObject(FSoftObjectPath(CustomizableObjectAssetPath));
+		UObject* FoundObject = UE::Mutable::Private::LoadObject(FSoftObjectPath(CustomizableObjectAssetPath));
 		if (!FoundObject)
 		{
 			UE_LOG(LogMutable, Error, TEXT("Failed to retrieve UObject from path %s"), *CustomizableObjectAssetPath);
diff --git a/Engine/Plugins/Mutable/Source/MutableValidation/Private/MuV/ValidationUtils.cpp b/Engine/Plugins/Mutable/Source/MutableValidation/Private/MuV/ValidationUtils.cpp
--- a/Engine/Plugins/Mutable/Source/MutableValidation/Private/MuV/ValidationUtils.cpp
+++ b/Engine/Plugins/Mutable/Source/MutableValidation/Private/MuV/ValidationUtils.cpp
@@ -218,7 +218,7 @@ bool TestCustomizableObject(UCustomizableObject& InTargetCustomizableObject, con
 	
 	// GHet the total size of the streaming data of the model ---------------------------------------------- //
 	{
-		const TSharedPtr<const mu::FModel> MutableModel = InTargetCustomizableObject.GetPrivate()->GetModel();
+		const TSharedPtr<const UE::Mutable::Private::FModel> MutableModel = InTargetCustomizableObject.GetPrivate()->GetModel();
 		check (MutableModel);
 
 		// Roms ---------------------- //
