[Mutable] Move BoneNames.cpp

diff --git a/Engine/Plugins/Mutable/Source/CustomizableObject/Private/BoneNames.cpp b/Engine/Plugins/Mutable/Source/CustomizableObject/Private/BoneNames.cpp
deleted file mode 100644
diff --git a/Engine/Plugins/Mutable/Source/CustomizableObject/Private/MuCO/BoneNames.cpp b/Engine/Plugins/Mutable/Source/CustomizableObject/Private/MuCO/BoneNames.cpp
new file mode 100644
--- /dev/null
+++ b/Engine/Plugins/Mutable/Source/CustomizableObject/Private/MuCO/BoneNames.cpp
@@ -0,0 +1,65 @@
+﻿// Copyright Epic Games, Inc. All Rights Reserved.
+
+#include "MuCO/BoneNames.h"
+
+#include "MuCO/CustomizableObjectPrivate.h"
+
+
+FBoneNames::FBoneNames(const UModelResources& ModelResources)
+{
+	BoneNamesMap.Reserve(ModelResources.BoneNamesMap.Num());
+
+	for (const TPair<FString, uint32>& Pair : ModelResources.BoneNamesMap)
+	{
+		BoneNamesMap.Add(Pair.Key, Pair.Value);
+	}
+}
+
+
+UE::Mutable::Private::FBoneName* FBoneNames::Find(const FName& BoneName)
+{
+	FScopeLock Lock(&CriticalSection);
+
+	return BoneNamesMap.Find(BoneName.ToString().ToLower());
+}
+
+
+UE::Mutable::Private::FBoneName FBoneNames::FindOrAdd(const FName& BoneName)
+{
+	FScopeLock Lock(&CriticalSection);
+
+	const FString BoneNameString = BoneName.ToString().ToLower();
+
+	if (UE::Mutable::Private::FBoneName* Result = BoneNamesMap.Find(BoneNameString))
+	{
+		return *Result;
+	}
+
+	// Go the slow way and add it to get a unique ID that the core can work with.
+	uint32 NewBoneId = CityHash32(reinterpret_cast<const char*>(*BoneNameString), BoneNameString.Len() * sizeof(FString::ElementType));
+
+	// See if the hash collides with an existing bone
+	bool bUnique = false;
+	while (!bUnique)
+	{
+		bUnique = true;
+		TMap<FString, UE::Mutable::Private::FBoneName>::TConstIterator MapIterator = BoneNamesMap.CreateConstIterator();
+		while (MapIterator)
+		{
+			if (MapIterator.Value() == NewBoneId)
+			{
+				bUnique = false;
+				break;
+			}
+			++MapIterator;
+		}
+
+		if (!bUnique)
+		{
+			NewBoneId++;
+		}
+	}
+
+	BoneNamesMap.Add(BoneNameString, NewBoneId);
+	return NewBoneId;
+}
