[Mutable] Moved MuT and MuR from /Public to /Internal.

Mutable Core should not be accessible from outside the Engine. We already had some headers on Internal but many were still in Public.
Having them in separate folders was a nightmare. Moving a type from a Public to an Internal header used to be impossible since it dragged most of the files, making the submit infeasable.
This submit fixes this by moving everything to Internal at once.

diff --git a/Engine/Plugins/Mutable/Source/CustomizableObject/Private/MuCO/CustomizableObjectSystem.cpp b/Engine/Plugins/Mutable/Source/CustomizableObject/Private/MuCO/CustomizableObjectSystem.cpp
--- a/Engine/Plugins/Mutable/Source/CustomizableObject/Private/MuCO/CustomizableObjectSystem.cpp
+++ b/Engine/Plugins/Mutable/Source/CustomizableObject/Private/MuCO/CustomizableObjectSystem.cpp
@@ -37,6 +37,8 @@
 #include "Interfaces/ITargetPlatform.h"
 #include "MuCO/ICustomizableObjectEditorModule.h"
 #include "MuT/UnrealPixelFormatOverride.h"
+#include "MuR/Parameters.h"
+#include "MuR/Types.h"
 #else
 #include "Engine/Engine.h"
 #endif
diff --git a/Engine/Plugins/Mutable/Source/CustomizableObject/Public/MuCO/CustomizableObjectParameterTypeDefinitions.h b/Engine/Plugins/Mutable/Source/CustomizableObject/Public/MuCO/CustomizableObjectParameterTypeDefinitions.h
--- a/Engine/Plugins/Mutable/Source/CustomizableObject/Public/MuCO/CustomizableObjectParameterTypeDefinitions.h
+++ b/Engine/Plugins/Mutable/Source/CustomizableObject/Public/MuCO/CustomizableObjectParameterTypeDefinitions.h
@@ -2,8 +2,6 @@
 
 #pragma once
 
-#include "MuR/Types.h"
-
 #include "CustomizableObjectParameterTypeDefinitions.generated.h"
 
 class UTexture;
diff --git a/Engine/Plugins/Mutable/Source/CustomizableObject/Public/MuCO/CustomizableObjectSystem.h b/Engine/Plugins/Mutable/Source/CustomizableObject/Public/MuCO/CustomizableObjectSystem.h
--- a/Engine/Plugins/Mutable/Source/CustomizableObject/Public/MuCO/CustomizableObjectSystem.h
+++ b/Engine/Plugins/Mutable/Source/CustomizableObject/Public/MuCO/CustomizableObjectSystem.h
@@ -8,8 +8,6 @@
 #include "Engine/StreamableManager.h"
 #include "AssetRegistry/AssetData.h"
 #include "MuCO/CustomizableObjectInstance.h"
-#include "MuR/Parameters.h"
-#include "MuR/Types.h"
 
 #include "CustomizableObjectSystem.generated.h"
 
diff --git a/Engine/Plugins/Mutable/Source/MutableTools/Internal/MuT/MutableToolsModule.h b/Engine/Plugins/Mutable/Source/MutableTools/Internal/MuT/MutableToolsModule.h
new file mode 100644
--- /dev/null
+++ b/Engine/Plugins/Mutable/Source/MutableTools/Internal/MuT/MutableToolsModule.h
@@ -0,0 +1,13 @@
+// Copyright Epic Games, Inc. All Rights Reserved.
+
+#pragma once
+
+#include "Modules/ModuleInterface.h"
+
+class FMutableToolsModule : public IModuleInterface
+{
+public:
+
+	virtual void StartupModule() override {};
+	virtual void ShutdownModule() override {};
+};
diff --git a/Engine/Plugins/Mutable/Source/MutableTools/Internal/MuT/NodeColorToSRGB.h b/Engine/Plugins/Mutable/Source/MutableTools/Internal/MuT/NodeColorToSRGB.h
new file mode 100644
--- /dev/null
+++ b/Engine/Plugins/Mutable/Source/MutableTools/Internal/MuT/NodeColorToSRGB.h
@@ -0,0 +1,44 @@
+// Copyright Epic Games, Inc. All Rights Reserved.
+
+#pragma once
+
+#include "MuR/Ptr.h"
+#include "MuR/RefCounted.h"
+#include "MuT/Node.h"
+#include "MuT/NodeColour.h"
+
+#define UE_API MUTABLETOOLS_API
+
+
+namespace UE::Mutable::Private
+{
+
+	/** Convert a linear color to sRGB
+	*/
+	class NodeColorToSRGB : public NodeColour
+	{
+	public:
+
+		Ptr<NodeColour> Color;
+
+	public:
+
+		// Node interface
+		virtual const FNodeType* GetType() const { return &StaticType; }
+		static const FNodeType* GetStaticType() { return &StaticType; }
+
+	protected:
+
+		/** Forbidden. Manage with the Ptr<> template. */
+		~NodeColorToSRGB() {}
+
+	private:
+
+		static UE_API FNodeType StaticType;
+
+	};
+
+
+}
+
+#undef UE_API
diff --git a/Engine/Plugins/Mutable/Source/MutableTools/Internal/MuT/NodeColourConstant.h b/Engine/Plugins/Mutable/Source/MutableTools/Internal/MuT/NodeColourConstant.h
new file mode 100644
--- /dev/null
+++ b/Engine/Plugins/Mutable/Source/MutableTools/Internal/MuT/NodeColourConstant.h
@@ -0,0 +1,45 @@
+// Copyright Epic Games, Inc. All Rights Reserved.
+
+#pragma once
+
+#include "MuR/Ptr.h"
+#include "MuR/RefCounted.h"
+#include "MuT/Node.h"
+#include "MuT/NodeColour.h"
+#include "Math/Vector4.h"
+#include "Math/MathFwd.h"
+
+#define UE_API MUTABLETOOLS_API
+
+namespace UE::Mutable::Private
+{
+
+	/** This node outputs a predefined colour value.
+	*/
+	class NodeColourConstant : public NodeColour
+	{
+	public:
+
+		FVector4f Value;
+
+	public:
+
+		// Node interface
+		virtual const FNodeType* GetType() const { return &StaticType; }
+		static const FNodeType* GetStaticType() { return &StaticType; }
+
+	protected:
+
+		/** Forbidden. Manage with the Ptr<> template. */
+		~NodeColourConstant() {}
+
+	private:
+
+		static UE_API FNodeType StaticType;
+
+	};
+
+
+}
+
+#undef UE_API
diff --git a/Engine/Plugins/Mutable/Source/MutableTools/Internal/MuT/NodeColourFromScalars.h b/Engine/Plugins/Mutable/Source/MutableTools/Internal/MuT/NodeColourFromScalars.h
new file mode 100644
--- /dev/null
+++ b/Engine/Plugins/Mutable/Source/MutableTools/Internal/MuT/NodeColourFromScalars.h
@@ -0,0 +1,48 @@
+// Copyright Epic Games, Inc. All Rights Reserved.
+
+#pragma once
+
+#include "MuR/Ptr.h"
+#include "MuR/RefCounted.h"
+#include "MuT/Node.h"
+#include "MuT/NodeColour.h"
+#include "MuT/NodeScalar.h"
+
+#define UE_API MUTABLETOOLS_API
+
+
+namespace UE::Mutable::Private
+{
+
+	/** Obtain a colour by sampling an image at specific homogeneous coordinates.
+	*/
+	class NodeColourFromScalars : public NodeColour
+	{
+	public:
+
+		Ptr<NodeScalar> X;
+		Ptr<NodeScalar> Y;
+		Ptr<NodeScalar> Z;
+		Ptr<NodeScalar> W;
+
+	public:
+
+		// Node interface
+		virtual const FNodeType* GetType() const { return &StaticType; }
+		static const FNodeType* GetStaticType() { return &StaticType; }
+
+	protected:
+
+		/** Forbidden. Manage with the Ptr<> template. */
+		~NodeColourFromScalars() {}
+
+	private:
+
+		static UE_API FNodeType StaticType;
+
+	};
+
+
+}
+
+#undef UE_API
diff --git a/Engine/Plugins/Mutable/Source/MutableTools/Internal/MuT/NodeColourMaterialBreak.h b/Engine/Plugins/Mutable/Source/MutableTools/Internal/MuT/NodeColourMaterialBreak.h
new file mode 100644
--- /dev/null
+++ b/Engine/Plugins/Mutable/Source/MutableTools/Internal/MuT/NodeColourMaterialBreak.h
@@ -0,0 +1,39 @@
+// Copyright Epic Games, Inc. All Rights Reserved.
+
+#pragma once
+
+#include "MuT/Node.h"
+#include "MuT/NodeColour.h"
+#include "MuR/Ptr.h"
+#include "MuT/NodeMaterial.h"
+
+#define UE_API MUTABLETOOLS_API
+
+namespace UE::Mutable::Private
+{
+	/** NodeColourMaterialBreak class. */
+	class NodeColourMaterialBreak : public NodeColour
+	{
+	public:
+
+		Ptr<NodeMaterial> MaterialSource;
+		FName ParameterName;
+
+	public:
+
+		// Node Interface
+		virtual const FNodeType* GetType() const override { return &StaticType; }
+		static const FNodeType* GetStaticType() { return &StaticType; }
+
+	protected:
+
+		/** Forbidden. Manage with the Ptr<> template. */
+		~NodeColourMaterialBreak() {}
+
+	private:
+
+		static UE_API FNodeType StaticType;
+	};
+}
+
+#undef UE_API
diff --git a/Engine/Plugins/Mutable/Source/MutableTools/Internal/MuT/NodeColourParameter.h b/Engine/Plugins/Mutable/Source/MutableTools/Internal/MuT/NodeColourParameter.h
new file mode 100644
--- /dev/null
+++ b/Engine/Plugins/Mutable/Source/MutableTools/Internal/MuT/NodeColourParameter.h
@@ -0,0 +1,50 @@
+// Copyright Epic Games, Inc. All Rights Reserved.
+
+#pragma once
+
+#include "MuR/Ptr.h"
+#include "MuR/RefCounted.h"
+#include "MuT/Node.h"
+#include "MuT/NodeColour.h"
+#include "MuT/NodeRange.h"
+#include "Math/Vector4.h"
+#include "Math/MathFwd.h"
+
+#define UE_API MUTABLETOOLS_API
+
+namespace UE::Mutable::Private
+{
+
+	/** Node that defines a colour model parameter.
+	*/
+	class NodeColourParameter : public NodeColour
+	{
+	public:
+
+		FVector4f DefaultValue;
+		FString Name;
+		FString Uid;
+
+		TArray<Ptr<NodeRange>> Ranges;
+
+	public:
+
+		// Node interface
+		virtual const FNodeType* GetType() const { return &StaticType; }
+		static const FNodeType* GetStaticType() { return &StaticType; }
+
+	protected:
+
+		/** Forbidden. Manage with the Ptr<> template. */
+		~NodeColourParameter() {}
+
+	private:
+
+		static UE_API FNodeType StaticType;
+
+	};
+
+
+}
+
+#undef UE_API
diff --git a/Engine/Plugins/Mutable/Source/MutableTools/Internal/MuT/NodeColourSampleImage.h b/Engine/Plugins/Mutable/Source/MutableTools/Internal/MuT/NodeColourSampleImage.h
new file mode 100644
--- /dev/null
+++ b/Engine/Plugins/Mutable/Source/MutableTools/Internal/MuT/NodeColourSampleImage.h
@@ -0,0 +1,48 @@
+// Copyright Epic Games, Inc. All Rights Reserved.
+
+#pragma once
+
+#include "MuR/Ptr.h"
+#include "MuR/RefCounted.h"
+#include "MuT/Node.h"
+#include "MuT/NodeColour.h"
+#include "MuT/NodeScalar.h"
+#include "MuT/NodeImage.h"
+
+#define UE_API MUTABLETOOLS_API
+
+
+namespace UE::Mutable::Private
+{
+
+	/** Obtain a colour by sampling an image at specific homogeneous coordinates.
+	*/
+	class NodeColourSampleImage : public NodeColour
+	{
+	public:
+
+		Ptr<NodeImage> Image;
+		Ptr<NodeScalar> X;
+		Ptr<NodeScalar> Y;
+
+	public:
+
+		// Node interface
+		virtual const FNodeType* GetType() const { return &StaticType; }
+		static const FNodeType* GetStaticType() { return &StaticType; }
+
+	protected:
+
+		/** Forbidden. Manage with the Ptr<> template. */
+		~NodeColourSampleImage() {}
+
+	private:
+
+		static UE_API FNodeType StaticType;
+
+	};
+
+
+}
+
+#undef UE_API
diff --git a/Engine/Plugins/Mutable/Source/MutableTools/Internal/MuT/NodeColourSwitch.h b/Engine/Plugins/Mutable/Source/MutableTools/Internal/MuT/NodeColourSwitch.h
new file mode 100644
--- /dev/null
+++ b/Engine/Plugins/Mutable/Source/MutableTools/Internal/MuT/NodeColourSwitch.h
@@ -0,0 +1,46 @@
+// Copyright Epic Games, Inc. All Rights Reserved.
+
+#pragma once
+
+#include "MuR/Ptr.h"
+#include "MuR/RefCounted.h"
+#include "MuT/Node.h"
+#include "MuT/NodeColour.h"
+#include "MuT/NodeScalar.h"
+
+#define UE_API MUTABLETOOLS_API
+
+
+namespace UE::Mutable::Private
+{
+
+    /** This node selects an output Colour from a set of input Colours based on a parameter. 
+	*/
+    class NodeColourSwitch : public NodeColour
+	{
+	public:
+
+		Ptr<NodeScalar> Parameter;
+		TArray<Ptr<NodeColour>> Options;
+
+	public:
+
+		/** Node type hierarchy data. */
+		virtual const FNodeType* GetType() const { return &StaticType; }
+		static const FNodeType* GetStaticType() { return &StaticType; }
+
+	protected:
+
+		/** Forbidden. Manage with the Ptr<> template. */
+		~NodeColourSwitch() {}
+
+	private:
+
+		static UE_API FNodeType StaticType;
+
+	};
+
+
+}
+
+#undef UE_API
diff --git a/Engine/Plugins/Mutable/Source/MutableTools/Internal/MuT/NodeColourTable.h b/Engine/Plugins/Mutable/Source/MutableTools/Internal/MuT/NodeColourTable.h
new file mode 100644
--- /dev/null
+++ b/Engine/Plugins/Mutable/Source/MutableTools/Internal/MuT/NodeColourTable.h
@@ -0,0 +1,49 @@
+// Copyright Epic Games, Inc. All Rights Reserved.
+
+#pragma once
+
+#include "MuR/Ptr.h"
+#include "MuR/RefCounted.h"
+#include "MuT/Node.h"
+#include "MuT/NodeColour.h"
+#include "MuT/Table.h"
+
+#define UE_API MUTABLETOOLS_API
+
+
+namespace UE::Mutable::Private
+{
+
+	/** This node provides the meshes stored in the column of a table.
+	*/
+	class NodeColourTable : public NodeColour
+	{
+	public:
+
+		FString ParameterName;
+		Ptr<FTable> Table;
+		FString ColumnName;
+		bool bNoneOption = false;
+		FString DefaultRowName;
+
+	public:
+
+		// Node interface
+		virtual const FNodeType* GetType() const override { return GetStaticType(); }
+		static const FNodeType* GetStaticType() { return &StaticType; }
+
+	protected:
+
+		/** Forbidden. Manage with the Ptr<> template. */
+		~NodeColourTable() {}
+
+	private:
+
+		static UE_API FNodeType StaticType;
+
+	};
+
+
+}
+
+#undef UE_API
diff --git a/Engine/Plugins/Mutable/Source/MutableTools/Internal/MuT/NodeComponentEdit.h b/Engine/Plugins/Mutable/Source/MutableTools/Internal/MuT/NodeComponentEdit.h
new file mode 100644
--- /dev/null
+++ b/Engine/Plugins/Mutable/Source/MutableTools/Internal/MuT/NodeComponentEdit.h
@@ -0,0 +1,47 @@
+// Copyright Epic Games, Inc. All Rights Reserved.
+
+#pragma once
+
+#include "MuR/Ptr.h"
+#include "MuR/RefCounted.h"
+#include "MuT/Node.h"
+#include "MuT/NodeComponent.h"
+
+#define UE_API MUTABLETOOLS_API
+
+namespace UE::Mutable::Private
+{
+
+	/** This node modifies a node of the parent object of the object that this node belongs to.
+	* It allows to extend, cut and morph the parent component's meshes.
+	* It also allows to patch the parent component's textures.
+	*/
+	class NodeComponentEdit : public NodeComponent
+	{
+	public:
+
+		Ptr<NodeComponent> Parent;
+
+	public:
+
+		// Node interface
+		virtual const FNodeType* GetType() const override { return GetStaticType(); }
+		static const FNodeType* GetStaticType() { return &StaticType; }
+
+        // NodeComponent interface
+		UE_API virtual const class NodeComponentNew* GetParentComponentNew() const override;
+
+	protected:
+
+		/** Forbidden. Manage with the Ptr<> template. */
+		~NodeComponentEdit() {}
+
+	private:
+
+		static UE_API FNodeType StaticType;
+
+	};
+
+}
+
+#undef UE_API
diff --git a/Engine/Plugins/Mutable/Source/MutableTools/Internal/MuT/NodeComponentSwitch.h b/Engine/Plugins/Mutable/Source/MutableTools/Internal/MuT/NodeComponentSwitch.h
new file mode 100644
--- /dev/null
+++ b/Engine/Plugins/Mutable/Source/MutableTools/Internal/MuT/NodeComponentSwitch.h
@@ -0,0 +1,49 @@
+// Copyright Epic Games, Inc. All Rights Reserved.
+
+#pragma once
+
+#include "MuR/Ptr.h"
+#include "MuR/RefCounted.h"
+#include "MuT/Node.h"
+#include "MuT/NodeComponent.h"
+#include "MuT/NodeScalar.h"
+
+#define UE_API MUTABLETOOLS_API
+
+
+namespace UE::Mutable::Private
+{
+
+    /** This node selects an output component from a set of input components based on a parameter. 
+	*/
+    class NodeComponentSwitch : public NodeComponent
+	{
+	public:
+
+		Ptr<NodeScalar> Parameter;
+		TArray<Ptr<NodeComponent>> Options;
+
+	public:
+
+		/** Node type hierarchy data. */
+		virtual const FNodeType* GetType() const { return &StaticType; }
+		static const FNodeType* GetStaticType() { return &StaticType; }
+
+		// NodeComponent interface
+		virtual const class NodeComponentNew* GetParentComponentNew() const override { check(false); return nullptr; }
+
+	protected:
+
+		/** Forbidden. Manage with the Ptr<> template. */
+		~NodeComponentSwitch() {}
+
+	private:
+
+		static UE_API FNodeType StaticType;
+
+	};
+
+
+}
+
+#undef UE_API
diff --git a/Engine/Plugins/Mutable/Source/MutableTools/Internal/MuT/NodeImageBinarise.h b/Engine/Plugins/Mutable/Source/MutableTools/Internal/MuT/NodeImageBinarise.h
new file mode 100644
--- /dev/null
+++ b/Engine/Plugins/Mutable/Source/MutableTools/Internal/MuT/NodeImageBinarise.h
@@ -0,0 +1,43 @@
+// Copyright Epic Games, Inc. All Rights Reserved.
+
+#pragma once
+
+#include "MuR/Ptr.h"
+#include "MuR/RefCounted.h"
+#include "MuT/Node.h"
+#include "MuT/NodeImage.h"
+#include "MuT/NodeScalar.h"
+
+#define UE_API MUTABLETOOLS_API
+
+namespace UE::Mutable::Private
+{
+
+	/** Node that multiplies the colors of an image, channel by channel. */
+	class NodeImageBinarise : public NodeImage
+	{
+	public:
+
+		Ptr<NodeImage> Base;
+		Ptr<NodeScalar> Threshold;
+
+	public:
+
+		// Node Interface
+		virtual const FNodeType* GetType() const override { return &StaticType; }
+		static const FNodeType* GetStaticType() { return &StaticType; }
+
+	protected:
+
+		/** Forbidden. Manage with the Ptr<> template. */
+		~NodeImageBinarise() {}
+
+	private:
+
+		static UE_API FNodeType StaticType;
+
+	};
+
+}
+
+#undef UE_API
diff --git a/Engine/Plugins/Mutable/Source/MutableTools/Internal/MuT/NodeImageConditional.h b/Engine/Plugins/Mutable/Source/MutableTools/Internal/MuT/NodeImageConditional.h
new file mode 100644
--- /dev/null
+++ b/Engine/Plugins/Mutable/Source/MutableTools/Internal/MuT/NodeImageConditional.h
@@ -0,0 +1,44 @@
+// Copyright Epic Games, Inc. All Rights Reserved.
+
+#pragma once
+
+#include "MuR/Ptr.h"
+#include "MuR/RefCounted.h"
+#include "MuT/Node.h"
+#include "MuT/NodeImage.h"
+#include "MuT/NodeBool.h"
+
+#define UE_API MUTABLETOOLS_API
+
+namespace UE::Mutable::Private
+{
+
+	/** This node selects an output image from a set of input images based on a parameter. */
+    class NodeImageConditional : public NodeImage
+	{
+	public:
+
+		Ptr<NodeBool> Parameter;
+		Ptr<NodeImage> True;
+		Ptr<NodeImage> False;
+
+	public:
+
+		// Node Interface
+		virtual const FNodeType* GetType() const override { return &StaticType; }
+		static const FNodeType* GetStaticType() { return &StaticType; }
+
+	protected:
+
+		/** Forbidden. Manage with the Ptr<> template. */
+		~NodeImageConditional() {}
+
+	private:
+
+		static UE_API FNodeType StaticType;
+
+	};
+
+}
+
+#undef UE_API
diff --git a/Engine/Plugins/Mutable/Source/MutableTools/Internal/MuT/NodeImageConstant.h b/Engine/Plugins/Mutable/Source/MutableTools/Internal/MuT/NodeImageConstant.h
new file mode 100644
--- /dev/null
+++ b/Engine/Plugins/Mutable/Source/MutableTools/Internal/MuT/NodeImageConstant.h
@@ -0,0 +1,56 @@
+// Copyright Epic Games, Inc. All Rights Reserved.
+
+#pragma once
+
+#include "MuR/Image.h"
+#include "MuR/Ptr.h"
+#include "MuR/RefCounted.h"
+#include "MuT/Node.h"
+#include "MuT/NodeImage.h"
+
+#define UE_API MUTABLETOOLS_API
+
+namespace UE::Mutable::Private
+{
+
+	/** Node that outputs a constant image.
+	* This node also supports "image references".
+	*/
+	class NodeImageConstant : public NodeImage
+	{
+	public:
+
+		/** */
+		Ptr<TResourceProxy<FImage>> Proxy;
+		
+		/** */
+		FSourceDataDescriptor SourceDataDescriptor;
+
+	public:
+
+		// Node Interface
+		virtual const FNodeType* GetType() const override { return GetStaticType(); }
+		static const FNodeType* GetStaticType() { return &StaticType; }
+
+		// Own Interface
+
+        //! Set the image to be output by this node
+        UE_API void SetValue(TSharedPtr<const FImage>);
+
+        //! Set the image proxy that will provide the image for this node when necessary
+        UE_API void SetValue( Ptr<TResourceProxy<FImage>> );
+
+	protected:
+
+		/** Forbidden. Manage with the Ptr<> template. */
+		~NodeImageConstant() {}
+
+	private:
+
+		static UE_API FNodeType StaticType;
+
+	};
+
+}
+
+#undef UE_API
diff --git a/Engine/Plugins/Mutable/Source/MutableTools/Internal/MuT/NodeImageFromMaterialParameter.h b/Engine/Plugins/Mutable/Source/MutableTools/Internal/MuT/NodeImageFromMaterialParameter.h
new file mode 100644
--- /dev/null
+++ b/Engine/Plugins/Mutable/Source/MutableTools/Internal/MuT/NodeImageFromMaterialParameter.h
@@ -0,0 +1,34 @@
+// Copyright Epic Games, Inc. All Rights Reserved.
+
+#pragma once
+#include "MuR/Ptr.h"
+#include "MuR/RefCounted.h"
+#include "MuT/Node.h"
+#include "MuT/NodeImage.h"
+
+#define UE_API MUTABLETOOLS_API
+
+namespace UE::Mutable::Private
+{
+	class NodeImageFromMaterialParameter : public NodeImage
+	{
+	public:
+
+		FName ImageParameterName;
+
+	public:
+
+		// Node Interface
+		virtual const FNodeType* GetType() const override { return &StaticType; }
+		static const FNodeType* GetStaticType() { return &StaticType; }
+
+	protected:
+
+		/** Forbidden. Manage with the Ptr<> template. */
+		~NodeImageFromMaterialParameter() {}
+
+	private:
+
+		static UE_API FNodeType StaticType;
+	};
+}
\ No newline at end of file
diff --git a/Engine/Plugins/Mutable/Source/MutableTools/Internal/MuT/NodeImageInterpolate.h b/Engine/Plugins/Mutable/Source/MutableTools/Internal/MuT/NodeImageInterpolate.h
new file mode 100644
--- /dev/null
+++ b/Engine/Plugins/Mutable/Source/MutableTools/Internal/MuT/NodeImageInterpolate.h
@@ -0,0 +1,49 @@
+// Copyright Epic Games, Inc. All Rights Reserved.
+
+#pragma once
+
+#include "MuR/Ptr.h"
+#include "MuR/RefCounted.h"
+#include "MuT/Node.h"
+#include "MuT/NodeImage.h"
+#include "MuT/NodeScalar.h"
+
+#define UE_API MUTABLETOOLS_API
+
+namespace UE::Mutable::Private
+{
+
+	/** Node that interpolates linearly between several images based on a weight.
+	* The different input images are uniformly distributed in the 0 to 1 range:
+	* - if two images are set, the first one is 0.0 and the second one is 1.0
+	* - if three images are set, the first one is 0.0 the second one is 0.5 and the third one 1.0
+	* regardless of the input slots used: If B and C are set, but not A, B will be at weight 0
+	* and C will be at weight 1
+	*/
+	class NodeImageInterpolate : public NodeImage
+	{
+	public:
+
+		Ptr<NodeScalar> Factor;
+		TArray<Ptr<NodeImage>> Targets;
+
+	public:
+
+		// Node Interface
+		virtual const FNodeType* GetType() const override { return &StaticType; }
+		static const FNodeType* GetStaticType() { return &StaticType; }
+
+	protected:
+
+		/** Forbidden. Manage with the Ptr<> template. */
+		~NodeImageInterpolate() {}
+
+	private:
+
+		static UE_API FNodeType StaticType;
+
+	};
+
+}
+
+#undef UE_API
diff --git a/Engine/Plugins/Mutable/Source/MutableTools/Internal/MuT/NodeImageInvert.h b/Engine/Plugins/Mutable/Source/MutableTools/Internal/MuT/NodeImageInvert.h
new file mode 100644
--- /dev/null
+++ b/Engine/Plugins/Mutable/Source/MutableTools/Internal/MuT/NodeImageInvert.h
@@ -0,0 +1,41 @@
+// Copyright Epic Games, Inc. All Rights Reserved.
+
+#pragma once
+
+#include "MuR/Ptr.h"
+#include "MuR/RefCounted.h"
+#include "MuT/Node.h"
+#include "MuT/NodeImage.h"
+
+#define UE_API MUTABLETOOLS_API
+
+namespace UE::Mutable::Private
+{
+
+	/** Node that inverts the colors of an image, channel by channel. */
+	class NodeImageInvert : public NodeImage
+	{
+	public:
+
+		Ptr<NodeImage> Base;
+
+	public:
+
+		// Node Interface
+		virtual const FNodeType* GetType() const override { return &StaticType; }
+		static const FNodeType* GetStaticType() { return &StaticType; }
+
+	protected:
+
+		/** Forbidden. Manage with the Ptr<> template. */
+		~NodeImageInvert() {}
+
+	private:
+
+		static UE_API FNodeType StaticType;
+
+	};
+
+}
+
+#undef UE_API
diff --git a/Engine/Plugins/Mutable/Source/MutableTools/Internal/MuT/NodeImageLayer.h b/Engine/Plugins/Mutable/Source/MutableTools/Internal/MuT/NodeImageLayer.h
new file mode 100644
--- /dev/null
+++ b/Engine/Plugins/Mutable/Source/MutableTools/Internal/MuT/NodeImageLayer.h
@@ -0,0 +1,47 @@
+// Copyright Epic Games, Inc. All Rights Reserved.
+
+#pragma once
+
+#include "HAL/Platform.h"
+#include "MuR/Image.h"
+#include "MuR/Ptr.h"
+#include "MuR/RefCounted.h"
+#include "MuT/Node.h"
+#include "MuT/NodeImage.h"
+
+#define UE_API MUTABLETOOLS_API
+
+namespace UE::Mutable::Private
+{
+
+	/** This node applies a layer blending effect on a base image using a mask and a blended image. */
+	class NodeImageLayer : public NodeImage
+	{
+	public:
+
+		Ptr<NodeImage> Base;
+		Ptr<NodeImage> Mask;
+		Ptr<NodeImage> Blended;
+		EBlendType Type;
+
+	public:
+
+		// Node Interface
+		virtual const FNodeType* GetType() const override { return &StaticType; }
+		static const FNodeType* GetStaticType() { return &StaticType; }
+
+	protected:
+
+		/** Forbidden. Manage with the Ptr<> template. */
+		~NodeImageLayer() {}
+
+	private:
+
+		static UE_API FNodeType StaticType;
+
+	};
+
+
+}
+
+#undef UE_API
diff --git a/Engine/Plugins/Mutable/Source/MutableTools/Internal/MuT/NodeImagePlainColour.h b/Engine/Plugins/Mutable/Source/MutableTools/Internal/MuT/NodeImagePlainColour.h
new file mode 100644
--- /dev/null
+++ b/Engine/Plugins/Mutable/Source/MutableTools/Internal/MuT/NodeImagePlainColour.h
@@ -0,0 +1,46 @@
+// Copyright Epic Games, Inc. All Rights Reserved.
+
+#pragma once
+
+#include "MuR/Ptr.h"
+#include "MuR/RefCounted.h"
+#include "MuT/Node.h"
+#include "MuT/NodeImage.h"
+#include "MuT/NodeColour.h"
+#include "MuR/Image.h"
+
+#define UE_API MUTABLETOOLS_API
+
+namespace UE::Mutable::Private
+{
+
+	/** Node that multiplies the colors of an image, channel by channel. */
+	class NodeImagePlainColour : public NodeImage
+	{
+	public:
+
+		Ptr<NodeColour> Colour;
+		int32 SizeX = 4;
+		int32 SizeY = 4;
+		EImageFormat Format = EImageFormat::RGBA_UByte;
+
+	public:
+
+		// Node Interface
+		virtual const FNodeType* GetType() const override { return &StaticType; }
+		static const FNodeType* GetStaticType() { return &StaticType; }
+
+	protected:
+
+		/** Forbidden. Manage with the Ptr<> template. */
+		~NodeImagePlainColour() {}
+
+	private:
+
+		static UE_API FNodeType StaticType;
+
+	};
+
+}
+
+#undef UE_API
diff --git a/Engine/Plugins/Mutable/Source/MutableTools/Internal/MuT/NodeImageProject.h b/Engine/Plugins/Mutable/Source/MutableTools/Internal/MuT/NodeImageProject.h
new file mode 100644
--- /dev/null
+++ b/Engine/Plugins/Mutable/Source/MutableTools/Internal/MuT/NodeImageProject.h
@@ -0,0 +1,57 @@
+// Copyright Epic Games, Inc. All Rights Reserved.
+
+#pragma once
+
+#include "MuR/Ptr.h"
+#include "MuR/RefCounted.h"
+#include "MuT/Node.h"
+#include "MuT/NodeImage.h"
+#include "MuT/NodeMesh.h"
+#include "MuT/NodeProjector.h"
+#include "MuT/NodeScalar.h"
+#include "MuR/Image.h"
+
+#define UE_API MUTABLETOOLS_API
+
+namespace UE::Mutable::Private
+{
+
+	/** */
+	class NodeImageProject : public NodeImage
+	{
+	public:
+
+		Ptr<NodeProjector> Projector;
+		Ptr<NodeMesh> Mesh;
+		Ptr<NodeScalar> AngleFadeStart;
+		Ptr<NodeScalar> AngleFadeEnd;
+		Ptr<NodeImage> Image;
+		Ptr<NodeImage> Mask;
+		FUintVector2 ImageSize;
+		uint8 Layout = 0;
+		bool bIsRGBFadingEnabled = true;
+		bool bIsAlphaFadingEnabled = true;
+		bool bEnableTextureSeamCorrection = true;
+		ESamplingMethod SamplingMethod = ESamplingMethod::Point;
+		EMinFilterMethod MinFilterMethod = EMinFilterMethod::None;
+
+	public:
+
+		// Node Interface
+		virtual const FNodeType* GetType() const override { return &StaticType; }
+		static const FNodeType* GetStaticType() { return &StaticType; }
+
+	protected:
+
+		/** Forbidden. Manage with the Ptr<> template. */
+		~NodeImageProject() {}
+
+	private:
+
+		static UE_API FNodeType StaticType;
+
+	};
+
+}
+
+#undef UE_API
diff --git a/Engine/Plugins/Mutable/Source/MutableTools/Internal/MuT/NodeImageResize.h b/Engine/Plugins/Mutable/Source/MutableTools/Internal/MuT/NodeImageResize.h
new file mode 100644
--- /dev/null
+++ b/Engine/Plugins/Mutable/Source/MutableTools/Internal/MuT/NodeImageResize.h
@@ -0,0 +1,46 @@
+// Copyright Epic Games, Inc. All Rights Reserved.
+
+#pragma once
+
+#include "MuR/Ptr.h"
+#include "MuR/RefCounted.h"
+#include "MuT/Node.h"
+#include "MuT/NodeImage.h"
+
+#define UE_API MUTABLETOOLS_API
+
+namespace UE::Mutable::Private
+{
+
+	/** Node that multiplies the colors of an image, channel by channel. */
+	class NodeImageResize : public NodeImage
+	{
+	public:
+
+		Ptr<NodeImage> Base;
+
+		bool bRelative = true;
+
+		float SizeX = 0.5f;
+		float SizeY = 0.5f;
+
+	public:
+
+		// Node Interface
+		virtual const FNodeType* GetType() const override { return &StaticType; }
+		static const FNodeType* GetStaticType() { return &StaticType; }
+
+	protected:
+
+		/** Forbidden. Manage with the Ptr<> template. */
+		~NodeImageResize() {}
+
+	private:
+
+		static UE_API FNodeType StaticType;
+
+	};
+
+}
+
+#undef UE_API
diff --git a/Engine/Plugins/Mutable/Source/MutableTools/Internal/MuT/NodeImageSaturate.h b/Engine/Plugins/Mutable/Source/MutableTools/Internal/MuT/NodeImageSaturate.h
new file mode 100644
--- /dev/null
+++ b/Engine/Plugins/Mutable/Source/MutableTools/Internal/MuT/NodeImageSaturate.h
@@ -0,0 +1,47 @@
+// Copyright Epic Games, Inc. All Rights Reserved.
+
+#pragma once
+
+#include "MuR/Ptr.h"
+#include "MuR/RefCounted.h"
+#include "MuT/Node.h"
+#include "MuT/NodeImage.h"
+#include "MuT/NodeScalar.h"
+
+#define UE_API MUTABLETOOLS_API
+
+namespace UE::Mutable::Private
+{
+
+	/** Change the saturation of an image.This node can be used to increase the saturation with a
+	* factor bigger than 1, or to decrease it or desaturate it completely with a factor smaller
+	* than 1 or 0.
+	*/
+	class NodeImageSaturate : public NodeImage
+	{
+	public:
+
+		Ptr<NodeImage> Source;
+
+		Ptr<NodeScalar> Factor;
+
+	public:
+
+		// Node Interface
+		virtual const FNodeType* GetType() const override { return &StaticType; }
+		static const FNodeType* GetStaticType() { return &StaticType; }
+
+	protected:
+
+		/** Forbidden. Manage with the Ptr<> template. */
+		~NodeImageSaturate() {}
+
+	private:
+
+		static UE_API FNodeType StaticType;
+
+	};
+
+}
+
+#undef UE_API
diff --git a/Engine/Plugins/Mutable/Source/MutableTools/Internal/MuT/NodeImageSwitch.h b/Engine/Plugins/Mutable/Source/MutableTools/Internal/MuT/NodeImageSwitch.h
new file mode 100644
--- /dev/null
+++ b/Engine/Plugins/Mutable/Source/MutableTools/Internal/MuT/NodeImageSwitch.h
@@ -0,0 +1,44 @@
+// Copyright Epic Games, Inc. All Rights Reserved.
+
+#pragma once
+
+#include "MuR/Ptr.h"
+#include "MuR/RefCounted.h"
+#include "MuT/Node.h"
+#include "MuT/NodeImage.h"
+#include "MuT/NodeScalar.h"
+
+#define UE_API MUTABLETOOLS_API
+
+namespace UE::Mutable::Private
+{
+
+	//! This node selects an output image from a set of input images based on a parameter.
+	class NodeImageSwitch : public NodeImage
+	{
+	public:
+
+		Ptr<NodeScalar> Parameter;
+
+		TArray<Ptr<NodeImage>> Options;
+
+	public:
+
+		// Node Interface
+		virtual const FNodeType* GetType() const override { return &StaticType; }
+		static const FNodeType* GetStaticType() { return &StaticType; }
+
+	protected:
+
+		/** Forbidden. Manage with the Ptr<> template. */
+		~NodeImageSwitch() {}
+
+	private:
+
+		static UE_API FNodeType StaticType;
+
+	};
+
+}
+
+#undef UE_API
diff --git a/Engine/Plugins/Mutable/Source/MutableTools/Internal/MuT/NodeImageSwizzle.h b/Engine/Plugins/Mutable/Source/MutableTools/Internal/MuT/NodeImageSwizzle.h
new file mode 100644
--- /dev/null
+++ b/Engine/Plugins/Mutable/Source/MutableTools/Internal/MuT/NodeImageSwizzle.h
@@ -0,0 +1,54 @@
+// Copyright Epic Games, Inc. All Rights Reserved.
+
+#pragma once
+
+#include "MuR/Image.h"
+#include "MuR/Ptr.h"
+#include "MuR/RefCounted.h"
+#include "MuT/Node.h"
+#include "MuT/NodeImage.h"
+
+#define UE_API MUTABLETOOLS_API
+
+namespace UE::Mutable::Private
+{
+
+	/** Node that composes a new image by gathering pixel data from channels in other images. */
+	class NodeImageSwizzle : public NodeImage
+	{
+	public:
+
+		EImageFormat NewFormat;
+		TArray<Ptr<NodeImage>> Sources;
+		TArray<int32> SourceChannels;
+
+	public:
+
+		// Node Interface
+		virtual const FNodeType* GetType() const override { return &StaticType; }
+		static const FNodeType* GetStaticType() { return &StaticType; }
+
+		// Own interface 
+		inline void SetFormat(EImageFormat InFormat)
+		{
+			NewFormat = InFormat;
+
+			int32 ChannelCount = GetImageFormatData(InFormat).Channels;
+			Sources.SetNum(ChannelCount);
+			SourceChannels.SetNum(ChannelCount);
+		}
+
+	protected:
+
+		/** Forbidden. Manage with the Ptr<> template. */
+		~NodeImageSwizzle() {}
+
+	private:
+
+		static UE_API FNodeType StaticType;
+
+	};
+
+}
+
+#undef UE_API
diff --git a/Engine/Plugins/Mutable/Source/MutableTools/Internal/MuT/NodeImageTable.h b/Engine/Plugins/Mutable/Source/MutableTools/Internal/MuT/NodeImageTable.h
new file mode 100644
--- /dev/null
+++ b/Engine/Plugins/Mutable/Source/MutableTools/Internal/MuT/NodeImageTable.h
@@ -0,0 +1,62 @@
+// Copyright Epic Games, Inc. All Rights Reserved.
+
+#pragma once
+
+#include "MuR/Ptr.h"
+#include "MuR/RefCounted.h"
+#include "MuR/ImageTypes.h"
+#include "MuT/Table.h"
+#include "MuT/Node.h"
+#include "MuT/NodeImage.h"
+
+#define UE_API MUTABLETOOLS_API
+
+
+namespace UE::Mutable::Private
+{
+
+
+	/** This node provides the meshes stored in the column of a table. */
+	class NodeImageTable : public NodeImage
+	{
+	public:
+
+		FString ParameterName;
+
+		Ptr<FTable> Table;
+
+		FString ColumnName;
+
+		uint16 MaxTextureSize = 0;
+
+		FImageDesc ReferenceImageDesc;
+
+		bool bNoneOption = false;
+
+		FString DefaultRowName;
+
+		/** */
+		FSourceDataDescriptor SourceDataDescriptor;
+
+	public:
+
+		// Node interface
+		virtual const FNodeType* GetType() const override { return GetStaticType(); }
+		static const FNodeType* GetStaticType() { return &StaticType; }
+
+	protected:
+
+		/** Forbidden. Manage with the Ptr<> template. */
+		~NodeImageTable() {}
+
+	private:
+
+		static UE_API FNodeType StaticType;
+
+
+	};
+
+
+}
+
+#undef UE_API
diff --git a/Engine/Plugins/Mutable/Source/MutableTools/Internal/MuT/NodeImageVariation.h b/Engine/Plugins/Mutable/Source/MutableTools/Internal/MuT/NodeImageVariation.h
new file mode 100644
--- /dev/null
+++ b/Engine/Plugins/Mutable/Source/MutableTools/Internal/MuT/NodeImageVariation.h
@@ -0,0 +1,49 @@
+// Copyright Epic Games, Inc. All Rights Reserved.
+
+#pragma once
+
+#include "MuR/Ptr.h"
+#include "MuR/RefCounted.h"
+#include "MuT/Node.h"
+#include "MuT/NodeImage.h"
+
+#define UE_API MUTABLETOOLS_API
+
+namespace UE::Mutable::Private
+{
+
+    /** */
+    class NodeImageVariation : public NodeImage
+    {
+	public:
+
+		Ptr<NodeImage> DefaultImage;
+
+		struct FVariation
+		{
+			Ptr<NodeImage> Image;
+			FString Tag;
+		};
+
+		TArray<FVariation> Variations;
+
+    public:
+
+		// Node Interface
+		virtual const FNodeType* GetType() const override { return &StaticType; }
+		static const FNodeType* GetStaticType() { return &StaticType; }
+
+    protected:
+
+        /** Forbidden. Manage with the Ptr<> template. */
+		~NodeImageVariation() {}
+
+    private:
+
+		static UE_API FNodeType StaticType;
+	
+	};
+
+}
+
+#undef UE_API
diff --git a/Engine/Plugins/Mutable/Source/MutableTools/Internal/MuT/NodeLOD.h b/Engine/Plugins/Mutable/Source/MutableTools/Internal/MuT/NodeLOD.h
new file mode 100644
--- /dev/null
+++ b/Engine/Plugins/Mutable/Source/MutableTools/Internal/MuT/NodeLOD.h
@@ -0,0 +1,49 @@
+// Copyright Epic Games, Inc. All Rights Reserved.
+
+#pragma once
+
+#include "MuR/Ptr.h"
+#include "MuR/RefCounted.h"
+#include "MuT/Node.h"
+#include "MuT/NodeSurface.h"
+#include "MuT/NodeModifier.h"
+
+#define UE_API MUTABLETOOLS_API
+
+namespace UE::Mutable::Private
+{
+
+
+	class NodeLOD : public Node
+	{
+	public:
+
+		//-----------------------------------------------------------------------------------------
+		// Node Interface
+		//-----------------------------------------------------------------------------------------
+
+		virtual const FNodeType* GetType() const override { return GetStaticType(); }
+		static const FNodeType* GetStaticType() { return &StaticType; }
+
+		//-----------------------------------------------------------------------------------------
+		// Own Interface
+		//-----------------------------------------------------------------------------------------
+
+		TArray<Ptr<NodeSurface>> Surfaces;
+
+	protected:
+
+		/** Forbidden. Manage with the Ptr<> template. */
+		~NodeLOD() {}
+
+	private:
+
+		static UE_API FNodeType StaticType;
+
+	};
+
+
+
+}
+
+#undef UE_API
diff --git a/Engine/Plugins/Mutable/Source/MutableTools/Internal/MuT/NodeMaterialParameter.h b/Engine/Plugins/Mutable/Source/MutableTools/Internal/MuT/NodeMaterialParameter.h
new file mode 100644
--- /dev/null
+++ b/Engine/Plugins/Mutable/Source/MutableTools/Internal/MuT/NodeMaterialParameter.h
@@ -0,0 +1,57 @@
+// Copyright Epic Games, Inc. All Rights Reserved.
+
+#pragma once
+
+#include "Containers/Map.h"
+#include "HAL/Platform.h"
+#include "MuR/Ptr.h"
+#include "MuR/RefCounted.h"
+#include "MuT/NodeRange.h"
+#include "MuT/NodeMaterial.h"
+#include "MuT/NodeImage.h"
+
+
+#define UE_API MUTABLETOOLS_API
+
+namespace UE::Mutable::Private
+{
+
+    /** NodeMaterialParameter class. */
+	class NodeMaterialParameter : public NodeMaterial
+	{
+	public:
+
+		/** Name of the parameter */
+		FString Name;
+
+		/** User provided ID to identify the parameter. */
+		FString UID;
+
+		/** Ranges for the parameter in case it is multidimensional. */
+		TArray<Ptr<NodeRange>> Ranges;
+
+		TArray<FString> ColorParameterNames;
+		TArray<FString> ScalarParameterNames;
+		TMap<FString, Ptr<NodeImage>> ImageParameterNodes;
+
+	public:
+
+		// Node Interface
+		virtual const FNodeType* GetType() const override { return &StaticType; }
+		static const FNodeType* GetStaticType() { return &StaticType; }
+
+	protected:
+
+		// Forbidden. Manage with the Ptr<> template.
+		inline ~NodeMaterialParameter() {}
+
+	private:
+
+		static UE_API FNodeType StaticType;
+
+	};
+
+
+}
+
+#undef UE_API
diff --git a/Engine/Plugins/Mutable/Source/MutableTools/Internal/MuT/NodeMaterialSwitch.h b/Engine/Plugins/Mutable/Source/MutableTools/Internal/MuT/NodeMaterialSwitch.h
new file mode 100644
--- /dev/null
+++ b/Engine/Plugins/Mutable/Source/MutableTools/Internal/MuT/NodeMaterialSwitch.h
@@ -0,0 +1,45 @@
+// Copyright Epic Games, Inc. All Rights Reserved.
+
+#pragma once
+
+#include "HAL/Platform.h"
+#include "MuR/Ptr.h"
+#include "MuR/RefCounted.h"
+#include "MuT/NodeScalar.h"
+#include "MuT/NodeMaterial.h"
+
+#define UE_API MUTABLETOOLS_API
+
+namespace UE::Mutable::Private
+{
+
+    /** NodeMaterialSwitch class. */
+	class NodeMaterialSwitch : public NodeMaterial
+	{
+	public:
+
+		Ptr<NodeScalar> Parameter;
+
+		TArray<Ptr<NodeMaterial>> Options;
+
+	public:
+
+		// Node Interface
+		virtual const FNodeType* GetType() const override { return &StaticType; }
+		static const FNodeType* GetStaticType() { return &StaticType; }
+
+	protected:
+
+		// Forbidden. Manage with the Ptr<> template.
+		inline ~NodeMaterialSwitch() {}
+
+	private:
+
+		static UE_API FNodeType StaticType;
+
+	};
+
+
+}
+
+#undef UE_API
diff --git a/Engine/Plugins/Mutable/Source/MutableTools/Internal/MuT/NodeMaterialTable.h b/Engine/Plugins/Mutable/Source/MutableTools/Internal/MuT/NodeMaterialTable.h
new file mode 100644
--- /dev/null
+++ b/Engine/Plugins/Mutable/Source/MutableTools/Internal/MuT/NodeMaterialTable.h
@@ -0,0 +1,48 @@
+// Copyright Epic Games, Inc. All Rights Reserved.
+
+#pragma once
+
+#include "HAL/Platform.h"
+#include "MuR/Ptr.h"
+#include "MuR/RefCounted.h"
+#include "MuT/Table.h"
+#include "MuT/NodeMaterial.h"
+
+#define UE_API MUTABLETOOLS_API
+
+namespace UE::Mutable::Private
+{
+
+    /** NodeMaterialTable class. */
+	class NodeMaterialTable : public NodeMaterial
+	{
+	public:
+		FString ParameterName;
+		FString ColumnName;
+		FString DefaultRowName;
+		
+		bool bNoneOption = false;
+		
+		Ptr<FTable> Table;
+
+	public:
+
+		// Node Interface
+		virtual const FNodeType* GetType() const override { return &StaticType; }
+		static const FNodeType* GetStaticType() { return &StaticType; }
+
+	protected:
+
+		// Forbidden. Manage with the Ptr<> template.
+		inline ~NodeMaterialTable() {}
+
+	private:
+
+		static UE_API FNodeType StaticType;
+
+	};
+
+
+}
+
+#undef UE_API
diff --git a/Engine/Plugins/Mutable/Source/MutableTools/Internal/MuT/NodeMaterialVariation.h b/Engine/Plugins/Mutable/Source/MutableTools/Internal/MuT/NodeMaterialVariation.h
new file mode 100644
--- /dev/null
+++ b/Engine/Plugins/Mutable/Source/MutableTools/Internal/MuT/NodeMaterialVariation.h
@@ -0,0 +1,48 @@
+// Copyright Epic Games, Inc. All Rights Reserved.
+
+#pragma once
+
+#include "HAL/Platform.h"
+#include "MuR/Ptr.h"
+#include "MuR/RefCounted.h"
+#include "MuT/NodeMaterial.h"
+
+#define UE_API MUTABLETOOLS_API
+
+namespace UE::Mutable::Private
+{
+
+    /** NodeMaterialVariation class. */
+	class NodeMaterialVariation : public NodeMaterial
+	{
+	public:
+		Ptr<NodeMaterial> DefaultMaterial;
+
+		struct FVariation
+		{
+			Ptr<NodeMaterial> Material;
+			FString Tag;
+		};
+
+		TArray<FVariation> Variations;
+	public:
+
+		// Node Interface
+		virtual const FNodeType* GetType() const override { return &StaticType; }
+		static const FNodeType* GetStaticType() { return &StaticType; }
+
+	protected:
+
+		// Forbidden. Manage with the Ptr<> template.
+		inline ~NodeMaterialVariation() {}
+
+	private:
+
+		static UE_API FNodeType StaticType;
+
+	};
+
+
+}
+
+#undef UE_API
diff --git a/Engine/Plugins/Mutable/Source/MutableTools/Internal/MuT/NodeMatrixParameter.h b/Engine/Plugins/Mutable/Source/MutableTools/Internal/MuT/NodeMatrixParameter.h
new file mode 100644
--- /dev/null
+++ b/Engine/Plugins/Mutable/Source/MutableTools/Internal/MuT/NodeMatrixParameter.h
@@ -0,0 +1,40 @@
+﻿// Copyright Epic Games, Inc. All Rights Reserved.
+
+#pragma once
+
+#include "MuT/NodeMatrix.h"
+
+#include "MuT/NodeRange.h"
+#include "Math/Matrix.h"
+
+#define UE_API MUTABLETOOLS_API
+
+
+namespace UE::Mutable::Private
+{
+
+	class NodeMatrixParameter : public NodeMatrix
+	{
+	public:
+		FMatrix44f DefaultValue;
+		FString Name;
+		FString Uid;
+
+		TArray<Ptr<NodeRange>> Ranges;
+		
+		// Node interface
+		const FNodeType* GetType() const override { return &StaticType; }
+		static const FNodeType* GetStaticType() { return &StaticType; }
+
+	protected:
+		/** Forbidden. Manage with the Ptr<> template. */
+		virtual ~NodeMatrixParameter() override = default;
+
+	private:
+
+		static UE_API FNodeType StaticType;
+	};
+
+}
+
+#undef UE_API
diff --git a/Engine/Plugins/Mutable/Source/MutableTools/Internal/MuT/NodeMeshClipDeform.h b/Engine/Plugins/Mutable/Source/MutableTools/Internal/MuT/NodeMeshClipDeform.h
new file mode 100644
--- /dev/null
+++ b/Engine/Plugins/Mutable/Source/MutableTools/Internal/MuT/NodeMeshClipDeform.h
@@ -0,0 +1,44 @@
+// Copyright Epic Games, Inc. All Rights Reserved.
+
+#pragma once
+
+#include "MuR/Ptr.h"
+#include "MuT/Node.h"
+#include "MuT/NodeImage.h"
+#include "MuT/NodeMesh.h"
+
+#define UE_API MUTABLETOOLS_API
+
+
+namespace UE::Mutable::Private
+{
+
+	//! 
+	class NodeMeshClipDeform : public NodeMesh
+	{
+	public:
+
+		Ptr<NodeMesh> BaseMesh;
+		Ptr<NodeMesh> ClipShape;
+		Ptr<NodeImage> ShapeWeights;
+
+	public:
+
+		// Node interface
+		virtual const FNodeType* GetType() const override { return &StaticType; }
+		static const FNodeType* GetStaticType() { return &StaticType; }
+
+	protected:
+
+		/** Forbidden. Manage with the Ptr<> template. */
+		inline ~NodeMeshClipDeform() {}
+
+	private:
+
+		static UE_API FNodeType StaticType;
+
+	};
+
+}
+
+#undef UE_API
diff --git a/Engine/Plugins/Mutable/Source/MutableTools/Internal/MuT/NodeMeshClipMorphPlane.h b/Engine/Plugins/Mutable/Source/MutableTools/Internal/MuT/NodeMeshClipMorphPlane.h
new file mode 100644
--- /dev/null
+++ b/Engine/Plugins/Mutable/Source/MutableTools/Internal/MuT/NodeMeshClipMorphPlane.h
@@ -0,0 +1,68 @@
+// Copyright Epic Games, Inc. All Rights Reserved.
+
+#pragma once
+
+#include "MuR/Ptr.h"
+#include "MuR/RefCounted.h"
+#include "MuT/Node.h"
+#include "MuT/NodeMesh.h"
+#include "MuT/NodeModifierMeshClipMorphPlane.h"
+
+#define UE_API MUTABLETOOLS_API
+
+
+namespace UE::Mutable::Private
+{
+
+    //! This node applies a geometric transform represented by a 4x4 matrix to a mesh
+    class NodeMeshClipMorphPlane : public NodeMesh
+	{
+	public:
+
+		Ptr<NodeMesh> Source;
+
+		FClipMorphPlaneParameters Parameters;
+
+	public:
+
+		// Node interface
+		virtual const FNodeType* GetType() const override { return &StaticType; }
+		static const FNodeType* GetStaticType() { return &StaticType; }
+
+		//-----------------------------------------------------------------------------------------
+		// Own Interface
+		//-----------------------------------------------------------------------------------------
+
+		//! Source mesh to be clipped and morphed
+        UE_API NodeMeshPtr GetSource() const;
+		UE_API void SetSource( NodeMesh* );
+
+		UE_API void SetPlane(FVector3f Center, FVector3f Normal);
+		UE_API void SetParams(float dist, float factor);
+		UE_API void SetMorphEllipse(float radius1, float radius2, float rotation);
+
+		//! Define an axis-aligned box that will select the vertices to be morphed.
+		//! If this is not set, or the size of the box is 0, all vertices will be affected.
+		//! Only one of Box or Bone Hierarchy can be used (the last one set)
+		UE_API void SetVertexSelectionBox(FVector3f Center, FVector3f Radius);
+
+		//! Define the root bone of the subhierarchy of the mesh that will be affected.
+		//! Only one of Box or Bone Hierarchy can be used (the last one set)
+		UE_API void SetVertexSelectionBone(const FBoneName& BoneId, float maxEffectRadius);
+
+
+	protected:
+
+		/** Forbidden. Manage with the Ptr<> template. */
+		~NodeMeshClipMorphPlane() {}
+
+	private:
+
+		static UE_API FNodeType StaticType;
+
+	};
+
+
+}
+
+#undef UE_API
diff --git a/Engine/Plugins/Mutable/Source/MutableTools/Internal/MuT/NodeMeshFragment.h b/Engine/Plugins/Mutable/Source/MutableTools/Internal/MuT/NodeMeshFragment.h
new file mode 100644
--- /dev/null
+++ b/Engine/Plugins/Mutable/Source/MutableTools/Internal/MuT/NodeMeshFragment.h
@@ -0,0 +1,49 @@
+// Copyright Epic Games, Inc. All Rights Reserved.
+
+#pragma once
+
+#include "MuR/Ptr.h"
+#include "MuT/NodeMesh.h"
+#include "MuT/NodeLayout.h"
+
+#define UE_API MUTABLETOOLS_API
+
+namespace UE::Mutable::Private
+{
+
+	/** This node extracts a fragment of a mesh by some selection on a texture layout space. */
+	class NodeMeshFragment : public NodeMesh
+	{
+	public:
+
+		/** Mesh from which a fragment will be extracted. */
+		Ptr<NodeMesh> SourceMesh;
+
+		/** Layout defining the blocks of mesh to extract by texture space. 
+		* If no layout is specified, all the blocks in the layout in the SourceMesh subgraph will be extracted.
+		*/
+		Ptr<NodeLayout> Layout;
+
+		/** Index of the UV channel in the source mesh to apply the Layout to decide what to extract. */
+		int32 LayoutIndex = 0;
+
+	public:
+
+		// Node Interface
+		virtual const FNodeType* GetType() const override { return GetStaticType(); }
+		static const FNodeType* GetStaticType() { return &StaticType; }
+
+	protected:
+
+		/** Forbidden. Manage with the Ptr<> template. */
+		~NodeMeshFragment() {}
+
+	private:
+
+		static UE_API FNodeType StaticType;
+
+	};
+
+}
+
+#undef UE_API
diff --git a/Engine/Plugins/Mutable/Source/MutableTools/Internal/MuT/NodeMeshMakeMorph.h b/Engine/Plugins/Mutable/Source/MutableTools/Internal/MuT/NodeMeshMakeMorph.h
new file mode 100644
--- /dev/null
+++ b/Engine/Plugins/Mutable/Source/MutableTools/Internal/MuT/NodeMeshMakeMorph.h
@@ -0,0 +1,45 @@
+// Copyright Epic Games, Inc. All Rights Reserved.
+
+#pragma once
+
+#include "MuR/Ptr.h"
+#include "MuR/RefCounted.h"
+#include "MuT/Node.h"
+#include "MuT/NodeMesh.h"
+
+#define UE_API MUTABLETOOLS_API
+
+
+namespace UE::Mutable::Private
+{
+
+	/** This node removes from a mesh A all the faces that are also part of a mesh B. */
+    class NodeMeshMakeMorph : public NodeMesh
+	{
+	public:
+
+		Ptr<NodeMesh> Base;
+		Ptr<NodeMesh> Target;
+
+		bool bOnlyPositionAndNormal = false;
+
+	public:
+
+		// Node Interface
+		virtual const FNodeType* GetType() const override { return &StaticType; }
+		static const FNodeType* GetStaticType() { return &StaticType; }
+
+	protected:
+
+		/** Forbidden. Manage with the Ptr<> template. */
+		~NodeMeshMakeMorph() {}
+
+	private:
+
+		static UE_API FNodeType StaticType;
+
+	};
+
+}
+
+#undef UE_API
diff --git a/Engine/Plugins/Mutable/Source/MutableTools/Internal/MuT/NodeMeshReshape.h b/Engine/Plugins/Mutable/Source/MutableTools/Internal/MuT/NodeMeshReshape.h
new file mode 100644
--- /dev/null
+++ b/Engine/Plugins/Mutable/Source/MutableTools/Internal/MuT/NodeMeshReshape.h
@@ -0,0 +1,57 @@
+// Copyright Epic Games, Inc. All Rights Reserved.
+
+#pragma once
+
+#include "MuR/Ptr.h"
+#include "MuT/Node.h"
+#include "MuT/NodeMesh.h"
+#include "MuR/Mesh.h"
+#include "MuR/Skeleton.h"
+
+#define UE_API MUTABLETOOLS_API
+
+namespace UE::Mutable::Private
+{
+
+	/** Node that morphs a base mesh with one or two weighted targets from a sequence. */
+	class NodeMeshReshape : public NodeMesh
+	{
+	public:
+
+		Ptr<NodeMesh> BaseMesh;
+		Ptr<NodeMesh> BaseShape;
+		Ptr<NodeMesh> TargetShape;
+		bool bReshapeVertices = true;
+		bool bRecomputeNormals = false;
+		bool bApplyLaplacian = false;
+		bool bReshapeSkeleton = false;
+		bool bReshapePhysicsVolumes = false;
+
+		EVertexColorUsage ColorRChannelUsage = EVertexColorUsage::None;
+		EVertexColorUsage ColorGChannelUsage = EVertexColorUsage::None;
+		EVertexColorUsage ColorBChannelUsage = EVertexColorUsage::None;
+		EVertexColorUsage ColorAChannelUsage = EVertexColorUsage::None;
+
+		TArray<FBoneName> BonesToDeform;
+		TArray<FBoneName> PhysicsToDeform;
+
+	public:
+
+		// Node Interface
+		virtual const FNodeType* GetType() const override { return &StaticType; }
+		static const FNodeType* GetStaticType() { return &StaticType; }
+
+	protected:
+
+		/** Forbidden. Manage with the Ptr<> template. */
+		~NodeMeshReshape() {}
+
+	private:
+
+		static UE_API FNodeType StaticType;
+
+	};
+
+}
+
+#undef UE_API
diff --git a/Engine/Plugins/Mutable/Source/MutableTools/Internal/MuT/NodeMeshSwitch.h b/Engine/Plugins/Mutable/Source/MutableTools/Internal/MuT/NodeMeshSwitch.h
new file mode 100644
--- /dev/null
+++ b/Engine/Plugins/Mutable/Source/MutableTools/Internal/MuT/NodeMeshSwitch.h
@@ -0,0 +1,42 @@
+// Copyright Epic Games, Inc. All Rights Reserved.
+
+#pragma once
+
+#include "MuR/Ptr.h"
+#include "MuT/NodeMesh.h"
+#include "MuT/NodeScalar.h"
+
+#define UE_API MUTABLETOOLS_API
+
+namespace UE::Mutable::Private
+{
+
+	/** This node selects an output Mesh from a set of input Meshs based on a parameter. */
+	class NodeMeshSwitch : public NodeMesh
+	{
+	public:
+
+		Ptr<NodeScalar> Parameter;
+
+		TArray<Ptr<NodeMesh>> Options;
+
+	public:
+
+		// Node Interface
+		virtual const FNodeType* GetType() const override { return &StaticType; }
+		static const FNodeType* GetStaticType() { return &StaticType; }
+
+	protected:
+
+		/** Forbidden. Manage with the Ptr<> template. */
+		~NodeMeshSwitch() {}
+
+	private:
+
+		static UE_API FNodeType StaticType;
+
+	};
+
+}
+
+#undef UE_API
diff --git a/Engine/Plugins/Mutable/Source/MutableTools/Internal/MuT/NodeModifier.h b/Engine/Plugins/Mutable/Source/MutableTools/Internal/MuT/NodeModifier.h
new file mode 100644
--- /dev/null
+++ b/Engine/Plugins/Mutable/Source/MutableTools/Internal/MuT/NodeModifier.h
@@ -0,0 +1,78 @@
+// Copyright Epic Games, Inc. All Rights Reserved.
+
+#pragma once
+
+#include "HAL/Platform.h"
+#include "MuR/Ptr.h"
+#include "MuR/RefCounted.h"
+#include "MuR/Serialisation.h"
+#include "MuT/Node.h"
+#include "Containers/UnrealString.h"
+
+#include "NodeModifier.generated.h"
+
+#define UE_API MUTABLETOOLS_API
+
+
+/** Despite being an UEnum, this is not always version-serialized (in MutableTools).
+* Beware of changing the enum options or order.
+*/
+UENUM()
+enum class EMutableMultipleTagPolicy : uint8
+{
+	OnlyOneRequired,
+	AllRequired
+};
+
+
+namespace UE::Mutable::Private
+{
+
+	// Forward definitions
+	class NodeModifier;
+	typedef Ptr<NodeModifier> NodeModifierPtr;
+	typedef Ptr<const NodeModifier> NodeModifierConst;
+
+	/** Parent of all node classes that apply modifiers to surfaces. */
+	class NodeModifier : public Node
+	{
+	public:
+
+		/** If not negative, this modifier will only be applied to the nodes of the component with matching id. Otherwise it will be applied to all components. */
+		int32 RequiredComponentId = -1;
+
+		/** Tags that target surface need to have enabled to receive this modifier. */
+		TArray<FString> RequiredTags;
+
+		/** In case of multiple tags in RequiredTags: are they all required, or one is enough? */
+		EMutableMultipleTagPolicy MultipleTagsPolicy = EMutableMultipleTagPolicy::OnlyOneRequired;
+
+		/** Wether the modifier has to be applied before the normal node operations or after. */
+		bool bApplyBeforeNormalOperations = true;
+
+		/** Tags enabled by this modifier. Other modifiers activated by these tags will be applied to this modifier's "child data" like meshes added by this modifier.
+		* Not to be confused with the RequiredTags.
+		*/
+		TArray<FString> EnableTags;
+
+	public:
+
+		// Node interface
+		virtual const FNodeType* GetType() const override { return &StaticType; }
+		static const FNodeType* GetStaticType() { return &StaticType; }
+
+	protected:
+
+		/** Forbidden. Manage with the Ptr<> template. */
+		inline ~NodeModifier() {}
+
+	private:
+
+		static UE_API FNodeType StaticType;
+
+	};
+
+
+}
+
+#undef UE_API
diff --git a/Engine/Plugins/Mutable/Source/MutableTools/Internal/MuT/NodeModifierMeshClipDeform.h b/Engine/Plugins/Mutable/Source/MutableTools/Internal/MuT/NodeModifierMeshClipDeform.h
new file mode 100644
--- /dev/null
+++ b/Engine/Plugins/Mutable/Source/MutableTools/Internal/MuT/NodeModifierMeshClipDeform.h
@@ -0,0 +1,50 @@
+// Copyright Epic Games, Inc. All Rights Reserved.
+
+#pragma once
+
+#include "HAL/Platform.h"
+#include "MuR/Ptr.h"
+#include "MuR/RefCounted.h"
+#include "MuR/Mesh.h"
+#include "MuT/Node.h"
+#include "MuT/NodeModifier.h"
+#include "MuT/NodeMesh.h"
+
+#define UE_API MUTABLETOOLS_API
+
+namespace UE::Mutable::Private
+{
+	
+	/** */
+	class NodeModifierMeshClipDeform : public NodeModifier
+	{
+	public:
+
+		Ptr<NodeMesh> ClipMesh;
+
+		EShapeBindingMethod BindingMethod = EShapeBindingMethod::ClipDeformClosestProject;
+
+		/** */
+		EFaceCullStrategy FaceCullStrategy = EFaceCullStrategy::AllVerticesCulled;
+
+	public:
+
+		// Node interface
+		virtual const FNodeType* GetType() const override { return &StaticType; }
+		static const FNodeType* GetStaticType() { return &StaticType; }
+
+	protected:
+
+		/** Forbidden. Manage with the Ptr<> template. */
+		inline ~NodeModifierMeshClipDeform() {}
+
+	private:
+
+		static UE_API FNodeType StaticType;
+
+	};
+
+
+}
+
+#undef UE_API
diff --git a/Engine/Plugins/Mutable/Source/MutableTools/Internal/MuT/NodeModifierMeshClipWithMesh.h b/Engine/Plugins/Mutable/Source/MutableTools/Internal/MuT/NodeModifierMeshClipWithMesh.h
new file mode 100644
--- /dev/null
+++ b/Engine/Plugins/Mutable/Source/MutableTools/Internal/MuT/NodeModifierMeshClipWithMesh.h
@@ -0,0 +1,47 @@
+// Copyright Epic Games, Inc. All Rights Reserved.
+
+#pragma once
+
+#include "MuR/Ptr.h"
+#include "MuR/RefCounted.h"
+#include "MuR/Types.h"
+#include "MuT/Node.h"
+#include "MuT/NodeModifier.h"
+#include "MuT/NodeMesh.h"
+
+#define UE_API MUTABLETOOLS_API
+
+namespace UE::Mutable::Private
+{
+
+	/** */
+	class NodeModifierMeshClipWithMesh : public NodeModifier
+	{
+	public:
+
+		/** */
+		Ptr<NodeMesh> ClipMesh;
+
+		/** */
+		EFaceCullStrategy FaceCullStrategy = EFaceCullStrategy::AllVerticesCulled;
+
+	public:
+
+		// Node interface
+		virtual const FNodeType* GetType() const override { return &StaticType; }
+		static const FNodeType* GetStaticType() { return &StaticType; }
+
+	protected:
+
+		/** Forbidden. Manage with the Ptr<> template. */
+		inline ~NodeModifierMeshClipWithMesh() {}
+
+	private:
+
+		static UE_API FNodeType StaticType;
+
+	};
+
+}
+
+#undef UE_API
diff --git a/Engine/Plugins/Mutable/Source/MutableTools/Internal/MuT/NodeObject.h b/Engine/Plugins/Mutable/Source/MutableTools/Internal/MuT/NodeObject.h
new file mode 100644
--- /dev/null
+++ b/Engine/Plugins/Mutable/Source/MutableTools/Internal/MuT/NodeObject.h
@@ -0,0 +1,51 @@
+// Copyright Epic Games, Inc. All Rights Reserved.
+
+#pragma once
+
+#include "HAL/Platform.h"
+#include "MuR/Ptr.h"
+#include "MuR/RefCounted.h"
+#include "MuT/Node.h"
+
+#define UE_API MUTABLETOOLS_API
+
+namespace UE::Mutable::Private
+{
+
+	/** Base class for any node that outputs an object. */
+	class NodeObject : public Node
+	{
+	public:
+
+		// Node Interface
+		virtual const FNodeType* GetType() const override { return &StaticType; }
+		static const FNodeType* GetStaticType() { return &StaticType; }
+
+		// Own interface
+
+		//! Get the name of the object.
+		virtual const FString& GetName() const = 0;
+
+		//! Set the name of the object.
+		virtual void SetName( const FString& ) = 0;
+
+		//! Get the uid of the object.
+		virtual const FString& GetUid() const = 0;
+
+		//! Set the uid of the object.
+		virtual void SetUid( const FString& ) = 0;
+
+	protected:
+
+		/** Forbidden. Manage with the Ptr<> template. */
+		inline ~NodeObject() {}
+
+	private:
+
+		static UE_API FNodeType StaticType;
+
+	};
+
+}
+
+#undef UE_API
diff --git a/Engine/Plugins/Mutable/Source/MutableTools/Internal/MuT/NodeObjectGroup.h b/Engine/Plugins/Mutable/Source/MutableTools/Internal/MuT/NodeObjectGroup.h
new file mode 100644
--- /dev/null
+++ b/Engine/Plugins/Mutable/Source/MutableTools/Internal/MuT/NodeObjectGroup.h
@@ -0,0 +1,80 @@
+// Copyright Epic Games, Inc. All Rights Reserved.
+
+#pragma once
+
+#include "MuR/Ptr.h"
+#include "MuR/RefCounted.h"
+#include "MuT/Node.h"
+#include "MuT/NodeObject.h"
+
+#define UE_API MUTABLETOOLS_API
+
+namespace UE::Mutable::Private
+{
+
+	/** Node that creates a group of objects and describes how they are selected. */
+	class NodeObjectGroup : public NodeObject
+	{
+	public:
+
+		FString Name;
+		FString Uid;
+
+		//! Typed of child selection
+		typedef enum
+		{
+			//! All objects in the group will always be enabled, and no parameter will be generated.
+			CS_ALWAYS_ALL,
+
+			//! Only one children may be selected, but it is allowed to have none.
+			//! An enumeration parameter will be generated but it may have a null value
+			CS_ONE_OR_NONE,
+
+			//! One and only one children has to be selected at all times
+			//! An enumeraation parameter will be generated and it cannot be null
+			CS_ALWAYS_ONE,
+
+			//! Each child in the group can be enabled or disabled individually.
+			//! A boolean parameter will be generated for every child.
+			CS_TOGGLE_EACH
+
+		} EChildSelection;
+
+		EChildSelection Type;
+
+		TArray<Ptr<NodeObject>> Children;
+
+		//! Default value for CS_ONE_OR_NONE or CS_ALWAYS_ONE groups
+		//! the value is the index of the child option in the group
+		//! -1 is the value for the None option.
+		//! 0 is the first child whether or not the NONE option is present.
+		int32 DefaultValue;
+
+	public:
+
+		// Node Interface
+		virtual const FNodeType* GetType() const override { return &StaticType; }
+		static const FNodeType* GetStaticType() { return &StaticType; }
+
+		// NodeObject Interface
+        UE_API virtual const FString& GetName() const override;
+		UE_API virtual void SetName( const FString&) override;
+		UE_API virtual const FString& GetUid() const override;
+		UE_API virtual void SetUid( const FString&) override;
+
+	protected:
+
+		/** Forbidden. Manage with the Ptr<> template. */
+		~NodeObjectGroup() {}
+
+	private:
+
+		static UE_API FNodeType StaticType;
+
+	};
+
+
+
+}
+
+#undef UE_API
diff --git a/Engine/Plugins/Mutable/Source/MutableTools/Internal/MuT/NodeProjector.h b/Engine/Plugins/Mutable/Source/MutableTools/Internal/MuT/NodeProjector.h
new file mode 100644
--- /dev/null
+++ b/Engine/Plugins/Mutable/Source/MutableTools/Internal/MuT/NodeProjector.h
@@ -0,0 +1,141 @@
+// Copyright Epic Games, Inc. All Rights Reserved.
+
+#pragma once
+
+#include "HAL/Platform.h"
+#include "MuR/Parameters.h"
+#include "MuR/Ptr.h"
+#include "MuR/RefCounted.h"
+#include "MuT/Node.h"
+#include "MuT/NodeRange.h"
+
+#define UE_API MUTABLETOOLS_API
+
+
+namespace UE::Mutable::Private
+{
+
+    //! Base class of any node that outputs a Projector.
+	class NodeProjector : public Node
+	{
+	public:
+
+		// Node Interface
+		virtual const FNodeType* GetType() const override { return &StaticType; }
+		static const FNodeType* GetStaticType() { return &StaticType; }
+
+	protected:
+
+		/** Forbidden. Manage with the Ptr<> template. */
+		inline ~NodeProjector() {}
+
+	private:
+
+		static UE_API FNodeType StaticType;
+
+	};
+
+
+	//! This node outputs a predefined Projector value.
+	class NodeProjectorConstant : public NodeProjector
+	{
+	public:
+
+		EProjectorType Type = EProjectorType::Planar;
+		FVector3f Position = FVector3f::ZeroVector;
+		FVector3f Direction = FVector3f::ZeroVector;
+		FVector3f Up = FVector3f::ZeroVector;
+		FVector3f Scale = FVector3f::ZeroVector;
+		float ProjectionAngle = 0.0f;
+
+	public:
+
+		// Node Interface
+		virtual const FNodeType* GetType() const override { return &StaticType; }
+		static const FNodeType* GetStaticType() { return &StaticType; }
+
+		// Own Interface
+
+		//! Get the value that this node returns
+        UE_API void GetValue( EProjectorType* OutType,
+			FVector3f* OutPos,
+			FVector3f* OutDir,
+			FVector3f* OutUp,
+			FVector3f* OutScaleU,
+			float* OutProjectionAngle ) const;
+
+		//! Set the value to be returned by this node
+        UE_API void SetValue( EProjectorType type, FVector3f pos, FVector3f dir, FVector3f up, FVector3f scale, float projectionAngle);
+
+	protected:
+
+		/** Forbidden. Manage with the Ptr<> template. */
+		~NodeProjectorConstant() {}
+
+	private:
+
+		static UE_API FNodeType StaticType;
+
+	};
+
+
+	//! Node that defines a Projector model parameter.
+	class NodeProjectorParameter : public NodeProjector
+	{
+	public:
+
+		EProjectorType Type = EProjectorType::Planar;
+		FVector3f Position = FVector3f::ZeroVector;
+		FVector3f Direction = FVector3f::ZeroVector;
+		FVector3f Up = FVector3f::ZeroVector;
+		FVector3f Scale = FVector3f::ZeroVector;
+		float ProjectionAngle = 0.0f;
+
+		FString Name;
+		FString UID;
+
+		TArray<Ptr<NodeRange>> Ranges;
+
+	public:
+
+		// Node Interface
+		virtual const FNodeType* GetType() const override { return &StaticType; }
+		static const FNodeType* GetStaticType() { return &StaticType; }
+
+		// Own Interface
+
+		//! Set the name of the parameter.
+		UE_API void SetName( const FString& );
+
+		//! Get the uid of the parameter. It will be exposed in the final compiled data.
+		UE_API const FString& GetUid() const;
+		UE_API void SetUid( const FString& );
+
+		//! Set the default value of the parameter.
+        UE_API void SetDefaultValue( EProjectorType type,
+			FVector3f pos,
+			FVector3f dir,
+			FVector3f up,
+			FVector3f scale,
+			float projectionAngle );
+
+        //! Set the number of ranges (dimensions) for this parameter.
+        //! By default a parameter has 0 ranges, meaning it only has one value.
+        UE_API void SetRangeCount( int32 i );
+        UE_API void SetRange( int32 i, Ptr<NodeRange> Range );
+
+	protected:
+
+		/** Forbidden. Manage with the Ptr<> template. */
+		~NodeProjectorParameter() {}
+
+	private:
+
+		static UE_API FNodeType StaticType;
+
+	};
+
+
+}
+
+#undef UE_API
diff --git a/Engine/Plugins/Mutable/Source/MutableTools/Internal/MuT/NodeScalarConstant.h b/Engine/Plugins/Mutable/Source/MutableTools/Internal/MuT/NodeScalarConstant.h
new file mode 100644
--- /dev/null
+++ b/Engine/Plugins/Mutable/Source/MutableTools/Internal/MuT/NodeScalarConstant.h
@@ -0,0 +1,41 @@
+// Copyright Epic Games, Inc. All Rights Reserved.
+
+#pragma once
+
+#include "MuR/Ptr.h"
+#include "MuR/RefCounted.h"
+#include "MuT/Node.h"
+#include "MuT/NodeScalar.h"
+
+#define UE_API MUTABLETOOLS_API
+
+namespace UE::Mutable::Private
+{
+
+	/** Node returning a scalar constant value. */
+	class NodeScalarConstant : public NodeScalar
+	{
+	public:
+
+		float Value;
+
+	public:
+
+		// Node Interface
+		virtual const FNodeType* GetType() const override { return &StaticType; }
+		static const FNodeType* GetStaticType() { return &StaticType; }
+
+	protected:
+
+		/** Forbidden. Manage with the Ptr<> template. */
+		~NodeScalarConstant() {}
+
+	private:
+
+		static UE_API FNodeType StaticType;
+
+	};
+
+}
+
+#undef UE_API
diff --git a/Engine/Plugins/Mutable/Source/MutableTools/Internal/MuT/NodeScalarEnumParameter.h b/Engine/Plugins/Mutable/Source/MutableTools/Internal/MuT/NodeScalarEnumParameter.h
new file mode 100644
--- /dev/null
+++ b/Engine/Plugins/Mutable/Source/MutableTools/Internal/MuT/NodeScalarEnumParameter.h
@@ -0,0 +1,54 @@
+// Copyright Epic Games, Inc. All Rights Reserved.
+
+#pragma once
+
+#include "MuR/Ptr.h"
+#include "MuR/RefCounted.h"
+#include "MuT/Node.h"
+#include "MuT/NodeRange.h"
+#include "MuT/NodeScalar.h"
+
+#define UE_API MUTABLETOOLS_API
+
+namespace UE::Mutable::Private
+{
+
+	/** Node that defines a scalar model parameter to be selected from a set of named values. */
+	class NodeScalarEnumParameter : public NodeScalar
+	{
+	public:
+
+		FString Name;
+		FString UID;
+		int32 DefaultValue = 0;
+
+		struct FOption
+		{
+			FString Name;
+			float Value = 0.0f;
+		};
+
+		TArray<FOption> Options;
+
+		TArray<Ptr<NodeRange>> Ranges;
+
+	public:
+
+		// Node Interface
+		virtual const FNodeType* GetType() const override { return &StaticType; }
+		static const FNodeType* GetStaticType() { return &StaticType; }
+
+	protected:
+
+		/** Forbidden. Manage with the Ptr<> template. */
+		~NodeScalarEnumParameter() {}
+
+	private:
+
+		static UE_API FNodeType StaticType;
+
+	};
+
+}
+
+#undef UE_API
diff --git a/Engine/Plugins/Mutable/Source/MutableTools/Internal/MuT/NodeScalarSwitch.h b/Engine/Plugins/Mutable/Source/MutableTools/Internal/MuT/NodeScalarSwitch.h
new file mode 100644
--- /dev/null
+++ b/Engine/Plugins/Mutable/Source/MutableTools/Internal/MuT/NodeScalarSwitch.h
@@ -0,0 +1,44 @@
+// Copyright Epic Games, Inc. All Rights Reserved.
+
+#pragma once
+
+#include "MuR/Ptr.h"
+#include "MuR/RefCounted.h"
+#include "MuT/Node.h"
+#include "MuT/NodeScalar.h"
+
+#define UE_API MUTABLETOOLS_API
+
+
+namespace UE::Mutable::Private
+{
+
+    /** This node selects an output Scalar from a set of input Scalars based on a parameter. */
+    class NodeScalarSwitch : public NodeScalar
+	{
+	public:
+
+		Ptr<NodeScalar> Parameter;
+		TArray<Ptr<NodeScalar>> Options;
+
+	public:
+
+		// Node interface
+		virtual const FNodeType* GetType() const override { return &StaticType; }
+		static const FNodeType* GetStaticType() { return &StaticType; }
+
+	protected:
+
+		/** Forbidden. Manage with the Ptr<> template. */
+		~NodeScalarSwitch() {}
+
+	private:
+
+		static UE_API FNodeType StaticType;
+
+	};
+
+
+}
+
+#undef UE_API
diff --git a/Engine/Plugins/Mutable/Source/MutableTools/Internal/MuT/NodeScalarTable.h b/Engine/Plugins/Mutable/Source/MutableTools/Internal/MuT/NodeScalarTable.h
new file mode 100644
--- /dev/null
+++ b/Engine/Plugins/Mutable/Source/MutableTools/Internal/MuT/NodeScalarTable.h
@@ -0,0 +1,46 @@
+// Copyright Epic Games, Inc. All Rights Reserved.
+
+#pragma once
+
+#include "MuR/Ptr.h"
+#include "MuR/RefCounted.h"
+#include "MuT/Node.h"
+#include "MuT/NodeScalar.h"
+#include "MuT/Table.h"
+
+#define UE_API MUTABLETOOLS_API
+
+namespace UE::Mutable::Private
+{
+
+	/** This node provides the meshes stored in the column of a table. */
+	class NodeScalarTable : public NodeScalar
+	{
+	public:
+
+		FString ParameterName;
+		Ptr<FTable> Table;
+		FString ColumnName;
+		bool bNoneOption = false;
+		FString DefaultRowName;
+
+	public:
+
+		// Node interface
+		virtual const FNodeType* GetType() const override { return GetStaticType(); }
+		static const FNodeType* GetStaticType() { return &StaticType; }
+
+	protected:
+
+		/** Forbidden. Manage with the Ptr<> template. */
+		~NodeScalarTable() {}
+
+	private:
+
+		static UE_API FNodeType StaticType;
+
+	};
+
+}
+
+#undef UE_API
diff --git a/Engine/Plugins/Mutable/Source/MutableTools/Internal/MuT/NodeScalarVariation.h b/Engine/Plugins/Mutable/Source/MutableTools/Internal/MuT/NodeScalarVariation.h
new file mode 100644
--- /dev/null
+++ b/Engine/Plugins/Mutable/Source/MutableTools/Internal/MuT/NodeScalarVariation.h
@@ -0,0 +1,49 @@
+// Copyright Epic Games, Inc. All Rights Reserved.
+
+#pragma once
+
+#include "MuR/Ptr.h"
+#include "MuR/RefCounted.h"
+#include "MuT/Node.h"
+#include "MuT/NodeScalar.h"
+
+#define UE_API MUTABLETOOLS_API
+
+namespace UE::Mutable::Private
+{
+
+	/** */
+    class NodeScalarVariation : public NodeScalar
+    {
+	public:
+
+		Ptr<NodeScalar> DefaultScalar;
+
+		struct FVariation
+		{
+			Ptr<NodeScalar> Scalar;
+			FString Tag;
+		};
+
+		TArray<FVariation> Variations;
+
+    public:
+ 
+		// Node Interface
+		virtual const FNodeType* GetType() const override { return &StaticType; }
+		static const FNodeType* GetStaticType() { return &StaticType; }
+
+    protected:
+
+		/** Forbidden. Manage with the Ptr<> template. */
+		~NodeScalarVariation() {}
+
+    private:
+
+		static UE_API FNodeType StaticType;
+
+	};
+
+}
+
+#undef UE_API
diff --git a/Engine/Plugins/Mutable/Source/MutableTools/Internal/MuT/NodeString.h b/Engine/Plugins/Mutable/Source/MutableTools/Internal/MuT/NodeString.h
new file mode 100644
--- /dev/null
+++ b/Engine/Plugins/Mutable/Source/MutableTools/Internal/MuT/NodeString.h
@@ -0,0 +1,39 @@
+// Copyright Epic Games, Inc. All Rights Reserved.
+
+#pragma once
+
+#include "HAL/Platform.h"
+#include "MuR/Ptr.h"
+#include "MuR/RefCounted.h"
+#include "MuT/Node.h"
+
+#define UE_API MUTABLETOOLS_API
+
+
+namespace UE::Mutable::Private
+{
+
+    /** Base class of any node that outputs a string value. */
+	class NodeString : public Node
+	{
+	public:
+
+		// Node Interface
+		virtual const FNodeType* GetType() const override { return &StaticType; }
+		static const FNodeType* GetStaticType() { return &StaticType; }
+
+	protected:
+
+		/** Forbidden. Manage with the Ptr<> template. */
+		~NodeString() {}
+
+	private:
+
+		static UE_API FNodeType StaticType;
+
+	};
+
+
+}
+
+#undef UE_API
diff --git a/Engine/Plugins/Mutable/Source/MutableTools/Internal/MuT/NodeSurface.h b/Engine/Plugins/Mutable/Source/MutableTools/Internal/MuT/NodeSurface.h
new file mode 100644
--- /dev/null
+++ b/Engine/Plugins/Mutable/Source/MutableTools/Internal/MuT/NodeSurface.h
@@ -0,0 +1,45 @@
+// Copyright Epic Games, Inc. All Rights Reserved.
+
+#pragma once
+
+#include "HAL/Platform.h"
+#include "MuR/Ptr.h"
+#include "MuR/RefCounted.h"
+#include "MuT/Node.h"
+
+#define UE_API MUTABLETOOLS_API
+
+
+namespace UE::Mutable::Private
+{
+
+	// Forward definitions
+    class NodeSurface;
+    typedef Ptr<NodeSurface> NodeSurfacePtr;
+    typedef Ptr<const NodeSurface> NodeSurfacePtrConst;
+
+
+    /** This class is the parent of all nodes that output a Surface. */
+    class NodeSurface : public Node
+	{
+	public:
+
+		// Node interface
+		virtual const FNodeType* GetType() const override { return &StaticType; }
+		static const FNodeType* GetStaticType() { return &StaticType; }
+
+	protected:
+
+		/** Forbidden. Manage with the Ptr<> template. */
+        inline ~NodeSurface() {}
+
+	private:
+
+		static UE_API FNodeType StaticType;
+
+	};
+
+
+}
+
+#undef UE_API
diff --git a/Engine/Plugins/Mutable/Source/MutableTools/Internal/MuT/NodeSurfaceNew.h b/Engine/Plugins/Mutable/Source/MutableTools/Internal/MuT/NodeSurfaceNew.h
new file mode 100644
--- /dev/null
+++ b/Engine/Plugins/Mutable/Source/MutableTools/Internal/MuT/NodeSurfaceNew.h
@@ -0,0 +1,106 @@
+// Copyright Epic Games, Inc. All Rights Reserved.
+
+#pragma once
+
+#include "MuR/Ptr.h"
+#include "MuR/RefCounted.h"
+#include "MuT/Node.h"
+#include "MuT/NodeSurface.h"
+#include "MuT/NodeMesh.h"
+#include "MuT/NodeMaterial.h"
+#include "MuT/NodeImage.h"
+#include "MuT/NodeScalar.h"
+#include "MuT/NodeString.h"
+#include "MuT/NodeColour.h"
+
+#include "Misc/Guid.h"
+
+#define UE_API MUTABLETOOLS_API
+
+namespace UE::Mutable::Private
+{
+
+	/** This node makes a new Surface (Mesh Material Section) from a mesh and several material parameters like images, vectors and scalars. */
+	class NodeSurfaceNew : public NodeSurface
+	{
+	public:
+
+		FString Name;
+
+		/** An optional, opaque id that will be returned in the surfaces of the created
+		* instances. Can be useful to identify surfaces on the application side.
+		*/
+		uint32 ExternalId = 0;
+
+		/** Source surface Guid, used to share layouts between LODs and generate a deterministic SharedSurfaceId. */
+		FGuid SurfaceGuid;
+
+		Ptr<NodeMesh> Mesh;
+
+		Ptr<NodeMaterial> Material;
+
+		struct FImageData
+		{
+			FString Name;
+			FString MaterialName;
+			FString MaterialParameterName;
+			Ptr<NodeImage> Image;
+
+			/** Index of the layout transform to apply to this image. It could be negative, to indicate no layout transform. */
+			int8 LayoutIndex = 0;
+		};
+
+		TArray<FImageData> Images;
+
+		struct FVectorData
+		{
+			FString Name;
+			Ptr<NodeColour> Vector;
+		};
+
+		TArray<FVectorData> Vectors;
+
+		struct FScalar
+		{
+			FString Name;
+			Ptr<NodeScalar> Scalar;
+		};
+
+		TArray<FScalar> Scalars;
+
+		struct FStringData
+		{
+			FString Name;
+			Ptr<NodeString> String;
+		};
+
+		TArray<FStringData> Strings;
+
+		/** Tags added to the surface:
+		* - the surface will be affected by modifier nodes with the same tag
+		* - the tag will be enabled when the surface is added to an object, and it can activate
+		* variations for any surface.
+		*/
+		TArray<FString> Tags;
+
+	public:
+
+		// Node interface
+		virtual const FNodeType* GetType() const override { return &StaticType; }
+		static const FNodeType* GetStaticType() { return &StaticType; }
+
+	protected:
+
+		/** Forbidden. Manage with the Ptr<> template. */
+		~NodeSurfaceNew() {}
+
+	private:
+
+		static UE_API FNodeType StaticType;
+
+	};
+
+
+}
+
+#undef UE_API
diff --git a/Engine/Plugins/Mutable/Source/MutableTools/Internal/MuT/Table.h b/Engine/Plugins/Mutable/Source/MutableTools/Internal/MuT/Table.h
new file mode 100644
--- /dev/null
+++ b/Engine/Plugins/Mutable/Source/MutableTools/Internal/MuT/Table.h
@@ -0,0 +1,87 @@
+// Copyright Epic Games, Inc. All Rights Reserved.
+
+#pragma once
+
+#include "MuR/Ptr.h"
+#include "MuR/RefCounted.h"
+#include "MuR/Serialisation.h"
+
+#define UE_API MUTABLETOOLS_API
+
+
+namespace UE::Mutable::Private
+{
+	// Forward declarations
+	class FTable;
+	class FMesh;
+	class FImage;
+
+	/** Types of the values for the table cells. */
+	enum class ETableColumnType : uint32
+	{
+		None,
+		Scalar,
+		Color,
+		Image,
+		Mesh,
+		String,
+		Material
+	};
+
+
+	//! A table that contains many rows and defines attributes like meshes, images,
+	//! colours, etc. for every column. It is useful to define a big number of similarly structured
+	//! objects, by using the NodeDatabase in a model expression.
+	//! \ingroup model
+	class FTable : public RefCounted
+	{
+	public:
+
+		UE_API FTable();
+
+		//-----------------------------------------------------------------------------------------
+		// Own interface
+		//-----------------------------------------------------------------------------------------
+
+		//!
+		UE_API void SetName(const FString&);
+		UE_API const FString& GetName() const;
+
+		//!
+		UE_API int32 AddColumn(const FString&, ETableColumnType );
+
+		//! Return the column index with the given name. -1 if not found.
+		UE_API int32 FindColumn(const FString&) const;
+
+		//!
+        UE_API void AddRow( uint32 id );
+
+		//!
+        UE_API void SetCell( int32 Column, uint32 RowId, int32 IntValue, const void* ErrorContext = nullptr);
+        UE_API void SetCell( int32 Column, uint32 RowId, float Value, const void* ErrorContext = nullptr);
+        UE_API void SetCell( int32 Column, uint32 RowId, const FVector4f& Value, const void* ErrorContext = nullptr);
+		UE_API void SetCell( int32 Column, uint32 RowId, TResourceProxy<FImage>* Value, const void* ErrorContext = nullptr);
+		UE_API void SetCell( int32 Column, uint32 RowId, const TSharedPtr<FMesh>& Value, const void* ErrorContext = nullptr);
+        UE_API void SetCell( int32 Column, uint32 RowId, const FString& Value, const void* ErrorContext = nullptr);
+
+
+		//-----------------------------------------------------------------------------------------
+		// Interface pattern
+		//-----------------------------------------------------------------------------------------
+		class Private;
+		UE_API Private* GetPrivate() const;
+
+	protected:
+
+		/** Forbidden. Manage with the Ptr<> template. */
+		UE_API ~FTable();
+
+	private:
+
+		Private* m_pD;
+
+	};
+
+}
+
+#undef UE_API
diff --git a/Engine/Plugins/Mutable/Source/MutableTools/Internal/MuT/TypeInfo.h b/Engine/Plugins/Mutable/Source/MutableTools/Internal/MuT/TypeInfo.h
new file mode 100644
--- /dev/null
+++ b/Engine/Plugins/Mutable/Source/MutableTools/Internal/MuT/TypeInfo.h
@@ -0,0 +1,37 @@
+// Copyright Epic Games, Inc. All Rights Reserved.
+
+#pragma once
+
+#include "HAL/PlatformMath.h"
+#include "MuR/Image.h"
+#include "MuR/MeshBufferSet.h"
+#include "MuR/Parameters.h"
+#include "MuR/ParametersPrivate.h"
+
+#define UE_API MUTABLETOOLS_API
+
+// This header file provides additional information about data types defined in the runtime, that
+// can be useful only for the tools. For instance, it provides strings for some enumeration values.
+
+namespace UE::Mutable::Private
+{
+
+	class TypeInfo
+	{
+	public:
+
+		static UE_API const char* s_imageFormatName[size_t(EImageFormat::Count)];
+
+		static UE_API const char* s_meshBufferSemanticName[uint32(EMeshBufferSemantic::Count)];
+
+		static UE_API const char* s_meshBufferFormatName[uint32(EMeshBufferFormat::Count) ];
+
+		static UE_API const char* s_blendModeName[uint32(EBlendType::_BT_COUNT) ];
+
+		static UE_API const char* s_projectorTypeName [ static_cast<uint32>(UE::Mutable::Private::EProjectorType::Count) ];
+
+	};
+
+}
+
+#undef UE_API
diff --git a/Engine/Plugins/Mutable/Source/MutableTools/Internal/MuT/UnrealPixelFormatOverride.h b/Engine/Plugins/Mutable/Source/MutableTools/Internal/MuT/UnrealPixelFormatOverride.h
new file mode 100644
--- /dev/null
+++ b/Engine/Plugins/Mutable/Source/MutableTools/Internal/MuT/UnrealPixelFormatOverride.h
@@ -0,0 +1,38 @@
+// Copyright Epic Games, Inc. All Rights Reserved.
+
+#pragma once
+
+#include "MuR/Image.h"
+#include "Interfaces/ITextureFormatManagerModule.h"
+#include "Interfaces/ITextureFormatModule.h"
+#include "Engine/TextureDefines.h"
+#include "TextureCompressorModule.h"
+#include "TextureBuildUtilities.h"
+#include "ImageCore.h"
+
+/** This function needs to be called from the game thread before any call to UnrealPixelFormatFunc can be done. */
+MUTABLETOOLS_API extern void PrepareUnrealCompression();
+
+/** */
+void FillBuildSettingsFromMutableFormat(FTextureBuildSettings& Settings, bool& bOutHasAlpha, UE::Mutable::Private::EImageFormat Format);
+
+/** */
+void MutableToImageCore(const UE::Mutable::Private::FImage* InMutable, FImage& CoreImage, int32 LOD, bool bSwizzleRGBHack);
+
+/** */
+bool ImageCoreToMutable(const FCompressedImage2D& Compressed, UE::Mutable::Private::FImage* Mutable, int32 LOD);
+
+/** Convert an Unreal platform pixel format to Mutable. */
+MUTABLETOOLS_API extern UE::Mutable::Private::EImageFormat UnrealToMutablePixelFormat(EPixelFormat PlatformFormat, bool bHasAlpha);
+
+/** Function that remaps some formats to the ones that provide more acceptable quality in Mutable. */
+MUTABLETOOLS_API extern UE::Mutable::Private::EImageFormat QualityAndPerformanceFix(UE::Mutable::Private::EImageFormat Format);
+
+/** Try to convert an image using unreal's compression. If the format is not supported, it will return false in bOutSuccess.
+* Can be called from any thread.
+* PrepareUnrealCompression needs to be called from the game thread once prior to using this function. 
+*/
+MUTABLETOOLS_API extern void UnrealPixelFormatFunc(bool& bOutSuccess, int32 Quality, UE::Mutable::Private::FImage* Target, const UE::Mutable::Private::FImage* Source, int32 OnlyLOD);
+
+/** Debug export the given image to a temp file in the saved folder. */
+MUTABLETOOLS_API void DebugImageDump(const UE::Mutable::Private::FImage*, const FString& FilePath);
diff --git a/Engine/Plugins/Mutable/Source/MutableTools/Public/MuT/MutableToolsModule.h b/Engine/Plugins/Mutable/Source/MutableTools/Public/MuT/MutableToolsModule.h
deleted file mode 100644
diff --git a/Engine/Plugins/Mutable/Source/MutableTools/Public/MuT/NodeColorToSRGB.h b/Engine/Plugins/Mutable/Source/MutableTools/Public/MuT/NodeColorToSRGB.h
deleted file mode 100644
diff --git a/Engine/Plugins/Mutable/Source/MutableTools/Public/MuT/NodeColourConstant.h b/Engine/Plugins/Mutable/Source/MutableTools/Public/MuT/NodeColourConstant.h
deleted file mode 100644
diff --git a/Engine/Plugins/Mutable/Source/MutableTools/Public/MuT/NodeColourFromScalars.h b/Engine/Plugins/Mutable/Source/MutableTools/Public/MuT/NodeColourFromScalars.h
deleted file mode 100644
diff --git a/Engine/Plugins/Mutable/Source/MutableTools/Public/MuT/NodeColourMaterialBreak.h b/Engine/Plugins/Mutable/Source/MutableTools/Public/MuT/NodeColourMaterialBreak.h
deleted file mode 100644
diff --git a/Engine/Plugins/Mutable/Source/MutableTools/Public/MuT/NodeColourParameter.h b/Engine/Plugins/Mutable/Source/MutableTools/Public/MuT/NodeColourParameter.h
deleted file mode 100644
diff --git a/Engine/Plugins/Mutable/Source/MutableTools/Public/MuT/NodeColourSampleImage.h b/Engine/Plugins/Mutable/Source/MutableTools/Public/MuT/NodeColourSampleImage.h
deleted file mode 100644
diff --git a/Engine/Plugins/Mutable/Source/MutableTools/Public/MuT/NodeColourSwitch.h b/Engine/Plugins/Mutable/Source/MutableTools/Public/MuT/NodeColourSwitch.h
deleted file mode 100644
diff --git a/Engine/Plugins/Mutable/Source/MutableTools/Public/MuT/NodeColourTable.h b/Engine/Plugins/Mutable/Source/MutableTools/Public/MuT/NodeColourTable.h
deleted file mode 100644
diff --git a/Engine/Plugins/Mutable/Source/MutableTools/Public/MuT/NodeComponentEdit.h b/Engine/Plugins/Mutable/Source/MutableTools/Public/MuT/NodeComponentEdit.h
deleted file mode 100644
diff --git a/Engine/Plugins/Mutable/Source/MutableTools/Public/MuT/NodeComponentSwitch.h b/Engine/Plugins/Mutable/Source/MutableTools/Public/MuT/NodeComponentSwitch.h
deleted file mode 100644
diff --git a/Engine/Plugins/Mutable/Source/MutableTools/Public/MuT/NodeImageBinarise.h b/Engine/Plugins/Mutable/Source/MutableTools/Public/MuT/NodeImageBinarise.h
deleted file mode 100644
diff --git a/Engine/Plugins/Mutable/Source/MutableTools/Public/MuT/NodeImageConditional.h b/Engine/Plugins/Mutable/Source/MutableTools/Public/MuT/NodeImageConditional.h
deleted file mode 100644
diff --git a/Engine/Plugins/Mutable/Source/MutableTools/Public/MuT/NodeImageConstant.h b/Engine/Plugins/Mutable/Source/MutableTools/Public/MuT/NodeImageConstant.h
deleted file mode 100644
diff --git a/Engine/Plugins/Mutable/Source/MutableTools/Public/MuT/NodeImageFromMaterialParameter.h b/Engine/Plugins/Mutable/Source/MutableTools/Public/MuT/NodeImageFromMaterialParameter.h
deleted file mode 100644
diff --git a/Engine/Plugins/Mutable/Source/MutableTools/Public/MuT/NodeImageInterpolate.h b/Engine/Plugins/Mutable/Source/MutableTools/Public/MuT/NodeImageInterpolate.h
deleted file mode 100644
diff --git a/Engine/Plugins/Mutable/Source/MutableTools/Public/MuT/NodeImageInvert.h b/Engine/Plugins/Mutable/Source/MutableTools/Public/MuT/NodeImageInvert.h
deleted file mode 100644
diff --git a/Engine/Plugins/Mutable/Source/MutableTools/Public/MuT/NodeImageLayer.h b/Engine/Plugins/Mutable/Source/MutableTools/Public/MuT/NodeImageLayer.h
deleted file mode 100644
diff --git a/Engine/Plugins/Mutable/Source/MutableTools/Public/MuT/NodeImagePlainColour.h b/Engine/Plugins/Mutable/Source/MutableTools/Public/MuT/NodeImagePlainColour.h
deleted file mode 100644
diff --git a/Engine/Plugins/Mutable/Source/MutableTools/Public/MuT/NodeImageProject.h b/Engine/Plugins/Mutable/Source/MutableTools/Public/MuT/NodeImageProject.h
deleted file mode 100644
diff --git a/Engine/Plugins/Mutable/Source/MutableTools/Public/MuT/NodeImageResize.h b/Engine/Plugins/Mutable/Source/MutableTools/Public/MuT/NodeImageResize.h
deleted file mode 100644
diff --git a/Engine/Plugins/Mutable/Source/MutableTools/Public/MuT/NodeImageSaturate.h b/Engine/Plugins/Mutable/Source/MutableTools/Public/MuT/NodeImageSaturate.h
deleted file mode 100644
diff --git a/Engine/Plugins/Mutable/Source/MutableTools/Public/MuT/NodeImageSwitch.h b/Engine/Plugins/Mutable/Source/MutableTools/Public/MuT/NodeImageSwitch.h
deleted file mode 100644
diff --git a/Engine/Plugins/Mutable/Source/MutableTools/Public/MuT/NodeImageSwizzle.h b/Engine/Plugins/Mutable/Source/MutableTools/Public/MuT/NodeImageSwizzle.h
deleted file mode 100644
diff --git a/Engine/Plugins/Mutable/Source/MutableTools/Public/MuT/NodeImageTable.h b/Engine/Plugins/Mutable/Source/MutableTools/Public/MuT/NodeImageTable.h
deleted file mode 100644
diff --git a/Engine/Plugins/Mutable/Source/MutableTools/Public/MuT/NodeImageVariation.h b/Engine/Plugins/Mutable/Source/MutableTools/Public/MuT/NodeImageVariation.h
deleted file mode 100644
diff --git a/Engine/Plugins/Mutable/Source/MutableTools/Public/MuT/NodeLOD.h b/Engine/Plugins/Mutable/Source/MutableTools/Public/MuT/NodeLOD.h
deleted file mode 100644
diff --git a/Engine/Plugins/Mutable/Source/MutableTools/Public/MuT/NodeMaterialParameter.h b/Engine/Plugins/Mutable/Source/MutableTools/Public/MuT/NodeMaterialParameter.h
deleted file mode 100644
diff --git a/Engine/Plugins/Mutable/Source/MutableTools/Public/MuT/NodeMaterialSwitch.h b/Engine/Plugins/Mutable/Source/MutableTools/Public/MuT/NodeMaterialSwitch.h
deleted file mode 100644
diff --git a/Engine/Plugins/Mutable/Source/MutableTools/Public/MuT/NodeMaterialTable.h b/Engine/Plugins/Mutable/Source/MutableTools/Public/MuT/NodeMaterialTable.h
deleted file mode 100644
diff --git a/Engine/Plugins/Mutable/Source/MutableTools/Public/MuT/NodeMaterialVariation.h b/Engine/Plugins/Mutable/Source/MutableTools/Public/MuT/NodeMaterialVariation.h
deleted file mode 100644
diff --git a/Engine/Plugins/Mutable/Source/MutableTools/Public/MuT/NodeMatrixParameter.h b/Engine/Plugins/Mutable/Source/MutableTools/Public/MuT/NodeMatrixParameter.h
deleted file mode 100644
diff --git a/Engine/Plugins/Mutable/Source/MutableTools/Public/MuT/NodeMeshClipDeform.h b/Engine/Plugins/Mutable/Source/MutableTools/Public/MuT/NodeMeshClipDeform.h
deleted file mode 100644
diff --git a/Engine/Plugins/Mutable/Source/MutableTools/Public/MuT/NodeMeshClipMorphPlane.h b/Engine/Plugins/Mutable/Source/MutableTools/Public/MuT/NodeMeshClipMorphPlane.h
deleted file mode 100644
diff --git a/Engine/Plugins/Mutable/Source/MutableTools/Public/MuT/NodeMeshFragment.h b/Engine/Plugins/Mutable/Source/MutableTools/Public/MuT/NodeMeshFragment.h
deleted file mode 100644
diff --git a/Engine/Plugins/Mutable/Source/MutableTools/Public/MuT/NodeMeshMakeMorph.h b/Engine/Plugins/Mutable/Source/MutableTools/Public/MuT/NodeMeshMakeMorph.h
deleted file mode 100644
diff --git a/Engine/Plugins/Mutable/Source/MutableTools/Public/MuT/NodeMeshReshape.h b/Engine/Plugins/Mutable/Source/MutableTools/Public/MuT/NodeMeshReshape.h
deleted file mode 100644
diff --git a/Engine/Plugins/Mutable/Source/MutableTools/Public/MuT/NodeMeshSwitch.h b/Engine/Plugins/Mutable/Source/MutableTools/Public/MuT/NodeMeshSwitch.h
deleted file mode 100644
diff --git a/Engine/Plugins/Mutable/Source/MutableTools/Public/MuT/NodeModifier.h b/Engine/Plugins/Mutable/Source/MutableTools/Public/MuT/NodeModifier.h
deleted file mode 100644
diff --git a/Engine/Plugins/Mutable/Source/MutableTools/Public/MuT/NodeModifierMeshClipDeform.h b/Engine/Plugins/Mutable/Source/MutableTools/Public/MuT/NodeModifierMeshClipDeform.h
deleted file mode 100644
diff --git a/Engine/Plugins/Mutable/Source/MutableTools/Public/MuT/NodeModifierMeshClipWithMesh.h b/Engine/Plugins/Mutable/Source/MutableTools/Public/MuT/NodeModifierMeshClipWithMesh.h
deleted file mode 100644
diff --git a/Engine/Plugins/Mutable/Source/MutableTools/Public/MuT/NodeObject.h b/Engine/Plugins/Mutable/Source/MutableTools/Public/MuT/NodeObject.h
deleted file mode 100644
diff --git a/Engine/Plugins/Mutable/Source/MutableTools/Public/MuT/NodeObjectGroup.h b/Engine/Plugins/Mutable/Source/MutableTools/Public/MuT/NodeObjectGroup.h
deleted file mode 100644
diff --git a/Engine/Plugins/Mutable/Source/MutableTools/Public/MuT/NodeProjector.h b/Engine/Plugins/Mutable/Source/MutableTools/Public/MuT/NodeProjector.h
deleted file mode 100644
diff --git a/Engine/Plugins/Mutable/Source/MutableTools/Public/MuT/NodeScalarConstant.h b/Engine/Plugins/Mutable/Source/MutableTools/Public/MuT/NodeScalarConstant.h
deleted file mode 100644
diff --git a/Engine/Plugins/Mutable/Source/MutableTools/Public/MuT/NodeScalarEnumParameter.h b/Engine/Plugins/Mutable/Source/MutableTools/Public/MuT/NodeScalarEnumParameter.h
deleted file mode 100644
diff --git a/Engine/Plugins/Mutable/Source/MutableTools/Public/MuT/NodeScalarSwitch.h b/Engine/Plugins/Mutable/Source/MutableTools/Public/MuT/NodeScalarSwitch.h
deleted file mode 100644
diff --git a/Engine/Plugins/Mutable/Source/MutableTools/Public/MuT/NodeScalarTable.h b/Engine/Plugins/Mutable/Source/MutableTools/Public/MuT/NodeScalarTable.h
deleted file mode 100644
diff --git a/Engine/Plugins/Mutable/Source/MutableTools/Public/MuT/NodeScalarVariation.h b/Engine/Plugins/Mutable/Source/MutableTools/Public/MuT/NodeScalarVariation.h
deleted file mode 100644
diff --git a/Engine/Plugins/Mutable/Source/MutableTools/Public/MuT/NodeString.h b/Engine/Plugins/Mutable/Source/MutableTools/Public/MuT/NodeString.h
deleted file mode 100644
diff --git a/Engine/Plugins/Mutable/Source/MutableTools/Public/MuT/NodeSurface.h b/Engine/Plugins/Mutable/Source/MutableTools/Public/MuT/NodeSurface.h
deleted file mode 100644
diff --git a/Engine/Plugins/Mutable/Source/MutableTools/Public/MuT/NodeSurfaceNew.h b/Engine/Plugins/Mutable/Source/MutableTools/Public/MuT/NodeSurfaceNew.h
deleted file mode 100644
diff --git a/Engine/Plugins/Mutable/Source/MutableTools/Public/MuT/Table.h b/Engine/Plugins/Mutable/Source/MutableTools/Public/MuT/Table.h
deleted file mode 100644
diff --git a/Engine/Plugins/Mutable/Source/MutableTools/Public/MuT/TypeInfo.h b/Engine/Plugins/Mutable/Source/MutableTools/Public/MuT/TypeInfo.h
deleted file mode 100644
diff --git a/Engine/Plugins/Mutable/Source/MutableTools/Public/MuT/UnrealPixelFormatOverride.h b/Engine/Plugins/Mutable/Source/MutableTools/Public/MuT/UnrealPixelFormatOverride.h
deleted file mode 100644
