[Mutable] Moved missing files MuT and MuR from /Public to /Internal

diff --git a/Engine/Plugins/Mutable/Source/MutableRuntime/Internal/MuR/ExtensionData.h b/Engine/Plugins/Mutable/Source/MutableRuntime/Internal/MuR/ExtensionData.h
new file mode 100644
--- /dev/null
+++ b/Engine/Plugins/Mutable/Source/MutableRuntime/Internal/MuR/ExtensionData.h
@@ -0,0 +1,63 @@
+// Copyright Epic Games, Inc. All Rights Reserved.
+
+#pragma once
+
+#include "MuR/Ptr.h"
+#include "MuR/RefCounted.h"
+#include "MuR/Serialisation.h"
+#include "MuR/Types.h"
+
+#define UE_API MUTABLERUNTIME_API
+
+namespace UE::Mutable::Private
+{
+	/** ExtensionData represents data types that Mutable doesn't support natively.
+	* Extensions can provide data, and functionality to operate on that data, without Mutable
+	* needing to know what the data refers to.
+	*/
+	class FExtensionData : public FResource
+	{
+	public:
+		static UE_API void Serialise(const FExtensionData* Data, FOutputArchive& Archive);
+		static UE_API TSharedPtr<FExtensionData> StaticUnserialise(FInputArchive& Archive);
+
+		// Resource interface
+		UE_API int32 GetDataSize() const override;
+
+		//! A stable hash of the contents
+		UE_API uint32 Hash() const;
+
+		UE_API void Serialise(FOutputArchive& Archive) const;
+		UE_API void Unserialise(FInputArchive& Archive);
+
+		inline bool operator==(const FExtensionData& Other) const
+		{
+			const bool bResult =
+				Other.Index == Index
+				&& Other.Origin == Origin;
+
+			return bResult;
+		}
+
+		enum class EOrigin : uint8
+		{
+			//! An invalid value used to indicate that this ExtensionData hasn't been initialized
+			Invalid,
+
+			//! This ExtensionData is a compile-time constant that's always loaded into memory
+			ConstantAlwaysLoaded,
+
+			//! This ExtensionData is a compile-time constant that's streamed in from disk when needed
+			ConstantStreamed,
+
+			//! This ExtensionData was generated at runtime
+			Runtime
+		};
+
+		int16 Index = -1;
+
+		EOrigin Origin = EOrigin::Invalid;
+	};
+}
+
+#undef UE_API
diff --git a/Engine/Plugins/Mutable/Source/MutableRuntime/Internal/MuR/Instance.h b/Engine/Plugins/Mutable/Source/MutableRuntime/Internal/MuR/Instance.h
new file mode 100644
--- /dev/null
+++ b/Engine/Plugins/Mutable/Source/MutableRuntime/Internal/MuR/Instance.h
@@ -0,0 +1,297 @@
+// Copyright Epic Games, Inc. All Rights Reserved.
+
+#pragma once
+
+#include "MuR/ResourceID.h"
+#include "HAL/PlatformMath.h"
+#include "MuR/ExtensionData.h"
+#include "MuR/Ptr.h"
+#include "MuR/RefCounted.h"
+#include "MuR/Types.h"
+#include "MuR/Material.h"
+
+#define UE_API MUTABLERUNTIME_API
+
+
+namespace UE::Mutable::Private
+{
+	struct FInstanceSurface
+	{
+		FName Name;
+		uint32 InternalId = 0;
+		uint32 ExternalId = 0;
+		uint32 SharedId = 0;
+
+		FMaterialId MaterialId;
+
+		struct FInstanceImage
+		{
+			FImageId Id;
+			FName Name;
+		};
+
+		TArray<FInstanceImage> Images;
+
+		struct FInstanceVector
+		{
+			FVector4f Value;
+			FName Name;
+		};
+
+		TArray<FInstanceVector> Vectors;
+
+		struct FInstanceScalar
+		{
+			float Value;
+			FName Name;
+		};
+
+		TArray<FInstanceScalar> Scalars;
+
+		struct FInstanceString
+		{
+			FString Value;
+			FName Name;
+		};
+
+		TArray<FInstanceString> Strings;
+	};
+
+
+	struct FInstanceLOD
+	{
+		FMeshId MeshId;
+		FName MeshName;
+
+		// The order must match the meshes surfaces
+		TArray<FInstanceSurface, TInlineAllocator<1>> Surfaces;
+	};
+
+
+	struct FInstanceComponent
+	{
+		uint16 Id;
+
+		FMaterialId OverlayMaterialId;
+
+		TArray<FInstanceLOD, TInlineAllocator<1>> LODs;
+	};
+
+	
+	struct NamedExtensionData
+	{
+		TSharedPtr<const FExtensionData> Data;
+		FName Name;
+	};
+
+	
+    /** A customised object created from a model and a set of parameter values.
+    * It corresponds to an "engine object" but the contents of its data depends on the Model, and
+    * it may contain any number of LODs, components, surfaces, meshes and images, even none.
+	*/
+	class FInstance : public FResource
+	{
+	public:
+
+		//-----------------------------------------------------------------------------------------
+		// Life cycle
+		//-----------------------------------------------------------------------------------------
+		
+        //! Clone this instance
+		UE_API TSharedPtr<FInstance> Clone() const;
+
+		// Resource interface
+		UE_API int32 GetDataSize() const override;
+
+		//-----------------------------------------------------------------------------------------
+		// Own interface
+		//-----------------------------------------------------------------------------------------
+
+        //! Type for the instance unique identifiers.
+        typedef uint32 FID;
+		
+		//! Get the number of components of this instance.
+		UE_API int32 GetComponentCount() const;
+
+		//! Get the Id of a component
+		//! \param ComponentIndex Index of the component, from 0 to GetComponentCount()-1
+		UE_API uint16 GetComponentId(int32 ComponentIndex) const;
+
+        //! Get the number of LODs in a component
+        //! \param ComponentIndex Index of the component, from 0 to GetComponentCount()-1
+        UE_API int32 GetLODCount( int32 ComponentIndex ) const;
+
+        //! Get the number of surfaces in a component
+        //! \param ComponentIndex Index of the component, from 0 to GetComponentCount()-1
+		//! \param LODIndex Index of the level of detail, from 0 to GetLODCount()-1
+		UE_API int32 GetSurfaceCount( int32 ComponentIndex, int32 LODIndex) const;
+
+        //! Get an id that can be used to match the surface data with the mesh surface data.
+        //! \param ComponentIndex Index of the component, from 0 to GetComponentCount()-1
+		//! \param LODIndex Index of the level of detail, from 0 to GetLODCount()-1
+		//! \param SurfaceIndex Index of the surface, from 0 to GetSurfaceCount()-1
+        UE_API uint32 GetSurfaceId( int32 ComponentIndex, int32 LODIndex, int32 SurfaceIndex) const;
+
+        //! Find a surface index from the internal id (as returned by GetSurfaceId).
+        //! \param ComponentIndex Index of the component, from 0 to GetComponentCount()-1
+		//! \param LODIndex Index of the level of detail, from 0 to GetLODCount()-1
+		//! \param id ID of the surface to look for.
+        UE_API int32 FindSurfaceById( int32 ComponentIndex, int32 LODIndex, uint32 id ) const;
+
+		//! Find the base surface index and Lod index when reusing surfaces between LODs. Return the surface index
+		//! and the LOD it belongs to.
+		//! \param ComponentIndex - Index of the component, from 0 to GetComponentCount()-1
+		//! \param SharedSurfaceId - Id of the surface to look for (as returned by GetSharedSurfaceId).
+		//! \param OutSurfaceIndex - Index of the surface in the OutLODIndex lod. 
+		//! \param OutLODIndex - Index of the first LOD where the surface can be found. 
+		UE_API void FindBaseSurfaceBySharedId(int32 ComponentIndex, uint32 SharedId, int32& OutSurfaceIndex, int32& OutLODIndex) const;
+
+		//! Get an id that can be used to find the same surface on other LODs
+		//! \param ComponentIndex Index of the component, from 0 to GetComponentCount()-1
+		//! \param lod Index of the level of detail, from 0 to GetLODCount()-1
+		//! \param SurfaceIndex Index of the surface, from 0 to GetSurfaceCount()-1
+		UE_API uint32 GetSharedSurfaceId(int32 ComponentIndex, int32 LODIndex, int32 SurfaceIndex) const;
+
+        //! Get an optional, opaque application-defined identifier for this surface. The meaning of
+        //! this ID depends on each application, and it is specified when creating the source data
+        //! that generates this surface.
+        //! See NodeSurfaceNew::SetCustomID.
+		//! \param ComponentIndex Index of the component, from 0 to GetComponentCount()-1
+		//! \param LODIndex Index of the level of detail, from 0 to GetLODCount()-1
+		//! \param SurfaceIndex Index of the surface, from 0 to GetSurfaceCount()-1
+        UE_API uint32 GetSurfaceCustomId( int32 ComponentIndex, int32 LODIndex, int32 SurfaceIndex) const;
+
+        //! Get the mesh resource id from a component
+        //! \param ComponentIndex Index of the component, from 0 to GetComponentCount()-1
+		//! \param LODIndex Index of the level of detail, from 0 to GetLODCount()-1
+  		UE_API FMeshId GetMeshId( int32 ComponentIndex, int32 LODIndex ) const;
+
+		//! Get the number of images in a component
+		//! \param ComponentIndex Index of the component, from 0 to GetComponentCount()-1
+		//! \param LODIndex Index of the level of detail, from 0 to GetLODCount()-1
+		//! \param SurfaceIndex Index of the surface, from 0 to GetSurfaceCount()-1
+        UE_API int32 GetImageCount( int32 ComponentIndex, int32 LODIndex, int32 SurfaceIndex) const;
+
+        //! Get an image resource id from a component
+		//! \param ComponentIndex Index of the component, from 0 to GetComponentCount()-1
+		//! \param LODIndex Index of the level of detail, from 0 to GetLODCount()-1
+		//! \param SurfaceIndex Index of the surface, from 0 to GetSurfaceCount()-1
+        //! \param ImageIndex Index of the image, from 0 to GetImageCount()-1
+ 		UE_API FImageId GetImageId( int32 ComponentIndex, int32 LODIndex, int32 SurfaceIndex, int32 ImageIndex) const;
+
+		//! Get the name of an image in a component
+		//! \param ComponentIndex Index of the component, from 0 to GetComponentCount()-1
+		//! \param LODIndex Index of the level of detail, from 0 to GetLODCount()-1
+		//! \param SurfaceIndex Index of the surface, from 0 to GetSurfaceCount()-1
+        //! \param ImageIndex Index of the image, from 0 to GetImageCount()-1
+		UE_API FName GetImageName( int32 ComponentIndex, int32 LODIndex, int32 SurfaceIndex, int32 ImageIndex ) const;
+
+		//! Get the number of vectors in a component
+		//! \param ComponentIndex Index of the component, from 0 to GetComponentCount()-1
+		//! \param LODIndex Index of the level of detail, from 0 to GetLODCount()-1
+		//! \param SurfaceIndex Index of the surface, from 0 to GetSurfaceCount()-1
+        UE_API int32 GetVectorCount( int32 ComponentIndex, int32 LODIndex, int32 SurfaceIndex) const;
+
+		//! Get a vector from a component
+		//! \param ComponentIndex Index of the component, from 0 to GetComponentCount()-1
+		//! \param LODIndex Index of the level of detail, from 0 to GetLODCount()-1
+		//! \param SurfaceIndex Index of the surface, from 0 to GetSurfaceCount()-1
+        //! \param VectorIndex Index of the vector, from 0 to GetVectorCount()-1
+        UE_API FVector4f GetVector( int32 ComponentIndex, int32 LODIndex, int32 SurfaceIndex, int32 VectorIndex) const;
+
+		//! Get the name of a vector in a component
+		//! \param ComponentIndex Index of the component, from 0 to GetComponentCount()-1
+		//! \param LODIndex Index of the level of detail, from 0 to GetLODCount()-1
+		//! \param SurfaceIndex Index of the surface, from 0 to GetSurfaceCount()-1
+        //! \param VectorIndex Index of the vector, from 0 to GetVectorCount()-1
+		UE_API FName GetVectorName( int32 ComponentIndex, int32 LODIndex, int32 SurfaceIndex, int32 VectorIndex ) const;
+
+        //! Get the number of scalar values in a component
+        //! \param ComponentIndex Index of the component, from 0 to GetComponentCount()-1
+		//! \param LODIndex Index of the level of detail, from 0 to GetLODCount()-1
+		//! \param SurfaceIndex Index of the surface, from 0 to GetSurfaceCount()-1
+        UE_API int32 GetScalarCount( int32 ComponentIndex, int32 LODIndex, int32 SurfaceIndex) const;
+
+        //! Get a scalar value from a component
+        //! \param ComponentIndex Index of the component, from 0 to GetComponentCount()-1
+		//! \param LODIndex Index of the level of detail, from 0 to GetLODCount()-1
+		//! \param SurfaceIndex Index of the surface, from 0 to GetSurfaceCount()-1
+        //! \param ScalarIndex Index of the scalar, from 0 to GetScalarCount()-1
+        UE_API float GetScalar( int32 ComponentIndex, int32 LODIndex, int32 SurfaceIndex, int32 ScalarIndex) const;
+
+        //! Get the name of a scalar from a component
+        //! \param ComponentIndex Index of the component, from 0 to GetComponentCount()-1
+		//! \param LODIndex Index of the level of detail, from 0 to GetLODCount()-1
+		//! \param SurfaceIndex Index of the surface, from 0 to GetSurfaceCount()-1
+        //! \param ScalarIndex Index of the scalar, from 0 to GetScalarCount()-1
+		UE_API FName GetScalarName( int32 ComponentIndex, int32 LODIndex, int32 SurfaceIndex, int32 ScalarIndex ) const;
+
+        //! Get the number of string values in a component
+        //! \param ComponentIndex Index of the component, from 0 to GetComponentCount()-1
+		//! \param LODIndex Index of the level of detail, from 0 to GetLODCount()-1
+		//! \param SurfaceIndex Index of the surface, from 0 to GetSurfaceCount()-1
+        UE_API int32 GetStringCount( int32 ComponentIndex, int32 LODIndex, int32 SurfaceIndex) const;
+
+        //! Get a string value from a component
+        //! \param ComponentIndex Index of the component, from 0 to GetComponentCount()-1
+		//! \param LODIndex Index of the level of detail, from 0 to GetLODCount()-1
+		//! \param SurfaceIndex Index of the surface, from 0 to GetSurfaceCount()-1
+        //! \param StringIndex Index of the string, from 0 to GetStringCount()-1
+        UE_API FString GetString( int32 ComponentIndex, int32 LODIndex, int32 SurfaceIndex, int32 StringIndex) const;
+
+        //! Get the name of a string from a component
+        //! \param ComponentIndex Index of the component, from 0 to GetComponentCount()-1
+		//! \param LODIndex Index of the level of detail, from 0 to GetLODCount()-1
+		//! \param SurfaceIndex Index of the surface, from 0 to GetSurfaceCount()-1
+        //! \param StringIndex Index of the string, from 0 to GetStringCount()-1
+        UE_API FName GetStringName( int32 ComponentIndex, int32 LODIndex, int32 SurfaceIndex, int32 StringIndex ) const;
+
+		//! Get the index of the overlay material.
+		//! \param ComponentIndex Index of the component, from 0 to GetComponentCount()-1
+		UE_API FMaterialId GetOverlayMaterialId(int32 ComponentIndex) const;
+
+		UE_API FMaterialId GetMaterialId(int32 ComponentIndex, int32 LODIndex, int32 SurfaceIndex) const;
+
+		//! Get the number of ExtensionData values in a component
+		UE_API int32 GetExtensionDataCount() const;
+
+		//! Get an ExtensionData value from a component
+		//! \param Index Index of the ExtensionData to fetch
+		//! \param OutExtensionData Receives the ExtensionData
+		//! \param OutName Receives the name associated with the ExtensionData. Guaranteed to be a valid string of non-zero length.
+		UE_API void GetExtensionData(int32 Index, TSharedPtr<const FExtensionData>& OutExtensionData, FName& OutName) const;
+		
+		int32 AddComponent();
+		
+		int32 AddLOD(int32 ComponentIndex);
+
+		void SetMesh(int32 ComponentIndex, int32 LODIndex, const FMeshId& MeshId, FName Name);
+
+		int32 AddSurface(int32 ComponentIndex, int32 LODIndex);
+
+		void SetSurfaceName(int32 ComponentIndex, int32 LODIndex, int32 SurfaceIndex, FName Name);
+
+		int32 AddImage(int32 ComponentIndex, int32 LODIndex, int32 SurfaceIndex, const FImageId& Imag, FName Name);
+
+		int32 AddVector(int32 ComponentIndex, int32 LODIndex, int32 SurfaceIndex, const FVector4f&, FName Name);
+
+		int32 AddScalar(int32 ComponentIndex, int32 LODIndex, int32 SurfaceIndex, float, FName Name);
+
+		int32 AddString(int32 ComponentIndex, int32 LODIndex, int32 SurfaceIndex, const FString& Value, FName Name);
+
+		void SetMaterialId(int32 ComponentIndex, int32 LODIndex, int32 SurfaceIndex, const FMaterialId& MaterialId);
+
+		void SetOverlayMaterialId(int32 ComponentIndex, const FMaterialId& MaterialId);
+
+		void SetExtensionData(const TSharedRef<const FExtensionData>& Data, FName Name);
+		
+		TArray<FInstanceComponent, TInlineAllocator<1>> Components;
+
+		// Every entry must have a valid ExtensionData and name
+		TArray<NamedExtensionData> ExtensionData;
+	};
+
+
+}
+
+#undef UE_API
diff --git a/Engine/Plugins/Mutable/Source/MutableRuntime/Internal/MuR/Skeleton.h b/Engine/Plugins/Mutable/Source/MutableRuntime/Internal/MuR/Skeleton.h
new file mode 100644
--- /dev/null
+++ b/Engine/Plugins/Mutable/Source/MutableRuntime/Internal/MuR/Skeleton.h
@@ -0,0 +1,120 @@
+// Copyright Epic Games, Inc. All Rights Reserved.
+
+#pragma once
+
+#include "MuR/Serialisation.h"
+#include "Containers/Array.h"
+#include "HAL/PlatformMath.h"
+#include "Math/Transform.h"
+#include "Math/VectorRegister.h"
+#include "Misc/AssertionMacros.h"
+#include "MuR/Ptr.h"
+#include "MuR/RefCounted.h"
+
+#define UE_API MUTABLERUNTIME_API
+
+
+namespace UE::Mutable::Private
+{	
+	// Bone name identifier
+	struct FBoneName
+	{
+		FBoneName() {};
+		FBoneName(uint32 InID) : Id(InID) {};
+
+		// Hash built from the bone name (FString)
+		uint32 Id = 0;
+
+		inline void Serialise(FOutputArchive& Arch) const
+		{
+			Arch << Id;
+		}
+
+		inline void Unserialise(FInputArchive& Arch)
+		{
+			Arch >> Id;
+		}
+
+		//!
+		inline bool operator==(const FBoneName& Other) const
+		{
+			return Id == Other.Id;
+		}
+	};
+
+	inline uint32 GetTypeHash(const FBoneName& Bone)
+	{
+		return Bone.Id;
+	}
+
+
+    /** Skeleton object. */
+	class FSkeleton
+	{
+	public:
+
+		//-----------------------------------------------------------------------------------------
+		// Life cycle
+		//-----------------------------------------------------------------------------------------
+
+        //! Deep clone this skeleton.
+        UE_API TSharedPtr<FSkeleton> Clone() const;
+
+		//! Serialisation
+        static UE_API void Serialise( const FSkeleton*, FOutputArchive& );
+        static UE_API TSharedPtr<FSkeleton> StaticUnserialise( FInputArchive& );
+
+
+		//-----------------------------------------------------------------------------------------
+		// Own interface
+		//-----------------------------------------------------------------------------------------
+
+        //! @return - Number of bones in the Skeleton
+		UE_API int32 GetBoneCount() const;
+        UE_API void SetBoneCount(int32 Count);
+
+		//! @return - FName of the bone at 'Index'. Only valid in the editor
+		UE_API const FName GetDebugName(int32 Index) const;
+		UE_API void SetDebugName(const int32 Index, const FName BoneName);
+
+        //! Get and set the parent bone of each bone. The parent can be -1 if the bone is a root.
+        UE_API int32 GetBoneParent(int32 BoneIndex) const;
+        UE_API void SetBoneParent(int32 BoneIndex, int32 ParentBoneIndex);
+
+		//! @return - BoneName of the Bone at 'Index'.
+		UE_API const FBoneName& GetBoneName(int32 Index) const;
+		UE_API void SetBoneName(int32 Index, const FBoneName& BoneName);
+
+		//! @return - Index in the Skeleton. INDEX_NONE if not found.
+		UE_API int32 FindBone(const FBoneName& BoneName) const;
+		
+
+	public:
+
+		//! DEBUG. FNames of the bones. Only valid in the editor. Do not serialize.
+		TArray<FName> DebugBoneNames;
+
+		//! Array of bone identifiers. 
+		TArray<FBoneName> BoneIds;
+
+		//! For each bone, index of the parent bone in the bone vectors. -1 means no parent.
+		//! This array must have the same size than the m_bones array.
+		TArray<int16> BoneParents;
+
+		//!
+		UE_API void Serialise(FOutputArchive&) const;
+
+		//!
+		UE_API void Unserialise(FInputArchive&);
+
+		//!
+		inline bool operator==(const FSkeleton& Other) const
+		{
+			return BoneIds == Other.BoneIds
+				&& BoneParents == Other.BoneParents;
+		}
+	};
+
+}
+
+#undef UE_API
diff --git a/Engine/Plugins/Mutable/Source/MutableRuntime/Internal/MuR/System.h b/Engine/Plugins/Mutable/Source/MutableRuntime/Internal/MuR/System.h
new file mode 100644
--- /dev/null
+++ b/Engine/Plugins/Mutable/Source/MutableRuntime/Internal/MuR/System.h
@@ -0,0 +1,246 @@
+// Copyright Epic Games, Inc. All Rights Reserved.
+
+#pragma once
+
+#include "MuR/Image.h"
+#include "MuR/Mesh.h"
+#include "MuR/Instance.h"
+#include "MuR/RefCounted.h"
+#include "MuR/Settings.h"
+#include "MuR/Types.h"
+#include "Tasks/Task.h"
+#include "Templates/SharedPointer.h"
+#include "Templates/Tuple.h"
+#include "HAL/Platform.h"
+
+#include "System.generated.h"
+
+class UTexture;
+class USkeletalMesh;
+
+#define UE_API MUTABLERUNTIME_API
+
+
+/** If set to 1, this enables some expensive Unreal Insights traces, but can lead to 5x slower mutable operation.
+* Other cheaper traces are enabled at all times.
+*/
+#define UE_MUTABLE_ENABLE_SLOW_TRACES	0
+
+
+/** Despite being an UEnum, this is not always version-serialized (in MutableTools).
+ * Beware of changing the enum options or order.
+ */
+UENUM()
+enum class ETextureCompressionStrategy : uint8
+{
+	/** Don't change the generated format. */
+	None,
+
+	/** If a texture depends on run-time parameters for an object state, don't compress. */
+	DontCompressRuntime,
+
+	/** Never compress the textures for this state. */
+	NeverCompress
+};
+
+MUTABLE_DEFINE_ENUM_SERIALISABLE(ETextureCompressionStrategy);
+
+
+namespace UE::Mutable::Private
+{
+	// Forward references
+	class FModel;
+	class FModelReader;
+	class FParameters;
+    class FMesh;
+	class FExtensionDataStreamer;
+
+
+	/** */
+	enum class EExecutionStrategy : uint8
+	{
+		/** Undefined. */
+		None = 0,
+
+		/** Always try to run operations that reduce working memory first. */
+		MinimizeMemory,
+
+		/** Always try to run operations that unlock more operations first. */
+		MaximizeConcurrency,
+
+		/** Utility value with the number of error types. */
+		Count
+	};
+
+
+    /** Interface to request external images used as parameters. */
+    class FExternalResourceProvider
+    {
+    public:
+        //! Ensure virtual destruction
+        virtual ~FExternalResourceProvider() = default;
+
+        /** Returns the completion event and a cleanup function that must be called once event is completed. */
+		virtual TTuple<UE::Tasks::FTask, TFunction<void()>> GetImageAsync(UTexture* Texture, uint8 MipmapsToSkip, TFunction<void(TSharedPtr<FImage>)>& ResultCallback) = 0;
+		virtual TTuple<UE::Tasks::FTask, TFunction<void()>> GetReferencedImageAsync(int32 Id, uint8 MipmapsToSkip, TFunction<void(TSharedPtr<FImage>)>& ResultCallback) { check(false); return {}; }
+
+		virtual FExtendedImageDesc GetImageDesc(UTexture* Texture) = 0;
+
+		/** Returns the completion event and a cleanup function that must be called once event is completed. */
+		virtual TTuple<UE::Tasks::FTask, TFunction<void()>> GetMeshAsync(USkeletalMesh* SkeleltalMesh, int32 InLODIndex, int32 InSectionIndex, TFunction<void(TSharedPtr<FMesh>)>& ResultCallback) = 0;
+    };
+
+
+    /** Main system class to load models and build instances. */
+	class FSystem
+	{
+    public:
+
+        //! This constant can be used in place of the lodMask in methods like BeginUpdate
+        static constexpr uint32 AllLODs = 0xffffffff;
+
+    public:
+
+		//! Constructor of a system object to build data.
+        //! \param Settings Optional class with the settings to use in this system. The default
+        //! value configures a production-ready system.
+		UE_API FSystem( const FSettings& Settings = FSettings());
+
+        //! Set a new provider for model data. 
+        UE_API void SetStreamingInterface(const TSharedPtr<FModelReader>& );
+
+        /** Set the working memory limit, overrding any set in the settings when the system was created.
+         * Refer to Settings::SetWorkingMemoryBudget for more information.
+		 */
+        UE_API void SetWorkingMemoryBytes( uint64 Bytes );
+
+        /** Removes all the possible working memory regardless of the budget set. This may make following 
+		* operations take longer.
+		*/
+		UE_API void ClearWorkingMemory();
+		
+		/** Set a function that will be used to convert image pixel formats instead of the internal conversion. 
+		* \warning The provided function can be called from any thread, and also concurrently.
+		* If this function fails (returns false in the first parameter) the internal function is attempted next.
+		* This is useful to provide higher-quality external compressors in editor or when cooking.
+		*/
+		UE_API void SetImagePixelConversionOverride(const FImageOperator::FImagePixelFormatFunc&);
+
+        //! Create a new instance from the given model. The instance can then be configured through
+        //! calls to BeginUpdate/EndUpdate.
+        //! A call to NewInstance must be paired to a call to ReleasesInstance when the instance is
+        //! no longer needed.
+        //! \return An identifier that is always bigger than 0.
+        UE_API FInstance::FID NewInstance(const TSharedPtr<const FModel>& Model, const TSharedPtr<FExternalResourceProvider>& ExternalResourceProvider);
+
+		UE_API void ReuseInstance(FInstance::FID InstanceId, const TSharedPtr<FExternalResourceProvider>& ExternalResourceProvider);
+		
+        //! \brief Update an instance with a new parameter set and/or state.
+        //!
+        //! \warning a call to BeginUpdate must be paired with a call to EndUpdate once the returned
+        //! data has been processed.
+        //! \param InstanceID The id of the instance to update, as created by a NewInstance call.
+        //! \param Params The parameters that customise this instance.
+        //! \param StateIndex The index of the state this instance will be set to. The states range
+        //! from 0 to Model::GetStateCount-1
+        //! \param LodMask Bitmask selecting the levels of detail to build (i-th bit selects i-th lod).
+        //! \return the instance data with all the LOD, components, and ids to generate the meshes 
+        //! and images. The returned Instance is only valid until the next call to EndUpdate with 
+        //! the same instanceID parameter.
+		//!
+		//! Very fast. Can be called in the Game Thread.
+		UE_API TSharedPtr<const FInstance> BeginUpdate_GameThread(FInstance::FID InstanceID,
+			const TSharedPtr<const FParameters>& Params,
+			const TSharedPtr<FMeshIdRegistry>& InMeshIdRegistry,
+			const TSharedPtr<FImageIdRegistry>& InImageIdRegistry,
+			const TSharedPtr<FMaterialIdRegistry>& InMaterialIdRegistry,
+			int32 StateIndex,
+			uint32 LodMask);
+
+		/** Potentially expensive. Must be called in the Mutable Thread. */
+		UE_API TSharedPtr<const FInstance> BeginUpdate_MutableThread(FInstance::FID InInstanceID,
+			const TSharedPtr<const FParameters>& InParams,
+			const TSharedPtr<FMeshIdRegistry>& InMeshIdRegistry,
+			const TSharedPtr<FImageIdRegistry>& InImageIdRegistry,
+			const TSharedPtr<FMaterialIdRegistry>& InMaterialIdRegistry,
+			int32 InStateIndex,
+			uint32 InLodMask);
+		
+		//! Only valid between BeginUpdate and EndUpdate
+		//! Calculate the description of an image, without generating it.
+		UE_API UE::Tasks::TTask<FExtendedImageDesc> GetImageDesc(FInstance::FID InstanceID, const FImageId& ImageId);
+
+		//! Only valid between BeginUpdate and EndUpdate
+		//! \param MipsToSkip Number of mips to skip compared from the full image.
+		//! If 0, all mip levels will be generated. If more levels than possible to discard are specified, 
+		//! the image will still contain a minimum number of mips specified at model compile time.
+		UE_API UE::Tasks::TTask<TSharedPtr<const FImage>> GetImage(FInstance::FID InstanceID, const FImageId& ImageId, int32 MipsToSkip = 0, int32 LOD = 0);
+
+        //! Only valid between BeginUpdate and EndUpdate
+		UE_API UE::Tasks::TTask<TSharedPtr<const FMesh>> GetMesh(FInstance::FID InstanceID, const FMeshId& MeshId, EMeshContentFlags MeshContentFilter = EMeshContentFlags::AllFlags);
+
+		//! Only valid between BeginUpdate and EndUpdate
+		//! Calculate the description of an image, without generating it.
+		UE_API FExtendedImageDesc GetImageDescInline(FInstance::FID InstanceID, const FImageId& ImageId);
+
+		//! Only valid between BeginUpdate and EndUpdate
+		//! \param MipsToSkip Number of mips to skip compared from the full image.
+		//! If 0, all mip levels will be generated. If more levels than possible to discard are specified, 
+		//! the image will still contain a minimum number of mips specified at model compile time.
+		UE_API TSharedPtr<const FImage> GetImageInline(FInstance::FID InstanceID, const FImageId& ImageId, int32 MipsToSkip = 0, int32 LOD = 0);
+
+        //! Only valid between BeginUpdate and EndUpdate
+		UE_API TSharedPtr<const FMesh> GetMeshInline(FInstance::FID InstanceID, const FImageId& MeshId, EMeshContentFlags MeshContentFilter = EMeshContentFlags::AllFlags);
+		
+		UE_API UE::Tasks::TTask<TSharedPtr<const FMaterial>> GetMaterial(FInstance::FID InstanceID, const FMaterialId& MaterialId);
+		
+        //! Invalidate and free the last Instance data returned by a call to BeginUpdate with
+        //! the same instance index. After a call to this method, that Instance cannot be used any
+        //! more and its content is undefined.
+        //! \param instance The index of the instance whose last data will be invalidated.
+        UE_API void EndUpdate(FInstance::FID InstanceID);
+
+        //! Completely destroy an instance. After a call to this method the given instance cannot be
+        //! updated any more, and its resources may have been freed.
+        //! \param instance The id of the instance to destroy.
+        UE_API void ReleaseInstance(FInstance::FID InstanceID);
+
+		//! Calculate the relevancy of every parameter. Some parameters may be unused depending on
+		//! the values of other parameters. This method will set to true the flags for parameters
+		//! that are relevant, and to false otherwise. This is useful to hide irrelevant parameters
+		//! in dynamic user interfaces.
+        //! \param pModel The model used to create the FParameters instance.
+        //! \param pParameters Parameter set that we want to find the relevancy of.
+        //! \param pFlags is a pointer to a preallocated array of booleans that contains at least
+		//! pParameters->GetCount() elements.
+        UE_API void GetParameterRelevancy( FInstance::FID InstanceID,
+									const TSharedPtr<const FParameters>& FParameters,
+									bool* Flags );
+
+		//-----------------------------------------------------------------------------------------
+		// Interface pattern
+		//-----------------------------------------------------------------------------------------
+		class Private;
+
+		UE_API Private* GetPrivate() const;
+
+    public:
+
+        // Prevent copy, move and assignment.
+        FSystem( const FSystem& ) = delete;
+		FSystem& operator=( const FSystem& ) = delete;
+        FSystem( FSystem&& ) = delete;
+        FSystem& operator=( FSystem&& ) = delete;
+
+		UE_API ~FSystem();
+
+	private:
+
+		Private* m_pD;
+
+	};
+
+
+}
+
+#undef UE_API
diff --git a/Engine/Plugins/Mutable/Source/MutableTools/Internal/MuT/NodeColour.h b/Engine/Plugins/Mutable/Source/MutableTools/Internal/MuT/NodeColour.h
new file mode 100644
--- /dev/null
+++ b/Engine/Plugins/Mutable/Source/MutableTools/Internal/MuT/NodeColour.h
@@ -0,0 +1,39 @@
+// Copyright Epic Games, Inc. All Rights Reserved.
+
+#pragma once
+
+#include "HAL/Platform.h"
+#include "MuR/Ptr.h"
+#include "MuR/RefCounted.h"
+#include "MuT/Node.h"
+
+#define UE_API MUTABLETOOLS_API
+
+
+namespace UE::Mutable::Private
+{
+
+    /** Base class of any node that outputs a colour.
+	*/
+    class NodeColour : public Node
+    {
+    public:
+
+        // Node Interface
+		virtual const FNodeType* GetType() const override { return &StaticType; }
+		static const FNodeType* GetStaticType() { return &StaticType; }
+
+    protected:
+
+        /** Forbidden. Manage with the Ptr<> template. */
+        inline ~NodeColour() {}
+
+	private:
+
+		static UE_API FNodeType StaticType;
+
+    };
+
+}
+
+#undef UE_API
diff --git a/Engine/Plugins/Mutable/Source/MutableTools/Internal/MuT/NodeColourArithmeticOperation.h b/Engine/Plugins/Mutable/Source/MutableTools/Internal/MuT/NodeColourArithmeticOperation.h
new file mode 100644
--- /dev/null
+++ b/Engine/Plugins/Mutable/Source/MutableTools/Internal/MuT/NodeColourArithmeticOperation.h
@@ -0,0 +1,54 @@
+// Copyright Epic Games, Inc. All Rights Reserved.
+
+#pragma once
+
+#include "MuR/Ptr.h"
+#include "MuR/RefCounted.h"
+#include "MuT/NodeColour.h"
+
+#define UE_API MUTABLETOOLS_API
+
+
+namespace UE::Mutable::Private
+{
+
+    /** Perform a per - component arithmetic operation between two colours. 
+	*/
+	class NodeColourArithmeticOperation : public NodeColour
+	{
+	public:
+
+		/** Possible arithmetic operations. */
+		enum class EOperation
+		{
+			Add,
+			Subtract,
+			Multiply,
+			Divide
+		};
+
+		EOperation Operation;
+		Ptr<NodeColour> A;
+		Ptr<NodeColour> B;
+
+	public:
+
+		/** Node type hierarchy data. */
+		virtual const FNodeType* GetType() const { return &StaticType; }
+		static const FNodeType* GetStaticType() { return &StaticType; }
+
+	protected:
+
+		/** Forbidden. Manage with the Ptr<> template. */
+		~NodeColourArithmeticOperation() {}
+
+	private:
+
+		static UE_API FNodeType StaticType;
+
+	};
+
+	
+}
+
+#undef UE_API
diff --git a/Engine/Plugins/Mutable/Source/MutableTools/Internal/MuT/NodeColourVariation.h b/Engine/Plugins/Mutable/Source/MutableTools/Internal/MuT/NodeColourVariation.h
new file mode 100644
--- /dev/null
+++ b/Engine/Plugins/Mutable/Source/MutableTools/Internal/MuT/NodeColourVariation.h
@@ -0,0 +1,51 @@
+// Copyright Epic Games, Inc. All Rights Reserved.
+
+#pragma once
+
+#include "MuR/Ptr.h"
+#include "MuR/RefCounted.h"
+#include "MuT/Node.h"
+#include "MuT/NodeColour.h"
+
+#define UE_API MUTABLETOOLS_API
+
+
+namespace UE::Mutable::Private
+{
+
+	/** Select different color subgraphs based on active tags. */
+    class NodeColourVariation : public NodeColour
+    {
+	public:
+
+		Ptr<NodeColour> DefaultColour;
+
+		struct FVariation
+		{
+			Ptr<NodeColour> Colour;
+			FString Tag;
+		};
+
+		TArray<FVariation> Variations;
+
+    public:
+
+		// Node Interface
+		virtual const FNodeType* GetType() const override { return &StaticType; }
+		static const FNodeType* GetStaticType() { return &StaticType; }
+
+    protected:
+
+        /** Forbidden. Manage with the Ptr<> template. */
+		~NodeColourVariation() {}
+
+	private:
+
+		static UE_API FNodeType StaticType;
+	
+	};
+
+
+} // namespace UE::Mutable::Private
+
+#undef UE_API
diff --git a/Engine/Plugins/Mutable/Source/MutableTools/Internal/MuT/NodeImageColourMap.h b/Engine/Plugins/Mutable/Source/MutableTools/Internal/MuT/NodeImageColourMap.h
new file mode 100644
--- /dev/null
+++ b/Engine/Plugins/Mutable/Source/MutableTools/Internal/MuT/NodeImageColourMap.h
@@ -0,0 +1,45 @@
+// Copyright Epic Games, Inc. All Rights Reserved.
+
+#pragma once
+
+#include "MuR/Ptr.h"
+#include "MuR/RefCounted.h"
+#include "MuT/Node.h"
+#include "MuT/NodeImage.h"
+
+#define UE_API MUTABLETOOLS_API
+
+namespace UE::Mutable::Private
+{
+
+	/** This node changes the colours of a selectd part of the image, applying a colour map from
+	* conteined in another image.
+	*/
+	class NodeImageColourMap : public NodeImage
+	{
+	public:
+
+		Ptr<NodeImage> Base;
+		Ptr<NodeImage> Mask;
+		Ptr<NodeImage> Map;
+
+	public:
+
+		// Node Interface
+		virtual const FNodeType* GetType() const override { return &StaticType; }
+		static const FNodeType* GetStaticType() { return &StaticType; }
+
+	protected:
+
+		/** Forbidden. Manage with the Ptr<> template. */
+		~NodeImageColourMap() {}
+
+	private:
+
+		static UE_API FNodeType StaticType;
+
+	};
+
+}
+
+#undef UE_API
diff --git a/Engine/Plugins/Mutable/Source/MutableTools/Internal/MuT/NodeImageLayerColour.h b/Engine/Plugins/Mutable/Source/MutableTools/Internal/MuT/NodeImageLayerColour.h
new file mode 100644
--- /dev/null
+++ b/Engine/Plugins/Mutable/Source/MutableTools/Internal/MuT/NodeImageLayerColour.h
@@ -0,0 +1,46 @@
+// Copyright Epic Games, Inc. All Rights Reserved.
+
+#pragma once
+
+#include "MuR/Image.h"
+#include "MuR/Ptr.h"
+#include "MuR/RefCounted.h"
+#include "MuT/Node.h"
+#include "MuT/NodeImage.h"
+#include "MuT/NodeColour.h"
+
+#define UE_API MUTABLETOOLS_API
+
+namespace UE::Mutable::Private
+{
+
+	/** This node applies a layer blending effect on a base image using a mask and a colour. */
+	class NodeImageLayerColour : public NodeImage
+	{
+	public:
+
+		Ptr<NodeImage> Base;
+		Ptr<NodeImage> Mask;
+		Ptr<NodeColour> Colour;
+		EBlendType Type;
+
+	public:
+
+		// Node Interface
+		virtual const FNodeType* GetType() const override { return &StaticType; }
+		static const FNodeType* GetStaticType() { return &StaticType; }
+
+	protected:
+
+		/** Forbidden. Manage with the Ptr<> template. */
+		~NodeImageLayerColour() {}
+
+	private:
+
+		static UE_API FNodeType StaticType;
+
+	};
+
+}
+
+#undef UE_API
diff --git a/Engine/Plugins/Mutable/Source/MutableTools/Internal/MuT/NodeImageParameter.h b/Engine/Plugins/Mutable/Source/MutableTools/Internal/MuT/NodeImageParameter.h
new file mode 100644
--- /dev/null
+++ b/Engine/Plugins/Mutable/Source/MutableTools/Internal/MuT/NodeImageParameter.h
@@ -0,0 +1,49 @@
+// Copyright Epic Games, Inc. All Rights Reserved.
+
+#pragma once
+
+#include "MuR/Ptr.h"
+#include "MuR/RefCounted.h"
+#include "MuT/Node.h"
+#include "MuT/NodeImage.h"
+#include "MuT/NodeRange.h"
+
+#define UE_API MUTABLETOOLS_API
+
+namespace UE::Mutable::Private
+{
+
+    /** */
+    class NodeImageParameter : public NodeImage
+	{
+	public:
+		/** Name of the parameter */
+		FString Name;
+
+		/** User provided ID to identify the parameter. */
+		FString UID;
+
+		/** Ranges for the parameter in case it is multidimensional. */
+		TArray<Ptr<NodeRange>> Ranges;
+
+	public:
+
+		// Node Interface
+		virtual const FNodeType* GetType() const override { return &StaticType; }
+		static const FNodeType* GetStaticType() { return &StaticType; }
+
+	protected:
+
+		/** Forbidden. Manage with the Ptr<> template. */
+		~NodeImageParameter() {}
+
+	private:
+
+		static UE_API FNodeType StaticType;
+
+	};
+
+
+}
+
+#undef UE_API
diff --git a/Engine/Plugins/Mutable/Source/MutableTools/Internal/MuT/NodeModifierMeshClipMorphPlane.h b/Engine/Plugins/Mutable/Source/MutableTools/Internal/MuT/NodeModifierMeshClipMorphPlane.h
new file mode 100644
--- /dev/null
+++ b/Engine/Plugins/Mutable/Source/MutableTools/Internal/MuT/NodeModifierMeshClipMorphPlane.h
@@ -0,0 +1,91 @@
+// Copyright Epic Games, Inc. All Rights Reserved.
+
+#pragma once
+
+#include "MuR/Ptr.h"
+#include "MuR/RefCounted.h"
+#include "MuR/Types.h"
+#include "MuR/Skeleton.h"
+#include "MuT/Node.h"
+#include "MuT/NodeModifier.h"
+
+#define UE_API MUTABLETOOLS_API
+
+namespace UE::Mutable::Private
+{
+
+	struct FBoneName;
+
+	struct FClipMorphPlaneParameters
+	{
+		// Morph field parameters
+
+		// Distance to the plane of last affected vertex
+		float DistanceToPlane = 0;
+
+		// "Linearity" factor of the influence.
+		float LinearityFactor = 0;
+
+		// Ellipse location
+		FVector3f Origin = FVector3f(0, 0, 0);
+		FVector3f Normal = FVector3f(0, 0, 0);
+		float Radius1 = 0, Radius2 = 0, Rotation = 0;
+
+		// Vertex selection box
+		EClipVertexSelectionType VertexSelectionType = EClipVertexSelectionType::None;
+		FVector3f SelectionBoxOrigin = FVector3f(0, 0, 0);
+		FVector3f SelectionBoxRadius = FVector3f(0, 0, 0);
+		FBoneName VertexSelectionBone;
+
+		// Max distance a vertex can have to the bone in order to be affected. A negative value
+		// means no limit.
+		float MaxEffectRadius = -1.0f;
+
+		EFaceCullStrategy FaceCullStrategy = EFaceCullStrategy::AllVerticesCulled;
+	};
+
+
+	/** */
+	class NodeModifierMeshClipMorphPlane : public NodeModifier
+	{
+	public:
+
+		FClipMorphPlaneParameters Parameters;
+
+	public:
+
+		// Node interface
+		virtual const FNodeType* GetType() const override { return &StaticType; }
+		static const FNodeType* GetStaticType() { return &StaticType; }
+
+		//-----------------------------------------------------------------------------------------
+        // Own interface
+		//-----------------------------------------------------------------------------------------
+
+		UE_API void SetPlane(FVector3f Center, FVector3f Normal);
+		UE_API void SetParams(float dist, float factor);
+		UE_API void SetMorphEllipse(float radius1, float radius2, float rotation);
+
+		//! Define an axis-aligned box that will select the vertices to be morphed.
+		//! Only one of Box or Bone Hierarchy can be used (the last one set)
+		UE_API void SetVertexSelectionBox(float centerX, float centerY, float centerZ, float radiusX, float radiusY, float radiusZ);
+
+		//! Define the root bone of the subhierarchy of the mesh that will be affected.
+		//! Only one of Box or Bone Hierarchy can be used (the last one set)
+		UE_API void SetVertexSelectionBone(const FBoneName& BoneId, float maxEffectRadius);
+
+	protected:
+
+		/** Forbidden. Manage with the Ptr<> template. */
+		~NodeModifierMeshClipMorphPlane() {}
+
+	private:
+
+		static UE_API FNodeType StaticType;
+
+	};
+
+
+}
+
+#undef UE_API
diff --git a/Engine/Plugins/Mutable/Source/MutableTools/Internal/MuT/NodeModifierMeshClipWithUVMask.h b/Engine/Plugins/Mutable/Source/MutableTools/Internal/MuT/NodeModifierMeshClipWithUVMask.h
new file mode 100644
--- /dev/null
+++ b/Engine/Plugins/Mutable/Source/MutableTools/Internal/MuT/NodeModifierMeshClipWithUVMask.h
@@ -0,0 +1,55 @@
+// Copyright Epic Games, Inc. All Rights Reserved.
+
+#pragma once
+
+#include "MuR/Ptr.h"
+#include "MuR/RefCounted.h"
+#include "MuT/Node.h"
+#include "MuT/NodeModifier.h"
+#include "MuT/NodeImage.h"
+#include "MuT/NodeLayout.h"
+
+#define UE_API MUTABLETOOLS_API
+
+
+namespace UE::Mutable::Private
+{
+
+	/** This node can clip part of a mesh using UV space-data (an image or a list of layout blocks). */
+	class NodeModifierMeshClipWithUVMask : public NodeModifier
+	{
+	public:
+
+		/** Image with the regions to remove. It will be interpreted as a bitmap. */
+		Ptr<NodeImage> ClipMask;
+
+		/** If the image above is null, clipping may hapopen with layout blocks. */
+		Ptr<NodeLayout> ClipLayout;
+
+		/** Layout index of the UVs to use inthe source mesh to ben clipped with the mask. */
+		uint8 LayoutIndex = 0;
+
+		/** */
+		EFaceCullStrategy FaceCullStrategy = EFaceCullStrategy::AllVerticesCulled;
+
+	public:
+
+		// Node interface
+		virtual const FNodeType* GetType() const override { return &StaticType; }
+		static const FNodeType* GetStaticType() { return &StaticType; }
+
+	protected:
+
+		/** Forbidden. Manage with the Ptr<> template. */
+		inline ~NodeModifierMeshClipWithUVMask() {}
+
+	private:
+
+		static UE_API FNodeType StaticType;
+
+	};
+
+
+}
+
+#undef UE_API
diff --git a/Engine/Plugins/Mutable/Source/MutableTools/Internal/MuT/NodeModifierMeshTransformInMesh.h b/Engine/Plugins/Mutable/Source/MutableTools/Internal/MuT/NodeModifierMeshTransformInMesh.h
new file mode 100644
--- /dev/null
+++ b/Engine/Plugins/Mutable/Source/MutableTools/Internal/MuT/NodeModifierMeshTransformInMesh.h
@@ -0,0 +1,46 @@
+// Copyright Epic Games, Inc. All Rights Reserved.
+
+#pragma once
+
+#include "MuR/Ptr.h"
+#include "MuR/Types.h"
+#include "MuT/Node.h"
+#include "MuT/NodeModifier.h"
+#include "MuT/NodeMesh.h"
+#include "MuT/NodeMatrix.h"
+
+#define UE_API MUTABLETOOLS_API
+
+namespace UE::Mutable::Private
+{
+
+/** */
+class NodeModifierMeshTransformInMesh : public NodeModifier
+{
+public:
+
+	/** */
+	Ptr<NodeMesh> BoundingMesh;
+
+	/** */
+	Ptr<NodeMatrix> MatrixNode;
+public:
+
+	// Node interface
+	virtual const FNodeType* GetType() const override { return &StaticType; }
+	static const FNodeType* GetStaticType() { return &StaticType; }
+
+protected:
+
+	/** Forbidden. Manage with the Ptr<> template. */
+	virtual ~NodeModifierMeshTransformInMesh() override {}
+
+private:
+
+	static UE_API FNodeType StaticType;
+
+};
+
+}
+
+#undef UE_API
diff --git a/Engine/Plugins/Mutable/Source/MutableTools/Internal/MuT/NodeModifierSurfaceEdit.h b/Engine/Plugins/Mutable/Source/MutableTools/Internal/MuT/NodeModifierSurfaceEdit.h
new file mode 100644
--- /dev/null
+++ b/Engine/Plugins/Mutable/Source/MutableTools/Internal/MuT/NodeModifierSurfaceEdit.h
@@ -0,0 +1,107 @@
+// Copyright Epic Games, Inc. All Rights Reserved.
+
+#pragma once
+
+#include "MuT/NodeModifier.h"
+#include "MuT/NodeMesh.h"
+#include "MuT/NodeImage.h"
+#include "MuT/NodeScalar.h"
+#include "MuR/ImageTypes.h"
+#include "MuR/RefCounted.h"
+#include "MuR/Ptr.h"
+
+#define UE_API MUTABLETOOLS_API
+
+
+namespace UE::Mutable::Private
+{
+
+	/** This node modifies a surface node.
+    * It allows to extend, cut and morph the parent Surface's meshes.
+    * It also allows to patch the parent Surface's textures.
+	*/
+    class NodeModifierSurfaceEdit : public NodeModifier
+	{
+	public:
+
+		/** Data for every modified texture. */
+		struct FTexture
+		{
+			/** Name used to match the image with the original one being modified. 
+			* This should match the MaterialParameterName in a NodeSurfaceNew::FImageData.
+			*/
+			FString MaterialParameterName;
+
+			/** Image to add if extgending. */
+			Ptr<NodeImage> Extend;
+
+			/** Image to blend if patching. */
+			Ptr<NodeImage> PatchImage;
+
+			/** Optional mask controlling the blending area. */
+			Ptr<NodeImage> PatchMask;
+
+			/** Rects in the parent layout homogeneous UV space to patch. */
+			TArray<FBox2f> PatchBlocks;
+
+			/** Type of patching operation. */
+			EBlendType PatchBlendType = EBlendType::BT_BLEND;
+
+			/** Patch alpha channel as well? */
+			bool bPatchApplyToAlpha = false;
+		};
+
+		struct FLOD
+		{
+			/** Mesh to remove from the modified surface. */
+			Ptr<NodeMesh> MeshRemove;
+
+			/** Mesh to add to the modified surface. */
+			Ptr<NodeMesh> MeshAdd;
+
+			/** Textures to modify. */
+			TArray<FTexture> Textures;
+		};
+
+		TArray<FLOD> LODs;
+
+		/** For remove operations, use this strategy to cull faces. */
+		EFaceCullStrategy FaceCullStrategy = EFaceCullStrategy::AllVerticesCulled;
+
+		/** Name of the morph to apply to the modified surface if it has it. */
+		FString MeshMorph;
+
+		/** Factor of the morph to apply. */
+		Ptr<NodeScalar> MorphFactor;
+
+		/** Source modifier guid, used to share resources between LODs and generate a deterministic surface guid. */
+		FGuid ModifierGuid;
+
+	public:
+
+		NodeModifierSurfaceEdit() 
+		{
+			// This modifier needs to be applied at the end of the operations.
+			bApplyBeforeNormalOperations = false;
+		}
+
+		// Node interface
+		virtual const FNodeType* GetType() const override { return &StaticType; }
+		static const FNodeType* GetStaticType() { return &StaticType; }
+
+	protected:
+
+		/** Forbidden. Manage with the Ptr<> template. */
+		~NodeModifierSurfaceEdit() {}
+
+	private:
+
+		static UE_API FNodeType StaticType;
+
+	};
+
+
+
+}
+
+#undef UE_API
