[Mutable] Enable Mesh LOD streaming by default.
Add an option to disable mesh LOD streaming per CO State.

diff --git a/Engine/Plugins/Mutable/Source/CustomizableObject/Internal/MuCO/CustomizableObjectPrivate.h b/Engine/Plugins/Mutable/Source/CustomizableObject/Internal/MuCO/CustomizableObjectPrivate.h
--- a/Engine/Plugins/Mutable/Source/CustomizableObject/Internal/MuCO/CustomizableObjectPrivate.h
+++ b/Engine/Plugins/Mutable/Source/CustomizableObject/Internal/MuCO/CustomizableObjectPrivate.h
@@ -888,6 +888,10 @@ struct FMutableStateData
 	UPROPERTY()
 	bool bLiveUpdateMode = false;
 
+	/** If this is enabled, Mesh streaming won't be used for this state, and all LODs will be generated when an instance is first updated. */
+	UPROPERTY()
+	bool bDisableMeshStreaming = false;
+
 	/** If this is enabled, texture streaming won't be used for this state, and full images will be generated when an instance is first updated. */
 	UPROPERTY()
 	bool bDisableTextureStreaming = false;
@@ -1699,6 +1703,8 @@ enum class ECustomizableObjectVersions : int32
 
 	LazyAddressOutsideFImage,
 
+	EnableMeshLODStreaming,
+
 	// -----<new versions can be added above this line>--------
 	LastCustomizableObjectVersion
 };
diff --git a/Engine/Plugins/Mutable/Source/CustomizableObject/Internal/MuCO/CustomizableObjectSystemPrivate.h b/Engine/Plugins/Mutable/Source/CustomizableObject/Internal/MuCO/CustomizableObjectSystemPrivate.h
--- a/Engine/Plugins/Mutable/Source/CustomizableObject/Internal/MuCO/CustomizableObjectSystemPrivate.h
+++ b/Engine/Plugins/Mutable/Source/CustomizableObject/Internal/MuCO/CustomizableObjectSystemPrivate.h
@@ -1008,7 +1008,7 @@ CUSTOMIZABLEOBJECT_API FTexturePlatformData* MutableCreateImagePlatformData(TSha
 
 
 /** Return true if Streaming is enabled for the given Object. */
-bool IsStreamingEnabled(const UCustomizableObject& Object);
+bool IsStreamingEnabled(const UCustomizableObject& Object, int32 State);
 
 
 template<typename Type>
diff --git a/Engine/Plugins/Mutable/Source/CustomizableObject/Private/MuCO/CustomizableObject.cpp b/Engine/Plugins/Mutable/Source/CustomizableObject/Private/MuCO/CustomizableObject.cpp
--- a/Engine/Plugins/Mutable/Source/CustomizableObject/Private/MuCO/CustomizableObject.cpp
+++ b/Engine/Plugins/Mutable/Source/CustomizableObject/Private/MuCO/CustomizableObject.cpp
@@ -2087,6 +2087,7 @@ FArchive& operator<<(FArchive& Ar, FMutableStateData& Struct)
 {
 	Ar << Struct.StateUIMetadata;
 	Ar << Struct.bLiveUpdateMode;
+	Ar << Struct.bDisableMeshStreaming;
 	Ar << Struct.bDisableTextureStreaming;
 	Ar << Struct.bReuseInstanceTextures;
 	Ar << Struct.ForcedParameterValues;
diff --git a/Engine/Plugins/Mutable/Source/CustomizableObject/Private/MuCO/CustomizableObjectInstance.cpp b/Engine/Plugins/Mutable/Source/CustomizableObject/Private/MuCO/CustomizableObjectInstance.cpp
--- a/Engine/Plugins/Mutable/Source/CustomizableObject/Private/MuCO/CustomizableObjectInstance.cpp
+++ b/Engine/Plugins/Mutable/Source/CustomizableObject/Private/MuCO/CustomizableObjectInstance.cpp
@@ -6771,7 +6771,7 @@ void UCustomizableObjectInstance::SetRequestedLODs(const TMap<FName, uint8>& InM
 		return;
 	}
 
-	if (IsStreamingEnabled(*CustomizableObject))
+	if (IsStreamingEnabled(*CustomizableObject, Descriptor.State))
 	{
 		return;
 	}
diff --git a/Engine/Plugins/Mutable/Source/CustomizableObject/Private/MuCO/CustomizableObjectSystem.cpp b/Engine/Plugins/Mutable/Source/CustomizableObject/Private/MuCO/CustomizableObjectSystem.cpp
--- a/Engine/Plugins/Mutable/Source/CustomizableObject/Private/MuCO/CustomizableObjectSystem.cpp
+++ b/Engine/Plugins/Mutable/Source/CustomizableObject/Private/MuCO/CustomizableObjectSystem.cpp
@@ -251,19 +251,20 @@ FUpdateContextPrivate::FUpdateContextPrivate(UCustomizableObjectInstance& InInst
 	Object = InObject;
 
 	CapturedDescriptor = Descriptor;
-	CapturedDescriptor.bStreamingEnabled = IsStreamingEnabled(*InObject);
 	InObject->GetPrivate()->ApplyStateForcedValuesToParameters(CapturedDescriptor);
 
-	CapturedDescriptorHash = FDescriptorHash(Descriptor);
-	NumObjectComponents = InObject->GetComponentCount();
-	
 	if (const UModelResources* ModelResources = InObject->GetPrivate()->GetModelResources())
 	{
+		CapturedDescriptor.bStreamingEnabled = IsStreamingEnabled(*InObject, Descriptor.State);
+
 		FirstLODAvailable = ModelResources->FirstLODAvailable;
 		FirstResidentLOD = ModelResources->NumLODsToStream;
 		ComponentNames = ModelResources->ComponentNamesPerObjectComponent;
 	}
 
+	CapturedDescriptorHash = FDescriptorHash(Descriptor);
+	NumObjectComponents = InObject->GetComponentCount();
+
 	UCustomizableObjectSystem* System = UCustomizableObjectSystem::GetInstanceChecked();
 	UCustomizableObjectSystemPrivate* SystemPrivate = System->GetPrivate();
 
@@ -3518,7 +3519,8 @@ namespace impl
 
 		const UModelResources& ModelResources = Instance->GetCustomizableObject()->GetPrivate()->GetModelResourcesChecked();
 
-		const FString& StateName = ObjectPrivate->GetStateName(InstancePrivate->GetState());
+		const int32 State = InstancePrivate->GetState();
+		const FString& StateName = ObjectPrivate->GetStateName(State);
 		const FMutableStateData* StateData = ModelResources.StateUIDataMap.Find(StateName);
 
 		Operation->bLiveUpdateMode = false;
@@ -3611,7 +3613,7 @@ namespace impl
 		Operation->InstanceID = Operation->bLiveUpdateMode ? InstancePrivate->LiveUpdateModeInstanceID : 0;
 		Operation->bUseMeshCache = Object->bEnableMeshCache && !Operation->bBake && !Operation->bLiveUpdateMode && UCustomizableObjectSystem::IsMeshCacheEnabled(true);
 		
-		Operation->bStreamMeshLODs = IsStreamingEnabled(*Object);
+		Operation->bStreamMeshLODs = IsStreamingEnabled(*Object, State);
 
 #if WITH_EDITOR
 		Operation->PixelFormatOverride = SystemPrivate->ImageFormatOverrideFunc;
@@ -4044,7 +4046,7 @@ int32 UCustomizableObjectSystem::TickInternal(const bool bBlocking)
 							LODUpdateDescriptorHash.FirstRequestedLOD = LODUpdateCandidate.FirstRequestedLOD;
 
 							if (CustomizableInstancePrivate->MinSquareDistFromComponentToPlayer < MaxSquareDistanceFound &&
-								(IsStreamingEnabled(*Object) || !LODUpdateDescriptorHash.IsSubset(CustomizableInstancePrivate->CommittedDescriptorHash) || CustomizableInstancePrivate->bAutomaticUpdateRequired))
+								(IsStreamingEnabled(*Object, CustomizableInstancePrivate->GetState()) || !LODUpdateDescriptorHash.IsSubset(CustomizableInstancePrivate->CommittedDescriptorHash) || CustomizableInstancePrivate->bAutomaticUpdateRequired))
 							{
 								MaxPriorityFound = LODUpdateCandidate.Priority;
 								MaxSquareDistanceFound = CustomizableInstancePrivate->MinSquareDistFromComponentToPlayer;
@@ -4952,9 +4954,15 @@ void ReconstructMorphTargets(const UE::Mutable::Private::FMesh& Mesh, const TArr
 }
 
 
-bool IsStreamingEnabled(const UCustomizableObject& Object)
+bool IsStreamingEnabled(const UCustomizableObject& Object, const int32 State)
 {
-	return (Object.bEnableMeshStreaming || bForceStreamMeshLODs) &&
+	const UModelResources& ModelResources = Object.GetPrivate()->GetModelResourcesChecked();
+	const FString& StateName = Object.GetPrivate()->GetStateName(State);
+	const FMutableStateData* StateData = ModelResources.StateUIDataMap.Find(StateName);
+
+	const bool bStateAllowsStreaming = StateData ? !StateData->bDisableMeshStreaming : true;
+
+	return ((Object.bEnableMeshStreaming && bStateAllowsStreaming) || bForceStreamMeshLODs) &&
 		bStreamMeshLODs &&
 		IStreamingManager::Get().IsRenderAssetStreamingEnabled(EStreamableRenderAssetType::SkeletalMesh);
 }
diff --git a/Engine/Plugins/Mutable/Source/CustomizableObject/Public/MuCO/CustomizableObject.h b/Engine/Plugins/Mutable/Source/CustomizableObject/Public/MuCO/CustomizableObject.h
--- a/Engine/Plugins/Mutable/Source/CustomizableObject/Public/MuCO/CustomizableObject.h
+++ b/Engine/Plugins/Mutable/Source/CustomizableObject/Public/MuCO/CustomizableObject.h
@@ -271,10 +271,9 @@ public:
 	UPROPERTY(EditAnywhere, Category = CustomizableObject)
 	bool bEnableMeshCache = false;
 
-	/** Experimental - If true, Mesh LODs will be streamed on demand. It requires streaming of SkeletalMeshes and Mutable.StreamMeshLODsEnabled to be enabled.
-	 *  Does not support Clothing and Morphs, yet. */
+	/** If true, Mesh LODs will be streamed on demand. It requires streaming of SkeletalMeshes and Mutable.StreamMeshLODsEnabled to be enabled. */
 	UPROPERTY(EditAnywhere, Category = CustomizableObject)
-	bool bEnableMeshStreaming = false;
+	bool bEnableMeshStreaming = true;
 	
 #if WITH_EDITORONLY_DATA
 private:
diff --git a/Engine/Plugins/Mutable/Source/CustomizableObjectEditor/Private/MuCOE/GenerateMutableSource/GenerateMutableSource.cpp b/Engine/Plugins/Mutable/Source/CustomizableObjectEditor/Private/MuCOE/GenerateMutableSource/GenerateMutableSource.cpp
--- a/Engine/Plugins/Mutable/Source/CustomizableObjectEditor/Private/MuCOE/GenerateMutableSource/GenerateMutableSource.cpp
+++ b/Engine/Plugins/Mutable/Source/CustomizableObjectEditor/Private/MuCOE/GenerateMutableSource/GenerateMutableSource.cpp
@@ -926,6 +926,7 @@ UE::Mutable::Private::Ptr<UE::Mutable::Private::NodeObject> GenerateMutableSourc
 			// UI Data
 			FMutableStateData StateUIData;
 			StateUIData.StateUIMetadata = State.UIMetadata;
+			StateUIData.bDisableMeshStreaming = State.bDisableTextureStreaming;
 			StateUIData.bDisableTextureStreaming = State.bDisableTextureStreaming;
 			StateUIData.bLiveUpdateMode = State.bLiveUpdateMode;
 			StateUIData.bReuseInstanceTextures = State.bReuseInstanceTextures;
diff --git a/Engine/Plugins/Mutable/Source/CustomizableObjectEditor/Private/MuCOE/Nodes/CustomizableObjectNodeObject.h b/Engine/Plugins/Mutable/Source/CustomizableObjectEditor/Private/MuCOE/Nodes/CustomizableObjectNodeObject.h
--- a/Engine/Plugins/Mutable/Source/CustomizableObjectEditor/Private/MuCOE/Nodes/CustomizableObjectNodeObject.h
+++ b/Engine/Plugins/Mutable/Source/CustomizableObjectEditor/Private/MuCOE/Nodes/CustomizableObjectNodeObject.h
@@ -39,6 +39,10 @@ struct FCustomizableObjectState
 	UPROPERTY(EditAnywhere, Category = CustomizableObject)
 	ETextureCompressionStrategy TextureCompressionStrategy = ETextureCompressionStrategy::None;
 
+	/** If this is enabled, Mesh streaming won't be used for this state, and all LODs will be generated when an instance is first updated. */
+	UPROPERTY(EditAnywhere, Category = CustomizableObject)
+	bool bDisableMeshStreaming = false;
+
 	/** If this is enabled, texture streaming won't be used for this state, and full images will be generated when an instance is first updated. */
 	UPROPERTY(EditAnywhere, Category = CustomizableObject)
 	bool bDisableTextureStreaming = false;
