[Mutable] Remove some cases where Skeletal Meshes were loaded synchronously.
Add flags to ignore AssetUserData and TexCoords on support meshes.
Add streamed resources and mesh sockets to the mu::FMesh in the second stage of the mesh conversion.
Fix UCustomizableObjectNodeSkeletalMesh checking if a SoftObjectPtr is loaded instead of not null.

#rnx

diff --git a/Engine/Plugins/Mutable/Source/CustomizableObject/Internal/MuCO/UnrealToMutableTextureConversionUtils.h b/Engine/Plugins/Mutable/Source/CustomizableObject/Internal/MuCO/UnrealToMutableTextureConversionUtils.h
--- a/Engine/Plugins/Mutable/Source/CustomizableObject/Internal/MuCO/UnrealToMutableTextureConversionUtils.h
+++ b/Engine/Plugins/Mutable/Source/CustomizableObject/Internal/MuCO/UnrealToMutableTextureConversionUtils.h
@@ -90,6 +90,7 @@ struct FMutableSourceMeshData
 	TSoftObjectPtr<UStreamableRenderAsset> Mesh;
 	TSoftClassPtr<UAnimInstance> AnimInstance;
 	TSoftObjectPtr<UStreamableRenderAsset> TableReferenceSkeletalMesh;
+
 	FName Component = NAME_None;
 
 	/** */
@@ -100,9 +101,9 @@ struct FMutableSourceMeshData
 	FMutableSourceSurfaceMetadata Metadata;
 
 	/** Selection of the mesh section*/
-	uint8 LODIndex=0;
-	uint8 SectionIndex=0;
-	bool bMeshMustExist=false;
+	int8 BaseLODIndex = INDEX_NONE;
+	int8 BaseSectionIndex = INDEX_NONE;
+	int8 LODOffset = INDEX_NONE;
 
 	/** Required mesh properties. */
 	EMutableMeshConversionFlags Flags = EMutableMeshConversionFlags::None;
@@ -111,9 +112,17 @@ struct FMutableSourceMeshData
 	bool bUseAllRealTimeMorphs = false;
 	TArray<FString> UsedRealTimeMorphTargetNames;
 
+	/** */
+	int32 SocketPriority = 0;
+
+	/** */
+	bool bOnlyConnectedLOD = false;
+
 	/** Context for log messages. */
 	const UObject* MessageContext = nullptr;
 
+	FString OptionalMessageInfo;
+
 	bool operator==(const FMutableSourceMeshData&) const = default;
 };
 
@@ -121,8 +130,11 @@ struct FMutableSourceMeshData
 inline uint32 GetTypeHash(const FMutableSourceMeshData& Key)
 {
 	uint32 GuidHash = GetTypeHash(Key.Mesh);
-	GuidHash = HashCombineFast(GuidHash, GetTypeHash(Key.LODIndex));
-	GuidHash = HashCombineFast(GuidHash, GetTypeHash(Key.SectionIndex));
+	GuidHash = HashCombineFast(GuidHash, GetTypeHash(Key.BaseLODIndex));
+	GuidHash = HashCombineFast(GuidHash, GetTypeHash(Key.BaseSectionIndex));
+	GuidHash = HashCombineFast(GuidHash, GetTypeHash(Key.LODOffset));
+	GuidHash = HashCombineFast(GuidHash, GetTypeHash(Key.SocketPriority));
+	GuidHash = HashCombineFast(GuidHash, GetTypeHash(Key.bOnlyConnectedLOD));
 	GuidHash = HashCombineFast(GuidHash, GetTypeHash(Key.Component));
 	GuidHash = HashCombineFast(GuidHash, GetTypeHash(Key.Flags));
 	return GuidHash;
diff --git a/Engine/Plugins/Mutable/Source/CustomizableObjectEditor/Private/MuCOE/CustomizableObjectCompiler.cpp b/Engine/Plugins/Mutable/Source/CustomizableObjectEditor/Private/MuCOE/CustomizableObjectCompiler.cpp
--- a/Engine/Plugins/Mutable/Source/CustomizableObjectEditor/Private/MuCOE/CustomizableObjectCompiler.cpp
+++ b/Engine/Plugins/Mutable/Source/CustomizableObjectEditor/Private/MuCOE/CustomizableObjectCompiler.cpp
@@ -990,7 +990,6 @@ void FCustomizableObjectCompiler::CompileInternal(bool bAsync)
 		ModelResources->AnimBPs = MoveTemp(GenerationContext.AnimBPAssets);
 
 		ModelResources->MaterialSlotNames = MoveTemp(GenerationContext.ReferencedMaterialSlotNames);
-		ModelResources->SocketArray = MoveTemp(GenerationContext.SocketArray);
 
 		TArray<FGeneratedImageProperties> ImageProperties;
 		GenerationContext.ImageProperties.GenerateValueArray(ImageProperties);
@@ -1086,8 +1085,6 @@ void FCustomizableObjectCompiler::CompileInternal(bool bAsync)
 			CurrentObject->Modify();
 		}
 
-		ModelResources->StreamedResourceDataEditor = MoveTemp(GenerationContext.StreamedResourceData);
-		
 		CompileTask = MakeShareable(new FCustomizableObjectCompileRunnable(MutableRoot, SharedThis(this)));
 		CompileTask->Options = CurrentOptions;
 		CompileTask->ReferencedTextures = NewCompileTimeReferencedTextures;
@@ -1715,6 +1712,8 @@ void FCustomizableObjectCompiler::FinishCompilationTask()
 
 		ModelResources->PhysicsAssets = MoveTemp(CompilationContext->PhysicsAssets);
 		ModelResources->AnimBpOverridePhysiscAssetsInfo = MoveTemp(CompilationContext->AnimBpOverridePhysicsAssetsInfo);
+		
+		ModelResources->SocketArray = MoveTemp(CompilationContext->SocketArray);
 
 		const int32 NumBones = CompilationContext->UniqueBoneNames.Num() + CompilationContext->RemappedBoneNames.Num();
 		ModelResources->BoneNamesMap.Reserve(NumBones);
@@ -1731,6 +1730,8 @@ void FCustomizableObjectCompiler::FinishCompilationTask()
 
 		ModelResources->SkinWeightProfilesInfo = MoveTemp(CompilationContext->SkinWeightProfilesInfo);
 
+		ModelResources->StreamedResourceDataEditor = MoveTemp(CompilationContext->StreamedResourceData);
+
 		TSharedPtr<FModelStreamableBulkData> ModelStreamableBulkData = PlatformData->ModelStreamableBulkData;
 		ModelStreamableBulkData->RealTimeMorphStreamables.Empty(32);
 
diff --git a/Engine/Plugins/Mutable/Source/CustomizableObjectEditor/Private/MuCOE/CustomizableObjectLayout.cpp b/Engine/Plugins/Mutable/Source/CustomizableObjectEditor/Private/MuCOE/CustomizableObjectLayout.cpp
--- a/Engine/Plugins/Mutable/Source/CustomizableObjectEditor/Private/MuCOE/CustomizableObjectLayout.cpp
+++ b/Engine/Plugins/Mutable/Source/CustomizableObjectEditor/Private/MuCOE/CustomizableObjectLayout.cpp
@@ -137,7 +137,7 @@ void UCustomizableObjectLayout::GenerateAutomaticBlocksFromUVs()
 	TSharedPtr<FMutableCompilationContext> CompilationContext = MakeShared<FMutableCompilationContext>(Object, Compiler, Options);
 	FMutableGraphGenerationContext GenerationContext(CompilationContext);
 
-	UE::Mutable::Private::Ptr<UE::Mutable::Private::NodeLayout> LayoutNode = CreateMutableLayoutNode(GenerationContext, this, false);
+	UE::Mutable::Private::Ptr<UE::Mutable::Private::NodeLayout> LayoutNode = CreateMutableLayoutNode(this, false);
 	check(LayoutNode);
 
 	if (PackingStrategy == ECustomizableObjectTextureLayoutPackingStrategy::Overlay)
@@ -164,15 +164,15 @@ void UCustomizableObjectLayout::GenerateAutomaticBlocksFromUVs()
 				EMutableMeshConversionFlags::IgnoreSkinning |
 				EMutableMeshConversionFlags::IgnorePhysics |
 				EMutableMeshConversionFlags::IgnoreMorphs |
+				EMutableMeshConversionFlags::IgnoreAUD |
 				EMutableMeshConversionFlags::DoNotCreateMeshMetadata;
 
 			GenerationContext.MeshGenerationFlags.Push(ShapeFlags);
 
-			bool bMeshMustExist = true;
 			bool bForceImmediateConversion = true;
 			FMutableSourceMeshData Source;
 			Source.Mesh = SkeletalMesh;
-			MutableMesh = ConvertSkeletalMeshToMutable(Source, bMeshMustExist, GetLOD(), GetMaterial(), GenerationContext, nullptr, bForceImmediateConversion);
+			MutableMesh = ConvertSkeletalMeshToMutable(Source, GetLOD(), GetMaterial(), GenerationContext, nullptr, bForceImmediateConversion);
 
 			GenerationContext.MeshGenerationFlags.Pop();
 		}
diff --git a/Engine/Plugins/Mutable/Source/CustomizableObjectEditor/Private/MuCOE/GenerateMutableSource/GenerateMutableSource.cpp b/Engine/Plugins/Mutable/Source/CustomizableObjectEditor/Private/MuCOE/GenerateMutableSource/GenerateMutableSource.cpp
--- a/Engine/Plugins/Mutable/Source/CustomizableObjectEditor/Private/MuCOE/GenerateMutableSource/GenerateMutableSource.cpp
+++ b/Engine/Plugins/Mutable/Source/CustomizableObjectEditor/Private/MuCOE/GenerateMutableSource/GenerateMutableSource.cpp
@@ -302,7 +302,7 @@ TSharedPtr<UE::Mutable::Private::FMesh> FMutableGraphGenerationContext::FindGene
 }
 
 
-int32 FMutableGraphGenerationContext::AddStreamedResource(UAssetUserData& AssetUserData)
+int32 FMutableCompilationContext::AddStreamedResource(UAssetUserData& AssetUserData)
 {
 	const uint32 AssetIdentifier = AssetUserData.GetUniqueID();
 
@@ -926,7 +926,7 @@ UE::Mutable::Private::Ptr<UE::Mutable::Private::NodeObject> GenerateMutableSourc
 			// UI Data
 			FMutableStateData StateUIData;
 			StateUIData.StateUIMetadata = State.UIMetadata;
-			StateUIData.bDisableMeshStreaming = State.bDisableTextureStreaming;
+			StateUIData.bDisableMeshStreaming = State.bDisableMeshStreaming;
 			StateUIData.bDisableTextureStreaming = State.bDisableTextureStreaming;
 			StateUIData.bLiveUpdateMode = State.bLiveUpdateMode;
 			StateUIData.bReuseInstanceTextures = State.bReuseInstanceTextures;
@@ -1510,7 +1510,7 @@ void PopulateReferenceSkeletalMeshesData(FMutableGraphGenerationContext& Generat
 			{
 				if (AssetUserData)
 				{
-					Data.AssetUserDataIndices.AddUnique(GenerationContext.AddStreamedResource(*AssetUserData));
+					Data.AssetUserDataIndices.AddUnique(GenerationContext.CompilationContext->AddStreamedResource(*AssetUserData));
 				}
 			}
 		}
@@ -1789,30 +1789,6 @@ TSharedPtr<UE::Mutable::Private::FMesh> GenerateMeshConstant(
 }
 
 
-void AddSocketTagsToMesh(const USkeletalMesh* SourceMesh, TSharedPtr<UE::Mutable::Private::FMesh> MutableMesh, FMutableGraphGenerationContext& GenerationContext)
-{
-	for (int32 SocketIndex = 0; SocketIndex < SourceMesh->NumSockets(); ++SocketIndex)
-	{
-		USkeletalMeshSocket* Socket = SourceMesh->GetSocketByIndex(SocketIndex);
-
-		FMutableRefSocket MutableSocket;
-		MutableSocket.SocketName = Socket->SocketName;
-		MutableSocket.BoneName = Socket->BoneName;
-		MutableSocket.RelativeLocation = Socket->RelativeLocation;
-		MutableSocket.RelativeRotation = Socket->RelativeRotation;
-		MutableSocket.RelativeScale = Socket->RelativeScale;
-		MutableSocket.bForceAlwaysAnimated = Socket->bForceAlwaysAnimated;
-
-		MutableSocket.Priority = GenerationContext.SocketPriorityStack.IsEmpty() ? 0 : GenerationContext.SocketPriorityStack.Top();
-			
-		int32 SocketArrayIndex = GenerationContext.SocketArray.AddUnique(MutableSocket);
-		FString SocketTag = FString::Printf(TEXT("__Socket:%d"), SocketArrayIndex);
-
-		AddTagToMutableMeshUnique(*MutableMesh, SocketTag);
-	}
-}
-
-
 FMutableComponentInfo::FMutableComponentInfo(FName InComponentName, USkeletalMesh* InRefSkeletalMesh)
 {
 	ComponentName = InComponentName;
diff --git a/Engine/Plugins/Mutable/Source/CustomizableObjectEditor/Private/MuCOE/GenerateMutableSource/GenerateMutableSource.h b/Engine/Plugins/Mutable/Source/CustomizableObjectEditor/Private/MuCOE/GenerateMutableSource/GenerateMutableSource.h
--- a/Engine/Plugins/Mutable/Source/CustomizableObjectEditor/Private/MuCOE/GenerateMutableSource/GenerateMutableSource.h
+++ b/Engine/Plugins/Mutable/Source/CustomizableObjectEditor/Private/MuCOE/GenerateMutableSource/GenerateMutableSource.h
@@ -485,6 +485,13 @@ struct FMutableCompilationContext
 
 	TArray<FAnimBpOverridePhysicsAssetsInfo> AnimBpOverridePhysicsAssetsInfo;
 
+	/** Resource Data constants */
+	TMap<uint32, int32> StreamedResourceIndices;
+	TArray<FCustomizableObjectResourceData> StreamedResourceData;
+
+	/** Stores the sockets provided by the part skeletal meshes, to be merged in the generated meshes */
+	TArray<FMutableRefSocket> SocketArray;
+
 	/** Data used for SkinWeightProfiles reconstruction. */
 	TArray<FMutableSkinWeightProfileInfo> SkinWeightProfilesInfo;
 
@@ -536,6 +543,10 @@ public:
 	/** */
 	uint32 GetSkinWeightProfileIdUnique(const FName ProfileName);
 
+	/** Adds a streamed resource of type AssetUserData.
+	 * Returns resource index in the array of streamed resources. */
+	int32 AddStreamedResource(UAssetUserData& AssetUserData);
+
 };
 
 
@@ -665,7 +676,7 @@ struct FMutableGraphGenerationContext
 	TMap<FMutableSourceMeshData, FGeneratedReferencedMesh> RuntimeReferencedMeshMap;
 
 	// Cache of runtime pass-through meshes and their IDs used in the core to identify them
-	// These meshes will become mutable images in the compiled model.
+	// These meshes will become mutable meshes in the compiled model.
 	TMap<FMutableSourceMeshData, FGeneratedReferencedMesh> CompileTimeMeshMap;
 
 	// Cache of runtime pass-through images and their IDs used in the core to identify them.
@@ -753,10 +764,6 @@ struct FMutableGraphGenerationContext
 	/** Find a mesh if already generated for a given source and flags. */
 	TSharedPtr<UE::Mutable::Private::FMesh> FindGeneratedMesh(const FGeneratedMeshData::FKey& Key);
 
-	/** Adds a streamed resource of type AssetUserData.
-	  * Returns resource index in the array of streamed resources. */
-	int32 AddStreamedResource(UAssetUserData& AssetUserData);
-
 	// Check if the Id of the node Node already exists, if it's new adds it to NodeIds array, otherwise, returns new Id
 	const FGuid GetNodeIdUnique(const UCustomizableObjectNode* Node);
 
@@ -840,9 +847,6 @@ struct FMutableGraphGenerationContext
 	/** Stores the anim BP assets gathered from the SkeletalMesh nodes during compilation, to be used in mesh generation in-game */
 	TArray<TSoftClassPtr<UAnimInstance>> AnimBPAssets;
 
-	/** Stores the sockets provided by the part skeletal meshes, to be merged in the generated meshes */
-	TArray<FMutableRefSocket> SocketArray;
-
 	/** Used to propagate the socket priority defined in group nodes to their child skeletal mesh nodes
 	* It's a stack because group nodes are recursive
 	*/
@@ -870,10 +874,6 @@ struct FMutableGraphGenerationContext
 	// Material to SharedSurfaceId
 	TMap<FSharedSurfaceKey, FGuid> SurfaceGuids;
 
-	/** Resource Data constants */
-	TMap<uint32, int32> StreamedResourceIndices;
-	TArray<FCustomizableObjectResourceData> StreamedResourceData;
-
 	/** Extension Data constants are collected here */
 	FExtensionDataCompilerInterface ExtensionDataCompilerInterface;
 	TArray<FCustomizableObjectResourceData> AlwaysLoadedExtensionData;
@@ -986,8 +986,6 @@ UE::Mutable::Private::Ptr<UE::Mutable::Private::NodeMeshApplyPose> CreateNodeMes
 /** Adds Tag to MutableMesh uniquely, returns the index were the tag has been inserted or the index where an intance of the tag has been found */
 int32 AddTagToMutableMeshUnique(UE::Mutable::Private::FMesh& MutableMesh, const FString& Tag);
 
-void AddSocketTagsToMesh(const USkeletalMesh* SourceMesh, TSharedPtr<UE::Mutable::Private::FMesh> MutableMesh, FMutableGraphGenerationContext& GenerationContext);
-
 // Generates the tag for an animation instance
 FString GenerateAnimationInstanceTag(const int32 AnimInstanceIndex, const FName& SlotIndex);
 
diff --git a/Engine/Plugins/Mutable/Source/CustomizableObjectEditor/Private/MuCOE/GenerateMutableSource/GenerateMutableSourceColor.cpp b/Engine/Plugins/Mutable/Source/CustomizableObjectEditor/Private/MuCOE/GenerateMutableSource/GenerateMutableSourceColor.cpp
--- a/Engine/Plugins/Mutable/Source/CustomizableObjectEditor/Private/MuCOE/GenerateMutableSource/GenerateMutableSourceColor.cpp
+++ b/Engine/Plugins/Mutable/Source/CustomizableObjectEditor/Private/MuCOE/GenerateMutableSource/GenerateMutableSourceColor.cpp
@@ -327,9 +327,7 @@ UE::Mutable::Private::Ptr<UE::Mutable::Private::NodeColour> GenerateMutableSourc
 					// Generating a new Color column if not exists
 					if (Table->FindColumn(ColumnName) == INDEX_NONE)
 					{
-						int32 Dummy = -1; // TODO MTBL-1512
-						bool Dummy2 = false;
-						bSuccess = GenerateTableColumn(TypedNodeTable, Pin, Table, ColumnName, Property, FMutableSourceMeshData(), Dummy, Dummy, GenerationContext.CurrentLOD, Dummy, Dummy2, GenerationContext);
+						bSuccess = GenerateTableColumn(TypedNodeTable, Pin, Table, ColumnName, Property, FMutableSourceMeshData(), GenerationContext);
 
 						if (!bSuccess)
 						{
diff --git a/Engine/Plugins/Mutable/Source/CustomizableObjectEditor/Private/MuCOE/GenerateMutableSource/GenerateMutableSourceComponent.cpp b/Engine/Plugins/Mutable/Source/CustomizableObjectEditor/Private/MuCOE/GenerateMutableSource/GenerateMutableSourceComponent.cpp
--- a/Engine/Plugins/Mutable/Source/CustomizableObjectEditor/Private/MuCOE/GenerateMutableSource/GenerateMutableSourceComponent.cpp
+++ b/Engine/Plugins/Mutable/Source/CustomizableObjectEditor/Private/MuCOE/GenerateMutableSource/GenerateMutableSourceComponent.cpp
@@ -261,7 +261,7 @@ UE::Mutable::Private::Ptr<UE::Mutable::Private::NodeComponent> GenerateMutableSo
 			FMutableSourceMeshData Source;
 			Source.Mesh = SkeletalMesh;
 			Source.bIsPassthrough = true;
-			TSharedPtr<UE::Mutable::Private::FMesh> MutableMesh = GenerateMutableSkeletalMesh(Source, 0, 0, 0, 0, {}, GenerationContext, Node);
+			TSharedPtr<UE::Mutable::Private::FMesh> MutableMesh = GenerateMutableSkeletalMesh(Source, 0, 0, {}, GenerationContext, Node, false);
 
 			MeshNode->Value = MutableMesh;
 		}
diff --git a/Engine/Plugins/Mutable/Source/CustomizableObjectEditor/Private/MuCOE/GenerateMutableSource/GenerateMutableSourceFloat.cpp b/Engine/Plugins/Mutable/Source/CustomizableObjectEditor/Private/MuCOE/GenerateMutableSource/GenerateMutableSourceFloat.cpp
--- a/Engine/Plugins/Mutable/Source/CustomizableObjectEditor/Private/MuCOE/GenerateMutableSource/GenerateMutableSourceFloat.cpp
+++ b/Engine/Plugins/Mutable/Source/CustomizableObjectEditor/Private/MuCOE/GenerateMutableSource/GenerateMutableSourceFloat.cpp
@@ -353,9 +353,7 @@ UE::Mutable::Private::Ptr<UE::Mutable::Private::NodeScalar> GenerateMutableSourc
 					// Generating a new Float column if not exists
 					if (Table->FindColumn(ColumnName) == INDEX_NONE)
 					{
-						int32 Dummy = -1; // TODO MTBL-1512
-						bool Dummy2 = false;
-						bSuccess = GenerateTableColumn(TypedNodeTable, Pin, Table, ColumnName, Property, FMutableSourceMeshData(), Dummy, Dummy, GenerationContext.CurrentLOD, Dummy, Dummy2, GenerationContext);
+						bSuccess = GenerateTableColumn(TypedNodeTable, Pin, Table, ColumnName, Property, FMutableSourceMeshData(), GenerationContext);
 
 						if (!bSuccess)
 						{
diff --git a/Engine/Plugins/Mutable/Source/CustomizableObjectEditor/Private/MuCOE/GenerateMutableSource/GenerateMutableSourceImage.cpp b/Engine/Plugins/Mutable/Source/CustomizableObjectEditor/Private/MuCOE/GenerateMutableSource/GenerateMutableSourceImage.cpp
--- a/Engine/Plugins/Mutable/Source/CustomizableObjectEditor/Private/MuCOE/GenerateMutableSource/GenerateMutableSourceImage.cpp
+++ b/Engine/Plugins/Mutable/Source/CustomizableObjectEditor/Private/MuCOE/GenerateMutableSource/GenerateMutableSourceImage.cpp
@@ -1103,11 +1103,8 @@ UE::Mutable::Private::Ptr<UE::Mutable::Private::NodeImage> GenerateMutableSource
 						// Generating a new Texture column if not exists
 						if (GeneratedTable->FindColumn(ColumnName) == INDEX_NONE)
 						{
-							int32 Dummy = -1; // TODO MTBL-1512
-							bool Dummy2 = false;
 							GenerationContext.CurrentTableColumnType = UE::Mutable::Private::ETableColumnType::Image;
-							bSuccess = GenerateTableColumn(TypedNodeTable, Pin, GeneratedTable, ColumnName, Property, 
-								FMutableSourceMeshData(), Dummy, Dummy, GenerationContext.CurrentLOD, Dummy, Dummy2, GenerationContext);
+							bSuccess = GenerateTableColumn(TypedNodeTable, Pin, GeneratedTable, ColumnName, Property, FMutableSourceMeshData(), GenerationContext);
 
 							if (!bSuccess)
 							{
diff --git a/Engine/Plugins/Mutable/Source/CustomizableObjectEditor/Private/MuCOE/GenerateMutableSource/GenerateMutableSourceLayout.cpp b/Engine/Plugins/Mutable/Source/CustomizableObjectEditor/Private/MuCOE/GenerateMutableSource/GenerateMutableSourceLayout.cpp
--- a/Engine/Plugins/Mutable/Source/CustomizableObjectEditor/Private/MuCOE/GenerateMutableSource/GenerateMutableSourceLayout.cpp
+++ b/Engine/Plugins/Mutable/Source/CustomizableObjectEditor/Private/MuCOE/GenerateMutableSource/GenerateMutableSourceLayout.cpp
@@ -34,7 +34,7 @@ UE::Mutable::Private::Ptr<UE::Mutable::Private::NodeLayout> CreateDefaultLayout(
 }
 
 
-UE::Mutable::Private::Ptr<UE::Mutable::Private::NodeLayout> CreateMutableLayoutNode(FMutableGraphGenerationContext& GenerationContext, const UCustomizableObjectLayout* UnrealLayout, bool bIgnoreLayoutWarnings)
+UE::Mutable::Private::Ptr<UE::Mutable::Private::NodeLayout> CreateMutableLayoutNode(const UCustomizableObjectLayout* UnrealLayout, bool bIgnoreLayoutWarnings)
 {
 	UE::Mutable::Private::Ptr<UE::Mutable::Private::NodeLayout> LayoutNode = new UE::Mutable::Private::NodeLayout;
 
@@ -60,14 +60,14 @@ UE::Mutable::Private::Ptr<UE::Mutable::Private::NodeLayout> CreateMutableLayoutN
 	LayoutNode->Blocks.SetNum(UnrealLayout->Blocks.Num());
 	for (int32 BlockIndex = 0; BlockIndex < UnrealLayout->Blocks.Num(); ++BlockIndex)
 	{
-		LayoutNode->Blocks[BlockIndex] = ToMutable(GenerationContext, UnrealLayout->Blocks[BlockIndex]);
+		LayoutNode->Blocks[BlockIndex] = ToMutable(UnrealLayout->Blocks[BlockIndex]);
 	}
 
 	return LayoutNode;
 }
 
 
-UE::Mutable::Private::FSourceLayoutBlock ToMutable(FMutableGraphGenerationContext& GenerationContext, const FCustomizableObjectLayoutBlock& UnrealBlock)
+UE::Mutable::Private::FSourceLayoutBlock ToMutable(const FCustomizableObjectLayoutBlock& UnrealBlock)
 {
 	UE::Mutable::Private::FSourceLayoutBlock MutableBlock;
 
diff --git a/Engine/Plugins/Mutable/Source/CustomizableObjectEditor/Private/MuCOE/GenerateMutableSource/GenerateMutableSourceLayout.h b/Engine/Plugins/Mutable/Source/CustomizableObjectEditor/Private/MuCOE/GenerateMutableSource/GenerateMutableSourceLayout.h
--- a/Engine/Plugins/Mutable/Source/CustomizableObjectEditor/Private/MuCOE/GenerateMutableSource/GenerateMutableSourceLayout.h
+++ b/Engine/Plugins/Mutable/Source/CustomizableObjectEditor/Private/MuCOE/GenerateMutableSource/GenerateMutableSourceLayout.h
@@ -6,11 +6,10 @@
 #include "MuCOE/CustomizableObjectLayout.h"
 
 class UEdGraphPin;
-struct FMutableGraphGenerationContext;
 
 
 extern UE::Mutable::Private::Ptr<UE::Mutable::Private::NodeLayout> CreateDefaultLayout();
 
-extern UE::Mutable::Private::Ptr<UE::Mutable::Private::NodeLayout> CreateMutableLayoutNode(FMutableGraphGenerationContext& GenerationContext, const UCustomizableObjectLayout*, bool bIgnoreLayoutWarnings);
+extern UE::Mutable::Private::Ptr<UE::Mutable::Private::NodeLayout> CreateMutableLayoutNode(const UCustomizableObjectLayout*, bool bIgnoreLayoutWarnings);
 
-extern UE::Mutable::Private::FSourceLayoutBlock ToMutable(FMutableGraphGenerationContext& GenerationContext, const FCustomizableObjectLayoutBlock&);
+extern UE::Mutable::Private::FSourceLayoutBlock ToMutable(const FCustomizableObjectLayoutBlock&);
diff --git a/Engine/Plugins/Mutable/Source/CustomizableObjectEditor/Private/MuCOE/GenerateMutableSource/GenerateMutableSourceMaterial.cpp b/Engine/Plugins/Mutable/Source/CustomizableObjectEditor/Private/MuCOE/GenerateMutableSource/GenerateMutableSourceMaterial.cpp
--- a/Engine/Plugins/Mutable/Source/CustomizableObjectEditor/Private/MuCOE/GenerateMutableSource/GenerateMutableSourceMaterial.cpp
+++ b/Engine/Plugins/Mutable/Source/CustomizableObjectEditor/Private/MuCOE/GenerateMutableSource/GenerateMutableSourceMaterial.cpp
@@ -433,10 +433,8 @@ UE::Mutable::Private::Ptr<UE::Mutable::Private::NodeMaterial> GenerateMutableSou
 					// Generating a new Float column if not exists
 					if (Table->FindColumn(ColumnName) == INDEX_NONE)
 					{
-						int32 Dummy = -1; // TODO MTBL-1512
-						bool Dummy2 = false;
 						GenerationContext.CurrentTableColumnType = UE::Mutable::Private::ETableColumnType::Material;
-						bSuccess = GenerateTableColumn(TypedNodeTable, Pin, Table, ColumnName, Property, FMutableSourceMeshData(), Dummy, Dummy, GenerationContext.CurrentLOD, Dummy, Dummy2, GenerationContext);
+						bSuccess = GenerateTableColumn(TypedNodeTable, Pin, Table, ColumnName, Property, FMutableSourceMeshData(), GenerationContext);
 
 						if (!bSuccess)
 						{
diff --git a/Engine/Plugins/Mutable/Source/CustomizableObjectEditor/Private/MuCOE/GenerateMutableSource/GenerateMutableSourceMesh.cpp b/Engine/Plugins/Mutable/Source/CustomizableObjectEditor/Private/MuCOE/GenerateMutableSource/GenerateMutableSourceMesh.cpp
--- a/Engine/Plugins/Mutable/Source/CustomizableObjectEditor/Private/MuCOE/GenerateMutableSource/GenerateMutableSourceMesh.cpp
+++ b/Engine/Plugins/Mutable/Source/CustomizableObjectEditor/Private/MuCOE/GenerateMutableSource/GenerateMutableSourceMesh.cpp
@@ -55,6 +55,7 @@
 #include "PhysicsEngine/PhysicsAsset.h"
 #include "PhysicsEngine/SkeletalBodySetup.h"
 #include "Engine/SkinnedAssetCommon.h"
+#include "Engine/SkeletalMeshSocket.h"
 #include "AssetRegistry/AssetRegistryModule.h"
 #include "Hash/CityHash.h"
 
@@ -209,83 +210,67 @@ void SetSurfaceFormat(FMutableCompilationContext& Context,
 }
 
 
-void GetLODAndSectionForAutomaticLODs(const FMutableGraphGenerationContext& Context, const UCustomizableObjectNode& Node, const USkeletalMesh& SkeletalMesh,
-	const int32 LODIndexConnected, const int32 SectionIndexConnected, int32& OutLODIndex, int32& OutSectionIndex, const bool bOnlyConnectedLOD)
+void GetLODAndSectionForAutomaticLODs(const FMutableCompilationContext& Context, const FMutableSourceMeshData& Source, const USkeletalMesh& SkeletalMesh, int32& OutLODIndex, int32& OutSectionIndex)
 {
-	OutLODIndex = LODIndexConnected;
-	OutSectionIndex = SectionIndexConnected;
+	OutLODIndex = Source.BaseLODIndex;
+	OutSectionIndex = Source.BaseSectionIndex;
 	
-	if (Context.CurrentAutoLODStrategy != ECustomizableObjectAutomaticLODStrategy::AutomaticFromMesh ||
-		bOnlyConnectedLOD)
+	if (Source.bOnlyConnectedLOD || Source.LODOffset == 0)
 	{
 		return;
 	}
 	
-	// When processing pins of the current LOD, indices will remain the same.
-	if (Context.CurrentLOD == Context.FromLOD)
-	{
-		return;
-	}
-
 	const FSkeletalMeshModel* ImportedModel = SkeletalMesh.GetImportedModel();
 	if (!ImportedModel)
 	{
 		return;
 	}
 
-	if (!ImportedModel->LODModels.IsValidIndex(LODIndexConnected) ||
-		!ImportedModel->LODModels[LODIndexConnected].Sections.IsValidIndex(SectionIndexConnected))
+	if (!ImportedModel->LODModels.IsValidIndex(Source.BaseLODIndex) ||
+		!ImportedModel->LODModels[Source.BaseLODIndex].Sections.IsValidIndex(Source.BaseSectionIndex))
 	{
 		return;
 	}
 	
-	const FSkelMeshSection& FromSection = ImportedModel->LODModels[LODIndexConnected].Sections[SectionIndexConnected];
-	const TArray<int32>& FromMaterialMap = SkeletalMesh.GetLODInfo(LODIndexConnected)->LODMaterialMap;
+	const FSkelMeshSection& FromSection = ImportedModel->LODModels[Source.BaseLODIndex].Sections[Source.BaseSectionIndex];
+	const TArray<int32>& FromMaterialMap = SkeletalMesh.GetLODInfo(Source.BaseLODIndex)->LODMaterialMap;
 	
 	// Material Index of the connected pin
-	const int32 SearchLODMaterialIndex = FromMaterialMap.IsValidIndex(SectionIndexConnected) && SkeletalMesh.GetMaterials().IsValidIndex(FromMaterialMap[SectionIndexConnected]) ?
-		FromMaterialMap[SectionIndexConnected] : 
+	const int32 SearchLODMaterialIndex = FromMaterialMap.IsValidIndex(Source.BaseSectionIndex) && SkeletalMesh.GetMaterials().IsValidIndex(FromMaterialMap[Source.BaseSectionIndex]) ?
+		FromMaterialMap[Source.BaseSectionIndex] :
 		FromSection.MaterialIndex;
 
-	const int32 CompilingLODIndex = LODIndexConnected + (Context.CurrentLOD - Context.FromLOD);
-	if (!ImportedModel->LODModels.IsValidIndex(CompilingLODIndex))
-	{
-		OutLODIndex = -1;
-		OutSectionIndex = -1;
-		return;
-	}
 
-	const FSkeletalMeshLODModel& LODModel = ImportedModel->LODModels[CompilingLODIndex];
-	const TArray<int32>& MaterialMap = SkeletalMesh.GetLODInfo(CompilingLODIndex)->LODMaterialMap;
+	const int32 CompilingLODIndex = Source.BaseLODIndex + Source.LODOffset;
+	OutLODIndex = INDEX_NONE;
+	OutSectionIndex = INDEX_NONE;
 
-	bool bFound = false;
-	for (int32 SectionIndex = 0; SectionIndex < LODModel.Sections.Num(); ++SectionIndex)
+	if (ImportedModel->LODModels.IsValidIndex(CompilingLODIndex))
 	{
-		const int32 MaterialIndex =  MaterialMap.IsValidIndex(SectionIndex) && SkeletalMesh.GetMaterials().IsValidIndex(MaterialMap[SectionIndex]) ?
-			MaterialMap[SectionIndex] :
-			LODModel.Sections[SectionIndex].MaterialIndex; // MaterialMap overrides the MaterialIndex in the section
-			
-		if (MaterialIndex == SearchLODMaterialIndex &&
-			!LODModel.Sections[SectionIndex].bDisabled)
+		const FSkeletalMeshLODModel& LODModel = ImportedModel->LODModels[CompilingLODIndex];
+		const TArray<int32>& MaterialMap = SkeletalMesh.GetLODInfo(CompilingLODIndex)->LODMaterialMap;
+
+		for (int32 SectionIndex = 0; SectionIndex < LODModel.Sections.Num(); ++SectionIndex)
 		{
-			if (!bFound)
-			{
-				OutLODIndex = CompilingLODIndex;
-				OutSectionIndex = SectionIndex;
-				bFound = true;
-			}
-			else
+			const int32 MaterialIndex = MaterialMap.IsValidIndex(SectionIndex) && SkeletalMesh.GetMaterials().IsValidIndex(MaterialMap[SectionIndex]) ?
+				MaterialMap[SectionIndex] :
+				LODModel.Sections[SectionIndex].MaterialIndex; // MaterialMap overrides the MaterialIndex in the section
+
+			if (MaterialIndex == SearchLODMaterialIndex &&
+				!LODModel.Sections[SectionIndex].bDisabled)
 			{
-				Context.Log(FText::Format(LOCTEXT("MeshMultipleMaterialIndex", "Mesh {0} contains multiple sections with the same Material Index"), FText::FromString(SkeletalMesh.GetName())), &Node);
+				if (OutLODIndex == INDEX_NONE && OutSectionIndex == INDEX_NONE)
+				{
+					OutLODIndex = CompilingLODIndex;
+					OutSectionIndex = SectionIndex;
+				}
+				else 
+				{
+					Context.Log(FText::Format(LOCTEXT("MeshMultipleMaterialIndex", "Mesh {0} contains multiple sections with the same Material Index"), FText::FromString(SkeletalMesh.GetName())), Source.MessageContext);
+				}
 			}
 		}
 	}
-
-	if (!bFound)
-	{
-		OutLODIndex = -1;
-		OutSectionIndex = -1;
-	}
 }
 
 
@@ -873,6 +858,7 @@ TSharedPtr<UE::Mutable::Private::FMesh> ConvertSkeletalMeshToMutable(TSharedRef<
 
 	const USkeletalMesh* SkeletalMesh = Cast<USkeletalMesh>(Source.Mesh.Get());
 	check(SkeletalMesh);
+
 	const FSkeletalMeshModel* ImportedModel = SkeletalMesh->GetImportedModel();
 	if (!ImportedModel)
 	{
@@ -884,9 +870,17 @@ TSharedPtr<UE::Mutable::Private::FMesh> ConvertSkeletalMeshToMutable(TSharedRef<
 	// Data tables may allow missing data. Don't throw warnings of missing LODs or sections.
 	const bool bIgnoreMissingData = Source.TableReferenceSkeletalMesh || Context.Object->GetPrivate()->bDisableTableMissingDataWarning;
 
-	if (!ImportedModel->LODModels.IsValidIndex(Source.LODIndex))
+	// If we are using automatic LODs and not generating the base LOD generate a mesh even if empty.
+	const bool bMeshMustExist = Source.LODOffset == 0;
+
+	// Find the correct LOD and Section indices when using automatic LODs 
+	int32 LODIndex = INDEX_NONE;
+	int32 SectionIndex = INDEX_NONE;
+	GetLODAndSectionForAutomaticLODs(Context, Source, *SkeletalMesh, LODIndex, SectionIndex);
+
+	if (!ImportedModel->LODModels.IsValidIndex(LODIndex))
 	{
-		if (!Source.bMeshMustExist)
+		if (!bMeshMustExist)
 		{
 			return MakeShared<UE::Mutable::Private::FMesh>(); // Return empty mesh to preserve the layouts
 		}
@@ -897,7 +891,7 @@ TSharedPtr<UE::Mutable::Private::FMesh> ConvertSkeletalMeshToMutable(TSharedRef<
 				const FString Msg = FString::Printf(
 					TEXT("The SkeletalMesh [%s] doesn't have the expected number of LODs [need %d, has %d]. Changed after reimporting?"),
 					*MeshName,
-					Source.LODIndex + 1,
+					LODIndex + 1,
 					ImportedModel->LODModels.Num());
 				Context.Log(FText::FromString(Msg), Source.MessageContext);
 			}
@@ -905,10 +899,11 @@ TSharedPtr<UE::Mutable::Private::FMesh> ConvertSkeletalMeshToMutable(TSharedRef<
 		}
 	}
 
-	const FSkeletalMeshLODModel& LODModel = ImportedModel->LODModels[Source.LODIndex];
-	if (!LODModel.Sections.IsValidIndex(Source.SectionIndex))
+
+	const FSkeletalMeshLODModel& LODModel = ImportedModel->LODModels[LODIndex];
+	if (!LODModel.Sections.IsValidIndex(SectionIndex))
 	{
-		if (!Source.bMeshMustExist)
+		if (!bMeshMustExist)
 		{
 			return MakeShared<UE::Mutable::Private::FMesh>(); // Return empty mesh to preserve the layouts
 		}
@@ -919,9 +914,9 @@ TSharedPtr<UE::Mutable::Private::FMesh> ConvertSkeletalMeshToMutable(TSharedRef<
 				const FString Msg = FString::Printf(
 					TEXT("The SkeletalMesh [%s] doesn't have the expected structure. Maybe the number of LODs [need %d, has %d] or Materials [need %d, has %d] has changed after reimporting?"),
 					*SkeletalMesh->GetName(),
-					Source.LODIndex + 1,
+					LODIndex + 1,
 					ImportedModel->LODModels.Num(),
-					Source.SectionIndex + 1,
+					SectionIndex + 1,
 					LODModel.Sections.Num());
 				Context.Log(FText::FromString(Msg), Source.MessageContext);
 			}
@@ -929,17 +924,48 @@ TSharedPtr<UE::Mutable::Private::FMesh> ConvertSkeletalMeshToMutable(TSharedRef<
 		}
 	}
 
-	const FSkelMeshSection& MeshSection = LODModel.Sections[Source.SectionIndex];
+	const FSkelMeshSection& MeshSection = LODModel.Sections[SectionIndex];
 
 	// Get the mesh generation flags to use
 	const EMutableMeshConversionFlags CurrentFlags = Source.Flags;
 	const bool bIgnoreSkinning = EnumHasAnyFlags(CurrentFlags, EMutableMeshConversionFlags::IgnoreSkinning);
 	const bool bIgnorePhysics = EnumHasAnyFlags(CurrentFlags, EMutableMeshConversionFlags::IgnorePhysics);
 	const bool bIgnoreMorphs = EnumHasAnyFlags(CurrentFlags, EMutableMeshConversionFlags::IgnoreMorphs);
+	const bool bIgnoreAUD = EnumHasAnyFlags(CurrentFlags, EMutableMeshConversionFlags::IgnoreAUD);
 	const bool bDoNotCreateMeshMetadata = EnumHasAnyFlags(CurrentFlags, EMutableMeshConversionFlags::DoNotCreateMeshMetadata);
 
 	TSharedPtr<UE::Mutable::Private::FMesh> MutableMesh = MakeShared<UE::Mutable::Private::FMesh>();
 
+	if (UCustomizableObjectSystem::GetInstance()->IsMutableAnimInfoDebuggingEnabled())
+	{
+		FString MeshTag = FString("__MeshPath:") + Source.Mesh->GetPathName();;
+		AddTagToMutableMeshUnique(*MutableMesh, MeshTag);
+	}
+
+	if (!bIgnoreAUD && Context.Object->GetPrivate()->IsAssetUserDataMergeEnabled())
+	{
+		if (const TArray<UAssetUserData*>* AssetUserDataArray = SkeletalMesh->GetAssetUserDataArray())
+		{
+			for (UAssetUserData* AssetUserData : *AssetUserDataArray)
+			{
+				if (!AssetUserData)
+				{
+					continue;
+				}
+
+				const int32 ResourceIndex = Context.AddStreamedResource(*AssetUserData);
+				if (ResourceIndex >= 0)
+				{
+					FCustomizableObjectStreameableResourceId ResourceId;
+					ResourceId.Id = (uint32)ResourceIndex;
+					ResourceId.Type = static_cast<uint8>(FCustomizableObjectStreameableResourceId::EType::AssetUserData);
+
+					MutableMesh->AddStreamedResource(BitCast<uint64>(ResourceId));
+				}
+			}
+		}
+	}
+
 	bool bBoneMapModified = false;
 	TArray<FBoneIndexType> BoneMap;
 	TArray<FBoneIndexType> RemappedBoneMapIndices;
@@ -1024,7 +1050,7 @@ TSharedPtr<UE::Mutable::Private::FMesh> ConvertSkeletalMeshToMutable(TSharedRef<
 
 		// Remove bones and build an array to remap indices of the BoneMap
 		TArray<FBoneIndexType> RemappedBones;
-		BuildRemappedBonesArray(*MutComponentInfo, SkeletalMesh, Source.LODIndex, SourceRequiredBones, RemappedBones);
+		BuildRemappedBonesArray(*MutComponentInfo, SkeletalMesh, LODIndex, SourceRequiredBones, RemappedBones);
 
 		// Build RequiredBones array
 		TArray<FBoneIndexType> RequiredBones;
@@ -1113,6 +1139,28 @@ TSharedPtr<UE::Mutable::Private::FMesh> ConvertSkeletalMeshToMutable(TSharedRef<
 		MutableMesh->SetBoneMap(MutableBoneMap);
 	}
 
+	if (!bIgnoreSkinning)
+	{
+		for (int32 SocketIndex = 0; SocketIndex < SkeletalMesh->NumSockets(); ++SocketIndex)
+		{
+			USkeletalMeshSocket* Socket = SkeletalMesh->GetSocketByIndex(SocketIndex);
+
+			FMutableRefSocket MutableSocket;
+			MutableSocket.SocketName = Socket->SocketName;
+			MutableSocket.BoneName = Socket->BoneName;
+			MutableSocket.RelativeLocation = Socket->RelativeLocation;
+			MutableSocket.RelativeRotation = Socket->RelativeRotation;
+			MutableSocket.RelativeScale = Socket->RelativeScale;
+			MutableSocket.bForceAlwaysAnimated = Socket->bForceAlwaysAnimated;
+			MutableSocket.Priority = Source.SocketPriority;
+
+			int32 SocketArrayIndex = Context.SocketArray.AddUnique(MutableSocket);
+			FString SocketTag = FString::Printf(TEXT("__Socket:%d"), SocketArrayIndex); // TODO PERE: Convert to StreamedResource and remove string
+
+			AddTagToMutableMeshUnique(*MutableMesh, SocketTag);
+		}
+	}
+
 	// Vertices
 	TArray<FSoftSkinVertex> Vertices;
 	LODModel.GetVertices(Vertices);
@@ -1341,16 +1389,16 @@ TSharedPtr<UE::Mutable::Private::FMesh> ConvertSkeletalMeshToMutable(TSharedRef<
 	if (!bIgnoreMorphs && !MorphName.IsEmpty())
 	{
 		UMorphTarget* MorphTarget = SkeletalMesh->FindMorphTarget(*MorphName);
-		if (MorphTarget && MorphTarget->GetMorphLODModels().IsValidIndex(Source.LODIndex))
+		if (MorphTarget && MorphTarget->GetMorphLODModels().IsValidIndex(LODIndex))
 		{
 			// Make a copy of the original vertex data since we will modify it.
 			MorphedVertexData = MakeUnique<FSoftSkinVertex[]>(VertexCount);
 			FMemory::Memcpy(MorphedVertexData.Get(), SourceVertexData, VertexCount*sizeof(FSoftSkinVertex));
 
 			// Bake the morph in the copy of the source vertex data
-			int32 MaterialVertexStart = SkeletalMesh->GetImportedModel()->LODModels[Source.LODIndex].Sections[Source.SectionIndex].GetVertexBufferIndex();
+			int32 MaterialVertexStart = SkeletalMesh->GetImportedModel()->LODModels[LODIndex].Sections[SectionIndex].GetVertexBufferIndex();
 
-			const FMorphTargetLODModel& MorphLODModel = MorphTarget->GetMorphLODModels()[Source.LODIndex];
+			const FMorphTargetLODModel& MorphLODModel = MorphTarget->GetMorphLODModels()[LODIndex];
 			for (const FMorphTargetDelta& MorphDelta : MorphLODModel.Vertices)
 			{
 				const int32 VertexIndex = MorphDelta.SourceIdx - MaterialVertexStart;
@@ -1604,15 +1652,15 @@ TSharedPtr<UE::Mutable::Private::FMesh> ConvertSkeletalMeshToMutable(TSharedRef<
 
 					const TArray<FMorphTargetLODModel>& MorphLODModels = MorphTarget->GetMorphLODModels();
 
-					if (Source.LODIndex >= MorphLODModels.Num() || !MorphLODModels[Source.LODIndex].SectionIndices.Contains(Source.SectionIndex))
+					if (LODIndex >= MorphLODModels.Num() || !MorphLODModels[LODIndex].SectionIndices.Contains(SectionIndex))
 					{
 						continue;
 					}
 
 					// The vertices should be sorted by SourceIdx
-					check(MorphLODModels[Source.LODIndex].Vertices.Num() < 2 || MorphLODModels[Source.LODIndex].Vertices[0].SourceIdx < MorphLODModels[Source.LODIndex].Vertices.Last().SourceIdx);
+					check(MorphLODModels[LODIndex].Vertices.Num() < 2 || MorphLODModels[LODIndex].Vertices[0].SourceIdx < MorphLODModels[LODIndex].Vertices.Last().SourceIdx);
 
-					const int32 VertexFoundIndex = Algo::BinarySearchBy(MorphLODModels[Source.LODIndex].Vertices, (uint32)VertexIdx,
+					const int32 VertexFoundIndex = Algo::BinarySearchBy(MorphLODModels[LODIndex].Vertices, (uint32)VertexIdx,
 						[](const FMorphTargetDelta& Element) { return Element.SourceIdx; });
 
 					if (VertexFoundIndex == INDEX_NONE)
@@ -1620,7 +1668,7 @@ TSharedPtr<UE::Mutable::Private::FMesh> ConvertSkeletalMeshToMutable(TSharedRef<
 						continue;
 					}
 
-					const FMorphTargetDelta& VertexFound = MorphLODModels[Source.LODIndex].Vertices[VertexFoundIndex];
+					const FMorphTargetDelta& VertexFound = MorphLODModels[LODIndex].Vertices[VertexFoundIndex];
 					const FName MorphTargetName = MorphTarget->GetFName();
 
 					int32 MorphTargetNameIndex = NameResolutionMap.Find(MorphTargetName);
@@ -1757,7 +1805,7 @@ TSharedPtr<UE::Mutable::Private::FMesh> ConvertSkeletalMeshToMutable(TSharedRef<
 
 		Invoke([&](int32& OutClothingAssetIndex, int32& OutPhysicsAssetIndex) -> void
 			{
-				const UClothingAssetBase* ClothingAssetBase = SkeletalMesh->GetSectionClothingAsset(Source.LODIndex, Source.SectionIndex);
+				const UClothingAssetBase* ClothingAssetBase = SkeletalMesh->GetSectionClothingAsset(LODIndex, SectionIndex);
 
 				if (!ClothingAssetBase)
 				{
@@ -1956,7 +2004,7 @@ TSharedPtr<UE::Mutable::Private::FMesh> ConvertSkeletalMeshToMutable(TSharedRef<
 			const FCustomizableObjectClothingAssetData& ClothingAssetData = Context.ClothingAssetsData[ClothingAssetIndex];
 
 			ClothingMeshData.ClothingAssetIndex = ClothingAssetIndex;
-			ClothingMeshData.ClothingAssetLOD = ClothingAssetData.LodMap[Source.LODIndex];
+			ClothingMeshData.ClothingAssetLOD = ClothingAssetData.LodMap[LODIndex];
 			ClothingMeshData.PhysicsAssetIndex = PhysicsAssetIndex;
 			ClothingMeshData.Data.Reserve(NumClothMappingDataVerts * ClothDataStride);
 
@@ -2476,11 +2524,11 @@ TSharedPtr<UE::Mutable::Private::FMesh> ConvertSkeletalMeshToMutable(TSharedRef<
 	{
 		FMutableGraphMeshGenerationData MeshData;
 
-		if (Context.Options.bSkinWeightProfilesEnabled && Source.LODIndex >= 0)
+		if (Context.Options.bSkinWeightProfilesEnabled && LODIndex >= 0)
 		{
-			if (ImportedModel && ImportedModel->LODModels.IsValidIndex(Source.LODIndex))
+			if (ImportedModel && ImportedModel->LODModels.IsValidIndex(LODIndex))
 			{
-				for (const auto& SkinWeightProfile : ImportedModel->LODModels[Source.LODIndex].SkinWeightProfiles)
+				for (const auto& SkinWeightProfile : ImportedModel->LODModels[LODIndex].SkinWeightProfiles)
 				{
 					const int32 ProfileId = static_cast<int32>(Context.GetSkinWeightProfileIdUnique(SkinWeightProfile.Key));
 					MeshData.SkinWeightProfilesSemanticIndices.AddUnique(ProfileId);
@@ -2488,15 +2536,15 @@ TSharedPtr<UE::Mutable::Private::FMesh> ConvertSkeletalMeshToMutable(TSharedRef<
 			}
 		}
 
-		if (ImportedModel->LODModels.IsValidIndex(Source.LODIndex) &&
-			ImportedModel->LODModels[Source.LODIndex].Sections.IsValidIndex(Source.SectionIndex))
+		if (ImportedModel->LODModels.IsValidIndex(LODIndex) &&
+			ImportedModel->LODModels[LODIndex].Sections.IsValidIndex(SectionIndex))
 		{
 			MeshData.bHasVertexColors = SkeletalMesh->GetHasVertexColors();
-			MeshData.NumTexCoordChannels = ImportedModel->LODModels[Source.LODIndex].NumTexCoords;
+			MeshData.NumTexCoordChannels = ImportedModel->LODModels[LODIndex].NumTexCoords;
 			MeshData.MaxBoneIndexTypeSizeBytes = MutableMesh->BoneMap.Num() > 256 ? 2 : 1;
-			MeshData.MaxNumBonesPerVertex = ImportedModel->LODModels[Source.LODIndex].GetMaxBoneInfluences();
+			MeshData.MaxNumBonesPerVertex = ImportedModel->LODModels[LODIndex].GetMaxBoneInfluences();
 			MeshData.bHasRealTimeMorphs = Context.Options.bRealTimeMorphTargetsEnabled && SkeletalMesh->GetMorphTargets().Num();
-			MeshData.bHasClothing = Context.Options.bClothingEnabled && ImportedModel->LODModels[Source.LODIndex].HasClothData();
+			MeshData.bHasClothing = Context.Options.bClothingEnabled && ImportedModel->LODModels[LODIndex].HasClothData();
 		}
 
 		TSharedPtr<UE::Mutable::Private::FMesh> ForcedFormat = MakeShared<UE::Mutable::Private::FMesh>();
@@ -2528,9 +2576,9 @@ TSharedPtr<UE::Mutable::Private::FMesh> ConvertSkeletalMeshToMutable(TSharedRef<
 
 TSharedPtr<UE::Mutable::Private::FMesh> ConvertSkeletalMeshToMutable(
 	FMutableSourceMeshData& Source,
-	bool bMeshMustExist, int32 LODIndex, int32 SectionIndex,
+	int32 LODIndex, int32 SectionIndex,
 	FMutableGraphGenerationContext& GenerationContext, 
-	const UCustomizableObjectNode* CurrentNode, 
+	const UCustomizableObjectNode* CurrentNode,
 	bool bForceImmediateConversion)
 {
 	MUTABLE_CPUPROFILER_SCOPE(ConvertSkeletalMeshToMutable);
@@ -2541,9 +2589,9 @@ TSharedPtr<UE::Mutable::Private::FMesh> ConvertSkeletalMeshToMutable(
 	}
 
 	// Prepare the data that is needed for the core mesh conversion.
-	Source.bMeshMustExist = bMeshMustExist;
-	Source.LODIndex = LODIndex;
-	Source.SectionIndex = SectionIndex;
+	Source.BaseLODIndex= LODIndex;
+	Source.BaseSectionIndex = SectionIndex;
+	Source.LODOffset = GenerationContext.CurrentLOD - GenerationContext.FromLOD;
 	Source.Flags = GenerationContext.MeshGenerationFlags.Last();
 	Source.MessageContext = CurrentNode;
 	Source.Component = GenerationContext.CurrentMeshComponent;
@@ -2772,27 +2820,20 @@ TSharedPtr<UE::Mutable::Private::FMesh> ConvertStaticMeshToMutable(const UStatic
 // Convert a Mesh constant to a mutable format. UniqueTags are the tags that make this Mesh unique that cannot be merged in the cache 
 // with the exact same Mesh with other tags
 TSharedPtr<UE::Mutable::Private::FMesh> GenerateMutableSkeletalMesh(FMutableSourceMeshData& Source, int32 LODIndexConnected, int32 SectionIndexConnected,
-									  int32 LODIndex, int32 SectionIndex, const FString& UniqueTags, FMutableGraphGenerationContext & GenerationContext, 
-									  const UCustomizableObjectNode* CurrentNode)
+										const FString& UniqueTags, FMutableGraphGenerationContext & GenerationContext, 
+										const UCustomizableObjectNode* CurrentNode,
+										const bool bOnlyConnectedLOD)
 {
 	// Get the mesh generation flags to use
 	EMutableMeshConversionFlags CurrentFlags = GenerationContext.MeshGenerationFlags.Last();
 
-	FMutableGraphGenerationContext::FGeneratedMeshData::FKey Key = { Source.Mesh, LODIndex, GenerationContext.CurrentLOD, SectionIndex, CurrentFlags, UniqueTags, CurrentNode };
+	FMutableGraphGenerationContext::FGeneratedMeshData::FKey Key = { Source.Mesh, LODIndexConnected, GenerationContext.CurrentLOD, SectionIndexConnected, CurrentFlags, UniqueTags, CurrentNode };
 	TSharedPtr<UE::Mutable::Private::FMesh> MutableMesh = GenerationContext.FindGeneratedMesh(Key);
 	if (MutableMesh)
 	{
 		return MutableMesh;
 	}
 
-	// If we are using automatic LODs and not generating the base LOD (the connected one) is not an error.
-	bool bMeshMustExist =
-		// Are we using automatic LODs?
-		GenerationContext.CurrentAutoLODStrategy != ECustomizableObjectAutomaticLODStrategy::AutomaticFromMesh
-		||
-		// Are we generating the base LOD?
-		((LODIndex == LODIndexConnected) && (SectionIndex == SectionIndexConnected));
-
 	// At some point we will want all meshes to be references at compile-time. For now, just create the actual pass-through meshes.
 	if (Source.bIsPassthrough)
 	{
@@ -2800,7 +2841,8 @@ TSharedPtr<UE::Mutable::Private::FMesh> GenerateMutableSkeletalMesh(FMutableSour
 	}
 	else
 	{
-		MutableMesh = ConvertSkeletalMeshToMutable(Source, bMeshMustExist, LODIndex, SectionIndex, GenerationContext, CurrentNode);
+		Source.bOnlyConnectedLOD = bOnlyConnectedLOD;
+		MutableMesh = ConvertSkeletalMeshToMutable(Source, LODIndexConnected, SectionIndexConnected, GenerationContext, CurrentNode);
 	}
 
 	if (MutableMesh)
@@ -2859,11 +2901,8 @@ TSharedPtr<UE::Mutable::Private::FMesh> BuildMorphedMutableMesh(const UEdGraphPi
 		return nullptr;
 	}
 
-	int32 LODIndexConnected = -1; // LOD which the pin is connected to
-	int32 SectionIndexConnected = -1;
-
-	int32 LODIndex = -1; // Initialization required to remove uninitialized warning.
-	int32 SectionIndex = -1;
+	int32 BaseLODIndex = -1; // LOD which the pin is connected to
+	int32 BaseSectionIndex = -1;
 
 	USkeletalMesh* SkeletalMesh = nullptr;
 	UCustomizableObjectNode* Node = Cast<UCustomizableObjectNode>(BaseSourcePin->GetOwningNode());
@@ -2875,7 +2914,7 @@ TSharedPtr<UE::Mutable::Private::FMesh> BuildMorphedMutableMesh(const UEdGraphPi
 
 	else if (const ICustomizableObjectNodeMeshInterface* TypedNodeSkeletalMesh = Cast<ICustomizableObjectNodeMeshInterface>(Node))
 	{
-		TypedNodeSkeletalMesh->GetPinSection(*BaseSourcePin, LODIndexConnected, SectionIndexConnected);
+		TypedNodeSkeletalMesh->GetPinSection(*BaseSourcePin, BaseLODIndex, BaseSectionIndex);
 		SkeletalMesh = Cast<USkeletalMesh>(UE::Mutable::Private::LoadObject(TypedNodeSkeletalMesh->GetMesh()));
 	}
 	
@@ -2885,7 +2924,7 @@ TSharedPtr<UE::Mutable::Private::FMesh> BuildMorphedMutableMesh(const UEdGraphPi
 
 		if (DataTable)
 		{
-			TypedNodeTable->GetPinLODAndSection(BaseSourcePin, LODIndexConnected, SectionIndexConnected);
+			TypedNodeTable->GetPinLODAndSection(BaseSourcePin, BaseLODIndex, BaseSectionIndex);
 			SkeletalMesh = Cast<USkeletalMesh>(GenerationContext.LoadObject(TypedNodeTable->GetSkeletalMeshAt(BaseSourcePin, DataTable, RowName)));
 		}
 	}
@@ -2899,12 +2938,10 @@ TSharedPtr<UE::Mutable::Private::FMesh> BuildMorphedMutableMesh(const UEdGraphPi
 
 	if (SkeletalMesh)
 	{
-		GetLODAndSectionForAutomaticLODs(GenerationContext, *Node, *SkeletalMesh, LODIndexConnected, SectionIndexConnected, LODIndex, SectionIndex, bOnlyConnectedLOD);
 		// Get the base mesh
 		FMutableSourceMeshData Source;
 		Source.Mesh = SkeletalMesh;
-		TSharedPtr<UE::Mutable::Private::FMesh> BaseSourceMesh = GenerateMutableSkeletalMesh(Source, LODIndexConnected, SectionIndexConnected,
-			LODIndex, SectionIndex, FString(), GenerationContext, Node);
+		TSharedPtr<UE::Mutable::Private::FMesh> BaseSourceMesh = GenerateMutableSkeletalMesh(Source, BaseLODIndex, BaseSectionIndex, FString(), GenerationContext, Node, bOnlyConnectedLOD);
 
 		if (BaseSourceMesh)
 		{
@@ -3726,187 +3763,117 @@ UE::Mutable::Private::Ptr<UE::Mutable::Private::NodeMesh> GenerateMutableSourceM
 	
 	if (const UCustomizableObjectNodeSkeletalMesh* TypedNodeSkel = Cast<UCustomizableObjectNodeSkeletalMesh>(Node))
 	{
+		MUTABLE_CPUPROFILER_SCOPE(GenerateMutableSourceMesh_SkeletalMesh);
 		UE::Mutable::Private::Ptr<UE::Mutable::Private::NodeMeshConstant> MeshNode = new UE::Mutable::Private::NodeMeshConstant();
 		Result = MeshNode;
 
-		if (TypedNodeSkel->SkeletalMesh.IsNull())
+		if (!TypedNodeSkel->SkeletalMesh.IsNull())
 		{
-			GenerationContext.Log(LOCTEXT("MissingSkeletalMesh", "No Skeletal Mesh set in the SkeletalMesh node."), Node);
-		}
-		else
-		{
-			// TODO: Avoid loading the mesh here.
-			USkeletalMesh* SkeletalMesh = GenerationContext.LoadObject(TypedNodeSkel->SkeletalMesh);
+			int32 BaseLODIndex = -1; // LOD which the pin is connected to
+			int32 BaseSectionIndex = -1;
+			TypedNodeSkel->GetPinSection(*Pin, BaseLODIndex, BaseSectionIndex);
 
-			if (SkeletalMesh)
+			// First process the mesh tags that are going to make the mesh unique and affect whether it's repeated in 
+			// the mesh cache or not
+			FString MeshUniqueTags;
+			FString AnimBPAssetTag;
+
+			if (!TypedNodeSkel->AnimInstance.IsNull())
 			{
-				int32 LODIndexConnected = -1; // LOD which the pin is connected to
-				int32 SectionIndexConnected = -1;
+				FName SlotIndex = TypedNodeSkel->AnimBlueprintSlotName;
+				const int32 AnimInstanceIndex = GenerationContext.AnimBPAssets.AddUnique(TypedNodeSkel->AnimInstance);
 
-				int32 LODIndex = -1;
-				int32 SectionIndex = -1;
+				AnimBPAssetTag = GenerateAnimationInstanceTag(AnimInstanceIndex, SlotIndex);
+				MeshUniqueTags += AnimBPAssetTag;
+			}
 
-				{
-					TypedNodeSkel->GetPinSection(*Pin, LODIndexConnected, SectionIndexConnected);
-				}
+			TArray<FString> ArrayAnimBPTags;
 
+			for (const FGameplayTag& GamePlayTag : TypedNodeSkel->AnimationGameplayTags)
+			{
+				const FString AnimBPTag = GenerateGameplayTag(GamePlayTag.ToString());
+				ArrayAnimBPTags.Add(AnimBPTag);
+				MeshUniqueTags += AnimBPTag;
+			}
 
-				GetLODAndSectionForAutomaticLODs(GenerationContext, *Node, *SkeletalMesh, LODIndexConnected, SectionIndexConnected, LODIndex, SectionIndex, bOnlyConnectedLOD);
+			FMutableSourceMeshData Source = BaseMeshData;
+			Source.Mesh = TypedNodeSkel->SkeletalMesh;
+			Source.AnimInstance = TypedNodeSkel->AnimInstance;
+			TSharedPtr<UE::Mutable::Private::FMesh > MutableMesh = GenerateMutableSkeletalMesh(Source, BaseLODIndex, BaseSectionIndex, MeshUniqueTags, GenerationContext, TypedNodeSkel, bOnlyConnectedLOD);
 
-				// First process the mesh tags that are going to make the mesh unique and affect whether it's repeated in 
-				// the mesh cache or not
-				FString MeshUniqueTags;
-				FString AnimBPAssetTag;
+			if (MutableMesh)
+			{
+				MeshNode->Value = MutableMesh;
 
 				if (!TypedNodeSkel->AnimInstance.IsNull())
 				{
-					FName SlotIndex = TypedNodeSkel->AnimBlueprintSlotName;
-					const int32 AnimInstanceIndex = GenerationContext.AnimBPAssets.AddUnique(TypedNodeSkel->AnimInstance);
-
-					AnimBPAssetTag = GenerateAnimationInstanceTag(AnimInstanceIndex, SlotIndex);
-					MeshUniqueTags += AnimBPAssetTag;
+					AddTagToMutableMeshUnique(*MutableMesh, AnimBPAssetTag);
 				}
 
-				TArray<FString> ArrayAnimBPTags;
-
-				for (const FGameplayTag& GamePlayTag : TypedNodeSkel->AnimationGameplayTags)
+				for (const FString& GamePlayTag : ArrayAnimBPTags)
 				{
-					const FString AnimBPTag = GenerateGameplayTag(GamePlayTag.ToString());
-					ArrayAnimBPTags.Add(AnimBPTag);
-					MeshUniqueTags += AnimBPTag;
+					AddTagToMutableMeshUnique(*MutableMesh, GamePlayTag);
 				}
+			}
 
-				TArray<FCustomizableObjectStreameableResourceId> StreamedResources;
-
-				if (GenerationContext.CompilationContext->Object->GetPrivate()->IsAssetUserDataMergeEnabled())
-				{
-					const TArray<UAssetUserData*>* AssetUserDataArray = SkeletalMesh->GetAssetUserDataArray();
-
-					if (AssetUserDataArray)
-					{
-						for (UAssetUserData* AssetUserData : *AssetUserDataArray)
-						{
-							if (!AssetUserData)
-							{
-								continue;
-							}
+			// Layouts
+			// 
+			// Generate Layouts when using manual LODs or when processing the base Automatic LOD (the other automatic LODs will use the base layouts).
+			const bool bGenerateLayouts = GenerationContext.CurrentAutoLODStrategy != ECustomizableObjectAutomaticLODStrategy::AutomaticFromMesh
+				|| GenerationContext.CurrentLOD == GenerationContext.FromLOD;
 
-							const int32 ResourceIndex = GenerationContext.AddStreamedResource(*AssetUserData);
-							if (ResourceIndex >= 0)
-							{
-								FCustomizableObjectStreameableResourceId ResourceId;
-								ResourceId.Id = (uint32)ResourceIndex;
-								ResourceId.Type = static_cast<uint8>(FCustomizableObjectStreameableResourceId::EType::AssetUserData);
+			if(bGenerateLayouts)
+			{
+				const FLayoutGenerationFlags& LayoutFlags = GenerationContext.LayoutGenerationFlags.Last();
 
-								StreamedResources.Add(ResourceId);
-							}
+				const TArray<UCustomizableObjectLayout*>& Layouts = TypedNodeSkel->GetLayouts(*Pin);
+				const int32 NumLayouts = Layouts.Num();
 
-							MeshUniqueTags += AssetUserData->GetPathName();
-						}
-					}
-				}
+				MeshNode->Layouts.SetNum(NumLayouts);
 
-				FMutableSourceMeshData Source = BaseMeshData;
-				Source.Mesh = SkeletalMesh;
-				Source.AnimInstance = TypedNodeSkel->AnimInstance;
-
-				TSharedPtr<UE::Mutable::Private::FMesh> MutableMesh = GenerateMutableSkeletalMesh(Source, LODIndexConnected, SectionIndexConnected,
-					LODIndex, SectionIndex, MeshUniqueTags, GenerationContext, TypedNodeSkel);
-				if (MutableMesh)
+				for (int32 LayoutIndex = 0; LayoutIndex < NumLayouts; ++LayoutIndex)
 				{
-					MeshNode->Value = MutableMesh;
-
-					if (!TypedNodeSkel->AnimInstance.IsNull())
-					{
-						AddTagToMutableMeshUnique(*MutableMesh, AnimBPAssetTag);
-					}
-
-					for (const FString& GamePlayTag : ArrayAnimBPTags)
-					{
-						AddTagToMutableMeshUnique(*MutableMesh, GamePlayTag);
-					}
-
-					for (FCustomizableObjectStreameableResourceId ResourceId : StreamedResources)
+					if (!LayoutFlags.TexturePinModes.IsValidIndex(LayoutIndex) ||
+						LayoutFlags.TexturePinModes[LayoutIndex] != EPinMode::Mutable)
 					{
-						MutableMesh->AddStreamedResource(BitCast<uint64>(ResourceId));
-					}
+						MeshNode->Layouts[LayoutIndex] = CreateDefaultLayout();
 
-					AddSocketTagsToMesh(SkeletalMesh, MutableMesh, GenerationContext);
+						if (Layouts.IsValidIndex(LayoutIndex) && Layouts[LayoutIndex])
+						{
+							// Keep packing strategy if possible, Overlay can be valid with EPinMode == Passthrough	
+							MeshNode->Layouts[LayoutIndex]->Strategy = ConvertLayoutStrategy(Layouts[LayoutIndex]->PackingStrategy);
+							MeshNode->Layouts[LayoutIndex]->TexCoordsIndex = Layouts[LayoutIndex]->GetUVChannel();
+						}
 
-					if (UCustomizableObjectSystem::GetInstance()->IsMutableAnimInfoDebuggingEnabled())
-					{
-						FString MeshPath;
-						SkeletalMesh->GetOuter()->GetPathName(nullptr, MeshPath);
-						FString MeshTag = FString("__MeshPath:") + MeshPath;
-						AddTagToMutableMeshUnique(*MutableMesh, MeshTag);
+						// Ignore layout
+						continue;
 					}
-				}
 
-				// Layouts
-				{
-					// When using Automatic From Mesh all LODs share the same base layout, hence we use LODIndexConnected (as the base layout) instead of the LODIndex.
-					const int32 LODIndexLayout = GenerationContext.CurrentAutoLODStrategy == ECustomizableObjectAutomaticLODStrategy::AutomaticFromMesh ?
-						LODIndexConnected :
-						LODIndex;
-					const int32 SectionIndexLayout = GenerationContext.CurrentAutoLODStrategy == ECustomizableObjectAutomaticLODStrategy::AutomaticFromMesh ?
-						SectionIndexConnected :
-						SectionIndex;
-
-
-					const FLayoutGenerationFlags& LayoutFlags = GenerationContext.LayoutGenerationFlags.Last();
-
-					const TArray<UCustomizableObjectLayout*>& Layouts = TypedNodeSkel->GetLayouts(*Pin);
-					const int32 NumLayouts = Layouts.Num();
-
-					MeshNode->Layouts.SetNum(NumLayouts);
-
-					for (int32 LayoutIndex = 0; LayoutIndex < NumLayouts; ++LayoutIndex)
+					const UCustomizableObjectLayout* Layout = Layouts.IsValidIndex(LayoutIndex) ? Layouts[LayoutIndex] : nullptr;
+					if (ensure(Layout))
 					{
-						if (!LayoutFlags.TexturePinModes.IsValidIndex(LayoutIndex) ||
-							LayoutFlags.TexturePinModes[LayoutIndex] != EPinMode::Mutable)
-						{
-							MeshNode->Layouts[LayoutIndex] = CreateDefaultLayout();
-
-							if (Layouts.IsValidIndex(LayoutIndex) && Layouts[LayoutIndex])
-							{
-								// Keep packing strategy if possible, Overlay can be valid with EPinMode == Passthrough	
-								MeshNode->Layouts[LayoutIndex]->Strategy = ConvertLayoutStrategy(Layouts[LayoutIndex]->PackingStrategy);
-								MeshNode->Layouts[LayoutIndex]->TexCoordsIndex = Layouts[LayoutIndex]->GetUVChannel();
-							}
+						UE::Mutable::Private::Ptr<UE::Mutable::Private::NodeLayout> LayoutNode = CreateMutableLayoutNode(Layout, bLinkedToExtendMaterial); // TODO PERE: Figure out
+						LayoutNode->SetMessageContext(Node);
 
-							// Ignore layout
-							continue;
-						}
-
-						const UCustomizableObjectLayout* Layout = Layouts.IsValidIndex(LayoutIndex) ? Layouts[LayoutIndex] : nullptr;
-						if (ensure(Layout))
-						{
-							UE::Mutable::Private::Ptr<UE::Mutable::Private::NodeLayout> LayoutNode = CreateMutableLayoutNode(GenerationContext, Layout, bLinkedToExtendMaterial); // TODO PERE: Figure out
-							LayoutNode->SetMessageContext(Node);
-
-							MeshNode->Layouts[LayoutIndex] = LayoutNode;
-						}
+						MeshNode->Layouts[LayoutIndex] = LayoutNode;
 					}
 				}
+			}
 
+			const FString MeshName = TypedNodeSkel->SkeletalMesh.GetLongPackageName().ToLower();
+			MeshNode->SourceDataDescriptor.SourceId = CityHash32(reinterpret_cast<const char*>(*MeshName), MeshName.Len() * sizeof(FString::ElementType));
+			MeshNode->SourceDataDescriptor.OptionalMaxLODSize = 0;
 
-				const FString MeshName = TypedNodeSkel->SkeletalMesh.GetLongPackageName().ToLower();
-				MeshNode->SourceDataDescriptor.SourceId = CityHash32(reinterpret_cast<const char*>(*MeshName), MeshName.Len() * sizeof(FString::ElementType));
-				MeshNode->SourceDataDescriptor.OptionalMaxLODSize = 0;
-
-				// Applying Mesh Morph Nodes
-				if (GenerationContext.MeshMorphStack.Num())
-				{
-					MorphResult = GenerateMorphMesh(Pin, GenerationContext.MeshMorphStack, 0, Result, GenerationContext, bOnlyConnectedLOD);
-				}
-			}
-			else
+			// Applying Mesh Morph Nodes
+			if (GenerationContext.MeshMorphStack.Num())
 			{
-				FText Path = TypedNodeSkel->SkeletalMesh.GetLongPackageFName().IsNone() ? FText::FromString("Unknown") : FText::FromName(TypedNodeSkel->SkeletalMesh.GetLongPackageFName());
-				FText Msg = FText::Format(LOCTEXT("SkeletalMeshNotLoaded", "Failed to load the skeletal mesh set in the node.[Asset path: {0}]"), Path);
-				GenerationContext.Log(Msg, Node);
+				MorphResult = GenerateMorphMesh(Pin, GenerationContext.MeshMorphStack, 0, Result, GenerationContext, bOnlyConnectedLOD);
 			}
 		}
+		else
+		{
+			GenerationContext.Log(LOCTEXT("NodeSkeletalMeshEmpty", "Node SkeletalMesh without a mesh assigned."), Node);
+		}
 	}
 
 	else if (const UCustomizableObjectNodeStaticMesh* TypedNodeStatic = Cast<UCustomizableObjectNodeStaticMesh>(Node))
@@ -3971,7 +3938,7 @@ UE::Mutable::Private::Ptr<UE::Mutable::Private::NodeMesh> GenerateMutableSourceM
 				
 				if (Layout)
 				{
-					LayoutNode = CreateMutableLayoutNode(GenerationContext, Layout, false);
+					LayoutNode = CreateMutableLayoutNode(Layout, false);
 				}
 				else
 				{
@@ -3995,6 +3962,8 @@ UE::Mutable::Private::Ptr<UE::Mutable::Private::NodeMesh> GenerateMutableSourceM
 
 	else if (UCustomizableObjectNodeMeshMorph* TypedNodeMorph = Cast<UCustomizableObjectNodeMeshMorph>(Node))
 	{
+		MUTABLE_CPUPROFILER_SCOPE(GenerateMutableSourceMesh_MeshMorph);
+
 		if (const UEdGraphPin* ConnectedPin = FollowInputPin(*TypedNodeMorph->MeshPin()))
 		{
 			FString MorphTargetName = TypedNodeMorph->MorphTargetName;
@@ -4035,6 +4004,8 @@ UE::Mutable::Private::Ptr<UE::Mutable::Private::NodeMesh> GenerateMutableSourceM
 
 	else if (const UCustomizableObjectNodeMeshMorphStackApplication* TypedNodeMeshMorphStackApp = Cast< UCustomizableObjectNodeMeshMorphStackApplication >(Node))
 	{
+		MUTABLE_CPUPROFILER_SCOPE(GenerateMutableSourceMesh_MeshMorphStack);
+
 		Result = nullptr;
 
 		if (const UEdGraphPin* ConnectedPin = FollowInputPin(*TypedNodeMeshMorphStackApp->GetStackPin()))
@@ -4171,6 +4142,7 @@ UE::Mutable::Private::Ptr<UE::Mutable::Private::NodeMesh> GenerateMutableSourceM
 
 	else if (const UCustomizableObjectNodeMeshReshape* TypedNodeReshape = Cast<const UCustomizableObjectNodeMeshReshape>(Node))
 	{
+		MUTABLE_CPUPROFILER_SCOPE(GenerateMutableSourceMesh_Reshape);
 		UE::Mutable::Private::Ptr<UE::Mutable::Private::NodeMeshReshape> MeshNode = new UE::Mutable::Private::NodeMeshReshape();
 		Result = MeshNode;
 
@@ -4294,7 +4266,9 @@ UE::Mutable::Private::Ptr<UE::Mutable::Private::NodeMesh> GenerateMutableSourceM
 		const EMutableMeshConversionFlags ShapeFlags = 
 				EMutableMeshConversionFlags::IgnoreSkinning |
 				EMutableMeshConversionFlags::IgnorePhysics |
-			EMutableMeshConversionFlags::DoNotCreateMeshMetadata;
+				EMutableMeshConversionFlags::IgnoreTexCoords |
+				EMutableMeshConversionFlags::IgnoreAUD |
+				EMutableMeshConversionFlags::DoNotCreateMeshMetadata;
 
 		GenerationContext.MeshGenerationFlags.Push( ShapeFlags );
 			
@@ -4383,6 +4357,8 @@ UE::Mutable::Private::Ptr<UE::Mutable::Private::NodeMesh> GenerateMutableSourceM
 
 	else if (const UCustomizableObjectNodeTable* TypedNodeTable = Cast<UCustomizableObjectNodeTable>(Node))
 	{
+		MUTABLE_CPUPROFILER_SCOPE(GenerateMutableSourceMesh_NodeTable);
+
 		TObjectPtr<UDataTable> Table = UE::Mutable::Private::LoadObject(TypedNodeTable->Table);
 		TObjectPtr<UScriptStruct> Structure = UE::Mutable::Private::LoadObject(TypedNodeTable->Structure);
 		const FString TableName = Table ? GetNameSafe(Table).ToLower() : GetNameSafe(Structure).ToLower();
@@ -4432,10 +4408,8 @@ UE::Mutable::Private::Ptr<UE::Mutable::Private::NodeMesh> GenerateMutableSourceM
 				{
 					UE::Mutable::Private::Ptr<UE::Mutable::Private::NodeMeshTable> MeshTableNode = new UE::Mutable::Private::NodeMeshTable();
 					
-					int32 LODIndexConnected = -1; // LOD which the pin is connected to
-					int32 SectionIndexConnected = -1;
-					int32 LODIndex = 0;
-					int32 SectionIndex = 0;
+					int32 BaseLODIndex = INDEX_NONE; // LOD which the pin is connected to
+					int32 BaseSectionIndex = INDEX_NONE;
 
 					// Getting the mutable table mesh column name
 					FString MutableColumnName = ColumnName;
@@ -4443,23 +4417,38 @@ UE::Mutable::Private::Ptr<UE::Mutable::Private::NodeMesh> GenerateMutableSourceM
 					if (Pin->PinType.PinCategory == Schema->PC_Mesh)
 					{
 						// LOD and sections are relevant for Skeletal and Static meshes but not for Pose Assets
-						TypedNodeTable->GetPinLODAndSection(Pin, LODIndexConnected, SectionIndexConnected);
+						TypedNodeTable->GetPinLODAndSection(Pin, BaseLODIndex, BaseSectionIndex);
 
 						if (DefaultSkeletalMesh)
 						{
-							GetLODAndSectionForAutomaticLODs(GenerationContext, *Node, *DefaultSkeletalMesh, LODIndexConnected, SectionIndexConnected, LODIndex, SectionIndex, bOnlyConnectedLOD);
+							int32 LODIndex = 0;
+							int32 SectionIndex = 0;
+
+							FMutableSourceMeshData Source;
+							Source.BaseLODIndex = BaseLODIndex;
+							Source.BaseSectionIndex = BaseSectionIndex;
+							Source.LODOffset = GenerationContext.CurrentLOD - GenerationContext.FromLOD;
+							Source.MessageContext = Node;
+							Source.bOnlyConnectedLOD = bOnlyConnectedLOD;
+							GetLODAndSectionForAutomaticLODs(*GenerationContext.CompilationContext, Source, *DefaultSkeletalMesh, LODIndex, SectionIndex);
+
 							MutableColumnName = TypedNodeTable->GenerateSkeletalMeshMutableColumName(ColumnName, LODIndex, SectionIndex);
 						}
 						else
 						{
-							MutableColumnName = TypedNodeTable->GenerateStaticMeshMutableColumName(ColumnName, SectionIndexConnected);
+							MutableColumnName = TypedNodeTable->GenerateStaticMeshMutableColumName(ColumnName, BaseSectionIndex);
 						}
 					}
 					
 					// Generating a new FMesh column if not exists
 					if (GeneratedTable->FindColumn(MutableColumnName) == INDEX_NONE)
 					{
-						bSuccess = GenerateTableColumn(TypedNodeTable, Pin, GeneratedTable, ColumnName, Property, BaseMeshData, LODIndexConnected, SectionIndexConnected, LODIndex, SectionIndex, bOnlyConnectedLOD, GenerationContext);
+						FMutableSourceMeshData Source = BaseMeshData;
+						Source.BaseLODIndex = BaseLODIndex;
+						Source.BaseSectionIndex = BaseSectionIndex;
+						Source.LODOffset = GenerationContext.CurrentLOD - GenerationContext.FromLOD;
+						Source.bOnlyConnectedLOD = bOnlyConnectedLOD; // True?
+						bSuccess = GenerateTableColumn(TypedNodeTable, Pin, GeneratedTable, ColumnName, Property, Source, GenerationContext);
 
 						if (!bSuccess)
 						{
@@ -4485,34 +4474,44 @@ UE::Mutable::Private::Ptr<UE::Mutable::Private::NodeMesh> GenerateMutableSourceM
 						{
 							const FLayoutGenerationFlags& LayoutFlags = GenerationContext.LayoutGenerationFlags.Last();
 
-							TArray<UCustomizableObjectLayout*> Layouts = TypedNodeTable->GetLayouts(Pin);
-							const int32 NumLayouts = Layouts.Num();
+							// Layouts
+							// 
+							// Generate Layouts when using manual LODs or when processing the base Automatic LOD (the other automatic LODs will use the base layouts).
+							const bool bGenerateLayouts = GenerationContext.CurrentAutoLODStrategy != ECustomizableObjectAutomaticLODStrategy::AutomaticFromMesh
+								|| GenerationContext.CurrentLOD == GenerationContext.FromLOD;
 
-							// Generating node Layouts
-							MeshTableNode->Layouts.SetNum(NumLayouts);
-							for (int32 LayoutIndex = 0; LayoutIndex < NumLayouts; ++LayoutIndex)
+							if (true || bGenerateLayouts)
 							{
-								if (!LayoutFlags.TexturePinModes.IsValidIndex(LayoutIndex) ||
-									LayoutFlags.TexturePinModes[LayoutIndex] != EPinMode::Mutable)
+
+								TArray<UCustomizableObjectLayout*> Layouts = TypedNodeTable->GetLayouts(Pin);
+								const int32 NumLayouts = Layouts.Num();
+
+								// Generating node Layouts
+								MeshTableNode->Layouts.SetNum(NumLayouts);
+								for (int32 LayoutIndex = 0; LayoutIndex < NumLayouts; ++LayoutIndex)
 								{
-									MeshTableNode->Layouts[LayoutIndex] = CreateDefaultLayout();
-							
-									// Keep packing strategy if possible, Overlay can be valid with EPinMode == Passthrough	
-									if (Layouts[LayoutIndex])
+									if (!LayoutFlags.TexturePinModes.IsValidIndex(LayoutIndex) ||
+										LayoutFlags.TexturePinModes[LayoutIndex] != EPinMode::Mutable)
 									{
-										MeshTableNode->Layouts[LayoutIndex]->Strategy = ConvertLayoutStrategy(Layouts[LayoutIndex]->PackingStrategy);
-										MeshTableNode->Layouts[LayoutIndex]->TexCoordsIndex = Layouts[LayoutIndex]->GetUVChannel();
+										MeshTableNode->Layouts[LayoutIndex] = CreateDefaultLayout();
+
+										// Keep packing strategy if possible, Overlay can be valid with EPinMode == Passthrough	
+										if (Layouts[LayoutIndex])
+										{
+											MeshTableNode->Layouts[LayoutIndex]->Strategy = ConvertLayoutStrategy(Layouts[LayoutIndex]->PackingStrategy);
+											MeshTableNode->Layouts[LayoutIndex]->TexCoordsIndex = Layouts[LayoutIndex]->GetUVChannel();
+										}
+
+										// Ignore layouts
+										continue;
 									}
-							
-									// Ignore layouts
-									continue;
+
+									// In tables, mimic the legacy behaviour and ignore all layout warnings beyond LOD 0.
+									bool bIgnoreLayoutWarnings = true;
+									UE::Mutable::Private::Ptr<UE::Mutable::Private::NodeLayout> LayoutNode = CreateMutableLayoutNode(Layouts[LayoutIndex], bIgnoreLayoutWarnings);
+
+									MeshTableNode->Layouts[LayoutIndex] = LayoutNode;
 								}
-							
-								// In tables, mimic the legacy behaviour and ignore all layout warnings beyond LOD 0.
-								bool bIgnoreLayoutWarnings = true;
-								UE::Mutable::Private::Ptr<UE::Mutable::Private::NodeLayout> LayoutNode = CreateMutableLayoutNode(GenerationContext, Layouts[LayoutIndex], bIgnoreLayoutWarnings);
-							
-								MeshTableNode->Layouts[LayoutIndex] = LayoutNode;
 							}
 
 							// Applying Mesh Morph Nodes
@@ -4548,13 +4547,13 @@ UE::Mutable::Private::Ptr<UE::Mutable::Private::NodeMesh> GenerateMutableSourceM
 			GenerationContext.SkeletalMeshParameterDefaultValues.Add(FName(TypedNodeParam->GetParameterName()), UE::Mutable::Private::LoadObject(TypedNodeParam->DefaultValue));
 		}
 
-		int32 LODIndex = 0;
-		int32 SectionIndex = 0;
-		
+		int32 BaseLODIndex = 0;
+		int32 BaseSectionIndex = 0;
+
 		{
-			TypedNodeParam->GetPinSection(*Pin, LODIndex, SectionIndex);
-			ParameterNode->LODIndex = LODIndex;
-			ParameterNode->SectionIndex = SectionIndex;
+			TypedNodeParam->GetPinSection(*Pin, BaseLODIndex, BaseSectionIndex);
+			ParameterNode->LODIndex = BaseLODIndex;
+			ParameterNode->SectionIndex = BaseSectionIndex;
 			ParameterNode->bAutomaticLODs = GenerationContext.CurrentAutoLODStrategy == ECustomizableObjectAutomaticLODStrategy::AutomaticFromMesh;
 		}
 		
@@ -4564,10 +4563,6 @@ UE::Mutable::Private::Ptr<UE::Mutable::Private::NodeMesh> GenerateMutableSourceM
 
 		// Layouts
 		{
-			const int32 LODIndexLayout = LODIndex;
-			const int32 SectionIndexLayout = SectionIndex;
-
-
 			const FLayoutGenerationFlags& LayoutFlags = GenerationContext.LayoutGenerationFlags.Last();
 
 			const TArray<UCustomizableObjectLayout*>& Layouts = TypedNodeParam->GetLayouts(*Pin);
@@ -4596,7 +4591,7 @@ UE::Mutable::Private::Ptr<UE::Mutable::Private::NodeMesh> GenerateMutableSourceM
 				if (ensure(Layout))
 				{
 					
-					UE::Mutable::Private::Ptr<UE::Mutable::Private::NodeLayout> LayoutNode = CreateMutableLayoutNode(GenerationContext, Layout, bLinkedToExtendMaterial);
+					UE::Mutable::Private::Ptr<UE::Mutable::Private::NodeLayout> LayoutNode = CreateMutableLayoutNode(Layout, bLinkedToExtendMaterial);
 					LayoutNode->SetMessageContext(Node);
 
 					ParameterNode->Layouts[LayoutIndex] = LayoutNode;
@@ -4610,13 +4605,14 @@ UE::Mutable::Private::Ptr<UE::Mutable::Private::NodeMesh> GenerateMutableSourceM
 					EMutableMeshConversionFlags::IgnoreSkinning |
 					EMutableMeshConversionFlags::IgnorePhysics |
 					EMutableMeshConversionFlags::IgnoreMorphs |
+					EMutableMeshConversionFlags::IgnoreAUD |
 					EMutableMeshConversionFlags::DoNotCreateMeshMetadata;
 				GenerationContext.MeshGenerationFlags.Push(ModifiersMeshFlags);
 
 				// Get the base mesh
 				FMutableSourceMeshData Source;
 				Source.Mesh = LayoutMesh;
-				ParameterNode->ReferenceMesh = ConvertSkeletalMeshToMutable(Source, true, LODIndex, SectionIndex, GenerationContext, Node);
+				ParameterNode->ReferenceMesh = ConvertSkeletalMeshToMutable(Source, BaseLODIndex, BaseSectionIndex, GenerationContext, Node);
 
 				GenerationContext.MeshGenerationFlags.Pop();
 			}
diff --git a/Engine/Plugins/Mutable/Source/CustomizableObjectEditor/Private/MuCOE/GenerateMutableSource/GenerateMutableSourceMesh.h b/Engine/Plugins/Mutable/Source/CustomizableObjectEditor/Private/MuCOE/GenerateMutableSource/GenerateMutableSourceMesh.h
--- a/Engine/Plugins/Mutable/Source/CustomizableObjectEditor/Private/MuCOE/GenerateMutableSource/GenerateMutableSourceMesh.h
+++ b/Engine/Plugins/Mutable/Source/CustomizableObjectEditor/Private/MuCOE/GenerateMutableSource/GenerateMutableSourceMesh.h
@@ -31,32 +31,23 @@ namespace UE::Mutable::Private
  * Do not confuse Section Index and Material Index, they are not the same.
  * 
  * @param Context 
- * @param Node 
+ * @param Source 
  * @param SkeletalMesh 
- * @param LODIndexConnected Connected pin LOD Index.
- * @param SectionIndexConnected Connected pin Section Index.
  * @param OutLODIndex Corrected Skeletal Mesh LOD Index.
  * @param OutSectionIndex Corrected Skeletal Mesh Section Index.
- * @param bOnlyConnectedLOD Corrected LOD and Section will unconditionally always be the connected ones.
  * @return When using Automatic LOD From Mesh, InOutLODIndex and InOutSectionIndex and will return -1 if the section is not found in the currently compiling LOD. */
-void GetLODAndSectionForAutomaticLODs(const FMutableGraphGenerationContext& Context, const UCustomizableObjectNode& Node, const USkeletalMesh& SkeletalMesh,
-	const int32 LODIndexConnected, const int32 SectionIndexConnected,
-	int32& OutLODIndex, int32& OutSectionIndex,
-	bool bOnlyConnectedLOD);
+void GetLODAndSectionForAutomaticLODs(const FMutableCompilationContext& Context, const FMutableSourceMeshData& Source, const USkeletalMesh& SkeletalMesh, int32& OutLODIndex, int32& OutSectionIndex);
 
 
 /** Converts an Unreal Skeletal Mesh to Mutable Mesh.
  *
  * @param Source. Mesh data required to generate the mutable mesh
- * @param LODIndexConnected LOD which the pin is connected to.
- * @param SectionIndexConnected Section which the pin is connected to.
  * @param LODIndex LOD we are generating. Will be different from LODIndexConnected only when using Automatic LOD From Mesh. 
  * @param SectionIndex Section we are generating. Will be different from SectionIndexConnected only when using Automatic LOD From Mesh.
  * @param GenerationContext 
  * @param CurrentNode The node being compiled, can be a skeletal mesh or a table node.
  * @return Mutable Mesh. Nullptr if there has been an error. Empty mesh if the Skeletal Mesh does not contain the requested LOD + Section. */
 TSharedPtr<UE::Mutable::Private::FMesh> ConvertSkeletalMeshToMutable(FMutableSourceMeshData& Source,
-                                         bool bMeshMustExist,
                                          int32 LODIndex, int32 SectionIndex,
                                          FMutableGraphGenerationContext& GenerationContext,
                                          const UCustomizableObjectNode* CurrentNode, 
@@ -71,17 +62,16 @@ TSharedPtr<UE::Mutable::Private::FMesh> ConvertStaticMeshToMutable(const UStatic
  * @param Source. Mesh data required to generate the mutable mesh
  * @param LODIndexConnected LOD which the pin is connected to.
  * @param SectionIndexConnected Section which the pin is connected to.
- * @param LODIndex LOD we are generating. Will be different from LODIndexConnected only when using Automatic LOD From Mesh. 
- * @param SectionIndex Section we are generating. Will be different from SectionIndexConnected only when using Automatic LOD From Mesh.
  * @param MeshUniqueTags
  * @param GenerationContext
  * @param CurrentNode The node being compiled, can be a skeletal mesh or a table node.
  * @return Mutable Mesh. Nullptr if there has been an error. Empty mesh if the Skeletal Mesh does not contain the requested LOD + Section. */
 TSharedPtr<UE::Mutable::Private::FMesh> GenerateMutableSkeletalMesh(FMutableSourceMeshData& Source,
                                 int32 LODIndexConnected, int32 SectionIndexConnected,
-                                int32 LODIndex, int32 SectionIndex, const FString& MeshUniqueTags, 
+                                const FString& MeshUniqueTags, 
                                 FMutableGraphGenerationContext& GenerationContext, 
-								const UCustomizableObjectNode* CurrentNode);
+								const UCustomizableObjectNode* CurrentNode,
+								const bool bOnlyConnectedLOD);
 
 /** */
 TSharedPtr<UE::Mutable::Private::FMesh> GenerateMutableStaticMesh(TSoftObjectPtr<UStreamableRenderAsset> Mesh, const TSoftClassPtr<UAnimInstance>& AnimBp,
diff --git a/Engine/Plugins/Mutable/Source/CustomizableObjectEditor/Private/MuCOE/GenerateMutableSource/GenerateMutableSourceModifier.cpp b/Engine/Plugins/Mutable/Source/CustomizableObjectEditor/Private/MuCOE/GenerateMutableSource/GenerateMutableSourceModifier.cpp
--- a/Engine/Plugins/Mutable/Source/CustomizableObjectEditor/Private/MuCOE/GenerateMutableSource/GenerateMutableSourceModifier.cpp
+++ b/Engine/Plugins/Mutable/Source/CustomizableObjectEditor/Private/MuCOE/GenerateMutableSource/GenerateMutableSourceModifier.cpp
@@ -78,6 +78,8 @@ UE::Mutable::Private::Ptr<UE::Mutable::Private::NodeModifier> GenerateMutableSou
 			EMutableMeshConversionFlags::IgnoreSkinning |
 			EMutableMeshConversionFlags::IgnorePhysics |
 			EMutableMeshConversionFlags::IgnoreMorphs |
+			EMutableMeshConversionFlags::IgnoreTexCoords |
+			EMutableMeshConversionFlags::IgnoreAUD |
 			EMutableMeshConversionFlags::DoNotCreateMeshMetadata;
 		GenerationContext.MeshGenerationFlags.Push(ModifiersMeshFlags);
 
@@ -111,6 +113,8 @@ UE::Mutable::Private::Ptr<UE::Mutable::Private::NodeModifier> GenerateMutableSou
 			EMutableMeshConversionFlags::IgnoreSkinning |
 			EMutableMeshConversionFlags::IgnorePhysics |
 			EMutableMeshConversionFlags::IgnoreMorphs |
+			EMutableMeshConversionFlags::IgnoreTexCoords |
+			EMutableMeshConversionFlags::IgnoreAUD |
 			EMutableMeshConversionFlags::DoNotCreateMeshMetadata;
 		GenerationContext.MeshGenerationFlags.Push(ModifiersMeshFlags);
 
@@ -163,6 +167,8 @@ UE::Mutable::Private::Ptr<UE::Mutable::Private::NodeModifier> GenerateMutableSou
 			EMutableMeshConversionFlags::IgnoreSkinning |
 			EMutableMeshConversionFlags::IgnorePhysics |
 			EMutableMeshConversionFlags::IgnoreMorphs |
+			EMutableMeshConversionFlags::IgnoreTexCoords |
+			EMutableMeshConversionFlags::IgnoreAUD |
 			EMutableMeshConversionFlags::DoNotCreateMeshMetadata;
 		GenerationContext.MeshGenerationFlags.Push(ModifiersMeshFlags);
 
@@ -209,6 +215,7 @@ UE::Mutable::Private::Ptr<UE::Mutable::Private::NodeModifier> GenerateMutableSou
 			EMutableMeshConversionFlags::IgnoreSkinning |
 			EMutableMeshConversionFlags::IgnorePhysics |
 			EMutableMeshConversionFlags::IgnoreMorphs |
+			EMutableMeshConversionFlags::IgnoreAUD |
 			EMutableMeshConversionFlags::DoNotCreateMeshMetadata;
 		GenerationContext.MeshGenerationFlags.Push(ModifiersMeshFlags);
 
@@ -382,6 +389,8 @@ UE::Mutable::Private::Ptr<UE::Mutable::Private::NodeModifier> GenerateMutableSou
 			EMutableMeshConversionFlags::IgnoreSkinning |
 			EMutableMeshConversionFlags::IgnorePhysics |
 			EMutableMeshConversionFlags::IgnoreMorphs |
+			EMutableMeshConversionFlags::IgnoreTexCoords |
+			EMutableMeshConversionFlags::IgnoreAUD |
 			EMutableMeshConversionFlags::DoNotCreateMeshMetadata;
 		GenerationContext.MeshGenerationFlags.Push(ModifiersMeshFlags);
 
@@ -425,7 +434,7 @@ UE::Mutable::Private::Ptr<UE::Mutable::Private::NodeModifier> GenerateMutableSou
 		ClipNode->RequiredTags = TypedNodeRemBlocks->GetNodeRequiredTags(&GenerationContext.MacroNodesStack);
 
 		check(TypedNodeRemBlocks->Layout->GetMesh().IsNull());
-		UE::Mutable::Private::Ptr<UE::Mutable::Private::NodeLayout> SourceLayout = CreateMutableLayoutNode(GenerationContext, TypedNodeRemBlocks->Layout, true);
+		UE::Mutable::Private::Ptr<UE::Mutable::Private::NodeLayout> SourceLayout = CreateMutableLayoutNode(TypedNodeRemBlocks->Layout, true);
 		ClipNode->ClipLayout = SourceLayout;
 		ClipNode->LayoutIndex = TypedNodeRemBlocks->ParentLayoutIndex;
 	}
@@ -436,6 +445,7 @@ UE::Mutable::Private::Ptr<UE::Mutable::Private::NodeModifier> GenerateMutableSou
 			EMutableMeshConversionFlags::IgnoreSkinning |
 			EMutableMeshConversionFlags::IgnorePhysics |
 			EMutableMeshConversionFlags::IgnoreMorphs |
+			EMutableMeshConversionFlags::IgnoreAUD |
 			EMutableMeshConversionFlags::DoNotCreateMeshMetadata;
 		GenerationContext.MeshGenerationFlags.Push(ModifiersMeshFlags);
 
@@ -515,6 +525,8 @@ UE::Mutable::Private::Ptr<UE::Mutable::Private::NodeModifier> GenerateMutableSou
 			EMutableMeshConversionFlags::IgnoreSkinning |
 			EMutableMeshConversionFlags::IgnorePhysics |
 			EMutableMeshConversionFlags::IgnoreMorphs |
+			EMutableMeshConversionFlags::IgnoreTexCoords |
+			EMutableMeshConversionFlags::IgnoreAUD |
 			EMutableMeshConversionFlags::DoNotCreateMeshMetadata;
 		GenerationContext.MeshGenerationFlags.Push(ModifiersMeshFlags);
 
@@ -607,6 +619,8 @@ UE::Mutable::Private::Ptr<UE::Mutable::Private::NodeModifier> GenerateMutableSou
 			EMutableMeshConversionFlags::IgnoreSkinning |
 			EMutableMeshConversionFlags::IgnorePhysics |
 			EMutableMeshConversionFlags::IgnoreMorphs |
+			EMutableMeshConversionFlags::IgnoreTexCoords |
+			EMutableMeshConversionFlags::IgnoreAUD |
 			EMutableMeshConversionFlags::DoNotCreateMeshMetadata;
 		GenerationContext.MeshGenerationFlags.Push(ModifiersMeshFlags);
 
diff --git a/Engine/Plugins/Mutable/Source/CustomizableObjectEditor/Private/MuCOE/GenerateMutableSource/GenerateMutableSourceTable.cpp b/Engine/Plugins/Mutable/Source/CustomizableObjectEditor/Private/MuCOE/GenerateMutableSource/GenerateMutableSourceTable.cpp
--- a/Engine/Plugins/Mutable/Source/CustomizableObjectEditor/Private/MuCOE/GenerateMutableSource/GenerateMutableSourceTable.cpp
+++ b/Engine/Plugins/Mutable/Source/CustomizableObjectEditor/Private/MuCOE/GenerateMutableSource/GenerateMutableSourceTable.cpp
@@ -30,7 +30,7 @@
 
 
 bool FillTableColumn(const UCustomizableObjectNodeTable* TableNode, UE::Mutable::Private::Ptr<UE::Mutable::Private::FTable> MutableTable, const FString& ColumnName, const FString& RowName, const uint32 RowId, uint8* CellData, const FProperty* ColumnProperty,
-	const FMutableSourceMeshData& BaseMeshData, const int LODIndexConnected, const int32 SectionIndexConnected, int32 LODIndex, int32 SectionIndex, const bool bOnlyConnectedLOD, FMutableGraphGenerationContext& GenerationContext)
+	const FMutableSourceMeshData& BaseMeshData, FMutableGraphGenerationContext& GenerationContext)
 {
 	int32 CurrentColumn;
 	UDataTable* DataTablePtr = GetDataTable(TableNode, GenerationContext);
@@ -45,22 +45,38 @@ bool FillTableColumn(const UCustomizableObjectNodeTable* TableNode, UE::Mutable:
 		
 		if (SoftObjectProperty->PropertyClass->IsChildOf(USkeletalMesh::StaticClass()))
 		{
-			// TODO: This is still forcing the load of the skeletal mesh asset. 
-			UObject* Object = GenerationContext.LoadObject(SoftObject);
-			USkeletalMesh* SkeletalMesh = Cast<USkeletalMesh>(Object);
+			// Generating an Empty cell
+			
+			// Getting reference Mesh
+			USkeletalMesh* ReferenceSkeletalMesh = TableNode->GetColumnDefaultAssetByType<USkeletalMesh>(ColumnPropertyName);
 
-			if(!SkeletalMesh)
+			if (!ReferenceSkeletalMesh)
 			{
-				// Generating an Empty cell
-				FString MutableColumnName = TableNode->GenerateSkeletalMeshMutableColumName(ColumnName, LODIndex, SectionIndex);
+				FString msg = FString::Printf(TEXT("Reference Skeletal Mesh not found for column [%s]."), *ColumnName);
+				GenerationContext.Log(FText::FromString(msg), TableNode);
 
-				CurrentColumn = MutableTable.get()->FindColumn(MutableColumnName);
+				return false;
+			}
 
-				if (CurrentColumn == -1)
-				{
-					CurrentColumn = MutableTable->AddColumn(MutableColumnName, UE::Mutable::Private::ETableColumnType::Mesh);
-				}
+			FString MutableColumnName;
 
+			{
+				// Current LOD and Section 
+				int32 LODIndex = INDEX_NONE;
+				int32 SectionIndex = INDEX_NONE;
+				GetLODAndSectionForAutomaticLODs(*GenerationContext.CompilationContext, BaseMeshData, *ReferenceSkeletalMesh, LODIndex, SectionIndex);
+
+				MutableColumnName = TableNode->GenerateSkeletalMeshMutableColumName(ColumnName, LODIndex, SectionIndex);
+			}
+
+			CurrentColumn = MutableTable.get()->FindColumn(MutableColumnName);
+			if (CurrentColumn == INDEX_NONE)
+			{
+				CurrentColumn = MutableTable->AddColumn(MutableColumnName, UE::Mutable::Private::ETableColumnType::Mesh);
+			}
+
+			if(SoftObject.IsNull())
+			{
 				TSharedPtr<UE::Mutable::Private::FMesh> EmptySkeletalMesh = nullptr;
 				MutableTable->SetCell(CurrentColumn, RowId, EmptySkeletalMesh);
 
@@ -175,57 +191,6 @@ bool FillTableColumn(const UCustomizableObjectNodeTable* TableNode, UE::Mutable:
 				}
 			}
 
-			// Getting reference Mesh
-			USkeletalMesh* ReferenceSkeletalMesh = TableNode->GetColumnDefaultAssetByType<USkeletalMesh>(ColumnPropertyName);
-
-			if (!ReferenceSkeletalMesh)
-			{
-				FString msg = FString::Printf(TEXT("Reference Skeletal Mesh not found for column [%s]."), *ColumnName);
-				GenerationContext.Log(FText::FromString(msg), TableNode);
-
-				return false;
-			}
-
-			GetLODAndSectionForAutomaticLODs(GenerationContext, *TableNode, *SkeletalMesh, LODIndexConnected, SectionIndexConnected, LODIndex, SectionIndex, bOnlyConnectedLOD);
-			
-			// Parameter used for LOD differences
-	
-			if (GenerationContext.CurrentAutoLODStrategy != ECustomizableObjectAutomaticLODStrategy::AutomaticFromMesh || 
-				SectionIndex == SectionIndexConnected)
-			{
-				const int32 NumLODs = SkeletalMesh->GetImportedModel()->LODModels.Num();
-
-				if (NumLODs <= LODIndex)
-				{
-					LODIndex = NumLODs - 1;
-
-					FString msg = FString::Printf(TEXT("Mesh from column [%s] row [%s] needs LOD %d but has less LODs than the reference mesh. LOD %d will be used instead. This can cause some performance penalties."),
-						*ColumnName, *RowName, LODIndex, LODIndex);
-
-					LogRowGenerationMessage(TableNode, DataTablePtr, GenerationContext, msg, RowName);
-				}
-			}
-
-			FSkeletalMeshModel* ImportedModel = SkeletalMesh->GetImportedModel();
-			
-			if (ImportedModel->LODModels.IsValidIndex(LODIndex)) // Ignore error since this Section is empty due to Automatic LODs From Mesh
-			{
-				if (!ImportedModel->LODModels[LODIndex].Sections.IsValidIndex(SectionIndex) && !bIgnoreMissingData)
-				{
-					FString msg = FString::Printf(TEXT("Mesh from column [%s] row [%s] does not have section %d at LOD %d"), *ColumnName, *RowName, SectionIndexConnected, LODIndex);
-					LogRowGenerationMessage(TableNode, DataTablePtr, GenerationContext, msg, RowName);
-				}
-			}
-
-			FString MutableColumnName = TableNode->GenerateSkeletalMeshMutableColumName(ColumnName, LODIndex, SectionIndex);
-
-			CurrentColumn = MutableTable.get()->FindColumn(MutableColumnName);
-
-			if (CurrentColumn == -1)
-			{
-				CurrentColumn = MutableTable->AddColumn(MutableColumnName, UE::Mutable::Private::ETableColumnType::Mesh);
-			}
-
 			// First process the mesh tags that are going to make the mesh unique and affect whether it's repeated in 
 			// the mesh cache or not
 			FString MeshUniqueTags;
@@ -241,44 +206,14 @@ bool FillTableColumn(const UCustomizableObjectNodeTable* TableNode, UE::Mutable:
 			{
 				MeshUniqueTags += GenerateGameplayTag(Tag.ToString());
 			}
-			
-			TArray<FCustomizableObjectStreameableResourceId> StreamedResources;
-
-			if (GenerationContext.CompilationContext->Object->GetPrivate()->IsAssetUserDataMergeEnabled())
-			{
-				const TArray<UAssetUserData*>* AssetUserDataArray = SkeletalMesh->GetAssetUserDataArray();
-
-				if (AssetUserDataArray)
-				{
-					for (UAssetUserData* AssetUserData : *AssetUserDataArray)
-					{
-						if (!AssetUserData)
-						{
-							continue;
-						}
-
-						const int32 ResourceIndex = GenerationContext.AddStreamedResource(*AssetUserData);
-						
-						if (ResourceIndex >= 0)
-						{
-							FCustomizableObjectStreameableResourceId ResourceId;
-							ResourceId.Id = (uint32)ResourceIndex;
-							ResourceId.Type = (uint8)FCustomizableObjectStreameableResourceId::EType::AssetUserData;
-
-							StreamedResources.Add(ResourceId);
-						}
-
-						MeshUniqueTags += AssetUserData->GetPathName();
-					}
-				}
-			}
 
 			//TODO: Add AnimBp physics to Tables.
+
 			FMutableSourceMeshData Source = BaseMeshData;
 			Source.Mesh = SoftObject.ToSoftObjectPath();
 			Source.TableReferenceSkeletalMesh = ReferenceSkeletalMesh;
-			TSharedPtr<UE::Mutable::Private::FMesh> MutableMesh = GenerateMutableSkeletalMesh(Source, LODIndexConnected, SectionIndexConnected,
-																LODIndex, SectionIndex, MeshUniqueTags, GenerationContext, TableNode);
+			Source.OptionalMessageInfo = FString::Printf(TEXT(" SkeletalMesh from column [%s] row [%s]"), *ColumnName, *RowName);
+			TSharedPtr<UE::Mutable::Private::FMesh> MutableMesh = GenerateMutableSkeletalMesh(Source, Source.BaseLODIndex, Source.BaseSectionIndex, MeshUniqueTags, GenerationContext, TableNode, Source.bOnlyConnectedLOD);
 
 			if (MutableMesh)
 			{
@@ -292,28 +227,7 @@ bool FillTableColumn(const UCustomizableObjectNodeTable* TableNode, UE::Mutable:
 					AddTagToMutableMeshUnique(*MutableMesh, GenerateGameplayTag(Tag.ToString()));
 				}
 
-				for (FCustomizableObjectStreameableResourceId ResourceId : StreamedResources)
-				{
-					MutableMesh->AddStreamedResource(BitCast<uint64>(ResourceId));
-				}
-
-				AddSocketTagsToMesh(SkeletalMesh, MutableMesh, GenerationContext);
-
-				if (UCustomizableObjectSystem::GetInstance()->IsMutableAnimInfoDebuggingEnabled())
-				{
-					FString MeshPath;
-					SkeletalMesh->GetOuter()->GetPathName(nullptr, MeshPath);
-					FString MeshTag = FString("__MeshPath:") + MeshPath;
-					AddTagToMutableMeshUnique(*MutableMesh, MeshTag);
-				}
-
-				MutableTable->SetCell(CurrentColumn, RowId, MutableMesh, SkeletalMesh);
-			}
-			else
-			{
-				FString msg = FString::Printf(TEXT("Error converting skeletal mesh LOD %d, Section %d from column [%s] row [%s] to mutable."),
-					LODIndex, SectionIndex, *ColumnName, *RowName);
-				LogRowGenerationMessage(TableNode, DataTablePtr, GenerationContext, msg, RowName);
+				MutableTable->SetCell(CurrentColumn, RowId, MutableMesh, TableNode);
 			}
 		}
 
@@ -339,7 +253,7 @@ bool FillTableColumn(const UCustomizableObjectNodeTable* TableNode, UE::Mutable:
 			}
 
 			// Parameter used for LOD differences
-			int32 CurrentLOD = LODIndex;
+			int32 CurrentLOD = BaseMeshData.BaseLODIndex;
 
 			int NumLODs = StaticMesh->GetRenderData()->LODResources.Num();
 
@@ -348,7 +262,7 @@ bool FillTableColumn(const UCustomizableObjectNodeTable* TableNode, UE::Mutable:
 				CurrentLOD = NumLODs - 1;
 
 				FString msg = FString::Printf(TEXT("Mesh from column [%s] row [%s] needs LOD %d but has less LODs than the reference mesh. LOD %d will be used instead. This can cause some performance penalties."),
-					*ColumnName, *RowName, LODIndex, CurrentLOD);
+					*ColumnName, *RowName, BaseMeshData.BaseLODIndex, CurrentLOD);
 				LogRowGenerationMessage(TableNode, DataTablePtr, GenerationContext, msg, RowName);
 			}
 
@@ -364,7 +278,7 @@ bool FillTableColumn(const UCustomizableObjectNodeTable* TableNode, UE::Mutable:
 				LogRowGenerationMessage(TableNode, DataTablePtr, GenerationContext, msg, RowName);
 			}
 
-			FString MutableColumnName = TableNode->GenerateStaticMeshMutableColumName(ColumnName, SectionIndex);
+			FString MutableColumnName = TableNode->GenerateStaticMeshMutableColumName(ColumnName, BaseMeshData.BaseSectionIndex);
 
 			CurrentColumn = MutableTable.get()->FindColumn(MutableColumnName);
 
@@ -374,7 +288,7 @@ bool FillTableColumn(const UCustomizableObjectNodeTable* TableNode, UE::Mutable:
 			}
 
 			constexpr bool bIsPassthrough = false;
-			TSharedPtr<UE::Mutable::Private::FMesh> MutableMesh = GenerateMutableStaticMesh(StaticMesh, TSoftClassPtr<UAnimInstance>(), CurrentLOD, SectionIndex,
+			TSharedPtr<UE::Mutable::Private::FMesh> MutableMesh = GenerateMutableStaticMesh(StaticMesh, TSoftClassPtr<UAnimInstance>(), CurrentLOD, BaseMeshData.BaseSectionIndex,
 														  FString(), GenerationContext, TableNode, nullptr, bIsPassthrough);
 
 			if (MutableMesh)
@@ -384,7 +298,7 @@ bool FillTableColumn(const UCustomizableObjectNodeTable* TableNode, UE::Mutable:
 			else
 			{
 				FString msg = FString::Printf(TEXT("Error converting skeletal mesh LOD %d, Section %d from column [%s] row [%s] to mutable."),
-					LODIndex, SectionIndex, *ColumnName, *RowName);
+					BaseMeshData.BaseLODIndex, BaseMeshData.BaseSectionIndex, *ColumnName, *RowName);
 				LogRowGenerationMessage(TableNode, DataTablePtr, GenerationContext, msg, RowName);
 			}
 		}
@@ -832,7 +746,7 @@ TArray<FName> GetRowsToCompile(const UDataTable& DataTable, const UCustomizableO
 
 
 bool GenerateTableColumn(const UCustomizableObjectNodeTable* TableNode, const UEdGraphPin* Pin, UE::Mutable::Private::Ptr<UE::Mutable::Private::FTable> MutableTable, const FString& DataTableColumnName, const FProperty* ColumnProperty, 
-	const FMutableSourceMeshData& BaseMeshData,	const int32 LODIndexConnected, const int32 SectionIndexConnected, const int32 LODIndex, const int32 SectionIndex, const bool bOnlyConnectedLOD, FMutableGraphGenerationContext& GenerationContext)
+	const FMutableSourceMeshData& BaseMeshData, FMutableGraphGenerationContext& GenerationContext)
 {
 	MUTABLE_CPUPROFILER_SCOPE(GenerateTableColumn);
 
@@ -852,52 +766,6 @@ bool GenerateTableColumn(const UCustomizableObjectNodeTable* TableNode, const UE
 	TArray<uint32> RowIds;
 	TArray<FName> RowNames = GetRowsToCompile(*DataTable, *TableNode, GenerationContext, RowIds);
 
-	// Pre-pass to request async loading of all data. This seems to be slightly faster because it avoids the sync after every separate load.
-	// Avoid pre-loading meshes, which should be loaded on-demand during the core compilation
-	//bool bPreload = true;
-	//if (const FSoftObjectProperty* SoftObjectProperty = CastField<FSoftObjectProperty>(ColumnProperty))
-	//{
-	//	if (SoftObjectProperty->PropertyClass->IsChildOf(USkeletalMesh::StaticClass()))
-	//	{
-	//		bPreload = false;
-	//	}
-	//	else if (SoftObjectProperty->PropertyClass->IsChildOf(UStaticMesh::StaticClass()))
-	//	{
-	//		bPreload = false;
-	//	}
-	//}
-
-	//if (bPreload)
-	//{
-	//	FCookLoadScope CookLoadScope(ECookLoadType::EditorOnly);
-
-	//	TArray<int32> LoadRequests;
-	//	LoadRequests.Reserve(RowNames.Num());
-
-	//	MUTABLE_CPUPROFILER_SCOPE(Preload);
-	//	for (int32 RowIndex = 0; RowIndex < RowNames.Num(); ++RowIndex)
-	//	{
-	//		if (uint8* CellData = UCustomizableObjectNodeTable::GetCellData(RowNames[RowIndex], *DataTable, *ColumnProperty))
-	//		{
-	//			// Getting property type
-	//			if (const FSoftObjectProperty* SoftObjectProperty = CastField<FSoftObjectProperty>(ColumnProperty))
-	//			{
-	//				const FSoftObjectPtr& Path = SoftObjectProperty->GetPropertyValue(CellData);
-	//				const FString PackageName = Path.GetLongPackageName();
-	//				if (!PackageName.IsEmpty())
-	//				{
-	//					LoadRequests.Add( LoadPackageAsync(PackageName) );
-	//				}
-	//			}
-	//		}
-	//	}
-
-	//	{
-	//		MUTABLE_CPUPROFILER_SCOPE(Flush);
-	//		FlushAsyncLoading(LoadRequests);
-	//	}
-	//}
-
 	// Variable to check if something failed during a Cell generation
 	bool bCellGenerated = true;
 	const UEdGraphSchema_CustomizableObject* Schema = GetDefault<UEdGraphSchema_CustomizableObject>();
@@ -918,7 +786,7 @@ bool GenerateTableColumn(const UCustomizableObjectNodeTable* TableNode, const UE
 		if (uint8* CellData = UCustomizableObjectNodeTable::GetCellData(RowNames[RowIndex], *DataTable, *ColumnProperty))
 		{
 			bCellGenerated = FillTableColumn(TableNode, MutableTable, DataTableColumnName, RowNames[RowIndex].ToString(), RowIds[RowIndex], CellData, ColumnProperty,
-				BaseMeshData, LODIndexConnected, SectionIndexConnected, LODIndex, SectionIndex, bOnlyConnectedLOD, GenerationContext);
+				BaseMeshData, GenerationContext);
 
 			// Stop the compilation if something fails
 			if (!bCellGenerated)
diff --git a/Engine/Plugins/Mutable/Source/CustomizableObjectEditor/Private/MuCOE/GenerateMutableSource/GenerateMutableSourceTable.h b/Engine/Plugins/Mutable/Source/CustomizableObjectEditor/Private/MuCOE/GenerateMutableSource/GenerateMutableSourceTable.h
--- a/Engine/Plugins/Mutable/Source/CustomizableObjectEditor/Private/MuCOE/GenerateMutableSource/GenerateMutableSourceTable.h
+++ b/Engine/Plugins/Mutable/Source/CustomizableObjectEditor/Private/MuCOE/GenerateMutableSource/GenerateMutableSourceTable.h
@@ -20,18 +20,11 @@ struct FMutableSourceMeshData;
  * @param MutableTable 
  * @param DataTableColumnName 
  * @param TableProperty 
- * @param LODIndexConnected LOD which the pin is connected to.
- * @param SectionIndexConnected Section which the pin is connected to.
- * @param LODIndex LOD we are generating. Will be different from LODIndexConnected only when using Automatic LOD From Mesh. 
- * @param SectionIndex Section we are generating. Will be different from SectionIndexConnected only when using Automatic LOD From Mesh.
- * @param bOnlyConnectedLOD Corrected LOD and Section will unconditionally always be the connected ones.
+ * @param BaseMeshData
  * @param GenerationContext 
  * @return  */
 bool GenerateTableColumn(const UCustomizableObjectNodeTable* TableNode, const UEdGraphPin* Pin, UE::Mutable::Private::Ptr<UE::Mutable::Private::FTable> MutableTable, const FString& DataTableColumnName, const FProperty* ColumnProperty,
 	const FMutableSourceMeshData& BaseMeshData,
-	int32 LODIndexConnected, int32 SectionIndexConnected,
-	int32 LODIndex, int32 SectionIndex,
-	bool bOnlyConnectedLOD,
 	FMutableGraphGenerationContext& GenerationContext);
 
 
@@ -45,19 +38,12 @@ bool GenerateTableColumn(const UCustomizableObjectNodeTable* TableNode, const UE
  * @param RowIdx 
  * @param CellData 
  * @param Property 
- * @param LODIndexConnected LOD which the pin is connected to.
- * @param SectionIndexConnected Section which the pin is connected to.
- * @param LODIndex LOD we are generating. Will be different from LODIndexConnected only when using Automatic LOD From Mesh. 
- * @param SectionIndex Section we are generating. Will be different from SectionIndexConnected only when using Automatic LOD From Mesh.
- * @param bOnlyConnectedLOD Corrected LOD and Section will unconditionally always be the connected ones.
+ * @param BaseMeshData 
  * @param GenerationContext 
  * @return  */
 bool FillTableColumn(const UCustomizableObjectNodeTable* TableNode, UE::Mutable::Private::Ptr<UE::Mutable::Private::FTable> MutableTable, const FString& ColumnName,
 	const FString& RowName, uint32 RowId, uint8* CellData, const FProperty* ColumnProperty,
-	const FMutableSourceMeshData& BaseMeshData,
-	int32 LODIndexConnected, int32 SectionIndexConnected, int32 LODIndex, int32 SectionIndex,
-	bool bOnlyConnectedLOD,
-	FMutableGraphGenerationContext& GenerationContext);
+	const FMutableSourceMeshData& BaseMeshData,	FMutableGraphGenerationContext& GenerationContext);
 
 
 UE::Mutable::Private::Ptr<UE::Mutable::Private::FTable> GenerateMutableSourceTable(const UDataTable* DataTable, const UCustomizableObjectNodeTable* TableNode, FMutableGraphGenerationContext& GenerationContext);
diff --git a/Engine/Plugins/Mutable/Source/CustomizableObjectEditor/Private/MuCOE/Nodes/CustomizableObjectNodeMeshInterface.h b/Engine/Plugins/Mutable/Source/CustomizableObjectEditor/Private/MuCOE/Nodes/CustomizableObjectNodeMeshInterface.h
--- a/Engine/Plugins/Mutable/Source/CustomizableObjectEditor/Private/MuCOE/Nodes/CustomizableObjectNodeMeshInterface.h
+++ b/Engine/Plugins/Mutable/Source/CustomizableObjectEditor/Private/MuCOE/Nodes/CustomizableObjectNodeMeshInterface.h
@@ -42,8 +42,7 @@ public:
 	 * In other words, the pin still represents a LOD 3 but the asset no longer has this third LOD.
 	 * 
 	 * @param Pin In. Node's owning pin.
-	 * @param LODIndex Out.
-	 * @param SectionIndex Out. */
+	 * @param LODIndex Out. */
 	virtual void GetPinSection(const UEdGraphPin& Pin, int32& OutLODIndex, int32& OutSectionIndex) const PURE_VIRTUAL(UCustomizableObjectNodeMesh::GetPinSection, );
 };
 
diff --git a/Engine/Plugins/Mutable/Source/CustomizableObjectEditor/Private/MuCOE/Nodes/CustomizableObjectNodeSkeletalMesh.cpp b/Engine/Plugins/Mutable/Source/CustomizableObjectEditor/Private/MuCOE/Nodes/CustomizableObjectNodeSkeletalMesh.cpp
--- a/Engine/Plugins/Mutable/Source/CustomizableObjectEditor/Private/MuCOE/Nodes/CustomizableObjectNodeSkeletalMesh.cpp
+++ b/Engine/Plugins/Mutable/Source/CustomizableObjectEditor/Private/MuCOE/Nodes/CustomizableObjectNodeSkeletalMesh.cpp
@@ -74,9 +74,9 @@ void UCustomizableObjectNodeSkeletalMesh::AllocateDefaultPins(UCustomizableObjec
 	const UEdGraphSchema_CustomizableObject* Schema = GetDefault<UEdGraphSchema_CustomizableObject>();
 
 	// Force the loading of the skeletal mesh so we can later access the internal WeakObjectPtr (may be null if no SKM is set)
-	UE::Mutable::Private::LoadObject(SkeletalMesh);
+	USkeletalMesh* LoadedMesh = UE::Mutable::Private::LoadObject(SkeletalMesh);
 	
-	if (!SkeletalMesh)
+	if (!LoadedMesh)
 	{
 		UCustomizableObjectNodeSkeletalMeshPinDataMesh* PinData = NewObject<UCustomizableObjectNodeSkeletalMeshPinDataMesh>(this);
 		PinData->Init(-1, -1, -1);
@@ -89,9 +89,9 @@ void UCustomizableObjectNodeSkeletalMesh::AllocateDefaultPins(UCustomizableObjec
 		DefaultPin = {};
 	}
 	
-	if (const FSkeletalMeshModel* ImportedModel = SkeletalMesh->GetImportedModel())
+	if (const FSkeletalMeshModel* ImportedModel = LoadedMesh->GetImportedModel())
 	{
-		const int32 NumLODs = SkeletalMesh->GetLODNum();
+		const int32 NumLODs = LoadedMesh->GetLODNum();
 		for (int32 LODIndex = 0; LODIndex < NumLODs; ++LODIndex)
 		{
 			const int32 NumSections = ImportedModel->LODModels[LODIndex].Sections.Num();
@@ -169,10 +169,10 @@ bool UCustomizableObjectNodeSkeletalMesh::HasPinViewer() const
 
 FText UCustomizableObjectNodeSkeletalMesh::GetNodeTitle(ENodeTitleType::Type TitleType) const
 {
-	if (SkeletalMesh)
+	if (!SkeletalMesh.IsNull())
 	{
 		FFormatNamedArguments Args;
-		Args.Add(TEXT("MeshName"), FText::FromString(SkeletalMesh->GetName()));
+		Args.Add(TEXT("MeshName"), FText::FromString(SkeletalMesh.GetAssetName()));
 
 		return FText::Format(LOCTEXT("SkeletalMesh_Title", "{MeshName}\nSkeletal Mesh"), Args);
 	}
@@ -328,7 +328,7 @@ bool UCustomizableObjectNodeSkeletalMesh::IsNodeOutDatedAndNeedsRefresh()
 {
 	const bool bOutdated = [&]()
 	{
-		if (!SkeletalMesh)
+		if (!SkeletalMesh) // TODO UE-312665: SkeletalMesh->IsNull();
 		{
 			return false;
 		}
@@ -703,15 +703,19 @@ UMaterialInterface* UCustomizableObjectNodeSkeletalMesh::GetMaterialInterfaceFor
 
 FSkeletalMaterial* UCustomizableObjectNodeSkeletalMesh::GetSkeletalMaterialFor(const int32 LODIndex, const int32 SectionIndex) const
 {
-	if (!SkeletalMesh)
+	if (SkeletalMesh.IsNull())
 	{
 		return nullptr;
 	}
 
+	// TODO UE-312665: Cache data and remove sync load
+	USkeletalMesh* LoadedMesh = UE::Mutable::Private::LoadObject(SkeletalMesh);
+	ensure(LoadedMesh);
+
 	const int32 SkeletalMeshMaterialIndex = GetSkeletalMaterialIndexFor(LODIndex, SectionIndex);
-	if (SkeletalMesh->GetMaterials().IsValidIndex(SkeletalMeshMaterialIndex))
+	if (LoadedMesh->GetMaterials().IsValidIndex(SkeletalMeshMaterialIndex))
 	{
-		return &SkeletalMesh->GetMaterials()[SkeletalMeshMaterialIndex];
+		return &LoadedMesh->GetMaterials()[SkeletalMeshMaterialIndex];
 	}
 	
 	return nullptr;
@@ -719,12 +723,19 @@ FSkeletalMaterial* UCustomizableObjectNodeSkeletalMesh::GetSkeletalMaterialFor(c
 
 const FSkelMeshSection* UCustomizableObjectNodeSkeletalMesh::GetSkeletalMeshSectionFor(const int32 LODIndex, const int32 SectionIndex) const
 {	
-	if (!SkeletalMesh)
+	if (SkeletalMesh.IsNull())
 	{
 		return nullptr;
 	}
 	
-	const FSkeletalMeshModel* ImportedModel = SkeletalMesh->GetImportedModel();
+	// TODO UE-312665: Cache data and remove sync load
+	USkeletalMesh* LoadedMesh = UE::Mutable::Private::LoadObject(SkeletalMesh);
+	if (!LoadedMesh)
+	{
+		return nullptr;
+	}
+
+	const FSkeletalMeshModel* ImportedModel = LoadedMesh->GetImportedModel();
 	if (!ImportedModel)
 	{
 		return nullptr;
@@ -746,7 +757,14 @@ const FSkelMeshSection* UCustomizableObjectNodeSkeletalMesh::GetSkeletalMeshSect
 
 int32 UCustomizableObjectNodeSkeletalMesh::GetSkeletalMaterialIndexFor(const int32 LODIndex, const int32 SectionIndex) const
 {
-	if (!SkeletalMesh)
+	if (SkeletalMesh.IsNull())
+	{
+		return INDEX_NONE;
+	}
+
+	// TODO UE-312665: Cache data and remove sync load
+	USkeletalMesh* LoadedMesh = UE::Mutable::Private::LoadObject(SkeletalMesh);
+	if (!LoadedMesh)
 	{
 		return INDEX_NONE;
 	}
@@ -755,7 +773,7 @@ int32 UCustomizableObjectNodeSkeletalMesh::GetSkeletalMaterialIndexFor(const int
 	int32 SkeletalMeshMaterialIndex = INDEX_NONE;
 
 	// Check if we have lod info map to get the correct material index
-	if (const FSkeletalMeshLODInfo* LodInfo = SkeletalMesh->GetLODInfo(LODIndex))
+	if (const FSkeletalMeshLODInfo* LodInfo = LoadedMesh->GetLODInfo(LODIndex))
 	{
 		if (LodInfo->LODMaterialMap.IsValidIndex(SectionIndex))
 		{
@@ -766,7 +784,7 @@ int32 UCustomizableObjectNodeSkeletalMesh::GetSkeletalMaterialIndexFor(const int
 	// Only deduce index when the explicit mapping is not found or there is no remap
 	if (SkeletalMeshMaterialIndex == INDEX_NONE)
 	{
-		FSkeletalMeshModel* ImportedModel = SkeletalMesh->GetImportedModel();
+		FSkeletalMeshModel* ImportedModel = LoadedMesh->GetImportedModel();
 		if (ImportedModel && ImportedModel->LODModels.IsValidIndex(LODIndex) && ImportedModel->LODModels[LODIndex].Sections.IsValidIndex(SectionIndex))
 		{
 			SkeletalMeshMaterialIndex = ImportedModel->LODModels[LODIndex].Sections[SectionIndex].MaterialIndex;
@@ -790,7 +808,6 @@ void SGraphNodeSkeletalMesh::Construct(const FArguments& InArgs, UEdGraphNode* I
 
 	// Thumbnail
 	AssetThumbnailPool = MakeShareable(new FAssetThumbnailPool(32));
-	USkeletalMesh* SkeletalMesh = UE::Mutable::Private::LoadObject(NodeSkeletalMesh->SkeletalMesh);
 
 	const FAssetRegistryModule& AssetRegistryModule = FModuleManager::LoadModuleChecked<FAssetRegistryModule>(AssetRegistryConstants::ModuleName);
 
