[Mutable] Convert the Sockets index to an ID to avoid indeterminism issues.
Fix some typos and comments.

#rnx

diff --git a/Engine/Plugins/Mutable/Source/CustomizableObject/Internal/MuCO/CustomizableObjectPrivate.h b/Engine/Plugins/Mutable/Source/CustomizableObject/Internal/MuCO/CustomizableObjectPrivate.h
--- a/Engine/Plugins/Mutable/Source/CustomizableObject/Internal/MuCO/CustomizableObjectPrivate.h
+++ b/Engine/Plugins/Mutable/Source/CustomizableObject/Internal/MuCO/CustomizableObjectPrivate.h
@@ -192,6 +192,7 @@ struct FCustomizableObjectStreameableResourceId
 		AssetUserData         = 1,
 		RealTimeMorphTarget   = 2,
 		Clothing              = 3,
+		Socket	              = 4,
 	};
 
 	uint64 Id   : 64 - 8;
@@ -435,7 +436,7 @@ struct FMutableRefSocket
 	CUSTOMIZABLEOBJECT_API bool operator ==(const FMutableRefSocket& Other) const;
 
 #if WITH_EDITORONLY_DATA
-	friend FArchive& operator<<(FArchive& Ar, FMutableRefSocket& Data);
+	CUSTOMIZABLEOBJECT_API friend FArchive& operator<<(FArchive& Ar, FMutableRefSocket& Data);
 #endif
 };
 
@@ -1265,7 +1266,7 @@ public:
 
 	/** Mesh sockets provided by the part skeletal meshes, to be merged in the generated meshes */
 	UPROPERTY()
-	TArray<FMutableRefSocket> SocketArray;
+	TMap<uint32, FMutableRefSocket> Sockets;
 
 	UPROPERTY()
 	TArray<FMutableSkinWeightProfileInfo> SkinWeightProfilesInfo;
@@ -1705,6 +1706,8 @@ enum class ECustomizableObjectVersions : int32
 
 	EnableMeshLODStreaming,
 
+	DeterministicSocketsIds,
+
 	// -----<new versions can be added above this line>--------
 	LastCustomizableObjectVersion
 };
diff --git a/Engine/Plugins/Mutable/Source/CustomizableObject/Private/MuCO/CustomizableObject.cpp b/Engine/Plugins/Mutable/Source/CustomizableObject/Private/MuCO/CustomizableObject.cpp
--- a/Engine/Plugins/Mutable/Source/CustomizableObject/Private/MuCO/CustomizableObject.cpp
+++ b/Engine/Plugins/Mutable/Source/CustomizableObject/Private/MuCO/CustomizableObject.cpp
@@ -3256,7 +3256,7 @@ void UCustomizableObjectPrivate::LogMemory()
 				Arch << ModelResources->AnimBpOverridePhysiscAssetsInfo;
 				Arch << ModelResources->MaterialSlotNames;
 				Arch << ModelResources->BoneNamesMap;
-				Arch << ModelResources->SocketArray;
+				Arch << ModelResources->Sockets;
 				Arch << ModelResources->SkinWeightProfilesInfo;
 				Arch << ModelResources->MeshMetadata;
 				Arch << ModelResources->SurfaceMetadata;
@@ -3414,19 +3414,13 @@ bool FMutableModelImageProperties::operator!=(const FMutableModelImageProperties
 
 bool FMutableRefSocket::operator==(const FMutableRefSocket& Other) const
 {
-	if (
-		SocketName == Other.SocketName &&
+	return SocketName == Other.SocketName &&
 		BoneName == Other.BoneName &&
 		RelativeLocation == Other.RelativeLocation &&
 		RelativeRotation == Other.RelativeRotation &&
 		RelativeScale == Other.RelativeScale &&
 		bForceAlwaysAnimated == Other.bForceAlwaysAnimated &&
-		Priority == Other.Priority)
-	{
-		return true;
-	}
-
-	return false;
+		Priority == Other.Priority;
 }
 
 
diff --git a/Engine/Plugins/Mutable/Source/CustomizableObject/Private/MuCO/CustomizableObjectInstance.cpp b/Engine/Plugins/Mutable/Source/CustomizableObject/Private/MuCO/CustomizableObjectInstance.cpp
--- a/Engine/Plugins/Mutable/Source/CustomizableObject/Private/MuCO/CustomizableObjectInstance.cpp
+++ b/Engine/Plugins/Mutable/Source/CustomizableObject/Private/MuCO/CustomizableObjectInstance.cpp
@@ -4052,35 +4052,34 @@ void BuildMeshSockets(USkeletalMesh* SkeletalMesh, const UModelResources& ModelR
 	{
 		MUTABLE_CPUPROFILER_SCOPE(BuildMeshSockets_MutableSockets);
 
-		for (int32 TagIndex = 0; TagIndex < MutableMesh->GetTagCount(); ++TagIndex)
+		const TArray<uint64>& StreamedResources = MutableMesh->GetStreamedResources();
+
+		for (uint64 ResourceId : StreamedResources)
 		{
-			FString Tag = MutableMesh->GetTag(TagIndex);
+			FCustomizableObjectStreameableResourceId TypedResourceId = BitCast<FCustomizableObjectStreameableResourceId>(ResourceId);
 
-			if (Tag.RemoveFromStart("__Socket:"))
+			if (TypedResourceId.Type == (uint8)FCustomizableObjectStreameableResourceId::EType::Socket)
 			{
-				check(Tag.IsNumeric());
-				const int32 MutableSocketIndex = FCString::Atoi(*Tag);
-
-				if (ModelResources.SocketArray.IsValidIndex(MutableSocketIndex))
+				check(TypedResourceId.Id != 0 && TypedResourceId.Id <= TNumericLimits<uint32>::Max());
+				if (const FMutableRefSocket* MutableSocket = ModelResources.Sockets.Find(TypedResourceId.Id))
 				{
-					const FMutableRefSocket& MutableSocket = ModelResources.SocketArray[MutableSocketIndex];
 					int32 IndexToWriteSocket = -1;
 
-					if (TTuple<int32, int32>* FoundSocket = SocketMap.Find(MutableSocket.SocketName))
+					if (TTuple<int32, int32>* FoundSocket = SocketMap.Find(MutableSocket->SocketName))
 					{
-						if (FoundSocket->Value < MutableSocket.Priority)
+						if (FoundSocket->Value < MutableSocket->Priority)
 						{
 							// Overwrite the existing socket because the new mesh part one is higher priority
 							IndexToWriteSocket = FoundSocket->Key;
-							FoundSocket->Value = MutableSocket.Priority;
+							FoundSocket->Value = MutableSocket->Priority;
 						}
 					}
 					else
 					{
 						// New Socket
-						USkeletalMeshSocket* Socket = NewObject<USkeletalMeshSocket>(SkeletalMesh, MutableSocket.SocketName);
+						USkeletalMeshSocket* Socket = NewObject<USkeletalMeshSocket>(SkeletalMesh, MutableSocket->SocketName);
 						IndexToWriteSocket = Sockets.Add(Socket);
-						SocketMap.Add(MutableSocket.SocketName, TTuple<int32, int32>(IndexToWriteSocket, MutableSocket.Priority));
+						SocketMap.Add(MutableSocket->SocketName, TTuple<int32, int32>(IndexToWriteSocket, MutableSocket->Priority));
 					}
 
 					if (IndexToWriteSocket >= 0)
@@ -4089,14 +4088,14 @@ void BuildMeshSockets(USkeletalMesh* SkeletalMesh, const UModelResources& ModelR
 
 						USkeletalMeshSocket* SocketToWrite = Sockets[IndexToWriteSocket];
 
-						SocketToWrite->SocketName = MutableSocket.SocketName;
-						SocketToWrite->BoneName = MutableSocket.BoneName;
+						SocketToWrite->SocketName = MutableSocket->SocketName;
+						SocketToWrite->BoneName = MutableSocket->BoneName;
 
-						SocketToWrite->RelativeLocation = MutableSocket.RelativeLocation;
-						SocketToWrite->RelativeRotation = MutableSocket.RelativeRotation;
-						SocketToWrite->RelativeScale = MutableSocket.RelativeScale;
+						SocketToWrite->RelativeLocation = MutableSocket->RelativeLocation;
+						SocketToWrite->RelativeRotation = MutableSocket->RelativeRotation;
+						SocketToWrite->RelativeScale = MutableSocket->RelativeScale;
 
-						SocketToWrite->bForceAlwaysAnimated = MutableSocket.bForceAlwaysAnimated;
+						SocketToWrite->bForceAlwaysAnimated = MutableSocket->bForceAlwaysAnimated;
 					}
 				}
 			}
diff --git a/Engine/Plugins/Mutable/Source/CustomizableObjectEditor/Private/MuCOE/CustomizableObjectCompiler.cpp b/Engine/Plugins/Mutable/Source/CustomizableObjectEditor/Private/MuCOE/CustomizableObjectCompiler.cpp
--- a/Engine/Plugins/Mutable/Source/CustomizableObjectEditor/Private/MuCOE/CustomizableObjectCompiler.cpp
+++ b/Engine/Plugins/Mutable/Source/CustomizableObjectEditor/Private/MuCOE/CustomizableObjectCompiler.cpp
@@ -1713,7 +1713,7 @@ void FCustomizableObjectCompiler::FinishCompilationTask()
 		ModelResources->PhysicsAssets = MoveTemp(CompilationContext->PhysicsAssets);
 		ModelResources->AnimBpOverridePhysiscAssetsInfo = MoveTemp(CompilationContext->AnimBpOverridePhysicsAssetsInfo);
 		
-		ModelResources->SocketArray = MoveTemp(CompilationContext->SocketArray);
+		ModelResources->Sockets = MoveTemp(CompilationContext->Sockets);
 
 		const int32 NumBones = CompilationContext->UniqueBoneNames.Num() + CompilationContext->RemappedBoneNames.Num();
 		ModelResources->BoneNamesMap.Reserve(NumBones);
diff --git a/Engine/Plugins/Mutable/Source/CustomizableObjectEditor/Private/MuCOE/GenerateMutableSource/GenerateMutableSource.h b/Engine/Plugins/Mutable/Source/CustomizableObjectEditor/Private/MuCOE/GenerateMutableSource/GenerateMutableSource.h
--- a/Engine/Plugins/Mutable/Source/CustomizableObjectEditor/Private/MuCOE/GenerateMutableSource/GenerateMutableSource.h
+++ b/Engine/Plugins/Mutable/Source/CustomizableObjectEditor/Private/MuCOE/GenerateMutableSource/GenerateMutableSource.h
@@ -490,7 +490,7 @@ struct FMutableCompilationContext
 	TArray<FCustomizableObjectResourceData> StreamedResourceData;
 
 	/** Stores the sockets provided by the part skeletal meshes, to be merged in the generated meshes */
-	TArray<FMutableRefSocket> SocketArray;
+	TMap<uint32, FMutableRefSocket> Sockets;
 
 	/** Data used for SkinWeightProfiles reconstruction. */
 	TArray<FMutableSkinWeightProfileInfo> SkinWeightProfilesInfo;
diff --git a/Engine/Plugins/Mutable/Source/CustomizableObjectEditor/Private/MuCOE/GenerateMutableSource/GenerateMutableSourceMaterial.cpp b/Engine/Plugins/Mutable/Source/CustomizableObjectEditor/Private/MuCOE/GenerateMutableSource/GenerateMutableSourceMaterial.cpp
--- a/Engine/Plugins/Mutable/Source/CustomizableObjectEditor/Private/MuCOE/GenerateMutableSource/GenerateMutableSourceMaterial.cpp
+++ b/Engine/Plugins/Mutable/Source/CustomizableObjectEditor/Private/MuCOE/GenerateMutableSource/GenerateMutableSourceMaterial.cpp
@@ -430,7 +430,7 @@ UE::Mutable::Private::Ptr<UE::Mutable::Private::NodeMaterial> GenerateMutableSou
 					// Materials use the parameter id as column names
 					ColumnName = GenerationContext.CurrentMaterialParameterId;
 
-					// Generating a new Float column if not exists
+					// Generating a new material column if it does not exist
 					if (Table->FindColumn(ColumnName) == INDEX_NONE)
 					{
 						GenerationContext.CurrentTableColumnType = UE::Mutable::Private::ETableColumnType::Material;
diff --git a/Engine/Plugins/Mutable/Source/CustomizableObjectEditor/Private/MuCOE/GenerateMutableSource/GenerateMutableSourceMesh.cpp b/Engine/Plugins/Mutable/Source/CustomizableObjectEditor/Private/MuCOE/GenerateMutableSource/GenerateMutableSourceMesh.cpp
--- a/Engine/Plugins/Mutable/Source/CustomizableObjectEditor/Private/MuCOE/GenerateMutableSource/GenerateMutableSourceMesh.cpp
+++ b/Engine/Plugins/Mutable/Source/CustomizableObjectEditor/Private/MuCOE/GenerateMutableSource/GenerateMutableSourceMesh.cpp
@@ -1141,6 +1141,10 @@ TSharedPtr<UE::Mutable::Private::FMesh> ConvertSkeletalMeshToMutable(TSharedRef<
 
 	if (!bIgnoreSkinning)
 	{
+		constexpr uint32 InvalidResourceId = 0;
+
+		TArray<uint8> SocketData;
+
 		for (int32 SocketIndex = 0; SocketIndex < SkeletalMesh->NumSockets(); ++SocketIndex)
 		{
 			USkeletalMeshSocket* Socket = SkeletalMesh->GetSocketByIndex(SocketIndex);
@@ -1154,10 +1158,54 @@ TSharedPtr<UE::Mutable::Private::FMesh> ConvertSkeletalMeshToMutable(TSharedRef<
 			MutableSocket.bForceAlwaysAnimated = Socket->bForceAlwaysAnimated;
 			MutableSocket.Priority = Source.SocketPriority;
 
-			int32 SocketArrayIndex = Context.SocketArray.AddUnique(MutableSocket);
-			FString SocketTag = FString::Printf(TEXT("__Socket:%d"), SocketArrayIndex); // TODO PERE: Convert to StreamedResource and remove string
+			SocketData.SetNum(0, EAllowShrinking::No);
+			FMemoryWriter MemoryWriter(SocketData);
+			MemoryWriter << MutableSocket;
+
+			uint32 SocketDataHash = CityHash32(reinterpret_cast<const char*>(SocketData.GetData()), SocketData.Num());
+			const FMutableRefSocket* FoundData = Context.Sockets.Find(SocketDataHash);
+
+			// NOTE: This way of unique hash generation guarantees all valid values can be used but given its 
+			// sequential nature a cascade of changes can occur if new meshes are added. Not many hash collisions 
+			// are expected so it should not be problematic.
+			if (FoundData && *FoundData != MutableSocket)
+			{
+				uint32 NumTries = 0;
+				for (; NumTries < TNumericLimits<uint32>::Max(); ++NumTries)
+				{
+					SocketDataHash++;
+
+					if (SocketDataHash == InvalidResourceId)
+					{
+						continue;
+					}
+
+					FoundData = Context.Sockets.Find(SocketDataHash);
 
-			AddTagToMutableMeshUnique(*MutableMesh, SocketTag);
+					if (!FoundData || *FoundData == MutableSocket)
+					{
+						break;
+					}
+				}
+
+				if (NumTries == TNumericLimits<uint32>::Max())
+				{
+					UE_LOG(LogMutable, Warning, TEXT("Maximum number of mesh sockets reached, some meshes may not work as expected."));
+					SocketDataHash = InvalidResourceId;
+				}
+				
+			}
+
+			if (SocketDataHash != InvalidResourceId)
+			{
+				Context.Sockets.Add(SocketDataHash, MutableSocket);
+
+				FCustomizableObjectStreameableResourceId StreamedMorphResource;
+				StreamedMorphResource.Id = SocketDataHash;
+				StreamedMorphResource.Type = static_cast<uint8>(FCustomizableObjectStreameableResourceId::EType::Socket);
+
+				MutableMesh->AddStreamedResource(BitCast<uint64>(StreamedMorphResource));
+			}
 		}
 	}
 
@@ -3799,7 +3847,7 @@ UE::Mutable::Private::Ptr<UE::Mutable::Private::NodeMesh> GenerateMutableSourceM
 			FMutableSourceMeshData Source = BaseMeshData;
 			Source.Mesh = TypedNodeSkel->SkeletalMesh;
 			Source.AnimInstance = TypedNodeSkel->AnimInstance;
-			TSharedPtr<UE::Mutable::Private::FMesh > MutableMesh = GenerateMutableSkeletalMesh(Source, BaseLODIndex, BaseSectionIndex, MeshUniqueTags, GenerationContext, TypedNodeSkel, bOnlyConnectedLOD);
+			TSharedPtr<UE::Mutable::Private::FMesh> MutableMesh = GenerateMutableSkeletalMesh(Source, BaseLODIndex, BaseSectionIndex, MeshUniqueTags, GenerationContext, TypedNodeSkel, bOnlyConnectedLOD);
 
 			if (MutableMesh)
 			{
@@ -4480,7 +4528,7 @@ UE::Mutable::Private::Ptr<UE::Mutable::Private::NodeMesh> GenerateMutableSourceM
 							const bool bGenerateLayouts = GenerationContext.CurrentAutoLODStrategy != ECustomizableObjectAutomaticLODStrategy::AutomaticFromMesh
 								|| GenerationContext.CurrentLOD == GenerationContext.FromLOD;
 
-							if (true || bGenerateLayouts)
+							if (bGenerateLayouts)
 							{
 
 								TArray<UCustomizableObjectLayout*> Layouts = TypedNodeTable->GetLayouts(Pin);
diff --git a/Engine/Plugins/Mutable/Source/CustomizableObjectEditor/Private/MuCOE/GenerateMutableSource/GenerateMutableSourceMesh.h b/Engine/Plugins/Mutable/Source/CustomizableObjectEditor/Private/MuCOE/GenerateMutableSource/GenerateMutableSourceMesh.h
--- a/Engine/Plugins/Mutable/Source/CustomizableObjectEditor/Private/MuCOE/GenerateMutableSource/GenerateMutableSourceMesh.h
+++ b/Engine/Plugins/Mutable/Source/CustomizableObjectEditor/Private/MuCOE/GenerateMutableSource/GenerateMutableSourceMesh.h
@@ -35,7 +35,7 @@ namespace UE::Mutable::Private
  * @param SkeletalMesh 
  * @param OutLODIndex Corrected Skeletal Mesh LOD Index.
  * @param OutSectionIndex Corrected Skeletal Mesh Section Index.
- * @return When using Automatic LOD From Mesh, InOutLODIndex and InOutSectionIndex and will return -1 if the section is not found in the currently compiling LOD. */
+ * @return When using Automatic LOD From Mesh, OutLODIndex and OutSectionIndex will return -1 if the section is not found in the currently compiling LOD. */
 void GetLODAndSectionForAutomaticLODs(const FMutableCompilationContext& Context, const FMutableSourceMeshData& Source, const USkeletalMesh& SkeletalMesh, int32& OutLODIndex, int32& OutSectionIndex);
 
 
@@ -65,6 +65,7 @@ TSharedPtr<UE::Mutable::Private::FMesh> ConvertStaticMeshToMutable(const UStatic
  * @param MeshUniqueTags
  * @param GenerationContext
  * @param CurrentNode The node being compiled, can be a skeletal mesh or a table node.
+ * @param bOnlyConnectedLOD
  * @return Mutable Mesh. Nullptr if there has been an error. Empty mesh if the Skeletal Mesh does not contain the requested LOD + Section. */
 TSharedPtr<UE::Mutable::Private::FMesh> GenerateMutableSkeletalMesh(FMutableSourceMeshData& Source,
                                 int32 LODIndexConnected, int32 SectionIndexConnected,
diff --git a/Engine/Plugins/Mutable/Source/CustomizableObjectEditor/Private/MuCOE/GenerateMutableSource/GenerateMutableSourceTable.cpp b/Engine/Plugins/Mutable/Source/CustomizableObjectEditor/Private/MuCOE/GenerateMutableSource/GenerateMutableSourceTable.cpp
--- a/Engine/Plugins/Mutable/Source/CustomizableObjectEditor/Private/MuCOE/GenerateMutableSource/GenerateMutableSourceTable.cpp
+++ b/Engine/Plugins/Mutable/Source/CustomizableObjectEditor/Private/MuCOE/GenerateMutableSource/GenerateMutableSourceTable.cpp
@@ -297,7 +297,7 @@ bool FillTableColumn(const UCustomizableObjectNodeTable* TableNode, UE::Mutable:
 			}
 			else
 			{
-				FString msg = FString::Printf(TEXT("Error converting skeletal mesh LOD %d, Section %d from column [%s] row [%s] to mutable."),
+				FString msg = FString::Printf(TEXT("Error converting static mesh LOD %d, Section %d from column [%s] row [%s] to mutable."),
 					BaseMeshData.BaseLODIndex, BaseMeshData.BaseSectionIndex, *ColumnName, *RowName);
 				LogRowGenerationMessage(TableNode, DataTablePtr, GenerationContext, msg, RowName);
 			}
diff --git a/Engine/Plugins/Mutable/Source/CustomizableObjectEditor/Private/MuCOE/Nodes/CustomizableObjectNodeMeshInterface.h b/Engine/Plugins/Mutable/Source/CustomizableObjectEditor/Private/MuCOE/Nodes/CustomizableObjectNodeMeshInterface.h
--- a/Engine/Plugins/Mutable/Source/CustomizableObjectEditor/Private/MuCOE/Nodes/CustomizableObjectNodeMeshInterface.h
+++ b/Engine/Plugins/Mutable/Source/CustomizableObjectEditor/Private/MuCOE/Nodes/CustomizableObjectNodeMeshInterface.h
@@ -42,7 +42,8 @@ public:
 	 * In other words, the pin still represents a LOD 3 but the asset no longer has this third LOD.
 	 * 
 	 * @param Pin In. Node's owning pin.
-	 * @param LODIndex Out. */
+	 * @param LODIndex Out.
+	 * @param SectionIndex Out. */
 	virtual void GetPinSection(const UEdGraphPin& Pin, int32& OutLODIndex, int32& OutSectionIndex) const PURE_VIRTUAL(UCustomizableObjectNodeMesh::GetPinSection, );
 };
 
diff --git a/Engine/Plugins/Mutable/Source/CustomizableObjectEditor/Private/MuCOE/Nodes/CustomizableObjectNodeSkeletalMesh.cpp b/Engine/Plugins/Mutable/Source/CustomizableObjectEditor/Private/MuCOE/Nodes/CustomizableObjectNodeSkeletalMesh.cpp
--- a/Engine/Plugins/Mutable/Source/CustomizableObjectEditor/Private/MuCOE/Nodes/CustomizableObjectNodeSkeletalMesh.cpp
+++ b/Engine/Plugins/Mutable/Source/CustomizableObjectEditor/Private/MuCOE/Nodes/CustomizableObjectNodeSkeletalMesh.cpp
@@ -710,7 +710,10 @@ FSkeletalMaterial* UCustomizableObjectNodeSkeletalMesh::GetSkeletalMaterialFor(c
 
 	// TODO UE-312665: Cache data and remove sync load
 	USkeletalMesh* LoadedMesh = UE::Mutable::Private::LoadObject(SkeletalMesh);
-	ensure(LoadedMesh);
+	if (!LoadedMesh)
+	{
+		return nullptr;
+	}
 
 	const int32 SkeletalMeshMaterialIndex = GetSkeletalMaterialIndexFor(LODIndex, SectionIndex);
 	if (LoadedMesh->GetMaterials().IsValidIndex(SkeletalMeshMaterialIndex))
diff --git a/Engine/Plugins/Mutable/Source/MutableTools/Private/MuT/CodeOptimiser.cpp b/Engine/Plugins/Mutable/Source/MutableTools/Private/MuT/CodeOptimiser.cpp
--- a/Engine/Plugins/Mutable/Source/MutableTools/Private/MuT/CodeOptimiser.cpp
+++ b/Engine/Plugins/Mutable/Source/MutableTools/Private/MuT/CodeOptimiser.cpp
@@ -581,26 +581,11 @@ namespace UE::Mutable::Private
 					case EOpType::ME_APPLYPOSE:
 					case EOpType::ME_REMOVEMASK:
 					case EOpType::ME_PREPARELAYOUT:
+					case EOpType::ME_ADDMETADATA:
 					case EOpType::IM_PLAINCOLOUR:
 						bHasSpecialOpInSubgraph = true;
 						break;
 
-					case EOpType::ME_ADDMETADATA:
-					{
-						const ASTOpMeshAddMetadata* AddMetadata = static_cast<const ASTOpMeshAddMetadata*>(SubgraphRoot.get());
-						if (AddMetadata->Source && AddMetadata->Source->GetParentCount() == 1)
-						{
-							// We could resolve the constant if there are no other constraints since this op's mesh op is only referenced once.
-							bGetFromChildren = true;
-						}
-						else
-						{
-							// Don't reduce constant to avoid duplicating mesh data.
-							bHasSpecialOpInSubgraph = true;
-						}
-						break;
-					}
-
 					case EOpType::IM_RASTERMESH:
 					{
 						const ASTOpImageRasterMesh* Raster = static_cast<const ASTOpImageRasterMesh*>(SubgraphRoot.get());
