[mutable] Fixed the generation of duplicated mu::ImageParameterNode nodes.
- This happened because what was being cached was a string of nodes (result = Format <- Resize <- ImageParam) wich was causing the cache to not always hit.
- Added a method to generate ImageParams that will only generate a new image param node if it does not already exist.
- Invalidated Mutable Model if repeated parameters are found.
- Increased CO version to invalidate old compiled data.
- Changed the "Scalar" and "Enum" parameter node caches to use an FGUID instead of the string representation of an FGUID as the cache key.

#rnx

diff --git a/Engine/Plugins/Mutable/Source/CustomizableObject/Internal/MuCO/CustomizableObjectPrivate.h b/Engine/Plugins/Mutable/Source/CustomizableObject/Internal/MuCO/CustomizableObjectPrivate.h
--- a/Engine/Plugins/Mutable/Source/CustomizableObject/Internal/MuCO/CustomizableObjectPrivate.h
+++ b/Engine/Plugins/Mutable/Source/CustomizableObject/Internal/MuCO/CustomizableObjectPrivate.h
@@ -1708,6 +1708,8 @@ enum class ECustomizableObjectVersions : int32
 
 	DeterministicSocketsIds,
 
+	CompilationEarlyOutOnRepeatedParameters,
+
 	// -----<new versions can be added above this line>--------
 	LastCustomizableObjectVersion
 };
diff --git a/Engine/Plugins/Mutable/Source/CustomizableObjectEditor/Private/MuCOE/GenerateMutableSource/GenerateMutableSource.h b/Engine/Plugins/Mutable/Source/CustomizableObjectEditor/Private/MuCOE/GenerateMutableSource/GenerateMutableSource.h
--- a/Engine/Plugins/Mutable/Source/CustomizableObjectEditor/Private/MuCOE/GenerateMutableSource/GenerateMutableSource.h
+++ b/Engine/Plugins/Mutable/Source/CustomizableObjectEditor/Private/MuCOE/GenerateMutableSource/GenerateMutableSource.h
@@ -33,6 +33,7 @@
 #include "MuR/Skeleton.h"
 #include "MuR/Mesh.h"
 #include "MuR/Ptr.h"
+#include "MuT/NodeImageParameter.h"
 #include "UObject/Package.h"
 #include "UObject/StrongObjectPtr.h"
 
@@ -622,10 +623,13 @@ struct FMutableGraphGenerationContext
 	TMap<FGeneratedGroupProjectorsKey, FGroupProjectorTempData> GeneratedGroupProjectors;
 
 	/** Key is the Node Uid. */
-	TMap<FString, UE::Mutable::Private::Ptr<UE::Mutable::Private::NodeScalarParameter>> GeneratedScalarParameters;
+	TMap<FGuid, UE::Mutable::Private::Ptr<UE::Mutable::Private::NodeScalarParameter>> GeneratedScalarParameters;
 
 	/** Key is the Node Uid. */
-	TMap<FString, UE::Mutable::Private::Ptr<UE::Mutable::Private::NodeScalarEnumParameter>> GeneratedEnumParameters;
+	TMap<FGuid, UE::Mutable::Private::Ptr<UE::Mutable::Private::NodeScalarEnumParameter>> GeneratedEnumParameters;
+
+	/** Key is the Node Uid. */
+	TMap<FGuid, UE::Mutable::Private::Ptr<UE::Mutable::Private::NodeImageParameter>> GeneratedImageParameters;
 	
 	struct FGeneratedCompositeDataTablesData
 	{
diff --git a/Engine/Plugins/Mutable/Source/CustomizableObjectEditor/Private/MuCOE/GenerateMutableSource/GenerateMutableSourceGroupProjector.cpp b/Engine/Plugins/Mutable/Source/CustomizableObjectEditor/Private/MuCOE/GenerateMutableSource/GenerateMutableSourceGroupProjector.cpp
--- a/Engine/Plugins/Mutable/Source/CustomizableObjectEditor/Private/MuCOE/GenerateMutableSource/GenerateMutableSourceGroupProjector.cpp
+++ b/Engine/Plugins/Mutable/Source/CustomizableObjectEditor/Private/MuCOE/GenerateMutableSource/GenerateMutableSourceGroupProjector.cpp
@@ -342,7 +342,7 @@ TOptional<FGroupProjectorTempData> GenerateMutableGroupProjector(UCustomizableOb
 		GroupProjectorTempData.CustomizableObjectNodeGroupProjectorParameter = ProjParamNode;
 
 		UE::Mutable::Private::Ptr<UE::Mutable::Private::NodeScalarParameter> NodeScalarParam;
-		if (UE::Mutable::Private::Ptr<UE::Mutable::Private::NodeScalarParameter>* FindResult = GenerationContext.GeneratedScalarParameters.Find(NumLayersParamUid.ToString()))
+		if (UE::Mutable::Private::Ptr<UE::Mutable::Private::NodeScalarParameter>* FindResult = GenerationContext.GeneratedScalarParameters.Find(NumLayersParamUid))
 		{
 			NodeScalarParam = *FindResult;
 		}
@@ -357,7 +357,7 @@ TOptional<FGroupProjectorTempData> GenerateMutableGroupProjector(UCustomizableOb
 			GenerationContext.ParameterUIDataMap.Add(NodeScalarParamName,
 		FMutableParameterData(ProjParamNode->ParamUIMetadata, EMutableParameterType::Int));
 					
-			GenerationContext.GeneratedScalarParameters.Add(NumLayersParamUid.ToString(), NodeScalarParam);
+			GenerationContext.GeneratedScalarParameters.Add(NumLayersParamUid, NodeScalarParam);
 		}
 
 		UE::Mutable::Private::Ptr<UE::Mutable::Private::NodeRangeFromScalar> NodeRangeFromScalar = new UE::Mutable::Private::NodeRangeFromScalar;
@@ -367,7 +367,7 @@ TOptional<FGroupProjectorTempData> GenerateMutableGroupProjector(UCustomizableOb
 		GroupProjectorTempData.NodeProjectorParameterPtr->SetRange(0, NodeRangeFromScalar);
 		
 		UE::Mutable::Private::Ptr<UE::Mutable::Private::NodeScalarEnumParameter> PoseEnumParameterNode;
-		if (UE::Mutable::Private::Ptr<UE::Mutable::Private::NodeScalarEnumParameter>* FindResult = GenerationContext.GeneratedEnumParameters.Find(SelectedPoseParamUid.ToString()))
+		if (UE::Mutable::Private::Ptr<UE::Mutable::Private::NodeScalarEnumParameter>* FindResult = GenerationContext.GeneratedEnumParameters.Find(SelectedPoseParamUid))
 		{
 			PoseEnumParameterNode = *FindResult;
 		}
@@ -381,14 +381,14 @@ TOptional<FGroupProjectorTempData> GenerateMutableGroupProjector(UCustomizableOb
 			PoseEnumParameterNode->DefaultValue = 0;
 			GroupProjectorTempData.PoseOptionsParameter = PoseEnumParameterNode;
 
-			GenerationContext.GeneratedEnumParameters.Add(SelectedPoseParamUid.ToString(), PoseEnumParameterNode);
+			GenerationContext.GeneratedEnumParameters.Add(SelectedPoseParamUid, PoseEnumParameterNode);
 
 			GenerationContext.ParameterUIDataMap.Add(PoseNodeEnumParamName,
 		FMutableParameterData(ProjParamNode->ParamUIMetadata, EMutableParameterType::Int));
 		}
 		
 		UE::Mutable::Private::Ptr<UE::Mutable::Private::NodeScalarParameter> OpacityParameterNode;
-		if (UE::Mutable::Private::Ptr<UE::Mutable::Private::NodeScalarParameter>* FindResult = GenerationContext.GeneratedScalarParameters.Find(OpacityParamUid.ToString()))
+		if (UE::Mutable::Private::Ptr<UE::Mutable::Private::NodeScalarParameter>* FindResult = GenerationContext.GeneratedScalarParameters.Find(OpacityParamUid))
 		{
 			OpacityParameterNode = *FindResult;
 		}
@@ -402,7 +402,7 @@ TOptional<FGroupProjectorTempData> GenerateMutableGroupProjector(UCustomizableOb
 			OpacityParameterNode->Ranges.SetNum(1);
 			OpacityParameterNode->Ranges[0] = NodeRangeFromScalar;
 
-			GenerationContext.GeneratedScalarParameters.Add(OpacityParamUid.ToString(), OpacityParameterNode);
+			GenerationContext.GeneratedScalarParameters.Add(OpacityParamUid, OpacityParameterNode);
 					
 			FMutableParamUIMetadata OpacityMetadata = ProjParamNode->ParamUIMetadata;
 			OpacityMetadata.ObjectFriendlyName = FString("Opacity");
@@ -451,7 +451,7 @@ TOptional<FGroupProjectorTempData> GenerateMutableGroupProjector(UCustomizableOb
 		}
 	
 		UE::Mutable::Private::Ptr<UE::Mutable::Private::NodeScalarEnumParameter> EnumParameterNode;
-		if (UE::Mutable::Private::Ptr<UE::Mutable::Private::NodeScalarEnumParameter>* FindResult = GenerationContext.GeneratedEnumParameters.Find(SelectedImageParamUid.ToString()))
+		if (UE::Mutable::Private::Ptr<UE::Mutable::Private::NodeScalarEnumParameter>* FindResult = GenerationContext.GeneratedEnumParameters.Find(SelectedImageParamUid))
 		{
 			EnumParameterNode = *FindResult;
 		}
@@ -467,7 +467,7 @@ TOptional<FGroupProjectorTempData> GenerateMutableGroupProjector(UCustomizableOb
 			EnumParameterNode->Ranges.SetNum(1);
 			EnumParameterNode->Ranges[0] = NodeRangeFromScalar;
 
-			GenerationContext.GeneratedEnumParameters.Add(SelectedImageParamUid.ToString(), EnumParameterNode);
+			GenerationContext.GeneratedEnumParameters.Add(SelectedImageParamUid, EnumParameterNode);
 					
 			FMutableParameterData ParameterUIData;
 			ParameterUIData.ParamUIMetadata = ProjParamNode->ParamUIMetadata;
diff --git a/Engine/Plugins/Mutable/Source/CustomizableObjectEditor/Private/MuCOE/GenerateMutableSource/GenerateMutableSourceImage.cpp b/Engine/Plugins/Mutable/Source/CustomizableObjectEditor/Private/MuCOE/GenerateMutableSource/GenerateMutableSourceImage.cpp
--- a/Engine/Plugins/Mutable/Source/CustomizableObjectEditor/Private/MuCOE/GenerateMutableSource/GenerateMutableSourceImage.cpp
+++ b/Engine/Plugins/Mutable/Source/CustomizableObjectEditor/Private/MuCOE/GenerateMutableSource/GenerateMutableSourceImage.cpp
@@ -63,7 +63,6 @@
 
 #define LOCTEXT_NAMESPACE "CustomizableObjectEditor"
 
-
 UE::Mutable::Private::Ptr<UE::Mutable::Private::NodeImage> ResizeTextureByNumMips(const UE::Mutable::Private::Ptr<UE::Mutable::Private::NodeImage>& ImageConstant, int32 MipsToSkip)
 {
 	if (MipsToSkip > 0)
@@ -106,7 +105,7 @@ UE::Mutable::Private::Ptr<UE::Mutable::Private::NodeImage> GenerateMutableSource
 	}
 
 	// Bool that determines if a node can be added to the cache of nodes.
-	// Most nodes need to be added to the cache but there are some that don't. For exampel, MacroInstanceNodes
+	// Most nodes need to be added to the cache but there are some that don't. For example, MacroInstanceNodes
 	bool bCacheNode = true;
 
 	UE::Mutable::Private::Ptr<UE::Mutable::Private::NodeImage> Result;
@@ -160,19 +159,7 @@ UE::Mutable::Private::Ptr<UE::Mutable::Private::NodeImage> GenerateMutableSource
 
 	else if (const UCustomizableObjectNodeTextureParameter* TypedNodeParam = Cast<UCustomizableObjectNodeTextureParameter>(Node))
 	{
-		UE::Mutable::Private::Ptr<UE::Mutable::Private::NodeImageParameter> TextureNode = new UE::Mutable::Private::NodeImageParameter();
-
-		TextureNode->Name = TypedNodeParam->GetParameterName(&GenerationContext.MacroNodesStack);
-		TextureNode->UID = GenerationContext.GetNodeIdUnique(Node).ToString();
-
-		if (TypedNodeParam->DefaultValue)
-		{
-			GenerationContext.TextureParameterDefaultValues.Add(FName(TypedNodeParam->GetParameterName()), TypedNodeParam->DefaultValue);
-		}
-
-		GenerationContext.ParameterUIDataMap.Add(TypedNodeParam->GetParameterName(&GenerationContext.MacroNodesStack), FMutableParameterData(
-			TypedNodeParam->ParamUIMetadata,
-			EMutableParameterType::Texture));
+		UE::Mutable::Private::Ptr<UE::Mutable::Private::NodeImageParameter> ImageParameterNode = GenerateMutableSourceImageParameter(Pin, GenerationContext);
 
 		// Force the same format that the default texture if any.
 		UE::Mutable::Private::Ptr<UE::Mutable::Private::NodeImageFormat> FormatNode = new UE::Mutable::Private::NodeImageFormat();
@@ -186,7 +173,7 @@ UE::Mutable::Private::Ptr<UE::Mutable::Private::NodeImage> GenerateMutableSource
 			// Force an "easy format" on the texture.
 			FormatNode->Format = UE::Mutable::Private::EImageFormat::RGBA_UByte;
 		}
-		FormatNode->Source = TextureNode;
+		FormatNode->Source = ImageParameterNode;
 
 		UE::Mutable::Private::Ptr<UE::Mutable::Private::NodeImageResize> ResizeNode = new UE::Mutable::Private::NodeImageResize();
 		ResizeNode->Base = FormatNode;
@@ -194,8 +181,7 @@ UE::Mutable::Private::Ptr<UE::Mutable::Private::NodeImage> GenerateMutableSource
 
 		FUintVector2 TextureSize(TypedNodeParam->TextureSizeX, TypedNodeParam->TextureSizeY);
 
-		const UTexture2D* ReferenceTexture = TypedNodeParam->ReferenceValue;
-		if (ReferenceTexture)
+		if (const UTexture2D* ReferenceTexture = TypedNodeParam->ReferenceValue)
 		{
 			const uint32 LODBias = ComputeLODBiasForTexture(GenerationContext, *TypedNodeParam->ReferenceValue, ReferenceTexture, ReferenceTextureSize);
 			TextureSize.X = FMath::Max(ReferenceTexture->Source.GetSizeX() >> LODBias, 1);
@@ -228,9 +214,7 @@ UE::Mutable::Private::Ptr<UE::Mutable::Private::NodeImage> GenerateMutableSource
 		UE::Mutable::Private::Ptr<UE::Mutable::Private::NodeImageConstant> ImageNode = new UE::Mutable::Private::NodeImageConstant();
 		Result = ImageNode;
 
-		UTexture2D* Texture = TypedNodeMesh->FindTextureForPin(Pin);
-
-		if (Texture)
+		if (UTexture2D* Texture = TypedNodeMesh->FindTextureForPin(Pin))
 		{
 			TSharedPtr<UE::Mutable::Private::FImage> ImageConstant = GenerateImageConstant(Texture, GenerationContext, false);
 			ImageNode->SetValue(ImageConstant);
@@ -1273,5 +1257,57 @@ UE::Mutable::Private::Ptr<UE::Mutable::Private::NodeImage> GenerateMutableSource
 }
 
 
+UE::Mutable::Private::Ptr<UE::Mutable::Private::NodeImageParameter> GenerateMutableSourceImageParameter(const UEdGraphPin* Pin, FMutableGraphGenerationContext& GenerationContext)
+{
+	check(Pin)
+	RETURN_ON_CYCLE(*Pin, GenerationContext)
+
+	CheckNumOutputs(*Pin, GenerationContext);
+	
+	UCustomizableObjectNodeParameter* ParameterNode = CastChecked<UCustomizableObjectNodeParameter>(Pin->GetOwningNode());
+
+	UE::Mutable::Private::Ptr<UE::Mutable::Private::NodeImageParameter> Result;
+
+	const FGuid NodeGUID = ParameterNode->NodeGuid;
+	if (const UE::Mutable::Private::Ptr<UE::Mutable::Private::NodeImageParameter>* GeneratedImageParameter = GenerationContext.GeneratedImageParameters.Find(NodeGUID))
+	{
+		return GeneratedImageParameter->get();
+	}
+	
+	if (const UCustomizableObjectNodeTextureParameter* TypedNodeParam = Cast<UCustomizableObjectNodeTextureParameter>(ParameterNode))
+	{
+		UE::Mutable::Private::Ptr<UE::Mutable::Private::NodeImageParameter> TextureNode = new UE::Mutable::Private::NodeImageParameter();
+
+		TextureNode->Name = TypedNodeParam->GetParameterName(&GenerationContext.MacroNodesStack);
+		TextureNode->UID = GenerationContext.GetNodeIdUnique(ParameterNode).ToString();
+
+		if (TypedNodeParam->DefaultValue)
+		{
+			GenerationContext.TextureParameterDefaultValues.Add(FName(TypedNodeParam->GetParameterName()), TypedNodeParam->DefaultValue);
+		}
+
+		GenerationContext.ParameterUIDataMap.Add(TypedNodeParam->GetParameterName(&GenerationContext.MacroNodesStack), FMutableParameterData(
+			TypedNodeParam->ParamUIMetadata,
+			EMutableParameterType::Texture));
+
+		Result = TextureNode;
+	}
+	else
+	{
+		GenerationContext.Log(LOCTEXT("UnimplementedNode", "Node type not implemented yet."), ParameterNode);
+	}
+
+	if (Result)
+	{
+		Result->SetMessageContext(ParameterNode);
+	}
+	
+	GenerationContext.GeneratedImageParameters.Add(NodeGUID,Result);
+	GenerationContext.GeneratedNodes.Add(ParameterNode);
+	
+	return Result;
+}
+
+
 #undef LOCTEXT_NAMESPACE
 
diff --git a/Engine/Plugins/Mutable/Source/CustomizableObjectEditor/Private/MuCOE/GenerateMutableSource/GenerateMutableSourceImage.h b/Engine/Plugins/Mutable/Source/CustomizableObjectEditor/Private/MuCOE/GenerateMutableSource/GenerateMutableSourceImage.h
--- a/Engine/Plugins/Mutable/Source/CustomizableObjectEditor/Private/MuCOE/GenerateMutableSource/GenerateMutableSourceImage.h
+++ b/Engine/Plugins/Mutable/Source/CustomizableObjectEditor/Private/MuCOE/GenerateMutableSource/GenerateMutableSourceImage.h
@@ -16,6 +16,7 @@ namespace UE::Mutable::Private
 {
 	class FImage;
 	class NodeImage;
+	class NodeImageParameter;
 }
 
 
@@ -25,3 +26,5 @@ UE::Mutable::Private::Ptr<UE::Mutable::Private::NodeImage> ResizeTextureByNumMip
 
 /** Convert a CustomizableObject Source Graph from an Image pin into a mutable source graph. */
 UE::Mutable::Private::Ptr<UE::Mutable::Private::NodeImage> GenerateMutableSourceImage(const UEdGraphPin* Pin, FMutableGraphGenerationContext& GenerationContext, int32 ReferenceTextureSize);
+
+UE::Mutable::Private::Ptr<UE::Mutable::Private::NodeImageParameter> GenerateMutableSourceImageParameter(const UEdGraphPin* Pin, FMutableGraphGenerationContext& GenerationContext);
\ No newline at end of file
diff --git a/Engine/Plugins/Mutable/Source/MutableTools/Private/MuT/Compiler.cpp b/Engine/Plugins/Mutable/Source/MutableTools/Private/MuT/Compiler.cpp
--- a/Engine/Plugins/Mutable/Source/MutableTools/Private/MuT/Compiler.cpp
+++ b/Engine/Plugins/Mutable/Source/MutableTools/Private/MuT/Compiler.cpp
@@ -509,6 +509,8 @@ namespace UE::Mutable::Private
         		// Check for repeated parameter names and report them to the user
 				if (Parameters.Num() > 1)
 				{
+					// todo: UE-316151 : Allow the repeated parameter names log to display a list of n nodes (instead of just 2)
+					bool bRepeatedParameterFound = false;
 					TArray<FString> RepeatedParameterNames;
 				
 					for (int32 Index = 1; Index < Parameters.Num(); ++Index)
@@ -516,14 +518,17 @@ namespace UE::Mutable::Private
 						const FString CurrentParameterName = Parameters[Index].Name;
 						const FString PreviousParameterName = Parameters[Index-1].Name;
 					
+						// Found a repeated parameter that has not been reported yet
 						if (CurrentParameterName.Compare(PreviousParameterName) == 0 && !RepeatedParameterNames.Contains(CurrentParameterName))
 						{
-							// Find the nodes defining these parameters
+							bRepeatedParameterFound = true;
+							
+							// Find the nodes defining these repeated parameters
 							TArray<const void*, TInlineAllocator<4>> MessageContexts;
 							{
 								for (const TPair<Ptr<const Node>, Ptr<ASTOpParameter>>& Entry : Generator.FirstPass.ParameterNodes.GenericParametersCache)
 								{
-									if (Entry.Value->Parameter.Name.Compare(CurrentParameterName))
+									if (Entry.Value->Parameter.Name.Compare(CurrentParameterName) == 0)
 									{
 										MessageContexts.AddUnique(Entry.Key->GetMessageContext());
 									}
@@ -532,7 +537,7 @@ namespace UE::Mutable::Private
 								{
 									for (const TPair<Ptr<ASTOpParameter>, FMeshGenerationResult>& ArrayEntry : Entry.Value)
 									{
-										if (ArrayEntry.Key && ArrayEntry.Key->Parameter.Name.Compare(CurrentParameterName))
+										if (ArrayEntry.Key && ArrayEntry.Key->Parameter.Name.Compare(CurrentParameterName) == 0)
 										{
 											MessageContexts.AddUnique(Entry.Key->GetMessageContext());
 										}
@@ -540,19 +545,28 @@ namespace UE::Mutable::Private
 								}
 							}
 
-							// Report repeated parameter name
-							if (MessageContexts.Num() >= 2)
+							// The "ensure" will fail if:
+							//		+ A single MessageContexts is found :	A single UE parameter node does get compiled to multiple core compilation parameter nodes. Ensure you are using the params cache.
+							//		+ No MessageContexts are found :		There are parameter core compilation nodes are being generated, and they are not getting a context set in them. Ensure all nodes generated get a context set.
+							if (ensureMsgf(MessageContexts.Num() >= 2, TEXT("An unexpected amount of parameter contexts (%u) have been found while reporting a repeated parameter."), MessageContexts.Num()))
 							{
+								// Report repeated parameter name alongside with the nodes that use it
 								const FString WarningText = FString::Printf(TEXT("Repeated parameter found : \"%s\". Please use a different name."), *CurrentParameterName);
-								GenErrorLog->Add(WarningText, ELMT_WARNING, MessageContexts[0], MessageContexts[1]);
+								GenErrorLog->Add(WarningText, ELMT_ERROR, MessageContexts[0], MessageContexts[1]);
 								RepeatedParameterNames.Add(CurrentParameterName);
 							}
-							else
-							{
-								ensure(false);
-							}
 						}
 					}
+
+					// Early out if a repeated parameter name has been found
+					if (bRepeatedParameterFound)
+					{
+						// Merge the log in the right order
+						GenErrorLog->Merge( m_pD->ErrorLog.Get() );
+						m_pD->ErrorLog = GenErrorLog;
+						return nullptr;
+					}
+
 				}
 			}
 		}
