[mutable] Fixed graph traversal incorrect processing of orphan pins
- Previously we were also using the orphaned pins to navigate the graph, which caused, in the case of the static mesh node, a check failure.

#rnx

diff --git a/Engine/Plugins/Mutable/Source/CustomizableObjectEditor/Private/MuCOE/GraphTraversal.cpp b/Engine/Plugins/Mutable/Source/CustomizableObjectEditor/Private/MuCOE/GraphTraversal.cpp
--- a/Engine/Plugins/Mutable/Source/CustomizableObjectEditor/Private/MuCOE/GraphTraversal.cpp
+++ b/Engine/Plugins/Mutable/Source/CustomizableObjectEditor/Private/MuCOE/GraphTraversal.cpp
@@ -62,7 +62,7 @@ TArray<UEdGraphPin*> FollowPinArray(const UEdGraphPin& Pin, bool bIgnoreOrphan,
 	{
 		const UEdGraphPin& CurrentPin = *PinsToVisit.Pop();
 
-		if (!bIgnoreOrphan && IsPinOrphan(CurrentPin))
+		if (bIgnoreOrphan && IsPinOrphan(CurrentPin))
 		{
 			continue;
 		}
@@ -75,7 +75,7 @@ TArray<UEdGraphPin*> FollowPinArray(const UEdGraphPin& Pin, bool bIgnoreOrphan,
 
 		for (UEdGraphPin* LinkedPin : CurrentPin.LinkedTo)
 		{
-			if (!bIgnoreOrphan && IsPinOrphan(*LinkedPin))
+			if (bIgnoreOrphan && IsPinOrphan(*LinkedPin))
 			{
 				continue;
 			}
