Fixing indeterministic asset registry queries in CustomizableObjectEditor during cook - switched them to only search on-disk packages

diff --git a/Engine/Plugins/Mutable/Source/CustomizableObjectEditor/Private/MuCOE/CustomizableObjectEditorModule.cpp b/Engine/Plugins/Mutable/Source/CustomizableObjectEditor/Private/MuCOE/CustomizableObjectEditorModule.cpp
--- a/Engine/Plugins/Mutable/Source/CustomizableObjectEditor/Private/MuCOE/CustomizableObjectEditorModule.cpp
+++ b/Engine/Plugins/Mutable/Source/CustomizableObjectEditor/Private/MuCOE/CustomizableObjectEditorModule.cpp
@@ -523,6 +523,7 @@ void GetReferencingCustomizableObjects(FName CustomizableObjectName, TArray<FNam
 
 	FARFilter Filter;
 	Filter.PackageNames = MoveTemp(ReferencedObjectNames);
+	Filter.bIncludeOnlyOnDiskAssets = IsRunningCookCommandlet();	// When cooking, only search on-disk packages (ensures deterministic results)
 
 	AssetRegistryModule.Get().GetAssets(Filter, AssetDataArray);
 
@@ -567,9 +568,6 @@ void GetReferencingPackages(const UCustomizableObject& Object, TArray<FAssetData
 		TArray<UCustomizableObjectNodeTable*> TableNodes;
 		ChildCustomizableObject->GetPrivate()->GetSource()->GetNodesOfClass(TableNodes);
 
-		FARFilter Filter;
-		Filter.ClassPaths.Add(FTopLevelAssetPath(UDataTable::StaticClass()));
-
 		for (const UCustomizableObjectNodeTable* TableNode : TableNodes)
 		{
 			TArray<FAssetData> DataTableAssets = TableNode->GetParentTables();
diff --git a/Engine/Plugins/Mutable/Source/CustomizableObjectEditor/Private/MuCOE/GraphTraversal.cpp b/Engine/Plugins/Mutable/Source/CustomizableObjectEditor/Private/MuCOE/GraphTraversal.cpp
--- a/Engine/Plugins/Mutable/Source/CustomizableObjectEditor/Private/MuCOE/GraphTraversal.cpp
+++ b/Engine/Plugins/Mutable/Source/CustomizableObjectEditor/Private/MuCOE/GraphTraversal.cpp
@@ -766,7 +766,7 @@ void GetNodeGroupObjectNodeMappingImmersive(UCustomizableObject* Object, IAssetR
 	AssetRegistry.GetReferencers(*Object->GetOuter()->GetPathName(), ArrayReferenceNames, UE::AssetRegistry::EDependencyCategory::Package, UE::AssetRegistry::EDependencyQuery::Hard);
 
 	FARFilter Filter;
-	Filter.bIncludeOnlyOnDiskAssets = false;
+	Filter.bIncludeOnlyOnDiskAssets = IsRunningCookCommandlet();	// when cooking, only search on-disk packages, to ensure deterministic results
 	
 	for (const FName& ReferenceName : ArrayReferenceNames)
 	{
diff --git a/Engine/Plugins/Mutable/Source/CustomizableObjectEditor/Private/MuCOE/Nodes/CustomizableObjectNodeTable.cpp b/Engine/Plugins/Mutable/Source/CustomizableObjectEditor/Private/MuCOE/Nodes/CustomizableObjectNodeTable.cpp
--- a/Engine/Plugins/Mutable/Source/CustomizableObjectEditor/Private/MuCOE/Nodes/CustomizableObjectNodeTable.cpp
+++ b/Engine/Plugins/Mutable/Source/CustomizableObjectEditor/Private/MuCOE/Nodes/CustomizableObjectNodeTable.cpp
@@ -1459,6 +1459,7 @@ TArray<FAssetData> UCustomizableObjectNodeTable::GetParentTables() const
 		FARFilter Filter;
 		Filter.ClassPaths.Add(FTopLevelAssetPath(UDataTable::StaticClass()));
 		Filter.bRecursivePaths = true;
+		Filter.bIncludeOnlyOnDiskAssets = IsRunningCookCommandlet();	// Only query on-disk assets while cooking - ensures results are deterministic
 
 		if (TObjectPtr<UScriptStruct> LoadedStructure = UE::Mutable::Private::LoadObject(Structure))
 		{
