[Mutable] Fix missing layouts when the base layout is defined in a LOD lower than the FirstLODAvailable.
Fix SkinWeightProfiles crash when the SkeletalMesh data is out of date.

#rnx

diff --git a/Engine/Plugins/Mutable/Source/CustomizableObjectEditor/Private/MuCOE/GenerateMutableSource/GenerateMutableSourceMesh.cpp b/Engine/Plugins/Mutable/Source/CustomizableObjectEditor/Private/MuCOE/GenerateMutableSource/GenerateMutableSourceMesh.cpp
--- a/Engine/Plugins/Mutable/Source/CustomizableObjectEditor/Private/MuCOE/GenerateMutableSource/GenerateMutableSourceMesh.cpp
+++ b/Engine/Plugins/Mutable/Source/CustomizableObjectEditor/Private/MuCOE/GenerateMutableSource/GenerateMutableSourceMesh.cpp
@@ -2191,7 +2191,11 @@ TSharedPtr<UE::Mutable::Private::FMesh> ConvertSkeletalMeshToMutable(TSharedRef<
 				continue;
 			}
 
-			check(Vertices.Num() == ImportedProfileData->SkinWeights.Num());
+			if (ImportedProfileData->SkinWeights.Num() != Vertices.Num())
+			{
+				ensure(false);
+				continue;
+			}
 
 			TArray<uint8> MutSkinWeights;
 			MutSkinWeights.SetNumZeroed(VertexCount * SkinWeightProfileVertexSize);
@@ -3868,7 +3872,7 @@ UE::Mutable::Private::Ptr<UE::Mutable::Private::NodeMesh> GenerateMutableSourceM
 			// 
 			// Generate Layouts when using manual LODs or when processing the base Automatic LOD (the other automatic LODs will use the base layouts).
 			const bool bGenerateLayouts = GenerationContext.CurrentAutoLODStrategy != ECustomizableObjectAutomaticLODStrategy::AutomaticFromMesh
-				|| GenerationContext.CurrentLOD == GenerationContext.FromLOD;
+				|| (GenerationContext.CurrentLOD == GenerationContext.FromLOD || GenerationContext.CurrentLOD <= GenerationContext.FirstLODAvailable[GenerationContext.CurrentMeshComponent]);
 
 			if(bGenerateLayouts)
 			{
@@ -4526,7 +4530,7 @@ UE::Mutable::Private::Ptr<UE::Mutable::Private::NodeMesh> GenerateMutableSourceM
 							// 
 							// Generate Layouts when using manual LODs or when processing the base Automatic LOD (the other automatic LODs will use the base layouts).
 							const bool bGenerateLayouts = GenerationContext.CurrentAutoLODStrategy != ECustomizableObjectAutomaticLODStrategy::AutomaticFromMesh
-								|| GenerationContext.CurrentLOD == GenerationContext.FromLOD;
+								|| (GenerationContext.CurrentLOD == GenerationContext.FromLOD || GenerationContext.CurrentLOD <= GenerationContext.FirstLODAvailable[GenerationContext.CurrentMeshComponent]);
 
 							if (bGenerateLayouts)
 							{
