[Mutable] Async Load SkeletalMeshes before converting them to Mutable FMesh format.
Optimize mesh conversion.
Remove unnecessary reformats and copies of source data.
Fix missing layouts when the base layout is connected below the FirstLODAvailable for a given platform.

#rnx

diff --git a/Engine/Plugins/Mutable/Source/CustomizableObjectEditor/Private/MuCOE/CustomizableObjectCompileRunnable.cpp b/Engine/Plugins/Mutable/Source/CustomizableObjectEditor/Private/MuCOE/CustomizableObjectCompileRunnable.cpp
--- a/Engine/Plugins/Mutable/Source/CustomizableObjectEditor/Private/MuCOE/CustomizableObjectCompileRunnable.cpp
+++ b/Engine/Plugins/Mutable/Source/CustomizableObjectEditor/Private/MuCOE/CustomizableObjectCompileRunnable.cpp
@@ -191,7 +191,7 @@ uint32 FCustomizableObjectCompileRunnable::Run()
 			FString MorphName = MorphNameRef;
 
 			// It would be great to be able to do this conversion in a worker thread, but the engine doesn't support it yet.
-			auto LoadMeshFunc = [MeshData, MorphName, ResolvedMesh, CompletionEvent, WeakCompiler = WeakCompiler]() mutable
+			auto ConvertMeshFunc = [MeshData, MorphName, ResolvedMesh, CompletionEvent, WeakCompiler = WeakCompiler, bRunImmediatlyIfPossible]() mutable
 				{
 					MUTABLE_CPUPROFILER_SCOPE(LoadMeshFunc);
 					check(IsInGameThread());
@@ -212,34 +212,106 @@ uint32 FCustomizableObjectCompileRunnable::Run()
 						return;
 					}
 
-					// Ensure we don't pull unwanted data into the package when cooking.
-					FCookLoadScope CookLoadScope(ECookLoadType::EditorOnly);
-
 					TSharedRef<FMeshConversionContext> MeshConversionContext = MakeShared<FMeshConversionContext>();
 					MeshConversionContext->Source = MeshData;
 					MeshConversionContext->MorphName = MorphName;
 					MeshConversionContext->CompilationContext = Compiler->CompilationContext;
 
-					// Todo: Load assets async
-					MeshConversionContext->ReferencedObjects.Emplace(TStrongObjectPtr<UObject>(UE::Mutable::Private::LoadObject(MeshData.Mesh)));
-					MeshConversionContext->ReferencedObjects.Emplace(TStrongObjectPtr<UObject>(UE::Mutable::Private::LoadObject(MeshData.TableReferenceSkeletalMesh)));
-					MeshConversionContext->ReferencedObjects.Emplace(TStrongObjectPtr<UObject>(UE::Mutable::Private::LoadClass(MeshData.AnimInstance)));
+					// Convert Synchronously
+					if (bRunImmediatlyIfPossible)
+					{
+						MeshConversionContext->ReferencedObjects.Emplace(TStrongObjectPtr<UObject>(UE::Mutable::Private::LoadObject(MeshData.Mesh)));
+						MeshConversionContext->ReferencedObjects.Emplace(TStrongObjectPtr<UObject>(UE::Mutable::Private::LoadObject(MeshData.TableReferenceSkeletalMesh)));
+						MeshConversionContext->ReferencedObjects.Emplace(TStrongObjectPtr<UObject>(UE::Mutable::Private::LoadClass(MeshData.AnimInstance)));
+
+						UE::Tasks::FTask ConvertTask = ConvertSkeletalMeshToMutable(MeshConversionContext);
+						ConvertTask.Wait();
+						
+						*ResolvedMesh = MeshConversionContext->Result;
+						CompletionEvent.Trigger();
+						return;
+					}
 
-					*ResolvedMesh = ConvertSkeletalMeshToMutable(MeshConversionContext);
+					auto ConvertSkeletalMesh = [ResolvedMesh, MeshConversionContext, CompletionEvent]()
+						{
+							UE::Tasks::FTask ConvertTask = ConvertSkeletalMeshToMutable(MeshConversionContext);
+
+							UE::Tasks::FTask LaunchTask = UE::Tasks::Launch(TEXT("ConvertSkeletalMeshToMutable Callback"),
+								[ResolvedMesh, MeshConversionContext, CompletionEvent]() mutable
+								{
+									*ResolvedMesh = MeshConversionContext->Result;
+									CompletionEvent.Trigger();
+								},
+								ConvertTask,
+								UE::Tasks::ETaskPriority::Inherit);
+						};
+
+					// Find assets to stream
+					TArray<FSoftObjectPath> AssetsToStream;
+					auto AddReferencedStreamableRenderAsset = [MeshConversionContext, &AssetsToStream](TSoftObjectPtr<UStreamableRenderAsset>& SoftObject)
+						{
+							if (UObject* Object = SoftObject.Get())
+							{
+								MeshConversionContext->ReferencedObjects.Add(TStrongObjectPtr<UObject>(Object));
+							}
+							else if (!SoftObject.IsNull())
+							{
+								AssetsToStream.Add(SoftObject.ToSoftObjectPath());
+							}
+						};
+
+					AddReferencedStreamableRenderAsset(MeshData.Mesh);
+					AddReferencedStreamableRenderAsset(MeshData.TableReferenceSkeletalMesh);
+
+					if (UObject* Object = MeshData.AnimInstance.Get())
+					{
+						MeshConversionContext->ReferencedObjects.Add(TStrongObjectPtr<UObject>(Object));
+					}
+					else if (!MeshData.AnimInstance.IsNull())
+					{
+						AssetsToStream.Add(MeshData.AnimInstance.ToSoftObjectPath());
+					}
 
-					CompletionEvent.Trigger();
+					if (AssetsToStream.IsEmpty())
+					{
+						ConvertSkeletalMesh();
+					}
+					else
+					{
+						UCustomizableObjectSystem* System = UCustomizableObjectSystem::GetInstanceChecked();
+						const TSharedRef<FMutableStreamableManager>& Streamable = System->GetPrivate()->StreamableManager;
+
+						auto OnAssetLoadCallback = [MeshConversionContext, ConvertSkeletalMesh]()
+							{
+								MeshConversionContext->ReferencedObjects.Emplace(TStrongObjectPtr<UObject>(MeshConversionContext->Source.Mesh.Get()));
+								MeshConversionContext->ReferencedObjects.Emplace(TStrongObjectPtr<UObject>(MeshConversionContext->Source.TableReferenceSkeletalMesh.Get()));
+								MeshConversionContext->ReferencedObjects.Emplace(TStrongObjectPtr<UObject>(MeshConversionContext->Source.AnimInstance.Get()));
+
+								ConvertSkeletalMesh();
+							};
+
+						if (IsRunningCookCommandlet())
+						{
+							Streamable->RequestSyncLoad(AssetsToStream);
+							OnAssetLoadCallback();
+						}
+						else
+						{
+							Streamable->RequestAsyncLoad(AssetsToStream, OnAssetLoadCallback, true);
+						}
+					}
 				};
 
 			if (IsInGameThread())
 			{
-				LoadMeshFunc();
+				ConvertMeshFunc();
 			}
 			else
 			{
 				check(!bRunImmediatlyIfPossible);
 				if (TSharedPtr<FCustomizableObjectCompiler> Compiler = WeakCompiler.Pin())
 				{
-					Compiler->AddGameThreadCompileTask(MoveTemp(LoadMeshFunc));	
+					Compiler->AddGameThreadCompileTask(MoveTemp(ConvertMeshFunc));
 				}
 				else
 				{
diff --git a/Engine/Plugins/Mutable/Source/CustomizableObjectEditor/Private/MuCOE/CustomizableObjectCompiler.cpp b/Engine/Plugins/Mutable/Source/CustomizableObjectEditor/Private/MuCOE/CustomizableObjectCompiler.cpp
--- a/Engine/Plugins/Mutable/Source/CustomizableObjectEditor/Private/MuCOE/CustomizableObjectCompiler.cpp
+++ b/Engine/Plugins/Mutable/Source/CustomizableObjectEditor/Private/MuCOE/CustomizableObjectCompiler.cpp
@@ -861,7 +861,7 @@ void FCustomizableObjectCompiler::ProcessCompileTasks()
 	check(IsInGameThread());
 
 	// See if there are compile-time tasks to run
-	constexpr double MaxSecondsPerFrame = 0.4;
+	constexpr double MaxSecondsPerFrame = 0.1;
 	double MaxTime = FPlatformTime::Seconds() + MaxSecondsPerFrame;
 
 	TFunction<void()> Task;
diff --git a/Engine/Plugins/Mutable/Source/CustomizableObjectEditor/Private/MuCOE/GenerateMutableSource/GenerateMutableSource.h b/Engine/Plugins/Mutable/Source/CustomizableObjectEditor/Private/MuCOE/GenerateMutableSource/GenerateMutableSource.h
--- a/Engine/Plugins/Mutable/Source/CustomizableObjectEditor/Private/MuCOE/GenerateMutableSource/GenerateMutableSource.h
+++ b/Engine/Plugins/Mutable/Source/CustomizableObjectEditor/Private/MuCOE/GenerateMutableSource/GenerateMutableSource.h
@@ -36,6 +36,7 @@
 #include "MuT/NodeImageParameter.h"
 #include "UObject/Package.h"
 #include "UObject/StrongObjectPtr.h"
+#include "Tasks/Task.h"
 
 class UCustomizableObjectNodeGroupProjectorParameter;
 class UCustomizableObjectNodeMaterialBase;
@@ -1023,4 +1024,4 @@ uint32 AddUniqueSurfaceMetadata(const FMutableSourceSurfaceMetadata& MeshData, F
 
 /** Convert a mesh to Mutable format.
 */
-TSharedPtr<UE::Mutable::Private::FMesh> ConvertSkeletalMeshToMutable(TSharedRef<FMeshConversionContext> MeshConversionContext);
+UE::Tasks::FTask ConvertSkeletalMeshToMutable(TSharedRef<FMeshConversionContext> MeshConversionContext);
diff --git a/Engine/Plugins/Mutable/Source/CustomizableObjectEditor/Private/MuCOE/GenerateMutableSource/GenerateMutableSourceMesh.cpp b/Engine/Plugins/Mutable/Source/CustomizableObjectEditor/Private/MuCOE/GenerateMutableSource/GenerateMutableSourceMesh.cpp
--- a/Engine/Plugins/Mutable/Source/CustomizableObjectEditor/Private/MuCOE/GenerateMutableSource/GenerateMutableSourceMesh.cpp
+++ b/Engine/Plugins/Mutable/Source/CustomizableObjectEditor/Private/MuCOE/GenerateMutableSource/GenerateMutableSourceMesh.cpp
@@ -16,6 +16,7 @@
 #include "MuCO/CustomizableObjectSystem.h"
 #include "MuCO/UnrealConversionUtils.h"
 #include "MuCO/MutableMeshBufferUtils.h"
+#include "MuCO/UnrealToMutableTextureConversionUtils.h"
 #include "MuCOE/CustomizableObjectCompiler.h"
 #include "MuCOE/CustomizableObjectLayout.h"
 #include "MuCOE/GenerateMutableSource/GenerateMutableSourceFloat.h"
@@ -52,92 +53,170 @@
 #include "MuT/NodeMeshParameter.h"
 #include "MuR/Types.h"
 #include "MuR/OpMeshFormat.h"
+#include "MuR/ConvertData.h"
 #include "PhysicsEngine/PhysicsAsset.h"
 #include "PhysicsEngine/SkeletalBodySetup.h"
 #include "Engine/SkinnedAssetCommon.h"
 #include "Engine/SkeletalMeshSocket.h"
 #include "AssetRegistry/AssetRegistryModule.h"
 #include "Hash/CityHash.h"
+#include "PackedNormal.h"
 
 #define LOCTEXT_NAMESPACE "CustomizableObjectEditor"
 
+struct FMutableConvertSkeletalMeshContext
+{
+	int32 NumBaseBuffers = 0;
+
+	// Skeleton
+	bool bShouldRemoveBonesInfluences = false;
+	TArray<FBoneIndexType> BoneMap;
+	TArray<FBoneIndexType> RemappedBoneMapIndices;
 
-void SetSurfaceFormat(FMutableCompilationContext& Context,
-	UE::Mutable::Private::FMeshBufferSet& OutVertexBufferFormat,
-	UE::Mutable::Private::FMeshBufferSet& OutIndexBufferFormat,
-	const FMutableGraphMeshGenerationData& MeshData)
+	// Skinning
+	int32 SourceNumBoneInfluences = 0;
+	UE::Mutable::Private::EMeshBufferFormat SourceSkinWeightFormat;
+	int32 NumBoneInfluences = 0;
+
+	// Bake Morph
+	TArray<FVector3f> MorphDeltaPosition;
+	TArray<FVector3f> MorphDeltaNormal;
+};
+
+void SetBuffersFormat(FMutableCompilationContext& Context,
+	FMutableConvertSkeletalMeshContext& MeshContext,
+	const USkeletalMesh& SourceMesh,
+	const FSkeletalMeshModel& SourceImportedModel,
+	int32 LODIndex,
+	int32 SectionIndex,
+	EMutableMeshConversionFlags Flags,
+	UE::Mutable::Private::FMesh& MutableMesh)
 {
+	if (!SourceImportedModel.LODModels.IsValidIndex(LODIndex)
+		|| !SourceImportedModel.LODModels[LODIndex].Sections.IsValidIndex(SectionIndex))
+	{
+		return;
+	}
+
+	const bool bIgnoreSkinning = EnumHasAnyFlags(Flags, EMutableMeshConversionFlags::IgnoreSkinning);
+	const bool bIgnorePhysics = EnumHasAnyFlags(Flags, EMutableMeshConversionFlags::IgnorePhysics);
+	const bool bIgnoreMorphs = EnumHasAnyFlags(Flags, EMutableMeshConversionFlags::IgnoreMorphs);
+	const bool bIgnoreTexCoords = EnumHasAnyFlags(Flags, EMutableMeshConversionFlags::IgnoreTexCoords);
+
+	const FSkeletalMeshLODModel& LODModel = SourceImportedModel.LODModels[LODIndex];
+	const FSkelMeshSection& MeshSection = LODModel.Sections[SectionIndex];
+
+	check(bIgnoreSkinning || MutableMesh.Skeleton);
+	const int32 MaxBoneWeightTypeSizeBytes = Context.Options.b16BitBoneWeightsEnabled ? 2 : 1;
+	const int32 MaxBoneIndexTypeSizeBytes = MeshSection.BoneMap.Num() > 256 ? 2 : 1;
+	const int32 MaxNumBonesPerVertex = MeshSection.GetMaxBoneInfluences();
+	MeshContext.SourceNumBoneInfluences = MaxNumBonesPerVertex;
+
 	// Limit skinning weights if necessary
 	// \todo: make it more flexible to support 3 or 5 or 1 weight, since there is support for this in 4.25
-	const int32 MutableBonesPerVertex = FGPUBaseSkinVertexFactory::UseUnlimitedBoneInfluences(MeshData.MaxNumBonesPerVertex, Context.Options.TargetPlatform) &&
-		MeshData.MaxNumBonesPerVertex < (int32)Context.Options.CustomizableObjectNumBoneInfluences ?
-		MeshData.MaxNumBonesPerVertex :
-		(int32)Context.Options.CustomizableObjectNumBoneInfluences;
+	int32 MutableBonesPerVertex = 0;
+	if (FGPUBaseSkinVertexFactory::GetUnlimitedBoneInfluences(Context.Options.TargetPlatform))
+	{
+		MutableBonesPerVertex = FMath::Min(MaxNumBonesPerVertex, EXTRA_BONE_INFLUENCES);
+	}
+	else
+	{
+		MutableBonesPerVertex = FMath::Min(MaxNumBonesPerVertex, (int32)Context.Options.CustomizableObjectNumBoneInfluences);
+	}
 
 	ensure(MutableBonesPerVertex <= MAX_TOTAL_INFLUENCES);
 
-	if (MutableBonesPerVertex != MeshData.MaxNumBonesPerVertex)
+	if (MutableBonesPerVertex != MaxNumBonesPerVertex)
+	{
+		UE_LOG(LogMutable, Verbose, TEXT("In object [%s] Mesh bone number adjusted from %d to %d."), *Context.Object->GetFullName(), MaxNumBonesPerVertex, MutableBonesPerVertex);
+	}
+	
+	const bool bHasTexCoords = !bIgnoreTexCoords;
+	const bool bHasSkinWeights = !bIgnoreSkinning && MutableBonesPerVertex > 0;
+	const bool bHasVertexColors = SourceMesh.GetHasVertexColors();
+	const bool bHasRealTimeMorphs = !bIgnoreMorphs && Context.Options.bRealTimeMorphTargetsEnabled && SourceMesh.GetMorphTargets().Num();
+	const bool bHasClothing = !bIgnoreSkinning && Context.Options.bClothingEnabled;
+	const bool bHasSkinWeightProfiles = !bIgnoreSkinning && Context.Options.bSkinWeightProfilesEnabled && LODModel.SkinWeightProfiles.Num();
+
+	int32 MutableBufferCount = 2; // Position & Tangents
+
+	if (bHasTexCoords)
 	{
-		UE_LOG(LogMutable, Verbose, TEXT("In object [%s] Mesh bone number adjusted from %d to %d."), *Context.Object->GetFullName(), MeshData.MaxNumBonesPerVertex, MutableBonesPerVertex);
+		++MutableBufferCount;
 	}
 
-	int MutableBufferCount = MUTABLE_VERTEXBUFFER_TEXCOORDS + 1;
-	if (MeshData.bHasVertexColors)
+	if (bHasVertexColors)
 	{
 		++MutableBufferCount;
 	}
 
-	if (MeshData.MaxNumBonesPerVertex > 0 && MeshData.MaxBoneIndexTypeSizeBytes > 0)
+	if (bHasSkinWeights)
 	{
 		++MutableBufferCount;
 	}
 
-	if (MeshData.bHasRealTimeMorphs)
+	MeshContext.NumBaseBuffers = MutableBufferCount;
+
+	if (bHasRealTimeMorphs)
 	{
 		MutableBufferCount += 2;
 	}
 
-	if (MeshData.bHasClothing)
+	if (bHasClothing)
 	{
 		MutableBufferCount += 2;
 	}
 
-	MutableBufferCount += MeshData.SkinWeightProfilesSemanticIndices.Num();
+	if (bHasSkinWeightProfiles)
+	{
+		MutableBufferCount += LODModel.SkinWeightProfiles.Num();
+	}
 
-	OutVertexBufferFormat.SetBufferCount(MutableBufferCount);
+	UE::Mutable::Private::FMeshBufferSet& VertexBuffers = MutableMesh.GetVertexBuffers();
+	VertexBuffers.SetBufferCount(MutableBufferCount);
 
 	int32 CurrentVertexBuffer = 0;
 
 	// Vertex buffer
-	MutableMeshBufferUtils::SetupVertexPositionsBuffer(CurrentVertexBuffer, OutVertexBufferFormat);
+	MutableMeshBufferUtils::SetupVertexPositionsBuffer(CurrentVertexBuffer, VertexBuffers);
+	check(!VertexBuffers.Buffers[CurrentVertexBuffer].HasPadding());
 	++CurrentVertexBuffer;
 
 	// Tangent buffer
-	MutableMeshBufferUtils::SetupTangentBuffer(CurrentVertexBuffer, OutVertexBufferFormat);
+	MutableMeshBufferUtils::SetupTangentBuffer(CurrentVertexBuffer, VertexBuffers);
+	check(!VertexBuffers.Buffers[CurrentVertexBuffer].HasPadding());
 	++CurrentVertexBuffer;
 
 	// Texture coords buffer
-	bool bHighPrecision = true;
-	MutableMeshBufferUtils::SetupTexCoordinatesBuffer(CurrentVertexBuffer, MeshData.NumTexCoordChannels, bHighPrecision, OutVertexBufferFormat);
-	++CurrentVertexBuffer;
+	if (bHasTexCoords)
+	{
+		bool bHighPrecision = true;
+		MutableMeshBufferUtils::SetupTexCoordinatesBuffer(CurrentVertexBuffer, LODModel.NumTexCoords, bHighPrecision, VertexBuffers);
+		check(!VertexBuffers.Buffers[CurrentVertexBuffer].HasPadding());
+		++CurrentVertexBuffer;
+	}
 
 	// Skin buffer
-	if (MeshData.MaxNumBonesPerVertex > 0 && MeshData.MaxBoneIndexTypeSizeBytes > 0)
+	if (bHasSkinWeights)
 	{
-		const int32 MaxBoneWeightTypeSizeBytes = Context.Options.b16BitBoneWeightsEnabled ? 2 : 1;
-		MutableMeshBufferUtils::SetupSkinBuffer(CurrentVertexBuffer, MeshData.MaxBoneIndexTypeSizeBytes, MaxBoneWeightTypeSizeBytes, MutableBonesPerVertex, OutVertexBufferFormat);
+		MutableMeshBufferUtils::SetupSkinBuffer(CurrentVertexBuffer, MaxBoneIndexTypeSizeBytes, MaxBoneWeightTypeSizeBytes, MutableBonesPerVertex, VertexBuffers);
+		check(!VertexBuffers.Buffers[CurrentVertexBuffer].HasPadding());
 		++CurrentVertexBuffer;
 	}
 
+
 	// Colour buffer
-	if (MeshData.bHasVertexColors)
+	if (bHasVertexColors)
 	{
-		MutableMeshBufferUtils::SetupVertexColorBuffer(CurrentVertexBuffer, OutVertexBufferFormat);
+		MutableMeshBufferUtils::SetupVertexColorBuffer(CurrentVertexBuffer, VertexBuffers);
+		check(!VertexBuffers.Buffers[CurrentVertexBuffer].HasPadding());
 		++CurrentVertexBuffer;
 	}
 
+	// Optional buffers that might be removed when converting the mesh
+
 	// MorphTarget vertex tracking info buffers
-	if (MeshData.bHasRealTimeMorphs)
+	if (bHasRealTimeMorphs)
 	{
 		using namespace UE::Mutable::Private;
 		{
@@ -149,7 +228,8 @@ void SetSurfaceFormat(FMutableCompilationContext& Context,
 			const int32 Components[ChannelCount] = { 1 };
 			const int32 Offsets[ChannelCount] = { 0 };
 
-			OutVertexBufferFormat.SetBuffer(CurrentVertexBuffer, ElementSize, ChannelCount, Semantics, SemanticIndices, Formats, Components, Offsets);
+			VertexBuffers.SetBuffer(CurrentVertexBuffer, ElementSize, ChannelCount, Semantics, SemanticIndices, Formats, Components, Offsets);
+			check(!VertexBuffers.Buffers[CurrentVertexBuffer].HasPadding());
 			++CurrentVertexBuffer;
 		}
 
@@ -162,13 +242,14 @@ void SetSurfaceFormat(FMutableCompilationContext& Context,
 			const int32 Components[ChannelCount] = { 1 };
 			const int32 Offsets[ChannelCount] = { 0 };
 
-			OutVertexBufferFormat.SetBuffer(CurrentVertexBuffer, ElementSize, ChannelCount, Semantics, SemanticIndices, Formats, Components, Offsets);
+			VertexBuffers.SetBuffer(CurrentVertexBuffer, ElementSize, ChannelCount, Semantics, SemanticIndices, Formats, Components, Offsets);
+			check(!VertexBuffers.Buffers[CurrentVertexBuffer].HasPadding());
 			++CurrentVertexBuffer;
 		}
 	}
 
 	//Clothing Data Buffer.
-	if (MeshData.bHasClothing)
+	if (bHasClothing)
 	{
 		{
 			using namespace UE::Mutable::Private;
@@ -180,7 +261,8 @@ void SetSurfaceFormat(FMutableCompilationContext& Context,
 			const int32 Components[ChannelCount] = { 1 };
 			const int32 Offsets[ChannelCount] = { 0 };
 
-			OutVertexBufferFormat.SetBuffer(CurrentVertexBuffer, ElementSize, ChannelCount, Semantics, SemanticIndices, Formats, Components, Offsets);
+			VertexBuffers.SetBuffer(CurrentVertexBuffer, ElementSize, ChannelCount, Semantics, SemanticIndices, Formats, Components, Offsets);
+			check(!VertexBuffers.Buffers[CurrentVertexBuffer].HasPadding());
 			++CurrentVertexBuffer;
 		}
 
@@ -194,19 +276,24 @@ void SetSurfaceFormat(FMutableCompilationContext& Context,
 			const int32 Components[ChannelCount] = { 1 };
 			const int32 Offsets[ChannelCount] = { 0 };
 
-			OutVertexBufferFormat.SetBuffer(CurrentVertexBuffer, ElementSize, ChannelCount, Semantics, SemanticIndices, Formats, Components, Offsets);
+			VertexBuffers.SetBuffer(CurrentVertexBuffer, ElementSize, ChannelCount, Semantics, SemanticIndices, Formats, Components, Offsets);
+			check(!VertexBuffers.Buffers[CurrentVertexBuffer].HasPadding());
 			++CurrentVertexBuffer;
 		}
 	}
 
-	for (int32 ProfileSemanticIndex : MeshData.SkinWeightProfilesSemanticIndices)
+	if (bHasSkinWeightProfiles)
 	{
-		MutableMeshBufferUtils::SetupSkinWeightProfileBuffer(CurrentVertexBuffer, MeshData.MaxBoneIndexTypeSizeBytes, 1, MutableBonesPerVertex, ProfileSemanticIndex, OutVertexBufferFormat);
-		++CurrentVertexBuffer;
+		for (const TPair<FName, FImportedSkinWeightProfileData>& Profile : LODModel.SkinWeightProfiles)
+		{
+			const uint32 ProfileId = static_cast<int32>(Context.GetSkinWeightProfileIdUnique(Profile.Key));
+			MutableMeshBufferUtils::SetupSkinWeightProfileBuffer(CurrentVertexBuffer, MaxBoneIndexTypeSizeBytes, MaxBoneWeightTypeSizeBytes, MutableBonesPerVertex, ProfileId, VertexBuffers);
+			++CurrentVertexBuffer;
+		}
 	}
 
 	// Index buffer
-	MutableMeshBufferUtils::SetupIndexBuffer(OutIndexBufferFormat);
+	MutableMeshBufferUtils::SetupIndexBuffer(MutableMesh.GetIndexBuffers());
 }
 
 
@@ -293,8 +380,7 @@ void BuildRemappedBonesArray(const FMutableComponentInfo& InComponentInfo, TObje
 	const int32 NumLODInfos = InSkeletalMesh->GetLODNum();
 
 	// Helper to know which bones have been removed
-	TArray<bool> RemovedBones;
-	RemovedBones.SetNumZeroed(NumBones);
+	TBitArray<> RemovedBones(false, NumBones);
 
 	for (const FBoneIndexType& RequiredBoneIndex : InRequiredBones)
 	{
@@ -333,7 +419,7 @@ void BuildRemappedBonesArray(const FMutableComponentInfo& InComponentInfo, TObje
 					[&BoneInfo](const FBoneReference& BoneReference) { return BoneReference.BoneName == BoneInfo.Name; });
 				
 				bBoneRemoved = BoneToRemove != nullptr;
-				RemovedBones[RequiredBoneIndex] = RemovedBones[RequiredBoneIndex] || bBoneRemoved;
+				RemovedBones[RequiredBoneIndex] |= bBoneRemoved;
 			}
 
 			// Fix up FinalBoneIndex if it has been removed. Root can't be removed
@@ -346,7 +432,7 @@ void BuildRemappedBonesArray(const FMutableComponentInfo& InComponentInfo, TObje
 }
 
 
-void TransferRemovedBonesInfluences(FBoneIndexType* InfluenceBones, uint16* InfluenceWeights, const int32 InfluenceCount, const TArray<FBoneIndexType>& RemappedBoneMapIndices)
+void RemoveBoneInfluences(uint16* InfluenceBones, uint16* InfluenceWeights, const int32 InfluenceCount, const TArray<FBoneIndexType>& RemappedBoneMapIndices)
 {
 	const int32 BoneMapBoneCount = RemappedBoneMapIndices.Num();
 
@@ -383,88 +469,126 @@ void TransferRemovedBonesInfluences(FBoneIndexType* InfluenceBones, uint16* Infl
 }
 
 
-void NormalizeWeights(FBoneIndexType* InfluenceBones, uint16* InfluenceWeights, const int32 InfluenceCount, const int32 MutableInfluenceCount,
-	int32* MutableMaxOrderedWeighsIndices, const int32 MaxSectionBoneMapIndex, const int32 MaxBoneWeight)
+void NormalizeWeights(uint16* InfluenceBones, uint16* InfluenceWeights, const int32 InfluenceCount, const int32 MutableInfluenceCount, const int32 MaxBoneWeight)
 {
-	// First get the indices of the 4 heaviest influences
-	for (int32 i = 0; i < MutableInfluenceCount; ++i)
-	{
-		int32 CurrentMaxWeight = -1;
+	uint16 MaxOrderedBones[MAX_TOTAL_INFLUENCES];
+	uint16 MaxOrderedWeights[MAX_TOTAL_INFLUENCES];
 
-		for (int32 j = 0; j < InfluenceCount; ++j)
-		{
-			bool bIndexAlreadyUsed = false;
-
-			for (int32 k = 0; k < i; ++k)
-			{
-				if (MutableMaxOrderedWeighsIndices[k] == j)
-				{
-					bIndexAlreadyUsed = true;
-					break;
-				}
-				else if (MutableMaxOrderedWeighsIndices[k] < 0)
-				{
-					break;
-				}
-			}
+	int32 TotalWeight = 0;
+	int32 UsedMask = 0;
 
-			if (!bIndexAlreadyUsed && InfluenceWeights[j] > CurrentMaxWeight
-				&& InfluenceBones[j] < MaxSectionBoneMapIndex)
+	// First sort the indices of the influences. Heaviest influences first
+	for (int8 i = 0; i < MutableInfluenceCount; ++i)
+	{
+		uint16 CurrentMaxWeight = 0;
+		int8 CurrentMaxIndex = 0;
+		
+		for (int8 j = 0; j < InfluenceCount; ++j)
+		{
+			if ((UsedMask & (1 << j)) == 0
+				&& InfluenceWeights[j] > CurrentMaxWeight)
 			{
-				MutableMaxOrderedWeighsIndices[i] = j;
 				CurrentMaxWeight = InfluenceWeights[j];
+				CurrentMaxIndex = j;
 			}
 		}
+
+		UsedMask |= (1 << CurrentMaxIndex);
+		MaxOrderedBones[i] = InfluenceBones[CurrentMaxIndex];
+		MaxOrderedWeights[i] = CurrentMaxWeight;
+		TotalWeight += CurrentMaxWeight;
 	}
 
-	// Copy 4 heaviest influences to 4 first indices
-	for (int32 i = 0; i < MutableInfluenceCount; ++i)
+	// Copy bone influences
+	FMemory::Memcpy(&InfluenceBones[0], &MaxOrderedBones[0], MutableInfluenceCount);
+
+	// Renormalize and copy weights
+	if (TotalWeight > 0)
 	{
-		if (i < InfluenceCount)
-		{
-			InfluenceWeights[i] = InfluenceWeights[MutableMaxOrderedWeighsIndices[i]];
-			InfluenceBones[i] = InfluenceBones[MutableMaxOrderedWeighsIndices[i]];
-		}
-		else
+		int32 AssignedWeight = 0;
+		for (int32 j = 1; j < MutableInfluenceCount; ++j)
 		{
-			InfluenceWeights[i] = 0;
-			InfluenceBones[i] = 0;
+			int32 Res = MaxOrderedWeights[j] * MaxBoneWeight / TotalWeight;
+			AssignedWeight += Res;
+			InfluenceWeights[j] = Res;
 		}
+
+		InfluenceWeights[0] = MaxBoneWeight - AssignedWeight;
 	}
+	else
+	{
+		FMemory::Memzero(InfluenceWeights, MutableInfluenceCount * sizeof(InfluenceWeights[0]));
+		InfluenceWeights[0] = MaxBoneWeight;
+	}
+}
 
-	// Actually renormalize the first 4 influences
-	int32 TotalWeight = 0;
+void ConvertVertexSkinningData(uint16* InfluenceBones, const UE::Mutable::Private::EMeshBufferFormat DestBonesFormat
+	, uint16* InfluenceWeights, const UE::Mutable::Private::EMeshBufferFormat DestWeightsFormat, const int32 DestNumInfluences
+	, const int32 MaxWeight, const FMutableConvertSkeletalMeshContext& MeshContext)
+{
+	bool bNormalize = DestNumInfluences != MeshContext.SourceNumBoneInfluences;
 
-	for (int32 j = 0; j < MutableInfluenceCount; ++j)
+	if (MeshContext.bShouldRemoveBonesInfluences)
 	{
-		TotalWeight += InfluenceWeights[j];
+		RemoveBoneInfluences(InfluenceBones, InfluenceWeights, MeshContext.SourceNumBoneInfluences, MeshContext.RemappedBoneMapIndices);
+		bNormalize = true;
 	}
 
-	if (TotalWeight > 0)
+	if (DestWeightsFormat != MeshContext.SourceSkinWeightFormat)
 	{
-		int32 AssignedWeight = 0;
-
-		for (int32 j = 1; j < MAX_TOTAL_INFLUENCES; ++j)
+		uint16* DestWeightData = InfluenceWeights;
+		if (DestWeightsFormat == UE::Mutable::Private::EMeshBufferFormat::NUInt8
+			&& MeshContext.SourceSkinWeightFormat == UE::Mutable::Private::EMeshBufferFormat::NUInt16)
 		{
-			if (j < MutableInfluenceCount)
+			for (int32 Index = 0; Index < MeshContext.SourceNumBoneInfluences; ++Index)
 			{
-				float Aux = InfluenceWeights[j];
-				int32 Res = FMath::RoundToInt(Aux / TotalWeight * MaxBoneWeight);
-				AssignedWeight += Res;
-				InfluenceWeights[j] = Res;
+				*DestWeightData /= (MAX_uint16 / MAX_uint8);
+				++DestWeightData;
 			}
-			else
+			bNormalize = true;
+		}
+		else if (DestWeightsFormat == UE::Mutable::Private::EMeshBufferFormat::NUInt16
+			&& MeshContext.SourceSkinWeightFormat == UE::Mutable::Private::EMeshBufferFormat::NUInt8)
+		{
+			for (int32 Index = 0; Index < MeshContext.SourceNumBoneInfluences; ++Index)
 			{
-				InfluenceWeights[j] = 0;
+				*DestWeightData *= (MAX_uint16 / MAX_uint8);
+				++DestWeightData;
 			}
 		}
+		else
+		{
+			check(false);
+		}
+	}
 
-		InfluenceWeights[0] = MaxBoneWeight - AssignedWeight;
+	if (bNormalize)
+	{
+		NormalizeWeights(InfluenceBones, InfluenceWeights, MeshContext.SourceNumBoneInfluences, DestNumInfluences, MaxWeight);
 	}
-	else
+
+	// Remove Padding
+	if (DestBonesFormat == UE::Mutable::Private::EMeshBufferFormat::UInt8)
 	{
-		FMemory::Memzero(InfluenceWeights, MutableInfluenceCount*sizeof(InfluenceWeights[0]));
-		InfluenceWeights[0] = MaxBoneWeight;
+		uint8* DestBoneData = reinterpret_cast<uint8*>(InfluenceBones);
+		for (int32 Index = 0; Index < DestNumInfluences; ++Index)
+		{
+			*DestBoneData = uint8(*InfluenceBones);
+			++DestBoneData;
+			++InfluenceBones;
+		}
+	}
+
+	// RemovePadding
+	if (DestWeightsFormat == UE::Mutable::Private::EMeshBufferFormat::NUInt8)
+	{
+		uint8* DestWeightData = reinterpret_cast<uint8*>(InfluenceWeights);
+		for (int32 Index = 0; Index < DestNumInfluences; ++Index)
+		{
+			*DestWeightData = uint8(*InfluenceWeights);
+			++DestWeightData;
+			++InfluenceWeights;
+		}
 	}
 }
 
@@ -575,7 +699,17 @@ TArray<TTuple<UPhysicsAsset*, int32>> GetPhysicsAssetsFromAnimInstance(const TSo
 	}
 
 	// TODO: Get Physics Asset using AssetRegistry
-	UClass* AnimInstanceClass = UE::Mutable::Private::LoadClass(AnimInstance);
+	UClass* AnimInstanceClass = nullptr;
+	if (IsInGameThread())
+	{
+		AnimInstanceClass = UE::Mutable::Private::LoadClass(AnimInstance);
+	}
+	else
+	{
+		AnimInstanceClass = AnimInstance.Get();
+		check(AnimInstanceClass);
+	}
+
 	UAnimBlueprintGeneratedClass* AnimClass = Cast<UAnimBlueprintGeneratedClass>(AnimInstanceClass);
 
 	if (AnimClass)
@@ -769,86 +903,258 @@ TSharedPtr<UE::Mutable::Private::FPhysicsBody> MakePhysicsBodyFromAsset(FMutable
 
 namespace UE::Mutable::Private
 {
-
-	/** Copy mesh data from a single source buffer into a set of destination buffers assuming the format and components is the same. */
-	void CopyBufferClearingPadding(UE::Tasks::FTaskEvent CompletionEvent, UE::Mutable::Private::FMeshBufferSet& BufferSet, const UE::Mutable::Private::FMeshBufferSet& SourceBufferFormat, const void* InSourceData)
+	UE::Tasks::FTask ConvertSkeletalMeshDefaultBuffers(TSharedPtr<UE::Mutable::Private::FMesh> MutableMesh, const FSkeletalMeshLODModel& LODModel, const FSkelMeshSection& MeshSection, TSharedPtr<FMutableConvertSkeletalMeshContext>& MeshContext)
 	{
-		MUTABLE_CPUPROFILER_SCOPE(CopyBufferClearingPadding);
+		TArrayView<const FSoftSkinVertex> Vertices(MeshSection.SoftVertices);
 
-		for (int32 BufferIndex = 0; BufferIndex < BufferSet.GetBufferCount(); ++BufferIndex)
+		const int32 NumVertices = Vertices.Num();
+		check(Vertices.Num() == MutableMesh->GetVertexCount());
+		if (NumVertices == 0)
 		{
-			CompletionEvent.AddPrerequisites(
+			return UE::Tasks::MakeCompletedTask<void>();
+		}
+
+		TArray<UE::Tasks::FTask> Prerequisites;
+
+		// Vertices
+		const int32 NumVertexBuffers = MutableMesh->GetVertexBuffers().GetBufferCount();
+		for (int32 BufferIndex = 0; BufferIndex < MeshContext->NumBaseBuffers; ++BufferIndex)
+		{
+			Prerequisites.Add(
 				UE::Tasks::Launch(TEXT("BufferConversion"),
-				[BufferIndex, &BufferSet, &SourceBufferFormat, InSourceData]()
-				{
-						int32 ElementCount = BufferSet.ElementCount;
-						int32 ChannelCount = BufferSet.Buffers[BufferIndex].Channels.Num();
-						int32 ElementSize = BufferSet.Buffers[BufferIndex].ElementSize;
+					[BufferIndex, MutableMesh, MeshContext, Vertices, NumVertices]()
+					{
+						const FMutableConvertSkeletalMeshContext& MeshContextRef = *MeshContext.Get();
 
-						check(SourceBufferFormat.Buffers.Num() == 1);
-						const UE::Mutable::Private::FMeshBuffer& SourceBuffer = SourceBufferFormat.Buffers[0];
+						UE::Mutable::Private::FMeshBuffer& VertexBuffer = MutableMesh->GetVertexBuffers().Buffers[BufferIndex];
+						uint8* DestData = VertexBuffer.Data.GetData();
+						const int32 DestDataSize = VertexBuffer.ElementSize;
 
-						uint8* TargetData = BufferSet.GetBufferData(BufferIndex);
-						TArray<const uint8*, TInlineAllocator<8>> SourceDataPerChannel;
-						SourceDataPerChannel.SetNumUninitialized(ChannelCount);
-						for (int32 ChannelIndex = 0; ChannelIndex < ChannelCount; ++ChannelIndex)
+						const int32 NumChannels = VertexBuffer.Channels.Num();
+						if (VertexBuffer.Channels[0].Semantic == UE::Mutable::Private::EMeshBufferSemantic::Position)
 						{
-							const UE::Mutable::Private::FMeshBufferChannel& DestinationChannel = BufferSet.Buffers[BufferIndex].Channels[ChannelIndex];
+							MUTABLE_CPUPROFILER_SCOPE(ConvertBuffers_Position);
 
-							int32 SourceBufferIndex = -1;
-							int32 SourceChannelIndex = -1;
-							SourceBufferFormat.FindChannel(DestinationChannel.Semantic, DestinationChannel.SemanticIndex, &SourceBufferIndex, &SourceChannelIndex);
-							check(SourceBufferIndex == 0 && SourceChannelIndex >= 0);
+							check(NumChannels == 1);
 
-							const UE::Mutable::Private::FMeshBufferChannel& SourceChannel = SourceBuffer.Channels[SourceChannelIndex];
-							check(SourceChannel.Format == DestinationChannel.Format);
-							check(SourceChannel.ComponentCount == DestinationChannel.ComponentCount);
+							FVector3f* DestPosition = reinterpret_cast<FVector3f*>(VertexBuffer.Data.GetData());
+							if (MeshContext->MorphDeltaPosition.IsEmpty())
+							{
+								for (const FSoftSkinVertex& Vertex : Vertices)
+								{
+									*DestPosition = Vertex.Position;
+									++DestPosition;
+								}
+							}
+							else
+							{
+								check(MeshContext->MorphDeltaPosition.Num() == NumVertices);
 
-							SourceDataPerChannel[ChannelIndex] = reinterpret_cast<const uint8*>(InSourceData) + SourceChannel.Offset;
+								const FVector3f* MorphDeltaPosition = MeshContext->MorphDeltaPosition.GetData();
+								for (const FSoftSkinVertex& Vertex : Vertices)
+								{
+									*DestPosition = Vertex.Position + *MorphDeltaPosition;
+									++DestPosition;
+									++MorphDeltaPosition;
+								}
+							}
 						}
+						else if (NumChannels == 2
+							&& VertexBuffer.Channels[0].Semantic == UE::Mutable::Private::EMeshBufferSemantic::Tangent
+							&& VertexBuffer.Channels[1].Semantic == UE::Mutable::Private::EMeshBufferSemantic::Normal)
+						{
+							MUTABLE_CPUPROFILER_SCOPE(ConvertBuffers_Tangent);
+
+							// TODO: ParallelFor x2
+							const int32 NumComponents = VertexBuffer.Channels[0].ComponentCount;
+							check(VertexBuffer.Channels[1].ComponentCount == NumComponents);
+
+							const UE::Mutable::Private::EMeshBufferFormat TangentFormat = VertexBuffer.Channels[0].Format;
+							const UE::Mutable::Private::EMeshBufferFormat NormalFormat = VertexBuffer.Channels[1].Format;
+
+							const int32 TangentDataSize = UE::Mutable::Private::GetMeshFormatData(TangentFormat).SizeInBytes * NumComponents;
+							const int32 NormalDataSize = UE::Mutable::Private::GetMeshFormatData(NormalFormat).SizeInBytes * NumComponents;
+							check(TangentDataSize == NormalDataSize && TangentDataSize == sizeof(FPackedNormal));
+
+							FPackedNormal* TypedDestData = reinterpret_cast<FPackedNormal*>(DestData);
+
+							if (MeshContext->MorphDeltaNormal.IsEmpty())
+							{
+								for (const FSoftSkinVertex& Vertex : Vertices)
+								{
+									// Tangent
+									*TypedDestData = Vertex.TangentX;
+									(*TypedDestData).Vector.W = 0;
+									++TypedDestData;
+
+									// Normal
+									*TypedDestData = Vertex.TangentZ;
+									FMatrix44f Mat(Vertex.TangentX, Vertex.TangentY, FVector3f(Vertex.TangentZ), FVector3f(0, 0, 0));
+									(*TypedDestData).Vector.W = Mat.RotDeterminant() < 0 ? -128 : 127;
+									++TypedDestData;
+								}
+							}
+							else
+							{
+								check(MeshContext->MorphDeltaNormal.Num() == NumVertices);
 
-						for (int32 Element = 0; Element < ElementCount; ++Element)
+								FVector3f* MorphDeltaNormal = MeshContext->MorphDeltaNormal.GetData();
+								for (const FSoftSkinVertex& Vertex : Vertices)
+								{
+									// Tangent
+									*TypedDestData = Vertex.TangentX;
+									(*TypedDestData).Vector.W = 0;
+									++TypedDestData;
+
+									// Normal
+									*MorphDeltaNormal += Vertex.TangentZ;
+									*TypedDestData = *MorphDeltaNormal;
+
+									FMatrix44f Mat(Vertex.TangentX, Vertex.TangentY, *MorphDeltaNormal, FVector3f(0, 0, 0));
+									(*TypedDestData).Vector.W = Mat.RotDeterminant() < 0 ? -128 : 127;
+									++TypedDestData;
+									++MorphDeltaNormal;
+								}
+							}
+						}
+						else if (NumChannels == 2
+							&& VertexBuffer.Channels[0].Semantic == UE::Mutable::Private::EMeshBufferSemantic::BoneIndices
+							&& VertexBuffer.Channels[1].Semantic == UE::Mutable::Private::EMeshBufferSemantic::BoneWeights)
 						{
-							int32 CurrentOffset = 0;
-							for (int32 ChannelIndex = 0; ChannelIndex < ChannelCount; ++ChannelIndex)
+							MUTABLE_CPUPROFILER_SCOPE(ConvertBuffers_Skin);
+							// TODO: ParallelFor x4
+							const int32 MaxSectionBoneInfluences = MeshContextRef.SourceNumBoneInfluences;
+							const int32 NumBoneInfluences = VertexBuffer.Channels[0].ComponentCount;
+
+							const UE::Mutable::Private::EMeshBufferFormat BoneIndicesFormat = VertexBuffer.Channels[0].Format;
+							const UE::Mutable::Private::EMeshBufferFormat BoneWeightFormat = VertexBuffer.Channels[1].Format;
+
+							const int32 BoneIndicesSizeBytes = UE::Mutable::Private::GetMeshFormatData(BoneIndicesFormat).SizeInBytes * NumBoneInfluences;
+							const int32 BoneWeightsSizeBytes = UE::Mutable::Private::GetMeshFormatData(BoneWeightFormat).SizeInBytes * NumBoneInfluences;
+
+							const uint16 MaxBoneWeightValue = BoneWeightFormat == UE::Mutable::Private::EMeshBufferFormat::NUInt16 ? MAX_uint16 : MAX_uint8;
+
+							uint16 InfluenceBones[MAX_TOTAL_INFLUENCES];
+							uint16 InfluenceWeights[MAX_TOTAL_INFLUENCES];
+
+							for (const FSoftSkinVertex& Vertex : Vertices)
 							{
-								const UE::Mutable::Private::FMeshBufferChannel& Channel = BufferSet.Buffers[BufferIndex].Channels[ChannelIndex];
+								FMemory::Memcpy(&InfluenceBones[0], &Vertex.InfluenceBones[0], sizeof(uint16) * MaxSectionBoneInfluences);
+								FMemory::Memcpy(&InfluenceWeights[0], &Vertex.InfluenceWeights[0], sizeof(uint16) * MaxSectionBoneInfluences);
+
+								ConvertVertexSkinningData(&InfluenceBones[0], BoneIndicesFormat,
+									&InfluenceWeights[0], BoneWeightFormat, NumBoneInfluences, MaxBoneWeightValue, MeshContextRef);
+
+								// Final copy
+								FMemory::Memcpy(DestData, &InfluenceBones[0], BoneIndicesSizeBytes);
+								DestData += BoneIndicesSizeBytes;
+								FMemory::Memcpy(DestData, &InfluenceWeights[0], BoneWeightsSizeBytes);
+								DestData += BoneWeightsSizeBytes;
+							}
+						}
+						else if (VertexBuffer.Channels[0].Semantic == UE::Mutable::Private::EMeshBufferSemantic::TexCoords)
+						{
+							MUTABLE_CPUPROFILER_SCOPE(ConvertBuffers_TexCoords);
 
-								int32 ChannelOffset = Channel.Offset;
+							constexpr size_t SoftSkinVertexUVsElemSize = sizeof(TDecay<decltype(DeclVal<FSoftSkinVertex>().UVs[0])>::Type);
 
-								int32 PreviousPadding = ChannelOffset - CurrentOffset;
-								if (PreviousPadding > 0)
+							if (DestDataSize == SoftSkinVertexUVsElemSize * NumChannels)
+							{
+								for (const FSoftSkinVertex& Vertex : Vertices)
 								{
-									FMemory::Memzero(TargetData, PreviousPadding);
-									TargetData += PreviousPadding;
-									CurrentOffset += PreviousPadding;
+									FMemory::Memcpy(DestData, &Vertex.UVs[0], DestDataSize);
+									DestData += DestDataSize;
 								}
+							}
+							else
+							{
+								const UE::Mutable::Private::EMeshBufferFormat DestFormat = VertexBuffer.Channels[0].Format;
+								const int32 DataTypeSize = UE::Mutable::Private::GetMeshFormatData(DestFormat).SizeInBytes;
 
-								int32 ChannelSize = Channel.ComponentCount * GetMeshFormatData(Channel.Format).SizeInBytes;
-								FMemory::Memcpy(TargetData, SourceDataPerChannel[ChannelIndex], ChannelSize);
-								TargetData += ChannelSize;
-								SourceDataPerChannel[ChannelIndex] += SourceBuffer.ElementSize;
-								CurrentOffset += ChannelSize;
+								for (const FSoftSkinVertex& Vertex : Vertices)
+								{
+									for (int32 TexCoordIndex = 0; TexCoordIndex < NumChannels; ++TexCoordIndex)
+									{
+										UE::Mutable::Private::ConvertData(0, DestData, DestFormat, &Vertex.UVs[TexCoordIndex].X, UE::Mutable::Private::EMeshBufferFormat::Float32);
+										DestData += DataTypeSize;
+										UE::Mutable::Private::ConvertData(0, DestData, DestFormat, &Vertex.UVs[TexCoordIndex].Y, UE::Mutable::Private::EMeshBufferFormat::Float32);
+										DestData += DataTypeSize;
+									}
+								}
 							}
+						}
+						else if (VertexBuffer.Channels[0].Semantic == UE::Mutable::Private::EMeshBufferSemantic::Color)
+						{
+							MUTABLE_CPUPROFILER_SCOPE(ConvertBuffers_Color);
+
+							check(NumChannels == 1);
+							check(VertexBuffer.Channels[0].ComponentCount == 4);
+							check(VertexBuffer.Channels[0].Format == UE::Mutable::Private::EMeshBufferFormat::NUInt8);
 
-							// Padding at the end?
-							int32 FinalPadding = ElementSize - CurrentOffset;
-							if (FinalPadding > 0)
+							uint32* TypedDestData = reinterpret_cast<uint32*>(DestData);
+							for (const FSoftSkinVertex& Vertex : Vertices)
 							{
-								FMemory::Memzero(TargetData, FinalPadding);
-								TargetData += FinalPadding;
+								*TypedDestData = Vertex.Color.DWColor();
+								++TypedDestData;
 							}
 						}
-					})
-			);
+						else
+						{
+							unimplemented();
+						}
+					}
+					));
 		}
+
+		// Indices
+		Prerequisites.Add(
+			UE::Tasks::Launch(TEXT("Convert Indices"),
+				[&MeshSection, MutableMesh, &LODModel]()
+				{
+					MUTABLE_CPUPROFILER_SCOPE(ConvertBuffers_Indices);
+
+					const uint32 VertexStart = MeshSection.BaseVertexIndex;
+					const uint32 VertexCount = MeshSection.NumVertices;
+					const uint32 IndexStart = MeshSection.BaseIndex;
+					const uint32 IndexCount = MeshSection.NumTriangles * 3;
+					MutableMesh->GetIndexBuffers().SetElementCount(IndexCount);
+
+					check(LODModel.IndexBuffer.IsValidIndex(IndexStart) && LODModel.IndexBuffer.IsValidIndex(IndexStart + IndexCount - 1));
+					const uint32* IndexDataPtr = &LODModel.IndexBuffer[IndexStart];
+
+					uint32* pDest = reinterpret_cast<uint32*>(MutableMesh->GetIndexBuffers().GetBufferData(0));
+
+					// 32-bit to 32-bit
+					uint32 VertexIndex = 0;
+					for (uint32 Index = 0; Index < IndexCount; ++Index)
+					{
+						VertexIndex = *IndexDataPtr - VertexStart;
+						if (VertexIndex < VertexCount)
+						{
+							*pDest = VertexIndex;
+						}
+						else
+						{
+							// Malformed mesh?
+							ensure(false);
+							*pDest = 0;
+						}
+						++pDest;
+						++IndexDataPtr;
+					}
+				}));
+
+		return UE::Tasks::Launch(TEXT("ConvertSkeletalMeshBuffers Completed"),
+			[]() {}, Prerequisites, LowLevelTasks::ETaskPriority::Inherit);
 	}
 
 }
 
 
-TSharedPtr<UE::Mutable::Private::FMesh> ConvertSkeletalMeshToMutable(TSharedRef<FMeshConversionContext> MeshConversionContext)
+UE::Tasks::FTask ConvertSkeletalMeshToMutable(TSharedRef<FMeshConversionContext> MeshConversionContext)
 {
+	check(IsInGameThread());
+
 	const FMutableSourceMeshData& Source = MeshConversionContext->Source;
 	FMutableCompilationContext& Context = *MeshConversionContext->CompilationContext;
 	const FString& MorphName = MeshConversionContext->MorphName;
@@ -864,7 +1170,7 @@ TSharedPtr<UE::Mutable::Private::FMesh> ConvertSkeletalMeshToMutable(TSharedRef<
 	{
 		const FString Msg = FString::Printf(TEXT("The SkeletalMesh [%s] doesn't have an imported resource."), *MeshName);
 		Context.Log(FText::FromString(Msg), Source.MessageContext);
-		return nullptr;
+		return UE::Tasks::MakeCompletedTask<void>();
 	}
 
 	// Data tables may allow missing data. Don't throw warnings of missing LODs or sections.
@@ -882,7 +1188,8 @@ TSharedPtr<UE::Mutable::Private::FMesh> ConvertSkeletalMeshToMutable(TSharedRef<
 	{
 		if (!bMeshMustExist)
 		{
-			return MakeShared<UE::Mutable::Private::FMesh>(); // Return empty mesh to preserve the layouts
+			MeshConversionContext->Result = MakeShared<UE::Mutable::Private::FMesh>(); // Return empty mesh to preserve the layouts
+			return UE::Tasks::MakeCompletedTask<void>();
 		}
 		else
 		{
@@ -895,7 +1202,7 @@ TSharedPtr<UE::Mutable::Private::FMesh> ConvertSkeletalMeshToMutable(TSharedRef<
 					ImportedModel->LODModels.Num());
 				Context.Log(FText::FromString(Msg), Source.MessageContext);
 			}
-			return nullptr;
+			return UE::Tasks::MakeCompletedTask<void>();
 		}
 	}
 
@@ -905,7 +1212,8 @@ TSharedPtr<UE::Mutable::Private::FMesh> ConvertSkeletalMeshToMutable(TSharedRef<
 	{
 		if (!bMeshMustExist)
 		{
-			return MakeShared<UE::Mutable::Private::FMesh>(); // Return empty mesh to preserve the layouts
+			MeshConversionContext->Result = MakeShared<UE::Mutable::Private::FMesh>(); // Return empty mesh to preserve the layouts
+			return UE::Tasks::MakeCompletedTask<void>();
 		}
 		else
 		{
@@ -920,7 +1228,7 @@ TSharedPtr<UE::Mutable::Private::FMesh> ConvertSkeletalMeshToMutable(TSharedRef<
 					LODModel.Sections.Num());
 				Context.Log(FText::FromString(Msg), Source.MessageContext);
 			}
-			return nullptr;
+			return UE::Tasks::MakeCompletedTask<void>();
 		}
 	}
 
@@ -932,9 +1240,11 @@ TSharedPtr<UE::Mutable::Private::FMesh> ConvertSkeletalMeshToMutable(TSharedRef<
 	const bool bIgnorePhysics = EnumHasAnyFlags(CurrentFlags, EMutableMeshConversionFlags::IgnorePhysics);
 	const bool bIgnoreMorphs = EnumHasAnyFlags(CurrentFlags, EMutableMeshConversionFlags::IgnoreMorphs);
 	const bool bIgnoreAUD = EnumHasAnyFlags(CurrentFlags, EMutableMeshConversionFlags::IgnoreAUD);
+	const bool bIgnoreTexCoords = EnumHasAnyFlags(CurrentFlags, EMutableMeshConversionFlags::IgnoreTexCoords);
 	const bool bDoNotCreateMeshMetadata = EnumHasAnyFlags(CurrentFlags, EMutableMeshConversionFlags::DoNotCreateMeshMetadata);
 
 	TSharedPtr<UE::Mutable::Private::FMesh> MutableMesh = MakeShared<UE::Mutable::Private::FMesh>();
+	TSharedPtr<FMutableConvertSkeletalMeshContext> MeshContext = MakeShared<FMutableConvertSkeletalMeshContext>();
 
 	if (UCustomizableObjectSystem::GetInstance()->IsMutableAnimInfoDebuggingEnabled())
 	{
@@ -968,7 +1278,6 @@ TSharedPtr<UE::Mutable::Private::FMesh> ConvertSkeletalMeshToMutable(TSharedRef<
 
 	bool bBoneMapModified = false;
 	TArray<FBoneIndexType> BoneMap;
-	TArray<FBoneIndexType> RemappedBoneMapIndices;
 
 	// Check if the Skeleton is valid and build the UE::Mutable::Private::Skeleton
 	if (!bIgnoreSkinning)
@@ -978,7 +1287,7 @@ TSharedPtr<UE::Mutable::Private::FMesh> ConvertSkeletalMeshToMutable(TSharedRef<
 		{
 			FString Msg = FString::Printf(TEXT("No skeleton provided when converting SkeletalMesh [%s]."), *MeshName);
 			Context.Log(FText::FromString(Msg), Source.MessageContext);
-			return nullptr;
+			return UE::Tasks::MakeCompletedTask<void>();
 		}
 
 		FMutableComponentInfo* MutComponentInfo = Context.GetComponentInfo(Source.Component);
@@ -1001,7 +1310,7 @@ TSharedPtr<UE::Mutable::Private::FMesh> ConvertSkeletalMeshToMutable(TSharedRef<
 				{
 					Context.Log(FText::FromString(ErrorMessage), Source.MessageContext, EMessageSeverity::Warning);
 				}
-				return nullptr;
+				return UE::Tasks::MakeCompletedTask<void>();
 			}
 
 			// Add the RefSkeleton ID to the mesh.
@@ -1042,7 +1351,7 @@ TSharedPtr<UE::Mutable::Private::FMesh> ConvertSkeletalMeshToMutable(TSharedRef<
 					*Skeleton->GetName(), *MeshName);
 				Context.Log(FText::FromString(Msg), Source.MessageContext);
 
-				return nullptr;
+				return UE::Tasks::MakeCompletedTask<void>();
 			}
 		}
 
@@ -1072,10 +1381,11 @@ TSharedPtr<UE::Mutable::Private::FMesh> ConvertSkeletalMeshToMutable(TSharedRef<
 			const FBoneIndexType FinalBoneIndex = BoneMapBoneIndex < NumRemappedBones ? RemappedBones[BoneMapBoneIndex] : 0;
 
 			const int32 BoneMapIndex = BoneMap.AddUnique(FinalBoneIndex);
-			RemappedBoneMapIndices.Add(BoneMapIndex);
+			MeshContext->RemappedBoneMapIndices.Add(BoneMapIndex);
 
 			bBoneMapModified = bBoneMapModified || SourceBoneMap[BoneIndex] != FinalBoneIndex;
 		}
+		MeshContext->bShouldRemoveBonesInfluences = bBoneMapModified;
 
 		// Create the skeleton, poses, and BoneMap for this mesh
 		TSharedPtr<UE::Mutable::Private::FSkeleton> MutableSkeleton = MakeShared<UE::Mutable::Private::FSkeleton>();
@@ -1157,6 +1467,7 @@ TSharedPtr<UE::Mutable::Private::FMesh> ConvertSkeletalMeshToMutable(TSharedRef<
 			MutableSocket.RelativeScale = Socket->RelativeScale;
 			MutableSocket.bForceAlwaysAnimated = Socket->bForceAlwaysAnimated;
 			MutableSocket.Priority = Source.SocketPriority;
+	
 
 			SocketData.SetNum(0, EAllowShrinking::No);
 			FMemoryWriter MemoryWriter(SocketData);
@@ -1209,326 +1520,26 @@ TSharedPtr<UE::Mutable::Private::FMesh> ConvertSkeletalMeshToMutable(TSharedRef<
 		}
 	}
 
+	// Set Mutable Mesh's final buffer formats
+	SetBuffersFormat(Context, *MeshContext.Get(), *SkeletalMesh, *ImportedModel, LODIndex, SectionIndex, Source.Flags, *MutableMesh.Get());
+
 	// Vertices
-	TArray<FSoftSkinVertex> Vertices;
-	LODModel.GetVertices(Vertices);
-	int32 VertexStart = MeshSection.GetVertexBufferIndex();
-	int32 VertexCount = MeshSection.GetNumVertices();
+	const FSkelMeshSection& Section = LODModel.Sections[SectionIndex];
+	TArrayView<const FSoftSkinVertex> Vertices(Section.SoftVertices);
 
+	const uint32 VertexCount = MeshSection.GetNumVertices();
 	MutableMesh->GetVertexBuffers().SetElementCount(VertexCount);
 
-	MutableMesh->GetVertexBuffers().SetBufferCount(2);
-
 	const int32 MaxSectionInfluences = MeshSection.MaxBoneInfluences;
 	const bool bUseUnlimitedInfluences = FGPUBaseSkinVertexFactory::UseUnlimitedBoneInfluences(MaxSectionInfluences, Context.Options.TargetPlatform);
 
-	// Create a mutable vertex buffer definition that matches the unreal soft vertex format.
-	UE::Mutable::Private::FMeshBufferSet UnrealSourceVertexFormat;
-	{
-		using namespace UE::Mutable::Private;
-		const int32 ElementSize = sizeof(FSoftSkinVertex);
-		const int32 ChannelCount = 11;
-		const EMeshBufferSemantic Semantics[ChannelCount] = { EMeshBufferSemantic::Position, EMeshBufferSemantic::Tangent, EMeshBufferSemantic::Binormal, EMeshBufferSemantic::Normal, EMeshBufferSemantic::TexCoords, EMeshBufferSemantic::TexCoords, EMeshBufferSemantic::TexCoords, EMeshBufferSemantic::TexCoords, EMeshBufferSemantic::Color, EMeshBufferSemantic::BoneIndices, EMeshBufferSemantic::BoneWeights };
-		const int32 SemanticIndices[ChannelCount] = { 0, 0, 0, 0, 0, 1, 2, 3, 0, 0, 0 };
-
-		// TODO: Remove BoneWeightFormat after merge
-		EMeshBufferFormat BoneWeightFormat = sizeof(TDecay<decltype(DeclVal<FSoftSkinVertex>().InfluenceWeights[0])>::Type) == 1 ? EMeshBufferFormat::NUInt8 : EMeshBufferFormat::NUInt16;
-		const EMeshBufferFormat Formats[ChannelCount] = { EMeshBufferFormat::Float32, EMeshBufferFormat::Float32, EMeshBufferFormat::Float32, EMeshBufferFormat::Float32, EMeshBufferFormat::Float32, EMeshBufferFormat::Float32, EMeshBufferFormat::Float32, EMeshBufferFormat::Float32, EMeshBufferFormat::NUInt8, EMeshBufferFormat::UInt16, BoneWeightFormat };
-
-		int32 Components[ChannelCount] = { 3, 3, 3, 4, 2, 2, 2, 2, 4, 4, 4 };
-		if (Context.Options.CustomizableObjectNumBoneInfluences != ECustomizableObjectNumBoneInfluences::Four &&
-			MaxSectionInfluences > 4)
-		{
-			int32 NewBoneInfluencesNum = (int32)Context.Options.CustomizableObjectNumBoneInfluences;
-
-			if (bUseUnlimitedInfluences &&
-				MaxSectionInfluences < NewBoneInfluencesNum)
-			{
-				Components[9] = MaxSectionInfluences;
-				Components[10] = MaxSectionInfluences;
-			}
-			else
-			{
-				Components[9] = NewBoneInfluencesNum;
-				Components[10] = NewBoneInfluencesNum;
-			}
-		}
-
-		constexpr size_t SoftSkinVertexUVsElemSize = sizeof(TDecay<decltype(DeclVal<FSoftSkinVertex>().UVs[0])>::Type);
-		const int32 Offsets[ChannelCount] =
-		{
-			STRUCT_OFFSET(FSoftSkinVertex, Position),
-			STRUCT_OFFSET(FSoftSkinVertex, TangentX),
-			STRUCT_OFFSET(FSoftSkinVertex, TangentY),
-			STRUCT_OFFSET(FSoftSkinVertex, TangentZ),
-			STRUCT_OFFSET(FSoftSkinVertex, UVs) + 0 * SoftSkinVertexUVsElemSize,
-			STRUCT_OFFSET(FSoftSkinVertex, UVs) + 1 * SoftSkinVertexUVsElemSize,
-			STRUCT_OFFSET(FSoftSkinVertex, UVs) + 2 * SoftSkinVertexUVsElemSize,
-			STRUCT_OFFSET(FSoftSkinVertex, UVs) + 3 * SoftSkinVertexUVsElemSize,
-			STRUCT_OFFSET(FSoftSkinVertex, Color),
-			STRUCT_OFFSET(FSoftSkinVertex, InfluenceBones),
-			STRUCT_OFFSET(FSoftSkinVertex, InfluenceWeights),
-		};
-
-		UnrealSourceVertexFormat.SetBufferCount(1);
-		UnrealSourceVertexFormat.SetBuffer(0, ElementSize, ChannelCount, Semantics, SemanticIndices, Formats, Components, Offsets);
-	}
-
-	// Create the mutable mesh with the same data without padding and separating the UVs, also add skinning only optionally.
-	{
-		using namespace UE::Mutable::Private;
-
-		// Base channels
-		{
-			const int32 ChannelCount = 5;
-			const EMeshBufferSemantic Semantics[ChannelCount] = { EMeshBufferSemantic::Position, EMeshBufferSemantic::Tangent, EMeshBufferSemantic::Binormal, EMeshBufferSemantic::Normal, EMeshBufferSemantic::Color };
-			const int32 SemanticIndices[ChannelCount] = { 0, 0, 0, 0, 0 };
-			const EMeshBufferFormat Formats[ChannelCount] = { EMeshBufferFormat::Float32, EMeshBufferFormat::Float32, EMeshBufferFormat::Float32, EMeshBufferFormat::Float32, EMeshBufferFormat::NUInt8 };
-			int32 Components[ChannelCount] = { 3, 3, 3, 4, 4 };
-			const int32 Offsets[ChannelCount] =
-			{
-				0,
-				sizeof(float) * 3,
-				sizeof(float) * 6,
-				sizeof(float) * 9,
-				sizeof(float) * 13,
-			};
-			const int32 ElementSize = sizeof(float) * 13 + 4;
-
-			MutableMesh->GetVertexBuffers().SetBuffer(0, ElementSize, ChannelCount, Semantics, SemanticIndices, Formats, Components, Offsets);
-			check(!MutableMesh->VertexBuffers.Buffers[0].HasPadding());
-		}
-
-		// Texture coordinates
-		{
-			const int32 ElementSize = 4 * 2 * sizeof(float);
-			const int32 ChannelCount = 4;
-			const EMeshBufferSemantic Semantics[ChannelCount] = { EMeshBufferSemantic::TexCoords, EMeshBufferSemantic::TexCoords, EMeshBufferSemantic::TexCoords, EMeshBufferSemantic::TexCoords };
-			const int32 SemanticIndices[ChannelCount] = { 0, 1, 2, 3 };
-			const EMeshBufferFormat Formats[ChannelCount] = { EMeshBufferFormat::Float32, EMeshBufferFormat::Float32, EMeshBufferFormat::Float32, EMeshBufferFormat::Float32 };
-			int32 Components[ChannelCount] = { 2, 2, 2, 2 };
-			const int32 Offsets[ChannelCount] =
-			{
-				0,
-				sizeof(float) * 2,
-				sizeof(float) * 4,
-				sizeof(float) * 6,
-			};
-
-			MutableMesh->GetVertexBuffers().SetBuffer(1, ElementSize, ChannelCount, Semantics, SemanticIndices, Formats, Components, Offsets);
-			check(!MutableMesh->VertexBuffers.Buffers[1].HasPadding());
-		}
-	}
-
-	if (!bIgnoreSkinning)
-	{
-		MutableMesh->GetVertexBuffers().SetBufferCount(3);
-
-		// Skinning buffer
-		using namespace UE::Mutable::Private;
-		const int32 ChannelCount = 2;
-		const EMeshBufferSemantic Semantics[ChannelCount] = { EMeshBufferSemantic::BoneIndices, EMeshBufferSemantic::BoneWeights };
-		const int32 SemanticIndices[ChannelCount] = { 0, 0 };
-
-		// TODO: Remove BoneWeightFormat after merge
-		EMeshBufferFormat BoneWeightFormat = sizeof(TDecay<decltype(DeclVal<FSoftSkinVertex>().InfluenceWeights[0])>::Type) == 1 ? EMeshBufferFormat::NUInt8 : EMeshBufferFormat::NUInt16;
-		const EMeshBufferFormat Formats[ChannelCount] = { EMeshBufferFormat::UInt16, BoneWeightFormat };
-
-		int32 Components[ChannelCount] = { 4, 4 };
-		if (Context.Options.CustomizableObjectNumBoneInfluences != ECustomizableObjectNumBoneInfluences::Four &&
-			MaxSectionInfluences > 4)
-		{
-			int32 NewBoneInfluencesNum = (int32)Context.Options.CustomizableObjectNumBoneInfluences;
-
-			if (bUseUnlimitedInfluences &&
-				MaxSectionInfluences < NewBoneInfluencesNum)
-			{
-				Components[0] = MaxSectionInfluences;
-				Components[1] = MaxSectionInfluences;
-			}
-			else
-			{
-				Components[0] = NewBoneInfluencesNum;
-				Components[1] = NewBoneInfluencesNum;
-			}
-		}
-
-		constexpr size_t SoftSkinVertexUVsElemSize = sizeof(TDecay<decltype(DeclVal<FSoftSkinVertex>().UVs[0])>::Type);
-		const int32 Offsets[ChannelCount] =
-		{
-			0,
-			Components[0] * sizeof(uint16),
-		};
-
-		int32 ElementSize = Components[0] * sizeof(uint16) + Components[1] * ((BoneWeightFormat == EMeshBufferFormat::NUInt8) ? 1 : 2);
+	int32 NextBufferIndex = MeshContext->NumBaseBuffers;
 
-		MutableMesh->GetVertexBuffers().SetBuffer(2, ElementSize, ChannelCount, Semantics, SemanticIndices, Formats, Components, Offsets);
-
-		// Fix bone weights if required (uint8 -> uint16)
-		if (BoneWeightFormat == EMeshBufferFormat::NUInt16 && Vertices.IsValidIndex(VertexStart))
-		{
-			FSoftSkinVertex FirstVertex = Vertices[VertexStart];
-
-			uint16 TotalWeight = 0;
-			for (int32 InfluenceIndex = 0; InfluenceIndex < MaxSectionInfluences; ++InfluenceIndex)
-			{
-				TotalWeight += FirstVertex.InfluenceWeights[InfluenceIndex];
-			}
-
-			if (TotalWeight <= 255)
-			{
-				for (int32 VertexIndex = VertexStart; VertexIndex < VertexStart + VertexCount && VertexIndex < Vertices.Num(); ++VertexIndex)
-				{
-					FSoftSkinVertex& Vertex = Vertices[VertexIndex];
-					for (int32 InfluenceIndex = 0; InfluenceIndex < MaxSectionInfluences; ++InfluenceIndex)
-					{
-						Vertex.InfluenceBones[InfluenceIndex] = Vertex.InfluenceBones[InfluenceIndex] * (65535 / 255);
-					}
-				}
-			}
-		}
-
-		const int32 MaxSectionBoneMapIndex = BoneMap.Num();
-
-		for (int32 VertexIndex = VertexStart; VertexIndex < VertexStart + VertexCount && VertexIndex < Vertices.Num(); ++VertexIndex)
-		{
-			FSoftSkinVertex& Vertex = Vertices[VertexIndex];
-
-			// Transfer removed bones influences to parent bones
-			if (bBoneMapModified)
-			{
-				TransferRemovedBonesInfluences(&Vertex.InfluenceBones[0], &Vertex.InfluenceWeights[0], MaxSectionInfluences, RemappedBoneMapIndices);
-			}
-
-			if (Context.Options.CustomizableObjectNumBoneInfluences == ECustomizableObjectNumBoneInfluences::Four)
-			{
-				// Normalize weights
-				const int32 MaxMutableWeights = 4;
-				int32 MaxOrderedWeighsIndices[MaxMutableWeights] = { -1, -1, -1, -1 };
-
-				const int32 MaxBoneWeightValue = BoneWeightFormat == EMeshBufferFormat::NUInt16 ? 65535 : 255;
-				NormalizeWeights(&Vertex.InfluenceBones[0], &Vertex.InfluenceWeights[0], MaxSectionInfluences, MaxMutableWeights,
-					&MaxOrderedWeighsIndices[0], MaxSectionBoneMapIndex, MaxBoneWeightValue);
-			}
-			else if (Context.Options.CustomizableObjectNumBoneInfluences == ECustomizableObjectNumBoneInfluences::Eight)
-			{
-				if (!bUseUnlimitedInfluences &&
-					MaxSectionInfluences < EXTRA_BONE_INFLUENCES) // EXTRA_BONE_INFLUENCES is ECustomizableObjectNumBoneInfluences::Eight
-				{
-					FMemory::Memzero(&Vertex.InfluenceWeights[MaxSectionInfluences], EXTRA_BONE_INFLUENCES - MaxSectionInfluences);
-				}
-			}
-			else if (Context.Options.CustomizableObjectNumBoneInfluences == ECustomizableObjectNumBoneInfluences::Twelve)
-			{
-				if (!bUseUnlimitedInfluences &&
-					MaxSectionInfluences < MAX_TOTAL_INFLUENCES) // MAX_TOTAL_INFLUENCES is ECustomizableObjectNumBoneInfluences::Twelve
-				{
-					FMemory::Memzero(&Vertex.InfluenceWeights[MaxSectionInfluences], MAX_TOTAL_INFLUENCES - MaxSectionInfluences);
-				}
-			}
-		}
-
-	}
-
-	const FSoftSkinVertex* SourceVertexData = Vertices.GetData() + VertexStart;
-
-	// Apply the morph if necessary
-	TUniquePtr<FSoftSkinVertex[]> MorphedVertexData;
-	if (!bIgnoreMorphs && !MorphName.IsEmpty())
-	{
-		UMorphTarget* MorphTarget = SkeletalMesh->FindMorphTarget(*MorphName);
-		if (MorphTarget && MorphTarget->GetMorphLODModels().IsValidIndex(LODIndex))
-		{
-			// Make a copy of the original vertex data since we will modify it.
-			MorphedVertexData = MakeUnique<FSoftSkinVertex[]>(VertexCount);
-			FMemory::Memcpy(MorphedVertexData.Get(), SourceVertexData, VertexCount*sizeof(FSoftSkinVertex));
-
-			// Bake the morph in the copy of the source vertex data
-			int32 MaterialVertexStart = SkeletalMesh->GetImportedModel()->LODModels[LODIndex].Sections[SectionIndex].GetVertexBufferIndex();
-
-			const FMorphTargetLODModel& MorphLODModel = MorphTarget->GetMorphLODModels()[LODIndex];
-			for (const FMorphTargetDelta& MorphDelta : MorphLODModel.Vertices)
-			{
-				const int32 VertexIndex = MorphDelta.SourceIdx - MaterialVertexStart;
-				if (VertexIndex >= 0 && VertexIndex < VertexCount)
-				{
-					FVector3f& Position = MorphedVertexData[VertexIndex].Position;
-					Position += MorphDelta.PositionDelta;
-
-					FVector4f& Normal = MorphedVertexData[VertexIndex].TangentZ;
-					Normal += FVector4f(MorphDelta.TangentZDelta, 0.0f);
-				}
-			}
-
-			SourceVertexData = MorphedVertexData.Get();
-		}
-	}
-
-	UE::Tasks::FTaskEvent CompletionEvent(TEXT("MeshConversionCompletion"));
-
-	UE::Mutable::Private::CopyBufferClearingPadding(CompletionEvent, MutableMesh->GetVertexBuffers(), UnrealSourceVertexFormat, SourceVertexData);
-
-	// Indices
-	CompletionEvent.AddPrerequisites(
-		UE::Tasks::Launch(TEXT("BufferConversion"),
-			[&MeshSection, MutableMesh, &LODModel, VertexStart, VertexCount]()
-			{
-				const uint32 IndexStart = MeshSection.BaseIndex;
-				const uint32 IndexCount = MeshSection.NumTriangles * 3;
-				MutableMesh->GetIndexBuffers().SetBufferCount(1);
-				MutableMesh->GetIndexBuffers().SetElementCount(IndexCount);
-
-				using namespace UE::Mutable::Private;
-
-				check(LODModel.IndexBuffer.IsValidIndex(IndexStart) && LODModel.IndexBuffer.IsValidIndex(IndexStart + IndexCount - 1));
-				const uint32* IndexDataPtr = &LODModel.IndexBuffer[IndexStart];
-
-				const int32 FinalElementSize = sizeof(uint32);
-				const int32 ChannelCount = 1;
-				const EMeshBufferSemantic Semantics[ChannelCount] = { EMeshBufferSemantic::VertexIndex };
-				const int32 SemanticIndices[ChannelCount] = { 0 };
-				// We force 32 bit indices, since merging meshes may create vertex buffers bigger than the initial mesh
-				// and for now the mutable runtime doesn't handle it.
-				// \TODO: go back to 16-bit indices when possible.
-				EMeshBufferFormat Formats[ChannelCount] = { EMeshBufferFormat::UInt32 };
-				const int32 Components[ChannelCount] = { 1 };
-				const int32 Offsets[ChannelCount] = { 0 };
-
-				MutableMesh->GetIndexBuffers().SetBuffer(0, FinalElementSize, ChannelCount, Semantics, SemanticIndices, Formats, Components, Offsets);
-
-				uint32* pDest = reinterpret_cast<uint32*>(MutableMesh->GetIndexBuffers().GetBufferData(0));
-
-				// 32-bit to 32-bit
-				uint32 VertexIndex = 0;
-				for (uint32 Index = 0; Index < IndexCount; ++Index)
-				{
-					VertexIndex = *IndexDataPtr - VertexStart;
-					if (VertexIndex < (uint32)VertexCount)
-					{
-						*pDest = VertexIndex;
-					}
-					else
-					{
-						// Malformed mesh?
-						ensure(false);
-						*pDest = 0;
-					}
-					++pDest;
-					++IndexDataPtr;
-				}
-			}
-		)
-	);
-
-	// TODO: For now we have to wait here because code belows touches the vertex buffers.
-	// Restructure so that this doesn't happen.
-	CompletionEvent.Trigger();
-	CompletionEvent.Wait();
+	// RealTimeMorphs and Clothing buffers might be removed if unused. Convert them before launching the tasks to convert non-optional buffers.
 
+	// RealTimeMorphs
 	uint32 MorphUniqueHash = 0;
-	
-	int32 NextBufferIndex = MutableMesh->VertexBuffers.Buffers.Num();
-	if (!bIgnoreMorphs && Context.Options.bRealTimeMorphTargetsEnabled)
+	if (!bIgnoreMorphs && Context.Options.bRealTimeMorphTargetsEnabled && SkeletalMesh->GetMorphTargets().Num())
 	{
 
 		// This call involves resolving every TObjectPtr<UMorphTarget> to a UMorphTarget*, so
@@ -1566,11 +1577,12 @@ TSharedPtr<UE::Mutable::Private::FMesh> ConvertSkeletalMeshToMutable(TSharedRef<
 				}
 				else
 				{
-					USkeletalMesh* TableReferenceSkeletalMesh = Cast<USkeletalMesh>(Source.TableReferenceSkeletalMesh.Get());
+					const USkeletalMesh* TableReferenceSkeletalMesh = Cast<USkeletalMesh>(Source.TableReferenceSkeletalMesh.Get());
+					check(TableReferenceSkeletalMesh);
+
 					for (const UMorphTarget* MorphTarget : SkeletalMeshMorphTargets)
 					{
 						check(MorphTarget && TableReferenceSkeletalMesh);
-
 						if (TableReferenceSkeletalMesh->FindMorphTarget(MorphTarget->GetFName()))
 						{
 							MorphTargetsNames.Add(MorphTarget->GetFName());
@@ -1631,36 +1643,22 @@ TSharedPtr<UE::Mutable::Private::FMesh> ConvertSkeletalMeshToMutable(TSharedRef<
 		{
 			MUTABLE_CPUPROFILER_SCOPE(RealTimeMorphTargetProcessing);
 
-			MutableMesh->GetVertexBuffers().SetBufferCount(NextBufferIndex + 2);
-			// MorphTarget vertex block offset.
-			{
-				using namespace UE::Mutable::Private;
-				const int32 ElementSize = sizeof(uint32);
-				const int32 ChannelCount = 1;
-				const EMeshBufferSemantic Semantics[ChannelCount] = { EMeshBufferSemantic::Other };
-				const int32 SemanticIndices[ChannelCount] = { 0 };
-				const EMeshBufferFormat Formats[ChannelCount] = { EMeshBufferFormat::UInt32 };
-				int32 Components[ChannelCount] = { 1 };
-				const int32 Offsets[ChannelCount] = { 0 };
-
-				MutableMesh->GetVertexBuffers().SetBuffer(NextBufferIndex, ElementSize, ChannelCount, Semantics, SemanticIndices, Formats, Components, Offsets, UE::Mutable::Private::EMemoryInitPolicy::Zeroed);
-			}
+			using namespace UE::Mutable::Private;
+			check(MutableMesh->GetVertexBuffers().GetBufferCount() >= NextBufferIndex + 2);
+
+			// MorphTarget block offset buffer
+			check(MutableMesh->GetVertexBuffers().Buffers[NextBufferIndex].Channels.Num() == 1);
+			check(MutableMesh->GetVertexBuffers().Buffers[NextBufferIndex].Channels[0].Semantic == EMeshBufferSemantic::Other);
+			check(MutableMesh->GetVertexBuffers().Buffers[NextBufferIndex].Channels[0].SemanticIndex == 0);
+			check(MutableMesh->GetVertexBuffers().Buffers[NextBufferIndex].Channels[0].Format == EMeshBufferFormat::UInt32);
 
 			// MorphTarget vertex morph count.
 
 			// MorphTarget vertex block id.
-			{
-				using namespace UE::Mutable::Private;
-				const int32 ElementSize = sizeof(uint32);
-				const int32 ChannelCount = 1;
-				const EMeshBufferSemantic Semantics[ChannelCount] = { EMeshBufferSemantic::Other };
-				const int32 SemanticIndices[ChannelCount] = { 1 };
-				const EMeshBufferFormat Formats[ChannelCount] = { EMeshBufferFormat::UInt32 };
-				int32 Components[ChannelCount] = { 1 };
-				const int32 Offsets[ChannelCount] = { 0 };
-
-				MutableMesh->GetVertexBuffers().SetBuffer(NextBufferIndex + 1, ElementSize, ChannelCount, Semantics, SemanticIndices, Formats, Components, Offsets, UE::Mutable::Private::EMemoryInitPolicy::Zeroed);
-			}
+			check(MutableMesh->GetVertexBuffers().Buffers[NextBufferIndex + 1].Channels.Num() == 1);
+			check(MutableMesh->GetVertexBuffers().Buffers[NextBufferIndex + 1].Channels[0].Semantic == EMeshBufferSemantic::Other);
+			check(MutableMesh->GetVertexBuffers().Buffers[NextBufferIndex + 1].Channels[0].SemanticIndex == 1);
+			check(MutableMesh->GetVertexBuffers().Buffers[NextBufferIndex + 1].Channels[0].Format == EMeshBufferFormat::UInt32);
 
 			// Setup MorphTarget reconstruction data.
 			TArrayView<uint32> VertexMorphsOffsetAndCountBufferView(reinterpret_cast<uint32*>(
@@ -1687,7 +1685,7 @@ TSharedPtr<UE::Mutable::Private::FMesh> ConvertSkeletalMeshToMutable(TSharedRef<
 			NameResolutionMap.Reserve(32);
 
 			TArray<FMorphTargetVertexData> MorphsUsed;
-			for (int32 VertexIdx = VertexStart; VertexIdx < VertexStart + VertexCount && VertexIdx < Vertices.Num(); ++VertexIdx)
+			for (uint32 VertexIdx = 0; VertexIdx < VertexCount; ++VertexIdx)
 			{
 				MorphsUsed.Reset(UsedMorphTargets.Num());
 
@@ -1708,7 +1706,7 @@ TSharedPtr<UE::Mutable::Private::FMesh> ConvertSkeletalMeshToMutable(TSharedRef<
 					// The vertices should be sorted by SourceIdx
 					check(MorphLODModels[LODIndex].Vertices.Num() < 2 || MorphLODModels[LODIndex].Vertices[0].SourceIdx < MorphLODModels[LODIndex].Vertices.Last().SourceIdx);
 
-					const int32 VertexFoundIndex = Algo::BinarySearchBy(MorphLODModels[LODIndex].Vertices, (uint32)VertexIdx,
+					const int32 VertexFoundIndex = Algo::BinarySearchBy(MorphLODModels[LODIndex].Vertices, (uint32)VertexIdx + MeshSection.BaseVertexIndex,
 						[](const FMorphTargetDelta& Element) { return Element.SourceIdx; });
 
 					if (VertexFoundIndex == INDEX_NONE)
@@ -1739,7 +1737,7 @@ TSharedPtr<UE::Mutable::Private::FMesh> ConvertSkeletalMeshToMutable(TSharedRef<
 
 					uint32 EncodedVertexOffsetAndCount =
 						(uint32(MorphsMeshData.Num()) & ((1 << Log2MaxNumVerts) - 1)) | (uint32(MorphsUsed.Num()) << Log2MaxNumVerts);
-					VertexMorphsOffsetAndCountBufferView[VertexIdx - VertexStart] = EncodedVertexOffsetAndCount;
+					VertexMorphsOffsetAndCountBufferView[VertexIdx] = EncodedVertexOffsetAndCount;
 
 					MorphsMeshData.Append(MorphsUsed);
 				}
@@ -1804,15 +1802,15 @@ TSharedPtr<UE::Mutable::Private::FMesh> ConvertSkeletalMeshToMutable(TSharedRef<
 				if (MorphUniqueHash != InvalidResourceId)
 				{
 					// Add the data block reference to the mesh vertices.
-					for (int32 VertexIdx = VertexStart; VertexIdx < VertexStart + VertexCount && VertexIdx < Vertices.Num(); ++VertexIdx)
+					for (uint32 VertexIdx = 0; VertexIdx < VertexCount; ++VertexIdx)
 					{
 						// Keep as invalid the vertices with no morphs, this way we can check if a vertex has morphs looking at either
 						// count and offset (Notice that if count is 0 then offset must be 0 as well) or the hash.
 						// This is useful because those two buffers are iterated in different phases when reconstructing the data back after
 						// mesh generation, but we need to know if a vertex has morphs in both phases. 
-						if (VertexMorphsOffsetAndCountBufferView[VertexIdx - VertexStart] != 0)
+						if (VertexMorphsOffsetAndCountBufferView[VertexIdx] != 0)
 						{
-							VertexMorphsResourceIdBufferView[VertexIdx - VertexStart] = MorphUniqueHash;
+							VertexMorphsResourceIdBufferView[VertexIdx] = MorphUniqueHash;
 						}
 					}
 
@@ -1838,12 +1836,15 @@ TSharedPtr<UE::Mutable::Private::FMesh> ConvertSkeletalMeshToMutable(TSharedRef<
 
 			NextBufferIndex += 2;
 		}
+		else
+		{
+			MutableMesh->GetVertexBuffers().RemoveBuffer(NextBufferIndex + 1);
+			MutableMesh->GetVertexBuffers().RemoveBuffer(NextBufferIndex);
+		}
 	}
 
-
+	// Clothing
 	uint32 ClothingUniqueHash = 0;
-	
-	// Clothing vertex info.
 	if (!bIgnoreSkinning && Context.Options.bClothingEnabled)
 	{
 		// Create new asset or find an already created one if the section has clothing assets.
@@ -1978,32 +1979,19 @@ TSharedPtr<UE::Mutable::Private::FMesh> ConvertSkeletalMeshToMutable(TSharedRef<
 
 		if (ClothingAssetIndex != INDEX_NONE && !MeshSection.ClothMappingDataLODs.IsEmpty())
 		{
-			MutableMesh->GetVertexBuffers().SetBufferCount(NextBufferIndex + 2);
-			{
-				using namespace UE::Mutable::Private;
-				const int32 ElementSize = sizeof(int32);
-				const int32 ChannelCount = 1;
-				const EMeshBufferSemantic Semantics[ChannelCount] = { EMeshBufferSemantic::Other };
-				const int32 SemanticIndices[ChannelCount] = { 2 };
-				const EMeshBufferFormat Formats[ChannelCount] = { EMeshBufferFormat::Int32 };
-				int32 Components[ChannelCount] = { 1 };
-				const int32 Offsets[ChannelCount] = { 0 };
-
-				MutableMesh->GetVertexBuffers().SetBuffer(NextBufferIndex, ElementSize, ChannelCount, Semantics, SemanticIndices, Formats, Components, Offsets, UE::Mutable::Private::EMemoryInitPolicy::Zeroed);
-			}
+			using namespace UE::Mutable::Private;
+			check(MutableMesh->GetVertexBuffers().GetBufferCount() >= NextBufferIndex + 2);
 
-			{
-				using namespace UE::Mutable::Private;
-				const int32 ElementSize = sizeof(uint32);
-				const int32 ChannelCount = 1;
-				const EMeshBufferSemantic Semantics[ChannelCount] = { EMeshBufferSemantic::Other };
-				const int32 SemanticIndices[ChannelCount] = { 3 };
-				const EMeshBufferFormat Formats[ChannelCount] = { EMeshBufferFormat::UInt32 };
-				int32 Components[ChannelCount] = { 1 };
-				const int32 Offsets[ChannelCount] = { 0 };
-
-				MutableMesh->GetVertexBuffers().SetBuffer(NextBufferIndex + 1, ElementSize, ChannelCount, Semantics, SemanticIndices, Formats, Components, Offsets, UE::Mutable::Private::EMemoryInitPolicy::Zeroed);
-			}
+			// Clothing buffers already initialized by the SetBuffer
+			check(MutableMesh->GetVertexBuffers().Buffers[NextBufferIndex].Channels.Num() == 1);
+			check(MutableMesh->GetVertexBuffers().Buffers[NextBufferIndex].Channels[0].Semantic == EMeshBufferSemantic::Other);
+			check(MutableMesh->GetVertexBuffers().Buffers[NextBufferIndex].Channels[0].SemanticIndex == 2);
+			check(MutableMesh->GetVertexBuffers().Buffers[NextBufferIndex].Channels[0].Format == EMeshBufferFormat::Int32);
+
+			check(MutableMesh->GetVertexBuffers().Buffers[NextBufferIndex + 1].Channels.Num() == 1);
+			check(MutableMesh->GetVertexBuffers().Buffers[NextBufferIndex + 1].Channels[0].Semantic == EMeshBufferSemantic::Other);
+			check(MutableMesh->GetVertexBuffers().Buffers[NextBufferIndex + 1].Channels[0].SemanticIndex == 3);
+			check(MutableMesh->GetVertexBuffers().Buffers[NextBufferIndex + 1].Channels[0].Format == EMeshBufferFormat::UInt32);
 
 			TArrayView<int32> ClothSectionIndexView(reinterpret_cast<int32*>(
 				MutableMesh->GetVertexBuffers().GetBufferData(NextBufferIndex)), VertexCount);
@@ -2159,28 +2147,64 @@ TSharedPtr<UE::Mutable::Private::FMesh> ConvertSkeletalMeshToMutable(TSharedRef<
 
 			NextBufferIndex += 2;
 		}
-
+		else
+		{
+			MutableMesh->VertexBuffers.RemoveBuffer(NextBufferIndex + 1);
+			MutableMesh->VertexBuffers.RemoveBuffer(NextBufferIndex);
+		}
 	}
 
+	if (!bIgnoreSkinning)
+	{
+		if (MeshSection.NumVertices > 0)
+		{
+			const FSoftSkinVertex& FirstVertex = Vertices[0];
 
-	// SkinWeightProfiles vertex info.
-	if (!bIgnoreSkinning && Context.Options.bSkinWeightProfilesEnabled)
+			uint16 TotalWeight = 0;
+			for (int32 InfluenceIndex = 0; InfluenceIndex < MaxSectionInfluences; ++InfluenceIndex)
+			{
+				TotalWeight += FirstVertex.InfluenceWeights[InfluenceIndex];
+			}
+
+			MeshContext->SourceSkinWeightFormat = TotalWeight > MAX_uint8 ? UE::Mutable::Private::EMeshBufferFormat::NUInt16 : UE::Mutable::Private::EMeshBufferFormat::NUInt8;
+		}
+	}
+
+	// Copy embedded Morph Data
+	if (!bIgnoreMorphs && !MorphName.IsEmpty())
 	{
-		using namespace UE::Mutable::Private;
+		UMorphTarget* MorphTarget = SkeletalMesh->FindMorphTarget(*MorphName);
+		if (MorphTarget && MorphTarget->GetMorphLODModels().IsValidIndex(LODIndex))
+		{
+			const int32 MorphVertexStart = SkeletalMesh->GetImportedModel()->LODModels[LODIndex].Sections[SectionIndex].GetVertexBufferIndex();
+
+			MeshContext->MorphDeltaPosition.SetNumZeroed(MeshSection.NumVertices);
+			MeshContext->MorphDeltaNormal.SetNumZeroed(MeshSection.NumVertices);
 
-		// TODO: Remove BoneWeightFormat after merge
-		const int32 BoneWeightTypeSizeBytes = sizeof(TDecay<decltype(DeclVal<FRawSkinWeight>().InfluenceWeights[0])>::Type);
-		EMeshBufferFormat BoneWeightFormat = BoneWeightTypeSizeBytes == 1 ? EMeshBufferFormat::NUInt8 : EMeshBufferFormat::NUInt16;
+			TArrayView<FVector3f> DeltaPosition(MeshContext->MorphDeltaPosition);
+			TArrayView<FVector3f> DeltaNormal(MeshContext->MorphDeltaNormal);
+
+			const FMorphTargetLODModel& MorphLODModel = MorphTarget->GetMorphLODModels()[LODIndex];
+			for (const FMorphTargetDelta& MorphDelta : MorphLODModel.Vertices)
+			{
+				const int32 VertexIndex = MorphDelta.SourceIdx - MorphVertexStart;
+				if (VertexIndex >= 0 && (uint32)VertexIndex < VertexCount)
+				{
+					DeltaPosition[VertexIndex] = MorphDelta.PositionDelta;
+					DeltaNormal[VertexIndex] = MorphDelta.TangentZDelta;
+				}
+			}
+		}
+	}
 
-		// Limit skinning weights if necessary
-		const int32 MutableBonesPerVertex = bUseUnlimitedInfluences ?
-			MaxSectionInfluences :
-			(int32)Context.Options.CustomizableObjectNumBoneInfluences;
-		const int32 BoneIndicesSize = MutableBonesPerVertex * sizeof(FBoneIndexType);
-		const int32 BoneWeightsSize = MutableBonesPerVertex * BoneWeightTypeSizeBytes;
-		const int32 SkinWeightProfileVertexSize = sizeof(int32) + BoneIndicesSize + BoneWeightsSize;
+	UE::Tasks::FTask ConvertSkeletalMeshBuffersTask = UE::Mutable::Private::ConvertSkeletalMeshDefaultBuffers(MutableMesh, LODModel, MeshSection, MeshContext);
 
-		const int32 MaxSectionBoneMapIndex = MeshSection.BoneMap.Num();
+	// SkinWeightProfiles vertex info.
+	if (!bIgnoreSkinning && Context.Options.bSkinWeightProfilesEnabled)
+	{
+		const FMutableConvertSkeletalMeshContext& MeshContextRef = *MeshContext.Get();
+		const int32 SourceDataSizeBytes = sizeof(uint16) * MeshContextRef.SourceNumBoneInfluences;
+		static_assert(sizeof(FBoneIndexType) == sizeof(uint16));
 
 		const TArray<FSkinWeightProfileInfo>& SkinWeightProfilesInfo = SkeletalMesh->GetSkinWeightProfiles();
 		for (const FSkinWeightProfileInfo& Profile : SkinWeightProfilesInfo)
@@ -2191,85 +2215,87 @@ TSharedPtr<UE::Mutable::Private::FMesh> ConvertSkeletalMeshToMutable(TSharedRef<
 				continue;
 			}
 
-			if (ImportedProfileData->SkinWeights.Num() != Vertices.Num())
+			if (ImportedProfileData->SkinWeights.Num() != LODModel.NumVertices)
 			{
-				ensure(false);
+				UE_LOG(LogMutable, Display, TEXT("Invalid data found in SkeletalMesh [%s] when converting SkinWeightProfiles. Regenerate and save the SkeletalMesh."), *SkeletalMesh->GetName());
 				continue;
 			}
 
-			TArray<uint8> MutSkinWeights;
-			MutSkinWeights.SetNumZeroed(VertexCount * SkinWeightProfileVertexSize);
-			uint8* MutSkinWeightData = MutSkinWeights.GetData();
+			const uint32 ProfileId = static_cast<int32>(Context.GetSkinWeightProfileIdUnique(Profile.Name));
+
+			int32 BufferIndex = INDEX_NONE;
+			int32 ChannelIndex = INDEX_NONE;
+			MutableMesh->GetVertexBuffers().FindChannel(UE::Mutable::Private::EMeshBufferSemantic::AltSkinWeight, ProfileId, &BufferIndex, &ChannelIndex);
 
-			for (int32 VertexIndex = VertexStart; VertexIndex < VertexStart + VertexCount; ++VertexIndex)
+			if (BufferIndex == INDEX_NONE)
 			{
-				FRawSkinWeight SkinWeight = ImportedProfileData->SkinWeights[VertexIndex];
+				check(false);
+				continue;
+			}
 
-				if (bBoneMapModified)
-				{
-					TransferRemovedBonesInfluences(&SkinWeight.InfluenceBones[0], &SkinWeight.InfluenceWeights[0], MaxSectionInfluences, RemappedBoneMapIndices);
-				}
+			UE::Mutable::Private::FMeshBuffer& VertexBuffer = MutableMesh->GetVertexBuffers().Buffers[BufferIndex];
+			uint8* DestData = VertexBuffer.Data.GetData();
+			const int32 DestDataSize = VertexBuffer.ElementSize;
 
-				if (Context.Options.CustomizableObjectNumBoneInfluences == ECustomizableObjectNumBoneInfluences::Four)
-				{
-					// Normalize weights
-					const int32 MaxMutableWeights = 4;
-					int32 MaxOrderedWeighsIndices[MaxMutableWeights] = { -1, -1, -1, -1 };
+			const int32 NumChannels = VertexBuffer.Channels.Num();
+			check(NumChannels == 3);
 
-					const int32 MaxBoneWeightValue = BoneWeightFormat == EMeshBufferFormat::NUInt16 ? 65535 : 255;
-					NormalizeWeights(&SkinWeight.InfluenceBones[0], &SkinWeight.InfluenceWeights[0], MaxSectionInfluences, MaxMutableWeights,
-						&MaxOrderedWeighsIndices[0], MaxSectionBoneMapIndex, MaxBoneWeightValue);
-				}
-				else if (MaxSectionInfluences < MutableBonesPerVertex)
-				{
-					FMemory::Memzero(&SkinWeight.InfluenceWeights[MaxSectionInfluences], MutableBonesPerVertex - MaxSectionInfluences);
-				}
+			const int32 MaxSectionBoneInfluences = MeshContextRef.SourceNumBoneInfluences;
+			const int32 NumBoneInfluences = VertexBuffer.Channels[0].ComponentCount;
+
+			const UE::Mutable::Private::EMeshBufferFormat BoneIndicesFormat = VertexBuffer.Channels[1].Format;
+			const UE::Mutable::Private::EMeshBufferFormat BoneWeightFormat = VertexBuffer.Channels[2].Format;
 
-				if (FMemory::Memcmp(&Vertices[VertexIndex].InfluenceBones[0], &SkinWeight.InfluenceBones[0], BoneIndicesSize) == 0
+			const int32 BoneIndicesSizeBytes = UE::Mutable::Private::GetMeshFormatData(BoneIndicesFormat).SizeInBytes * NumBoneInfluences;
+			const int32 BoneWeightsSizeBytes = UE::Mutable::Private::GetMeshFormatData(BoneWeightFormat).SizeInBytes * NumBoneInfluences;
+
+			const uint16 MaxBoneWeightValue = BoneWeightFormat == UE::Mutable::Private::EMeshBufferFormat::NUInt16 ? MAX_uint16 : MAX_uint8;
+
+			uint16 InfluenceBones[MAX_TOTAL_INFLUENCES];
+			uint16 InfluenceWeights[MAX_TOTAL_INFLUENCES];
+
+			for (uint32 VertexIndex = 0; VertexIndex < VertexCount; ++VertexIndex)
+			{
+				const FRawSkinWeight& SkinWeight = ImportedProfileData->SkinWeights[VertexIndex + MeshSection.BaseVertexIndex];
+
+				if (FMemory::Memcmp(&Vertices[VertexIndex].InfluenceBones[0], &SkinWeight.InfluenceBones[0], SourceDataSizeBytes) == 0
 					&&
-					FMemory::Memcmp(&Vertices[VertexIndex].InfluenceWeights[0], &SkinWeight.InfluenceWeights[0], BoneWeightsSize) == 0)
+					FMemory::Memcmp(&Vertices[VertexIndex].InfluenceWeights[0], &SkinWeight.InfluenceWeights[0], SourceDataSizeBytes) == 0)
 				{
-					MutSkinWeightData += SkinWeightProfileVertexSize;
+					DestData += DestDataSize;
 					continue;
 				}
 
+				FMemory::Memcpy(&InfluenceBones[0], &SkinWeight.InfluenceBones[0], SourceDataSizeBytes);
+				FMemory::Memcpy(&InfluenceWeights[0], &SkinWeight.InfluenceWeights[0], SourceDataSizeBytes);
+
+				ConvertVertexSkinningData(&InfluenceBones[0], BoneIndicesFormat,
+					&InfluenceWeights[0], BoneWeightFormat, NumBoneInfluences, MaxBoneWeightValue, MeshContextRef);
+
 				int32 SkinWeightVertexHash = 0;
-				for (int32 InfluenceIndex = 0; InfluenceIndex < MutableBonesPerVertex; ++InfluenceIndex)
+				for (int32 InfluenceIndex = 0; InfluenceIndex < NumBoneInfluences; ++InfluenceIndex)
 				{
 					SkinWeightVertexHash = HashCombine(SkinWeightVertexHash, SkinWeight.InfluenceBones[InfluenceIndex]);
 					SkinWeightVertexHash = HashCombine(SkinWeightVertexHash, SkinWeight.InfluenceWeights[InfluenceIndex]);
 				}
 
-				FMemory::Memcpy(MutSkinWeightData, &SkinWeightVertexHash, sizeof(int32));
-				MutSkinWeightData += sizeof(int32);
-				FMemory::Memcpy(MutSkinWeightData, &SkinWeight.InfluenceBones[0], BoneIndicesSize);
-				MutSkinWeightData += BoneIndicesSize;
-				FMemory::Memcpy(MutSkinWeightData, &SkinWeight.InfluenceWeights[0], BoneWeightsSize);
-				MutSkinWeightData += BoneWeightsSize;
+				FMemory::Memcpy(DestData, &SkinWeightVertexHash, sizeof(int32));
+				DestData += sizeof(int32);
+				FMemory::Memcpy(DestData, &SkinWeight.InfluenceBones[0], BoneIndicesSizeBytes);
+				DestData += BoneIndicesSizeBytes;
+				FMemory::Memcpy(DestData, &SkinWeight.InfluenceWeights[0], BoneWeightsSizeBytes);
+				DestData += BoneWeightsSizeBytes;
 			}
 
-			const uint32 ProfileId = Context.GetSkinWeightProfileIdUnique(Profile.Name);
 			const int32 ProfileIndex = Context.SkinWeightProfilesInfo.AddUnique({ Profile.Name, ProfileId, false, 0 });
 			FMutableSkinWeightProfileInfo& MutSkinWeightProfileInfo = Context.SkinWeightProfilesInfo[ProfileIndex];
 
 			const FName PlatformName = *Context.Options.TargetPlatform->PlatformName();
 			MutSkinWeightProfileInfo.DefaultProfile = MutSkinWeightProfileInfo.DefaultProfile || Profile.DefaultProfile.GetValueForPlatform(PlatformName);
 			MutSkinWeightProfileInfo.DefaultProfileFromLODIndex = FMath::Min(MutSkinWeightProfileInfo.DefaultProfileFromLODIndex, Profile.DefaultProfileFromLODIndex.GetValueForPlatform(PlatformName));
-
-			// Set up SkinWeightPRofile BufferData
-			const int32 ElementSize = sizeof(int32) + sizeof(FBoneIndexType) + BoneWeightTypeSizeBytes;
-			const int32 ChannelCount = 3;
-			const EMeshBufferSemantic Semantics[ChannelCount] = { EMeshBufferSemantic::AltSkinWeight, EMeshBufferSemantic::BoneIndices, EMeshBufferSemantic::BoneWeights };
-			const int32 SemanticIndices[ChannelCount] = { static_cast<int32>(ProfileId), static_cast<int32>(ProfileId), static_cast<int32>(ProfileId) };
-			const EMeshBufferFormat Formats[ChannelCount] = { EMeshBufferFormat::Int32, EMeshBufferFormat::UInt16, BoneWeightFormat };
-			const int32 Components[ChannelCount] = { 1, MutableBonesPerVertex, MutableBonesPerVertex };
-			const int32 Offsets[ChannelCount] = { 0, sizeof(int32), sizeof(int32) + BoneIndicesSize };
-
-			MutableMesh->GetVertexBuffers().SetBufferCount(NextBufferIndex + 1);
-			MutableMesh->GetVertexBuffers().SetBuffer(NextBufferIndex, ElementSize, ChannelCount, Semantics, SemanticIndices, Formats, Components, Offsets);
-			FMemory::Memcpy(MutableMesh->GetVertexBuffers().GetBufferData(NextBufferIndex), MutSkinWeights.GetData(), VertexCount * SkinWeightProfileVertexSize);
-			++NextBufferIndex;
 		}
+
+		NextBufferIndex += LODModel.SkinWeightProfiles.Num();
 	}
 
 
@@ -2525,6 +2551,19 @@ TSharedPtr<UE::Mutable::Private::FMesh> ConvertSkeletalMeshToMutable(TSharedRef<
 		}
 	}
 
+
+	// Ensure Surface Data
+	UE::Mutable::Private::FMeshSurface& MeshSurface = MutableMesh->Surfaces.Emplace_GetRef();
+	MeshSurface.BoneMapCount = MutableMesh->BoneMap.Num();
+
+	MeshSurface.SubMeshes.Emplace(
+		UE::Mutable::Private::FSurfaceSubMesh
+		{
+			0, MeshSection.NumVertices,
+			0, (int32)MeshSection.NumTriangles * 3,
+			0
+		});
+
 	if (!bDoNotCreateMeshMetadata)
 	{
 		FMutableMeshMetadata MeshMetadata;
@@ -2556,73 +2595,14 @@ TSharedPtr<UE::Mutable::Private::FMesh> ConvertSkeletalMeshToMutable(TSharedRef<
 			UE_LOG(LogMutable, Error, TEXT("Maximum number of meshes reached."));
 		}
 
-		// Ensure Surface Data
-		UE::Mutable::Private::FMeshSurface& MeshSurface = MutableMesh->Surfaces.Emplace_GetRef();
 
-		MeshSurface.SubMeshes.Emplace(
-			UE::Mutable::Private::FSurfaceSubMesh
-			{
-				0, MutableMesh->VertexBuffers.GetElementCount(),
-				0, MutableMesh->IndexBuffers.GetElementCount(),
-				MeshMetadataUniqueHash
-			});
+		MeshSurface.SubMeshes.Last().ExternalId = MeshMetadataUniqueHash;
 
-		MeshSurface.BoneMapCount = MutableMesh->BoneMap.Num();
 	}
 
+	MeshConversionContext->Result = MutableMesh;
 
-	// Convert to the expected mesh format
-	// TODO: generate directly in this format.
-	{
-		FMutableGraphMeshGenerationData MeshData;
-
-		if (Context.Options.bSkinWeightProfilesEnabled && LODIndex >= 0)
-		{
-			if (ImportedModel && ImportedModel->LODModels.IsValidIndex(LODIndex))
-			{
-				for (const auto& SkinWeightProfile : ImportedModel->LODModels[LODIndex].SkinWeightProfiles)
-				{
-					const int32 ProfileId = static_cast<int32>(Context.GetSkinWeightProfileIdUnique(SkinWeightProfile.Key));
-					MeshData.SkinWeightProfilesSemanticIndices.AddUnique(ProfileId);
-				}
-			}
-		}
-
-		if (ImportedModel->LODModels.IsValidIndex(LODIndex) &&
-			ImportedModel->LODModels[LODIndex].Sections.IsValidIndex(SectionIndex))
-		{
-			MeshData.bHasVertexColors = SkeletalMesh->GetHasVertexColors();
-			MeshData.NumTexCoordChannels = ImportedModel->LODModels[LODIndex].NumTexCoords;
-			MeshData.MaxBoneIndexTypeSizeBytes = MutableMesh->BoneMap.Num() > 256 ? 2 : 1;
-			MeshData.MaxNumBonesPerVertex = ImportedModel->LODModels[LODIndex].GetMaxBoneInfluences();
-			MeshData.bHasRealTimeMorphs = Context.Options.bRealTimeMorphTargetsEnabled && SkeletalMesh->GetMorphTargets().Num();
-			MeshData.bHasClothing = Context.Options.bClothingEnabled && ImportedModel->LODModels[LODIndex].HasClothData();
-		}
-
-		TSharedPtr<UE::Mutable::Private::FMesh> ForcedFormat = MakeShared<UE::Mutable::Private::FMesh>();
-		SetSurfaceFormat(Context, ForcedFormat->GetVertexBuffers(), ForcedFormat->GetIndexBuffers(), MeshData);
-
-		TSharedPtr<UE::Mutable::Private::FMesh> FormattedMesh = MakeShared<UE::Mutable::Private::FMesh>();
-		bool bOutSuccess = false;
-		constexpr bool bKeepSystemBuffers = true;
-		constexpr bool bFormatVertices = true;
-		constexpr bool bFormatIndices = true;
-		constexpr bool bIgnoreMissingChannels = true;
-		UE::Mutable::Private::MeshFormat(FormattedMesh.Get(), MutableMesh.Get(), ForcedFormat.Get(),
-			bKeepSystemBuffers,
-			bFormatVertices,
-			bFormatIndices,
-			bIgnoreMissingChannels,
-			bOutSuccess);
-
-		UE::Mutable::Private::MeshOptimizeBuffers(FormattedMesh.Get());
-
-		check(bOutSuccess);
-
-		MutableMesh = FormattedMesh;
-	}
-
-	return MutableMesh;
+	return ConvertSkeletalMeshBuffersTask;
 }
 
 
@@ -2667,17 +2647,20 @@ TSharedPtr<UE::Mutable::Private::FMesh> ConvertSkeletalMeshToMutable(
 	else
 	{
 		// At some point this shouldn't happen anymore because all mesh conversion could be moved to the core compilation stage.
+
 		TSharedRef<FMeshConversionContext> MeshConversionContext = MakeShared<FMeshConversionContext>();
 		MeshConversionContext->Source = Source;
-		//MeshConversionContext->MorphName = FString(); // NoMorph
+		//MeshConversionContext->MorphName = FString(); No Morph
 		MeshConversionContext->CompilationContext = GenerationContext.CompilationContext;
 
-		// Load assets sync
 		MeshConversionContext->ReferencedObjects.Emplace(TStrongObjectPtr<UObject>(UE::Mutable::Private::LoadObject(Source.Mesh)));
 		MeshConversionContext->ReferencedObjects.Emplace(TStrongObjectPtr<UObject>(UE::Mutable::Private::LoadObject(Source.TableReferenceSkeletalMesh)));
 		MeshConversionContext->ReferencedObjects.Emplace(TStrongObjectPtr<UObject>(UE::Mutable::Private::LoadClass(Source.AnimInstance)));
 
-		MutableMesh = ConvertSkeletalMeshToMutable(MeshConversionContext);
+		UE::Tasks::FTask MeshConversionTask = ConvertSkeletalMeshToMutable(MeshConversionContext);
+		MeshConversionTask.Wait();
+
+		MutableMesh = MeshConversionContext->Result;
 	}
 
 	return MutableMesh;
diff --git a/Engine/Plugins/Mutable/Source/CustomizableObjectEditor/Private/MuCOE/GenerateMutableSource/GenerateMutableSourceMesh.h b/Engine/Plugins/Mutable/Source/CustomizableObjectEditor/Private/MuCOE/GenerateMutableSource/GenerateMutableSourceMesh.h
--- a/Engine/Plugins/Mutable/Source/CustomizableObjectEditor/Private/MuCOE/GenerateMutableSource/GenerateMutableSourceMesh.h
+++ b/Engine/Plugins/Mutable/Source/CustomizableObjectEditor/Private/MuCOE/GenerateMutableSource/GenerateMutableSourceMesh.h
@@ -95,10 +95,3 @@ TSharedPtr<UE::Mutable::Private::FMesh> BuildMorphedMutableMesh(const UEdGraphPi
  * @param bOnlyConnectedLOD Corrected LOD and Section will unconditionally always be the connected ones.
  * @return  Mutable Mesh Node. */
 UE::Mutable::Private::Ptr<UE::Mutable::Private::NodeMesh> GenerateMutableSourceMesh(const UEdGraphPin* Pin, FMutableGraphGenerationContext& GenerationContext, const FMutableSourceMeshData& BaseMeshData, bool bLinkedToExtendMaterial, bool bOnlyConnectedLOD);
-
-
-/** */
-extern void SetSurfaceFormat(FMutableCompilationContext&,
-	UE::Mutable::Private::FMeshBufferSet& OutVertexBufferFormat, 
-	UE::Mutable::Private::FMeshBufferSet& OutIndexBufferFormat, 
-	const FMutableGraphMeshGenerationData& MeshData);
