[Mutable] Fix cooked data invalidation due to a redistribution of the bulk data files dedicated to meshes and images.

#rnx
#tests Cooked build

diff --git a/Engine/Plugins/Mutable/Source/CustomizableObject/Internal/MuCO/CustomizableObjectPrivate.h b/Engine/Plugins/Mutable/Source/CustomizableObject/Internal/MuCO/CustomizableObjectPrivate.h
--- a/Engine/Plugins/Mutable/Source/CustomizableObject/Internal/MuCO/CustomizableObjectPrivate.h
+++ b/Engine/Plugins/Mutable/Source/CustomizableObject/Internal/MuCO/CustomizableObjectPrivate.h
@@ -1008,6 +1008,13 @@ namespace UE::Mutable::Private
 
 		friend uint32 GetTypeHash(const FFileCategoryID& Key);
 		bool operator==(const FFileCategoryID& Other) const = default;
+		friend FArchive& operator<<(FArchive& Ar, FFileCategoryID& Data)
+		{
+			Ar << Data.DataType;
+			Ar << Data.ResourceType;
+			Ar << Data.Flags;
+			return Ar;
+		};
 	};
 
 
@@ -1025,6 +1032,13 @@ namespace UE::Mutable::Private
 	};
 
 
+	struct FFileCategoryOverride
+	{
+		FFileCategoryID Id;
+		uint32 NumFiles = 0;
+	};
+
+
 	/** Group bulk data by categories. */
 	struct FFileBucket
 	{
@@ -1094,10 +1108,13 @@ namespace UE::Mutable::Private
 	 *  Resources will be split into two buckets for non-optional and optional BulkData.
 	 */
 	void CUSTOMIZABLEOBJECT_API GenerateBulkDataFilesListWithFileLimit(
+		FGuid ObjectId,
 		TSharedPtr<const UE::Mutable::Private::FModel, ESPMode::ThreadSafe> Model,
 		FModelStreamableBulkData& ModelStreamableBulkData,
 		uint32 NumFilesPerBucket,
+		bool bAllowFileCategoryOverride,
 		bool bAllowSplit,
+		const ITargetPlatform& TargetPlatform,
 		TArray<FFile>& OutBulkDataFiles);
 
 	/** Generate the list of BulkData files with a soft restriction to the size of the files.
diff --git a/Engine/Plugins/Mutable/Source/CustomizableObject/Private/MuCO/CustomizableObject.cpp b/Engine/Plugins/Mutable/Source/CustomizableObject/Private/MuCO/CustomizableObject.cpp
--- a/Engine/Plugins/Mutable/Source/CustomizableObject/Private/MuCO/CustomizableObject.cpp
+++ b/Engine/Plugins/Mutable/Source/CustomizableObject/Private/MuCO/CustomizableObject.cpp
@@ -2473,10 +2473,13 @@ namespace UE::Mutable::Private
 	}
 	
 	void GenerateBulkDataFilesListWithFileLimit(
+		FGuid ObjectId,
 		TSharedPtr<const UE::Mutable::Private::FModel, ESPMode::ThreadSafe> Model,
 		FModelStreamableBulkData& ModelStreamableBulkData,
 		uint32 NumFilesPerBucket,
+		bool bAllowFileCategoryOverride,
 		bool bAllowSplit,
+		const ITargetPlatform& TargetPlatform,
 		TArray<FFile>& OutBulkDataFiles)
 	{
 		MUTABLE_CPUPROFILER_SCOPE(GenerateBulkDataFilesListWithFileLimit);
@@ -2555,11 +2558,60 @@ namespace UE::Mutable::Private
 		
 		// Limited number of files in each bucket. Find the ideal file distribution between categories based on the accumulated size of their resources.
 		TArray<FFileCategory> Categories;
+		TArray<FFileCategoryOverride> FileCategoryOverrides;
+
+		using namespace UE::DerivedData;
+		FCacheKey DDCKey;
+
+		if (bAllowFileCategoryOverride)
+		{
+			// Generate DDC Key;
+			TArray<uint8> Bytes;
+			FMemoryWriter Ar(Bytes, /*bIsPersistent=*/ true);
+
+			Ar << ObjectId;
+
+			int32 CategoryOverrideDerivedDataVersion = 0x3c507a2f;
+			Ar << CategoryOverrideDerivedDataVersion;
+
+			FString PlatformName = TargetPlatform.PlatformName();
+			Ar << PlatformName;
+
+			DDCKey.Bucket = FCacheBucket(TEXT("CustomizableObject"));
+			DDCKey.Hash = FIoHashBuilder::HashBuffer(MakeMemoryView(Bytes));
+
+
+			// Request FileCategory Overrides for TargetPlatform
+			FCacheGetValueRequest DDCRequest;
+			DDCRequest.Key = DDCKey;
+
+			FRequestOwner RequestOwner(UE::DerivedData::EPriority::Blocking);
+			GetCache().GetValue({ DDCRequest }, RequestOwner, [&FileCategoryOverrides](FCacheGetValueResponse&& Response) mutable
+				{
+					if (Response.Status == EStatus::Ok)
+					{
+						const FCompressedBuffer& CompressedBuffer = Response.Value.GetData();
+						if (ensure(!CompressedBuffer.IsNull()))
+						{
+							check(Response.Value.GetRawSize() % sizeof(FFileCategoryOverride) == 0);
+							
+							const int32 NumCategories = Response.Value.GetRawSize() / sizeof(FFileCategoryOverride);
+							FileCategoryOverrides.SetNum(NumCategories);
+							
+							bool bSuccess = CompressedBuffer.TryDecompressTo(MakeMemoryView(FileCategoryOverrides));
+						}
+					}
+				});
+			RequestOwner.Wait();
+		}
+
 
 		for (FFileBucket& Bucket : FileBuckets)
 		{
 			uint32 NumFiles = 0;
 
+			bool bCanOverride = bAllowFileCategoryOverride;
+
 			for (FFileCategory& Category : Bucket.Categories)
 			{
 				if (Category.DataSize > 0)
@@ -2570,10 +2622,57 @@ namespace UE::Mutable::Private
 
 					NumFiles += Category.NumFiles;
 				}
+
+				if (bAllowFileCategoryOverride)
+				{
+					FFileCategoryOverride* CategoryOverride = FileCategoryOverrides.FindByPredicate(
+						[Category](const FFileCategoryOverride& C) { return C.Id == Category.Id; });
+
+					// Can't override if there's no override data or if the override doesn't have the minimum required number of files.
+					const bool bIsValidOverride = CategoryOverride && (Category.DataSize == 0 || CategoryOverride->NumFiles >= 0);
+					bCanOverride = bCanOverride && bIsValidOverride;
+
+					if (!CategoryOverride)
+					{
+						FileCategoryOverrides.Add(FFileCategoryOverride({ .Id = Category.Id, .NumFiles = Category.NumFiles }));
+					}
+				}
+			}
+
+			if (bCanOverride)
+			{
+				for (FFileCategory& Category : Bucket.Categories)
+				{
+					FFileCategoryOverride* CategoryOverride = FileCategoryOverrides.FindByPredicate(
+						[Category](const FFileCategoryOverride& C) { return C.Id == Category.Id; });
+					if (ensure(CategoryOverride))
+					{
+						UE_LOG(LogMutable, Display, TEXT("Id [%s] - NumFiles [%d] -> [%d]."), *ObjectId.ToString(), Category.NumFiles, CategoryOverride->NumFiles);
+						Category.NumFiles = CategoryOverride->NumFiles;
+					}
+				}
 			}
 
 			Categories.Append(Bucket.Categories);
 		}
+
+		if (bAllowFileCategoryOverride && !FileCategoryOverrides.IsEmpty())
+		{
+			for (auto& CategoryOverride : FileCategoryOverrides)
+			{
+				UE_LOG(LogMutable, Display, TEXT("Id [%s] - Storing Override NumFiles [%d]."), *ObjectId.ToString(), CategoryOverride.NumFiles);
+			}
+
+			// Upload FileCategory Overrides to DDC.
+			FCachePutValueRequest DDCRequest;
+			DDCRequest.Key = DDCKey;
+			DDCRequest.Value = FValue::Compress(FSharedBuffer::MakeView(FileCategoryOverrides.GetData(), FileCategoryOverrides.Num() * sizeof(FFileCategoryOverride)));
+
+			FRequestOwner RequestOwner(UE::DerivedData::EPriority::Blocking);
+			RequestOwner.KeepAlive();
+
+			GetCache().PutValue({ DDCRequest }, RequestOwner);
+		}
 		
 		// Function to create the list of bulk data files. Blocks will be grouped by source Id.
 		const auto CreateFileList = [Categories, bAllowSplit](const FFileCategoryID& CategoryID, const FClassifyNode& Node, TArray<FFile>& OutBulkDataFiles)
diff --git a/Engine/Plugins/Mutable/Source/CustomizableObjectEditor/Private/MuCOE/CustomizableObjectCompileRunnable.cpp b/Engine/Plugins/Mutable/Source/CustomizableObjectEditor/Private/MuCOE/CustomizableObjectCompileRunnable.cpp
--- a/Engine/Plugins/Mutable/Source/CustomizableObjectEditor/Private/MuCOE/CustomizableObjectCompileRunnable.cpp
+++ b/Engine/Plugins/Mutable/Source/CustomizableObjectEditor/Private/MuCOE/CustomizableObjectCompileRunnable.cpp
@@ -435,6 +435,7 @@ FCustomizableObjectSaveDDRunnable::FCustomizableObjectSaveDDRunnable(const TShar
 
 	UCustomizableObject* CustomizableObject = InRequest->GetCustomizableObject();
 	CustomizableObjectName = GetNameSafe(CustomizableObject);
+	CustomizableObjectId = GenerateIdentifier(*CustomizableObject);
 
 	CustomizableObjectHeader.InternalVersion = GetECustomizableObjectVersionEnumHash();
 	CustomizableObjectHeader.VersionId = CustomizableObject->GetPrivate()->GetVersionId();
@@ -509,7 +510,8 @@ void FCustomizableObjectSaveDDRunnable::CachePlatformData()
 		// Generate list of files and update streamable blocks ids and offsets
 		if (Options.bUseBulkData)
 		{
-			UE::Mutable::Private::GenerateBulkDataFilesListWithFileLimit(PlatformData->Model, *PlatformData->ModelStreamableBulkData.Get(), MAX_uint8, false, PlatformData->BulkDataFiles);
+			UE::Mutable::Private::GenerateBulkDataFilesListWithFileLimit(CustomizableObjectId, PlatformData->Model, *PlatformData->ModelStreamableBulkData.Get(), MAX_uint8, 
+				Options.bIsCooking, false, *Options.TargetPlatform, PlatformData->BulkDataFiles);
 		}
 		else
 		{
@@ -555,7 +557,8 @@ void FCustomizableObjectSaveDDRunnable::StoreCachedPlatformDataInDDC(bool& bStor
 
 		// Generate list of files and update streamable blocks ids and offsets
 		constexpr int32 MaxDDCFiles = 1 << 13;
-		UE::Mutable::Private::GenerateBulkDataFilesListWithFileLimit(PlatformData->Model, ModelStreamablesDDC, MaxDDCFiles, !Options.bIsCooking, BulkDataFilesDDC);
+		UE::Mutable::Private::GenerateBulkDataFilesListWithFileLimit(CustomizableObjectId, PlatformData->Model, ModelStreamablesDDC, MaxDDCFiles, 
+			Options.bIsCooking, !Options.bIsCooking, *Options.TargetPlatform, BulkDataFilesDDC);
 
 		if (!Options.bIsCooking)
 		{
diff --git a/Engine/Plugins/Mutable/Source/CustomizableObjectEditor/Private/MuCOE/CustomizableObjectCompileRunnable.h b/Engine/Plugins/Mutable/Source/CustomizableObjectEditor/Private/MuCOE/CustomizableObjectCompileRunnable.h
--- a/Engine/Plugins/Mutable/Source/CustomizableObjectEditor/Private/MuCOE/CustomizableObjectCompileRunnable.h
+++ b/Engine/Plugins/Mutable/Source/CustomizableObjectEditor/Private/MuCOE/CustomizableObjectCompileRunnable.h
@@ -108,9 +108,10 @@ private:
 
 	FCompilationOptions Options;
 
-	MutableCompiledDataStreamHeader CustomizableObjectHeader;
-
 	FString CustomizableObjectName;
+	FGuid CustomizableObjectId;
+
+	MutableCompiledDataStreamHeader CustomizableObjectHeader;
 
 	// Used to save files to disk
 	FString FullFileName;
diff --git a/Engine/Plugins/Mutable/Source/CustomizableObjectEditor/Private/MuCOE/CustomizableObjectCompiler.cpp b/Engine/Plugins/Mutable/Source/CustomizableObjectEditor/Private/MuCOE/CustomizableObjectCompiler.cpp
--- a/Engine/Plugins/Mutable/Source/CustomizableObjectEditor/Private/MuCOE/CustomizableObjectCompiler.cpp
+++ b/Engine/Plugins/Mutable/Source/CustomizableObjectEditor/Private/MuCOE/CustomizableObjectCompiler.cpp
@@ -1544,8 +1544,8 @@ void FCustomizableObjectCompiler::FinishLoadingModelDataFromDDC()
 					if (CVarMutableUseBulkData.GetValueOnAnyThread())
 					{
 						const uint32 NumBulkDataFilesPerBucket = MAX_uint8;
-						UE::Mutable::Private::GenerateBulkDataFilesListWithFileLimit(PlatformData->Model, *PlatformData->ModelStreamableBulkData.Get(),
-							NumBulkDataFilesPerBucket, false, PlatformData->BulkDataFiles);
+						UE::Mutable::Private::GenerateBulkDataFilesListWithFileLimit(GenerateIdentifier(*CustomizableObject), PlatformData->Model, *PlatformData->ModelStreamableBulkData.Get(),
+							NumBulkDataFilesPerBucket, Options.bIsCooking, false, *Options.TargetPlatform, PlatformData->BulkDataFiles);
 					}
 					else
 					{
