[Mutable] Fix Skin Weight Profiles conversion.

#rnx

diff --git a/Engine/Plugins/Mutable/Source/CustomizableObject/Internal/MuCO/CustomizableObjectPrivate.h b/Engine/Plugins/Mutable/Source/CustomizableObject/Internal/MuCO/CustomizableObjectPrivate.h
--- a/Engine/Plugins/Mutable/Source/CustomizableObject/Internal/MuCO/CustomizableObjectPrivate.h
+++ b/Engine/Plugins/Mutable/Source/CustomizableObject/Internal/MuCO/CustomizableObjectPrivate.h
@@ -1727,6 +1727,8 @@ enum class ECustomizableObjectVersions : int32
 
 	CompilationEarlyOutOnRepeatedParameters,
 
+	FixSkinWeightProfilesConversion,
+
 	// -----<new versions can be added above this line>--------
 	LastCustomizableObjectVersion
 };
diff --git a/Engine/Plugins/Mutable/Source/CustomizableObjectEditor/Private/MuCOE/GenerateMutableSource/GenerateMutableSourceMesh.cpp b/Engine/Plugins/Mutable/Source/CustomizableObjectEditor/Private/MuCOE/GenerateMutableSource/GenerateMutableSourceMesh.cpp
--- a/Engine/Plugins/Mutable/Source/CustomizableObjectEditor/Private/MuCOE/GenerateMutableSource/GenerateMutableSourceMesh.cpp
+++ b/Engine/Plugins/Mutable/Source/CustomizableObjectEditor/Private/MuCOE/GenerateMutableSource/GenerateMutableSourceMesh.cpp
@@ -514,6 +514,15 @@ void NormalizeWeights(uint16* InfluenceBones, uint16* InfluenceWeights, const in
 		}
 
 		InfluenceWeights[0] = MaxBoneWeight - AssignedWeight;
+		
+		for (int32 j = 1; j < MutableInfluenceCount; ++j)
+		{
+			if (InfluenceWeights[j] == 0)
+			{
+				FMemory::Memzero(&InfluenceWeights[j], MutableInfluenceCount - j);
+				break;
+			}
+		}
 	}
 	else
 	{
@@ -2241,7 +2250,7 @@ UE::Tasks::FTask ConvertSkeletalMeshToMutable(TSharedRef<FMeshConversionContext>
 			check(NumChannels == 3);
 
 			const int32 MaxSectionBoneInfluences = MeshContextRef.SourceNumBoneInfluences;
-			const int32 NumBoneInfluences = VertexBuffer.Channels[0].ComponentCount;
+			const int32 NumBoneInfluences = VertexBuffer.Channels[1].ComponentCount;
 
 			const UE::Mutable::Private::EMeshBufferFormat BoneIndicesFormat = VertexBuffer.Channels[1].Format;
 			const UE::Mutable::Private::EMeshBufferFormat BoneWeightFormat = VertexBuffer.Channels[2].Format;
@@ -2262,6 +2271,7 @@ UE::Tasks::FTask ConvertSkeletalMeshToMutable(TSharedRef<FMeshConversionContext>
 					&&
 					FMemory::Memcmp(&Vertices[VertexIndex].InfluenceWeights[0], &SkinWeight.InfluenceWeights[0], SourceDataSizeBytes) == 0)
 				{
+					FMemory::Memset(DestData, 0, DestDataSize);
 					DestData += DestDataSize;
 					continue;
 				}
@@ -2275,15 +2285,15 @@ UE::Tasks::FTask ConvertSkeletalMeshToMutable(TSharedRef<FMeshConversionContext>
 				int32 SkinWeightVertexHash = 0;
 				for (int32 InfluenceIndex = 0; InfluenceIndex < NumBoneInfluences; ++InfluenceIndex)
 				{
-					SkinWeightVertexHash = HashCombine(SkinWeightVertexHash, SkinWeight.InfluenceBones[InfluenceIndex]);
-					SkinWeightVertexHash = HashCombine(SkinWeightVertexHash, SkinWeight.InfluenceWeights[InfluenceIndex]);
+					SkinWeightVertexHash = HashCombine(SkinWeightVertexHash, InfluenceBones[InfluenceIndex]);
+					SkinWeightVertexHash = HashCombine(SkinWeightVertexHash, InfluenceWeights[InfluenceIndex]);
 				}
 
 				FMemory::Memcpy(DestData, &SkinWeightVertexHash, sizeof(int32));
 				DestData += sizeof(int32);
-				FMemory::Memcpy(DestData, &SkinWeight.InfluenceBones[0], BoneIndicesSizeBytes);
+				FMemory::Memcpy(DestData, &InfluenceBones[0], BoneIndicesSizeBytes);
 				DestData += BoneIndicesSizeBytes;
-				FMemory::Memcpy(DestData, &SkinWeight.InfluenceWeights[0], BoneWeightsSizeBytes);
+				FMemory::Memcpy(DestData, &InfluenceWeights[0], BoneWeightsSizeBytes);
 				DestData += BoneWeightsSizeBytes;
 			}
 
