[Mutable] Add mesh description to generated meshes.
- Removed unnecessary casts
- Renamed profiler visible operation name from RegenerateImportedModels to "RegenrateEditorImportedModels" so the user knows this is something related to the editor operation.

#rnx

diff --git a/Engine/Plugins/Mutable/Source/CustomizableObject/CustomizableObject.Build.cs b/Engine/Plugins/Mutable/Source/CustomizableObject/CustomizableObject.Build.cs
--- a/Engine/Plugins/Mutable/Source/CustomizableObject/CustomizableObject.Build.cs
+++ b/Engine/Plugins/Mutable/Source/CustomizableObject/CustomizableObject.Build.cs
@@ -25,7 +25,8 @@ public class CustomizableObject : ModuleRules
 			//"ClothingSystemEditor",
 			"UMG",
 			"ImageCore",
-			"TraceLog"
+			"TraceLog",
+			"MeshDescription"
 		]);
 
         PublicDependencyModuleNames.AddRange([
diff --git a/Engine/Plugins/Mutable/Source/CustomizableObject/Private/MuCO/CustomizableObjectInstance.cpp b/Engine/Plugins/Mutable/Source/CustomizableObject/Private/MuCO/CustomizableObjectInstance.cpp
--- a/Engine/Plugins/Mutable/Source/CustomizableObject/Private/MuCO/CustomizableObjectInstance.cpp
+++ b/Engine/Plugins/Mutable/Source/CustomizableObject/Private/MuCO/CustomizableObjectInstance.cpp
@@ -54,6 +54,7 @@
 
 #include "MuCO/Plugins/IMutableClothingModule.h"
 #include "Rendering/SkeletalMeshModel.h"
+#include "MeshDescription.h"
 
 #include UE_INLINE_GENERATED_CPP_BY_NAME(CustomizableObjectInstance)
 
@@ -6040,7 +6041,12 @@ void UCustomizableInstancePrivate::BuildMaterials(const TSharedRef<FUpdateContex
 				if (FoundSurfaceMetadata)
 				{
 					MaterialSlot.MaterialSlotName = FoundSurfaceMetadata->MaterialSlotName;
+
+#if WITH_EDITOR
+					MaterialSlot.ImportedMaterialSlotName = FoundSurfaceMetadata->MaterialSlotName;
+#endif
 				}
+
 				if (RefSkeletalMeshData)
 				{
 					SetMeshUVChannelDensity(MaterialSlot.UVChannelData, RefSkeletalMeshData->Settings.DefaultUVChannelDensity);
@@ -7227,7 +7233,7 @@ void CalculateBonesToRemove(const FSkeletalMeshLODRenderData& LODResource, const
 
 void UCustomizableInstancePrivate::RegenerateImportedModels()
 {
-	MUTABLE_CPUPROFILER_SCOPE(RegenerateImportedModels);
+	MUTABLE_CPUPROFILER_SCOPE(RegenerateEditorImportedModels);
 
 	struct FMeshDataConvertJob
 	{
@@ -7260,7 +7266,7 @@ void UCustomizableInstancePrivate::RegenerateImportedModels()
 			continue;
 		}
 
-		const bool bIsTransientMesh = static_cast<bool>(SkeletalMesh->HasAllFlags(EObjectFlags::RF_Transient));
+		const bool bIsTransientMesh = SkeletalMesh->HasAllFlags(EObjectFlags::RF_Transient);
 
 		if (!bIsTransientMesh)
 		{
@@ -7460,8 +7466,6 @@ void UCustomizableInstancePrivate::RegenerateImportedModels()
 
 				FSkeletalMeshLODInfo* LODInfo = SkeletalMesh->GetLODInfo(LODIndex);
 				LODInfo->BuildGUID = LODInfo->ComputeDeriveDataCacheKey(SkeletalMeshLODGroupSettings);
-
-				LODModel.BuildStringID = LODModel.GetLODModelDeriveDataKey();
 			}
 		}
 
@@ -7597,6 +7601,79 @@ void UCustomizableInstancePrivate::RegenerateImportedModels()
 			}
 		});
 	}
+
+	TArray<UE::Tasks::FTask> MeshDescriptionGenerationTasks;
+	
+	for (const TTuple<FName, TObjectPtr<USkeletalMesh>>& Tuple : SkeletalMeshes)
+	{
+		USkeletalMesh* SkeletalMesh = Tuple.Get<1>();
+
+		if (!SkeletalMesh)
+		{
+			continue;
+		}
+
+		const bool bIsTransientMesh = SkeletalMesh->HasAllFlags(EObjectFlags::RF_Transient);
+
+		if (!bIsTransientMesh)
+		{
+			// This must be a pass-through referenced mesh so don't do anything to it
+			continue;
+		}
+
+		FSkeletalMeshRenderData* RenderData = SkeletalMesh->GetResourceForRendering();
+		if (!RenderData || RenderData->IsInitialized())
+		{
+			continue;
+		}
+
+		const int32 LODCount = SkeletalMesh->GetLODNum();
+		MeshDescriptionGenerationTasks.Reserve(MeshDescriptionGenerationTasks.Num() + LODCount);
+
+		TStrongObjectPtr<USkeletalMesh> StrongMesh(SkeletalMesh);
+		
+		for (int32 LODIndex = 0; LODIndex < LODCount; ++LODIndex)
+		{
+			if (!SkeletalMesh->HasMeshDescription(LODIndex))
+			{
+				UE::Tasks::FTask GenerateMeshDescriptionTask = UE::Tasks::Launch(TEXT("Generate Mesh Description"),
+				[StrongMesh, LODIndex]()
+				{
+					USkeletalMesh* SkeletalMesh = StrongMesh.Get();
+					FSkeletalMeshLODInfo* MeshLODInfo = SkeletalMesh->GetLODInfo(LODIndex);
+
+					const FSkeletalMeshLODModel& LODModel = SkeletalMesh->GetImportedModel()->LODModels[LODIndex];
+					FMeshDescription MeshDescription;
+					LODModel.GetMeshDescription(SkeletalMesh, LODIndex, MeshDescription);
+					SkeletalMesh->CreateMeshDescription(LODIndex, MoveTemp(MeshDescription));
+					SkeletalMesh->CommitMeshDescription(LODIndex);
+
+					// Ensure normals aren't automatically computed when we rebuild.
+					FSkeletalMeshBuildSettings& BuildSettings = MeshLODInfo->BuildSettings;
+					BuildSettings.bRecomputeNormals = false;
+
+					// Reset the reduction settings so that we don't re-reduce the mesh and possibly lose morph targets
+					// in the process.
+					FSkeletalMeshOptimizationSettings& ReductionSettings = MeshLODInfo->ReductionSettings;
+
+					//Remove the reduction settings
+					ReductionSettings.NumOfTrianglesPercentage = 1.0f;
+					ReductionSettings.NumOfVertPercentage = 1.0f;
+					ReductionSettings.MaxNumOfTrianglesPercentage = MAX_uint32;
+					ReductionSettings.MaxNumOfVertsPercentage = MAX_uint32;
+					ReductionSettings.TerminationCriterion = SMTC_NumOfTriangles;
+					MeshLODInfo->bHasBeenSimplified = false;
+
+					LODModel.BuildStringID = LODModel.GetLODModelDeriveDataKey();
+				});
+
+				MeshDescriptionGenerationTasks.Add(GenerateMeshDescriptionTask);
+			}
+		}
+	}
+
+	// Wait for all Mesh descriptors to be completed before allowing other work to operate over the Skeletal mesh
+	UE::Tasks::Wait(MeshDescriptionGenerationTasks);
 }
 
 #endif
