[mutable] Fix for "SetDisplayNamePropertyAsPinNameOfTableNodes" backwards compatibility code not running in some cases.
- Fixed a problem with the order of the COCustomVersion enum values that was causing the "SetDisplayNamePropertyAsPinNameOfTableNodes" compatibility code to not be run.
- Added a new entry in the enum that calls for the code that was being skipped.

#rnx

diff --git a/Engine/Plugins/Mutable/Source/CustomizableObject/Internal/MuCO/CustomizableObjectCustomVersion.h b/Engine/Plugins/Mutable/Source/CustomizableObject/Internal/MuCO/CustomizableObjectCustomVersion.h
--- a/Engine/Plugins/Mutable/Source/CustomizableObject/Internal/MuCO/CustomizableObjectCustomVersion.h
+++ b/Engine/Plugins/Mutable/Source/CustomizableObject/Internal/MuCO/CustomizableObjectCustomVersion.h
@@ -202,6 +202,8 @@ struct FCustomizableObjectCustomVersion
 
 		FixStaticMeshNodeLODMaterialDataInvalidationBackOut,
 		
+		SetDisplayNamePropertyAsPinNameOfTableNodes2,			// Added in order to apply the same fix for cases where the version was skipped
+
 		// -----<new versions can be added above this line>-------------------------------------------------
 		VersionPlusOne,
 		LatestVersion = VersionPlusOne - 1
diff --git a/Engine/Plugins/Mutable/Source/CustomizableObjectEditor/Private/MuCOE/Nodes/CustomizableObjectNodeTable.cpp b/Engine/Plugins/Mutable/Source/CustomizableObjectEditor/Private/MuCOE/Nodes/CustomizableObjectNodeTable.cpp
--- a/Engine/Plugins/Mutable/Source/CustomizableObjectEditor/Private/MuCOE/Nodes/CustomizableObjectNodeTable.cpp
+++ b/Engine/Plugins/Mutable/Source/CustomizableObjectEditor/Private/MuCOE/Nodes/CustomizableObjectNodeTable.cpp
@@ -331,6 +331,11 @@ void UCustomizableObjectNodeTable::BackwardsCompatibleFixup(int32 CustomizableOb
 			CompilationFilterOptions.Add({ CompilationFilterColumn_DEPRECATED, CompilationFilters_DEPRECATED, FilterOperationType_DEPRECATED });
 		}
 	}
+
+	if (CustomizableObjectCustomVersion == FCustomizableObjectCustomVersion::SetDisplayNamePropertyAsPinNameOfTableNodes2)
+	{
+		BackwardsCompatibleFixup(FCustomizableObjectCustomVersion::SetDisplayNamePropertyAsPinNameOfTableNodes);
+	}
 }
 
 
