[mutable] Fixed bug preventing the activation of some sections due to the wrong handling of repeated section tags.
- Previously, if two sections did share the same tag, both sections were required to be active for the tag to be enabled. Now, as it is expected, only one section is required to activate the tag.
- Replaced some code with more functional code (removed pointer operations).
- Fixed some typos
- Improved overall code readability by enforcing UE coding standards.

#rnx

diff --git a/Engine/Plugins/Mutable/Source/CustomizableObjectEditor/Private/MuCOE/GenerateMutableSource/GenerateMutableSourceComponent.cpp b/Engine/Plugins/Mutable/Source/CustomizableObjectEditor/Private/MuCOE/GenerateMutableSource/GenerateMutableSourceComponent.cpp
--- a/Engine/Plugins/Mutable/Source/CustomizableObjectEditor/Private/MuCOE/GenerateMutableSource/GenerateMutableSourceComponent.cpp
+++ b/Engine/Plugins/Mutable/Source/CustomizableObjectEditor/Private/MuCOE/GenerateMutableSource/GenerateMutableSourceComponent.cpp
@@ -90,7 +90,7 @@ void GenerateMutableSourceComponentMesh(FMutableGraphGenerationContext& Generati
 
 			TArray<UEdGraphPin*> ConnectedLODPins = FollowInputPinArray(*LODPin);
 
-			// Proccess non modifier nodes.
+			// Process non modifier nodes.
 			for (UEdGraphPin* const ChildNodePin : ConnectedLODPins)
 			{
 				// Modifiers are shared for all components and are processed per LOD and not component.
diff --git a/Engine/Plugins/Mutable/Source/MutableTools/Private/MuT/CodeGenerator_FirstPass.cpp b/Engine/Plugins/Mutable/Source/MutableTools/Private/MuT/CodeGenerator_FirstPass.cpp
--- a/Engine/Plugins/Mutable/Source/MutableTools/Private/MuT/CodeGenerator_FirstPass.cpp
+++ b/Engine/Plugins/Mutable/Source/MutableTools/Private/MuT/CodeGenerator_FirstPass.cpp
@@ -230,42 +230,50 @@ namespace UE::Mutable::Private
 
         case NodeSurfaceVariation::VariationType::Tag:
         {
-            // Any of the tags in the variations would prevent the default surface
+			// Process the default case
+        		
+            // Make the default case nodes not active if any of the variation tags is enabled.
 			TArray<FString> OldNegativeTags = CurrentNegativeTags;
-            for (int32 v=0; v< InNode->Variations.Num(); ++v)
             {
-                CurrentNegativeTags.Add(InNode->Variations[v].Tag);
+				for (int32 VariationIndex=0; VariationIndex< InNode->Variations.Num(); ++VariationIndex)
+				{
+					CurrentNegativeTags.Add(InNode->Variations[VariationIndex].Tag);
             }
 
-            for(const Ptr<NodeSurface>& n: InNode->DefaultSurfaces)
+				for(const Ptr<NodeSurface>& SurfaceNode: InNode->DefaultSurfaces)
             {
-				Generate_Generic(n.get());
+					Generate_Generic(SurfaceNode.get());
 			}
-            for(const Ptr<NodeModifier>& n: InNode->DefaultModifiers)
+				for(const Ptr<NodeModifier>& ModifierNode: InNode->DefaultModifiers)
             {
-				Generate_Modifier(n.get());
+					Generate_Modifier(ModifierNode.get());
             }
 
+				// Restore the NegativeTags to be able to process the non default cases
             CurrentNegativeTags = OldNegativeTags;
+			}
+
 
-            for (int32 v=0; v< InNode->Variations.Num(); ++v)
+        	// Process the variations so each node generated from each of them does have as the activation tag the one defined in the variation.
+            for (int32 VariationIndex=0; VariationIndex< InNode->Variations.Num(); ++VariationIndex)
             {
-                CurrentPositiveTags.Add(InNode->Variations[v].Tag);
-                for (const Ptr<NodeSurface>& s : InNode->Variations[v].Surfaces)
+                CurrentPositiveTags.Add(InNode->Variations[VariationIndex].Tag);
                 {
-					Generate_Generic(s.get());
+                	for (const Ptr<NodeSurface>& SurfaceNode : InNode->Variations[VariationIndex].Surfaces)
+                	{
+                		Generate_Generic(SurfaceNode.get());
                 }
 
-                for (const Ptr<NodeModifier>& s : InNode->Variations[v].Modifiers)
+                	for (const Ptr<NodeModifier>& ModifierNode : InNode->Variations[VariationIndex].Modifiers)
                 {
-					Generate_Modifier(s.get());
+                		Generate_Modifier(ModifierNode.get());
 				}
-
-                CurrentPositiveTags.Pop();
+                }
+                CurrentPositiveTags.Pop();		// Restore the state prior to processing this variation
 
                 // Tags have an order in a variation node: the current tag should prevent any following
                 // variation surface
-                CurrentNegativeTags.Add(InNode->Variations[v].Tag);
+                CurrentNegativeTags.Add(InNode->Variations[VariationIndex].Tag);
             }
 
             CurrentNegativeTags = OldNegativeTags;
@@ -276,7 +284,7 @@ namespace UE::Mutable::Private
 
         case NodeSurfaceVariation::VariationType::State:
         {
-            int32 StateCount = States.Num();
+            const int32 StateCount = States.Num();
 
             // Default
             {
@@ -289,11 +297,11 @@ namespace UE::Mutable::Private
                             ? AllTrue
                             : CurrentStateCondition.Last();
 
-                    for (const NodeSurfaceVariation::FVariation& v: InNode->Variations)
+                    for (const NodeSurfaceVariation::FVariation& Variation: InNode->Variations)
                     {
                         for( int32 StateIndex=0; StateIndex < StateCount; ++StateIndex)
                         {
-                            if (States[StateIndex].Name==v.Tag)
+                            if (States[StateIndex].Name==Variation.Tag)
                             {
                                 // Remove this state from the default options, since it has its own variation
 								DefaultStates[StateIndex] = false;
@@ -304,13 +312,13 @@ namespace UE::Mutable::Private
 
                 CurrentStateCondition.Add(DefaultStates);
 
-                for (const Ptr<NodeSurface>& n : InNode->DefaultSurfaces)
+                for (const Ptr<NodeSurface>& SurfaceNode : InNode->DefaultSurfaces)
                 {
-					Generate_Generic(n.get());
+					Generate_Generic(SurfaceNode.get());
 				}
-                for (const Ptr<NodeModifier>& n : InNode->DefaultModifiers)
+                for (const Ptr<NodeModifier>& ModifierNode : InNode->DefaultModifiers)
                 {
-					Generate_Modifier(n.get());
+					Generate_Modifier(ModifierNode.get());
 				}
 
                 CurrentStateCondition.Pop();
@@ -333,13 +341,13 @@ namespace UE::Mutable::Private
 
                 CurrentStateCondition.Add(variationStates);
 
-                for (const Ptr<NodeSurface>& n : v.Surfaces)
+                for (const Ptr<NodeSurface>& SurfaceNode : v.Surfaces)
                 {
-					Generate_Generic(n.get());
+					Generate_Generic(SurfaceNode.get());
 				}
-                for (const Ptr<NodeModifier>& n : v.Modifiers)
+                for (const Ptr<NodeModifier>& ModifierNode : v.Modifiers)
                 {
-					Generate_Modifier(n.get());
+					Generate_Modifier(ModifierNode.get());
 				}
 
                 CurrentStateCondition.Pop();
diff --git a/Engine/Plugins/Mutable/Source/MutableTools/Private/MuT/CodeGenerator_SecondPass.cpp b/Engine/Plugins/Mutable/Source/MutableTools/Private/MuT/CodeGenerator_SecondPass.cpp
--- a/Engine/Plugins/Mutable/Source/MutableTools/Private/MuT/CodeGenerator_SecondPass.cpp
+++ b/Engine/Plugins/Mutable/Source/MutableTools/Private/MuT/CodeGenerator_SecondPass.cpp
@@ -2,6 +2,7 @@
 
 #include "MuT/CodeGenerator_SecondPass.h"
 
+#include "ASTOpBoolOr.h"
 #include "Misc/AssertionMacros.h"
 #include "MuR/MutableTrace.h"
 #include "MuR/Operations.h"
@@ -60,13 +61,10 @@ namespace UE::Mutable::Private
 		for (auto s : posTag) { if (TagsPerTag[tagIndex].Contains(s)) { key.posTag.Add(s); } }
 		for (auto s : negSurf) { if (SurfacesPerTag[tagIndex].Contains(s)) { key.negSurf.Add(s); } }
 		for (auto s : posSurf) { if (SurfacesPerTag[tagIndex].Contains(s)) { key.posSurf.Add(s); } }
-
+		
+		if (Ptr<ASTOp>* Found = TagConditionGenerationCache.Find(key))
 		{
-			Ptr<ASTOp>* Found = TagConditionGenerationCache.Find(key);
-			if (Found)
-			{
-				return *Found;
-			}
+			return *Found;
 		}
 		
 		Ptr<ASTOp> c;
@@ -131,7 +129,7 @@ namespace UE::Mutable::Private
 			}
 			else
 			{
-				Ptr<ASTOpBoolAnd> o = new ASTOpBoolAnd();
+				Ptr<ASTOpBoolOr> o = new ASTOpBoolOr();
 				o->A = fullCondition;
 				o->B = c;
 				c = OpPool.Add(o);
@@ -169,22 +167,20 @@ namespace UE::Mutable::Private
 
 		Ptr<ASTOp> c;
 
-		for (const FString& t : PositiveTags)
+		for (const FString& PositiveTag : PositiveTags)
 		{
-			const FirstPassGenerator::FTag* it = FirstPass->Tags.FindByPredicate([&](const FirstPassGenerator::FTag& e) { return e.Tag == t; });
-			if (!it)
+			const int32 PositiveTagIndex = FirstPass->Tags.IndexOfByPredicate([&](const FirstPassGenerator::FTag& e) { return e.Tag == PositiveTag; });
+			if (PositiveTagIndex == INDEX_NONE)
 			{
-				// This could happen if a tag is in a variation but noone defines it.
+				// This could happen if a tag is in a variation but no one defines it.
 				// This surface depends on a tag that will never be active, so it will never be used.
 				return OpPool.Add(new ASTOpConstantBool(false));
 			}
-
-			int32 tagIndex = int32(it - &FirstPass->Tags[0]);
-
+			
 			TSet<int32> positiveSurfacesVisited = posSurf;
 			positiveSurfacesVisited.Add(Index);
 
-			Ptr<ASTOp> tagCondition = GenerateTagCondition(tagIndex,
+			Ptr<ASTOp> tagCondition = GenerateTagCondition(PositiveTagIndex,
 				positiveSurfacesVisited,
 				negSurf,
 				posTag,
@@ -238,51 +234,49 @@ namespace UE::Mutable::Private
 		}
 
 
-		for (const FString& t : NegativeTags)
+		for (const FString& NegativeTag : NegativeTags)
 		{
-			const FirstPassGenerator::FTag* it = FirstPass->Tags.FindByPredicate([&](const FirstPassGenerator::FTag& e) { return e.Tag == t; });
-			if (!it)
+			const int32 NegativeTagIndex = FirstPass->Tags.IndexOfByPredicate([&](const FirstPassGenerator::FTag& e) { return e.Tag == NegativeTag; });
+			if (NegativeTagIndex == INDEX_NONE)
 			{
-				// This could happen if a tag is in a variation but noone defines it.
+				// This could happen if a tag is in a variation but none defines it.
 				continue;
 			}
-
-			int32 tagIndex = int32(it - FirstPass->Tags.GetData());
-
-			TSet<int32> positiveSurfacesVisited = negSurf;
-			TSet<int32> negativeSurfacesVisited = posSurf;
-			negativeSurfacesVisited.Add(Index);
-			TSet<int32> positiveTagsVisited = negTag;
-			TSet<int32> negativeTagsVisited = posTag;
-			Ptr<ASTOp> tagCondition = GenerateTagCondition(tagIndex,
-				positiveSurfacesVisited,
-				negativeSurfacesVisited,
-				positiveTagsVisited,
-				negativeTagsVisited);
+			
+			TSet<int32> PositiveSurfacesVisited = negSurf;
+			TSet<int32> NegativeSurfacesVisited = posSurf;
+			NegativeSurfacesVisited.Add(Index);
+			TSet<int32> PositiveTagsVisited = negTag;
+			TSet<int32> NegativeTagsVisited = posTag;
+			Ptr<ASTOp> TagCondition = GenerateTagCondition(NegativeTagIndex,
+				PositiveSurfacesVisited,
+				NegativeSurfacesVisited,
+				PositiveTagsVisited,
+				NegativeTagsVisited);
 
 			// No condition is equal to a conditional with a true constant
-			if (!tagCondition)
+			if (!TagCondition)
 			{
-				tagCondition = new ASTOpConstantBool(true);
+				TagCondition = new ASTOpConstantBool(true);
 			}
 
 			// TODO: Optimise the tag condition here
 
 			// If the tag is a constant ...
-			bool isConstant = false;
-			bool constantValue = false;
-			if (tagCondition && tagCondition->GetOpType()== EOpType::BO_CONSTANT)
+			bool bIsConstant = false;
+			bool bIsConstantValue = false;
+			if (TagCondition && TagCondition->GetOpType()== EOpType::BO_CONSTANT)
 			{
-				const ASTOpConstantBool* constOp = static_cast<const ASTOpConstantBool*>(tagCondition.get());
-				isConstant = true;
-				constantValue = constOp->bValue;
+				const ASTOpConstantBool* constOp = static_cast<const ASTOpConstantBool*>(TagCondition.get());
+				bIsConstant = true;
+				bIsConstantValue = constOp->bValue;
 			}
 
 
-			if (!isConstant && tagCondition)
+			if (!bIsConstant && TagCondition)
 			{
 				Ptr<ASTOpBoolNot> n = new ASTOpBoolNot;
-				n->A = tagCondition;
+				n->A = TagCondition;
 
 				if (!c)
 				{
@@ -296,7 +290,7 @@ namespace UE::Mutable::Private
 					c = OpPool.Add(o);
 				}
 			}
-			else if (isConstant && constantValue == true)
+			else if (bIsConstant && bIsConstantValue == true)
 			{
 				// No expression here means always true which becomes always false
 				Ptr<ASTOpConstantBool> f = new ASTOpConstantBool(false);
@@ -324,69 +318,65 @@ namespace UE::Mutable::Private
 		SurfacesPerTag.SetNum(FirstPass->Tags.Num());
 		TagsPerTag.Empty();
 		TagsPerTag.SetNum(FirstPass->Tags.Num());
-		for (int32 t = 0; t < FirstPass->Tags.Num(); ++t)
+		for (int32 TagIndex = 0; TagIndex < FirstPass->Tags.Num(); ++TagIndex)
 		{
 			TSet<int32> PendingSurfaces;
-			for (int32 s : FirstPass->Tags[t].Surfaces)
+			for (int32 SurfaceIndex : FirstPass->Tags[TagIndex].Surfaces)
 			{
-				PendingSurfaces.Add(s);
+				PendingSurfaces.Add(SurfaceIndex);
 			}
 
-			TSet<int32> processedSurfs;
+			TSet<int32> ProcessedSurfaces;
 			while (!PendingSurfaces.IsEmpty())
 			{
-				int32 cs = *PendingSurfaces.begin();
-				PendingSurfaces.Remove(cs);
+				const int32 PendingSurfaceIndex = *PendingSurfaces.begin();
+				PendingSurfaces.Remove(PendingSurfaceIndex);
 
-				if (processedSurfs.Contains(cs))
+				if (ProcessedSurfaces.Contains(PendingSurfaceIndex))
 				{
 					continue;
 				}
 
-				processedSurfs.Add(cs);
-
-				SurfacesPerTag[t].Add(cs);
-
-				FirstPassGenerator::FSurface& csurf = FirstPass->Surfaces[cs];
-				for (const FString& sct : csurf.PositiveTags)
+				ProcessedSurfaces.Add(PendingSurfaceIndex);
+				SurfacesPerTag[TagIndex].Add(PendingSurfaceIndex);
+				
+				FirstPassGenerator::FSurface& CurrentSurface = FirstPass->Surfaces[PendingSurfaceIndex];
+				for (const FString& PositiveTag : CurrentSurface.PositiveTags) 
 				{
-					const FirstPassGenerator::FTag* it = FirstPass->Tags.FindByPredicate([&](const FirstPassGenerator::FTag& e) { return e.Tag == sct; });
-					if (!it)
+					const int32 PositiveTagIndex = FirstPass->Tags.IndexOfByPredicate([&](const FirstPassGenerator::FTag& e) { return e.Tag == PositiveTag; });
+					if (PositiveTagIndex == INDEX_NONE)
 					{
-						// This could happen if a tag is in a variation but noone defines it.
+						// This could happen if a tag is in a variation but no one defines it.
 						continue;
 					}
+					
+					TagsPerTag[TagIndex].Add(PositiveTagIndex);
 
-					int32 ct = int32(it - &FirstPass->Tags[0]);
-
-					TagsPerTag[t].Add(ct);
-
-					for (int32 s : FirstPass->Tags[ct].Surfaces)
+					for (int32 Surface : FirstPass->Tags[PositiveTagIndex].Surfaces)
 					{
-						if (SurfacesPerTag[t].Contains(s))
+						if (SurfacesPerTag[TagIndex].Contains(Surface))
 						{
-							PendingSurfaces.Add(s);
+							PendingSurfaces.Add(Surface);
 						}
 					}
 				}
-				for (const FString& sct : csurf.NegativeTags)
+				for (const FString& NegativeTag : CurrentSurface.NegativeTags)
 				{
-					const FirstPassGenerator::FTag* it = FirstPass->Tags.FindByPredicate([&](const FirstPassGenerator::FTag& e) { return e.Tag == sct; });
-					if (!it)
+					// Look for the NegativeTag being processed and store it.
+					const int32 NegativeTagIndex = FirstPass->Tags.IndexOfByPredicate([&](const FirstPassGenerator::FTag& e) { return e.Tag == NegativeTag; });
+					if (NegativeTagIndex == INDEX_NONE)
 					{
-						// This could happen if a tag is in a variation but noone defines it.
+						// This could happen if a tag is in a variation but no one defines it.
 						continue;
 					}
+					
+					TagsPerTag[TagIndex].Add(NegativeTagIndex);
 
-					size_t ct = it - &FirstPass->Tags[0];
-
-					TagsPerTag[t].Add(ct);
-
-					for (size_t s : FirstPass->Tags[ct].Surfaces)
+					for (size_t Surface : FirstPass->Tags[NegativeTagIndex].Surfaces)
 					{
-						if (SurfacesPerTag[t].Contains(s))
+						if (SurfacesPerTag[TagIndex].Contains(Surface))
 						{
-							PendingSurfaces.Add(s);
+							PendingSurfaces.Add(Surface);
 						}
 					}
 				}
@@ -402,19 +392,17 @@ namespace UE::Mutable::Private
 		{
 			FirstPassGenerator::FSurface& Surface = FirstPass->Surfaces[SurfaceIndex];
 
-			{
-				Ptr<ASTOp> c = GenerateDataCodition(
-					SurfaceIndex, 
-					Surface.PositiveTags,
-					Surface.NegativeTags,
-					Empty, Empty, Empty, Empty);
+			Ptr<ASTOp> c = GenerateDataCodition(
+				SurfaceIndex, 
+				Surface.PositiveTags,
+				Surface.NegativeTags,
+				Empty, Empty, Empty, Empty);
 
-				Ptr<ASTOpBoolAnd> ConditionOp = new ASTOpBoolAnd();
-				ConditionOp->A = Surface.ObjectCondition;
-				ConditionOp->B = c;
+			Ptr<ASTOpBoolAnd> ConditionOp = new ASTOpBoolAnd();
+			ConditionOp->A = Surface.ObjectCondition;
+			ConditionOp->B = c;
 
-				Surface.FinalCondition = ConditionOp;
-			}
+			Surface.FinalCondition = ConditionOp;
 			
 			// TODO: Convert to modifiers that enable tags?
 			//for (int32 EditIndex = 0; EditIndex < Surface.Edits.Num(); ++EditIndex)
