[Mutable] Add MeshDescription to generated Skeletal Meshes.

#rnx

diff --git a/Engine/Plugins/Mutable/Source/CustomizableObject/CustomizableObject.Build.cs b/Engine/Plugins/Mutable/Source/CustomizableObject/CustomizableObject.Build.cs
--- a/Engine/Plugins/Mutable/Source/CustomizableObject/CustomizableObject.Build.cs
+++ b/Engine/Plugins/Mutable/Source/CustomizableObject/CustomizableObject.Build.cs
@@ -25,7 +25,8 @@ public class CustomizableObject : ModuleRules
 			//"ClothingSystemEditor",
 			"UMG",
 			"ImageCore",
-			"TraceLog"
+			"TraceLog",
+			"MeshDescription"
 		]);
 
         PublicDependencyModuleNames.AddRange([
diff --git a/Engine/Plugins/Mutable/Source/CustomizableObject/Internal/MuCO/CustomizableObjectInstancePrivate.h b/Engine/Plugins/Mutable/Source/CustomizableObject/Internal/MuCO/CustomizableObjectInstancePrivate.h
--- a/Engine/Plugins/Mutable/Source/CustomizableObject/Internal/MuCO/CustomizableObjectInstancePrivate.h
+++ b/Engine/Plugins/Mutable/Source/CustomizableObject/Internal/MuCO/CustomizableObjectInstancePrivate.h
@@ -312,7 +312,7 @@ public:
 
 	/** */
 #if WITH_EDITORONLY_DATA
-	UE_API void RegenerateImportedModels();
+	UE_API void RegenerateImportedModels(const TSharedRef<FUpdateContextPrivate>& OperationData);
 #endif
 	
 	// 
diff --git a/Engine/Plugins/Mutable/Source/CustomizableObject/Private/MuCO/CustomizableObjectInstance.cpp b/Engine/Plugins/Mutable/Source/CustomizableObject/Private/MuCO/CustomizableObjectInstance.cpp
--- a/Engine/Plugins/Mutable/Source/CustomizableObject/Private/MuCO/CustomizableObjectInstance.cpp
+++ b/Engine/Plugins/Mutable/Source/CustomizableObject/Private/MuCO/CustomizableObjectInstance.cpp
@@ -6040,7 +6040,12 @@ void UCustomizableInstancePrivate::BuildMaterials(const TSharedRef<FUpdateContex
 				if (FoundSurfaceMetadata)
 				{
 					MaterialSlot.MaterialSlotName = FoundSurfaceMetadata->MaterialSlotName;
+
+#if WITH_EDITOR
+					MaterialSlot.ImportedMaterialSlotName = FoundSurfaceMetadata->MaterialSlotName;
+#endif
 				}
+
 				if (RefSkeletalMeshData)
 				{
 					SetMeshUVChannelDensity(MaterialSlot.UVChannelData, RefSkeletalMeshData->Settings.DefaultUVChannelDensity);
@@ -7225,9 +7230,9 @@ void CalculateBonesToRemove(const FSkeletalMeshLODRenderData& LODResource, const
 	}
 }
 
-void UCustomizableInstancePrivate::RegenerateImportedModels()
+void UCustomizableInstancePrivate::RegenerateImportedModels(const TSharedRef<FUpdateContextPrivate>& OperationData)
 {
-	MUTABLE_CPUPROFILER_SCOPE(RegenerateImportedModels);
+	MUTABLE_CPUPROFILER_SCOPE(RegenerateEditorImportedModels);
 
 	struct FMeshDataConvertJob
 	{
@@ -7597,6 +7602,93 @@ void UCustomizableInstancePrivate::RegenerateImportedModels()
 			}
 		});
 	}
+
+	TArray<UE::Tasks::FTask> CommitMeshDescriptionTasks;
+
+	for (const TTuple<FName, TObjectPtr<USkeletalMesh>>& Tuple : SkeletalMeshes)
+	{
+		USkeletalMesh* SkeletalMesh = Tuple.Get<1>();
+
+		if (!SkeletalMesh)
+		{
+			continue;
+		}
+
+		const bool bIsTransientMesh = SkeletalMesh->HasAllFlags(EObjectFlags::RF_Transient);
+
+		if (!bIsTransientMesh)
+		{
+			// This must be a pass-through referenced mesh so don't do anything to it
+			continue;
+		}
+
+		FSkeletalMeshRenderData* RenderData = SkeletalMesh->GetResourceForRendering();
+		if (!RenderData || RenderData->IsInitialized())
+		{
+			continue;
+		}
+
+		const int32 LODCount = SkeletalMesh->GetLODNum();
+
+		TArray<UE::Tasks::FTask> MeshDescriptionGenerationTasks;
+		MeshDescriptionGenerationTasks.Reserve(LODCount);
+
+		TStrongObjectPtr<USkeletalMesh> StrongMesh(SkeletalMesh);
+
+		for (int32 LODIndex = 0; LODIndex < LODCount; ++LODIndex)
+		{
+			if (ensure(!SkeletalMesh->HasMeshDescription(LODIndex)))
+			{
+				MeshDescriptionGenerationTasks.Add(UE::Tasks::Launch(TEXT("Generate Mesh Description"),
+					[StrongMesh, LODIndex]()
+					{
+						USkeletalMesh* SkeletalMesh = StrongMesh.Get();
+						const FSkeletalMeshLODModel& LODModel = SkeletalMesh->GetImportedModel()->LODModels[LODIndex];
+
+						FMeshDescription MeshDescription;
+						LODModel.GetMeshDescription(SkeletalMesh, LODIndex, MeshDescription);
+						SkeletalMesh->CreateMeshDescription(LODIndex, MoveTemp(MeshDescription));
+					}));
+			}
+		}
+
+		CommitMeshDescriptionTasks.Add(UE::Tasks::Launch(TEXT("Commit Mesh Descriptions Task"),
+			[StrongMesh]()
+			{
+				USkeletalMesh* SkeletalMesh = StrongMesh.Get();
+
+				for (int32 LODIndex = 0; LODIndex < SkeletalMesh->GetLODNum(); ++LODIndex)
+				{
+					SkeletalMesh->CommitMeshDescription(LODIndex);
+
+					// Ensure normals aren't automatically computed when we rebuild.
+					FSkeletalMeshLODInfo* MeshLODInfo = SkeletalMesh->GetLODInfo(LODIndex);
+					FSkeletalMeshBuildSettings& BuildSettings = MeshLODInfo->BuildSettings;
+					BuildSettings.bRecomputeNormals = false;
+
+					// Reset the reduction settings so that we don't re-reduce the mesh and possibly lose morph targets
+					// in the process.
+					FSkeletalMeshOptimizationSettings& ReductionSettings = MeshLODInfo->ReductionSettings;
+
+					//Remove the reduction settings
+					ReductionSettings.NumOfTrianglesPercentage = 1.0f;
+					ReductionSettings.NumOfVertPercentage = 1.0f;
+					ReductionSettings.MaxNumOfTrianglesPercentage = MAX_uint32;
+					ReductionSettings.MaxNumOfVertsPercentage = MAX_uint32;
+					ReductionSettings.TerminationCriterion = SMTC_NumOfTriangles;
+					MeshLODInfo->bHasBeenSimplified = false;
+
+					const FSkeletalMeshLODModel& LODModel = SkeletalMesh->GetImportedModel()->LODModels[LODIndex];
+					LODModel.BuildStringID = LODModel.GetLODModelDeriveDataKey();
+
+				}
+			}, MeshDescriptionGenerationTasks));
+	}
+
+	if (OperationData->bBake)
+	{
+		UE::Tasks::Wait(CommitMeshDescriptionTasks);
+	}
 }
 
 #endif
diff --git a/Engine/Plugins/Mutable/Source/CustomizableObject/Private/MuCO/CustomizableObjectSystem.cpp b/Engine/Plugins/Mutable/Source/CustomizableObject/Private/MuCO/CustomizableObjectSystem.cpp
--- a/Engine/Plugins/Mutable/Source/CustomizableObject/Private/MuCO/CustomizableObjectSystem.cpp
+++ b/Engine/Plugins/Mutable/Source/CustomizableObject/Private/MuCO/CustomizableObjectSystem.cpp
@@ -3200,7 +3200,7 @@ namespace impl
 					MUTABLE_CPUPROFILER_SCOPE(UpdateSkeletalMesh_PostBeginUpdate2);
 					
 #if WITH_EDITORONLY_DATA
-					CustomizableInstancePrivateData->RegenerateImportedModels();
+					CustomizableInstancePrivateData->RegenerateImportedModels(OperationData);
 #endif
 					CustomizableInstancePrivateData->PostEditChangePropertyWithoutEditor();
 				}
diff --git a/Engine/Plugins/Mutable/Source/CustomizableObjectEditor/CustomizableObjectEditor.Build.cs b/Engine/Plugins/Mutable/Source/CustomizableObjectEditor/CustomizableObjectEditor.Build.cs
--- a/Engine/Plugins/Mutable/Source/CustomizableObjectEditor/CustomizableObjectEditor.Build.cs
+++ b/Engine/Plugins/Mutable/Source/CustomizableObjectEditor/CustomizableObjectEditor.Build.cs
@@ -69,7 +69,9 @@ public class CustomizableObjectEditor : ModuleRules
 				"TextureCompressor",
 				"TextureBuildUtilities",
 				"ImageCore",
-				"KismetWidgets"
+				"KismetWidgets",
+
+				"MeshDescription"
 			}
 		);
 
diff --git a/Engine/Plugins/Mutable/Source/CustomizableObjectEditor/Private/MuCOE/CustomizableObjectInstanceBakingUtils.cpp b/Engine/Plugins/Mutable/Source/CustomizableObjectEditor/Private/MuCOE/CustomizableObjectInstanceBakingUtils.cpp
--- a/Engine/Plugins/Mutable/Source/CustomizableObjectEditor/Private/MuCOE/CustomizableObjectInstanceBakingUtils.cpp
+++ b/Engine/Plugins/Mutable/Source/CustomizableObjectEditor/Private/MuCOE/CustomizableObjectInstanceBakingUtils.cpp
@@ -23,6 +23,8 @@
 #include "MuCO/CustomizableObjectMipDataProvider.h"
 #include "MuT/UnrealPixelFormatOverride.h"
 #include "Rendering/SkeletalMeshModel.h"
+#include "MeshDescription.h"
+
 
 #define LOCTEXT_NAMESPACE "CustomizableObjectEditor" 
 
@@ -989,9 +991,11 @@ bool BakeCustomizableObjectInstance(
 
 		if (USkeletalMesh* SkeletalMesh = Cast<USkeletalMesh>(DuplicatedMesh))
 		{
+			const int32 NumLODs = Mesh->GetLODNum();
+
 			SkeletalMesh->SetNumSourceModels(0);
-			for (int32 LODIndex = 0; LODIndex < Mesh->GetLODNum(); ++LODIndex)
-        	{
+			for (int32 LODIndex = 0; LODIndex < NumLODs; ++LODIndex)
+			{
         		SkeletalMesh->AddLODInfo(*Mesh->GetLODInfo(LODIndex));
         	}
 		
@@ -1050,6 +1054,21 @@ bool BakeCustomizableObjectInstance(
 				SkeletalMesh->GetPhysicsAsset()->SetPreviewMesh(SkeletalMesh);
 			}
 
+			// Copy MeshDescriptions
+			for (int32 LODIndex = 0; LODIndex < NumLODs; ++LODIndex)
+			{
+				if (Mesh->HasMeshDescription(LODIndex))
+				{
+					FMeshDescription MeshDescription;
+					Mesh->CloneMeshDescription(LODIndex, MeshDescription);
+					SkeletalMesh->CreateMeshDescription(LODIndex, MoveTemp(MeshDescription));
+					SkeletalMesh->CommitMeshDescription(LODIndex);
+
+					const FSkeletalMeshLODModel& LODModel = SkeletalMesh->GetImportedModel()->LODModels[LODIndex];
+					LODModel.BuildStringID = LODModel.GetLODModelDeriveDataKey();
+				}
+			}
+
 			// Generate render data
 			SkeletalMesh->Build();
 		}
