[mutable] Updated Required Tags and MultipleTagPolicy comments in favor of more descriptive ones.

#rnx

diff --git a/Engine/Plugins/Mutable/Source/CustomizableObjectEditor/Private/MuCOE/Nodes/CustomizableObjectNodeMeshMorphDetails.cpp b/Engine/Plugins/Mutable/Source/CustomizableObjectEditor/Private/MuCOE/Nodes/CustomizableObjectNodeMeshMorphDetails.cpp
--- a/Engine/Plugins/Mutable/Source/CustomizableObjectEditor/Private/MuCOE/Nodes/CustomizableObjectNodeMeshMorphDetails.cpp
+++ b/Engine/Plugins/Mutable/Source/CustomizableObjectEditor/Private/MuCOE/Nodes/CustomizableObjectNodeMeshMorphDetails.cpp
@@ -144,7 +144,7 @@ void FCustomizableObjectNodeMeshMorphDetails::OnMorphTargetComboBoxSelectionChan
 
 bool FCustomizableObjectNodeMeshMorphDetails::IsMorphNameSelectorWidgetEnabled() const
 {
-	// Disabled if there is a string node linked to the "Required Tags" pin.
+	// Disabled if there is a string node linked to the "Target Tags" pin.
 	return !(Node->MorphTargetNamePin() && FollowInputPin(*Node->MorphTargetNamePin()));
 }
 
diff --git a/Engine/Plugins/Mutable/Source/CustomizableObjectEditor/Private/MuCOE/Nodes/CustomizableObjectNodeModifierBase.cpp b/Engine/Plugins/Mutable/Source/CustomizableObjectEditor/Private/MuCOE/Nodes/CustomizableObjectNodeModifierBase.cpp
--- a/Engine/Plugins/Mutable/Source/CustomizableObjectEditor/Private/MuCOE/Nodes/CustomizableObjectNodeModifierBase.cpp
+++ b/Engine/Plugins/Mutable/Source/CustomizableObjectEditor/Private/MuCOE/Nodes/CustomizableObjectNodeModifierBase.cpp
@@ -269,7 +269,7 @@ void UCustomizableObjectNodeModifierBase::BackwardsCompatibleFixup(int32 Customi
 	{
 		if (!RequiredTagsPinRef.Get())
 		{
-			RequiredTagsPinRef = CustomCreatePin(EGPD_Input, UEdGraphSchema_CustomizableObject::PC_String, "Required Tags", true);
+			RequiredTagsPinRef = CustomCreatePin(EGPD_Input, UEdGraphSchema_CustomizableObject::PC_String, "Target Tags", true);
 		}
 	}
 }
@@ -346,5 +346,5 @@ FString UCustomizableObjectNodeModifierBase::MakeNodeAutoTag( UEdGraphNode* Node
 
 void UCustomizableObjectNodeModifierBase::AllocateDefaultPins(UCustomizableObjectNodeRemapPins* RemapPins)
 {
-	RequiredTagsPinRef = CustomCreatePin(EGPD_Input, UEdGraphSchema_CustomizableObject::PC_String, "Required Tags", true);
+	RequiredTagsPinRef = CustomCreatePin(EGPD_Input, UEdGraphSchema_CustomizableObject::PC_String, "Target Tags", true);
 }
diff --git a/Engine/Plugins/Mutable/Source/CustomizableObjectEditor/Private/MuCOE/Nodes/CustomizableObjectNodeModifierBase.h b/Engine/Plugins/Mutable/Source/CustomizableObjectEditor/Private/MuCOE/Nodes/CustomizableObjectNodeModifierBase.h
--- a/Engine/Plugins/Mutable/Source/CustomizableObjectEditor/Private/MuCOE/Nodes/CustomizableObjectNodeModifierBase.h
+++ b/Engine/Plugins/Mutable/Source/CustomizableObjectEditor/Private/MuCOE/Nodes/CustomizableObjectNodeModifierBase.h
@@ -42,7 +42,7 @@ public:
 	UPROPERTY(EditAnywhere, BlueprintReadWrite, Category = Modifier)
 	TArray<FString> RequiredTags;
 
-	/** Policy to use required tags in case more than one is added. */
+	/** Policy to use when multiple Target Mesh Section tags are set. */
 	UPROPERTY(EditAnywhere, Category = Modifier)
 	EMutableMultipleTagPolicy MultipleTagPolicy = EMutableMultipleTagPolicy::OnlyOneRequired;
 
diff --git a/Engine/Plugins/Mutable/Source/CustomizableObjectEditor/Private/MuCOE/Nodes/CustomizableObjectNodeModifierBaseDetails.cpp b/Engine/Plugins/Mutable/Source/CustomizableObjectEditor/Private/MuCOE/Nodes/CustomizableObjectNodeModifierBaseDetails.cpp
--- a/Engine/Plugins/Mutable/Source/CustomizableObjectEditor/Private/MuCOE/Nodes/CustomizableObjectNodeModifierBaseDetails.cpp
+++ b/Engine/Plugins/Mutable/Source/CustomizableObjectEditor/Private/MuCOE/Nodes/CustomizableObjectNodeModifierBaseDetails.cpp
@@ -57,7 +57,7 @@ void FCustomizableObjectNodeModifierBaseDetails::CustomizeDetails( IDetailLayout
 	TagsPolicyPropertyHandle = DetailBuilder.GetProperty(GET_MEMBER_NAME_CHECKED(UCustomizableObjectNodeModifierBase, MultipleTagPolicy), UCustomizableObjectNodeModifierBase::StaticClass());
 	TagsPolicyPropertyHandle->SetOnPropertyValueChanged(FSimpleDelegate::CreateSP(this, &FCustomizableObjectNodeModifierBaseDetails::OnRequiredTagsPropertyChanged));
 
-	ModifierCategory.AddCustomRow(FText::FromString(TEXT("Required Tags")))
+	ModifierCategory.AddCustomRow(FText::FromString(TEXT("Target Tags")))
 		.PropertyHandleList({ RequiredTagsPropertyHandle })
 		.NameContent()
 			.VAlign(VAlign_Fill)
@@ -68,7 +68,7 @@ void FCustomizableObjectNodeModifierBaseDetails::CustomizeDetails( IDetailLayout
 				.Padding(FMargin(0,4.0f,0,4.0f))
 				[
 					SNew(STextBlock)
-						.Text(LOCTEXT("ModifierDetails_RequiredTags", "Required Tags"))
+						.Text(LOCTEXT("ModifierDetails_RequiredTags", "Target Tags"))
 						.Font(IDetailLayoutBuilder::GetDetailFont())
 						.IsEnabled(this, &FCustomizableObjectNodeModifierBaseDetails::IsTagsPropertyWidgetEnabled)
 						.ToolTipText(this, &FCustomizableObjectNodeModifierBaseDetails::TagsPropertyWidgetTooltip)
@@ -100,7 +100,7 @@ bool FCustomizableObjectNodeModifierBaseDetails::IsTagsPropertyWidgetEnabled() c
 {
 	const UEdGraphPin* ReqTagsPin = Node->RequiredTagsPin();
 
-	// Disabled if there is a string node linked to the "Required Tags" pin.
+	// Disabled if there is a string node linked to the "Target Tags" pin.
 	return !(ReqTagsPin && FollowInputPinArray(*ReqTagsPin).Num());
 }
 
@@ -111,11 +111,11 @@ FText FCustomizableObjectNodeModifierBaseDetails::TagsPropertyWidgetTooltip() co
 
 	if (ReqTagsPin && FollowInputPinArray(*ReqTagsPin).Num())
 	{
-		return LOCTEXT("RequiredTagsWidgetTooltip_Ignored", "Disabled. When there are string nodes linked to the Required Tags pin, the tag list is ignored.");
+		return LOCTEXT("RequiredTagsWidgetTooltip_Ignored", "Disabled. When there are string nodes linked to the Target Tags pin, the tag list is ignored.");
 	}
 	else
 	{
-		return LOCTEXT("RequiredTagsWidgetTooltip", "List of Tags that this node will Required.");
+		return LOCTEXT("RequiredTagsWidgetTooltip", "Collection of tags to be used to decide which Mesh Sections will get the action of this Modifier.");
 	}
 }
 
diff --git a/Engine/Plugins/Mutable/Source/CustomizableObjectEditor/Private/MuCOE/Nodes/CustomizableObjectNodeModifierMorphMeshSectionDetails.cpp b/Engine/Plugins/Mutable/Source/CustomizableObjectEditor/Private/MuCOE/Nodes/CustomizableObjectNodeModifierMorphMeshSectionDetails.cpp
--- a/Engine/Plugins/Mutable/Source/CustomizableObjectEditor/Private/MuCOE/Nodes/CustomizableObjectNodeModifierMorphMeshSectionDetails.cpp
+++ b/Engine/Plugins/Mutable/Source/CustomizableObjectEditor/Private/MuCOE/Nodes/CustomizableObjectNodeModifierMorphMeshSectionDetails.cpp
@@ -326,7 +326,7 @@ void FCustomizableObjectNodeModifierMorphMeshSectionDetails::AddMorphsFromNode(
 
 bool FCustomizableObjectNodeModifierMorphMeshSectionDetails::IsMorphNameSelectorWidgetEnabled() const
 {
-	// Disabled if there is a string node linked to the "Required Tags" pin.
+	// Disabled if there is a string node linked to the "Target Tags" pin.
 	return !(Node->MorphTargetNamePin() && FollowInputPin(*Node->MorphTargetNamePin()));
 }
 
diff --git a/Engine/Plugins/Mutable/Source/MutableTools/Internal/MuT/NodeModifier.h b/Engine/Plugins/Mutable/Source/MutableTools/Internal/MuT/NodeModifier.h
--- a/Engine/Plugins/Mutable/Source/MutableTools/Internal/MuT/NodeModifier.h
+++ b/Engine/Plugins/Mutable/Source/MutableTools/Internal/MuT/NodeModifier.h
@@ -20,7 +20,10 @@
 UENUM()
 enum class EMutableMultipleTagPolicy : uint8
 {
+	/** All the Mesh Sections that do have at least one of the Target tags will be modified. */
 	OnlyOneRequired,
+
+	/** Only the Mesh Sections that do have all the Target tags will be modified. */
 	AllRequired
 };
 
