[Mutable] Fix invalid data distribution override that could cause a division by zero when computing the number of files to use per data type.

#rnx

diff --git a/Engine/Plugins/Mutable/Source/CustomizableObject/Private/MuCO/CustomizableObject.cpp b/Engine/Plugins/Mutable/Source/CustomizableObject/Private/MuCO/CustomizableObject.cpp
--- a/Engine/Plugins/Mutable/Source/CustomizableObject/Private/MuCO/CustomizableObject.cpp
+++ b/Engine/Plugins/Mutable/Source/CustomizableObject/Private/MuCO/CustomizableObject.cpp
@@ -2629,7 +2629,7 @@ namespace UE::Mutable::Private
 						[Category](const FFileCategoryOverride& C) { return C.Id == Category.Id; });
 
 					// Can't override if there's no override data or if the override doesn't have the minimum required number of files.
-					const bool bIsValidOverride = CategoryOverride && (Category.DataSize == 0 || CategoryOverride->NumFiles >= 0);
+					const bool bIsValidOverride = CategoryOverride && (Category.DataSize == 0 || CategoryOverride->NumFiles > 0);
 					bCanOverride = bCanOverride && bIsValidOverride;
 
 					if (!CategoryOverride)
