[Mutable] Fix non-deterministic mesh conversion due to uninitialized data when an Alt SkinWeight Profile has invalid data.

#rnx

diff --git a/Engine/Plugins/Mutable/Source/CustomizableObjectEditor/Private/MuCOE/GenerateMutableSource/GenerateMutableSourceMesh.cpp b/Engine/Plugins/Mutable/Source/CustomizableObjectEditor/Private/MuCOE/GenerateMutableSource/GenerateMutableSourceMesh.cpp
--- a/Engine/Plugins/Mutable/Source/CustomizableObjectEditor/Private/MuCOE/GenerateMutableSource/GenerateMutableSourceMesh.cpp
+++ b/Engine/Plugins/Mutable/Source/CustomizableObjectEditor/Private/MuCOE/GenerateMutableSource/GenerateMutableSourceMesh.cpp
@@ -519,7 +519,7 @@ void NormalizeWeights(uint16* InfluenceBones, uint16* InfluenceWeights, const in
 		{
 			if (InfluenceWeights[j] == 0)
 			{
-				FMemory::Memzero(&InfluenceWeights[j], MutableInfluenceCount - j);
+				FMemory::Memzero(&InfluenceWeights[j], (MutableInfluenceCount - j) * sizeof(uint16));
 				break;
 			}
 		}
@@ -2224,14 +2224,15 @@ UE::Tasks::FTask ConvertSkeletalMeshToMutable(TSharedRef<FMeshConversionContext>
 				continue;
 			}
 
-			if (ImportedProfileData->SkinWeights.Num() != LODModel.NumVertices)
-			{
-				UE_LOG(LogMutable, Display, TEXT("Invalid data found in SkeletalMesh [%s] when converting SkinWeightProfiles. Regenerate and save the SkeletalMesh."), *SkeletalMesh->GetName());
-				continue;
-			}
-
 			const uint32 ProfileId = static_cast<int32>(Context.GetSkinWeightProfileIdUnique(Profile.Name));
 
+			const int32 ProfileIndex = Context.SkinWeightProfilesInfo.AddUnique({ Profile.Name, ProfileId, false, 0 });
+			FMutableSkinWeightProfileInfo& MutSkinWeightProfileInfo = Context.SkinWeightProfilesInfo[ProfileIndex];
+
+			const FName PlatformName = *Context.Options.TargetPlatform->PlatformName();
+			MutSkinWeightProfileInfo.DefaultProfile = MutSkinWeightProfileInfo.DefaultProfile || Profile.DefaultProfile.GetValueForPlatform(PlatformName);
+			MutSkinWeightProfileInfo.DefaultProfileFromLODIndex = FMath::Min(MutSkinWeightProfileInfo.DefaultProfileFromLODIndex, Profile.DefaultProfileFromLODIndex.GetValueForPlatform(PlatformName));
+
 			int32 BufferIndex = INDEX_NONE;
 			int32 ChannelIndex = INDEX_NONE;
 			MutableMesh->GetVertexBuffers().FindChannel(UE::Mutable::Private::EMeshBufferSemantic::AltSkinWeight, ProfileId, &BufferIndex, &ChannelIndex);
@@ -2246,6 +2247,14 @@ UE::Tasks::FTask ConvertSkeletalMeshToMutable(TSharedRef<FMeshConversionContext>
 			uint8* DestData = VertexBuffer.Data.GetData();
 			const int32 DestDataSize = VertexBuffer.ElementSize;
 
+			if (ImportedProfileData->SkinWeights.Num() != LODModel.NumVertices)
+			{
+				UE_LOG(LogMutable, Display, TEXT("Invalid data found in SkeletalMesh [%s] when converting SkinWeightProfiles. Regenerate and save the SkeletalMesh."), *SkeletalMesh->GetName());
+				
+				FMemory::Memset(DestData, 0, DestDataSize * VertexCount);
+				continue;
+			}
+
 			const int32 NumChannels = VertexBuffer.Channels.Num();
 			check(NumChannels == 3);
 
@@ -2296,13 +2305,6 @@ UE::Tasks::FTask ConvertSkeletalMeshToMutable(TSharedRef<FMeshConversionContext>
 				FMemory::Memcpy(DestData, &InfluenceWeights[0], BoneWeightsSizeBytes);
 				DestData += BoneWeightsSizeBytes;
 			}
-
-			const int32 ProfileIndex = Context.SkinWeightProfilesInfo.AddUnique({ Profile.Name, ProfileId, false, 0 });
-			FMutableSkinWeightProfileInfo& MutSkinWeightProfileInfo = Context.SkinWeightProfilesInfo[ProfileIndex];
-
-			const FName PlatformName = *Context.Options.TargetPlatform->PlatformName();
-			MutSkinWeightProfileInfo.DefaultProfile = MutSkinWeightProfileInfo.DefaultProfile || Profile.DefaultProfile.GetValueForPlatform(PlatformName);
-			MutSkinWeightProfileInfo.DefaultProfileFromLODIndex = FMath::Min(MutSkinWeightProfileInfo.DefaultProfileFromLODIndex, Profile.DefaultProfileFromLODIndex.GetValueForPlatform(PlatformName));
 		}
 
 		NextBufferIndex += LODModel.SkinWeightProfiles.Num();
