[Mutable] Fix missing MinLOD and MinQualityLevelLOD on generated SkeletalMeshes.

#rnx

diff --git a/Engine/Plugins/Mutable/Source/CustomizableObject/Internal/MuCO/CustomizableObjectPrivate.h b/Engine/Plugins/Mutable/Source/CustomizableObject/Internal/MuCO/CustomizableObjectPrivate.h
--- a/Engine/Plugins/Mutable/Source/CustomizableObject/Internal/MuCO/CustomizableObjectPrivate.h
+++ b/Engine/Plugins/Mutable/Source/CustomizableObject/Internal/MuCO/CustomizableObjectPrivate.h
@@ -1729,6 +1729,8 @@ enum class ECustomizableObjectVersions : int32
 
 	FixSkinWeightProfilesConversion,
 
+	FixMissingLODSettings,
+
 	// -----<new versions can be added above this line>--------
 	LastCustomizableObjectVersion
 };
diff --git a/Engine/Plugins/Mutable/Source/CustomizableObjectEditor/Private/MuCOE/GenerateMutableSource/GenerateMutableSourceComponent.cpp b/Engine/Plugins/Mutable/Source/CustomizableObjectEditor/Private/MuCOE/GenerateMutableSource/GenerateMutableSourceComponent.cpp
--- a/Engine/Plugins/Mutable/Source/CustomizableObjectEditor/Private/MuCOE/GenerateMutableSource/GenerateMutableSourceComponent.cpp
+++ b/Engine/Plugins/Mutable/Source/CustomizableObjectEditor/Private/MuCOE/GenerateMutableSource/GenerateMutableSourceComponent.cpp
@@ -610,6 +610,7 @@ void FirstPass(UCustomizableObjectNodeComponentMesh& Node, FMutableGraphGenerati
 	ComponentInfo.Node = &Node;
 	ComponentInfo.AccumulateBonesToRemovePerLOD(Node.LODReductionSettings, Node.NumLODs);
 	ComponentInfo.MacroContext = GenerationContext.MacroNodesStack;
+	ComponentInfo.LODSettings = LODSettings;
 
 	GenerationContext.CompilationContext->ComponentInfos.Add(ComponentInfo);
 
