[Mutable] Fix mesh conversion using the default sockets priority when converting meshes to Mutable instead of the current socket priority. Generated meshes would have the default socket data from the base mesh.

#rnx

diff --git a/Engine/Plugins/Mutable/Source/CustomizableObject/Internal/MuCO/CustomizableObjectPrivate.h b/Engine/Plugins/Mutable/Source/CustomizableObject/Internal/MuCO/CustomizableObjectPrivate.h
--- a/Engine/Plugins/Mutable/Source/CustomizableObject/Internal/MuCO/CustomizableObjectPrivate.h
+++ b/Engine/Plugins/Mutable/Source/CustomizableObject/Internal/MuCO/CustomizableObjectPrivate.h
@@ -1731,6 +1731,8 @@ enum class ECustomizableObjectVersions : int32
 
 	FixMissingLODSettings,
 
+	FixSocketPriority,
+
 	// -----<new versions can be added above this line>--------
 	LastCustomizableObjectVersion
 };
diff --git a/Engine/Plugins/Mutable/Source/CustomizableObjectEditor/Private/MuCOE/GenerateMutableSource/GenerateMutableSource.cpp b/Engine/Plugins/Mutable/Source/CustomizableObjectEditor/Private/MuCOE/GenerateMutableSource/GenerateMutableSource.cpp
--- a/Engine/Plugins/Mutable/Source/CustomizableObjectEditor/Private/MuCOE/GenerateMutableSource/GenerateMutableSource.cpp
+++ b/Engine/Plugins/Mutable/Source/CustomizableObjectEditor/Private/MuCOE/GenerateMutableSource/GenerateMutableSource.cpp
@@ -267,6 +267,9 @@ FMutableGraphGenerationContext::FMutableGraphGenerationContext(const TSharedPtr<
 
 	// Default flags for mesh layout generation.
 	LayoutGenerationFlags.Push(FLayoutGenerationFlags());
+
+	// Default SocketPriority for mesh generation.
+	SocketPriorityStack.Push(0);
 }
 
 
diff --git a/Engine/Plugins/Mutable/Source/CustomizableObjectEditor/Private/MuCOE/GenerateMutableSource/GenerateMutableSourceMesh.cpp b/Engine/Plugins/Mutable/Source/CustomizableObjectEditor/Private/MuCOE/GenerateMutableSource/GenerateMutableSourceMesh.cpp
--- a/Engine/Plugins/Mutable/Source/CustomizableObjectEditor/Private/MuCOE/GenerateMutableSource/GenerateMutableSourceMesh.cpp
+++ b/Engine/Plugins/Mutable/Source/CustomizableObjectEditor/Private/MuCOE/GenerateMutableSource/GenerateMutableSourceMesh.cpp
@@ -2639,6 +2639,7 @@ TSharedPtr<UE::Mutable::Private::FMesh> ConvertSkeletalMeshToMutable(
 	Source.Flags = GenerationContext.MeshGenerationFlags.Last();
 	Source.MessageContext = CurrentNode;
 	Source.Component = GenerationContext.CurrentMeshComponent;
+	Source.SocketPriority = GenerationContext.SocketPriorityStack.Last();
 
 	// Prepare source data for realtime morph targets
 	const UCustomizableObjectNodeSkeletalMesh* NodeTypedSkMesh = Cast<UCustomizableObjectNodeSkeletalMesh>(CurrentNode);
