[Mutable] Fix crash when generating empty components.

diff --git a/Engine/Plugins/Mutable/Source/CustomizableObject/Private/MuCO/CustomizableObjectInstance.cpp b/Engine/Plugins/Mutable/Source/CustomizableObject/Private/MuCO/CustomizableObjectInstance.cpp
--- a/Engine/Plugins/Mutable/Source/CustomizableObject/Private/MuCO/CustomizableObjectInstance.cpp
+++ b/Engine/Plugins/Mutable/Source/CustomizableObject/Private/MuCO/CustomizableObjectInstance.cpp
@@ -2289,9 +2289,8 @@ bool UCustomizableInstancePrivate::UpdateSkeletalMesh_PostBeginUpdate0(UCustomiz
 		uint32 MeshHash = INDEX_NONE;
 		for (int32 LODIndex = FirstLOD; LODIndex < Component->LODCount && LODIndex < MAX_MESH_LOD_COUNT; ++LODIndex)
 		{
-			const UE::Mutable::Private::FGeneratedMeshKey* MeshKey = (*MeshDescriptors)[LODIndex].GetKey();
-			check(MeshKey);
-			const uint32 MeshIDHash = GetTypeHashPersistent(*MeshKey);
+			const UE::Mutable::Private::FMeshId& MeshId = (*MeshDescriptors)[LODIndex];
+			const uint32 MeshIDHash = UE::Mutable::Private::GetTypeHashPersistent(MeshId);
 
 			if (!Context->bBake)
 			{
@@ -6221,7 +6220,7 @@ void UCustomizableInstancePrivate::BuildMaterials(const TSharedRef<FUpdateContex
 
 								FString MutableTextureName;
 								{
-									const uint32 MutableImageHash = UE::Mutable::Private::GetTypeHashPersistent(*Image.ImageID.GetKey());
+									const uint32 MutableImageHash = UE::Mutable::Private::GetTypeHashPersistent(Image.ImageID);
 									MutableTextureName = FString::Printf(TEXT("T_%s_%lu"), *Props.TextureParameterName, MutableImageHash);
 									MutableTextureName.ReplaceInline(TEXT(" "), TEXT("_"));
 
diff --git a/Engine/Plugins/Mutable/Source/MutableRuntime/Internal/MuR/ResourceID.h b/Engine/Plugins/Mutable/Source/MutableRuntime/Internal/MuR/ResourceID.h
--- a/Engine/Plugins/Mutable/Source/MutableRuntime/Internal/MuR/ResourceID.h
+++ b/Engine/Plugins/Mutable/Source/MutableRuntime/Internal/MuR/ResourceID.h
@@ -1,4 +1,4 @@
-﻿// Copyright Epic Games, Inc. All Rights Reserved.
+// Copyright Epic Games, Inc. All Rights Reserved.
 
 #pragma once
 
@@ -24,13 +24,13 @@ namespace UE::Mutable::Private
 	};
 	
 	MUTABLERUNTIME_API uint32 GetTypeHash(const FGeneratedMeshKey& Key);
-
-	MUTABLERUNTIME_API uint32 GetTypeHashPersistent(const FGeneratedMeshKey& Key);
 	
 	typedef TRegistry<FGeneratedMeshKey, FGeneratedMeshData>::FHandle FMeshId;
 
 	typedef TRegistry<FGeneratedMeshKey, FGeneratedMeshData> FMeshIdRegistry;
 
+	MUTABLERUNTIME_API uint32 GetTypeHashPersistent(const FMeshId& Key);
+
 	// Image
 	struct FGeneratedImageKey
 	{
@@ -48,12 +48,12 @@ namespace UE::Mutable::Private
 	
 	MUTABLERUNTIME_API uint32 GetTypeHash(const FGeneratedImageKey& Key);
 
-	MUTABLERUNTIME_API uint32 GetTypeHashPersistent(const FGeneratedImageKey& Key);
-
 	typedef TRegistry<FGeneratedImageKey, FGeneratedImageData>::FHandle FImageId;
 
 	typedef TRegistry<FGeneratedImageKey, FGeneratedImageData> FImageIdRegistry;
 
+	MUTABLERUNTIME_API uint32 GetTypeHashPersistent(const FImageId& Key);
+
 	// Material
 	struct FGeneratedMaterialKey
 	{
@@ -70,10 +70,10 @@ namespace UE::Mutable::Private
 	};
 	
 	MUTABLERUNTIME_API uint32 GetTypeHash(const FGeneratedMaterialKey& Key);
-
-	MUTABLERUNTIME_API uint32 GetTypeHashPersistent(const FGeneratedMaterialKey& Key);
 	
 	typedef TRegistry<FGeneratedMaterialKey, FGeneratedMaterialData>::FHandle FMaterialId;
 
 	typedef TRegistry<FGeneratedMaterialKey, FGeneratedMaterialData> FMaterialIdRegistry;
+
+	MUTABLERUNTIME_API uint32 GetTypeHashPersistent(const FMaterialId& Key);
 }
diff --git a/Engine/Plugins/Mutable/Source/MutableRuntime/Private/MuR/ResourceID.cpp b/Engine/Plugins/Mutable/Source/MutableRuntime/Private/MuR/ResourceID.cpp
--- a/Engine/Plugins/Mutable/Source/MutableRuntime/Private/MuR/ResourceID.cpp
+++ b/Engine/Plugins/Mutable/Source/MutableRuntime/Private/MuR/ResourceID.cpp
@@ -11,9 +11,16 @@ uint32 UE::Mutable::Private::GetTypeHash(const FGeneratedMeshKey& Key)
 }
 
 
-uint32 UE::Mutable::Private::GetTypeHashPersistent(const FGeneratedMeshKey& Key)
+uint32 UE::Mutable::Private::GetTypeHashPersistent(const FMeshId& Id)
 {
-	return HashCombine(Key.Address, static_cast<uint32>(FXxHash64::HashBuffer(Key.ParameterValuesBlob.GetData(), Key.ParameterValuesBlob.Num()).Hash));
+	if (const FGeneratedMeshKey* Key = Id.GetKey())
+	{
+		return HashCombine(Key->Address, static_cast<uint32>(FXxHash64::HashBuffer(Key->ParameterValuesBlob.GetData(), Key->ParameterValuesBlob.Num()).Hash));
+	}
+	else 
+	{
+		return 0;
+	}
 }
 
 
@@ -23,9 +30,16 @@ uint32 UE::Mutable::Private::GetTypeHash(const FGeneratedImageKey& Key)
 }
 
 
-uint32 UE::Mutable::Private::GetTypeHashPersistent(const FGeneratedImageKey& Key)
+uint32 UE::Mutable::Private::GetTypeHashPersistent(const FImageId& Id)
 {
-	return HashCombine(Key.Address, static_cast<uint32>(FXxHash64::HashBuffer(Key.ParameterValuesBlob.GetData(), Key.ParameterValuesBlob.Num()).Hash));
+	if (const FGeneratedImageKey* Key = Id.GetKey())
+	{
+		return HashCombine(Key->Address, static_cast<uint32>(FXxHash64::HashBuffer(Key->ParameterValuesBlob.GetData(), Key->ParameterValuesBlob.Num()).Hash));
+	}
+	else
+	{
+		return 0;
+	}
 }
 
 
@@ -35,7 +49,14 @@ uint32 UE::Mutable::Private::GetTypeHash(const FGeneratedMaterialKey& Key)
 }
 
 
-uint32 UE::Mutable::Private::GetTypeHashPersistent(const FGeneratedMaterialKey& Key)
+uint32 UE::Mutable::Private::GetTypeHashPersistent(const FMaterialId& Id)
 {
-	return HashCombine(Key.Address, static_cast<uint32>(FXxHash64::HashBuffer(Key.ParameterValuesBlob.GetData(), Key.ParameterValuesBlob.Num()).Hash));
+	if (const FGeneratedMaterialKey* Key = Id.GetKey())
+	{
+		return HashCombine(Key->Address, static_cast<uint32>(FXxHash64::HashBuffer(Key->ParameterValuesBlob.GetData(), Key->ParameterValuesBlob.Num()).Hash));
+	}
+	else
+	{
+		return 0;
+	}
 }
\ No newline at end of file
