[Mutable] Add option to regenerate cooked data distribution.

#rnx

diff --git a/Engine/Plugins/Mutable/Source/CustomizableObject/Internal/MuCO/CustomizableObjectPrivate.h b/Engine/Plugins/Mutable/Source/CustomizableObject/Internal/MuCO/CustomizableObjectPrivate.h
--- a/Engine/Plugins/Mutable/Source/CustomizableObject/Internal/MuCO/CustomizableObjectPrivate.h
+++ b/Engine/Plugins/Mutable/Source/CustomizableObject/Internal/MuCO/CustomizableObjectPrivate.h
@@ -63,6 +63,12 @@ FGuid CUSTOMIZABLEOBJECT_API GenerateIdentifier(const UCustomizableObject& Custo
 
 FString CUSTOMIZABLEOBJECT_API GetModelResourcesNameForPlatform(const UCustomizableObject& CustomizableObject, const ITargetPlatform& Platform);
 
+#if WITH_EDITOR
+
+FGuid CUSTOMIZABLEOBJECT_API GenerateDataDistributionIdentifier(const UCustomizableObject& CustomizableObject);
+
+#endif
+
 // A USTRUCT version of FMeshToMeshVertData in SkeletalMeshTypes.h
 // We are taking advantage of the padding data to store from which asset this data comes from
 // maintaining the same memory footprint than the original.
@@ -1838,6 +1844,9 @@ public:
 	CUSTOMIZABLEOBJECT_API FGuid GetVersionId() const;
 
 	CUSTOMIZABLEOBJECT_API void SaveEmbeddedData(FArchive& Ar);
+
+	// Regenerate the DataDistributionId. Used to keep the same cooked data distribution between builds.
+	CUSTOMIZABLEOBJECT_API void UpdateDataDistributionId();
 	
 	// Add a profile that stores the values of the parameters used by the CustomInstance.
 	CUSTOMIZABLEOBJECT_API FReply AddNewParameterProfile(FString Name, class UCustomizableObjectInstance& CustomInstance);
@@ -2008,6 +2017,11 @@ public:
 	UPROPERTY()
 	bool bDisableTableMissingDataWarning = false;
 
+	/** Used to keep the same cooked data distribution between builds. 
+	  * This Id is part of the DDC key of the cooked data distribution. */
+	UPROPERTY()
+	FGuid CookedDataDistributionId;
+
 #endif
 	
 	static constexpr int32 DerivedDataVersion = 0x290ba346;
diff --git a/Engine/Plugins/Mutable/Source/CustomizableObject/Private/MuCO/CustomizableObject.cpp b/Engine/Plugins/Mutable/Source/CustomizableObject/Private/MuCO/CustomizableObject.cpp
--- a/Engine/Plugins/Mutable/Source/CustomizableObject/Private/MuCO/CustomizableObject.cpp
+++ b/Engine/Plugins/Mutable/Source/CustomizableObject/Private/MuCO/CustomizableObject.cpp
@@ -495,6 +495,17 @@ FString GetModelResourcesNameForPlatform(const UCustomizableObject& Customizable
 	return GenerateIdentifier(CustomizableObject).ToString() + InTargetPlatform.PlatformName();
 }
 
+#if WITH_EDITOR
+FGuid GenerateDataDistributionIdentifier(const UCustomizableObject& CustomizableObject)
+{
+	if (CustomizableObject.GetPrivate()->CookedDataDistributionId.IsValid())
+	{
+		return CustomizableObject.GetPrivate()->CookedDataDistributionId;
+	}
+
+	return GenerateIdentifier(CustomizableObject);
+}
+#endif
 
 bool UCustomizableObjectPrivate::LoadModelResources(FArchive& MemoryReader, const ITargetPlatform* InTargetPlatform, bool bIsCooking)
 {
@@ -901,6 +912,11 @@ void UCustomizableObjectPrivate::SaveEmbeddedData(FArchive& Ar)
 	}
 }
 
+void UCustomizableObjectPrivate::UpdateDataDistributionId()
+{
+	CookedDataDistributionId = FGuid::NewGuid();
+}
+
 #endif // End WITH_EDITOR 
 
 void UCustomizableObjectPrivate::LoadEmbeddedData(FArchive& Ar)
diff --git a/Engine/Plugins/Mutable/Source/CustomizableObjectEditor/Private/MuCOE/CustomizableObjectCompileRunnable.cpp b/Engine/Plugins/Mutable/Source/CustomizableObjectEditor/Private/MuCOE/CustomizableObjectCompileRunnable.cpp
--- a/Engine/Plugins/Mutable/Source/CustomizableObjectEditor/Private/MuCOE/CustomizableObjectCompileRunnable.cpp
+++ b/Engine/Plugins/Mutable/Source/CustomizableObjectEditor/Private/MuCOE/CustomizableObjectCompileRunnable.cpp
@@ -435,7 +435,7 @@ FCustomizableObjectSaveDDRunnable::FCustomizableObjectSaveDDRunnable(const TShar
 
 	UCustomizableObject* CustomizableObject = InRequest->GetCustomizableObject();
 	CustomizableObjectName = GetNameSafe(CustomizableObject);
-	CustomizableObjectId = GenerateIdentifier(*CustomizableObject);
+	CustomizableObjectId = GenerateDataDistributionIdentifier(*CustomizableObject);
 
 	CustomizableObjectHeader.InternalVersion = GetECustomizableObjectVersionEnumHash();
 	CustomizableObjectHeader.VersionId = CustomizableObject->GetPrivate()->GetVersionId();
@@ -556,7 +556,7 @@ void FCustomizableObjectSaveDDRunnable::StoreCachedPlatformDataInDDC(bool& bStor
 		MUTABLE_CPUPROFILER_SCOPE(SerializeBulkDataForDDC);
 
 		// Generate list of files and update streamable blocks ids and offsets
-		constexpr int32 MaxDDCFiles = 1 << 13;
+		const int32 MaxDDCFiles = Options.bIsCooking ? MAX_uint8 : 1 << 13;
 		UE::Mutable::Private::GenerateBulkDataFilesListWithFileLimit(CustomizableObjectId, PlatformData->Model, ModelStreamablesDDC, MaxDDCFiles, 
 			Options.bIsCooking, !Options.bIsCooking, *Options.TargetPlatform, BulkDataFilesDDC);
 
diff --git a/Engine/Plugins/Mutable/Source/CustomizableObjectEditor/Private/MuCOE/CustomizableObjectCompiler.cpp b/Engine/Plugins/Mutable/Source/CustomizableObjectEditor/Private/MuCOE/CustomizableObjectCompiler.cpp
--- a/Engine/Plugins/Mutable/Source/CustomizableObjectEditor/Private/MuCOE/CustomizableObjectCompiler.cpp
+++ b/Engine/Plugins/Mutable/Source/CustomizableObjectEditor/Private/MuCOE/CustomizableObjectCompiler.cpp
@@ -1544,7 +1544,7 @@ void FCustomizableObjectCompiler::FinishLoadingModelDataFromDDC()
 					if (CVarMutableUseBulkData.GetValueOnAnyThread())
 					{
 						const uint32 NumBulkDataFilesPerBucket = MAX_uint8;
-						UE::Mutable::Private::GenerateBulkDataFilesListWithFileLimit(GenerateIdentifier(*CustomizableObject), PlatformData->Model, *PlatformData->ModelStreamableBulkData.Get(),
+						UE::Mutable::Private::GenerateBulkDataFilesListWithFileLimit(GenerateDataDistributionIdentifier(*CustomizableObject), PlatformData->Model, *PlatformData->ModelStreamableBulkData.Get(),
 							NumBulkDataFilesPerBucket, Options.bIsCooking, false, *Options.TargetPlatform, PlatformData->BulkDataFiles);
 					}
 					else
diff --git a/Engine/Plugins/Mutable/Source/CustomizableObjectEditor/Private/MuCOE/CustomizableObjectEditor.cpp b/Engine/Plugins/Mutable/Source/CustomizableObjectEditor/Private/MuCOE/CustomizableObjectEditor.cpp
--- a/Engine/Plugins/Mutable/Source/CustomizableObjectEditor/Private/MuCOE/CustomizableObjectEditor.cpp
+++ b/Engine/Plugins/Mutable/Source/CustomizableObjectEditor/Private/MuCOE/CustomizableObjectEditor.cpp
@@ -488,6 +488,12 @@ void FCustomizableObjectEditor::BindCommands()
 		FExecuteAction::CreateSP(this, &FCustomizableObjectEditor::CodeViewer),
 		FCanExecuteAction(),
 		FIsActionChecked());
+
+	ToolkitCommands->MapAction(
+		Commands.UpdateCookDataDistributionId,
+		FExecuteAction::CreateSP(this, &FCustomizableObjectEditor::UpdateCookDataDistributionId),
+		FCanExecuteAction(),
+		FIsActionChecked());
 	
 	// References
 	ToolkitCommands->MapAction(
@@ -1128,6 +1134,8 @@ TSharedRef<SWidget> FCustomizableObjectEditor::GenerateCompileOptionsMenuContent
 
 	MenuBuilder.BeginSection("Packaging", LOCTEXT("MutableCompilePackagingHeading", "Packaging"));
 	{
+		MenuBuilder.AddMenuEntry(FCustomizableObjectEditorCommands::Get().UpdateCookDataDistributionId);
+
 		// Unfortunately SNumericDropDown doesn't work with integers at the time of writing.
 		TArray<SNumericDropDown<float>::FNamedValue> EmbeddedOptions;
 		EmbeddedOptions.Add(SNumericDropDown<float>::FNamedValue(0, FText::FromString(TEXT("0")), FText::FromString(TEXT("Disabled"))));
@@ -1500,6 +1508,13 @@ void FCustomizableObjectEditor::CodeViewer() const
 }
 
 
+void FCustomizableObjectEditor::UpdateCookDataDistributionId()
+{
+	CustomizableObject->GetPrivate()->UpdateDataDistributionId();
+	CustomizableObject->Modify();
+}
+
+
 void FCustomizableObjectEditor::ClearGatheredReferences()
 {
 	CustomizableObject->GetPrivate()->ReferencedObjects = {};
diff --git a/Engine/Plugins/Mutable/Source/CustomizableObjectEditor/Private/MuCOE/CustomizableObjectEditor.h b/Engine/Plugins/Mutable/Source/CustomizableObjectEditor/Private/MuCOE/CustomizableObjectEditor.h
--- a/Engine/Plugins/Mutable/Source/CustomizableObjectEditor/Private/MuCOE/CustomizableObjectEditor.h
+++ b/Engine/Plugins/Mutable/Source/CustomizableObjectEditor/Private/MuCOE/CustomizableObjectEditor.h
@@ -178,6 +178,9 @@ public:
 	void GraphViewer() const;
 	void CodeViewer() const;
 
+	/** Regenerate CookDataDistributionId to generate a new entry of cooked data distribution in the DDC next time the object is cooked. */
+	void UpdateCookDataDistributionId();
+
 	/** Clear game asset references saved in the Customizable Object. */
 	void ClearGatheredReferences();
 	
diff --git a/Engine/Plugins/Mutable/Source/CustomizableObjectEditor/Private/MuCOE/CustomizableObjectEditorActions.cpp b/Engine/Plugins/Mutable/Source/CustomizableObjectEditor/Private/MuCOE/CustomizableObjectEditorActions.cpp
--- a/Engine/Plugins/Mutable/Source/CustomizableObjectEditor/Private/MuCOE/CustomizableObjectEditorActions.cpp
+++ b/Engine/Plugins/Mutable/Source/CustomizableObjectEditor/Private/MuCOE/CustomizableObjectEditorActions.cpp
@@ -28,6 +28,7 @@ void FCustomizableObjectEditorCommands::RegisterCommands()
 	UI_COMMAND(CompileOptions_UseDiskCompilation, "Enable compiling using the disk as memory.", "This is very slow but supports compiling huge objects. It requires a lot of free space in the OS disk.", EUserInterfaceActionType::ToggleButton, FInputChord());
 	UI_COMMAND(GraphViewer, "Graph Viewer", "Open the Customizable Object Graph Viewer.", EUserInterfaceActionType::Button, FInputChord());
 	UI_COMMAND(CodeViewer, "Code Viewer", "Open the Customizable Object Code Viewer.", EUserInterfaceActionType::Button, FInputChord());
+	UI_COMMAND(UpdateCookDataDistributionId, "Regenerate Cooked Data Distribution", "Optimize cooked data distribution based on current graph and generated data. Data distribution is stored in the DDC. \n\nWarning! Changing the data distribution may result in big patch sizes.", EUserInterfaceActionType::Button, FInputChord());
 
 	UI_COMMAND(PerformanceAnalyzer, "Performance Analyzer", "Open the Performance Analyzer window to set up and perform all tests relevant to Customizable Objects.", EUserInterfaceActionType::Button, FInputChord());
 	UI_COMMAND(TextureAnalyzer, "Texture Memory Analyzer", "Open the Texture Analyzer window to check all the information of the textures created by Mutable.", EUserInterfaceActionType::Button, FInputChord());
diff --git a/Engine/Plugins/Mutable/Source/CustomizableObjectEditor/Private/MuCOE/CustomizableObjectEditorActions.h b/Engine/Plugins/Mutable/Source/CustomizableObjectEditor/Private/MuCOE/CustomizableObjectEditorActions.h
--- a/Engine/Plugins/Mutable/Source/CustomizableObjectEditor/Private/MuCOE/CustomizableObjectEditorActions.h
+++ b/Engine/Plugins/Mutable/Source/CustomizableObjectEditor/Private/MuCOE/CustomizableObjectEditorActions.h
@@ -20,6 +20,7 @@ public:
 	TSharedPtr< FUICommandInfo > CompileOptions_UseDiskCompilation;
 	TSharedPtr< FUICommandInfo > GraphViewer;
 	TSharedPtr< FUICommandInfo > CodeViewer;
+	TSharedPtr< FUICommandInfo > UpdateCookDataDistributionId;
 
 	TSharedPtr< FUICommandInfo > PerformanceAnalyzer;
 	TSharedPtr< FUICommandInfo > ResetPerformanceReportOptions;
