[Mutable] Remove cooked references to private objects in AlwaysLoadedExtensionData.

#rnx

diff --git a/Engine/Plugins/Mutable/Source/CustomizableObject/Private/MuCO/CustomizableObject.cpp b/Engine/Plugins/Mutable/Source/CustomizableObject/Private/MuCO/CustomizableObject.cpp
--- a/Engine/Plugins/Mutable/Source/CustomizableObject/Private/MuCO/CustomizableObject.cpp
+++ b/Engine/Plugins/Mutable/Source/CustomizableObject/Private/MuCO/CustomizableObject.cpp
@@ -630,6 +630,14 @@ void UModelResources::InitCookData(UObject& InCustomizableObject)
 
 		StreamedExtensionData.Emplace(Container);
 	}
+
+	for (FCustomizableObjectResourceData& ResourceData : AlwaysLoadedExtensionData)
+	{
+		for (const UCustomizableObjectExtension* Extension : AllExtensions)
+		{
+			Extension->MovePrivateReferencesToContainer(ResourceData.Data, &InCustomizableObject);
+		}
+	}
 }
 
 
