[mutable] Fixed CO Compilation race condition due to the wrong mutex being used.

#rnx

diff --git a/Engine/Plugins/Mutable/Source/MutableTools/Private/MuT/CodeGenerator_Scalar.cpp b/Engine/Plugins/Mutable/Source/MutableTools/Private/MuT/CodeGenerator_Scalar.cpp
--- a/Engine/Plugins/Mutable/Source/MutableTools/Private/MuT/CodeGenerator_Scalar.cpp
+++ b/Engine/Plugins/Mutable/Source/MutableTools/Private/MuT/CodeGenerator_Scalar.cpp
@@ -55,7 +55,7 @@ namespace UE::Mutable::Private
 		Key.Node = Untyped;
 		Key.Options = Options;
 		{
-			UE::TUniqueLock Lock(GeneratedStrings.Mutex);
+			UE::TUniqueLock Lock(GeneratedScalars.Mutex);
 			FGeneratedScalarsMap::ValueType* Found = GeneratedScalars.Map.Find(Key);
 			if (Found)
 			{
