[Mutable] Replace ensure with a log that notifies the user about corrupted SkinWeightProfiles data.

#rnx

diff --git a/Engine/Plugins/Mutable/Source/CustomizableObjectEditor/Private/MuCOE/GenerateMutableSource/GenerateMutableSourceMesh.cpp b/Engine/Plugins/Mutable/Source/CustomizableObjectEditor/Private/MuCOE/GenerateMutableSource/GenerateMutableSourceMesh.cpp
--- a/Engine/Plugins/Mutable/Source/CustomizableObjectEditor/Private/MuCOE/GenerateMutableSource/GenerateMutableSourceMesh.cpp
+++ b/Engine/Plugins/Mutable/Source/CustomizableObjectEditor/Private/MuCOE/GenerateMutableSource/GenerateMutableSourceMesh.cpp
@@ -2193,7 +2193,7 @@ TSharedPtr<UE::Mutable::Private::FMesh> ConvertSkeletalMeshToMutable(TSharedRef<
 
 			if (ImportedProfileData->SkinWeights.Num() != Vertices.Num())
 			{
-				ensure(false);
+				UE_LOG(LogMutable, Display, TEXT("Invalid data found in SkeletalMesh [%s] when converting SkinWeightProfiles. Regenerate and save the SkeletalMesh."), *SkeletalMesh->GetName());
 				continue;
 			}
 
