[Mutable] Fixed compilation loop in PIE

#rnx

diff --git a/Engine/Plugins/Mutable/Source/CustomizableObjectEditor/Private/MuCOE/GenerateMutableSource/GenerateMutableSourceComponent.cpp b/Engine/Plugins/Mutable/Source/CustomizableObjectEditor/Private/MuCOE/GenerateMutableSource/GenerateMutableSourceComponent.cpp
--- a/Engine/Plugins/Mutable/Source/CustomizableObjectEditor/Private/MuCOE/GenerateMutableSource/GenerateMutableSourceComponent.cpp
+++ b/Engine/Plugins/Mutable/Source/CustomizableObjectEditor/Private/MuCOE/GenerateMutableSource/GenerateMutableSourceComponent.cpp
@@ -468,6 +468,11 @@ UE::Mutable::Private::Ptr<UE::Mutable::Private::NodeComponent> GenerateMutableSo
 		GenerationContext.GeneratedNodes.Add(Node);
 	}
 
+	if (Result)
+	{
+		Result->SetMessageContext(Node);
+	}
+
 	return Result;
 }
 
diff --git a/Engine/Plugins/Mutable/Source/MutableTools/Private/MuT/CodeGenerator.cpp b/Engine/Plugins/Mutable/Source/MutableTools/Private/MuT/CodeGenerator.cpp
--- a/Engine/Plugins/Mutable/Source/MutableTools/Private/MuT/CodeGenerator.cpp
+++ b/Engine/Plugins/Mutable/Source/MutableTools/Private/MuT/CodeGenerator.cpp
@@ -329,7 +329,7 @@ namespace UE::Mutable::Private
 	}
 
 
-	Ptr<NodeScalar> CodeGenerator::GenerateTableVariableNode(Ptr<const Node> InNode, const FTableCacheKey& CacheKey, bool bAddNoneOption, const FString& DefaultRowName)
+	Ptr<NodeScalar> CodeGenerator::GenerateTableVariableNode(const Node& InNode, const FTableCacheKey& CacheKey, bool bAddNoneOption, const FString& DefaultRowName)
 	{
 		Ptr<NodeScalarEnumParameter> Result = new NodeScalarEnumParameter;
 
@@ -338,9 +338,10 @@ namespace UE::Mutable::Private
 		{
 			ParamName = CacheKey.Table->GetName();
 		}
-		Result->Name = ParamName;
 
+		Result->Name = ParamName;
 		Result->DefaultValue = 0;
+		Result->SetMessageContext(InNode.GetMessageContext());
 
 		int32 CurrentRow = 0;
 		int32 RowCount = CacheKey.Table->GetPrivate()->Rows.Num();
diff --git a/Engine/Plugins/Mutable/Source/MutableTools/Private/MuT/CodeGenerator.h b/Engine/Plugins/Mutable/Source/MutableTools/Private/MuT/CodeGenerator.h
--- a/Engine/Plugins/Mutable/Source/MutableTools/Private/MuT/CodeGenerator.h
+++ b/Engine/Plugins/Mutable/Source/MutableTools/Private/MuT/CodeGenerator.h
@@ -416,7 +416,7 @@ namespace UE::Mutable::Private
 		void GenerateComponent_Variation(const FComponentGenerationOptions&, FGenericGenerationResult&, const NodeComponentVariation*);
 
         /** */
-		Ptr<NodeScalar> GenerateTableVariableNode(Ptr<const Node>, const FTableCacheKey&, bool bAddNoneOption, const FString& DefaultRowName);
+		Ptr<NodeScalar> GenerateTableVariableNode(const Node&, const FTableCacheKey&, bool bAddNoneOption, const FString& DefaultRowName);
 
         //!
         Ptr<ASTOp> GenerateMissingBoolCode(const TCHAR* strWhere, bool value, const void* errorContext );
@@ -1054,7 +1054,7 @@ namespace UE::Mutable::Private
 			if (!VariableNode)
 			{
 				// Create the table variable expression
-				VariableNode = GenerateTableVariableNode(&node, CacheKey, node.bNoneOption, node.DefaultRowName);
+				VariableNode = GenerateTableVariableNode(node, CacheKey, node.bNoneOption, node.DefaultRowName);
 				GeneratedTableNodes.Map.Add(CacheKey, VariableNode);
 			}
 		}
@@ -1132,7 +1132,7 @@ namespace UE::Mutable::Private
 			if (!Variable)
 			{
 				// Create the table variable expression
-				Variable = GenerateTableVariableNode(&node, CacheKey, node.bNoneOption, node.DefaultRowName);
+				Variable = GenerateTableVariableNode(node, CacheKey, node.bNoneOption, node.DefaultRowName);
 
 				GeneratedTableNodes.Map.Add(CacheKey, Variable);
 			}
diff --git a/Engine/Plugins/Mutable/Source/MutableTools/Private/MuT/Compiler.cpp b/Engine/Plugins/Mutable/Source/MutableTools/Private/MuT/Compiler.cpp
--- a/Engine/Plugins/Mutable/Source/MutableTools/Private/MuT/Compiler.cpp
+++ b/Engine/Plugins/Mutable/Source/MutableTools/Private/MuT/Compiler.cpp
@@ -530,6 +530,7 @@ namespace UE::Mutable::Private
 								{
 									if (Entry.Value->Parameter.Name.Compare(CurrentParameterName) == 0)
 									{
+										check(Entry.Key->GetMessageContext());
 										MessageContexts.AddUnique(Entry.Key->GetMessageContext());
 									}
 								}
@@ -539,6 +540,7 @@ namespace UE::Mutable::Private
 									{
 										if (ArrayEntry.Key && ArrayEntry.Key->Parameter.Name.Compare(CurrentParameterName) == 0)
 										{
+											check(Entry.Key->GetMessageContext());
 											MessageContexts.AddUnique(Entry.Key->GetMessageContext());
 										}
 									}
