[Mutable] External Operations.

Allow opreatoins declared outside Mutable. New External Operations can be declared in non-Mutable related plugins. Once an External Operation is registered, it automatically appears as CO node. These nodes are automatically compiled into Mutable VM code.
Users can define hteir own data types and also use already existing internal types (through adapters). This submit adds support to the following internal types:
- Scalars (float)
- Colors (FVector4f)
- FMesh
- FTexture
- FMaterial

Added an example plugin which show how to create and use new External Operations.

diff --git a/Engine/Plugins/Experimental/MutableExtensionExample/MutableExtensionExample.uplugin b/Engine/Plugins/Experimental/MutableExtensionExample/MutableExtensionExample.uplugin
new file mode 100644
--- /dev/null
+++ b/Engine/Plugins/Experimental/MutableExtensionExample/MutableExtensionExample.uplugin
@@ -0,0 +1,33 @@
+{
+	"FileVersion": 3,
+	"Version": 1,
+	"VersionName": "0.1",
+	"FriendlyName": "Mutable Extension Example",
+	"Description": "Example plugin which adds new external operations to Mutable.",
+	"Category": "Animation",
+	"CreatedBy": "Epic Games, Inc.",
+	"CreatedByURL": "https://epicgames.com",
+	"DocsURL": "",
+	"MarketplaceURL": "",
+	"SupportURL": "",
+	"CanContainContent": false,
+	"IsBetaVersion": false,
+	"IsExperimentalVersion": true,
+	"Installed": false,
+	"EnabledByDefault": false,
+	"Modules": 
+	[
+		{
+			"Name": "MutableExtensionExample",
+			"Type": "Runtime",
+			"LoadingPhase": "Default"
+		}
+	],
+	"Plugins": 
+	[
+		{
+			"Name": "Mutable",
+			"Enabled": true
+		}
+	]
+}
\ No newline at end of file
diff --git a/Engine/Plugins/Experimental/MutableExtensionExample/Source/MutableExtensionExample/MutableExtensionExample.Build.cs b/Engine/Plugins/Experimental/MutableExtensionExample/Source/MutableExtensionExample/MutableExtensionExample.Build.cs
new file mode 100644
--- /dev/null
+++ b/Engine/Plugins/Experimental/MutableExtensionExample/Source/MutableExtensionExample/MutableExtensionExample.Build.cs
@@ -0,0 +1,27 @@
+// Copyright Epic Games, Inc. All Rights Reserved.
+
+namespace UnrealBuildTool.Rules
+{
+	public class MutableExtensionExample : ModuleRules
+	{
+        public MutableExtensionExample(ReadOnlyTargetRules Target) : base(Target)
+		{
+			PrivateIncludePaths.Add(ModuleDirectory + "/Private");
+			PublicIncludePaths.Add(ModuleDirectory + "/Public");
+			
+			PublicDependencyModuleNames.AddRange(
+				new string[]
+				{
+					"Core",
+					"CoreUObject",
+					"MutableRuntime"
+				});
+
+			PrivateDependencyModuleNames.AddRange(
+				new string[]
+				{
+					"CustomizableObject"
+				});
+		}
+	}
+}
diff --git a/Engine/Plugins/Experimental/MutableExtensionExample/Source/MutableExtensionExample/Private/MeshIdentity.cpp b/Engine/Plugins/Experimental/MutableExtensionExample/Source/MutableExtensionExample/Private/MeshIdentity.cpp
new file mode 100644
--- /dev/null
+++ b/Engine/Plugins/Experimental/MutableExtensionExample/Source/MutableExtensionExample/Private/MeshIdentity.cpp
@@ -0,0 +1,36 @@
+﻿// Copyright Epic Games, Inc. All Rights Reserved.
+
+#include "MeshIdentity.h"
+
+#include "MuR/External/MeshAdapter.h"
+
+
+const FName FMeshIdentity::InputMesh = TEXT("Mesh");
+
+
+FName FMeshIdentity::GetName() const
+{
+	return TEXT("Mesh Identity");
+}
+
+
+TArray<TPair<FName, const UScriptStruct*>> FMeshIdentity::GetInputs() const
+{
+	TArray<TPair<FName, const UScriptStruct*>> Inputs;
+
+	Inputs.Emplace(InputMesh, UE::Mutable::FMeshAdapter::StaticStruct());
+	
+	return Inputs;
+}
+
+
+TPair<FName, const UScriptStruct*> FMeshIdentity::GetOutput() const
+{
+	return MakeTuple(TEXT("Mesh"), UE::Mutable::FMeshAdapter::StaticStruct());
+}
+
+
+void FMeshIdentity::Evaluate(UE::Mutable::FContext& Context) const
+{
+	Context.GetOutput() = Context.GetInput(InputMesh);
+}
diff --git a/Engine/Plugins/Experimental/MutableExtensionExample/Source/MutableExtensionExample/Private/MeshNoise.cpp b/Engine/Plugins/Experimental/MutableExtensionExample/Source/MutableExtensionExample/Private/MeshNoise.cpp
new file mode 100644
--- /dev/null
+++ b/Engine/Plugins/Experimental/MutableExtensionExample/Source/MutableExtensionExample/Private/MeshNoise.cpp
@@ -0,0 +1,62 @@
+﻿// Copyright Epic Games, Inc. All Rights Reserved.
+
+#include "MeshNoise.h"
+
+#include "MuR/External/FloatAdapter.h"
+#include "MuR/External/MeshAdapter.h"
+
+
+const FName FMeshNoise::InputMesh = TEXT("Mesh");
+
+const FName FMeshNoise::InputFactor = TEXT("Factor");
+
+
+FName FMeshNoise::GetName() const
+{
+	return TEXT("Mesh Noise");
+}
+
+
+TArray<TPair<FName, const UScriptStruct*>> FMeshNoise::GetInputs() const
+{
+	TArray<TPair<FName, const UScriptStruct*>> Inputs;
+
+	Inputs.Emplace(InputMesh, UE::Mutable::FMeshAdapter::StaticStruct());
+	Inputs.Emplace(InputFactor, UE::Mutable::FFloatAdapter::StaticStruct());
+	
+	return Inputs;
+}
+
+
+TPair<FName, const UScriptStruct*> FMeshNoise::GetOutput() const
+{
+	return MakeTuple(TEXT("Mesh"), UE::Mutable::FMeshAdapter::StaticStruct());
+}
+
+
+void FMeshNoise::Evaluate(UE::Mutable::FContext& Context) const
+{
+	// Since we are going to edit the Mesh input, we need to make it non-constant using the CopyOrMove.
+	TSharedPtr<FInstancedStruct> Input = UE::Mutable::CopyOrMove(Context.GetInput(InputMesh));
+
+	// Cast the input to a Mesh Adapter.
+	UE::Mutable::FMeshAdapter& Mesh = Input->GetMutable<UE::Mutable::FMeshAdapter>();
+
+	// Get the Factor input
+	const float Factor = Context.GetInput(InputFactor)->Get<UE::Mutable::FFloatAdapter>().GetValue();
+	
+	// Move each vertex some random noise.
+	FRandomStream Random(Seed);
+	
+	for (int32 Index = 0; Index < Mesh.GetNumVertices(); ++Index)
+	{
+		FVector3f Vertex = Mesh.GetVertex(Index);
+		
+		Vertex += static_cast<FVector3f>(Random.GetUnitVector()) * Factor;
+		
+		Mesh.SetVertex(Index, Vertex);
+	}
+	
+	// Output the modified Mesh Adapter.
+	Context.GetOutput() = Input;
+}
diff --git a/Engine/Plugins/Experimental/MutableExtensionExample/Source/MutableExtensionExample/Private/MutableExtensionExampleModule.cpp b/Engine/Plugins/Experimental/MutableExtensionExample/Source/MutableExtensionExample/Private/MutableExtensionExampleModule.cpp
new file mode 100644
--- /dev/null
+++ b/Engine/Plugins/Experimental/MutableExtensionExample/Source/MutableExtensionExample/Private/MutableExtensionExampleModule.cpp
@@ -0,0 +1,28 @@
+// Copyright Epic Games, Inc. All Rights Reserved.
+
+#include "MutableExtensionExampleModule.h"
+
+#include "MeshIdentity.h"
+#include "MeshNoise.h"
+#include "MuCO/ICustomizableObjectModule.h"
+
+
+IMPLEMENT_MODULE(FMutableExtensionExampleModule, MutableExtensionExample);
+
+
+void FMutableExtensionExampleModule::StartupModule()
+{
+	ICustomizableObjectModule& Module = ICustomizableObjectModule::Get();
+
+	Module.RegisterExternalOperation(FMeshIdentity::StaticStruct());
+	Module.RegisterExternalOperation(FMeshNoise::StaticStruct());
+}
+
+
+void FMutableExtensionExampleModule::ShutdownModule()
+{
+	ICustomizableObjectModule& Module = ICustomizableObjectModule::Get();
+
+	Module.UnregisterExternalOperation(FMeshIdentity::StaticStruct());
+	Module.UnregisterExternalOperation(FMeshNoise::StaticStruct());
+}
diff --git a/Engine/Plugins/Experimental/MutableExtensionExample/Source/MutableExtensionExample/Public/MeshIdentity.h b/Engine/Plugins/Experimental/MutableExtensionExample/Source/MutableExtensionExample/Public/MeshIdentity.h
new file mode 100644
--- /dev/null
+++ b/Engine/Plugins/Experimental/MutableExtensionExample/Source/MutableExtensionExample/Public/MeshIdentity.h
@@ -0,0 +1,22 @@
+﻿// Copyright Epic Games, Inc. All Rights Reserved.
+
+#pragma once
+
+#include "MuR/External/Operation.h"
+
+#include "MeshIdentity.generated.h"
+
+
+USTRUCT()
+struct FMeshIdentity : public UE::Mutable::FExternalOperation
+{
+	GENERATED_BODY()
+	
+	virtual ~FMeshIdentity() override = default;
+	virtual FName GetName() const override;
+	virtual TArray<TPair<FName, const UScriptStruct*>> GetInputs() const override;
+	virtual TPair<FName, const UScriptStruct*> GetOutput() const override;
+	virtual void Evaluate(UE::Mutable::FContext& Context) const override;
+
+	const static FName InputMesh;
+};
\ No newline at end of file
diff --git a/Engine/Plugins/Experimental/MutableExtensionExample/Source/MutableExtensionExample/Public/MeshNoise.h b/Engine/Plugins/Experimental/MutableExtensionExample/Source/MutableExtensionExample/Public/MeshNoise.h
new file mode 100644
--- /dev/null
+++ b/Engine/Plugins/Experimental/MutableExtensionExample/Source/MutableExtensionExample/Public/MeshNoise.h
@@ -0,0 +1,30 @@
+﻿// Copyright Epic Games, Inc. All Rights Reserved.
+
+#pragma once
+
+#include "MuR/External/Operation.h"
+
+#include "MeshNoise.generated.h"
+
+
+USTRUCT()
+struct FMeshNoise : public UE::Mutable::FExternalOperation
+{
+	GENERATED_BODY()
+	
+	virtual ~FMeshNoise() override = default;
+	virtual FName GetName() const override;
+	virtual TArray<TPair<FName, const UScriptStruct*>> GetInputs() const override;
+	virtual TPair<FName, const UScriptStruct*> GetOutput() const override;
+	virtual void Evaluate(UE::Mutable::FContext& Context) const override;
+
+	const static FName InputMesh;
+
+	const static FName InputFactor;
+
+	/** Operation constant. This value will be copied when compiling a Customizable Object, hence it can not be changed at runtime.
+	 *
+	 * If a value needs to change at runtime, it must be a Parameter. */
+	UPROPERTY(EditAnywhere, Category = "Mesh Noise")
+	int32 Seed = 0;
+};
\ No newline at end of file
diff --git a/Engine/Plugins/Experimental/MutableExtensionExample/Source/MutableExtensionExample/Public/MutableExtensionExampleModule.h b/Engine/Plugins/Experimental/MutableExtensionExample/Source/MutableExtensionExample/Public/MutableExtensionExampleModule.h
new file mode 100644
--- /dev/null
+++ b/Engine/Plugins/Experimental/MutableExtensionExample/Source/MutableExtensionExample/Public/MutableExtensionExampleModule.h
@@ -0,0 +1,13 @@
+// Copyright Epic Games, Inc. All Rights Reserved.
+
+#pragma once
+
+#include "Modules/ModuleInterface.h"
+
+
+class FMutableExtensionExampleModule : public IModuleInterface
+{
+	virtual void StartupModule() override;
+	virtual void ShutdownModule() override;
+};
+
diff --git a/Engine/Plugins/Mutable/Source/CustomizableObject/Internal/MuCO/CustomizableObjectPrivate.h b/Engine/Plugins/Mutable/Source/CustomizableObject/Internal/MuCO/CustomizableObjectPrivate.h
--- a/Engine/Plugins/Mutable/Source/CustomizableObject/Internal/MuCO/CustomizableObjectPrivate.h
+++ b/Engine/Plugins/Mutable/Source/CustomizableObject/Internal/MuCO/CustomizableObjectPrivate.h
@@ -1352,6 +1352,14 @@ public:
 	// UE::Mutable::Private::FExtensionData::Index is an index into this array when UE::Mutable::Private::ExtensionData::Origin is ConstantStreamed
 	UPROPERTY()
 	TArray<FCustomizableObjectStreamedResourceData> StreamedExtensionData;
+
+	// Key is UE::Mutable::Private::OP::ADDRESS. UPROPERTY does not support typedefs. 
+	UPROPERTY()
+	TMap<uint32, FInstancedStruct> ExternalOperations;
+	
+	// Value is the compiled version of the operation.
+	UPROPERTY()
+	TMap<TObjectPtr<const UScriptStruct>, uint32> ExternalOperationVersions;
 	
 	/** Max number of LODs in the compiled Model. */
 	UPROPERTY()
@@ -1747,12 +1755,13 @@ class UCustomizableObjectPrivate : public UObject
 
 public:
 	/** Must be called after unlocking the CustomizableObject. */
-	CUSTOMIZABLEOBJECT_API void SetModel(const TSharedPtr<UE::Mutable::Private::FModel, ESPMode::ThreadSafe>& Model, const FGuid Identifier);
-	CUSTOMIZABLEOBJECT_API const TSharedPtr<UE::Mutable::Private::FModel, ESPMode::ThreadSafe>& GetModel();
-	CUSTOMIZABLEOBJECT_API const TSharedPtr<const UE::Mutable::Private::FModel, ESPMode::ThreadSafe> GetModel() const;
+	CUSTOMIZABLEOBJECT_API void SetModel(const TSharedPtr<UE::Mutable::Private::FModel>& Model, const FGuid Identifier);
+	CUSTOMIZABLEOBJECT_API const TSharedPtr<UE::Mutable::Private::FModel>& GetModel();
+	CUSTOMIZABLEOBJECT_API const TSharedPtr<const UE::Mutable::Private::FModel> GetModel() const;
 
 	CUSTOMIZABLEOBJECT_API UModelResources* GetModelResources();
 	CUSTOMIZABLEOBJECT_API const UModelResources* GetModelResources() const;
+	CUSTOMIZABLEOBJECT_API UModelResources& GetModelResourcesChecked();
 	CUSTOMIZABLEOBJECT_API const UModelResources& GetModelResourcesChecked() const;
 
 
diff --git a/Engine/Plugins/Mutable/Source/CustomizableObject/Internal/MuCO/ICustomizableObjectModulePrivate.h b/Engine/Plugins/Mutable/Source/CustomizableObject/Internal/MuCO/ICustomizableObjectModulePrivate.h
new file mode 100644
--- /dev/null
+++ b/Engine/Plugins/Mutable/Source/CustomizableObject/Internal/MuCO/ICustomizableObjectModulePrivate.h
@@ -0,0 +1,21 @@
+// Copyright Epic Games, Inc. All Rights Reserved.
+
+#pragma once
+
+#include "Modules/ModuleManager.h"
+#include "MuCO/ICustomizableObjectModule.h"
+
+/**
+ * The public interface of the CustomizableObject module
+ */
+class ICustomizableObjectModulePrivate : public ICustomizableObjectModule
+{
+public:
+	static ICustomizableObjectModulePrivate& Get()
+	{
+		return FModuleManager::LoadModuleChecked<ICustomizableObjectModulePrivate>("CustomizableObject");
+	}
+	
+	virtual TArray<TStrongObjectPtr<const UScriptStruct>> GetRegisteredExternalOperations() const = 0;
+};
+
diff --git a/Engine/Plugins/Mutable/Source/CustomizableObject/Internal/MuCO/UnrealToMutableTextureConversionUtils.h b/Engine/Plugins/Mutable/Source/CustomizableObject/Internal/MuCO/UnrealToMutableTextureConversionUtils.h
--- a/Engine/Plugins/Mutable/Source/CustomizableObject/Internal/MuCO/UnrealToMutableTextureConversionUtils.h
+++ b/Engine/Plugins/Mutable/Source/CustomizableObject/Internal/MuCO/UnrealToMutableTextureConversionUtils.h
@@ -140,6 +140,21 @@ inline uint32 GetTypeHash(const FMutableSourceMeshData& Key)
 	return GuidHash;
 }
 
+
+struct FSourceExtensionOptions
+{
+	// GenerateMutableSourceMesh
+	FMutableSourceMeshData MeshOptions;
+	bool bLinkedToExtendMaterial = false;
+	bool bOnlyConnectedLOD = false;
+	
+	// GenerateMutableSourceImage
+	int32 ReferenceTextureSize = 0;
+
+	bool operator==(const FSourceExtensionOptions& Other) const = default;
+};
+
+
 enum class EUnrealToMutableConversionError 
 {
     Success,
diff --git a/Engine/Plugins/Mutable/Source/CustomizableObject/Private/MuCO/CustomizableObject.cpp b/Engine/Plugins/Mutable/Source/CustomizableObject/Private/MuCO/CustomizableObject.cpp
--- a/Engine/Plugins/Mutable/Source/CustomizableObject/Private/MuCO/CustomizableObject.cpp
+++ b/Engine/Plugins/Mutable/Source/CustomizableObject/Private/MuCO/CustomizableObject.cpp
@@ -49,6 +49,7 @@
 #endif
 
 
+#include "ExternalOperationProvider.h"
 #include "MuCO/CustomizableObjectCustomVersion.h"
 #include "MuCO/CustomizableObjectResourceDataTypes.h"
 #include "MuCO/CustomizableObjectSkeletalMesh.h"
@@ -2153,7 +2154,7 @@ void UModelStreamableData::PostLoad()
 }
 
 
-void UCustomizableObjectPrivate::SetModel(const TSharedPtr<UE::Mutable::Private::FModel, ESPMode::ThreadSafe>& Model, const FGuid Id)
+void UCustomizableObjectPrivate::SetModel(const TSharedPtr<UE::Mutable::Private::FModel>& Model, const FGuid Id)
 {
 	if (MutableModel == Model
 #if WITH_EDITOR
@@ -2175,6 +2176,12 @@ void UCustomizableObjectPrivate::SetModel(const TSharedPtr<UE::Mutable::Private:
 	
 	MutableModel = Model;
 
+	if (Model)
+	{
+		UModelResources& LocalModelResources = GetModelResourcesChecked();
+		MutableModel->GetPrivate()->Program.ExternalOperationProvider = MakeShared<UE::Mutable::Private::FExternalOperationProvider>(LocalModelResources);
+	}
+	
 	// Create parameter properties
 	UpdateParameterPropertiesFromModel(Model);
 
@@ -2244,6 +2251,14 @@ const UModelResources* UCustomizableObjectPrivate::GetModelResources() const
 #endif
 }
 
+
+UModelResources& UCustomizableObjectPrivate::GetModelResourcesChecked()
+{
+	const UCustomizableObjectPrivate* ConstThis = this;
+	return *const_cast<UModelResources*>(&ConstThis->GetModelResourcesChecked());
+}
+
+
 const UModelResources& UCustomizableObjectPrivate::GetModelResourcesChecked() const
 {
 	const UModelResources* LocalModelResources = GetModelResources();
diff --git a/Engine/Plugins/Mutable/Source/CustomizableObject/Private/MuCO/CustomizableObjectModule.cpp b/Engine/Plugins/Mutable/Source/CustomizableObject/Private/MuCO/CustomizableObjectModule.cpp
--- a/Engine/Plugins/Mutable/Source/CustomizableObject/Private/MuCO/CustomizableObjectModule.cpp
+++ b/Engine/Plugins/Mutable/Source/CustomizableObject/Private/MuCO/CustomizableObjectModule.cpp
@@ -1,5 +1,7 @@
 // Copyright Epic Games, Inc. All Rights Reserved.
 
+#include "MuCO/ICustomizableObjectModulePrivate.h"
+
 #include "MuCO/CustomizableObjectInstanceUsagePrivate.h"
 #include "Algo/Find.h"
 #include "Engine/Engine.h"
@@ -18,12 +20,13 @@
 #include "UObject/UObjectIterator.h"
 #include "GPUSkinPublicDefs.h"
 #include "Components/SkeletalMeshComponent.h"
+#include "MuR/External/Operation.h"
 
 /**
  * Customizable Object module implementation (private)
  */
 
-class FCustomizableObjectModule : public ICustomizableObjectModule
+class FCustomizableObjectModule : public ICustomizableObjectModulePrivate
 {
 public:
 
@@ -39,6 +42,11 @@ public:
 	TArrayView<const UCustomizableObjectExtension* const> GetRegisteredExtensions() const override;
 	TArrayView<const FRegisteredCustomizableObjectPinType> GetExtendedPinTypes() const override;
 	TArrayView<const FRegisteredObjectNodeInputPin> GetAdditionalObjectNodePins() const override;
+	virtual void RegisterExternalOperation(const UScriptStruct* Operation) override;
+	virtual void UnregisterExternalOperation(const UScriptStruct* Operation) override;
+
+	// ICustomizableObjectModulePrivate interface
+	virtual TArray<TStrongObjectPtr<const UScriptStruct>> GetRegisteredExternalOperations() const override;
 
 private:
 	void RefreshExtensionData();
@@ -46,7 +54,7 @@ private:
 	static void InitializeSystem();
 
 	// Command to look for Customizable Object Instance in the player pawn of the current world and open a DGGUI to edit its parameters
-	IConsoleCommand* LaunchDGGUICommand;
+	IConsoleCommand* LaunchDGGUICommand = nullptr;
 	static void ToggleDGGUI(const TArray<FString>& Arguments);
 
 	// Ensure extensions aren't garbage collected
@@ -56,6 +64,8 @@ private:
 
 	TArray<FRegisteredCustomizableObjectPinType> ExtendedPinTypes;
 	TArray<FRegisteredObjectNodeInputPin> AdditionalObjectNodePins;
+
+	TArray<TStrongObjectPtr<const UScriptStruct>> RegisteredOperations;
 };
 
 
@@ -176,6 +186,41 @@ TArrayView<const FRegisteredObjectNodeInputPin> FCustomizableObjectModule::GetAd
 	return MakeArrayView(AdditionalObjectNodePins);
 }
 
+
+void FCustomizableObjectModule::RegisterExternalOperation(const UScriptStruct* Operation)
+{
+	if (!Operation)
+	{
+		return;
+	}
+
+	if (!Operation->IsChildOf(UE::Mutable::FExternalOperation::StaticStruct()))
+	{
+		UE_LOG(LogMutable, Error, TEXT("Unable to register operation. Operation that does not inherit from FMutableOperation: %s"), *Operation->GetName());
+		return;
+	}
+	
+	RegisteredOperations.AddUnique(TStrongObjectPtr(Operation));
+}
+
+
+void FCustomizableObjectModule::UnregisterExternalOperation(const UScriptStruct* Operation)
+{
+	if (!Operation)
+	{
+		return;
+	}
+	
+	RegisteredOperations.Remove(TStrongObjectPtr(Operation));
+}
+
+
+TArray<TStrongObjectPtr<const UScriptStruct>> FCustomizableObjectModule::GetRegisteredExternalOperations() const
+{
+	return RegisteredOperations;
+}
+
+
 void FCustomizableObjectModule::RefreshExtensionData()
 {
 	ExtendedPinTypes.Reset();
diff --git a/Engine/Plugins/Mutable/Source/CustomizableObject/Private/MuCO/ExternalOperationProvider.cpp b/Engine/Plugins/Mutable/Source/CustomizableObject/Private/MuCO/ExternalOperationProvider.cpp
new file mode 100644
--- /dev/null
+++ b/Engine/Plugins/Mutable/Source/CustomizableObject/Private/MuCO/ExternalOperationProvider.cpp
@@ -0,0 +1,29 @@
+﻿// Copyright Epic Games, Inc. All Rights Reserved.
+
+#include "MuCO/ExternalOperationProvider.h"
+
+#include "CustomizableObject/Internal/MuCO/CustomizableObjectPrivate.h"
+
+
+namespace UE::Mutable::Private
+{
+	FExternalOperationProvider::FExternalOperationProvider(UModelResources& InModelResources)
+	{
+		ExternalOperations = InModelResources.ExternalOperations;
+	}
+
+
+	void FExternalOperationProvider::AddReferencedObjects(FReferenceCollector& Collector)
+	{
+		for (TPair<OP::ADDRESS, FInstancedStruct>& Pair : ExternalOperations)
+		{
+			Collector.AddPropertyReferences(Pair.Value.GetScriptStruct(), Pair.Value.GetMutableMemory());
+		}
+	}
+
+
+	const FInstancedStruct& FExternalOperationProvider::Get(OP::ADDRESS Address) const
+	{
+		return ExternalOperations[Address];
+	}
+}
diff --git a/Engine/Plugins/Mutable/Source/CustomizableObject/Private/MuCO/ExternalOperationProvider.h b/Engine/Plugins/Mutable/Source/CustomizableObject/Private/MuCO/ExternalOperationProvider.h
new file mode 100644
--- /dev/null
+++ b/Engine/Plugins/Mutable/Source/CustomizableObject/Private/MuCO/ExternalOperationProvider.h
@@ -0,0 +1,35 @@
+﻿// Copyright Epic Games, Inc. All Rights Reserved.
+
+#pragma once
+
+#include "MuR/ExternalOperationProvider.h"
+#include "UObject/GCObject.h"
+
+#include "MuR/Operations.h"
+
+struct FInstancedStruct;
+class UModelResources;
+
+
+namespace UE::Mutable::Private
+{
+	class FExternalOperationProvider : public IExternalOperationProvider, public FGCObject
+	{
+	public:
+		virtual void AddReferencedObjects(FReferenceCollector& Collector) override; 
+		virtual FString GetReferencerName() const override
+		{
+			return TEXT("FExternalOperationProvider");
+		}
+
+		// IExternalOperationProvider interface
+		virtual const FInstancedStruct& Get(OP::ADDRESS Address) const override;
+		
+		// Own interface
+		FExternalOperationProvider(UModelResources& InModelResources);
+	
+	private:
+		TMap<OP::ADDRESS, FInstancedStruct> ExternalOperations;
+	};
+}
+
diff --git a/Engine/Plugins/Mutable/Source/CustomizableObject/Public/MuCO/ICustomizableObjectModule.h b/Engine/Plugins/Mutable/Source/CustomizableObject/Public/MuCO/ICustomizableObjectModule.h
--- a/Engine/Plugins/Mutable/Source/CustomizableObject/Public/MuCO/ICustomizableObjectModule.h
+++ b/Engine/Plugins/Mutable/Source/CustomizableObject/Public/MuCO/ICustomizableObjectModule.h
@@ -64,7 +64,6 @@ public:
 	 *
 	 * These may only be called from the game thread
 	 */
-
 	virtual void RegisterExtension(const UCustomizableObjectExtension* Extension) = 0;
 	virtual void UnregisterExtension(const UCustomizableObjectExtension* Extension) = 0;
 	virtual TArrayView<const UCustomizableObjectExtension* const> GetRegisteredExtensions() const = 0;
@@ -78,6 +77,19 @@ public:
 	 */
 	virtual TArrayView<const FRegisteredCustomizableObjectPinType> GetExtendedPinTypes() const = 0;
 	virtual TArrayView<const FRegisteredObjectNodeInputPin> GetAdditionalObjectNodePins() const = 0;
+	
+	/** EXPERIMENTAL. USE WITH CAUTION!
+	 *
+	 * Register a new External Operation.
+	 *
+	 * All External Operations must inherit from UE::Mutable::FMutableOperation.
+	 * It is only necessary to register operations in editor (to be able to create nodes). No need to register them at runtime. */
+	virtual void RegisterExternalOperation(const UScriptStruct* Operation) = 0;
+
+	/** EXPERIMENTAL. USE WITH CAUTION!
+	 * 
+	 * Unregister an External Operation. */
+	virtual void UnregisterExternalOperation(const UScriptStruct* Operation) = 0;
 };
 
 
diff --git a/Engine/Plugins/Mutable/Source/CustomizableObjectEditor/Private/MuCOE/CustomizableObjectCompileRunnable.cpp b/Engine/Plugins/Mutable/Source/CustomizableObjectEditor/Private/MuCOE/CustomizableObjectCompileRunnable.cpp
--- a/Engine/Plugins/Mutable/Source/CustomizableObjectEditor/Private/MuCOE/CustomizableObjectCompileRunnable.cpp
+++ b/Engine/Plugins/Mutable/Source/CustomizableObjectEditor/Private/MuCOE/CustomizableObjectCompileRunnable.cpp
@@ -354,7 +354,7 @@ uint32 FCustomizableObjectCompileRunnable::Run()
 	UE::Mutable::Private::Ptr<UE::Mutable::Private::Compiler> Compiler = new UE::Mutable::Private::Compiler(CompilerOptions, WaitCallback);
 
 	UE_LOG(LogMutable, Verbose, TEXT("PROFILE: [ %16.8f ] FCustomizableObjectCompileRunnable Compile start."), FPlatformTime::Seconds());
-	Model = Compiler->Compile(MutableRoot);
+	Model = Compiler->Compile(MutableRoot, ExternalOperations);
 
 	// Dump all the log messages from the compiler
 	TSharedPtr<const UE::Mutable::Private::FErrorLog> pLog = Compiler->GetLog();
diff --git a/Engine/Plugins/Mutable/Source/CustomizableObjectEditor/Private/MuCOE/CustomizableObjectCompileRunnable.h b/Engine/Plugins/Mutable/Source/CustomizableObjectEditor/Private/MuCOE/CustomizableObjectCompileRunnable.h
--- a/Engine/Plugins/Mutable/Source/CustomizableObjectEditor/Private/MuCOE/CustomizableObjectCompileRunnable.h
+++ b/Engine/Plugins/Mutable/Source/CustomizableObjectEditor/Private/MuCOE/CustomizableObjectCompileRunnable.h
@@ -69,6 +69,8 @@ public:
 
 	TSharedPtr<UE::Mutable::Private::FModel, ESPMode::ThreadSafe> Model;
 
+	TMap<UE::Mutable::Private::OP::ADDRESS, FInstancedStruct> ExternalOperations;
+	
 	FCompilationOptions Options;
 
 	TWeakPtr<FCustomizableObjectCompiler> WeakCompiler;
diff --git a/Engine/Plugins/Mutable/Source/CustomizableObjectEditor/Private/MuCOE/CustomizableObjectCompiler.cpp b/Engine/Plugins/Mutable/Source/CustomizableObjectEditor/Private/MuCOE/CustomizableObjectCompiler.cpp
--- a/Engine/Plugins/Mutable/Source/CustomizableObjectEditor/Private/MuCOE/CustomizableObjectCompiler.cpp
+++ b/Engine/Plugins/Mutable/Source/CustomizableObjectEditor/Private/MuCOE/CustomizableObjectCompiler.cpp
@@ -36,6 +36,7 @@
 #include "DerivedDataCache.h"
 #include "DerivedDataCacheKey.h"
 #include "DerivedDataRequestOwner.h"
+#include "MuR/External/Operation.h"
 
 class UTexture2D;
 
@@ -1708,6 +1709,14 @@ void FCustomizableObjectCompiler::FinishCompilationTask()
 	{
 		TObjectPtr<UModelResources> ModelResources = PlatformData->ModelResources.Get();
 
+		ModelResources->ExternalOperations = MoveTemp(CompileTask->ExternalOperations);
+
+		for (TPair<UE::Mutable::Private::OP::ADDRESS, FInstancedStruct> ExternalOperation : ModelResources->ExternalOperations)
+		{
+			const uint32 Version = ExternalOperation.Value.GetPtr<UE::Mutable::FExternalOperation>()->GetVersion();
+			ModelResources->ExternalOperationVersions.Add(ExternalOperation.Value.GetScriptStruct(), Version);
+		}
+		
 		ModelResources->Skeletons = MoveTemp(CompilationContext->ReferencedSkeletons);
 
 		ModelResources->PhysicsAssets = MoveTemp(CompilationContext->PhysicsAssets);
diff --git a/Engine/Plugins/Mutable/Source/CustomizableObjectEditor/Private/MuCOE/CustomizableObjectEditorModule.cpp b/Engine/Plugins/Mutable/Source/CustomizableObjectEditor/Private/MuCOE/CustomizableObjectEditorModule.cpp
--- a/Engine/Plugins/Mutable/Source/CustomizableObjectEditor/Private/MuCOE/CustomizableObjectEditorModule.cpp
+++ b/Engine/Plugins/Mutable/Source/CustomizableObjectEditor/Private/MuCOE/CustomizableObjectEditorModule.cpp
@@ -110,6 +110,9 @@
 #include "MuCOE/CustomizableObjectGraph.h"
 #include "Nodes/CustomizableObjectNodeComponentMeshDetails.h"
 #include "MuCOE/GenerateMutableSource/GenerateMutableSourceTable.h"
+#include "MuR/External/Operation.h"
+#include "Nodes/CONodeExternal.h"
+#include "Nodes/CONodeExternalDetails.h"
 #include "Nodes/CustomizableObjectNodeTextureParameter.h"
 #include "PhysicsEngine/PhysicsAsset.h"
 
@@ -277,6 +280,7 @@ void FCustomizableObjectEditorModule::StartupModule()
 	RegisterCustomDetails(PropertyModule, UCustomizableObjectNodeMacroInstance::StaticClass(), FOnGetDetailCustomizationInstance::CreateStatic(&FCustomizableObjectNodeMacroInstanceDetails::MakeInstance));
 	RegisterCustomDetails(PropertyModule, UCustomizableObjectNodeParameter::StaticClass(), FOnGetDetailCustomizationInstance::CreateStatic(&FCustomizableObjectNodeParameterDetails::MakeInstance));
 	RegisterCustomDetails(PropertyModule, UCONodeMaterialBreak::StaticClass(), FOnGetDetailCustomizationInstance::CreateStatic(&FCustomizableObjectNodeDetails::MakeInstance));
+	RegisterCustomDetails(PropertyModule, UCONodeExternal::StaticClass(), FOnGetDetailCustomizationInstance::CreateStatic(&FCONodeExternalDetails::MakeInstance));
 
 	// Other Objects
 	RegisterCustomDetails(PropertyModule, UCustomizableObject::StaticClass(), FOnGetDetailCustomizationInstance::CreateStatic(&FCustomizableObjectDetails::MakeInstance));
@@ -615,6 +619,7 @@ struct FCompilationOutOfDateContext
 	TArray<FName> AddedPackages;
 	TArray<FName> RemovedPackages;
 
+	bool bExternalOperationVersionDiff = false;
 	bool bVersionDiff = false;
 	
 	ICustomizableObjectEditorModule::IsCompilationOutOfDateCallback Callback;
@@ -709,6 +714,25 @@ void FCustomizableObjectEditorModule::IsCompilationOutOfDate(const UCustomizable
 		Context->bVersionDiff = ModelResources->ReleaseVersion != VersionBridge->GetCurrentVersionAsString();
 	}
 
+	for (TPair<TObjectPtr<const UScriptStruct>, unsigned> Pair : Object.GetPrivate()->GetModelResources()->ExternalOperationVersions)
+	{
+		FInstancedStruct OperationInstancedStruct;
+		OperationInstancedStruct.InitializeAs(Pair.Key);
+
+		const UE::Mutable::FExternalOperation* Operation = OperationInstancedStruct.GetPtr<UE::Mutable::FExternalOperation>();
+		if (!Operation)
+		{
+			Context->bExternalOperationVersionDiff = true;
+			break;
+		}
+		
+		if (Operation->GetVersion() != Pair.Value)
+		{
+			Context->bExternalOperationVersionDiff = true;
+			break;
+		}
+	}
+
 	check(MaxTime == TNumericLimits<float>::Max()  || bSkipIndirectReferences); // If async, bSkipIndirectReferences must be false since it is very expensive and can not be split in subtasks.
 
 	// Check that we have the exact same set of participating object as before. This can change due to indirect references and versioning.
diff --git a/Engine/Plugins/Mutable/Source/CustomizableObjectEditor/Private/MuCOE/EdGraphSchema_CustomizableObject.cpp b/Engine/Plugins/Mutable/Source/CustomizableObjectEditor/Private/MuCOE/EdGraphSchema_CustomizableObject.cpp
--- a/Engine/Plugins/Mutable/Source/CustomizableObjectEditor/Private/MuCOE/EdGraphSchema_CustomizableObject.cpp
+++ b/Engine/Plugins/Mutable/Source/CustomizableObjectEditor/Private/MuCOE/EdGraphSchema_CustomizableObject.cpp
@@ -102,10 +102,18 @@
 #include "Toolkits/ToolkitManager.h"
 #include "PropertyEditorModule.h"
 #include "Animation/PoseAsset.h"
+#include "MuCO/ICustomizableObjectModulePrivate.h"
 #include "MuCOE/GraphTraversal.h"
 #include "MuCOE/Nodes/CustomizableObjectNodeComponentMeshAddTo.h"
 #include "MuCOE/Nodes/CustomizableObjectNodeTransformConstant.h"
 #include "MuCOE/Nodes/CustomizableObjectNodeTransformParameter.h"
+#include "MuR/External/FloatAdapter.h"
+#include "MuR/External/MaterialAdapter.h"
+#include "MuR/External/MeshAdapter.h"
+#include "MuR/External/Operation.h"
+#include "MuR/External/TextureAdapter.h"
+#include "MuR/External/VectorAdapter.h"
+#include "Nodes/CONodeExternal.h"
 #include "Nodes/CONodeSchema.h"
 
 #include UE_INLINE_GENERATED_CPP_BY_NAME(EdGraphSchema_CustomizableObject)
@@ -116,9 +124,10 @@ class IToolkit;
 
 #define LOCTEXT_NAMESPACE "CustomizableObjectSchema"
 
-
 #define SNAP_GRID (16) // @todo ensure this is the same as SNodePanel::GetSnapGridSize()
 
+using namespace UE::Mutable;
+
 namespace UE::Mutable::Private
 {
 	// Maximum distance a drag can be off a node edge to require 'push off' from node
@@ -643,9 +652,26 @@ void UEdGraphSchema_CustomizableObject::GetGraphContextActions(FGraphContextMenu
 
 		AddNewNodeCategoryActionsFiltered(MaterialTemplateNodes, ContextMenuBuilder, NC_Material.ToString(), 2, Filter);
 	}
+	
+	// Extensions
+	{
+		for (const TStrongObjectPtr<const UScriptStruct>& OperationType : ICustomizableObjectModulePrivate::Get().GetRegisteredExternalOperations())
+		{
+			FInstancedStruct OperationInstancedStruct;
+			OperationInstancedStruct.InitializeAs(OperationType.Get());
 
+			if (const FExternalOperation* Operation = OperationInstancedStruct.GetPtr<FExternalOperation>())
+			{
+				UCONodeExternal* Node = ContextMenuBuilder.CreateTemplateNode<UCONodeExternal>();
+				Node->OperationInstancedStruct = OperationInstancedStruct;
+
+				AddNewNodeActionFiltered(Node, ContextMenuBuilder, GetPinCategoryName(Operation->GetOutput().Value).ToString(), GeneralGrouping, Filter);
+			}
+		}
+	}
+	
+	// External Pin Nodes
 	{
-		// External Pin Nodes
 		TArray<FName> PinTypes({ PC_MeshSection, PC_Modifier, PC_Mesh, PC_Texture, PC_Projector, PC_GroupProjector, PC_Color, PC_Float, PC_Bool, PC_Enum, PC_Transform, PC_Stack, PC_PassthroughTexture, PC_Material, PC_PoseAsset, PC_Component });
 
 		// Add pin types from extensions
@@ -654,6 +680,29 @@ void UEdGraphSchema_CustomizableObject::GetGraphContextActions(FGraphContextMenu
 			PinTypes.AddUnique(PinType.PinType.Name);
 		}
 
+		for (const TStrongObjectPtr<const UScriptStruct>& OperationType : ICustomizableObjectModulePrivate::Get().GetRegisteredExternalOperations())
+		{
+			FInstancedStruct OperationInstancedStruct;
+			OperationInstancedStruct.InitializeAs(OperationType.Get());
+			
+			if (const UE::Mutable::FExternalOperation* Operation = OperationInstancedStruct.GetPtr<UE::Mutable::FExternalOperation>())
+			{
+				// Inputs
+				for (TPair<FName, const UScriptStruct*> Pair : Operation->GetInputs())
+				{
+					const FName Category = UEdGraphSchema_CustomizableObject::GetPinCategoryName(Pair.Value);
+					PinTypes.AddUnique(Category);
+				}
+
+				// Output
+				{
+					TPair<FName, const UScriptStruct*> Pair = Operation->GetOutput();
+					const FName Category = UEdGraphSchema_CustomizableObject::GetPinCategoryName(Pair.Value);
+					PinTypes.AddUnique(Category);
+				}
+			}
+		}
+
 		for (const FName& PinCategory : PinTypes)
 		{
 			UCustomizableObjectNodeExternalPin* CustomizableObjectNodeExternalPin = ContextMenuBuilder.CreateTemplateNode<UCustomizableObjectNodeExternalPin>();
@@ -858,7 +907,8 @@ FLinearColor UEdGraphSchema_CustomizableObject::GetPinTypeColor(const FName& Typ
 		}
 	}
 
-	return FLinearColor(0.750000f, 0.600000f, 0.400000f, 1.000000f);
+	const uint32 Hash = GetTypeHash(TypeString);
+	return FLinearColor::MakeRandomSeededColor(Hash);
 }
 
 
@@ -1480,6 +1530,35 @@ FName UEdGraphSchema_CustomizableObject::GetPinCategoryName(const FName& PinCate
 }
 
 
+FName UEdGraphSchema_CustomizableObject::GetPinCategoryName(const UScriptStruct* Type)
+{
+	if (Type == FMeshAdapter::StaticStruct())
+	{
+		return PC_Mesh;
+	}
+	else if (Type == FTextureAdapter::StaticStruct())
+	{
+		return PC_Texture;
+	}
+	else if (Type == FFloatAdapter::StaticStruct())
+	{
+		return PC_Float;
+	}
+	else if (Type == FVectorAdapter::StaticStruct())
+	{
+		return PC_Color;
+	}
+	else if (Type == FMaterialAdapter::StaticStruct())
+	{
+		return PC_Material;
+	}
+	else
+	{
+		return Type->GetFName();
+	}
+}
+
+
 FText UEdGraphSchema_CustomizableObject::GetPinCategoryFriendlyName(const FName& PinCategory)
 {
 	if (PinCategory == PC_Object)
diff --git a/Engine/Plugins/Mutable/Source/CustomizableObjectEditor/Private/MuCOE/GenerateMutableSource/GenerateMutableSource.cpp b/Engine/Plugins/Mutable/Source/CustomizableObjectEditor/Private/MuCOE/GenerateMutableSource/GenerateMutableSource.cpp
--- a/Engine/Plugins/Mutable/Source/CustomizableObjectEditor/Private/MuCOE/GenerateMutableSource/GenerateMutableSource.cpp
+++ b/Engine/Plugins/Mutable/Source/CustomizableObjectEditor/Private/MuCOE/GenerateMutableSource/GenerateMutableSource.cpp
@@ -418,6 +418,18 @@ uint32 GetTypeHash(const FGeneratedKey& Key)
 	return Hash;
 }
 
+uint32 GetTypeHash(const FGeneratedSourceExtensionKey& Key)
+{
+	uint32 Hash = GetTypeHash(Key.Pin);
+
+	Hash = HashCombine(Hash, GetTypeHash(Key.Options.MeshOptions));
+	Hash = HashCombine(Hash, GetTypeHash(Key.Options.bLinkedToExtendMaterial));
+	Hash = HashCombine(Hash, GetTypeHash(Key.Options.bOnlyConnectedLOD));
+	Hash = HashCombine(Hash, GetTypeHash(Key.Options.ReferenceTextureSize));
+	
+	return Hash;
+}
+
 
 FGraphCycleKey::FGraphCycleKey(const UEdGraphPin& Pin, const FString& Id, const UCustomizableObjectNodeMacroInstance* MacroContext) :
 	Pin(Pin),
@@ -814,7 +826,7 @@ uint32 AddUniqueSurfaceMetadata(const FMutableSourceSurfaceMetadata& SourceMetad
 		SurfaceMetadata.bCastShadow = MeshSection->bCastShadow;
 	}
 
-	uint32 SurfaceMetadataUniqueHash = Private::GenerateUniquePersistentHash(
+	uint32 SurfaceMetadataUniqueHash = GenerateMutableSourcePrivate::GenerateUniquePersistentHash(
 			SurfaceMetadata, CompilationContext.SurfaceMetadata, HashSurfaceMetadataFunc, CompareSurfaceMetadataFunc);
 
 	if (SurfaceMetadataUniqueHash != 0)
diff --git a/Engine/Plugins/Mutable/Source/CustomizableObjectEditor/Private/MuCOE/GenerateMutableSource/GenerateMutableSource.h b/Engine/Plugins/Mutable/Source/CustomizableObjectEditor/Private/MuCOE/GenerateMutableSource/GenerateMutableSource.h
--- a/Engine/Plugins/Mutable/Source/CustomizableObjectEditor/Private/MuCOE/GenerateMutableSource/GenerateMutableSource.h
+++ b/Engine/Plugins/Mutable/Source/CustomizableObjectEditor/Private/MuCOE/GenerateMutableSource/GenerateMutableSource.h
@@ -29,6 +29,7 @@
 #include "MuT/NodeSurfaceNew.h"
 #include "MuT/NodeMaterialConstant.h"
 #include "MuT/NodeMaterialParameter.h"
+#include "MuT/NodeExternal.h"
 #include "MuT/Table.h"
 #include "MuR/Skeleton.h"
 #include "MuR/Mesh.h"
@@ -179,14 +180,26 @@ public:
 uint32 GetTypeHash(const FGeneratedKey& Key);
 
 
-struct FGeneratedImageKey
+struct FGeneratedSourceExtensionKey
 {
-	FGeneratedImageKey(const UEdGraphPin* InPin)
+	const UEdGraphPin* Pin = nullptr;
+	FSourceExtensionOptions Options;
+
+	bool operator==(const FGeneratedSourceExtensionKey& Other) const = default;
+};
+
+
+uint32 GetTypeHash(const FGeneratedSourceExtensionKey& Key);
+
+
+struct FImageKey
+{
+	FImageKey(const UEdGraphPin* InPin)
 	{
 		Pin = InPin;
 	}
 
-	bool operator==(const FGeneratedImageKey& Other) const
+	bool operator==(const FImageKey& Other) const
 	{
 		return Pin == Other.Pin;
 	}
@@ -251,7 +264,13 @@ struct FGeneratedData
 };
 
 
-inline uint32 GetTypeHash(const FGeneratedImageKey& Key)
+struct FGeneratedSourceExtensionData
+{
+	UE::Mutable::Private::Ptr<UE::Mutable::Private::NodeExternal> Node;
+};
+
+
+inline uint32 GetTypeHash(const FImageKey& Key)
 {
 	uint32 GuidHash = GetTypeHash(Key.Pin->PinId);
 
@@ -598,6 +617,7 @@ struct FMutableGraphGenerationContext
 	
 	// Cache of generated pins per LOD
 	TMap<FGeneratedKey, FGeneratedData> Generated;
+	TMap<FGeneratedSourceExtensionKey, FGeneratedSourceExtensionData> GeneratedSourceExtensions;
 
 	/** Set of all generated nodes. */
 	TSet<UCustomizableObjectNode*> GeneratedNodes;
@@ -649,7 +669,7 @@ struct FMutableGraphGenerationContext
 
 	// Cache of generated images, because sometimes they are reused by LOD, we use this as a second
 	// level cache
-	TMap<FGeneratedImageKey, UE::Mutable::Private::NodeImagePtr> GeneratedImages;
+	TMap<FImageKey, UE::Mutable::Private::NodeImagePtr> GeneratedImages;
 
 	/** Data stored per-generated passthrough texture. */
 	struct FGeneratedReferencedTexture
@@ -916,7 +936,7 @@ struct FMutableGraphGenerationContext
 };
 
 
-namespace Private
+namespace GenerateMutableSourcePrivate
 {
 	template<class HashableType, class HashDataSetType, class HashFuncType, class CompareFuncType>
 	uint32 GenerateUniquePersistentHash(const HashableType& HashableData, const HashDataSetType& HashDataSet, HashFuncType&& HashFunc, CompareFuncType&& CompareFunc)
@@ -969,7 +989,7 @@ namespace Private
 
 		return UniqueHash;
 	}
-} //Private
+}
 
 //
 UE::Mutable::Private::Ptr<UE::Mutable::Private::NodeObject> GenerateMutableSource(const class UEdGraphPin* Pin, FMutableGraphGenerationContext& GenerationContext);
diff --git a/Engine/Plugins/Mutable/Source/CustomizableObjectEditor/Private/MuCOE/GenerateMutableSource/GenerateMutableSourceColor.cpp b/Engine/Plugins/Mutable/Source/CustomizableObjectEditor/Private/MuCOE/GenerateMutableSource/GenerateMutableSourceColor.cpp
--- a/Engine/Plugins/Mutable/Source/CustomizableObjectEditor/Private/MuCOE/GenerateMutableSource/GenerateMutableSourceColor.cpp
+++ b/Engine/Plugins/Mutable/Source/CustomizableObjectEditor/Private/MuCOE/GenerateMutableSource/GenerateMutableSourceColor.cpp
@@ -2,6 +2,7 @@
 
 #include "MuCOE/GenerateMutableSource/GenerateMutableSourceColor.h"
 
+#include "GenerateMutableSourceExternal.h"
 #include "MuCOE/CustomizableObjectCompiler.h"
 #include "MuCOE/CustomizableObjectGraph.h"
 #include "MuCOE/EdGraphSchema_CustomizableObject.h"
@@ -12,6 +13,7 @@
 #include "MuCOE/GenerateMutableSource/GenerateMutableSourceTable.h"
 #include "MuCOE/GenerateMutableSource/GenerateMutableSource.h"
 #include "MuCOE/GraphTraversal.h"
+#include "MuCOE/Nodes/CONodeExternal.h"
 #include "MuCOE/Nodes/CustomizableObjectNodeColorArithmeticOp.h"
 #include "MuCOE/Nodes/CustomizableObjectNodeColorConstant.h"
 #include "MuCOE/Nodes/CustomizableObjectNodeColorFromFloats.h"
@@ -33,6 +35,7 @@
 #include "MuT/NodeColourTable.h"
 #include "MuT/NodeColourVariation.h"
 #include "MuT/NodeColorToSRGB.h"
+#include "MuT/NodeColourExternal.h"
 #include "MuT/NodeColourMaterialBreak.h"
 #include "MuT/NodeImageBinarise.h"
 
@@ -359,6 +362,17 @@ UE::Mutable::Private::Ptr<UE::Mutable::Private::NodeColour> GenerateMutableSourc
 			GenerationContext.Log(LOCTEXT("ColorTableError", "Couldn't find the data table of the node."), Node);
 		}
 	}
+	
+	else if (Cast<UCONodeExternal>(Node))
+	{
+		UE::Mutable::Private::Ptr<UE::Mutable::Private::NodeColourExternal> MeshExtension = new UE::Mutable::Private::NodeColourExternal();
+
+		FSourceExtensionOptions Options;
+		MeshExtension->Node = GenerateMutableSourceExtension(Pin, GenerationContext, Options);
+		
+		Result = MeshExtension;
+
+	}
 
 	else if (const UCustomizableObjectNodeMacroInstance* TypedNodeMacro = Cast<UCustomizableObjectNodeMacroInstance>(Node))
 	{
diff --git a/Engine/Plugins/Mutable/Source/CustomizableObjectEditor/Private/MuCOE/GenerateMutableSource/GenerateMutableSourceExternal.cpp b/Engine/Plugins/Mutable/Source/CustomizableObjectEditor/Private/MuCOE/GenerateMutableSource/GenerateMutableSourceExternal.cpp
new file mode 100644
--- /dev/null
+++ b/Engine/Plugins/Mutable/Source/CustomizableObjectEditor/Private/MuCOE/GenerateMutableSource/GenerateMutableSourceExternal.cpp
@@ -0,0 +1,165 @@
+﻿// Copyright Epic Games, Inc. All Rights Reserved.
+
+
+#include "GenerateMutableSourceExternal.h"
+
+#include "GenerateMutableSourceColor.h"
+#include "GenerateMutableSourceFloat.h"
+#include "GenerateMutableSourceImage.h"
+#include "GenerateMutableSourceMacro.h"
+#include "GenerateMutableSourceMesh.h"
+#include "MuCOE/GraphTraversal.h"
+#include "MuCOE/Nodes/CONodeExternal.h"
+#include "MuR/External/FloatAdapter.h"
+#include "MuR/External/MaterialAdapter.h"
+#include "MuR/External/MeshAdapter.h"
+#include "MuR/External/Operation.h"
+#include "MuR/External/TextureAdapter.h"
+#include "MuR/External/VectorAdapter.h"
+
+#define LOCTEXT_NAMESPACE "CustomizableObjectEditor"
+
+using namespace UE::Mutable;
+using namespace UE::Mutable::Private;
+
+
+Ptr<NodeExternal> GenerateMutableSourceExtension(const UEdGraphPin* Pin, FMutableGraphGenerationContext& GenerationContext, const FSourceExtensionOptions& Options)
+{
+	check(Pin)
+	RETURN_ON_CYCLE(*Pin, GenerationContext)
+
+	CheckNumOutputs(*Pin, GenerationContext);
+	
+	UCustomizableObjectNode* InNode = CastChecked<UCustomizableObjectNode>(Pin->GetOwningNode());
+
+	FGeneratedSourceExtensionKey Key;
+	Key.Pin = Pin;
+	Key.Options = Options;
+	
+	if (const FGeneratedSourceExtensionData* Generated = GenerationContext.GeneratedSourceExtensions.Find(Key))
+	{
+		return Generated->Node;
+	}
+
+	if (InNode->IsNodeOutDatedAndNeedsRefresh())
+	{
+		InNode->SetRefreshNodeWarning();
+	}
+	
+	UCustomizableObjectNode* EditorNode = CastChecked<UCustomizableObjectNode>(Pin->GetOwningNode());
+
+	Ptr<NodeExternal> Result;
+	bool bCacheNode = true;
+	
+	if (const UCONodeExternal* TypedNodeExtension = Cast<UCONodeExternal>(EditorNode))
+	{
+		[&]()
+		{
+			const UE::Mutable::FExternalOperation* Operation = TypedNodeExtension->OperationInstancedStruct.GetPtr<UE::Mutable::FExternalOperation>();
+			if (!Operation)
+			{
+				FText Msg = FText::Format(LOCTEXT("ExtensionNotLoaded", "Node is using an external operation which has not been loaded: {0}"), FText::FromName(TypedNodeExtension->CachedOperationName));
+				GenerationContext.Log(Msg, TypedNodeExtension, EMessageSeverity::Error);
+				return;
+			}
+			
+			Ptr<NodeExternal> ExtensionNode = new NodeExternal();
+            				
+            ExtensionNode->OperationInstancedStruct = TypedNodeExtension->OperationInstancedStruct;
+
+            TArray<TPair<FName, const UScriptStruct*>> Inputs = Operation->GetInputs();
+            for (const TPair<FName, const UScriptStruct*>& Pair : Inputs)
+            {
+            	const FName& InputName = Pair.Key;
+            
+            	const FEdGraphPinReference* InputResult = TypedNodeExtension->InputPins.FindByPredicate([&](const FEdGraphPinReference& Element)
+            	{
+            		return TypedNodeExtension->GetPinData<UCONodeExtensionPinData>(*Element.Get()).Name == InputName;
+            	});
+        
+            	if (!InputResult)
+            	{
+            		FText Msg = FText::Format(LOCTEXT("ExtensionInputNotFound", "Can not find input in external operation: {0}"), FText::FromName(InputName));
+            		GenerationContext.Log(Msg, TypedNodeExtension, EMessageSeverity::Error);
+            		continue;
+            	}
+
+            	const UEdGraphPin* ConnectedPin = FollowInputPin(*InputResult->Get());
+            	if (!ConnectedPin)
+            	{
+            		FText Msg = FText::Format(LOCTEXT("ExtensionInputNotConnected", "External operation input not connected: {0}"), FText::FromName(InputName));
+            		GenerationContext.Log(Msg, TypedNodeExtension, EMessageSeverity::Error);
+            		continue;
+            	}
+
+            	const UScriptStruct* InputType = Pair.Value;
+            
+            	if (InputType == FMeshAdapter::StaticStruct())
+            	{
+            		Ptr<NodeMesh> SourceNode = GenerateMutableSourceMesh(ConnectedPin, GenerationContext, Options.MeshOptions, Options.bLinkedToExtendMaterial, Options.bOnlyConnectedLOD);
+            		ExtensionNode->Inputs.Add(SourceNode);
+            	}
+            	else if (InputType == FTextureAdapter::StaticStruct())
+            	{
+            		Ptr<NodeImage> SourceNode = GenerateMutableSourceImage(ConnectedPin, GenerationContext, Options.ReferenceTextureSize);
+            		ExtensionNode->Inputs.Add(SourceNode);
+            	}
+            	else if (InputType == FFloatAdapter::StaticStruct())
+            	{
+            		Ptr<NodeScalar> SourceNode = GenerateMutableSourceFloat(ConnectedPin, GenerationContext);
+            		ExtensionNode->Inputs.Add(SourceNode);
+            	}
+            	else if (InputType == FVectorAdapter::StaticStruct())
+            	{
+            		Ptr<NodeColour> SourceNode = GenerateMutableSourceColor(ConnectedPin, GenerationContext);
+            		ExtensionNode->Inputs.Add(SourceNode);
+            	}
+            	else if (InputType == FMaterialAdapter::StaticStruct())
+            	{
+            		Ptr<NodeColour> SourceNode = GenerateMutableSourceColor(ConnectedPin, GenerationContext);
+            		ExtensionNode->Inputs.Add(SourceNode);
+            	}
+            	else
+            	{
+            		Ptr<Node> SourceNode = GenerateMutableSourceExtension(ConnectedPin, GenerationContext, Options);
+            		ExtensionNode->Inputs.Add(SourceNode);
+            	}	
+            }
+
+			Result = ExtensionNode;
+		}();		
+	}
+
+	else if (const UCustomizableObjectNodeMacroInstance* TypedNodeMacro = Cast<UCustomizableObjectNodeMacroInstance>(InNode))
+	{
+		bCacheNode = false;
+
+		Result = GenerateMutableSourceMacro<NodeExternal>(*Pin, GenerationContext, GenerateMutableSourceExtension, Options);
+	}
+
+	else if (const UCustomizableObjectNodeTunnel* TypedNodeTunnel = Cast<UCustomizableObjectNodeTunnel>(InNode))
+	{
+		bCacheNode = false;
+
+		Result = GenerateMutableSourceMacro<NodeExternal>(*Pin, GenerationContext, GenerateMutableSourceExtension, Options);
+	}
+	
+	if (Result)
+	{
+		Result->SetMessageContext(InNode);
+	}
+
+	if (bCacheNode)
+	{
+		FGeneratedSourceExtensionData CacheData;
+		CacheData.Node = Result;
+	
+		GenerationContext.GeneratedSourceExtensions.Add(Key, CacheData);
+		GenerationContext.GeneratedNodes.Add(InNode);	
+	}
+
+	return Result;
+}
+
+
+#undef LOCTEXT_NAMESPACE
\ No newline at end of file
diff --git a/Engine/Plugins/Mutable/Source/CustomizableObjectEditor/Private/MuCOE/GenerateMutableSource/GenerateMutableSourceExternal.h b/Engine/Plugins/Mutable/Source/CustomizableObjectEditor/Private/MuCOE/GenerateMutableSource/GenerateMutableSourceExternal.h
new file mode 100644
--- /dev/null
+++ b/Engine/Plugins/Mutable/Source/CustomizableObjectEditor/Private/MuCOE/GenerateMutableSource/GenerateMutableSourceExternal.h
@@ -0,0 +1,13 @@
+﻿// Copyright Epic Games, Inc. All Rights Reserved.
+
+#pragma once
+
+#include "MuT/NodeExternal.h"
+
+class UEdGraphPin;
+struct FSourceExtensionOptions;
+struct FMutableGraphGenerationContext;
+
+
+UE::Mutable::Private::Ptr<UE::Mutable::Private::NodeExternal> GenerateMutableSourceExtension(const UEdGraphPin* Pin, FMutableGraphGenerationContext& GenerationContext, const FSourceExtensionOptions& Options);
+
diff --git a/Engine/Plugins/Mutable/Source/CustomizableObjectEditor/Private/MuCOE/GenerateMutableSource/GenerateMutableSourceFloat.cpp b/Engine/Plugins/Mutable/Source/CustomizableObjectEditor/Private/MuCOE/GenerateMutableSource/GenerateMutableSourceFloat.cpp
--- a/Engine/Plugins/Mutable/Source/CustomizableObjectEditor/Private/MuCOE/GenerateMutableSource/GenerateMutableSourceFloat.cpp
+++ b/Engine/Plugins/Mutable/Source/CustomizableObjectEditor/Private/MuCOE/GenerateMutableSource/GenerateMutableSourceFloat.cpp
@@ -2,6 +2,7 @@
 
 #include "MuCOE/GenerateMutableSource/GenerateMutableSourceFloat.h"
 
+#include "GenerateMutableSourceExternal.h"
 #include "Curves/CurveFloat.h"
 #include "Curves/CurveLinearColor.h"
 #include "Curves/CurveVector.h"
@@ -13,6 +14,7 @@
 #include "MuCOE/GenerateMutableSource/GenerateMutableSourceMaterial.h"
 #include "MuCOE/GenerateMutableSource/GenerateMutableSourceTable.h"
 #include "MuCOE/GraphTraversal.h"
+#include "MuCOE/Nodes/CONodeExternal.h"
 #include "MuCOE/Nodes/CONodeMaterialBreak.h"
 #include "MuCOE/Nodes/CustomizableObjectNodeCurve.h"
 #include "MuCOE/Nodes/CustomizableObjectNodeEnumParameter.h"
@@ -24,12 +26,14 @@
 #include "MuCOE/Nodes/CustomizableObjectNodeMacroInstance.h"
 #include "MuCOE/Nodes/CustomizableObjectNodeTable.h"
 #include "MuCOE/Nodes/CustomizableObjectNodeTunnel.h"
+#include "MuT/NodeMeshExternal.h"
 #include "MuT/NodeScalarConstant.h"
 #include "MuT/NodeScalarCurve.h"
 #include "MuT/NodeScalarSwitch.h"
 #include "MuT/NodeScalarTable.h"
 #include "MuT/NodeScalarVariation.h"
 #include "MuT/NodeScalarArithmeticOperation.h"
+#include "MuT/NodeScalarExternal.h"
 #include "MuT/NodeScalarMaterialBreak.h"
 
 #define LOCTEXT_NAMESPACE "CustomizableObjectEditor"
@@ -386,6 +390,17 @@ UE::Mutable::Private::Ptr<UE::Mutable::Private::NodeScalar> GenerateMutableSourc
 		}
 	}
 
+	else if (Cast<UCONodeExternal>(Node))
+	{
+		UE::Mutable::Private::Ptr<UE::Mutable::Private::NodeScalarExternal> MeshExtension = new UE::Mutable::Private::NodeScalarExternal();
+
+		FSourceExtensionOptions Options;
+		MeshExtension->Node = GenerateMutableSourceExtension(Pin, GenerationContext, Options);
+		
+		Result = MeshExtension;
+
+	}
+
 	else if (const UCustomizableObjectNodeMacroInstance* TypedNodeMacro = Cast<UCustomizableObjectNodeMacroInstance>(Node))
 	{
 		bCacheNode = false;
diff --git a/Engine/Plugins/Mutable/Source/CustomizableObjectEditor/Private/MuCOE/GenerateMutableSource/GenerateMutableSourceImage.cpp b/Engine/Plugins/Mutable/Source/CustomizableObjectEditor/Private/MuCOE/GenerateMutableSource/GenerateMutableSourceImage.cpp
--- a/Engine/Plugins/Mutable/Source/CustomizableObjectEditor/Private/MuCOE/GenerateMutableSource/GenerateMutableSourceImage.cpp
+++ b/Engine/Plugins/Mutable/Source/CustomizableObjectEditor/Private/MuCOE/GenerateMutableSource/GenerateMutableSourceImage.cpp
@@ -2,11 +2,13 @@
 
 #include "MuCOE/GenerateMutableSource/GenerateMutableSourceImage.h"
 
+#include "GenerateMutableSourceExternal.h"
 #include "Engine/TextureLODSettings.h"
 
 #include "MuCOE/CustomizableObjectCompiler.h"
 #include "MuCOE/CustomizableObjectGraph.h"
 #include "MuCOE/EdGraphSchema_CustomizableObject.h"
+#include "MuCOE/GenerateMutableSource/GenerateMutableSource.h"
 #include "MuCOE/GenerateMutableSource/GenerateMutableSourceColor.h"
 #include "MuCOE/GenerateMutableSource/GenerateMutableSourceFloat.h"
 #include "MuCOE/GenerateMutableSource/GenerateMutableSourceMaterial.h"
@@ -40,6 +42,7 @@
 #include "MuCOE/Nodes/CustomizableObjectNodeTunnel.h"
 #include "MuCOE/UnrealEditorPortabilityHelpers.h"
 #include "MuCO/UnrealToMutableTextureConversionUtils.h"
+#include "MuCOE/Nodes/CONodeExternal.h"
 #include "MuT/NodeImageBinarise.h"
 #include "MuT/NodeImageColourMap.h"
 #include "MuT/NodeImageFormat.h"
@@ -59,6 +62,7 @@
 #include "MuT/NodeImageTransform.h"
 #include "MuT/NodeImageVariation.h"
 #include "MuT/NodeImageConstant.h"
+#include "MuT/NodeImageExternal.h"
 #include "MuT/NodeScalarConstant.h"
 
 #define LOCTEXT_NAMESPACE "CustomizableObjectEditor"
@@ -118,7 +122,7 @@ UE::Mutable::Private::Ptr<UE::Mutable::Private::NodeImage> GenerateMutableSource
 		if (BaseTexture)
 		{
 			// Check the specific image cache
-			FGeneratedImageKey ImageKey = FGeneratedImageKey(Pin);
+			FImageKey ImageKey = FImageKey(Pin);
 			UE::Mutable::Private::Ptr<UE::Mutable::Private::NodeImage> ImageNode;
 			
 			const uint32 MipsToSkip = ComputeLODBiasForTexture(GenerationContext, *BaseTexture, nullptr, ReferenceTextureSize);
@@ -1147,6 +1151,17 @@ UE::Mutable::Private::Ptr<UE::Mutable::Private::NodeImage> GenerateMutableSource
 		}
 	}
 
+	else if (Cast<UCONodeExternal>(Node))
+	{
+		UE::Mutable::Private::Ptr<UE::Mutable::Private::NodeImageExternal> MeshExtension = new UE::Mutable::Private::NodeImageExternal();
+
+		FSourceExtensionOptions Options;
+		Options.ReferenceTextureSize = ReferenceTextureSize;
+		MeshExtension->Node = GenerateMutableSourceExtension(Pin, GenerationContext, Options);
+		
+		Result = MeshExtension;
+	}
+	
 	else if (const UCustomizableObjectNodeMacroInstance* TypedNodeMacro = Cast<UCustomizableObjectNodeMacroInstance>(Node))
 	{
 		// Can't use the GenerateMutableSourceMacro function here because GenerateMutableSourceImage needs the extra parameter ReferenceTextureSize 
diff --git a/Engine/Plugins/Mutable/Source/CustomizableObjectEditor/Private/MuCOE/GenerateMutableSource/GenerateMutableSourceImage.h b/Engine/Plugins/Mutable/Source/CustomizableObjectEditor/Private/MuCOE/GenerateMutableSource/GenerateMutableSourceImage.h
--- a/Engine/Plugins/Mutable/Source/CustomizableObjectEditor/Private/MuCOE/GenerateMutableSource/GenerateMutableSourceImage.h
+++ b/Engine/Plugins/Mutable/Source/CustomizableObjectEditor/Private/MuCOE/GenerateMutableSource/GenerateMutableSourceImage.h
@@ -14,14 +14,11 @@ struct FMutableGraphGenerationContext;
 
 namespace UE::Mutable::Private
 {
-	class FImage;
 	class NodeImage;
 	class NodeImageParameter;
 }
 
 
-//TSharedPtr<UE::Mutable::Private::FImage> ConvertTextureUnrealToMutable(UTexture2D* Texture, const UCustomizableObjectNode* Node, FCustomizableObjectCompiler* Compiler, bool bIsNormalComposite);
-
 UE::Mutable::Private::Ptr<UE::Mutable::Private::NodeImage> ResizeTextureByNumMips(const UE::Mutable::Private::Ptr<UE::Mutable::Private::NodeImage>& Image, int32 MipsToSkip);
 
 /** Convert a CustomizableObject Source Graph from an Image pin into a mutable source graph. */
diff --git a/Engine/Plugins/Mutable/Source/CustomizableObjectEditor/Private/MuCOE/GenerateMutableSource/GenerateMutableSourceMacro.h b/Engine/Plugins/Mutable/Source/CustomizableObjectEditor/Private/MuCOE/GenerateMutableSource/GenerateMutableSourceMacro.h
--- a/Engine/Plugins/Mutable/Source/CustomizableObjectEditor/Private/MuCOE/GenerateMutableSource/GenerateMutableSourceMacro.h
+++ b/Engine/Plugins/Mutable/Source/CustomizableObjectEditor/Private/MuCOE/GenerateMutableSource/GenerateMutableSourceMacro.h
@@ -9,8 +9,9 @@
 
 #define LOCTEXT_NAMESPACE "CustomizableObjectEditor"
 
-template<class T>
-UE::Mutable::Private::Ptr<T> GenerateMutableSourceMacro(const UEdGraphPin& Pin, FMutableGraphGenerationContext& GenerationContext, const TFunction<UE::Mutable::Private::Ptr<T>(const UEdGraphPin*, FMutableGraphGenerationContext&)>& GenerationFunction)
+
+template<typename T, typename... ArgTypes>
+UE::Mutable::Private::Ptr<T> GenerateMutableSourceMacro(const UEdGraphPin& Pin, FMutableGraphGenerationContext& GenerationContext, UE::Mutable::Private::Ptr<T> (*GenerationFunction)(const UEdGraphPin*, FMutableGraphGenerationContext&, ArgTypes...), ArgTypes&&... Args)
 {
 	UE::Mutable::Private::Ptr<T> Result;
 	UCustomizableObjectNode* Node = CastChecked<UCustomizableObjectNode>(Pin.GetOwningNode());
@@ -22,7 +23,7 @@ UE::Mutable::Private::Ptr<T> GenerateMutableSourceMacro(const UEdGraphPin& Pin,
 			if (const UEdGraphPin* FollowPin = FollowInputPin(*OutputPin))
 			{
 				GenerationContext.MacroNodesStack.Push(TypedNodeMacro);
-				Result = GenerationFunction(FollowPin, GenerationContext);
+				Result = GenerationFunction(FollowPin, GenerationContext, Forward<ArgTypes>(Args)...);
 				GenerationContext.MacroNodesStack.Pop();
 			}
 			else
@@ -50,7 +51,7 @@ UE::Mutable::Private::Ptr<T> GenerateMutableSourceMacro(const UEdGraphPin& Pin,
 		{
 			if (const UEdGraphPin* FollowPin = FollowInputPin(*InputPin))
 			{
-				Result = GenerationFunction(FollowPin, GenerationContext);
+				Result = GenerationFunction(FollowPin, GenerationContext, Forward<ArgTypes>(Args)...);
 			}
 		}
 		else
diff --git a/Engine/Plugins/Mutable/Source/CustomizableObjectEditor/Private/MuCOE/GenerateMutableSource/GenerateMutableSourceMaterial.cpp b/Engine/Plugins/Mutable/Source/CustomizableObjectEditor/Private/MuCOE/GenerateMutableSource/GenerateMutableSourceMaterial.cpp
--- a/Engine/Plugins/Mutable/Source/CustomizableObjectEditor/Private/MuCOE/GenerateMutableSource/GenerateMutableSourceMaterial.cpp
+++ b/Engine/Plugins/Mutable/Source/CustomizableObjectEditor/Private/MuCOE/GenerateMutableSource/GenerateMutableSourceMaterial.cpp
@@ -2,6 +2,7 @@
 
 #include "MuCOE/GenerateMutableSource/GenerateMutableSourceMaterial.h"
 
+#include "GenerateMutableSourceExternal.h"
 #include "Engine/TextureLODSettings.h"
 #include "Interfaces/ITargetPlatform.h"
 #include "MuCOE/CustomizableObjectCompiler.h"
@@ -33,6 +34,8 @@
 
 #include "Materials/MaterialInterface.h"
 #include "Materials/MaterialInstance.h"
+#include "MuCOE/Nodes/CONodeExternal.h"
+#include "MuT/NodeMaterialExternal.h"
 
 
 #define LOCTEXT_NAMESPACE "CustomizableObjectEditor"
@@ -467,6 +470,17 @@ UE::Mutable::Private::Ptr<UE::Mutable::Private::NodeMaterial> GenerateMutableSou
 		}
 	}
 
+	else if (Cast<UCONodeExternal>(Node))
+	{
+		UE::Mutable::Private::Ptr<UE::Mutable::Private::NodeMaterialExternal> MeshExtension = new UE::Mutable::Private::NodeMaterialExternal();
+
+		FSourceExtensionOptions Options;
+		MeshExtension->Node = GenerateMutableSourceExtension(Pin, GenerationContext, Options);
+		
+		Result = MeshExtension;
+
+	}
+
 	else if (const UCustomizableObjectNodeMacroInstance* TypedNodeMacro = Cast<UCustomizableObjectNodeMacroInstance>(Node))
 	{
 		bCacheNode = false;
diff --git a/Engine/Plugins/Mutable/Source/CustomizableObjectEditor/Private/MuCOE/GenerateMutableSource/GenerateMutableSourceMesh.cpp b/Engine/Plugins/Mutable/Source/CustomizableObjectEditor/Private/MuCOE/GenerateMutableSource/GenerateMutableSourceMesh.cpp
--- a/Engine/Plugins/Mutable/Source/CustomizableObjectEditor/Private/MuCOE/GenerateMutableSource/GenerateMutableSourceMesh.cpp
+++ b/Engine/Plugins/Mutable/Source/CustomizableObjectEditor/Private/MuCOE/GenerateMutableSource/GenerateMutableSourceMesh.cpp
@@ -9,6 +9,8 @@
 #include "AnimGraphNode_RigidBody.h"
 #include "ClothConfigBase.h"
 #include "ClothingAsset.h"
+#include "GenerateMutableSourceExternal.h"
+#include "GenerateMutableSourceImage.h"
 #include "MeshUtilities.h"
 #include "Modules/ModuleManager.h"
 #include "MuCO/CustomizableObjectPrivate.h"
@@ -61,6 +63,13 @@
 #include "AssetRegistry/AssetRegistryModule.h"
 #include "Hash/CityHash.h"
 #include "PackedNormal.h"
+#include "MuCOE/Nodes/CONodeExternal.h"
+#include "MuR/External/MeshAdapter.h"
+#include "MuR/External/Operation.h"
+#include "MuR/External/TextureAdapter.h"
+#include "MuT/NodeExternal.h"
+#include "MuT/NodeExtensionData.h"
+#include "MuT/NodeMeshExternal.h"
 
 #define LOCTEXT_NAMESPACE "CustomizableObjectEditor"
 
@@ -2583,7 +2592,7 @@ UE::Tasks::FTask ConvertSkeletalMeshToMutable(TSharedRef<FMeshConversionContext>
 				return FMemory::Memcmp(&A, &B, sizeof(FMutableMeshMetadata)) == 0;
 			};
 
-		const uint32 MeshMetadataUniqueHash = Private::GenerateUniquePersistentHash(
+		const uint32 MeshMetadataUniqueHash = GenerateMutableSourcePrivate::GenerateUniquePersistentHash(
 			MeshMetadata, Context.MeshMetadata, HashMeshMetadataFunc, CompareMeshMetadataFunc);
 
 		if (MeshMetadataUniqueHash != 0)
@@ -4666,6 +4675,19 @@ UE::Mutable::Private::Ptr<UE::Mutable::Private::NodeMesh> GenerateMutableSourceM
 		}
 	}
 
+	else if (Cast<UCONodeExternal>(Node))
+	{
+		UE::Mutable::Private::Ptr<UE::Mutable::Private::NodeMeshExternal> MeshExtension = new UE::Mutable::Private::NodeMeshExternal();
+
+		FSourceExtensionOptions Options;
+		Options.MeshOptions = BaseMeshData;
+		Options.bLinkedToExtendMaterial = bLinkedToExtendMaterial;
+		Options.bOnlyConnectedLOD = bOnlyConnectedLOD;
+		MeshExtension->Node = GenerateMutableSourceExtension(Pin, GenerationContext, Options);
+		
+		Result = MeshExtension;
+	}
+	
 	else if (const UCustomizableObjectNodeMacroInstance* TypedNodeMacro = Cast<UCustomizableObjectNodeMacroInstance>(Node))
 	{
 		// Can't use the GenerateMutableSourceMacro function here because GenerateMutableSourceMesh needs some extra parameters
diff --git a/Engine/Plugins/Mutable/Source/CustomizableObjectEditor/Private/MuCOE/GraphTraversal.cpp b/Engine/Plugins/Mutable/Source/CustomizableObjectEditor/Private/MuCOE/GraphTraversal.cpp
--- a/Engine/Plugins/Mutable/Source/CustomizableObjectEditor/Private/MuCOE/GraphTraversal.cpp
+++ b/Engine/Plugins/Mutable/Source/CustomizableObjectEditor/Private/MuCOE/GraphTraversal.cpp
@@ -29,6 +29,7 @@
 #include "MuCOE/Nodes/CustomizableObjectNodeAnimationPose.h"
 #include "MuCOE/Nodes/CustomizableObjectNodeObjectGroup.h"
 #include "MuCOE/Nodes/CustomizableObjectNodeReroute.h"
+#include "Nodes/CONodeExternal.h"
 #include "UObject/UObjectIterator.h"
 
 
@@ -749,6 +750,14 @@ const UEdGraphPin* FindMeshBaseSource(const UEdGraphPin& Pin, const bool bOnlyLo
 		}
 	}
 
+	else if (Cast<UCONodeExternal>(Node))
+	{
+		if (!bOnlyLookForStaticMesh)
+		{
+			return &Pin;
+		}
+	}
+
 	else
 	{
 		unimplemented(); // Case missing.
diff --git a/Engine/Plugins/Mutable/Source/CustomizableObjectEditor/Private/MuCOE/Nodes/CONodeExternal.cpp b/Engine/Plugins/Mutable/Source/CustomizableObjectEditor/Private/MuCOE/Nodes/CONodeExternal.cpp
new file mode 100644
--- /dev/null
+++ b/Engine/Plugins/Mutable/Source/CustomizableObjectEditor/Private/MuCOE/Nodes/CONodeExternal.cpp
@@ -0,0 +1,82 @@
+﻿// Copyright Epic Games, Inc. All Rights Reserved.
+
+#include "MuCOE/Nodes/CONodeExternal.h"
+
+#include "SCONodeExternal.h"
+#include "MuCOE/EdGraphSchema_CustomizableObject.h"
+#include "MuR/External/TextureAdapter.h"
+#include "MuR/External/MeshAdapter.h"
+#include "MuR/External/Operation.h"
+
+#define LOCTEXT_NAMESPACE "CustomizableObjectEditor"
+
+using namespace UE::Mutable;
+
+
+FText UCONodeExternal::GetNodeTitle(ENodeTitleType::Type TitleType) const
+{
+	const FExternalOperation* Operation = OperationInstancedStruct.GetPtr<FExternalOperation>();
+	const FName Name = Operation ? Operation->GetName() : CachedOperationName;
+	
+	if (TitleType == ENodeTitleType::ListView)
+	{
+		return FText::FromName(Name);
+	}
+	else
+	{
+		return FText::Format(LOCTEXT("Extension_Name", "{0}\nExternal"), FText::FromName(Name));
+	}
+}
+
+
+FLinearColor UCONodeExternal::GetNodeTitleColor() const
+{
+	const FName Category = UEdGraphSchema_CustomizableObject::GetPinCategoryName(OutputPin.Get()->PinType.PinCategory);
+	return UEdGraphSchema_CustomizableObject::GetPinTypeColor(Category);
+}
+
+
+TSharedPtr<SGraphNode> UCONodeExternal::CreateVisualWidget()
+{
+	return SNew(SCONodeExternal, this);
+}
+
+
+void UCONodeExternal::AllocateDefaultPins(UCustomizableObjectNodeRemapPins* RemapPins)
+{
+	Super::AllocateDefaultPins(RemapPins);
+
+	const FExternalOperation* Operation = OperationInstancedStruct.GetPtr<FExternalOperation>();
+	if (!Operation)
+	{
+		return;
+	}
+	
+	// Inputs
+	for (TPair<FName, const UScriptStruct*> Pair : Operation->GetInputs())
+	{
+		UCONodeExtensionPinData* PinData = NewObject<UCONodeExtensionPinData>(this);
+		PinData->Name = Pair.Key;
+		
+		const FName Category = UEdGraphSchema_CustomizableObject::GetPinCategoryName(Pair.Value);
+
+		UEdGraphPin* Pin = CustomCreatePin(EGPD_Input, Category, Pair.Key, PinData);
+		InputPins.Add(Pin);
+	}
+
+	// Output
+	{
+		const TPair<FName, const UScriptStruct*> Pair = Operation->GetOutput();
+
+		UCONodeExtensionPinData* PinData = NewObject<UCONodeExtensionPinData>(this);
+		PinData->Name = Pair.Key;
+
+		const FName Category = UEdGraphSchema_CustomizableObject::GetPinCategoryName(Pair.Value);
+		OutputPin = CustomCreatePin(EGPD_Output, Category, Pair.Key, PinData);
+	}
+
+	CachedOperationName = Operation->GetName();
+}
+
+
+#undef LOCTEXT_NAMESPACE
diff --git a/Engine/Plugins/Mutable/Source/CustomizableObjectEditor/Private/MuCOE/Nodes/CONodeExternal.h b/Engine/Plugins/Mutable/Source/CustomizableObjectEditor/Private/MuCOE/Nodes/CONodeExternal.h
new file mode 100644
--- /dev/null
+++ b/Engine/Plugins/Mutable/Source/CustomizableObjectEditor/Private/MuCOE/Nodes/CONodeExternal.h
@@ -0,0 +1,47 @@
+﻿// Copyright Epic Games, Inc. All Rights Reserved.
+
+#pragma once
+
+#include "MuCOE/Nodes/CustomizableObjectNode.h"
+
+#include "CONodeExternal.generated.h"
+
+
+UCLASS()
+class UCONodeExtensionPinData : public UCustomizableObjectNodePinData
+{
+	GENERATED_BODY()
+
+public:
+	UPROPERTY()
+	FName Name;
+};
+
+
+UCLASS()
+class UCONodeExternal : public UCustomizableObjectNode
+{
+	GENERATED_BODY()
+
+public:
+	// EdGraphNode interface
+	virtual FText GetNodeTitle(ENodeTitleType::Type TitleType) const override;
+	virtual FLinearColor GetNodeTitleColor() const override;
+	virtual TSharedPtr<SGraphNode> CreateVisualWidget() override;
+
+	// UCustomizableObjectNode interface
+	virtual void AllocateDefaultPins(UCustomizableObjectNodeRemapPins* RemapPins) override;
+	
+	UPROPERTY()
+	FInstancedStruct OperationInstancedStruct;
+
+	UPROPERTY()
+	TArray<FEdGraphPinReference> InputPins;
+
+	UPROPERTY()
+	FEdGraphPinReference OutputPin;
+	
+	/** In case of the operation not being loaded, the node uses this cached name so we can at least show the user which operation it was. */
+	UPROPERTY()
+	FName CachedOperationName;
+};
diff --git a/Engine/Plugins/Mutable/Source/CustomizableObjectEditor/Private/MuCOE/Nodes/CONodeExternalDetails.cpp b/Engine/Plugins/Mutable/Source/CustomizableObjectEditor/Private/MuCOE/Nodes/CONodeExternalDetails.cpp
new file mode 100644
--- /dev/null
+++ b/Engine/Plugins/Mutable/Source/CustomizableObjectEditor/Private/MuCOE/Nodes/CONodeExternalDetails.cpp
@@ -0,0 +1,83 @@
+﻿// Copyright Epic Games, Inc. All Rights Reserved.
+
+#include "MuCOE/Nodes/CONodeExternalDetails.h"
+
+#include "Editor.h"
+#include "CONodeExternal.h"
+#include "DetailCategoryBuilder.h"
+#include "DetailLayoutBuilder.h"
+#include "IStructureDetailsView.h"
+
+#define LOCTEXT_NAMESPACE "CustomizableObjectEditor"
+
+
+TSharedRef<IDetailCustomization> FCONodeExternalDetails::MakeInstance()
+{
+	return MakeShareable(new FCONodeExternalDetails);
+}
+
+
+void FCONodeExternalDetails::CustomizeDetails(IDetailLayoutBuilder& DetailBuilder)
+{
+	TSharedPtr<const IDetailsView> DetailsView = DetailBuilder.GetDetailsViewSharedPtr();
+	if (!DetailsView || !DetailsView->GetSelectedObjects().Num())
+	{
+		return;
+	}
+
+	UCONodeExternal* Node = Cast<UCONodeExternal>(DetailsView->GetSelectedObjects()[0]);
+	if (!Node)
+	{
+		return;
+	}
+
+	const UStruct* Struct = Node->OperationInstancedStruct.GetScriptStruct();
+	if (!Struct)
+	{
+		return;
+	}
+
+	TSharedRef<FStructOnScope> StructData = MakeShared<FStructOnScope>(Struct, Node->OperationInstancedStruct.GetMutableMemory());
+
+	for (FProperty* Property = Struct->PropertyLink; Property; Property = Property->PropertyLinkNext)
+	{
+		FName PropertyName = Property->GetFName();
+		TSharedPtr<IPropertyHandle> PropertyHandle = DetailBuilder.AddStructurePropertyData(StructData, PropertyName);
+		IDetailPropertyRow& DetailPropertyRow = DetailBuilder.AddPropertyToCategory(PropertyHandle.ToSharedRef());
+
+		TSharedPtr<IPropertyHandle> PropertyRowHandle = DetailPropertyRow.GetPropertyHandle();
+		
+		PropertyRowHandle->SetOnPropertyValuePreChange(FSimpleDelegate::CreateLambda([Node]()
+		{
+			GEditor->BeginTransaction(LOCTEXT("EditProperty", "Edited property"));
+			Node->Modify();
+		}));
+		
+		PropertyRowHandle->SetOnPropertyValueChangedWithData(TDelegate<void(const FPropertyChangedEvent&)>::CreateLambda([this](const FPropertyChangedEvent& InEvent)
+		{
+			GEditor->EndTransaction();
+		}));
+
+		PropertyRowHandle->SetOnChildPropertyValuePreChange(FSimpleDelegate::CreateLambda([Node]()
+		{
+			GEditor->BeginTransaction(LOCTEXT("EditChildProperty", "Edited child property"))	;
+			Node->Modify();
+		}));
+
+		PropertyRowHandle->SetOnChildPropertyValueChangedWithData(TDelegate<void(const FPropertyChangedEvent&)>::CreateLambda([this](const FPropertyChangedEvent& InEvent)
+		{
+			GEditor->EndTransaction();
+		}));
+	}
+}
+
+
+void FCONodeExternalDetails::CustomizeDetails(const TSharedPtr<IDetailLayoutBuilder>& DetailBuilder)
+{
+	IDetailCustomization::CustomizeDetails(DetailBuilder);
+
+	DetailBuilderPtr = DetailBuilder;
+}
+
+
+#undef LOCTEXT_NAMESPACE
diff --git a/Engine/Plugins/Mutable/Source/CustomizableObjectEditor/Private/MuCOE/Nodes/CONodeExternalDetails.h b/Engine/Plugins/Mutable/Source/CustomizableObjectEditor/Private/MuCOE/Nodes/CONodeExternalDetails.h
new file mode 100644
--- /dev/null
+++ b/Engine/Plugins/Mutable/Source/CustomizableObjectEditor/Private/MuCOE/Nodes/CONodeExternalDetails.h
@@ -0,0 +1,23 @@
+﻿// Copyright Epic Games, Inc. All Rights Reserved.
+
+#pragma once
+
+#include "IDetailCustomization.h"
+
+class FReply;
+class IDetailLayoutBuilder;
+class UCustomizableObjectNode;
+
+
+class FCONodeExternalDetails : public IDetailCustomization
+{
+public:
+	static TSharedRef<IDetailCustomization> MakeInstance();
+	
+	// IDetailCustomization interface
+	virtual void CustomizeDetails(IDetailLayoutBuilder& DetailBuilder) override;
+	virtual void CustomizeDetails(const TSharedPtr<IDetailLayoutBuilder>& DetailBuilder) override;
+
+private: 
+	TSharedPtr<IDetailLayoutBuilder> DetailBuilderPtr;
+};
diff --git a/Engine/Plugins/Mutable/Source/CustomizableObjectEditor/Private/MuCOE/Nodes/CustomizableObjectNode.cpp b/Engine/Plugins/Mutable/Source/CustomizableObjectEditor/Private/MuCOE/Nodes/CustomizableObjectNode.cpp
--- a/Engine/Plugins/Mutable/Source/CustomizableObjectEditor/Private/MuCOE/Nodes/CustomizableObjectNode.cpp
+++ b/Engine/Plugins/Mutable/Source/CustomizableObjectEditor/Private/MuCOE/Nodes/CustomizableObjectNode.cpp
@@ -60,9 +60,9 @@ UEdGraphPin* UCustomizableObjectNode::CustomCreatePinSimple(EEdGraphPinDirection
 }
 
 
-UEdGraphPin* UCustomizableObjectNode::CustomCreatePin(EEdGraphPinDirection Direction, const FName& Type, const FName& Name, bool bIsArray)
+UEdGraphPin* UCustomizableObjectNode::CustomCreatePin(EEdGraphPinDirection Direction, const FName& Category, const FName& Name, bool bIsArray)
 {
-	UEdGraphPin* Pin = CreatePin(Direction, Type, Name);
+	UEdGraphPin* Pin = CreatePin(Direction, Category, Name);
 
 	Pin->PinFriendlyName = FText::FromName(Name);
 	if (bIsArray)
@@ -74,9 +74,9 @@ UEdGraphPin* UCustomizableObjectNode::CustomCreatePin(EEdGraphPinDirection Direc
 }
 
 
-UEdGraphPin* UCustomizableObjectNode::CustomCreatePin(EEdGraphPinDirection Direction, const FName& Type, const FName& Name, UCustomizableObjectNodePinData* PinData)
+UEdGraphPin* UCustomizableObjectNode::CustomCreatePin(EEdGraphPinDirection Direction, const FName& Category, const FName& Name, UCustomizableObjectNodePinData* PinData)
 {
-	UEdGraphPin* Pin = CreatePin(Direction, Type, Name);
+	UEdGraphPin* Pin = CreatePin(Direction, Category, Name);
 	if (Pin && PinData)
 	{
 		AddPinData(*Pin, *PinData);		
diff --git a/Engine/Plugins/Mutable/Source/CustomizableObjectEditor/Private/MuCOE/Nodes/SCONodeExternal.cpp b/Engine/Plugins/Mutable/Source/CustomizableObjectEditor/Private/MuCOE/Nodes/SCONodeExternal.cpp
new file mode 100644
--- /dev/null
+++ b/Engine/Plugins/Mutable/Source/CustomizableObjectEditor/Private/MuCOE/Nodes/SCONodeExternal.cpp
@@ -0,0 +1,41 @@
+﻿// Copyright Epic Games, Inc. All Rights Reserved.
+
+#include "SCONodeExternal.h"
+
+#include "CONodeExternal.h"
+#include "Widgets/SBoxPanel.h"
+
+#define LOCTEXT_NAMESPACE "CustomizableObjectEditor"
+
+
+void SCONodeExternal::Construct(const FArguments& InArgs, UEdGraphNode* InGraphNode)
+{
+	GraphNode = InGraphNode;
+	
+	UpdateGraphNode();
+}
+
+
+void SCONodeExternal::CreateBelowPinControls(TSharedPtr<SVerticalBox> MainBox)
+{
+	SCustomizableObjectNode::CreateBelowPinControls(MainBox);
+
+	if (UCONodeExternal* Node = Cast<UCONodeExternal>(GraphNode))
+	{
+		if (!Node->OperationInstancedStruct.IsValid())
+		{
+			MainBox->AddSlot()
+			.Padding(FMargin(0.0, 2.0, 0.0, 0.0))
+			.AutoHeight()
+			[
+				SNew(SErrorText)
+				.BackgroundColor(FAppStyle::GetColor("ErrorReporting.BackgroundColor"))
+				.ErrorText(LOCTEXT("Experimental", "NOT LOADED"))
+			];
+		}
+	}
+}
+
+
+#undef LOCTEXT_NAMESPACE
+
diff --git a/Engine/Plugins/Mutable/Source/CustomizableObjectEditor/Private/MuCOE/Nodes/SCONodeExternal.h b/Engine/Plugins/Mutable/Source/CustomizableObjectEditor/Private/MuCOE/Nodes/SCONodeExternal.h
new file mode 100644
--- /dev/null
+++ b/Engine/Plugins/Mutable/Source/CustomizableObjectEditor/Private/MuCOE/Nodes/SCONodeExternal.h
@@ -0,0 +1,22 @@
+﻿// Copyright Epic Games, Inc. All Rights Reserved.
+
+#pragma once
+
+#include "SCustomizableObjectNode.h"
+
+class SGraphPin;
+class UEdGraphNode;
+class UEdGraphPin;
+
+
+class SCONodeExternal : public SCustomizableObjectNode
+{
+public:
+	SLATE_BEGIN_ARGS(SCONodeExternal) {}
+	SLATE_END_ARGS();
+
+	void Construct(const FArguments& InArgs, UEdGraphNode* InGraphNode);
+
+	// SGraphNode interface
+	virtual void CreateBelowPinControls(TSharedPtr<SVerticalBox> MainBox) override;
+};
diff --git a/Engine/Plugins/Mutable/Source/CustomizableObjectEditor/Private/MuCOE/SMutableGraphViewer.cpp b/Engine/Plugins/Mutable/Source/CustomizableObjectEditor/Private/MuCOE/SMutableGraphViewer.cpp
--- a/Engine/Plugins/Mutable/Source/CustomizableObjectEditor/Private/MuCOE/SMutableGraphViewer.cpp
+++ b/Engine/Plugins/Mutable/Source/CustomizableObjectEditor/Private/MuCOE/SMutableGraphViewer.cpp
@@ -79,6 +79,7 @@ namespace MutableGraphTreeViewColumns
 	static const FName Name("Name");
 };
 
+
 static const char* MutableNodeNames[] =
 {
 	"None",
@@ -102,6 +103,7 @@ static const char* MutableNodeNames[] =
 	"MeshReshape",
 	"MeshClipDeform",
 	"MeshParameter",
+	"MeshExternal",
 
 	"Image",
 	"ImageConstant",
@@ -129,6 +131,7 @@ static const char* MutableNodeNames[] =
 	"ImageNormalComposite",
 	"ImageTransform",
 	"ImageFromMaterialParameter",
+	"ImageExternal",
 
 	"Bool",
 	"BoolConstant",
@@ -147,6 +150,7 @@ static const char* MutableNodeNames[] =
 	"ColorSwitch",
 	"ColorVariation",
 	"ColorToSRGB",
+	"ColorExternal",
 
 	"Scalar",
 	"ScalarConstant",
@@ -157,6 +161,7 @@ static const char* MutableNodeNames[] =
 	"ScalarArithmeticOperation",
 	"ScalarVariation",
 	"ScalarTable",
+	"ScalarExternal",
 
 	"String",
 	"StringConstant",
@@ -200,6 +205,8 @@ static const char* MutableNodeNames[] =
 	"ModifierTransformInMesh",
 	"ModifierTransformWithBone",
 
+	"Extension",
+	
 	"ExtensionData",
 	"ExtensionDataConstant",
 	"ExtensionDataSwitch",
@@ -215,6 +222,7 @@ static const char* MutableNodeNames[] =
 	"MaterialSwitch",
 	"MaterialVariation",
 	"MaterialParameter",
+	"MaterialExternal",
 
 	"ImageMaterialBreak",
 	"ScalarMaterialBreak",
diff --git a/Engine/Plugins/Mutable/Source/CustomizableObjectEditor/Public/MuCOE/EdGraphSchema_CustomizableObject.h b/Engine/Plugins/Mutable/Source/CustomizableObjectEditor/Public/MuCOE/EdGraphSchema_CustomizableObject.h
--- a/Engine/Plugins/Mutable/Source/CustomizableObjectEditor/Public/MuCOE/EdGraphSchema_CustomizableObject.h
+++ b/Engine/Plugins/Mutable/Source/CustomizableObjectEditor/Public/MuCOE/EdGraphSchema_CustomizableObject.h
@@ -73,6 +73,7 @@ public:
 	// Own interface
 	static UE_API FLinearColor GetPinTypeColor(const FName& PinType);
 	static UE_API FName GetPinCategoryName(const FName& PinCategory);
+	static UE_API FName GetPinCategoryName(const UScriptStruct* Type);
 	static UE_API FText GetPinCategoryFriendlyName(const FName& PinCategory);
 	static UE_API bool IsPassthrough(const FName& PinCategory);
 };
diff --git a/Engine/Plugins/Mutable/Source/CustomizableObjectEditor/Public/MuCOE/Nodes/CustomizableObjectNode.h b/Engine/Plugins/Mutable/Source/CustomizableObjectEditor/Public/MuCOE/Nodes/CustomizableObjectNode.h
--- a/Engine/Plugins/Mutable/Source/CustomizableObjectEditor/Public/MuCOE/Nodes/CustomizableObjectNode.h
+++ b/Engine/Plugins/Mutable/Source/CustomizableObjectEditor/Public/MuCOE/Nodes/CustomizableObjectNode.h
@@ -218,13 +218,12 @@ public:
 	UE_API UEdGraphPin* CustomCreatePinSimple( EEdGraphPinDirection Direction, const FName& Category, bool bIsArray=false );
 	
 	// Wrapper of the pin creation method, to cope with multiple Unreal Engine versions
-	UE_API UEdGraphPin* CustomCreatePin( EEdGraphPinDirection Direction, const FName& Type, const FName& Name, bool bIsArray=false );
+	UE_API UEdGraphPin* CustomCreatePin( EEdGraphPinDirection Direction, const FName& Category, const FName& Name, bool bIsArray=false );
 
 	/** Create pin with attached data.
 	 *
      * @param PinData Pin data to be saved. */
-	UE_API UEdGraphPin* CustomCreatePin(EEdGraphPinDirection Direction, const FName& Type, const FName& Name, UCustomizableObjectNodePinData* PinData);
-
+	UE_API UEdGraphPin* CustomCreatePin(EEdGraphPinDirection Direction, const FName& Category, const FName& Name, UCustomizableObjectNodePinData* PinData);
 	
 	virtual void AddAttachedErrorData(const FAttachedErrorDataView& ) {};
 	virtual void ResetAttachedErrorData() {};
@@ -337,7 +336,7 @@ protected:
 	 *
 	 * @param RemapPins remap pins action. Use this object to save any information required to later perfrom the remap pins action. Can be null.
 	 */
-	virtual void AllocateDefaultPins(UCustomizableObjectNodeRemapPins* RemapPins) {};
+	virtual void AllocateDefaultPins(UCustomizableObjectNodeRemapPins* RemapPins) {}
 
 	/**
 	 * Creates and returns the node default node remap pins action.
diff --git a/Engine/Plugins/Mutable/Source/MutableRuntime/Internal/MuR/ExternalOperationProvider.h b/Engine/Plugins/Mutable/Source/MutableRuntime/Internal/MuR/ExternalOperationProvider.h
new file mode 100644
--- /dev/null
+++ b/Engine/Plugins/Mutable/Source/MutableRuntime/Internal/MuR/ExternalOperationProvider.h
@@ -0,0 +1,19 @@
+﻿// Copyright Epic Games, Inc. All Rights Reserved.
+
+#pragma once
+
+#include "Operations.h"
+#include "StructUtils/InstancedStruct.h"
+
+
+namespace UE::Mutable::Private
+{
+	class IExternalOperationProvider
+	{
+	public:
+		virtual ~IExternalOperationProvider() = default;
+		
+		virtual const FInstancedStruct& Get(OP::ADDRESS Address) const = 0;
+	};
+}
+
diff --git a/Engine/Plugins/Mutable/Source/MutableRuntime/Internal/MuR/ModelPrivate.h b/Engine/Plugins/Mutable/Source/MutableRuntime/Internal/MuR/ModelPrivate.h
--- a/Engine/Plugins/Mutable/Source/MutableRuntime/Internal/MuR/ModelPrivate.h
+++ b/Engine/Plugins/Mutable/Source/MutableRuntime/Internal/MuR/ModelPrivate.h
@@ -2,6 +2,7 @@
 
 #pragma once
 
+#include "ExternalOperationProvider.h"
 #include "MuR/Model.h"
 #include "MuR/System.h"
 
@@ -12,6 +13,7 @@
 #include "MuR/Mesh.h"
 #include "MuR/ParametersPrivate.h"
 #include "MuR/MutableRuntimeModule.h"
+#include "StructUtils/InstancedStruct.h"
 
 #define UE_API MUTABLERUNTIME_API
 
@@ -376,6 +378,8 @@ namespace UE::Mutable::Private
 		/** Constant Material Data */
 		TArray<TSharedPtr<const FMaterial>> ConstantMaterials;
 
+    	TSharedPtr<const IExternalOperationProvider> ExternalOperationProvider;
+    	
 #if WITH_EDITOR
 		/** 
 		 * State of the program. True unless the streamed resources were destroyed,
diff --git a/Engine/Plugins/Mutable/Source/MutableRuntime/Internal/MuR/Operations.h b/Engine/Plugins/Mutable/Source/MutableRuntime/Internal/MuR/Operations.h
--- a/Engine/Plugins/Mutable/Source/MutableRuntime/Internal/MuR/Operations.h
+++ b/Engine/Plugins/Mutable/Source/MutableRuntime/Internal/MuR/Operations.h
@@ -351,6 +351,13 @@ namespace UE::Mutable::Private
 		//! Extract a layout from a mesh
 		LA_FROMMESH,
 
+    	//-----------------------------------------------------------------------------------------
+		// External operations
+		//-----------------------------------------------------------------------------------------
+
+    	//! External operation. Returns unknown type.
+		UN_EXTERNAL, // TODO GMT Instruction for each type
+    	
         //-----------------------------------------------------------------------------------------
         // Utility values
         //-----------------------------------------------------------------------------------------
@@ -1129,6 +1136,11 @@ namespace UE::Mutable::Private
 			ADDRESS MaterialParameter;
 			STRING_ADDRESS ParameterName;
 		};
+
+    	struct ExternalArgs
+    	{
+    		int32 NumOperants;
+    	};
     };
 
 
@@ -1156,13 +1168,15 @@ namespace UE::Mutable::Private
 		ExtensionData,
 		Matrix,
 		Material,
+		InstancedStruct,
 
 		// Supporting data types : Never returned as an actual data type for any operation.
 		Shape,
 		Curve,
 		Skeleton,
 		PhysicsAsset,
-		
+
+		Unknown, // Not a type by itself. Unknown at compile time but at runtime it has a valid type. 
 		Count
 	};
 
diff --git a/Engine/Plugins/Mutable/Source/MutableRuntime/Internal/MuR/SystemPrivate.h b/Engine/Plugins/Mutable/Source/MutableRuntime/Internal/MuR/SystemPrivate.h
--- a/Engine/Plugins/Mutable/Source/MutableRuntime/Internal/MuR/SystemPrivate.h
+++ b/Engine/Plugins/Mutable/Source/MutableRuntime/Internal/MuR/SystemPrivate.h
@@ -503,6 +503,7 @@ namespace UE::Mutable::Private
 		TMemoryTrackedArray<FProjector> ProjectorResults;
 		TMemoryTrackedArray<TSharedPtr<const String>> StringResults;
 		TMemoryTrackedArray<TSharedPtr<const FExtensionData>> ExtensionDataResults;
+		TMemoryTrackedArray<TSharedPtr<const FInstancedStruct>> InstancedStructResults;
 		TMemoryTrackedArray<FMatrix44f> MatrixResults;
 		TMemoryTrackedArray<TSharedPtr<const FMaterial>> MaterialResults;
 
@@ -885,6 +886,17 @@ namespace UE::Mutable::Private
 			return ExtensionDataResults[Data->DataTypeIndex];
 		}
 		
+		TSharedPtr<const FInstancedStruct> GetInstancedStruct(FCacheAddress At)
+		{
+			if (!At.At) return nullptr;
+			FOpExecutionData* Data = OpExecutionData.get_ptr(At);
+			if (!Data) return nullptr;
+
+			check(Data->DataType == EDataType::InstancedStruct);
+			return InstancedStructResults[Data->DataTypeIndex];
+		}
+		
+		
 		void SetValidDesc(FCacheAddress at)
 		{
 			check(at.Type == FScheduledOp::EType::ImageDesc);
@@ -1043,6 +1055,26 @@ namespace UE::Mutable::Private
 			check(Data->DataTypeIndex != 0);
 		}
 
+		void SetInstancedStruct(FCacheAddress At, TSharedPtr<const FInstancedStruct> Value)
+		{
+			check(At.At < OpExecutionData.size_code());
+			FOpExecutionData* Data = OpExecutionData.get_ptr(At);
+			check(Data->DataType == EDataType::InstancedStruct || Data->DataType == EDataType::None);
+			Data->DataType = EDataType::InstancedStruct;
+			Data->IsValueValid = true;
+
+			if (!Data->DataTypeIndex)
+			{
+				Data->DataTypeIndex = InstancedStructResults.Num();
+				InstancedStructResults.Add(Value);
+			}
+			else
+			{
+				InstancedStructResults[Data->DataTypeIndex] = Value;
+			}
+			check(Data->DataTypeIndex != 0);
+		}
+		
 		void SetImage(FCacheAddress At, TSharedPtr<const FImage> Value)
 		{
 			check(At.At < OpExecutionData.size_code());
diff --git a/Engine/Plugins/Mutable/Source/MutableRuntime/MutableRuntime.Build.cs b/Engine/Plugins/Mutable/Source/MutableRuntime/MutableRuntime.Build.cs
--- a/Engine/Plugins/Mutable/Source/MutableRuntime/MutableRuntime.Build.cs
+++ b/Engine/Plugins/Mutable/Source/MutableRuntime/MutableRuntime.Build.cs
@@ -14,6 +14,8 @@ namespace UnrealBuildTool.Rules
 			IWYUSupport = IWYUSupport.KeepAsIsForNow;
 			//bUseUnity = false;
 
+			bAllowUETypesInNamespaces = true;
+
 			PublicDependencyModuleNames.AddRange(
 				new string[] {
 				"Core",
diff --git a/Engine/Plugins/Mutable/Source/MutableRuntime/Private/MuR/CodeRunner.cpp b/Engine/Plugins/Mutable/Source/MutableRuntime/Private/MuR/CodeRunner.cpp
--- a/Engine/Plugins/Mutable/Source/MutableRuntime/Private/MuR/CodeRunner.cpp
+++ b/Engine/Plugins/Mutable/Source/MutableRuntime/Private/MuR/CodeRunner.cpp
@@ -58,6 +58,13 @@
 #include "MuR/PhysicsBody.h"
 #include "MuR/Skeleton.h"
 #include "MuR/SystemPrivate.h"
+#include "MuR/External/FloatAdapter.h"
+#include "MuR/External/MaterialAdapter.h"
+#include "MuR/External/TextureAdapter.h"
+#include "MuR/External/MeshAdapter.h"
+#include "MuR/External/Operation.h"
+#include "MuR/External/VectorAdapter.h"
+#include "StructUtils/InstancedStruct.h"
 #include "Templates/Tuple.h"
 
 #if MUTABLE_DEBUG_CODERUNNER_TASK_SCHEDULE_CALLSTACK
@@ -6671,7 +6678,7 @@ namespace UE::Mutable::Private
 
 
     //---------------------------------------------------------------------------------------------
-    void CodeRunner::RunCode( const FScheduledOp& Item, const FParameters* pParams, const TSharedPtr<const FModel>& InModel, uint32 LodMask)
+    void CodeRunner::RunCode( const FScheduledOp& Item, const FParameters* pParams, const TSharedPtr<const FModel>& InModel, uint32 LodMask, bool& bOutFailed)
     {
 		//UE_LOG( LogMutableCore, Log, TEXT("Running :%5d , %d "), Item.At, Item.Stage );
 		check( Item.Type == FScheduledOp::EType::Full );
@@ -6730,6 +6737,153 @@ namespace UE::Mutable::Private
             RunCode_InstanceAddResource(Item, InModel, pParams);
             break;
 
+		case EOpType::UN_EXTERNAL:
+			{
+				const uint8* Data = Program.GetOpArgsPointer(Item.At);
+
+				OP::ExternalArgs Args = Program.GetOpArgs<OP::ExternalArgs>(Item.At);
+				Data += sizeof(OP::ExternalArgs);
+				
+				const FInstancedStruct& OperationInstancedStruct = Program.ExternalOperationProvider->Get(Item.At);
+				if (!OperationInstancedStruct.IsValid()) // Plugin not loaded
+				{
+                    UE_LOG(LogMutableCore, Error, TEXT("Trying to evaluate an External Operation which has not been loaded."));
+					bOutFailed = true;
+					break;
+				}
+					
+				const FExternalOperation* Operation = OperationInstancedStruct.GetPtr<FExternalOperation>();
+				
+				switch (Item.Stage)
+				{
+				case 0:
+					{
+						TArray<FScheduledOp, TInlineAllocator<8>> Dependencies;
+
+						for (int32 Index = 0; Index < Args.NumOperants; ++Index)
+						{
+							OP::ADDRESS OperantAt;
+							FMemory::Memcpy(&OperantAt, Data, sizeof(OP::ADDRESS));
+							Data += sizeof(OP::ADDRESS);
+
+							Dependencies.Add(FScheduledOp(OperantAt, Item));
+						}
+
+						AddOp(FScheduledOp(Item.At, Item, 1), Dependencies);
+						
+						break;
+					}
+
+				case 1:
+					{
+						FContext Context;
+
+						const TArray<TPair<FName, const UScriptStruct*>>& Inputs = Operation->GetInputs();
+						check(Inputs.Num() == Args.NumOperants);
+
+						// Inputs
+						for (int32 Index = 0; Index < Args.NumOperants; ++Index)
+						{
+							OP::ADDRESS OperantAt;
+							FMemory::Memcpy(&OperantAt, Data, sizeof(OP::ADDRESS));
+							Data += sizeof(OP::ADDRESS);
+
+							const TPair<FName, const UScriptStruct*>& Input = Inputs[Index];
+
+							const UScriptStruct* Type = Input.Value;
+							if (Type == FMeshAdapter::StaticStruct())
+							{
+								TSharedPtr<const FMesh> Value = LoadMesh(FCacheAddress(OperantAt, Item));
+	
+								TSharedPtr<FInstancedStruct> InstancedStruct = MakeShared<FInstancedStruct>();
+								InstancedStruct->InitializeAs<FMeshAdapter>().Mesh = ConstCastSharedPtr<FMesh>(Value);
+
+								Context.Inputs.Add(Input.Key, MoveTemp(InstancedStruct));
+							}
+							else if (Type == FTextureAdapter::StaticStruct())
+							{
+								TSharedPtr<const FImage> Value = LoadImage(FCacheAddress(OperantAt, Item));
+
+								TSharedPtr<FInstancedStruct> InstancedStruct = MakeShared<FInstancedStruct>();
+								InstancedStruct->InitializeAs<FTextureAdapter>().Image = ConstCastSharedPtr<FImage>(Value);
+
+								Context.Inputs.Add(Input.Key, MoveTemp(InstancedStruct));
+							}
+							else if (Type == FFloatAdapter::StaticStruct())
+							{
+								float Value = LoadScalar(FCacheAddress(OperantAt, Item));
+
+								TSharedPtr<FInstancedStruct> InstancedStruct = MakeShared<FInstancedStruct>();
+								InstancedStruct->InitializeAs<FFloatAdapter>().Value = Value;
+
+								Context.Inputs.Add(Input.Key, MoveTemp(InstancedStruct));
+							}
+							else if (Type == FVectorAdapter::StaticStruct())
+							{
+								FVector4f Value = LoadColor(FCacheAddress(OperantAt, Item));
+
+								TSharedPtr<FInstancedStruct> InstancedStruct = MakeShared<FInstancedStruct>();
+								InstancedStruct->InitializeAs<FVectorAdapter>().Value = Value;
+
+								Context.Inputs.Add(Input.Key, MoveTemp(InstancedStruct));
+							}
+							else if (Type == FMaterialAdapter::StaticStruct())
+							{
+								TSharedPtr<const FMaterial> Value = LoadMaterial(FCacheAddress(OperantAt, Item));
+
+								TSharedPtr<FInstancedStruct> InstancedStruct = MakeShared<FInstancedStruct>();
+								InstancedStruct->InitializeAs<FMaterialAdapter>().Material = ConstCastSharedPtr<FMaterial>(Value);
+
+								Context.Inputs.Add(Input.Key, MoveTemp(InstancedStruct));
+							}
+							else
+							{
+								TSharedPtr<const FInstancedStruct> Value = LoadInstancedStruct(FCacheAddress(OperantAt, Item));
+
+								Context.Inputs.Add(Input.Key, Value);
+							}
+						}
+
+						// Evaluate
+						Operation->Evaluate(Context);
+
+						// Store output
+						const UStruct* Type = Context.Output->GetScriptStruct();
+						if (Type == FMeshAdapter::StaticStruct())
+						{
+							StoreMesh(FCacheAddress(Item.At, Item), Context.Output.Get()->Get<FMeshAdapter>().Mesh);
+						}
+						else if (Type == FTextureAdapter::StaticStruct())
+						{
+							StoreImage(FCacheAddress(Item.At, Item), Context.Output.Get()->Get<FTextureAdapter>().Image);
+						}
+						else if (Type == FFloatAdapter::StaticStruct())
+						{
+							StoreScalar(FCacheAddress(Item.At, Item), Context.Output.Get()->Get<FFloatAdapter>().Value);
+						}
+						else if (Type == FVectorAdapter::StaticStruct())
+						{
+							StoreColor(FCacheAddress(Item.At, Item), Context.Output.Get()->Get<FVectorAdapter>().Value);
+						}
+						else if (Type == FMaterialAdapter::StaticStruct())
+						{
+							StoreMaterial(FCacheAddress(Item.At, Item), Context.Output.Get()->Get<FMaterialAdapter>().Material);
+						}
+						else
+						{
+							StoreInstancedStruct(FCacheAddress(Item.At, Item), Context.Output);
+						}
+
+						break;
+					}
+
+				default:
+					unimplemented();
+				}
+
+				break;	
+			}
+			
 		default:
 		{
 			EDataType DataType = GetOpDataType(type);
diff --git a/Engine/Plugins/Mutable/Source/MutableRuntime/Private/MuR/CodeRunner.h b/Engine/Plugins/Mutable/Source/MutableRuntime/Private/MuR/CodeRunner.h
--- a/Engine/Plugins/Mutable/Source/MutableRuntime/Private/MuR/CodeRunner.h
+++ b/Engine/Plugins/Mutable/Source/MutableRuntime/Private/MuR/CodeRunner.h
@@ -119,7 +119,7 @@ namespace  UE::Mutable::Private
     	
         TSharedPtr<FRangeIndex> BuildCurrentOpRangeIndex( const FScheduledOp&, int32 ParameterIndex);
 
-        void RunCode( const FScheduledOp&, const FParameters*, const TSharedPtr<const FModel>&, uint32 LodMask);
+        void RunCode( const FScheduledOp&, const FParameters*, const TSharedPtr<const FModel>&, uint32 LodMask, bool& bOutFailed);
 
         void RunCode_Conditional(const FScheduledOp&, const FModel* );
         void RunCode_Switch(const FScheduledOp&, const FModel* );
@@ -507,6 +507,11 @@ namespace  UE::Mutable::Private
 		{
 			return System->WorkingMemoryManager.CurrentInstanceCache->GetExtensionData(From);
 		}
+    	
+    	inline TSharedPtr<const FInstancedStruct> LoadInstancedStruct(const FCacheAddress& From)
+		{
+			return System->WorkingMemoryManager.CurrentInstanceCache->GetInstancedStruct(From);
+		}
 
 		inline void StoreValidDesc(const FCacheAddress& To)
 		{
@@ -578,6 +583,11 @@ namespace  UE::Mutable::Private
 			System->WorkingMemoryManager.CurrentInstanceCache->SetExtensionData(To, Resource);
 		}
 
+		inline void StoreInstancedStruct(const FCacheAddress& To, TSharedPtr<const FInstancedStruct> Resource)
+		{
+			System->WorkingMemoryManager.CurrentInstanceCache->SetInstancedStruct(To, Resource);
+		}
+    	
 		inline TSharedPtr<FImage> CreateImage(uint32 SizeX, uint32 SizeY, uint32 Lods, EImageFormat Format, EInitializationType Init)
 		{
 			TSharedPtr<FImage> Result = System->WorkingMemoryManager.CreateImage(SizeX, SizeY, Lods, Format, Init);
diff --git a/Engine/Plugins/Mutable/Source/MutableRuntime/Private/MuR/CodeRunnerBegin.cpp b/Engine/Plugins/Mutable/Source/MutableRuntime/Private/MuR/CodeRunnerBegin.cpp
--- a/Engine/Plugins/Mutable/Source/MutableRuntime/Private/MuR/CodeRunnerBegin.cpp
+++ b/Engine/Plugins/Mutable/Source/MutableRuntime/Private/MuR/CodeRunnerBegin.cpp
@@ -6,6 +6,12 @@
 #include "MuR/MutableTrace.h"
 #include "MuR/Operations.h"
 #include "MuR/SystemPrivate.h"
+#include "MuR/External/FloatAdapter.h"
+#include "MuR/External/MaterialAdapter.h"
+#include "MuR/External/MeshAdapter.h"
+#include "MuR/External/Operation.h"
+#include "MuR/External/TextureAdapter.h"
+#include "MuR/External/VectorAdapter.h"
 
 
 namespace UE::Mutable::Private
@@ -362,6 +368,7 @@ namespace UE::Mutable::Private
 				case EDataType::Image:		
 				case EDataType::Layout:		
 				case EDataType::Material:		
+				case EDataType::Unknown:		
 					break;
 
 				default:
@@ -2358,7 +2365,93 @@ namespace UE::Mutable::Private
 					StoreNone(Item);
 					break;
 				}
-				
+
+			case EOpType::UN_EXTERNAL:
+				{
+					const uint8* Data = Program.GetOpArgsPointer(Item.At);
+
+					OP::ExternalArgs Args = Program.GetOpArgs<OP::ExternalArgs>(Item.At);
+					Data += sizeof(OP::ExternalArgs);
+
+					const FInstancedStruct& OperationInstancedStruct = Program.ExternalOperationProvider->Get(Item.At);
+					if (!OperationInstancedStruct.IsValid()) // Plugin not loaded
+					{
+						StoreNone(Item);
+						break;
+					}
+
+					switch(Item.Stage)
+					{
+					case 0:
+						{
+							PushOp(FScheduledOpInline(Item.At, 1));
+
+							for (int32 Index = 0; Index < Args.NumOperants; ++Index)
+							{
+								OP::ADDRESS OperantAt;
+								FMemory::Memcpy(&OperantAt, Data, sizeof(OP::ADDRESS));
+								Data += sizeof(OP::ADDRESS);
+
+								PushOp(FScheduledOpInline(OperantAt, Item));
+							}
+							
+							break;
+						}
+
+					case 1:
+						{
+							const FExternalOperation* Operation = OperationInstancedStruct.GetPtr<FExternalOperation>();
+
+							const TArray<TPair<FName, const UScriptStruct*>>& Inputs = Operation->GetInputs();
+							check(Inputs.Num() == Args.NumOperants);
+							
+							for (int32 Index = 0; Index < Args.NumOperants; ++Index)
+							{
+								OP::ADDRESS OperantAt;
+								FMemory::Memcpy(&OperantAt, Data, sizeof(OP::ADDRESS));
+								Data += sizeof(OP::ADDRESS);
+
+								const TPair<FName, const UScriptStruct*>& Input = Inputs[Index];
+
+								const UScriptStruct* InputType = Input.Value;
+								
+								if (InputType == FMeshAdapter::StaticStruct())
+								{
+									// Nothing.
+								}
+								else if (InputType == FTextureAdapter::StaticStruct())
+								{
+									// Nothing.
+								}
+								else if (InputType == FFloatAdapter::StaticStruct())
+								{
+									LoadFloat(OperantAt);
+								}
+								else if (InputType == FVectorAdapter::StaticStruct())
+								{
+									// Nothing.
+								}
+								else if (InputType == FMaterialAdapter::StaticStruct())
+								{
+									// Nothing.
+								}
+								else
+								{
+									// Nothing.
+								}
+							}
+
+							StoreNone(Item);
+							break;
+						}
+						
+					default:
+						unimplemented()
+					}
+					
+					break;
+				}
+			
 			default:
 				unimplemented();
 			}
diff --git a/Engine/Plugins/Mutable/Source/MutableRuntime/Private/MuR/CodeRunnerTasked.cpp b/Engine/Plugins/Mutable/Source/MutableRuntime/Private/MuR/CodeRunnerTasked.cpp
--- a/Engine/Plugins/Mutable/Source/MutableRuntime/Private/MuR/CodeRunnerTasked.cpp
+++ b/Engine/Plugins/Mutable/Source/MutableRuntime/Private/MuR/CodeRunnerTasked.cpp
@@ -303,7 +303,12 @@ namespace UE::Mutable::Private
 				else
 				{
 					// Run immediately
-					RunCode(item, Params, Model, LODMask);
+					bool bFailed = false;
+					RunCode(item, Params, Model, LODMask, bFailed);
+					if (bFailed)
+					{
+						return AbortRun();
+					}
 
 					if (ScheduledStagePerOp[item] == item.Stage + 1)
 					{
diff --git a/Engine/Plugins/Mutable/Source/MutableRuntime/Private/MuR/External/FloatAdapter.cpp b/Engine/Plugins/Mutable/Source/MutableRuntime/Private/MuR/External/FloatAdapter.cpp
new file mode 100644
--- /dev/null
+++ b/Engine/Plugins/Mutable/Source/MutableRuntime/Private/MuR/External/FloatAdapter.cpp
@@ -0,0 +1,34 @@
+﻿// Copyright Epic Games, Inc. All Rights Reserved.
+
+#include "MuR/External/FloatAdapter.h"
+
+
+namespace UE::Mutable
+{
+	using namespace Private;
+	
+
+	FFloatAdapter::FFloatAdapter(const FFloatAdapter& Other)
+	{
+		Value = Other.Value;
+	}
+
+	
+	FFloatAdapter& FFloatAdapter::operator=(const FFloatAdapter& Other)
+	{
+		Value = Other.Value;
+		return *this;
+	}
+
+
+	float FFloatAdapter::GetValue() const
+	{
+		return Value;
+	}
+
+
+	void FFloatAdapter::SetValue(float InValue)
+	{
+		Value = InValue;
+	}
+}
diff --git a/Engine/Plugins/Mutable/Source/MutableRuntime/Private/MuR/External/MaterialAdapter.cpp b/Engine/Plugins/Mutable/Source/MutableRuntime/Private/MuR/External/MaterialAdapter.cpp
new file mode 100644
--- /dev/null
+++ b/Engine/Plugins/Mutable/Source/MutableRuntime/Private/MuR/External/MaterialAdapter.cpp
@@ -0,0 +1,121 @@
+﻿// Copyright Epic Games, Inc. All Rights Reserved.
+
+#include "MuR/External/MaterialAdapter.h"
+
+#include "MuR/Material.h"
+#include "MuR/External/FloatAdapter.h"
+#include "MuR/External/TextureAdapter.h"
+#include "MuR/External/VectorAdapter.h"
+
+
+namespace UE::Mutable
+{
+	using namespace Private;
+
+
+	FMaterialAdapter::FMaterialAdapter()
+	{
+		Material = MakeShared<FMaterial>();
+	}
+	
+
+	FMaterialAdapter::FMaterialAdapter(const FMaterialAdapter& Other)
+	{
+		*Material = *Other.Material;
+	}
+
+	
+	FMaterialAdapter& FMaterialAdapter::operator=(const FMaterialAdapter& Other)
+	{
+		*Material = *Other.Material;
+		return *this;
+	}
+
+
+	TSharedPtr<const FFloatAdapter> FMaterialAdapter::GetFloat(const FName& ParameterName) const
+	{
+		float* Result = Material->ScalarParameters.Find(ParameterName);
+		if (!Result)
+		{
+			return nullptr;
+		}
+
+		TSharedPtr<FFloatAdapter> Adapter = MakeShared<FFloatAdapter>();
+		Adapter->Value = *Result;
+		
+		return Adapter;
+	}
+
+
+	bool FMaterialAdapter::SetFloat(const FName& ParameterName, const TSharedPtr<const FFloatAdapter>& Adapter)
+	{
+		if (!Adapter)
+		{
+			return false;
+		}
+
+		Material->ScalarParameters.Add(ParameterName, Adapter->GetValue());
+		return true;
+	}
+
+
+	TSharedPtr<const FVectorAdapter> FMaterialAdapter::GetVector(const FName& ParameterName) const
+	{
+		FVector4f* Result = Material->ColorParameters.Find(ParameterName);
+		if (!Result)
+		{
+			return nullptr;
+		}
+
+		TSharedPtr<FVectorAdapter> Adapter = MakeShared<FVectorAdapter>();
+		Adapter->Value = *Result;
+		
+		return Adapter;
+	}
+
+
+	bool FMaterialAdapter::SetVector(const FName& ParameterName, const TSharedPtr<const FVectorAdapter>& Adapter)
+	{
+		if (!Adapter)
+		{
+			return false;
+		}
+		
+		Material->ColorParameters.Add(ParameterName, Adapter->GetValue());
+		return true;
+	}
+
+
+	TSharedPtr<const FTextureAdapter> FMaterialAdapter::GetTexture(const FName& ParameterName) const
+	{
+		TVariant<OP::ADDRESS, TSharedPtr<const FImage>>* Result = Material->ImageParameters.Find(ParameterName);
+		if (!Result)
+		{
+			return nullptr;
+		}
+
+		if (ensure(Result->IsType<TSharedPtr<const FImage>>()))
+		{
+			TSharedPtr<FTextureAdapter> Adapter = MakeShared<FTextureAdapter>(); 
+			Adapter->Image = ConstCastSharedPtr<FImage>(Result->Get<TSharedPtr<const FImage>>()); // TODO GMT If the image is lazy what should we do? Force eval before calling an UN_EXTERNAL?
+
+			return Adapter;
+		}
+		else
+		{
+			return nullptr;
+		}
+	}
+
+	
+	void FMaterialAdapter::SetTexture(const FName& ParameterName, const TSharedPtr<FTextureAdapter>& Adapter)
+	{
+		if (Adapter)
+		{
+			TVariant<OP::ADDRESS, TSharedPtr<const FImage>> Image;
+			Image.Set<TSharedPtr<const FImage>>(Adapter->Image);
+			
+			Material->ImageParameters.Add(ParameterName, Image);
+		}
+	}
+}
diff --git a/Engine/Plugins/Mutable/Source/MutableRuntime/Private/MuR/External/MeshAdapter.cpp b/Engine/Plugins/Mutable/Source/MutableRuntime/Private/MuR/External/MeshAdapter.cpp
new file mode 100644
--- /dev/null
+++ b/Engine/Plugins/Mutable/Source/MutableRuntime/Private/MuR/External/MeshAdapter.cpp
@@ -0,0 +1,54 @@
+﻿// Copyright Epic Games, Inc. All Rights Reserved.
+
+#include "MuR/External/MeshAdapter.h"
+
+#include "MuR/Mesh.h"
+#include "MuR/MeshPrivate.h"
+
+
+namespace UE::Mutable
+{
+	using namespace Private;
+
+
+	FMeshAdapter::FMeshAdapter()
+	{
+		Mesh = MakeShared<FMesh>();
+	}
+	
+
+	FMeshAdapter::FMeshAdapter(const FMeshAdapter& Other)
+	{
+		Mesh->CopyFrom(*Other.Mesh);
+	}
+
+	
+	FMeshAdapter& FMeshAdapter::operator=(const FMeshAdapter& Other)
+	{
+		Mesh->CopyFrom(*Other.Mesh);
+		return *this;
+	}
+
+
+	int32 FMeshAdapter::GetNumVertices() const
+	{
+		return Mesh->GetVertexCount();
+	}
+
+	FVector3f FMeshAdapter::GetVertex(int32 Index) const
+	{
+		MeshBufferIterator<EMeshBufferFormat::Float32, float, 3> BasePositionIter(Mesh->VertexBuffers, EMeshBufferSemantic::Position, 0);
+		BasePositionIter += Index;
+
+		return BasePositionIter.GetAsVec3f();
+	}
+
+	
+	void FMeshAdapter::SetVertex(int32 Index, const FVector3f& Vertex)
+	{
+		MeshBufferIterator<EMeshBufferFormat::Float32, float, 3> BasePositionIter(Mesh->VertexBuffers, EMeshBufferSemantic::Position, 0);
+		BasePositionIter += Index;
+
+		BasePositionIter.SetFromVec3f(Vertex);
+	}
+}
diff --git a/Engine/Plugins/Mutable/Source/MutableRuntime/Private/MuR/External/Operation.cpp b/Engine/Plugins/Mutable/Source/MutableRuntime/Private/MuR/External/Operation.cpp
new file mode 100644
--- /dev/null
+++ b/Engine/Plugins/Mutable/Source/MutableRuntime/Private/MuR/External/Operation.cpp
@@ -0,0 +1,20 @@
+﻿// Copyright Epic Games, Inc. All Rights Reserved.
+
+#include "MuR/External/Operation.h"
+
+#include "StructUtils/InstancedStruct.h"
+
+
+namespace UE::Mutable
+{
+	TSharedPtr<const FInstancedStruct> FContext::GetInput(const FName& OperantName)
+	{
+		return Inputs.FindAndRemoveChecked(OperantName); // If not found it means that the input was already requested.
+	}
+
+
+	TSharedPtr<const FInstancedStruct>& FContext::GetOutput()
+	{
+		return Output;
+	}
+}
diff --git a/Engine/Plugins/Mutable/Source/MutableRuntime/Private/MuR/External/TextureAdapter.cpp b/Engine/Plugins/Mutable/Source/MutableRuntime/Private/MuR/External/TextureAdapter.cpp
new file mode 100644
--- /dev/null
+++ b/Engine/Plugins/Mutable/Source/MutableRuntime/Private/MuR/External/TextureAdapter.cpp
@@ -0,0 +1,30 @@
+﻿// Copyright Epic Games, Inc. All Rights Reserved.
+
+#include "MuR/External/TextureAdapter.h"
+
+#include "MuR/Image.h"
+
+
+namespace UE::Mutable
+{
+	using namespace Private;
+	
+
+	FTextureAdapter::FTextureAdapter()
+	{
+		Image = MakeShared<FImage>();		
+	}
+
+
+	FTextureAdapter::FTextureAdapter(const FTextureAdapter& Other)
+	{
+		Image->Copy(Other.Image.Get());
+	}
+
+
+	FTextureAdapter& FTextureAdapter::operator=(const FTextureAdapter& Other)
+	{
+		Image->Copy(Other.Image.Get());
+		return *this;
+	}
+}
diff --git a/Engine/Plugins/Mutable/Source/MutableRuntime/Private/MuR/External/VectorAdapter.cpp b/Engine/Plugins/Mutable/Source/MutableRuntime/Private/MuR/External/VectorAdapter.cpp
new file mode 100644
--- /dev/null
+++ b/Engine/Plugins/Mutable/Source/MutableRuntime/Private/MuR/External/VectorAdapter.cpp
@@ -0,0 +1,34 @@
+﻿// Copyright Epic Games, Inc. All Rights Reserved.
+
+#include "MuR/External/VectorAdapter.h"
+
+
+namespace UE::Mutable
+{
+	using namespace Private;
+	
+
+	FVectorAdapter::FVectorAdapter(const FVectorAdapter& Other)
+	{
+		Value = Other.Value;
+	}
+
+	
+	FVectorAdapter& FVectorAdapter::operator=(const FVectorAdapter& Other)
+	{
+		Value = Other.Value;
+		return *this;
+	}
+
+
+	FVector4f FVectorAdapter::GetValue() const
+	{
+		return Value;
+	}
+
+
+	void FVectorAdapter::SetValue(const FVector4f& InValue)
+	{
+		Value = InValue;
+	}
+}
diff --git a/Engine/Plugins/Mutable/Source/MutableRuntime/Private/MuR/Operations.cpp b/Engine/Plugins/Mutable/Source/MutableRuntime/Private/MuR/Operations.cpp
--- a/Engine/Plugins/Mutable/Source/MutableRuntime/Private/MuR/Operations.cpp
+++ b/Engine/Plugins/Mutable/Source/MutableRuntime/Private/MuR/Operations.cpp
@@ -5,6 +5,7 @@
 
 #include "MuR/Image.h"
 #include "MuR/ModelPrivate.h"
+#include "MuR/External/Operation.h"
 
 
 namespace UE::Mutable::Private
@@ -152,6 +153,8 @@ namespace UE::Mutable::Private
 		{ EDataType::Layout,		},	// LA_MERGE
 		{ EDataType::Layout,		},	// LA_REMOVEBLOCKS
 		{ EDataType::Layout,		},	// LA_FROMMESH
+
+		{ EDataType::Unknown,	},	// UN_EXTERNAL
 	};
 
     // clang-format on
@@ -905,7 +908,26 @@ namespace UE::Mutable::Private
 			f(args.Mesh);
 			break;
 		}
+        
+        case EOpType::UN_EXTERNAL:
+        {
+        	const uint8* Data = program.GetOpArgsPointer(at);
+
+        	OP::ExternalArgs Args = program.GetOpArgs<OP::ExternalArgs>(at);
+        	Data += sizeof(OP::ExternalArgs);
+        		
+        	for (int32 Index = 0; Index < Args.NumOperants; ++Index)
+        	{
+        		OP::ADDRESS OperantAt;
+        		FMemory::Memcpy(&OperantAt, Data, sizeof(OP::ADDRESS));
+        		Data += sizeof(OP::ADDRESS);
 
+        		f(OperantAt);
+        	}
+
+       		break;
+        }
+        	
         default:
 			check( false );
             break;
diff --git a/Engine/Plugins/Mutable/Source/MutableRuntime/Public/MuR/External/FloatAdapter.h b/Engine/Plugins/Mutable/Source/MutableRuntime/Public/MuR/External/FloatAdapter.h
new file mode 100644
--- /dev/null
+++ b/Engine/Plugins/Mutable/Source/MutableRuntime/Public/MuR/External/FloatAdapter.h
@@ -0,0 +1,46 @@
+﻿// Copyright Epic Games, Inc. All Rights Reserved.
+
+#pragma once
+
+#include "MuR/Image.h"
+
+#include "FloatAdapter.generated.h"
+
+namespace UE::Mutable
+{
+	struct FMaterialAdapter;
+
+	namespace Private
+	{
+		class FMesh;
+		class CodeRunner;
+	}
+
+
+	USTRUCT()
+	struct FFloatAdapter
+	{
+		GENERATED_BODY()
+		
+		friend Private::CodeRunner;
+		friend FMaterialAdapter;
+
+		FFloatAdapter() = default;
+		
+		FFloatAdapter(const FFloatAdapter& Other);
+		
+		FFloatAdapter(FFloatAdapter&& Other) = delete;
+
+		FFloatAdapter& operator=(const FFloatAdapter& Other);
+		
+		FFloatAdapter& operator=(FFloatAdapter&& Other) = delete;
+
+		// Mesh interface
+		MUTABLERUNTIME_API float GetValue() const;
+
+		MUTABLERUNTIME_API void SetValue(float Value);
+	
+	private:
+		float Value = false;
+	};
+}
diff --git a/Engine/Plugins/Mutable/Source/MutableRuntime/Public/MuR/External/MaterialAdapter.h b/Engine/Plugins/Mutable/Source/MutableRuntime/Public/MuR/External/MaterialAdapter.h
new file mode 100644
--- /dev/null
+++ b/Engine/Plugins/Mutable/Source/MutableRuntime/Public/MuR/External/MaterialAdapter.h
@@ -0,0 +1,54 @@
+﻿// Copyright Epic Games, Inc. All Rights Reserved.
+
+#pragma once
+
+#include "MuR/Image.h"
+
+#include "MaterialAdapter.generated.h"
+
+namespace UE::Mutable
+{
+	struct FFloatAdapter;
+	struct FVectorAdapter;
+	struct FTextureAdapter;
+
+	namespace Private
+	{
+		class FMaterial;
+		class CodeRunner;
+	}
+
+	
+	USTRUCT()
+	struct FMaterialAdapter
+	{
+		GENERATED_BODY()
+		
+		friend Private::CodeRunner;
+
+		FMaterialAdapter();
+		
+		FMaterialAdapter(const FMaterialAdapter& Other);
+		
+		FMaterialAdapter(FMaterialAdapter&& Other) = delete;
+
+		FMaterialAdapter& operator=(const FMaterialAdapter& Other);
+		
+		FMaterialAdapter& operator=(FMaterialAdapter&& Other) = delete;
+		
+		TSharedPtr<const FFloatAdapter> GetFloat(const FName& ParameterName) const;
+
+		bool SetFloat(const FName& ParameterName, const TSharedPtr<const FFloatAdapter>& Adapter);
+		
+		TSharedPtr<const FVectorAdapter> GetVector(const FName& ParameterName) const;
+		
+		bool SetVector(const FName& ParameterName, const TSharedPtr<const FVectorAdapter>& Adapter);
+		
+		TSharedPtr<const FTextureAdapter> GetTexture(const FName& ParameterName) const;
+
+		void SetTexture(const FName& ParameterName, const TSharedPtr<FTextureAdapter>& Adapter);
+	
+	private:
+		TSharedPtr<Private::FMaterial> Material;
+	};
+}
diff --git a/Engine/Plugins/Mutable/Source/MutableRuntime/Public/MuR/External/MeshAdapter.h b/Engine/Plugins/Mutable/Source/MutableRuntime/Public/MuR/External/MeshAdapter.h
new file mode 100644
--- /dev/null
+++ b/Engine/Plugins/Mutable/Source/MutableRuntime/Public/MuR/External/MeshAdapter.h
@@ -0,0 +1,45 @@
+﻿// Copyright Epic Games, Inc. All Rights Reserved.
+
+#pragma once
+
+#include "MuR/Image.h"
+
+#include "MeshAdapter.generated.h"
+
+namespace UE::Mutable
+{
+	namespace Private
+	{
+		class FMesh;
+		class CodeRunner;
+	}
+
+
+	USTRUCT()
+	struct FMeshAdapter
+	{
+		GENERATED_BODY()
+		
+		friend Private::CodeRunner;
+
+		FMeshAdapter();
+		
+		FMeshAdapter(const FMeshAdapter& Other);
+		
+		FMeshAdapter(FMeshAdapter&& Other) = delete;
+
+		FMeshAdapter& operator=(const FMeshAdapter& Other);
+		
+		FMeshAdapter& operator=(FMeshAdapter&& Other) = delete;
+
+		// Mesh interface
+		MUTABLERUNTIME_API int32 GetNumVertices() const;
+
+		MUTABLERUNTIME_API FVector3f GetVertex(int32 Index) const;
+
+		MUTABLERUNTIME_API void SetVertex(int32 Index, const FVector3f& Vertex);
+
+	private:
+		TSharedPtr<Private::FMesh> Mesh;
+	};
+}
diff --git a/Engine/Plugins/Mutable/Source/MutableRuntime/Public/MuR/External/Operation.h b/Engine/Plugins/Mutable/Source/MutableRuntime/Public/MuR/External/Operation.h
new file mode 100644
--- /dev/null
+++ b/Engine/Plugins/Mutable/Source/MutableRuntime/Public/MuR/External/Operation.h
@@ -0,0 +1,89 @@
+﻿// Copyright Epic Games, Inc. All Rights Reserved.
+
+#pragma once
+
+#include "StructUtils/InstancedStruct.h"
+
+#include "Operation.generated.h"
+
+
+namespace UE::Mutable
+{
+	namespace Private
+	{
+		class CodeRunner;
+	}
+
+
+	/** Copy (or move) a value to remove it's const.
+	 *
+	 * Instead of copy, a value will be moved if it is uniquely referenced. A value may not be unique due to:
+	 * - Being being used in more than one subsequent input.
+	 * - User holding additional references to it. Use MoveTemp whenever possible. */
+	template<class T>
+	TSharedPtr<T> CopyOrMove(const TSharedPtr<const T>&& Source)
+	{
+		if (Source.IsUnique())
+		{
+			return ConstCastSharedPtr<T>(Source);
+		}
+		else
+		{
+			return MakeShared<T>(*Source);
+		}
+	}
+	
+
+	/** Context of the External Operation. */
+	class FContext
+	{
+		friend Private::CodeRunner;
+		
+	public:
+		/** Returns and removes the input from the Context.
+		 *
+		 * Once called the owner is the requester. This allows CopyOrMove to move (see CopyOrMove for details). */
+		MUTABLERUNTIME_API TSharedPtr<const FInstancedStruct> GetInput(const FName& OperantName);
+
+		/** Return the pointer to store the ouput. */
+		MUTABLERUNTIME_API TSharedPtr<const FInstancedStruct>& GetOutput();
+		
+	private:
+		TMap<const FName, TSharedPtr<const FInstancedStruct>> Inputs;
+
+		TSharedPtr<const FInstancedStruct> Output = nullptr;
+	};
+	
+
+	/** External Operation Base class.
+	 *
+	 * All External Operations must inherit from this class.
+	 *
+	 * Due to Mutable being a Pure Functional VM, all operations must not store information (no members, globals, statics...).
+	 * Furthermore, due to cache, given the same inputs an operation must always produce the same output (no randomness).
+	 *
+	 * An operation can have UProperty members but they are considered constants (they are copied when compiling and must not be modified at runtime).
+	 * UProperties marked as Visible/EditAnywhere are automatically shown in the Node Details panel. */
+	USTRUCT()
+	struct FExternalOperation
+	{
+		GENERATED_BODY()
+		
+		virtual ~FExternalOperation() = default;
+
+		/** Name of the operation. */
+		virtual FName GetName() const { return {}; }
+
+		/** Version of the operation. If it differs from the compiled one, it will automatically trigger a compilation. */
+		virtual uint32 GetVersion() const { return 0; }
+		
+		/** Static. It must always return the same inputs. */
+		virtual TArray<TPair<FName, const UScriptStruct*>> GetInputs() const { return {}; }
+
+		/** Static. It must always return the same output. */
+		virtual TPair<FName, const UScriptStruct*> GetOutput() const { return {}; }
+
+		/** This function must be pure functional (same inputs, same output). */
+		virtual void Evaluate(FContext& Context) const {}
+	};
+}
\ No newline at end of file
diff --git a/Engine/Plugins/Mutable/Source/MutableRuntime/Public/MuR/External/TextureAdapter.h b/Engine/Plugins/Mutable/Source/MutableRuntime/Public/MuR/External/TextureAdapter.h
new file mode 100644
--- /dev/null
+++ b/Engine/Plugins/Mutable/Source/MutableRuntime/Public/MuR/External/TextureAdapter.h
@@ -0,0 +1,41 @@
+﻿// Copyright Epic Games, Inc. All Rights Reserved.
+
+#pragma once
+
+
+#include "TextureAdapter.generated.h"
+
+
+namespace UE::Mutable
+{
+	struct FMaterialAdapter;
+
+	namespace Private
+	{
+		class FImage;
+		class CodeRunner;
+	}
+
+
+	USTRUCT()
+	struct FTextureAdapter
+	{
+		GENERATED_BODY()
+		
+		friend Private::CodeRunner;
+		friend FMaterialAdapter;
+
+		FTextureAdapter();
+		
+		FTextureAdapter(const FTextureAdapter& Other);
+		
+		FTextureAdapter(FTextureAdapter&& Other) = delete;
+
+		FTextureAdapter& operator=(const FTextureAdapter& Other);
+		
+		FTextureAdapter& operator=(FTextureAdapter&& Other) = delete;
+		
+	private:
+		TSharedPtr<Private::FImage> Image;
+	};
+}
diff --git a/Engine/Plugins/Mutable/Source/MutableRuntime/Public/MuR/External/VectorAdapter.h b/Engine/Plugins/Mutable/Source/MutableRuntime/Public/MuR/External/VectorAdapter.h
new file mode 100644
--- /dev/null
+++ b/Engine/Plugins/Mutable/Source/MutableRuntime/Public/MuR/External/VectorAdapter.h
@@ -0,0 +1,46 @@
+﻿// Copyright Epic Games, Inc. All Rights Reserved.
+
+#pragma once
+
+#include "MuR/Image.h"
+
+#include "VectorAdapter.generated.h"
+
+namespace UE::Mutable
+{
+	struct FMaterialAdapter;
+
+	namespace Private
+	{
+		class FMesh;
+		class CodeRunner;
+	}
+
+
+	USTRUCT()
+	struct FVectorAdapter
+	{
+		GENERATED_BODY()
+		
+		friend Private::CodeRunner;
+		friend FMaterialAdapter;
+
+		FVectorAdapter() = default;
+		
+		FVectorAdapter(const FVectorAdapter& Other);
+		
+		FVectorAdapter(FVectorAdapter&& Other) = delete;
+
+		FVectorAdapter& operator=(const FVectorAdapter& Other);
+		
+		FVectorAdapter& operator=(FVectorAdapter&& Other) = delete;
+
+		// Mesh interface
+		MUTABLERUNTIME_API FVector4f GetValue() const;
+
+		MUTABLERUNTIME_API void SetValue(const FVector4f& Value);
+	
+	private:
+		FVector4f Value = FVector4f::Zero();
+	};
+}
diff --git a/Engine/Plugins/Mutable/Source/MutableTools/Internal/MuT/AST.h b/Engine/Plugins/Mutable/Source/MutableTools/Internal/MuT/AST.h
--- a/Engine/Plugins/Mutable/Source/MutableTools/Internal/MuT/AST.h
+++ b/Engine/Plugins/Mutable/Source/MutableTools/Internal/MuT/AST.h
@@ -309,6 +309,8 @@ namespace UE::Mutable::Private
 		};
 
 		FAdditionalData AdditionalData;
+
+		TMap<OP::ADDRESS, FInstancedStruct>* ExternalOperations = nullptr;
 	};
 
 
diff --git a/Engine/Plugins/Mutable/Source/MutableTools/Internal/MuT/Compiler.h b/Engine/Plugins/Mutable/Source/MutableTools/Internal/MuT/Compiler.h
--- a/Engine/Plugins/Mutable/Source/MutableTools/Internal/MuT/Compiler.h
+++ b/Engine/Plugins/Mutable/Source/MutableTools/Internal/MuT/Compiler.h
@@ -9,6 +9,7 @@
 #include "MuR/System.h"
 #include "Templates/SharedPointer.h"
 #include "HAL/PlatformMath.h"
+#include "StructUtils/InstancedStruct.h"
 #include "Tasks/Task.h"
 
 #define UE_API MUTABLETOOLS_API
@@ -161,7 +162,7 @@ namespace UE::Mutable::Private
         UE_API Compiler( Ptr<CompilerOptions> Options, TFunction<void()>& InWaitCallback );
 
         //! Compile the expression into a run-time model.
-		UE_API TSharedPtr<FModel> Compile( const Ptr<Node>& pNode );
+		UE_API TSharedPtr<FModel> Compile( const Ptr<Node>& pNode, TMap<OP::ADDRESS, FInstancedStruct>& ExternalOperations);
 
         //! Return the log of messages of all the compile operations executed so far.
         UE_API TSharedPtr<FErrorLog> GetLog() const;
diff --git a/Engine/Plugins/Mutable/Source/MutableTools/Internal/MuT/Node.h b/Engine/Plugins/Mutable/Source/MutableTools/Internal/MuT/Node.h
--- a/Engine/Plugins/Mutable/Source/MutableTools/Internal/MuT/Node.h
+++ b/Engine/Plugins/Mutable/Source/MutableTools/Internal/MuT/Node.h
@@ -49,6 +49,7 @@ namespace UE::Mutable::Private
 			MeshReshape,
 			MeshClipDeform,
 			MeshParameter,
+			MeshExternal,
 
 			Image,
 			ImageConstant,
@@ -76,6 +77,7 @@ namespace UE::Mutable::Private
 			ImageNormalComposite,
 			ImageTransform,
 			ImageFromMaterialParameter,
+			ImageExternal,
 
 			Bool,
 			BoolConstant,
@@ -94,6 +96,7 @@ namespace UE::Mutable::Private
 			ColorSwitch,
 			ColorVariation,
 			ColorLinearToSRGB,
+			ColorExternal,
 
 			Scalar,
 			ScalarConstant,
@@ -104,6 +107,7 @@ namespace UE::Mutable::Private
 			ScalarArithmeticOperation,
 			ScalarVariation,
 			ScalarTable,
+			ScalarExternal,
 
 			String,
 			StringConstant,
@@ -147,6 +151,8 @@ namespace UE::Mutable::Private
 			ModifierTransformInMesh,
 			ModifierTransformWithBone,
 
+			External,
+			
 			ExtensionData,
 			ExtensionDataConstant,
 			ExtensionDataSwitch,
@@ -162,6 +168,7 @@ namespace UE::Mutable::Private
 			MaterialSwitch,
 			MaterialVariation,
 			MaterialParameter,
+			MaterialExternal,
 			
 			ImageMaterialBreak,
 			ScalarMaterialBreak,
diff --git a/Engine/Plugins/Mutable/Source/MutableTools/Internal/MuT/NodeColourExternal.h b/Engine/Plugins/Mutable/Source/MutableTools/Internal/MuT/NodeColourExternal.h
new file mode 100644
--- /dev/null
+++ b/Engine/Plugins/Mutable/Source/MutableTools/Internal/MuT/NodeColourExternal.h
@@ -0,0 +1,33 @@
+// Copyright Epic Games, Inc. All Rights Reserved.
+
+#pragma once
+
+#include "NodeColour.h"
+
+#include "MuR/Ptr.h"
+#include "MuT/NodeExternal.h"
+
+#define UE_API MUTABLETOOLS_API
+
+namespace UE::Mutable::Private
+{
+	class NodeColourExternal : public NodeColour
+	{
+	public:
+		// Node Interface
+		virtual const FNodeType* GetType() const override { return &StaticType; }
+		static const FNodeType* GetStaticType() { return &StaticType; }
+
+	protected:
+		// Forbidden. Manage with the Ptr<> template.
+		virtual ~NodeColourExternal() override {}
+
+	private:
+		static UE_API FNodeType StaticType;
+
+	public:
+		Ptr<NodeExternal> Node;
+	};
+}
+
+#undef UE_API
diff --git a/Engine/Plugins/Mutable/Source/MutableTools/Internal/MuT/NodeExternal.h b/Engine/Plugins/Mutable/Source/MutableTools/Internal/MuT/NodeExternal.h
new file mode 100644
--- /dev/null
+++ b/Engine/Plugins/Mutable/Source/MutableTools/Internal/MuT/NodeExternal.h
@@ -0,0 +1,36 @@
+﻿// Copyright Epic Games, Inc. All Rights Reserved.
+
+#pragma once
+
+#include "MuR/Ptr.h"
+#include "MuT/Node.h"
+#include "StructUtils/InstancedStruct.h"
+
+#define UE_API MUTABLETOOLS_API
+
+
+namespace UE::Mutable::Private
+{
+	class NodeExternal : public Node
+	{
+	public:
+		// Node interface
+		virtual const FNodeType* GetType() const override { return GetStaticType(); }
+		static const FNodeType* GetStaticType() { return &StaticType; }
+
+	protected:
+		/** Forbidden. Manage with the Ptr<> template. */
+		virtual ~NodeExternal() override {}
+
+	private:
+		static UE_API FNodeType StaticType;
+
+	public:
+		FInstancedStruct OperationInstancedStruct;
+
+		TArray<Ptr<Node>> Inputs;
+	};
+}
+
+
+#undef UE_API
diff --git a/Engine/Plugins/Mutable/Source/MutableTools/Internal/MuT/NodeImageExternal.h b/Engine/Plugins/Mutable/Source/MutableTools/Internal/MuT/NodeImageExternal.h
new file mode 100644
--- /dev/null
+++ b/Engine/Plugins/Mutable/Source/MutableTools/Internal/MuT/NodeImageExternal.h
@@ -0,0 +1,33 @@
+// Copyright Epic Games, Inc. All Rights Reserved.
+
+#pragma once
+
+#include "NodeImage.h"
+
+#include "MuR/Ptr.h"
+#include "MuT/NodeExternal.h"
+
+#define UE_API MUTABLETOOLS_API
+
+namespace UE::Mutable::Private
+{
+	class NodeImageExternal : public NodeImage
+	{
+	public:
+		// Node Interface
+		virtual const FNodeType* GetType() const override { return &StaticType; }
+		static const FNodeType* GetStaticType() { return &StaticType; }
+
+	protected:
+		// Forbidden. Manage with the Ptr<> template.
+		virtual ~NodeImageExternal() override {}
+
+	private:
+		static UE_API FNodeType StaticType;
+
+	public:
+		Ptr<NodeExternal> Node;
+	};
+}
+
+#undef UE_API
diff --git a/Engine/Plugins/Mutable/Source/MutableTools/Internal/MuT/NodeMaterialExternal.h b/Engine/Plugins/Mutable/Source/MutableTools/Internal/MuT/NodeMaterialExternal.h
new file mode 100644
--- /dev/null
+++ b/Engine/Plugins/Mutable/Source/MutableTools/Internal/MuT/NodeMaterialExternal.h
@@ -0,0 +1,33 @@
+﻿// Copyright Epic Games, Inc. All Rights Reserved.
+
+#pragma once
+
+#include "NodeMaterial.h"
+
+#include "MuR/Ptr.h"
+#include "MuT/NodeExternal.h"
+
+#define UE_API MUTABLETOOLS_API
+
+namespace UE::Mutable::Private
+{
+	class NodeMaterialExternal : public NodeMaterial
+	{
+	public:
+		// Node Interface
+		virtual const FNodeType* GetType() const override { return &StaticType; }
+		static const FNodeType* GetStaticType() { return &StaticType; }
+
+	protected:
+		// Forbidden. Manage with the Ptr<> template.
+		virtual ~NodeMaterialExternal() override {}
+
+	private:
+		static UE_API FNodeType StaticType;
+
+	public:
+		Ptr<NodeExternal> Node;
+	};
+}
+
+#undef UE_API
diff --git a/Engine/Plugins/Mutable/Source/MutableTools/Internal/MuT/NodeMeshExternal.h b/Engine/Plugins/Mutable/Source/MutableTools/Internal/MuT/NodeMeshExternal.h
new file mode 100644
--- /dev/null
+++ b/Engine/Plugins/Mutable/Source/MutableTools/Internal/MuT/NodeMeshExternal.h
@@ -0,0 +1,33 @@
+// Copyright Epic Games, Inc. All Rights Reserved.
+
+#pragma once
+
+#include "NodeMesh.h"
+
+#include "MuR/Ptr.h"
+#include "MuT/NodeExternal.h"
+
+#define UE_API MUTABLETOOLS_API
+
+namespace UE::Mutable::Private
+{
+	class NodeMeshExternal : public NodeMesh
+	{
+	public:
+		// Node Interface
+		virtual const FNodeType* GetType() const override { return &StaticType; }
+		static const FNodeType* GetStaticType() { return &StaticType; }
+
+	protected:
+		// Forbidden. Manage with the Ptr<> template.
+		virtual ~NodeMeshExternal() override {}
+
+	private:
+		static UE_API FNodeType StaticType;
+
+	public:
+		Ptr<NodeExternal> Node;
+	};
+}
+
+#undef UE_API
diff --git a/Engine/Plugins/Mutable/Source/MutableTools/Internal/MuT/NodeScalarExternal.h b/Engine/Plugins/Mutable/Source/MutableTools/Internal/MuT/NodeScalarExternal.h
new file mode 100644
--- /dev/null
+++ b/Engine/Plugins/Mutable/Source/MutableTools/Internal/MuT/NodeScalarExternal.h
@@ -0,0 +1,33 @@
+// Copyright Epic Games, Inc. All Rights Reserved.
+
+#pragma once
+
+#include "NodeScalar.h"
+
+#include "MuR/Ptr.h"
+#include "MuT/NodeExternal.h"
+
+#define UE_API MUTABLETOOLS_API
+
+namespace UE::Mutable::Private
+{
+	class NodeScalarExternal : public NodeScalar
+	{
+	public:
+		// Node Interface
+		virtual const FNodeType* GetType() const override { return &StaticType; }
+		static const FNodeType* GetStaticType() { return &StaticType; }
+
+	protected:
+		// Forbidden. Manage with the Ptr<> template.
+		virtual ~NodeScalarExternal() override {}
+
+	private:
+		static UE_API FNodeType StaticType;
+
+	public:
+		Ptr<NodeExternal> Node;
+	};
+}
+
+#undef UE_API
diff --git a/Engine/Plugins/Mutable/Source/MutableTools/Private/MuT/ASTOpExternal.cpp b/Engine/Plugins/Mutable/Source/MutableTools/Private/MuT/ASTOpExternal.cpp
new file mode 100644
--- /dev/null
+++ b/Engine/Plugins/Mutable/Source/MutableTools/Private/MuT/ASTOpExternal.cpp
@@ -0,0 +1,105 @@
+﻿// Copyright Epic Games, Inc. All Rights Reserved.
+
+#include "MuT/ASTOpExternal.h"
+
+
+namespace UE::Mutable::Private
+{
+	EOpType ASTOpExternal::GetOpType() const
+	{
+		return EOpType::UN_EXTERNAL;
+	}
+
+
+	uint64 ASTOpExternal::Hash() const
+	{
+		uint64 Result = GetTypeHash(OperationInstancedStruct.GetScriptStruct()); // FInstancedStruct can not be hashed. At least we can hash the type.
+		
+		for (const ASTChild& Input : Inputs)
+		{
+			hash_combine(Result, Input.child().get());
+		}
+		return Result;
+	}
+
+
+	void ASTOpExternal::ForEachChild(const TFunctionRef<void(ASTChild&)> Func)
+	{
+		for (ASTChild& Input : Inputs)
+		{
+			Func(Input);
+		}
+	}
+
+
+	bool ASTOpExternal::IsEqual(const ASTOp& OtherUntyped) const
+	{
+		if (OtherUntyped.GetOpType() == GetOpType())
+		{
+			const ASTOpExternal* Other = static_cast<const ASTOpExternal*>(&OtherUntyped);
+
+			if (OperationInstancedStruct != Other->OperationInstancedStruct)
+			{
+				return false;
+			}
+
+			if (Inputs.Num() != Other->Inputs.Num())
+			{
+				return false;
+			}
+			
+			for (int Index = 0; Index < Inputs.Num(); Index++)
+			{
+				if (Inputs[Index] != Other->Inputs[Index])
+				{
+					return false;
+				}
+			}
+			
+			return true;
+		}
+		
+		return false;
+	}
+
+
+	Ptr<ASTOp> ASTOpExternal::Clone(MapChildFuncRef MapChild) const
+	{
+		Ptr<ASTOpExternal> New = new ASTOpExternal();
+		New->OperationInstancedStruct = OperationInstancedStruct;
+		
+		for (const ASTChild& Input : Inputs)
+		{
+			New->Inputs.Emplace(New, MapChild(Input.child()));
+		}
+		
+		return New;
+	}
+
+
+	void ASTOpExternal::Link(FProgram& Program, FLinkerOptions* Options)
+	{
+		if (!linkedAddress)
+		{
+			OP::ExternalArgs Args;
+			FMemory::Memzero(Args);
+			Args.NumOperants = Inputs.Num();
+
+			linkedAddress = (OP::ADDRESS)Program.OpAddress.Num();
+			Program.OpAddress.Add((uint32_t)Program.ByteCode.Num());
+			
+			AppendCode(Program.ByteCode, EOpType::UN_EXTERNAL);
+			AppendCode(Program.ByteCode, Args);
+
+			for (const ASTChild& Input : Inputs)
+			{
+				AppendCode(Program.ByteCode, Input->linkedAddress);
+			}
+
+			if (Options->ExternalOperations)
+			{
+				Options->ExternalOperations->Add(linkedAddress, OperationInstancedStruct);
+			}
+		}
+	}
+}
diff --git a/Engine/Plugins/Mutable/Source/MutableTools/Private/MuT/ASTOpExternal.h b/Engine/Plugins/Mutable/Source/MutableTools/Private/MuT/ASTOpExternal.h
new file mode 100644
--- /dev/null
+++ b/Engine/Plugins/Mutable/Source/MutableTools/Private/MuT/ASTOpExternal.h
@@ -0,0 +1,31 @@
+﻿// Copyright Epic Games, Inc. All Rights Reserved.
+
+#pragma once
+
+#include "MuR/Image.h"
+#include "MuR/Operations.h"
+#include "MuR/Ptr.h"
+#include "MuT/AST.h"
+#include "StructUtils/InstancedStruct.h"
+
+
+namespace UE::Mutable::Private
+{
+	struct FProgram;
+
+	class ASTOpExternal final : public ASTOp
+	{
+	public:
+		virtual EOpType GetOpType() const override;
+		virtual uint64 Hash() const override;
+		virtual void ForEachChild(const TFunctionRef<void(ASTChild&)>) override;
+		virtual bool IsEqual(const ASTOp&) const override;
+		virtual Ptr<ASTOp> Clone(MapChildFuncRef) const override;
+		virtual void Link(FProgram&, FLinkerOptions*) override;
+
+		FInstancedStruct OperationInstancedStruct;
+
+		TArray<ASTChild> Inputs;
+	};
+}
+
diff --git a/Engine/Plugins/Mutable/Source/MutableTools/Private/MuT/CodeGenerator.h b/Engine/Plugins/Mutable/Source/MutableTools/Private/MuT/CodeGenerator.h
--- a/Engine/Plugins/Mutable/Source/MutableTools/Private/MuT/CodeGenerator.h
+++ b/Engine/Plugins/Mutable/Source/MutableTools/Private/MuT/CodeGenerator.h
@@ -2,6 +2,7 @@
 
 #pragma once
 
+#include "ASTOpExternal.h"
 #include "Containers/Map.h"
 #include "HAL/PlatformMath.h"
 #include "Misc/AssertionMacros.h"
@@ -20,10 +21,12 @@
 #include "MuT/ErrorLog.h"
 #include "MuT/Node.h"
 #include "MuT/NodeBool.h"
+#include "MuT/NodeColourExternal.h"
 #include "MuT/NodeColourSampleImage.h"
 #include "MuT/NodeComponentEdit.h"
 #include "MuT/NodeComponentNew.h"
 #include "MuT/NodeComponentSwitch.h"
+#include "MuT/NodeExternal.h"
 #include "MuT/NodeExtensionData.h"
 #include "MuT/NodeImageProject.h"
 #include "MuT/NodeLOD.h"
@@ -43,6 +46,14 @@
 
 namespace UE::Mutable::Private
 {
+	class NodeMaterialExternal;
+	class NodeScalarExternal;
+}
+
+namespace UE::Mutable::Private
+{
+	class NodeImageExternal;
+	class NodeMeshExternal;
 	class ASTOpParameter;
 	class FLayout;
 	class NodeColourArithmeticOperation;
@@ -540,6 +551,7 @@ namespace UE::Mutable::Private
 		void GenerateImage_Transform(const FImageGenerationOptions&, FImageGenerationResult&, const NodeImageTransform*);
 		void GenerateImage_MaterialBreak(const FImageGenerationOptions&, FImageGenerationResult&, const NodeImageMaterialBreak*);
 		void GenerateImage_FromMaterialParameter(const FImageGenerationOptions&, FImageGenerationResult&, const NodeImageFromMaterialParameter*);
+		void GenerateImage_External(const FImageGenerationOptions&, FImageGenerationResult&, const NodeImageExternal*);
 
     	void GenerateCrop(const FImageGenerationOptions& Options, FImageGenerationResult& Result, const NodeImage&);
 
@@ -598,8 +610,26 @@ namespace UE::Mutable::Private
 		{
 			Ptr<ASTOp> op;
 		};
+    	
+		struct FGeneratedMaterialCacheKey
+		{
+			Ptr<const Node> Node;
+			FMaterialGenerationOptions Options;
+
+			friend FORCEINLINE uint32 GetTypeHash(const FGeneratedMaterialCacheKey& InKey)
+			{
+				uint32 KeyHash = 0;
 
-		typedef TMap<FGeneratedCacheKey, FMaterialGenerationResult> FGeneratedMaterialsMap;
+				KeyHash = HashCombineFast(KeyHash, ::GetTypeHash(InKey.Node.get()));
+				KeyHash = HashCombineFast(KeyHash, GetTypeHash(InKey.Options));
+
+				return KeyHash;
+			}
+
+			FORCEINLINE bool operator==(const FGeneratedMaterialCacheKey& Other) const = default;
+		};
+    	
+		typedef TMap<FGeneratedMaterialCacheKey, FMaterialGenerationResult> FGeneratedMaterialsMap;
 		struct FGeneratedMaterials
 		{
 			UE::FMutex Mutex;
@@ -613,6 +643,7 @@ namespace UE::Mutable::Private
 		void GenerateMaterial_Table(const FMaterialGenerationOptions&, FMaterialGenerationResult& result, const NodeMaterialTable* node);
 		void GenerateMaterial_Variation(const FMaterialGenerationOptions&, FMaterialGenerationResult& result, const NodeMaterialVariation* node);
 		void GenerateMaterial_Parameter(const FMaterialGenerationOptions&, FMaterialGenerationResult& result, const NodeMaterialParameter* node);
+		void GenerateMaterial_External(const FMaterialGenerationOptions&, FMaterialGenerationResult& result, const NodeMaterialExternal* node);
 
 		/** */
 		Ptr<ASTOp> GenerateMissingMaterialCode(const TCHAR* strWhere, const void* ErrorContext);
@@ -783,29 +814,103 @@ namespace UE::Mutable::Private
 		FGenerateLayoutConstantState GenerateLayoutConstantState;
 
 		FMeshTask GenerateMesh(const FMeshGenerationStaticOptions&, FMeshOptionsTask, const Ptr<const NodeMesh>&);
-		FMeshTask GenerateMesh_Constant(const FMeshGenerationStaticOptions&, FMeshOptionsTask, const NodeMeshConstant* );
-		FMeshTask GenerateMesh_Format(const FMeshGenerationStaticOptions&, FMeshOptionsTask, const NodeMeshFormat* );
-		FMeshTask GenerateMesh_Morph(const FMeshGenerationStaticOptions&, FMeshOptionsTask, const NodeMeshMorph* );
-		FMeshTask GenerateMesh_MakeMorph(const FMeshGenerationStaticOptions&, FMeshOptionsTask, const NodeMeshMakeMorph* );
-        FMeshTask GenerateMesh_Fragment(const FMeshGenerationStaticOptions&, FMeshOptionsTask, const NodeMeshFragment* );
-        FMeshTask GenerateMesh_Switch(const FMeshGenerationStaticOptions&, FMeshOptionsTask, const NodeMeshSwitch* );
-        FMeshTask GenerateMesh_Transform(const FMeshGenerationStaticOptions&, FMeshOptionsTask, const NodeMeshTransform* );
-        FMeshTask GenerateMesh_ClipMorphPlane(const FMeshGenerationStaticOptions&, FMeshOptionsTask, const NodeMeshClipMorphPlane* );
-        FMeshTask GenerateMesh_ClipWithMesh(const FMeshGenerationStaticOptions&, FMeshOptionsTask, const NodeMeshClipWithMesh* );
-        FMeshTask GenerateMesh_ApplyPose(const FMeshGenerationStaticOptions&, FMeshOptionsTask, const NodeMeshApplyPose* );
-        FMeshTask GenerateMesh_Variation(const FMeshGenerationStaticOptions&, FMeshOptionsTask, const NodeMeshVariation* );
+		FMeshTask GenerateMesh_Constant(const FMeshGenerationStaticOptions&, FMeshOptionsTask, const NodeMeshConstant*);
+		FMeshTask GenerateMesh_Format(const FMeshGenerationStaticOptions&, FMeshOptionsTask, const NodeMeshFormat*);
+		FMeshTask GenerateMesh_Morph(const FMeshGenerationStaticOptions&, FMeshOptionsTask, const NodeMeshMorph*);
+		FMeshTask GenerateMesh_MakeMorph(const FMeshGenerationStaticOptions&, FMeshOptionsTask, const NodeMeshMakeMorph*);
+        FMeshTask GenerateMesh_Fragment(const FMeshGenerationStaticOptions&, FMeshOptionsTask, const NodeMeshFragment*);
+        FMeshTask GenerateMesh_Switch(const FMeshGenerationStaticOptions&, FMeshOptionsTask, const NodeMeshSwitch*);
+        FMeshTask GenerateMesh_Transform(const FMeshGenerationStaticOptions&, FMeshOptionsTask, const NodeMeshTransform*);
+        FMeshTask GenerateMesh_ClipMorphPlane(const FMeshGenerationStaticOptions&, FMeshOptionsTask, const NodeMeshClipMorphPlane*);
+        FMeshTask GenerateMesh_ClipWithMesh(const FMeshGenerationStaticOptions&, FMeshOptionsTask, const NodeMeshClipWithMesh*);
+        FMeshTask GenerateMesh_ApplyPose(const FMeshGenerationStaticOptions&, FMeshOptionsTask, const NodeMeshApplyPose*);
+        FMeshTask GenerateMesh_Variation(const FMeshGenerationStaticOptions&, FMeshOptionsTask, const NodeMeshVariation*);
 		FMeshTask GenerateMesh_Table(const FMeshGenerationStaticOptions&, FMeshOptionsTask, const NodeMeshTable*);
 		FMeshTask GenerateMesh_Reshape(const FMeshGenerationStaticOptions&, FMeshOptionsTask, const NodeMeshReshape*);
 		FMeshTask GenerateMesh_ClipDeform(const FMeshGenerationStaticOptions&, FMeshOptionsTask, const NodeMeshClipDeform*);
 		FMeshTask GenerateMesh_Parameter(const FMeshGenerationStaticOptions&, FMeshOptionsTask, const NodeMeshParameter*);
+    	FMeshTask GenerateMesh_External(const FMeshGenerationStaticOptions&, FMeshOptionsTask, const NodeMeshExternal*);
 
 		Ptr<ASTOp> GenerateLayoutOpsAndResult(const FMeshGenerationDynamicOptions&, Ptr<ASTOp> LastMeshOp, const TArray<Ptr<NodeLayout>>& OriginalLayouts, uint32 MeshId, FMeshGenerationResult& OutResult);
 
-		//-----------------------------------------------------------------------------------------
+    	//!
+    	TSharedPtr<const FLayout> GenerateLayout(Ptr<const NodeLayout> SourceLayout, uint32 MeshIDPrefix);
 
-		//!
-		TSharedPtr<const FLayout> GenerateLayout(Ptr<const NodeLayout> SourceLayout, uint32 MeshIDPrefix);
+    	//-----------------------------------------------------------------------------------------
 
+    	struct FExternalGenerationOptions
+    	{
+    		// GenerateMesh
+    		TOptional<FMeshGenerationStaticOptions> MeshStaticOptions;
+    		TOptional<FMeshOptionsTask> MeshDynamicOptions;
+
+    		// GenerateImage
+    		TOptional<FImageGenerationOptions> ImageOptions;
+
+    		// GenerateScalar
+    		TOptional<FGenericGenerationOptions> ScalarOptions;
+    		
+    		// GenerateColor
+    		TOptional<FGenericGenerationOptions> ColourOptions;
+
+    		// GenerateMaterial
+    		TOptional<FMaterialGenerationOptions> MaterialOptions;
+
+    		bool operator==(const FExternalGenerationOptions& Other) const
+    		{
+    			return
+    				//MeshStaticOptions == Other.MeshStaticOptions && // Meshes are not cached
+    				//MeshDynamicOptions == Other.MeshDynamicOptions &&
+   					ImageOptions == Other.ImageOptions &&
+ 					ScalarOptions == Other.ScalarOptions &&
+ 					ColourOptions == Other.ColourOptions &&
+ 					MaterialOptions == Other.MaterialOptions;
+    		}
+
+    		friend uint32 GetTypeHash(const FExternalGenerationOptions& InKey)
+    		{
+    			uint32 Hash = 0;
+
+				// Hash = HashCombineFast(Hash, GetTypeHash(InKey.MeshStaticOptions)); // Meshes are not cached
+				// Hash = HashCombineFast(Hash, GetTypeHash(InKey.MeshDynamicOptions));
+				Hash = HashCombineFast(Hash, GetTypeHash(InKey.ImageOptions));
+				Hash = HashCombineFast(Hash, GetTypeHash(InKey.ScalarOptions));
+				Hash = HashCombineFast(Hash, GetTypeHash(InKey.MaterialOptions));
+    			
+    			return Hash;
+    		}
+    	};
+    	
+    	struct FGeneratedExternalCacheKey
+    	{
+    		Ptr<const NodeExternal> Node;
+    		FExternalGenerationOptions Options;
+    		
+    		friend uint32 GetTypeHash(const FGeneratedExternalCacheKey& InKey)
+    		{
+    			uint32 Hash = 0;
+
+    			Hash = HashCombineFast(Hash, ::GetTypeHash(InKey.Node.get()));
+    			Hash = HashCombineFast(Hash, GetTypeHash(InKey.Options));
+    			
+    			return Hash;
+    		}
+
+    		bool operator==(const FGeneratedExternalCacheKey& Other) const = default;
+    	};
+    	
+    	struct FExternalGenerationResult
+    	{
+    		Ptr<ASTOp> Op;
+    	};
+    	
+    	typedef TMap<FGeneratedExternalCacheKey, FExternalGenerationResult> FGeneratedExternalMap;
+    	FGeneratedExternalMap GeneratedExtension;
+
+    	void GenerateExternal(FExternalGenerationResult& OutResult, const FExternalGenerationOptions& Options, const Ptr<const NodeExternal>& ExtensionNode);
+
+		//-----------------------------------------------------------------------------------------
+    	
 		struct FExtensionDataGenerationResult
 		{
 			Ptr<ASTOp> Op;
@@ -889,6 +994,7 @@ namespace UE::Mutable::Private
 		void GenerateScalar_Variation(FScalarGenerationResult&, const FGenericGenerationOptions&, const Ptr<const NodeScalarVariation>&);
 		void GenerateScalar_Table(FScalarGenerationResult&, const FGenericGenerationOptions&, const Ptr<const NodeScalarTable>&);
 		void GenerateScalar_MaterialBreak(FScalarGenerationResult&, const FGenericGenerationOptions&, const Ptr<const NodeScalarMaterialBreak>&);
+		void GenerateScalar_External(FScalarGenerationResult&, const FGenericGenerationOptions&, const Ptr<const NodeScalarExternal>&);
 		Ptr<ASTOp> GenerateMissingScalarCode(const TCHAR* strWhere, float value, const void* errorContext);
 
 		//-----------------------------------------------------------------------------------------
@@ -917,6 +1023,7 @@ namespace UE::Mutable::Private
 		void GenerateColor_ToSRGB(FColorGenerationResult&, const FGenericGenerationOptions&, const Ptr<const NodeColorToSRGB>&);
 		void GenerateColor_Table(FColorGenerationResult&, const FGenericGenerationOptions&, const Ptr<const NodeColourTable>&);
 		void GenerateColor_MaterialBreak(FColorGenerationResult&, const FGenericGenerationOptions&, const Ptr<const NodeColourMaterialBreak>&);
+		void GenerateColor_External(FColorGenerationResult&, const FGenericGenerationOptions&, const Ptr<const NodeColourExternal>&);
 		Ptr<ASTOp> GenerateMissingColourCode(const TCHAR* strWhere, const void* errorContext);
 
 		//-----------------------------------------------------------------------------------------
diff --git a/Engine/Plugins/Mutable/Source/MutableTools/Private/MuT/CodeGenerator_Color.cpp b/Engine/Plugins/Mutable/Source/MutableTools/Private/MuT/CodeGenerator_Color.cpp
--- a/Engine/Plugins/Mutable/Source/MutableTools/Private/MuT/CodeGenerator_Color.cpp
+++ b/Engine/Plugins/Mutable/Source/MutableTools/Private/MuT/CodeGenerator_Color.cpp
@@ -37,6 +37,7 @@
 #include "MuT/NodeColourTable.h"
 #include "MuT/NodeColourVariation.h"
 #include "MuT/NodeColorToSRGB.h"
+#include "MuT/NodeColourExternal.h"
 #include "MuT/NodeImage.h"
 #include "MuT/NodeImageProject.h"
 #include "MuT/NodeMeshTable.h"
@@ -125,6 +126,11 @@ namespace UE::Mutable::Private
 			const NodeColourMaterialBreak* Table = static_cast<const NodeColourMaterialBreak*>(Untyped.get());
 			GenerateColor_MaterialBreak(Result, Options, Table);
 		}
+		else if (Untyped->GetType() == NodeColourExternal::GetStaticType())
+		{
+			const NodeColourExternal* Node = static_cast<const NodeColourExternal*>(Untyped.get());
+			GenerateColor_External(Result, Options, Node);
+		}
 		else
 		{
 			check(false);
@@ -567,6 +573,18 @@ namespace UE::Mutable::Private
 	}
 
 
+	void CodeGenerator::GenerateColor_External(FColorGenerationResult& Result, const FGenericGenerationOptions& InOptions, const Ptr<const NodeColourExternal>& InNode)
+	{
+		FExternalGenerationOptions Options;
+		Options.ColourOptions = InOptions;
+					
+		FExternalGenerationResult ExtensionResult;
+		GenerateExternal(ExtensionResult, Options, InNode->Node);
+
+		Result.op = ExtensionResult.Op;
+	}
+
+
 	//---------------------------------------------------------------------------------------------
 	UE::Mutable::Private::Ptr<ASTOp> CodeGenerator::GenerateMissingColourCode(const TCHAR* StrWhere, const void* ErrorContext)
 	{
diff --git a/Engine/Plugins/Mutable/Source/MutableTools/Private/MuT/CodeGenerator_External.cpp b/Engine/Plugins/Mutable/Source/MutableTools/Private/MuT/CodeGenerator_External.cpp
new file mode 100644
--- /dev/null
+++ b/Engine/Plugins/Mutable/Source/MutableTools/Private/MuT/CodeGenerator_External.cpp
@@ -0,0 +1,117 @@
+// Copyright Epic Games, Inc. All Rights Reserved.
+
+#include "CodeGenerator.h"
+
+
+namespace UE::Mutable::Private
+{
+	void CodeGenerator::GenerateExternal(FExternalGenerationResult& OutResult, const FExternalGenerationOptions& Options, const Ptr<const NodeExternal>& ExtensionNode)
+	{
+		if (!ExtensionNode)
+		{
+			OutResult = FExternalGenerationResult();
+			return;
+		}
+
+		// See if it was already generated
+		FGeneratedExternalCacheKey Key;
+		Key.Node = ExtensionNode.get();
+		Key.Options = Options;
+		
+		FGeneratedExternalMap::ValueType* CachedResult = GeneratedExtension.Find(Key);
+		if (CachedResult)
+		{
+			OutResult = *CachedResult;
+			return;
+		}
+    			
+		Ptr<ASTOpExternal> External = new ASTOpExternal();
+		External->OperationInstancedStruct = ExtensionNode->OperationInstancedStruct;
+
+		FMeshGenerationResult Result;
+		Result.MeshOp = External;
+
+		const int32 NumInputs = ExtensionNode->Inputs.Num();
+    			
+		External->Inputs.Reserve(NumInputs);
+		
+		for (int32 Index = 0; Index < NumInputs; ++Index)
+		{
+			const Ptr<Node>& Input = ExtensionNode->Inputs[Index];
+    				
+			if (Input->GetType()->IsA(NodeMesh::GetStaticType()))
+			{
+				FMeshGenerationStaticOptions MeshStaticOptions = Options.MeshStaticOptions.IsSet() ?
+					Options.MeshStaticOptions.GetValue() :
+					FMeshGenerationStaticOptions(0, 0);
+
+				FMeshOptionsTask MeshDynamicOptions = Options.MeshDynamicOptions.IsSet() ?
+					Options.MeshDynamicOptions.GetValue() :
+					UE::Tasks::MakeCompletedTask<FMeshGenerationDynamicOptions>();
+    					
+				FMeshTask Task = GenerateMesh(MeshStaticOptions, MeshDynamicOptions, static_pointer_cast<NodeMesh>(Input));
+				Task.Wait();
+
+				External->Inputs.Add(ASTChild(External.get() ,Task.GetResult().MeshOp));
+			}
+			else if (Input->GetType()->IsA(NodeImage::GetStaticType()))
+			{
+				FImageGenerationOptions ImageOptions = Options.ImageOptions.IsSet() ?
+					Options.ImageOptions.GetValue() :
+					FImageGenerationOptions(0, 0);
+    					
+				FImageGenerationResult InputResult;
+				GenerateImage(ImageOptions, InputResult, static_pointer_cast<NodeImage>(Input));
+
+				External->Inputs.Add(ASTChild(External.get(), InputResult.op));
+			}
+			else if (Input->GetType()->IsA(NodeScalar::GetStaticType()))
+			{
+				FGenericGenerationOptions ScalarOptions = Options.ScalarOptions.IsSet() ?
+					Options.ScalarOptions.GetValue() :
+					FGenericGenerationOptions();
+    					
+				FScalarGenerationResult InputResult;
+				GenerateScalar(InputResult, ScalarOptions, static_pointer_cast<NodeScalar>(Input));
+
+				External->Inputs.Add(ASTChild(External.get(), InputResult.op));
+			}
+			else if (Input->GetType()->IsA(NodeColour::GetStaticType()))
+			{
+				 FGenericGenerationOptions ColourOptions = Options.ColourOptions.IsSet() ?
+					Options.ColourOptions.GetValue() :
+					FGenericGenerationOptions();
+    					
+				FColorGenerationResult InputResult;
+				GenerateColor(InputResult, ColourOptions, static_pointer_cast<NodeColour>(Input));
+
+				External->Inputs.Add(ASTChild(External.get(), InputResult.op));
+			}
+			else if (Input->GetType()->IsA(NodeMaterial::GetStaticType()))
+			{
+				FMaterialGenerationOptions MaterialOptions = Options.MaterialOptions.IsSet() ?
+				   Options.MaterialOptions.GetValue() :
+				   FMaterialGenerationOptions();
+    					
+				FMaterialGenerationResult InputResult;
+				GenerateMaterial(MaterialOptions, InputResult, static_pointer_cast<NodeMaterial>(Input));
+
+				External->Inputs.Add(ASTChild(External.get(), InputResult.op));
+			}
+			else if (Input->GetType()->IsA(NodeExternal::GetStaticType()))
+			{
+				FExternalGenerationResult InputResult;
+				GenerateExternal(InputResult, Options, static_pointer_cast<NodeExternal>(Input));
+
+				External->Inputs[Index] = InputResult.Op;
+			}
+			else
+			{
+				External->Inputs.Add(ASTChild(External.get()));
+				unimplemented(); // GenerateMutableSourceExtension has compiled an input type that is not supported.
+			}
+		}
+
+		OutResult.Op = External;
+	}
+}
\ No newline at end of file
diff --git a/Engine/Plugins/Mutable/Source/MutableTools/Private/MuT/CodeGenerator_Image.cpp b/Engine/Plugins/Mutable/Source/MutableTools/Private/MuT/CodeGenerator_Image.cpp
--- a/Engine/Plugins/Mutable/Source/MutableTools/Private/MuT/CodeGenerator_Image.cpp
+++ b/Engine/Plugins/Mutable/Source/MutableTools/Private/MuT/CodeGenerator_Image.cpp
@@ -60,6 +60,7 @@
 #include "MuT/NodeImageColourMap.h"
 #include "MuT/NodeImageConditional.h"
 #include "MuT/NodeImageConstant.h"
+#include "MuT/NodeImageExternal.h"
 #include "MuT/NodeImageFormat.h"
 #include "MuT/NodeImageFromMaterialParameter.h"
 #include "MuT/NodeImageInterpolate.h"
@@ -140,6 +141,7 @@ namespace UE::Mutable::Private
 		case Node::EType::ImageTransform: GenerateImage_Transform(Options, Result, static_cast<const NodeImageTransform*>(Node)); break;
 		case Node::EType::ImageMaterialBreak: GenerateImage_MaterialBreak(Options, Result, static_cast<const NodeImageMaterialBreak*>(Node)); break;
 		case Node::EType::ImageFromMaterialParameter: GenerateImage_FromMaterialParameter(Options, Result, static_cast<const NodeImageFromMaterialParameter*>(Node)); break;
+		case Node::EType::ImageExternal: GenerateImage_External(Options, Result, static_cast<const NodeImageExternal*>(Node)); break;
 		default: check(false);
 		}
 
@@ -1921,6 +1923,18 @@ namespace UE::Mutable::Private
 	}
 
 
+	void CodeGenerator::GenerateImage_External(const FImageGenerationOptions& InOptions, FImageGenerationResult& Result, const NodeImageExternal* InNode)
+	{
+		FExternalGenerationOptions Options;
+		Options.ImageOptions = InOptions;
+					
+		FExternalGenerationResult ExtensionResult;
+		GenerateExternal(ExtensionResult, Options, InNode->Node);
+
+		Result.op = ExtensionResult.Op;
+	}
+
+
 	TSharedPtr<FImage> CodeGenerator::GenerateMissingImage(EImageFormat Format)
 	{
 		MUTABLE_CPUPROFILER_SCOPE(GenerateMissingImage);
diff --git a/Engine/Plugins/Mutable/Source/MutableTools/Private/MuT/CodeGenerator_Material.cpp b/Engine/Plugins/Mutable/Source/MutableTools/Private/MuT/CodeGenerator_Material.cpp
--- a/Engine/Plugins/Mutable/Source/MutableTools/Private/MuT/CodeGenerator_Material.cpp
+++ b/Engine/Plugins/Mutable/Source/MutableTools/Private/MuT/CodeGenerator_Material.cpp
@@ -21,6 +21,7 @@
 #include "MuT/NodeScalar.h"
 #include "MuT/NodeMaterial.h"
 #include "MuT/NodeMaterialConstant.h"
+#include "MuT/NodeMaterialExternal.h"
 #include "MuT/NodeMaterialSwitch.h"
 #include "MuT/NodeMaterialVariation.h"
 #include "MuT/NodeMaterialTable.h"
@@ -43,7 +44,7 @@ namespace UE::Mutable::Private
 		}
 
 		// See if it was already generated
-		FGeneratedCacheKey Key;
+		FGeneratedMaterialCacheKey Key;
 		Key.Node = InUntypedNode;
 		Key.Options = InOptions;
 		{
@@ -80,6 +81,10 @@ namespace UE::Mutable::Private
 			GenerateMaterial_Parameter(InOptions, Result, static_cast<const NodeMaterialParameter*>(Node));
 			break;
 
+		case Node::EType::MaterialExternal:
+			GenerateMaterial_External(InOptions, Result, static_cast<const NodeMaterialExternal*>(Node));
+			break;
+
 		default: 
 			check(false);
 		}
@@ -319,6 +324,18 @@ namespace UE::Mutable::Private
 	}
 
 
+	void CodeGenerator::GenerateMaterial_External(const FMaterialGenerationOptions& InOptions, FMaterialGenerationResult& Result, const NodeMaterialExternal* InNode)
+	{
+		FExternalGenerationOptions Options;
+		Options.MaterialOptions = InOptions;
+					
+		FExternalGenerationResult ExtensionResult;
+		GenerateExternal(ExtensionResult, Options, InNode->Node);
+
+		Result.op = ExtensionResult.Op;
+	}
+
+
 	Ptr<ASTOp> CodeGenerator::GenerateMissingMaterialCode(const TCHAR* strWhere, const void* ErrorContext)
 	{
 		// Log an error message
diff --git a/Engine/Plugins/Mutable/Source/MutableTools/Private/MuT/CodeGenerator_Mesh.cpp b/Engine/Plugins/Mutable/Source/MutableTools/Private/MuT/CodeGenerator_Mesh.cpp
--- a/Engine/Plugins/Mutable/Source/MutableTools/Private/MuT/CodeGenerator_Mesh.cpp
+++ b/Engine/Plugins/Mutable/Source/MutableTools/Private/MuT/CodeGenerator_Mesh.cpp
@@ -1,5 +1,6 @@
 // Copyright Epic Games, Inc. All Rights Reserved.
 
+#include "ASTOpExternal.h"
 #include "Containers/Array.h"
 #include "Math/IntPoint.h"
 #include "Math/UnrealMathSSE.h"
@@ -39,6 +40,7 @@
 #include "MuT/CodeGenerator_FirstPass.h"
 #include "MuT/CompilerPrivate.h"
 #include "MuT/ErrorLog.h"
+#include "MuT/NodeExternal.h"
 #include "MuT/NodeImageProject.h"
 #include "MuT/NodeLayout.h"
 #include "MuT/NodeMesh.h"
@@ -47,6 +49,7 @@
 #include "MuT/NodeMeshClipMorphPlane.h"
 #include "MuT/NodeMeshClipWithMesh.h"
 #include "MuT/NodeMeshConstant.h"
+#include "MuT/NodeMeshExternal.h"
 #include "MuT/NodeMeshFormat.h"
 #include "MuT/NodeMeshFragment.h"
 #include "MuT/NodeMeshMakeMorph.h"
@@ -95,6 +98,7 @@ namespace UE::Mutable::Private
 		case Node::EType::MeshReshape: Result = GenerateMesh_Reshape(StaticOptions, InOptions, static_cast<const NodeMeshReshape*>(Node)); break;
 		case Node::EType::MeshClipDeform: Result = GenerateMesh_ClipDeform(StaticOptions, InOptions, static_cast<const NodeMeshClipDeform*>(Node)); break;
 		case Node::EType::MeshParameter: Result = GenerateMesh_Parameter(StaticOptions, InOptions, static_cast<const NodeMeshParameter*>(Node)); break;
+		case Node::EType::MeshExternal: Result = GenerateMesh_External(StaticOptions, InOptions, static_cast<const NodeMeshExternal*>(Node)); break;
 		default: check(false);
 		}
 
@@ -1478,4 +1482,24 @@ namespace UE::Mutable::Private
 		return ApplyModifiersTask;
 	}
 
+
+	FMeshTask CodeGenerator::GenerateMesh_External(const FMeshGenerationStaticOptions& StaticOptions, FMeshOptionsTask InOptions, const NodeMeshExternal* Node)
+	{
+		return UE::Tasks::Launch(TEXT("MutableExtension"),
+			[Node, StaticOptions, InOptions, this]() mutable
+			{
+				FExternalGenerationOptions Options;
+				Options.MeshStaticOptions = StaticOptions;
+				Options.MeshDynamicOptions = InOptions;
+					
+				FExternalGenerationResult ExtensionResult;
+				GenerateExternal(ExtensionResult, Options, Node->Node);
+
+				FMeshGenerationResult MeshResult;
+				MeshResult.MeshOp = ExtensionResult.Op;
+					
+				return MeshResult;
+			}
+		);
+	}
 }
diff --git a/Engine/Plugins/Mutable/Source/MutableTools/Private/MuT/CodeGenerator_Scalar.cpp b/Engine/Plugins/Mutable/Source/MutableTools/Private/MuT/CodeGenerator_Scalar.cpp
--- a/Engine/Plugins/Mutable/Source/MutableTools/Private/MuT/CodeGenerator_Scalar.cpp
+++ b/Engine/Plugins/Mutable/Source/MutableTools/Private/MuT/CodeGenerator_Scalar.cpp
@@ -30,6 +30,7 @@
 #include "MuT/NodeScalarConstant.h"
 #include "MuT/NodeScalarCurve.h"
 #include "MuT/NodeScalarEnumParameter.h"
+#include "MuT/NodeScalarExternal.h"
 #include "MuT/NodeScalarMaterialBreak.h"
 #include "MuT/NodeScalarParameter.h"
 #include "MuT/NodeScalarSwitch.h"
@@ -110,6 +111,11 @@ namespace UE::Mutable::Private
 			const NodeScalarMaterialBreak* Table = static_cast<const NodeScalarMaterialBreak*>(Untyped.get());
 			GenerateScalar_MaterialBreak(result, Options, Table);
 		}
+		else if (Untyped->GetType() == NodeScalarExternal::GetStaticType())
+		{
+			const NodeScalarExternal* Node = static_cast<const NodeScalarExternal*>(Untyped.get());
+			GenerateScalar_External(result, Options, Node);
+		}
 		else
 		{
 			check(false);
@@ -448,6 +454,18 @@ namespace UE::Mutable::Private
 	}
 
 
+	void CodeGenerator::GenerateScalar_External(FScalarGenerationResult& Result, const FGenericGenerationOptions& InOptions, const Ptr<const NodeScalarExternal>& InNode)
+	{
+		FExternalGenerationOptions Options;
+		Options.ScalarOptions = InOptions;
+					
+		FExternalGenerationResult ExtensionResult;
+		GenerateExternal(ExtensionResult, Options, InNode->Node);
+
+		Result.op = ExtensionResult.Op;
+	}
+
+
 	UE::Mutable::Private::Ptr<ASTOp> CodeGenerator::GenerateMissingScalarCode(const TCHAR* strWhere, float value, const void* errorContext)
 	{
 		// Log a warning
diff --git a/Engine/Plugins/Mutable/Source/MutableTools/Private/MuT/Compiler.cpp b/Engine/Plugins/Mutable/Source/MutableTools/Private/MuT/Compiler.cpp
--- a/Engine/Plugins/Mutable/Source/MutableTools/Private/MuT/Compiler.cpp
+++ b/Engine/Plugins/Mutable/Source/MutableTools/Private/MuT/Compiler.cpp
@@ -190,6 +190,8 @@ namespace UE::Mutable::Private
 		{ true,		{ 0, 0, 0, 0, 0, 0, 0, 0, 0, 0, 0, 0, 0, 0, 0, 0, 0, 0, 0, 0, 0, 0, 0, 0, 0, 0, 0, 0, 0 }		},	// LA_MERGE
 		{ true,		{ 0, 0, 0, 0, 0, 0, 0, 0, 0, 0, 0, 0, 0, 0, 0, 0, 0, 0, 0, 0, 0, 0, 0, 0, 0, 0, 0, 0, 0 }		},	// LA_REMOVEBLOCKS
 		{ true,		{ 0, 0, 0, 0, 0, 0, 0, 0, 0, 0, 0, 0, 0, 0, 0, 0, 0, 0, 0, 0, 0, 0, 0, 0, 0, 0, 0, 0, 0 }		},	// LA_FROMMESH
+
+		{ true,		{ 0, 0, 0, 0, 0, 0, 0, 0, 0, 0, 0, 0, 0, 0, 0, 0, 0, 0, 0, 0, 0, 0, 0, 0, 0, 0, 0, 0, 0 }		},	// UN_EXTERNAL
 	};
 
 	// clang-format on
@@ -438,7 +440,7 @@ namespace UE::Mutable::Private
     }
 
     //---------------------------------------------------------------------------------------------
-	TSharedPtr<FModel> Compiler::Compile( const Ptr<Node>& pNode )
+	TSharedPtr<FModel> Compiler::Compile( const Ptr<Node>& pNode, TMap<OP::ADDRESS, FInstancedStruct>& ExternalOperations)
     {
         MUTABLE_CPUPROFILER_SCOPE(Compile);
 
@@ -597,6 +599,8 @@ namespace UE::Mutable::Private
 		FImageOperator ImOp = FImageOperator::GetDefault(m_pD->Options->GetPrivate()->ImageFormatFunc);
 		FLinkerOptions LinkerOptions(ImOp);
 
+		LinkerOptions.ExternalOperations = &ExternalOperations;
+			
 		for(FStateCompilationData& State : States )
         {
 			LinkerOptions.MinTextureResidentMipCount = m_pD->Options->GetPrivate()->MinTextureResidentMipCount;
diff --git a/Engine/Plugins/Mutable/Source/MutableTools/Private/MuT/ErrorLog.cpp b/Engine/Plugins/Mutable/Source/MutableTools/Private/MuT/ErrorLog.cpp
--- a/Engine/Plugins/Mutable/Source/MutableTools/Private/MuT/ErrorLog.cpp
+++ b/Engine/Plugins/Mutable/Source/MutableTools/Private/MuT/ErrorLog.cpp
@@ -334,7 +334,9 @@ namespace UE::Mutable::Private
 		TEXT("LA_PACK          "),
 		TEXT("LA_MERGE         "),
 		TEXT("LA_REMOVEBLOCKS  "),
-		TEXT("LA_FROMMESH	   "),
+		TEXT("LA_FROMMESH      "),
+
+    	TEXT("UN_EXTERNAL      "),
 	};
 
 	static_assert(UE_ARRAY_COUNT(s_opNames) == int32(EOpType::COUNT));
diff --git a/Engine/Plugins/Mutable/Source/MutableTools/Private/MuT/Node.cpp b/Engine/Plugins/Mutable/Source/MutableTools/Private/MuT/Node.cpp
--- a/Engine/Plugins/Mutable/Source/MutableTools/Private/MuT/Node.cpp
+++ b/Engine/Plugins/Mutable/Source/MutableTools/Private/MuT/Node.cpp
@@ -94,16 +94,22 @@
 #include "MuT/NodeStringConstant.h"
 #include "MuT/NodeStringParameter.h"
 #include "MuT/NodeBool.h"
+#include "MuT/NodeColourExternal.h"
+#include "MuT/NodeExternal.h"
+#include "MuT/NodeImageExternal.h"
 #include "MuT/NodeProjector.h"
 #include "MuT/NodeRange.h"
 #include "MuT/NodeRangeFromScalar.h"
 #include "MuT/NodeMaterial.h"
 #include "MuT/NodeImageMaterialBreak.h"
 #include "MuT/NodeMaterialConstant.h"
+#include "MuT/NodeMaterialExternal.h"
 #include "MuT/NodeMaterialTable.h"
 #include "MuT/NodeMaterialSwitch.h"
 #include "MuT/NodeMaterialVariation.h"
 #include "MuT/NodeMaterialParameter.h"
+#include "MuT/NodeMeshExternal.h"
+#include "MuT/NodeScalarExternal.h"
 
 
 namespace UE::Mutable::Private
@@ -137,12 +143,15 @@ namespace UE::Mutable::Private
 	FNodeType NodeScalarEnumParameter::StaticType = FNodeType(Node::EType::ScalarEnumParameter, NodeScalar::GetStaticType());
 	FNodeType NodeScalarTable::StaticType = FNodeType(Node::EType::ScalarTable, NodeScalar::GetStaticType());
 	FNodeType NodeScalarCurve::StaticType = FNodeType(Node::EType::ScalarCurve, NodeScalar::GetStaticType());
+	FNodeType NodeScalarExternal::StaticType = FNodeType(Node::EType::ScalarExternal, NodeScalar::GetStaticType());
 
 	FNodeType NodeSurface::StaticType = FNodeType(Node::EType::Surface, Node::GetStaticType());
 	FNodeType NodeSurfaceNew::StaticType = FNodeType(Node::EType::SurfaceNew, NodeSurface::GetStaticType());
 	FNodeType NodeSurfaceSwitch::StaticType = FNodeType(Node::EType::SurfaceSwitch, NodeSurface::GetStaticType());
 	FNodeType NodeSurfaceVariation::StaticType = FNodeType(Node::EType::SurfaceVariation, NodeSurface::GetStaticType());
 
+	FNodeType NodeExternal::StaticType = FNodeType(Node::EType::External, Node::GetStaticType());
+	
 	FNodeType NodeLOD::StaticType = FNodeType(Node::EType::LOD, Node::GetStaticType());
 	FNodeType NodeExtensionData::StaticType = FNodeType(Node::EType::ExtensionData, Node::GetStaticType());
 	FNodeType NodeExtensionDataConstant::StaticType = FNodeType(Node::EType::ExtensionDataConstant, NodeExtensionData::GetStaticType());
@@ -172,6 +181,7 @@ namespace UE::Mutable::Private
 	FNodeType NodeImageVariation::StaticType = FNodeType(Node::EType::ImageVariation, NodeImage::GetStaticType());
 	FNodeType NodeImageColourMap::StaticType = FNodeType(Node::EType::ImageColorMap, NodeImage::GetStaticType());
 	FNodeType NodeImageFromMaterialParameter::StaticType = FNodeType(Node::EType::ImageFromMaterialParameter, NodeImage::GetStaticType());
+	FNodeType NodeImageExternal::StaticType = FNodeType(Node::EType::ImageExternal, NodeImage::GetStaticType());
 
 	FNodeType NodeColour::StaticType = FNodeType(Node::EType::Color, Node::GetStaticType());
 	FNodeType NodeColourConstant::StaticType = FNodeType(Node::EType::ColorConstant, NodeColour::GetStaticType());
@@ -183,6 +193,7 @@ namespace UE::Mutable::Private
 	FNodeType NodeColourSampleImage::StaticType = FNodeType(Node::EType::ColorSampleImage, NodeColour::GetStaticType());
 	FNodeType NodeColourFromScalars::StaticType = FNodeType(Node::EType::ColorFromScalars, NodeColour::GetStaticType());
 	FNodeType NodeColorToSRGB::StaticType = FNodeType(Node::EType::ColorLinearToSRGB, NodeColour::GetStaticType());
+	FNodeType NodeColourExternal::StaticType = FNodeType(Node::EType::ColorExternal, NodeColour::GetStaticType());
 
 	FNodeType NodeMesh::StaticType = FNodeType(Node::EType::Mesh, Node::GetStaticType());
 	FNodeType NodeMeshConstant::StaticType = FNodeType(Node::EType::MeshConstant, NodeMesh::GetStaticType());
@@ -200,6 +211,7 @@ namespace UE::Mutable::Private
 	FNodeType NodeMeshFormat::StaticType = FNodeType(Node::EType::MeshFormat, NodeMesh::GetStaticType());
 	FNodeType NodeMeshVariation::StaticType = FNodeType(Node::EType::MeshVariation, NodeMesh::GetStaticType());
 	FNodeType NodeMeshTable::StaticType = FNodeType(Node::EType::MeshTable, NodeMesh::GetStaticType());
+	FNodeType NodeMeshExternal::StaticType = FNodeType(Node::EType::MeshExternal, NodeMesh::GetStaticType());
 
 	FNodeType NodeModifier::StaticType = FNodeType(Node::EType::Modifier, Node::GetStaticType());
 	FNodeType NodeModifierMeshClipDeform::StaticType = FNodeType(Node::EType::ModifierMeshClipDeform, NodeModifier::GetStaticType());
@@ -231,6 +243,7 @@ namespace UE::Mutable::Private
 	FNodeType NodeMaterialSwitch::StaticType = FNodeType(Node::EType::MaterialSwitch, NodeMaterial::GetStaticType());
 	FNodeType NodeMaterialVariation::StaticType = FNodeType(Node::EType::MaterialVariation, NodeMaterial::GetStaticType());
 	FNodeType NodeMaterialParameter::StaticType = FNodeType(Node::EType::MaterialParameter, NodeMaterial::GetStaticType());
+	FNodeType NodeMaterialExternal::StaticType = FNodeType(Node::EType::MaterialExternal, NodeMaterial::GetStaticType());
 	
 	FNodeType NodeImageMaterialBreak::StaticType = FNodeType(Node::EType::ImageMaterialBreak, NodeImage::GetStaticType());
 	FNodeType NodeScalarMaterialBreak::StaticType = FNodeType(Node::EType::ScalarMaterialBreak, NodeScalar::GetStaticType());
