[Mutable] Fix ambiguous Ptr in some platforms.

diff --git a/Engine/Plugins/Mutable/Source/CustomizableObjectEditor/Private/MuCOE/GenerateMutableSource/GenerateMutableSourceExternal.cpp b/Engine/Plugins/Mutable/Source/CustomizableObjectEditor/Private/MuCOE/GenerateMutableSource/GenerateMutableSourceExternal.cpp
--- a/Engine/Plugins/Mutable/Source/CustomizableObjectEditor/Private/MuCOE/GenerateMutableSource/GenerateMutableSourceExternal.cpp
+++ b/Engine/Plugins/Mutable/Source/CustomizableObjectEditor/Private/MuCOE/GenerateMutableSource/GenerateMutableSourceExternal.cpp
@@ -23,7 +23,7 @@ using namespace UE::Mutable;
 using namespace UE::Mutable::Private;
 
 
-Ptr<NodeExternal> GenerateMutableSourceExtension(const UEdGraphPin* Pin, FMutableGraphGenerationContext& GenerationContext, const FSourceExtensionOptions& Options)
+UE::Mutable::Private::Ptr<NodeExternal> GenerateMutableSourceExtension(const UEdGraphPin* Pin, FMutableGraphGenerationContext& GenerationContext, const FSourceExtensionOptions& Options)
 {
 	check(Pin)
 	RETURN_ON_CYCLE(*Pin, GenerationContext)
@@ -48,7 +48,7 @@ Ptr<NodeExternal> GenerateMutableSourceExtension(const UEdGraphPin* Pin, FMutabl
 	
 	UCustomizableObjectNode* EditorNode = CastChecked<UCustomizableObjectNode>(Pin->GetOwningNode());
 
-	Ptr<NodeExternal> Result;
+	UE::Mutable::Private::Ptr<NodeExternal> Result;
 	bool bCacheNode = true;
 	
 	if (const UCONodeExternal* TypedNodeExtension = Cast<UCONodeExternal>(EditorNode))
@@ -63,7 +63,7 @@ Ptr<NodeExternal> GenerateMutableSourceExtension(const UEdGraphPin* Pin, FMutabl
 				return;
 			}
 			
-			Ptr<NodeExternal> ExtensionNode = new NodeExternal();
+			UE::Mutable::Private::Ptr<NodeExternal> ExtensionNode = new NodeExternal();
             				
             ExtensionNode->OperationInstancedStruct = TypedNodeExtension->OperationInstancedStruct;
 
@@ -96,32 +96,32 @@ Ptr<NodeExternal> GenerateMutableSourceExtension(const UEdGraphPin* Pin, FMutabl
             
             	if (InputType == FMeshAdapter::StaticStruct())
             	{
-            		Ptr<NodeMesh> SourceNode = GenerateMutableSourceMesh(ConnectedPin, GenerationContext, Options.MeshOptions, Options.bLinkedToExtendMaterial, Options.bOnlyConnectedLOD);
+            		UE::Mutable::Private::Ptr<NodeMesh> SourceNode = GenerateMutableSourceMesh(ConnectedPin, GenerationContext, Options.MeshOptions, Options.bLinkedToExtendMaterial, Options.bOnlyConnectedLOD);
             		ExtensionNode->Inputs.Add(SourceNode);
             	}
             	else if (InputType == FTextureAdapter::StaticStruct())
             	{
-            		Ptr<NodeImage> SourceNode = GenerateMutableSourceImage(ConnectedPin, GenerationContext, Options.ReferenceTextureSize);
+            		UE::Mutable::Private::Ptr<NodeImage> SourceNode = GenerateMutableSourceImage(ConnectedPin, GenerationContext, Options.ReferenceTextureSize);
             		ExtensionNode->Inputs.Add(SourceNode);
             	}
             	else if (InputType == FFloatAdapter::StaticStruct())
             	{
-            		Ptr<NodeScalar> SourceNode = GenerateMutableSourceFloat(ConnectedPin, GenerationContext);
+            		UE::Mutable::Private::Ptr<NodeScalar> SourceNode = GenerateMutableSourceFloat(ConnectedPin, GenerationContext);
             		ExtensionNode->Inputs.Add(SourceNode);
             	}
             	else if (InputType == FVectorAdapter::StaticStruct())
             	{
-            		Ptr<NodeColour> SourceNode = GenerateMutableSourceColor(ConnectedPin, GenerationContext);
+            		UE::Mutable::Private::Ptr<NodeColour> SourceNode = GenerateMutableSourceColor(ConnectedPin, GenerationContext);
             		ExtensionNode->Inputs.Add(SourceNode);
             	}
             	else if (InputType == FMaterialAdapter::StaticStruct())
             	{
-            		Ptr<NodeColour> SourceNode = GenerateMutableSourceColor(ConnectedPin, GenerationContext);
+            		UE::Mutable::Private::Ptr<NodeColour> SourceNode = GenerateMutableSourceColor(ConnectedPin, GenerationContext);
             		ExtensionNode->Inputs.Add(SourceNode);
             	}
             	else
             	{
-            		Ptr<Node> SourceNode = GenerateMutableSourceExtension(ConnectedPin, GenerationContext, Options);
+            		UE::Mutable::Private::Ptr<Node> SourceNode = GenerateMutableSourceExtension(ConnectedPin, GenerationContext, Options);
             		ExtensionNode->Inputs.Add(SourceNode);
             	}	
             }
