[Mutable] Fixed ambiguous Private namespace.

#nrx

diff --git a/Engine/Plugins/Mutable/Source/MutableRuntime/Private/MuR/External/FloatAdapter.cpp b/Engine/Plugins/Mutable/Source/MutableRuntime/Private/MuR/External/FloatAdapter.cpp
--- a/Engine/Plugins/Mutable/Source/MutableRuntime/Private/MuR/External/FloatAdapter.cpp
+++ b/Engine/Plugins/Mutable/Source/MutableRuntime/Private/MuR/External/FloatAdapter.cpp
@@ -5,9 +5,6 @@
 
 namespace UE::Mutable
 {
-	using namespace Private;
-	
-
 	FFloatAdapter::FFloatAdapter(const FFloatAdapter& Other)
 	{
 		Value = Other.Value;
diff --git a/Engine/Plugins/Mutable/Source/MutableRuntime/Private/MuR/External/MaterialAdapter.cpp b/Engine/Plugins/Mutable/Source/MutableRuntime/Private/MuR/External/MaterialAdapter.cpp
--- a/Engine/Plugins/Mutable/Source/MutableRuntime/Private/MuR/External/MaterialAdapter.cpp
+++ b/Engine/Plugins/Mutable/Source/MutableRuntime/Private/MuR/External/MaterialAdapter.cpp
@@ -10,12 +10,9 @@
 
 namespace UE::Mutable
 {
-	using namespace Private;
-
-
 	FMaterialAdapter::FMaterialAdapter()
 	{
-		Material = MakeShared<FMaterial>();
+		Material = MakeShared<Private::FMaterial>();
 	}
 	
 
@@ -88,16 +85,16 @@ namespace UE::Mutable
 
 	TSharedPtr<const FTextureAdapter> FMaterialAdapter::GetTexture(const FName& ParameterName) const
 	{
-		TVariant<OP::ADDRESS, TSharedPtr<const FImage>>* Result = Material->ImageParameters.Find(ParameterName);
+		TVariant<Private::OP::ADDRESS, TSharedPtr<const Private::FImage>>* Result = Material->ImageParameters.Find(ParameterName);
 		if (!Result)
 		{
 			return nullptr;
 		}
 
-		if (ensure(Result->IsType<TSharedPtr<const FImage>>()))
+		if (ensure(Result->IsType<TSharedPtr<const Private::FImage>>()))
 		{
 			TSharedPtr<FTextureAdapter> Adapter = MakeShared<FTextureAdapter>(); 
-			Adapter->Image = ConstCastSharedPtr<FImage>(Result->Get<TSharedPtr<const FImage>>()); // TODO GMT If the image is lazy what should we do? Force eval before calling an UN_EXTERNAL?
+			Adapter->Image = ConstCastSharedPtr<Private::FImage>(Result->Get<TSharedPtr<const Private::FImage>>()); // TODO GMT If the image is lazy what should we do? Force eval before calling an UN_EXTERNAL?
 
 			return Adapter;
 		}
@@ -112,8 +109,8 @@ namespace UE::Mutable
 	{
 		if (Adapter)
 		{
-			TVariant<OP::ADDRESS, TSharedPtr<const FImage>> Image;
-			Image.Set<TSharedPtr<const FImage>>(Adapter->Image);
+			TVariant<Private::OP::ADDRESS, TSharedPtr<const Private::FImage>> Image;
+			Image.Set<TSharedPtr<const Private::FImage>>(Adapter->Image);
 			
 			Material->ImageParameters.Add(ParameterName, Image);
 		}
diff --git a/Engine/Plugins/Mutable/Source/MutableRuntime/Private/MuR/External/MeshAdapter.cpp b/Engine/Plugins/Mutable/Source/MutableRuntime/Private/MuR/External/MeshAdapter.cpp
--- a/Engine/Plugins/Mutable/Source/MutableRuntime/Private/MuR/External/MeshAdapter.cpp
+++ b/Engine/Plugins/Mutable/Source/MutableRuntime/Private/MuR/External/MeshAdapter.cpp
@@ -8,12 +8,9 @@
 
 namespace UE::Mutable
 {
-	using namespace Private;
-
-
 	FMeshAdapter::FMeshAdapter()
 	{
-		Mesh = MakeShared<FMesh>();
+		Mesh = MakeShared<Private::FMesh>();
 	}
 	
 
@@ -37,7 +34,7 @@ namespace UE::Mutable
 
 	FVector3f FMeshAdapter::GetVertex(int32 Index) const
 	{
-		MeshBufferIterator<EMeshBufferFormat::Float32, float, 3> BasePositionIter(Mesh->VertexBuffers, EMeshBufferSemantic::Position, 0);
+		Private::MeshBufferIterator<Private::EMeshBufferFormat::Float32, float, 3> BasePositionIter(Mesh->VertexBuffers, Private::EMeshBufferSemantic::Position, 0);
 		BasePositionIter += Index;
 
 		return BasePositionIter.GetAsVec3f();
@@ -46,7 +43,7 @@ namespace UE::Mutable
 	
 	void FMeshAdapter::SetVertex(int32 Index, const FVector3f& Vertex)
 	{
-		MeshBufferIterator<EMeshBufferFormat::Float32, float, 3> BasePositionIter(Mesh->VertexBuffers, EMeshBufferSemantic::Position, 0);
+		Private::MeshBufferIterator<Private::EMeshBufferFormat::Float32, float, 3> BasePositionIter(Mesh->VertexBuffers, Private::EMeshBufferSemantic::Position, 0);
 		BasePositionIter += Index;
 
 		BasePositionIter.SetFromVec3f(Vertex);
diff --git a/Engine/Plugins/Mutable/Source/MutableRuntime/Private/MuR/External/TextureAdapter.cpp b/Engine/Plugins/Mutable/Source/MutableRuntime/Private/MuR/External/TextureAdapter.cpp
--- a/Engine/Plugins/Mutable/Source/MutableRuntime/Private/MuR/External/TextureAdapter.cpp
+++ b/Engine/Plugins/Mutable/Source/MutableRuntime/Private/MuR/External/TextureAdapter.cpp
@@ -7,12 +7,9 @@
 
 namespace UE::Mutable
 {
-	using namespace Private;
-	
-
 	FTextureAdapter::FTextureAdapter()
 	{
-		Image = MakeShared<FImage>();		
+		Image = MakeShared<Private::FImage>();		
 	}
 
 
diff --git a/Engine/Plugins/Mutable/Source/MutableRuntime/Private/MuR/External/VectorAdapter.cpp b/Engine/Plugins/Mutable/Source/MutableRuntime/Private/MuR/External/VectorAdapter.cpp
--- a/Engine/Plugins/Mutable/Source/MutableRuntime/Private/MuR/External/VectorAdapter.cpp
+++ b/Engine/Plugins/Mutable/Source/MutableRuntime/Private/MuR/External/VectorAdapter.cpp
@@ -5,9 +5,6 @@
 
 namespace UE::Mutable
 {
-	using namespace Private;
-	
-
 	FVectorAdapter::FVectorAdapter(const FVectorAdapter& Other)
 	{
 		Value = Other.Value;
