[Mutable] Fix ambiguous types due to using namespace in Unity builds.

diff --git a/Engine/Plugins/Mutable/Source/CustomizableObjectEditor/Private/MuCOE/EdGraphSchema_CustomizableObject.cpp b/Engine/Plugins/Mutable/Source/CustomizableObjectEditor/Private/MuCOE/EdGraphSchema_CustomizableObject.cpp
--- a/Engine/Plugins/Mutable/Source/CustomizableObjectEditor/Private/MuCOE/EdGraphSchema_CustomizableObject.cpp
+++ b/Engine/Plugins/Mutable/Source/CustomizableObjectEditor/Private/MuCOE/EdGraphSchema_CustomizableObject.cpp
@@ -126,7 +126,6 @@ class IToolkit;
 
 #define SNAP_GRID (16) // @todo ensure this is the same as SNodePanel::GetSnapGridSize()
 
-using namespace UE::Mutable;
 
 namespace UE::Mutable::Private
 {
@@ -660,7 +659,7 @@ void UEdGraphSchema_CustomizableObject::GetGraphContextActions(FGraphContextMenu
 			FInstancedStruct OperationInstancedStruct;
 			OperationInstancedStruct.InitializeAs(OperationType.Get());
 
-			if (const FExternalOperation* Operation = OperationInstancedStruct.GetPtr<FExternalOperation>())
+			if (const UE::Mutable::FExternalOperation* Operation = OperationInstancedStruct.GetPtr<UE::Mutable::FExternalOperation>())
 			{
 				UCONodeExternal* Node = ContextMenuBuilder.CreateTemplateNode<UCONodeExternal>();
 				Node->OperationInstancedStruct = OperationInstancedStruct;
@@ -1532,23 +1531,23 @@ FName UEdGraphSchema_CustomizableObject::GetPinCategoryName(const FName& PinCate
 
 FName UEdGraphSchema_CustomizableObject::GetPinCategoryName(const UScriptStruct* Type)
 {
-	if (Type == FMeshAdapter::StaticStruct())
+	if (Type == UE::Mutable::FMeshAdapter::StaticStruct())
 	{
 		return PC_Mesh;
 	}
-	else if (Type == FTextureAdapter::StaticStruct())
+	else if (Type == UE::Mutable::FTextureAdapter::StaticStruct())
 	{
 		return PC_Texture;
 	}
-	else if (Type == FFloatAdapter::StaticStruct())
+	else if (Type == UE::Mutable::FFloatAdapter::StaticStruct())
 	{
 		return PC_Float;
 	}
-	else if (Type == FVectorAdapter::StaticStruct())
+	else if (Type == UE::Mutable::FVectorAdapter::StaticStruct())
 	{
 		return PC_Color;
 	}
-	else if (Type == FMaterialAdapter::StaticStruct())
+	else if (Type == UE::Mutable::FMaterialAdapter::StaticStruct())
 	{
 		return PC_Material;
 	}
diff --git a/Engine/Plugins/Mutable/Source/CustomizableObjectEditor/Private/MuCOE/GenerateMutableSource/GenerateMutableSourceExternal.cpp b/Engine/Plugins/Mutable/Source/CustomizableObjectEditor/Private/MuCOE/GenerateMutableSource/GenerateMutableSourceExternal.cpp
--- a/Engine/Plugins/Mutable/Source/CustomizableObjectEditor/Private/MuCOE/GenerateMutableSource/GenerateMutableSourceExternal.cpp
+++ b/Engine/Plugins/Mutable/Source/CustomizableObjectEditor/Private/MuCOE/GenerateMutableSource/GenerateMutableSourceExternal.cpp
@@ -19,11 +19,8 @@
 
 #define LOCTEXT_NAMESPACE "CustomizableObjectEditor"
 
-using namespace UE::Mutable;
-using namespace UE::Mutable::Private;
 
-
-UE::Mutable::Private::Ptr<NodeExternal> GenerateMutableSourceExtension(const UEdGraphPin* Pin, FMutableGraphGenerationContext& GenerationContext, const FSourceExtensionOptions& Options)
+UE::Mutable::Private::Ptr<UE::Mutable::Private::NodeExternal> GenerateMutableSourceExtension(const UEdGraphPin* Pin, FMutableGraphGenerationContext& GenerationContext, const FSourceExtensionOptions& Options)
 {
 	check(Pin)
 	RETURN_ON_CYCLE(*Pin, GenerationContext)
@@ -48,7 +45,7 @@ UE::Mutable::Private::Ptr<NodeExternal> GenerateMutableSourceExtension(const UEd
 	
 	UCustomizableObjectNode* EditorNode = CastChecked<UCustomizableObjectNode>(Pin->GetOwningNode());
 
-	UE::Mutable::Private::Ptr<NodeExternal> Result;
+	UE::Mutable::Private::Ptr<UE::Mutable::Private::NodeExternal> Result;
 	bool bCacheNode = true;
 	
 	if (const UCONodeExternal* TypedNodeExtension = Cast<UCONodeExternal>(EditorNode))
@@ -63,7 +60,7 @@ UE::Mutable::Private::Ptr<NodeExternal> GenerateMutableSourceExtension(const UEd
 				return;
 			}
 			
-			UE::Mutable::Private::Ptr<NodeExternal> ExtensionNode = new NodeExternal();
+			UE::Mutable::Private::Ptr<UE::Mutable::Private::NodeExternal> ExtensionNode = new UE::Mutable::Private::NodeExternal();
             				
             ExtensionNode->OperationInstancedStruct = TypedNodeExtension->OperationInstancedStruct;
 
@@ -94,34 +91,34 @@ UE::Mutable::Private::Ptr<NodeExternal> GenerateMutableSourceExtension(const UEd
 
             	const UScriptStruct* InputType = Pair.Value;
             
-            	if (InputType == FMeshAdapter::StaticStruct())
+            	if (InputType == UE::Mutable::FMeshAdapter::StaticStruct())
             	{
-            		UE::Mutable::Private::Ptr<NodeMesh> SourceNode = GenerateMutableSourceMesh(ConnectedPin, GenerationContext, Options.MeshOptions, Options.bLinkedToExtendMaterial, Options.bOnlyConnectedLOD);
+            		UE::Mutable::Private::Ptr<UE::Mutable::Private::NodeMesh> SourceNode = GenerateMutableSourceMesh(ConnectedPin, GenerationContext, Options.MeshOptions, Options.bLinkedToExtendMaterial, Options.bOnlyConnectedLOD);
             		ExtensionNode->Inputs.Add(SourceNode);
             	}
-            	else if (InputType == FTextureAdapter::StaticStruct())
+            	else if (InputType == UE::Mutable::FTextureAdapter::StaticStruct())
             	{
-            		UE::Mutable::Private::Ptr<NodeImage> SourceNode = GenerateMutableSourceImage(ConnectedPin, GenerationContext, Options.ReferenceTextureSize);
+            		UE::Mutable::Private::Ptr<UE::Mutable::Private::NodeImage> SourceNode = GenerateMutableSourceImage(ConnectedPin, GenerationContext, Options.ReferenceTextureSize);
             		ExtensionNode->Inputs.Add(SourceNode);
             	}
-            	else if (InputType == FFloatAdapter::StaticStruct())
+            	else if (InputType == UE::Mutable::FFloatAdapter::StaticStruct())
             	{
-            		UE::Mutable::Private::Ptr<NodeScalar> SourceNode = GenerateMutableSourceFloat(ConnectedPin, GenerationContext);
+            		UE::Mutable::Private::Ptr<UE::Mutable::Private::NodeScalar> SourceNode = GenerateMutableSourceFloat(ConnectedPin, GenerationContext);
             		ExtensionNode->Inputs.Add(SourceNode);
             	}
-            	else if (InputType == FVectorAdapter::StaticStruct())
+            	else if (InputType == UE::Mutable::FVectorAdapter::StaticStruct())
             	{
-            		UE::Mutable::Private::Ptr<NodeColour> SourceNode = GenerateMutableSourceColor(ConnectedPin, GenerationContext);
+            		UE::Mutable::Private::Ptr<UE::Mutable::Private::NodeColour> SourceNode = GenerateMutableSourceColor(ConnectedPin, GenerationContext);
             		ExtensionNode->Inputs.Add(SourceNode);
             	}
-            	else if (InputType == FMaterialAdapter::StaticStruct())
+            	else if (InputType == UE::Mutable::FMaterialAdapter::StaticStruct())
             	{
-            		UE::Mutable::Private::Ptr<NodeColour> SourceNode = GenerateMutableSourceColor(ConnectedPin, GenerationContext);
+            		UE::Mutable::Private::Ptr<UE::Mutable::Private::NodeColour> SourceNode = GenerateMutableSourceColor(ConnectedPin, GenerationContext);
             		ExtensionNode->Inputs.Add(SourceNode);
             	}
             	else
             	{
-            		UE::Mutable::Private::Ptr<Node> SourceNode = GenerateMutableSourceExtension(ConnectedPin, GenerationContext, Options);
+            		UE::Mutable::Private::Ptr<UE::Mutable::Private::Node> SourceNode = GenerateMutableSourceExtension(ConnectedPin, GenerationContext, Options);
             		ExtensionNode->Inputs.Add(SourceNode);
             	}	
             }
@@ -134,14 +131,14 @@ UE::Mutable::Private::Ptr<NodeExternal> GenerateMutableSourceExtension(const UEd
 	{
 		bCacheNode = false;
 
-		Result = GenerateMutableSourceMacro<NodeExternal>(*Pin, GenerationContext, GenerateMutableSourceExtension, Options);
+		Result = GenerateMutableSourceMacro<UE::Mutable::Private::NodeExternal>(*Pin, GenerationContext, GenerateMutableSourceExtension, Options);
 	}
 
 	else if (const UCustomizableObjectNodeTunnel* TypedNodeTunnel = Cast<UCustomizableObjectNodeTunnel>(InNode))
 	{
 		bCacheNode = false;
 
-		Result = GenerateMutableSourceMacro<NodeExternal>(*Pin, GenerationContext, GenerateMutableSourceExtension, Options);
+		Result = GenerateMutableSourceMacro<UE::Mutable::Private::NodeExternal>(*Pin, GenerationContext, GenerateMutableSourceExtension, Options);
 	}
 	
 	if (Result)
diff --git a/Engine/Plugins/Mutable/Source/CustomizableObjectEditor/Private/MuCOE/Nodes/CONodeExternal.cpp b/Engine/Plugins/Mutable/Source/CustomizableObjectEditor/Private/MuCOE/Nodes/CONodeExternal.cpp
--- a/Engine/Plugins/Mutable/Source/CustomizableObjectEditor/Private/MuCOE/Nodes/CONodeExternal.cpp
+++ b/Engine/Plugins/Mutable/Source/CustomizableObjectEditor/Private/MuCOE/Nodes/CONodeExternal.cpp
@@ -10,12 +10,10 @@
 
 #define LOCTEXT_NAMESPACE "CustomizableObjectEditor"
 
-using namespace UE::Mutable;
-
 
 FText UCONodeExternal::GetNodeTitle(ENodeTitleType::Type TitleType) const
 {
-	const FExternalOperation* Operation = OperationInstancedStruct.GetPtr<FExternalOperation>();
+	const UE::Mutable::FExternalOperation* Operation = OperationInstancedStruct.GetPtr<UE::Mutable::FExternalOperation>();
 	const FName Name = Operation ? Operation->GetName() : CachedOperationName;
 	
 	if (TitleType == ENodeTitleType::ListView)
@@ -46,7 +44,7 @@ void UCONodeExternal::AllocateDefaultPins(UCustomizableObjectNodeRemapPins* Rema
 {
 	Super::AllocateDefaultPins(RemapPins);
 
-	const FExternalOperation* Operation = OperationInstancedStruct.GetPtr<FExternalOperation>();
+	const UE::Mutable::FExternalOperation* Operation = OperationInstancedStruct.GetPtr<UE::Mutable::FExternalOperation>();
 	if (!Operation)
 	{
 		return;
