[Mutable] Fix GetModelResources crash when cooking.

diff --git a/Engine/Plugins/Mutable/Source/CustomizableObject/Internal/MuCO/CustomizableObjectPrivate.h b/Engine/Plugins/Mutable/Source/CustomizableObject/Internal/MuCO/CustomizableObjectPrivate.h
--- a/Engine/Plugins/Mutable/Source/CustomizableObject/Internal/MuCO/CustomizableObjectPrivate.h
+++ b/Engine/Plugins/Mutable/Source/CustomizableObject/Internal/MuCO/CustomizableObjectPrivate.h
@@ -1755,7 +1755,7 @@ class UCustomizableObjectPrivate : public UObject
 
 public:
 	/** Must be called after unlocking the CustomizableObject. */
-	CUSTOMIZABLEOBJECT_API void SetModel(const TSharedPtr<UE::Mutable::Private::FModel>& Model, const FGuid Identifier);
+	CUSTOMIZABLEOBJECT_API void SetModel(const TSharedPtr<UE::Mutable::Private::FModel>& Model, const FGuid Identifier, bool bIsCooking);
 	CUSTOMIZABLEOBJECT_API const TSharedPtr<UE::Mutable::Private::FModel>& GetModel();
 	CUSTOMIZABLEOBJECT_API const TSharedPtr<const UE::Mutable::Private::FModel> GetModel() const;
 
@@ -1833,7 +1833,7 @@ public:
 	/** Generic Load methods to read compiled data */
 	CUSTOMIZABLEOBJECT_API bool LoadModelResources(FArchive& Ar, const ITargetPlatform* InTargetPlatform, bool bSkipEditorOnlyData = false);
 	CUSTOMIZABLEOBJECT_API void LoadModelStreamableBulk(FArchive& Ar, bool bIsCooking);
-	CUSTOMIZABLEOBJECT_API void LoadModel(FArchive& Ar);
+	CUSTOMIZABLEOBJECT_API void LoadModel(FArchive& Ar, bool bIsCooking);
 
 	/** Load compiled data for the running platform from disk, this is used to load Editor Compilations. */
 	CUSTOMIZABLEOBJECT_API void LoadCompiledDataFromDisk();
diff --git a/Engine/Plugins/Mutable/Source/CustomizableObject/Private/MuCO/CustomizableObject.cpp b/Engine/Plugins/Mutable/Source/CustomizableObject/Private/MuCO/CustomizableObject.cpp
--- a/Engine/Plugins/Mutable/Source/CustomizableObject/Private/MuCO/CustomizableObject.cpp
+++ b/Engine/Plugins/Mutable/Source/CustomizableObject/Private/MuCO/CustomizableObject.cpp
@@ -243,7 +243,7 @@ void UCustomizableObject::PreSave(FObjectPreSaveContext ObjectSaveContext)
 			// Clear model resources
 			GetPrivate()->SetModelResources(nullptr, true /* bIsCooking */);
 			GetPrivate()->GetModelStreamableBulkData(true).Reset();
-			GetPrivate()->SetModel(nullptr, FGuid());
+			GetPrivate()->SetModel(nullptr, FGuid(), true);
 		}
 	}
 }
@@ -303,7 +303,7 @@ bool UCustomizableObjectPrivate::TryLoadCompiledCookDataForPlatform(const ITarge
 
 	SetModelResources(PlatformData->ModelResources.Get(), true);
 	SetModelStreamableBulkData(PlatformData->ModelStreamableBulkData, true);
-	SetModel(PlatformData->Model, GenerateIdentifier(*GetPublic()));
+	SetModel(PlatformData->Model, GenerateIdentifier(*GetPublic()), true);
 
 	return GetModel() != nullptr;
 }
@@ -512,9 +512,9 @@ void UCustomizableObjectPrivate::LoadModelStreamableBulk(FArchive& MemoryReader,
 }
 
 
-void UCustomizableObjectPrivate::LoadModel(FArchive& MemoryReader)
+void UCustomizableObjectPrivate::LoadModel(FArchive& MemoryReader, bool bIsCooking)
 {
-	SetModel(LoadModel_Internal(MemoryReader), GenerateIdentifier(*GetPublic()));
+	SetModel(LoadModel_Internal(MemoryReader), GenerateIdentifier(*GetPublic()), bIsCooking);
 }
 
 
@@ -700,7 +700,7 @@ void UCustomizableObjectPrivate::LoadCompiledDataFromDisk()
 			if (!bOutOfDate)
 			{
 				LoadModelStreamableBulk(MemoryReader, /* bIsCooking */false);
-				LoadModel(MemoryReader);
+				LoadModel(MemoryReader, false);
 			}
 			else
 			{
@@ -920,7 +920,7 @@ void UCustomizableObjectPrivate::LoadEmbeddedData(FArchive& Ar)
 		UE::Mutable::Private::FInputArchive Arch(&Stream);
 		TSharedPtr<UE::Mutable::Private::FModel, ESPMode::ThreadSafe> Model = UE::Mutable::Private::FModel::StaticUnserialise(Arch);
 
-		SetModel(Model, FGuid());
+		SetModel(Model, FGuid(), false);
 	}
 }
 
@@ -2154,7 +2154,7 @@ void UModelStreamableData::PostLoad()
 }
 
 
-void UCustomizableObjectPrivate::SetModel(const TSharedPtr<UE::Mutable::Private::FModel>& Model, const FGuid Id)
+void UCustomizableObjectPrivate::SetModel(const TSharedPtr<UE::Mutable::Private::FModel>& Model, const FGuid Id, bool bIsCooking)
 {
 	if (MutableModel == Model
 #if WITH_EDITOR
@@ -2178,8 +2178,9 @@ void UCustomizableObjectPrivate::SetModel(const TSharedPtr<UE::Mutable::Private:
 
 	if (Model)
 	{
-		UModelResources& LocalModelResources = GetModelResourcesChecked();
-		MutableModel->GetPrivate()->Program.ExternalOperationProvider = MakeShared<UE::Mutable::Private::FExternalOperationProvider>(LocalModelResources);
+		UModelResources* LocalModelResources = GetModelResources(bIsCooking);
+		check(LocalModelResources);
+		MutableModel->GetPrivate()->Program.ExternalOperationProvider = MakeShared<UE::Mutable::Private::FExternalOperationProvider>(*LocalModelResources);
 	}
 	
 	// Create parameter properties
diff --git a/Engine/Plugins/Mutable/Source/CustomizableObjectEditor/Private/MuCOE/CustomizableObjectCompiler.cpp b/Engine/Plugins/Mutable/Source/CustomizableObjectEditor/Private/MuCOE/CustomizableObjectCompiler.cpp
--- a/Engine/Plugins/Mutable/Source/CustomizableObjectEditor/Private/MuCOE/CustomizableObjectCompiler.cpp
+++ b/Engine/Plugins/Mutable/Source/CustomizableObjectEditor/Private/MuCOE/CustomizableObjectCompiler.cpp
@@ -1163,13 +1163,13 @@ void FCustomizableObjectCompiler::CompleteRequest(ECompilationStatePrivate State
 			{
 				CurrentObject->GetPrivate()->SetModelResources(PlatformData->ModelResources.Get(), CurrentOptions.bIsCooking);
 				CurrentObject->GetPrivate()->SetModelStreamableBulkData(PlatformData->ModelStreamableBulkData, CurrentOptions.bIsCooking);
-				CurrentObject->GetPrivate()->SetModel(PlatformData->Model, GenerateIdentifier(*CurrentObject));
+				CurrentObject->GetPrivate()->SetModel(PlatformData->Model, GenerateIdentifier(*CurrentObject), CurrentOptions.bIsCooking);
 			}
 			else
 			{
 				CurrentObject->GetPrivate()->SetModelResources(nullptr, CurrentOptions.bIsCooking);
 				CurrentObject->GetPrivate()->SetModelStreamableBulkData(nullptr, CurrentOptions.bIsCooking);
-				CurrentObject->GetPrivate()->SetModel(nullptr, {});
+				CurrentObject->GetPrivate()->SetModel(nullptr, {}, CurrentOptions.bIsCooking);
 			}
 
 			CurrentObject->GetPrivate()->PostCompile();
