[Mutable] Disabled extensions due to a crashes in a packaged build.

#rnx

diff --git a/Engine/Plugins/Mutable/Source/CustomizableObject/Private/MuCO/CustomizableObject.cpp b/Engine/Plugins/Mutable/Source/CustomizableObject/Private/MuCO/CustomizableObject.cpp
--- a/Engine/Plugins/Mutable/Source/CustomizableObject/Private/MuCO/CustomizableObject.cpp
+++ b/Engine/Plugins/Mutable/Source/CustomizableObject/Private/MuCO/CustomizableObject.cpp
@@ -2178,7 +2178,8 @@ void UCustomizableObjectPrivate::SetModel(const TSharedPtr<UE::Mutable::Private:
 
 	if (Model)
 	{
-		
+		// TODO GMT Extensions disabled due to a crash in packaged build
+		/*
 		UModelResources* LocalModelResources = nullptr;
 #if WITH_EDITORONLY_DATA
 		LocalModelResources = GetModelResources(bIsCooking);
@@ -2188,6 +2189,7 @@ void UCustomizableObjectPrivate::SetModel(const TSharedPtr<UE::Mutable::Private:
 
 		check(LocalModelResources);
 		MutableModel->GetPrivate()->Program.ExternalOperationProvider = MakeShared<UE::Mutable::Private::FExternalOperationProvider>(*LocalModelResources);
+		*/
 	}
 	
 	// Create parameter properties
