[mutable] Fixed LOCTEXT repeated key
- Made the key of an SCONode LOCTEXT more specific.
- Fixed badly named key for a "SCONodeExternal" LOCTEXT referring to an Instanced Struct not being loaded.

#rnx
#FailedValidations Preflight

diff --git a/Engine/Plugins/Mutable/Source/CustomizableObjectEditor/Private/MuCOE/Nodes/SCONodeExternal.cpp b/Engine/Plugins/Mutable/Source/CustomizableObjectEditor/Private/MuCOE/Nodes/SCONodeExternal.cpp
--- a/Engine/Plugins/Mutable/Source/CustomizableObjectEditor/Private/MuCOE/Nodes/SCONodeExternal.cpp
+++ b/Engine/Plugins/Mutable/Source/CustomizableObjectEditor/Private/MuCOE/Nodes/SCONodeExternal.cpp
@@ -20,7 +20,7 @@ void SCONodeExternal::CreateBelowPinControls(TSharedPtr<SVerticalBox> MainBox)
 {
 	SCustomizableObjectNode::CreateBelowPinControls(MainBox);
 
-	if (UCONodeExternal* Node = Cast<UCONodeExternal>(GraphNode))
+	if (const UCONodeExternal* Node = Cast<UCONodeExternal>(GraphNode))
 	{
 		if (!Node->OperationInstancedStruct.IsValid())
 		{
@@ -30,7 +30,7 @@ void SCONodeExternal::CreateBelowPinControls(TSharedPtr<SVerticalBox> MainBox)
 			[
 				SNew(SErrorText)
 				.BackgroundColor(FAppStyle::GetColor("ErrorReporting.BackgroundColor"))
-				.ErrorText(LOCTEXT("Experimental", "NOT LOADED"))
+				.ErrorText(LOCTEXT("SCONodeExternal_NotLoaded", "NOT LOADED"))
 			];
 		}
 	}
diff --git a/Engine/Plugins/Mutable/Source/CustomizableObjectEditor/Private/MuCOE/Nodes/SCustomizableObjectNode.cpp b/Engine/Plugins/Mutable/Source/CustomizableObjectEditor/Private/MuCOE/Nodes/SCustomizableObjectNode.cpp
--- a/Engine/Plugins/Mutable/Source/CustomizableObjectEditor/Private/MuCOE/Nodes/SCustomizableObjectNode.cpp
+++ b/Engine/Plugins/Mutable/Source/CustomizableObjectEditor/Private/MuCOE/Nodes/SCustomizableObjectNode.cpp
@@ -30,7 +30,7 @@ void SCustomizableObjectNode::CreateBelowPinControls(TSharedPtr<SVerticalBox> Ma
 {
 	SGraphNode::CreateBelowPinControls(MainBox);
 
-	if (UCustomizableObjectNode* Node = Cast<UCustomizableObjectNode>(GraphNode))
+	if (const UCustomizableObjectNode* Node = Cast<UCustomizableObjectNode>(GraphNode))
 	{
 		if (Node->IsExperimental())
 		{
@@ -40,7 +40,7 @@ void SCustomizableObjectNode::CreateBelowPinControls(TSharedPtr<SVerticalBox> Ma
 			[
 				SNew(SErrorText)
 				.BackgroundColor(FAppStyle::GetColor("ErrorReporting.WarningBackgroundColor"))
-				.ErrorText(LOCTEXT("Experimental", "EXPERIMENTAL"))
+				.ErrorText(LOCTEXT("SCONode_Experimental", "EXPERIMENTAL"))
 			];
 		}
 	}
