[mutable] Fixed invalid interactions with the StaticMesh TSoftObjtPtr
- Added backwards compatibility code to handle the regeneration of the output pins alongside the overriding of the "LODs" property.
- A TSoftObjPtr does require loading before being able to access the data it points.
- Improved property value changes management.
- Removed unnecessary includes from source files.
- Renamed incorrectly named variables and removed unused ones.
- Changed the type of the int variables (int to int32)

#rnx

diff --git a/Engine/Plugins/Mutable/Source/CustomizableObject/Internal/MuCO/CustomizableObjectCustomVersion.h b/Engine/Plugins/Mutable/Source/CustomizableObject/Internal/MuCO/CustomizableObjectCustomVersion.h
--- a/Engine/Plugins/Mutable/Source/CustomizableObject/Internal/MuCO/CustomizableObjectCustomVersion.h
+++ b/Engine/Plugins/Mutable/Source/CustomizableObject/Internal/MuCO/CustomizableObjectCustomVersion.h
@@ -201,6 +201,8 @@ struct FCustomizableObjectCustomVersion
 		FixStaticMeshNodeLODMaterialDataInvalidation,
 
 		FixStaticMeshNodeLODMaterialDataInvalidationBackOut,
+
+		FixStaticMeshNodeLODMaterialDataInvalidation2,
 		
 		// -----<new versions can be added above this line>-------------------------------------------------
 		VersionPlusOne,
diff --git a/Engine/Plugins/Mutable/Source/CustomizableObjectEditor/Private/MuCOE/GenerateMutableSource/GenerateMutableSourceMesh.cpp b/Engine/Plugins/Mutable/Source/CustomizableObjectEditor/Private/MuCOE/GenerateMutableSource/GenerateMutableSourceMesh.cpp
--- a/Engine/Plugins/Mutable/Source/CustomizableObjectEditor/Private/MuCOE/GenerateMutableSource/GenerateMutableSourceMesh.cpp
+++ b/Engine/Plugins/Mutable/Source/CustomizableObjectEditor/Private/MuCOE/GenerateMutableSource/GenerateMutableSourceMesh.cpp
@@ -3950,67 +3950,62 @@ UE::Mutable::Private::Ptr<UE::Mutable::Private::NodeMesh> GenerateMutableSourceM
 		UE::Mutable::Private::Ptr<UE::Mutable::Private::NodeMeshConstant> MeshNode = new UE::Mutable::Private::NodeMeshConstant();
 		Result = MeshNode;
 
-		if (TypedNodeStatic->StaticMesh)
-		{			
-			int32 LODIndex = 0; // TODO MTBL-1474
-			int32 SectionIndex = 0;
+		int32 LODIndex = 0; // TODO MTBL-1474
+		int32 SectionIndex = 0;
 
-			// Find out what material do we need
-			[&LODIndex, &SectionIndex, &TypedNodeStatic, &Pin]()
+		// Find out what material do we need
+		[&LODIndex, &SectionIndex, &TypedNodeStatic, &Pin]()
+		{
+			for (; LODIndex < TypedNodeStatic->LODs.Num(); ++LODIndex)
 			{
-				for (; LODIndex < TypedNodeStatic->LODs.Num(); ++LODIndex)
+				for (; SectionIndex < TypedNodeStatic->LODs[LODIndex].Materials.Num(); ++SectionIndex)
 				{
-					for (; SectionIndex < TypedNodeStatic->LODs[LODIndex].Materials.Num(); ++SectionIndex)
+					if (TypedNodeStatic->LODs[LODIndex].Materials[SectionIndex].MeshPinRef.Get() == Pin)
 					{
-						if (TypedNodeStatic->LODs[LODIndex].Materials[SectionIndex].MeshPinRef.Get() == Pin)
-						{
-							return;
-						}
+						return;
 					}
 				}
+			}
 
-				LODIndex = -1;
-				SectionIndex = -1;
-			}();
+			LODIndex = -1;
+			SectionIndex = -1;
+		}();
 			
-			check(TypedNodeStatic->LODs.IsValidIndex(LODIndex) && SectionIndex < TypedNodeStatic->LODs[LODIndex].Materials.Num());
+		check(TypedNodeStatic->LODs.IsValidIndex(LODIndex) && SectionIndex < TypedNodeStatic->LODs[LODIndex].Materials.Num());
 
-			constexpr bool bIsPassthrough = false;
-			TSharedPtr<UE::Mutable::Private::FMesh> MutableMesh = GenerateMutableStaticMesh(TypedNodeStatic->StaticMesh, TSoftClassPtr<UAnimInstance>(), LODIndex, SectionIndex,
-														  FString(), GenerationContext, TypedNodeStatic, nullptr, bIsPassthrough);
-			if (MutableMesh)
-			{
-				MeshNode->Value = MutableMesh;
+		constexpr bool bIsPassthrough = false;
+		TSharedPtr<UE::Mutable::Private::FMesh> MutableMesh = GenerateMutableStaticMesh(TypedNodeStatic->StaticMesh, TSoftClassPtr<UAnimInstance>(), LODIndex, SectionIndex,
+													  FString(), GenerationContext, TypedNodeStatic, nullptr, bIsPassthrough);
+		if (MutableMesh)
+		{
+			MeshNode->Value = MutableMesh;
 
-				// Layouts
-				MeshNode->Layouts.SetNum(1);
+			// Layouts
+			MeshNode->Layouts.SetNum(1);
 
-				UE::Mutable::Private::Ptr<UE::Mutable::Private::NodeLayout> LayoutNode;
+			UE::Mutable::Private::Ptr<UE::Mutable::Private::NodeLayout> LayoutNode;
 
-				const TArray<UCustomizableObjectLayout*>& Layouts = TypedNodeStatic->GetLayouts(*Pin);
-				const UCustomizableObjectLayout* Layout = Layouts.IsValidIndex(0) ? Layouts[0] : nullptr;
+			const TArray<UCustomizableObjectLayout*>& Layouts = TypedNodeStatic->GetLayouts(*Pin);
 				
-				if (Layout)
-				{
-					LayoutNode = CreateMutableLayoutNode(Layout, false);
-				}
-				else
-				{
-					LayoutNode = CreateDefaultLayout();
-				}
-
-				MeshNode->Layouts[0] = LayoutNode;
-				LayoutNode->SetMessageContext(Node);  // We need it here because we create multiple nodes.
-
-
-				const FString MeshName = TypedNodeStatic->StaticMesh.GetLongPackageName().ToLower();
-				MeshNode->SourceDataDescriptor.SourceId = CityHash32(reinterpret_cast<const char*>(*MeshName), MeshName.Len() * sizeof(FString::ElementType));
-				MeshNode->SourceDataDescriptor.OptionalMaxLODSize = 0;
+			if (const UCustomizableObjectLayout* Layout = Layouts.IsValidIndex(0) ? Layouts[0] : nullptr)
+			{
+				LayoutNode = CreateMutableLayoutNode(Layout, false);
 			}
 			else
 			{
-				Result = nullptr;
+				LayoutNode = CreateDefaultLayout();
 			}
+
+			MeshNode->Layouts[0] = LayoutNode;
+			LayoutNode->SetMessageContext(Node);  // We need it here because we create multiple nodes.
+			
+			const FString MeshName = TypedNodeStatic->StaticMesh.GetLongPackageName().ToLower();
+			MeshNode->SourceDataDescriptor.SourceId = CityHash32(reinterpret_cast<const char*>(*MeshName), MeshName.Len() * sizeof(FString::ElementType));
+			MeshNode->SourceDataDescriptor.OptionalMaxLODSize = 0;
+		}
+		else
+		{
+			Result = nullptr;
 		}
 	}
 
diff --git a/Engine/Plugins/Mutable/Source/CustomizableObjectEditor/Private/MuCOE/Nodes/CustomizableObjectNodeStaticMesh.cpp b/Engine/Plugins/Mutable/Source/CustomizableObjectEditor/Private/MuCOE/Nodes/CustomizableObjectNodeStaticMesh.cpp
--- a/Engine/Plugins/Mutable/Source/CustomizableObjectEditor/Private/MuCOE/Nodes/CustomizableObjectNodeStaticMesh.cpp
+++ b/Engine/Plugins/Mutable/Source/CustomizableObjectEditor/Private/MuCOE/Nodes/CustomizableObjectNodeStaticMesh.cpp
@@ -6,13 +6,12 @@
 #include "ISinglePropertyView.h"
 #include "Modules/ModuleManager.h"
 #include "MuCO/CustomizableObjectCustomVersion.h"
+#include "MuCO/LoadUtils.h"
 #include "MuCOE/CustomizableObjectEditor.h"
 #include "MuCOE/CustomizableObjectEditorStyle.h"
-#include "MuCOE/CustomizableObjectEditorUtilities.h"
 #include "MuCOE/CustomizableObjectLayout.h"
 #include "MuCOE/EdGraphSchema_CustomizableObject.h"
 #include "MuCOE/GraphTraversal.h"
-#include "MuCOE/MutableUtils.h"
 #include "MuCOE/Nodes/CustomizableObjectNodeLayoutBlocks.h"
 #include "MuCOE/RemapPins/CustomizableObjectNodeRemapPinsByNameDefaultPin.h"
 #include "MuCOE/UnrealEditorPortabilityHelpers.h"
@@ -34,7 +33,7 @@ static const TCHAR* STATIC_MESH_PIN_NAME = TEXT("Static Mesh");
 void UCustomizableObjectNodeStaticMesh::PostEditChangeProperty(FPropertyChangedEvent& PropertyChangedEvent)
 {
 	FProperty* PropertyThatChanged = PropertyChangedEvent.Property;
-	if (PropertyThatChanged && PropertyThatChanged->GetName() == TEXT("StaticMesh"))
+	if (PropertyThatChanged && PropertyThatChanged->GetName() == GET_MEMBER_NAME_STRING_CHECKED(UCustomizableObjectNodeStaticMesh, StaticMesh))
 	{
 		ReconstructNode();
 	}
@@ -92,27 +91,22 @@ void UCustomizableObjectNodeStaticMesh::AllocateDefaultPins(UCustomizableObjectN
 	
 	const UEdGraphSchema_CustomizableObject* Schema = GetDefault<UEdGraphSchema_CustomizableObject>();
 
-	if (StaticMesh)
+	if (UE::Mutable::Private::LoadObject(StaticMesh))
 	{
-		const int NumLODs = StaticMesh->GetRenderData()->LODResources.Num();
+		const int32 NumLODs = StaticMesh->GetRenderData()->LODResources.Num();
 
 		LODs.Reset(NumLODs);
 		LODs.SetNum(NumLODs);
-		for (int32 i = 0; i < NumLODs; ++i)
+		for (int32 LODIndex = 0; LODIndex < NumLODs; ++LODIndex)
 		{
-			FString LODName = FString::Printf(TEXT("LOD %d, "), i);
+			const int32 NumMaterials = StaticMesh->GetRenderData()->LODResources[LODIndex].Sections.Num();
 
-			const int NumMaterials = StaticMesh->GetRenderData()->LODResources[i].Sections.Num();
-
-			LODs[i].Materials.Reset(NumMaterials);
-			LODs[i].Materials.SetNum(NumMaterials);
+			LODs[LODIndex].Materials.Reset(NumMaterials);
+			LODs[LODIndex].Materials.SetNum(NumMaterials);
 			for (int32 MaterialIndex = 0; MaterialIndex < NumMaterials; ++MaterialIndex)
 			{
-				LODs[i].Materials.Reset(NumMaterials);
-				LODs[i].Materials.SetNum(NumMaterials);
-				
 				UMaterialInterface* MaterialInterface = 0;
-				int StaticMeshMaterialIndex = StaticMesh->GetRenderData()->LODResources[i].Sections[MaterialIndex].MaterialIndex;
+				int32 StaticMeshMaterialIndex = StaticMesh->GetRenderData()->LODResources[LODIndex].Sections[MaterialIndex].MaterialIndex;
 				if (StaticMeshMaterialIndex >= 0 && StaticMeshMaterialIndex < StaticMesh->GetStaticMaterials().Num())
 				{
 					MaterialInterface = StaticMesh->GetStaticMaterials()[StaticMeshMaterialIndex].MaterialInterface;
@@ -123,13 +117,13 @@ void UCustomizableObjectNodeStaticMesh::AllocateDefaultPins(UCustomizableObjectN
 				{
 					MaterialName = MaterialInterface->GetName();
 				}
-				LODs[i].Materials[MaterialIndex].Name = MaterialName;
-
+				LODs[LODIndex].Materials[MaterialIndex].Name = MaterialName;
+				
 				// Mesh
 				const FString MeshName = MaterialName + TEXT(" Mesh");
 				UEdGraphPin* Pin = CustomCreatePin(EGPD_Output, Schema->PC_Mesh, FName(*MeshName));
 				Pin->bDefaultValueIsIgnored = true;
-				LODs[i].Materials[MaterialIndex].MeshPinRef.SetPin(Pin);
+				LODs[LODIndex].Materials[MaterialIndex].MeshPinRef.SetPin(Pin);
 
 				// Images
 				if (MaterialInterface)
@@ -144,7 +138,7 @@ void UCustomizableObjectNodeStaticMesh::AllocateDefaultPins(UCustomizableObjectN
 						FString ImageName = ImageNames[ImageIndex].Name.ToString();
 						Pin = CustomCreatePin(EGPD_Output, Schema->PC_Texture, *ImageName);
 						Pin->bDefaultValueIsIgnored = true;
-						LODs[i].Materials[MaterialIndex].ImagePinsRef.Add(Pin);
+						LODs[LODIndex].Materials[MaterialIndex].ImagePinsRef.Add(Pin);
 					}
 				}
 			}
@@ -165,7 +159,7 @@ void UCustomizableObjectNodeStaticMesh::AllocateDefaultPins(UCustomizableObjectN
 
 FText UCustomizableObjectNodeStaticMesh::GetNodeTitle(ENodeTitleType::Type TitleType) const
 {
-	if (StaticMesh)
+	if (UE::Mutable::Private::LoadObject(StaticMesh))
 	{
 		FFormatNamedArguments Args;
 		Args.Add(TEXT("MeshName"), FText::FromString(StaticMesh->GetName()));
@@ -181,17 +175,17 @@ FText UCustomizableObjectNodeStaticMesh::GetNodeTitle(ENodeTitleType::Type Title
 
 UTexture2D* UCustomizableObjectNodeStaticMesh::FindTextureForPin(const UEdGraphPin* Pin) const
 {
-	if (!StaticMesh)
+	if (!UE::Mutable::Private::LoadObject(StaticMesh))
 	{
 		return 0;
 	}
 
-	for (int LODIndex = 0; LODIndex < LODs.Num(); ++LODIndex)
+	for (int32 LODIndex = 0; LODIndex < LODs.Num(); ++LODIndex)
 	{
-		for (int MaterialIndex = 0; MaterialIndex<LODs[LODIndex].Materials.Num(); ++MaterialIndex)
+		for (int32 MaterialIndex = 0; MaterialIndex<LODs[LODIndex].Materials.Num(); ++MaterialIndex)
 		{
 			UMaterialInterface* MaterialInterface = 0;
-			int StaticMeshMaterialIndex = -1;
+			int32 StaticMeshMaterialIndex = -1;
 			if (StaticMesh->GetRenderData()->LODResources.Num() > LODIndex && StaticMesh->GetRenderData()->LODResources[LODIndex].Sections.Num()>MaterialIndex)
 			{
 				StaticMeshMaterialIndex = StaticMesh->GetRenderData()->LODResources[LODIndex].Sections[MaterialIndex].MaterialIndex;
@@ -204,7 +198,7 @@ UTexture2D* UCustomizableObjectNodeStaticMesh::FindTextureForPin(const UEdGraphP
 
 			if (MaterialInterface)
 			{
-				for (int ImageIndex = 0; ImageIndex < LODs[LODIndex].Materials[MaterialIndex].ImagePinsRef.Num(); ++ImageIndex)
+				for (int32 ImageIndex = 0; ImageIndex < LODs[LODIndex].Materials[MaterialIndex].ImagePinsRef.Num(); ++ImageIndex)
 				{
 					if (LODs[LODIndex].Materials[MaterialIndex].ImagePinsRef[ImageIndex].Get() == Pin)
 					{
@@ -278,26 +272,26 @@ bool UCustomizableObjectNodeStaticMesh::IsNodeOutDatedAndNeedsRefresh()
 {
 	bool Result = false;
 
-	if (StaticMesh)
+	if (UE::Mutable::Private::LoadObject(StaticMesh))
 	{
-		int32 numLods = StaticMesh->GetRenderData()->LODResources.Num();
-		if (LODs.Num() != numLods)
+		const int32 LODCount = StaticMesh->GetRenderData()->LODResources.Num();
+		if (LODs.Num() != LODCount)
 		{
 			Result = true;
 		}
 
-		for (int LODIndex = 0; !Result && LODIndex < LODs.Num(); ++LODIndex)
+		for (int32 LODIndex = 0; !Result && LODIndex < LODs.Num(); ++LODIndex)
 		{
-			int32 NumMaterials = StaticMesh->GetRenderData()->LODResources[LODIndex].Sections.Num();
+			const int32 NumMaterials = StaticMesh->GetRenderData()->LODResources[LODIndex].Sections.Num();
 			if (LODs[LODIndex].Materials.Num() != NumMaterials)
 			{
 				Result = true;
 			}
 
-			for (int MaterialIndex = 0; !Result && MaterialIndex < NumMaterials; ++MaterialIndex)
+			for (int32 MaterialIndex = 0; !Result && MaterialIndex < NumMaterials; ++MaterialIndex)
 			{
 				UMaterialInterface* MaterialInterface = nullptr;
-				int SkeletalMeshMaterialIndex = StaticMesh->GetRenderData()->LODResources[LODIndex].Sections[MaterialIndex].MaterialIndex;
+				int32 SkeletalMeshMaterialIndex = StaticMesh->GetRenderData()->LODResources[LODIndex].Sections[MaterialIndex].MaterialIndex;
 				if (SkeletalMeshMaterialIndex >= 0 && SkeletalMeshMaterialIndex < StaticMesh->GetStaticMaterials().Num())
 				{
 					MaterialInterface = StaticMesh->GetStaticMaterials()[SkeletalMeshMaterialIndex].MaterialInterface;
@@ -322,7 +316,7 @@ bool UCustomizableObjectNodeStaticMesh::IsNodeOutDatedAndNeedsRefresh()
 						break;
 					}
 
-					for (int ImageIndex = 0; ImageIndex < ImageInfos.Num(); ++ImageIndex)
+					for (int32 ImageIndex = 0; ImageIndex < ImageInfos.Num(); ++ImageIndex)
 					{
 						const FString ImagePinName = LODs[LODIndex].Materials[MaterialIndex].ImagePinsRef[ImageIndex].Get()->GetName();
 						if (!ImagePinName.Contains(ImageInfos[ImageIndex].Name.ToString()))
@@ -346,12 +340,118 @@ bool UCustomizableObjectNodeStaticMesh::IsNodeOutDatedAndNeedsRefresh()
 	return Result;
 }
 
+
 FString UCustomizableObjectNodeStaticMesh::GetRefreshMessage() const
 {
 	return "Node data outdated. Please refresh node.";
 }
 
 
+void UCustomizableObjectNodeStaticMesh::BackwardsCompatibleFixup(int32 CustomizableObjectCustomVersion)
+{
+	Super::BackwardsCompatibleFixup(CustomizableObjectCustomVersion);
+
+	if (CustomizableObjectCustomVersion == FCustomizableObjectCustomVersion::FixStaticMeshNodeLODMaterialDataInvalidation2)
+	{
+		if (!UE::Mutable::Private::LoadObject(StaticMesh))
+		{
+			return;
+		}
+
+		// Ensure the remapping is actually needed, if not skip the remapping
+		if (DefaultPin.Get() == nullptr)
+		{
+			bool bInvalidMeshPinReferenceFound = false;
+			for (const FCustomizableObjectNodeStaticMeshLOD& LODData : LODs)
+			{
+				// There is no need to remap the pins if they are already correct.
+				if (LODData.Materials.IsValidIndex(0) && LODData.Materials[0].MeshPinRef.Get() == nullptr)
+				{
+					bInvalidMeshPinReferenceFound = true;
+					break;
+				}
+			}
+
+			// If the mesh pin references are correct there is no need to fix them.
+			if (!bInvalidMeshPinReferenceFound)
+			{
+				return;
+			}
+		}
+		
+		// The pin references in the LODs array are wrong so they need to be regenerated and remapped
+		
+		// Remap the pins using the old pin names but DO NOT auto-set one as default
+		FixupReconstructPins(NewObject<UCustomizableObjectNodeRemapPinsByNameDefaultPin>(), [&](UCustomizableObjectNode* Node, UCustomizableObjectNodeRemapPins* Action)
+		{
+			// Generate the pins again so we get non default pins and valid info into the LODs array
+				
+			if (StaticMesh->HasValidRenderData(false, INDEX_NONE))
+			{
+				// Make the first mesh pin the default one if the DefaultPin was already set (no LOD pins were present before)
+				if (UCustomizableObjectNodeRemapPinsByNameDefaultPin* RemapPinsCustom = Cast<UCustomizableObjectNodeRemapPinsByNameDefaultPin>(Action))
+				{
+					RemapPinsCustom->DefaultPin = DefaultPin.Get();
+				}
+				
+				const int32 NumLODs = StaticMesh->GetRenderData()->LODResources.Num();
+				LODs.Reset(NumLODs);
+				LODs.SetNum(NumLODs);
+				for (int32 LODIndex = 0; LODIndex < NumLODs; ++LODIndex)
+				{
+					const int32 NumMaterials = StaticMesh->GetRenderData()->LODResources[LODIndex].Sections.Num();
+
+					LODs[LODIndex].Materials.Reset(NumMaterials);
+					LODs[LODIndex].Materials.SetNum(NumMaterials);
+					for (int32 MaterialIndex = 0; MaterialIndex < NumMaterials; ++MaterialIndex)
+					{
+						UMaterialInterface* MaterialInterface = nullptr;
+						const int32 StaticMeshMaterialIndex = StaticMesh->GetRenderData()->LODResources[LODIndex].Sections[MaterialIndex].MaterialIndex;
+						if (StaticMeshMaterialIndex >= 0 && StaticMeshMaterialIndex < StaticMesh->GetStaticMaterials().Num())
+						{
+							MaterialInterface = StaticMesh->GetStaticMaterials()[StaticMeshMaterialIndex].MaterialInterface;
+						}
+
+						FString MaterialName = TEXT("Unnamed Material");
+						if (MaterialInterface)
+						{
+							MaterialName = MaterialInterface->GetName();
+						}
+						LODs[LODIndex].Materials[MaterialIndex].Name = MaterialName;
+						
+						// Mesh
+						const FString MeshName = MaterialName + TEXT(" Mesh");
+						UEdGraphPin* Pin = CustomCreatePin(EGPD_Output, UEdGraphSchema_CustomizableObject::PC_Mesh, FName(*MeshName));
+						Pin->bDefaultValueIsIgnored = true;
+						LODs[LODIndex].Materials[MaterialIndex].MeshPinRef.SetPin(Pin);
+
+						// Images
+						if (MaterialInterface)
+						{
+							if (const UMaterial* Material = MaterialInterface->GetMaterial())
+							{
+								TArray<FMaterialParameterInfo> ImageNames;
+								TArray<FGuid> ImageIds;
+								
+								Material->GetAllTextureParameterInfo(ImageNames, ImageIds);
+								for (int32 ImageIndex = 0; ImageIndex < ImageNames.Num(); ++ImageIndex)
+								{
+									FString ImageName = ImageNames[ImageIndex].Name.ToString();
+									Pin = CustomCreatePin(EGPD_Output, UEdGraphSchema_CustomizableObject::PC_Texture, *ImageName);
+									Pin->bDefaultValueIsIgnored = true;
+									LODs[LODIndex].Materials[MaterialIndex].ImagePinsRef.Add(Pin);
+								}
+							}	
+						}
+					}
+				}
+			}
+		});
+	}
+}
+
+
+
 FText UCustomizableObjectNodeStaticMesh::GetTooltipText() const
 {
 	return LOCTEXT("Static_Mesh_Tooltip", "Static Mesh");
@@ -381,7 +481,7 @@ void UCustomizableObjectNodeStaticMesh::GetPinSection(const UEdGraphPin& Pin, in
 			}
 			
 			// Is an Image pin?
-			for (int ImageIndex = 0; ImageIndex < Material.ImagePinsRef.Num(); ++ImageIndex)
+			for (int32 ImageIndex = 0; ImageIndex < Material.ImagePinsRef.Num(); ++ImageIndex)
 			{
 				if (Material.ImagePinsRef[ImageIndex].Get() == &Pin)
 				{
@@ -397,7 +497,7 @@ void UCustomizableObjectNodeStaticMesh::GetPinSection(const UEdGraphPin& Pin, in
 
 UMaterialInterface* UCustomizableObjectNodeStaticMesh::GetMaterialFor(const UEdGraphPin* Pin) const
 {
-	if (StaticMesh)
+	if (UE::Mutable::Private::LoadObject(StaticMesh))
 	{
 		check(Pin);
 		int32 LODIndex;
@@ -431,7 +531,8 @@ TSharedPtr<SGraphNode> UCustomizableObjectNodeStaticMesh::CreateVisualWidget()
 void SGraphNodeStaticMesh::Construct(const FArguments& InArgs, UEdGraphNode* InGraphNode)
 {
 	NodeStaticMesh = Cast< UCustomizableObjectNodeStaticMesh >(InGraphNode);
-
+	UE::Mutable::Private::LoadObject(NodeStaticMesh->StaticMesh);
+	
 	WidgetSize = 128.0f;
 	ThumbnailSize = 128;
 
@@ -447,7 +548,7 @@ void SGraphNodeStaticMesh::Construct(const FArguments& InArgs, UEdGraphNode* InG
 	SingleDetails.NotifyHook = Editor.Get();
 	SingleDetails.bHideAssetThumbnail = true;
 
-	StaticMeshSelector = PropPlugin.CreateSingleProperty(NodeStaticMesh, "StaticMesh", SingleDetails);	
+	StaticMeshSelector = PropPlugin.CreateSingleProperty(NodeStaticMesh, GET_MEMBER_NAME_STRING_CHECKED(UCustomizableObjectNodeStaticMesh, StaticMesh), SingleDetails);	
 
 	SCustomizableObjectNode::Construct({}, InGraphNode);
 }
diff --git a/Engine/Plugins/Mutable/Source/CustomizableObjectEditor/Private/MuCOE/Nodes/CustomizableObjectNodeStaticMesh.h b/Engine/Plugins/Mutable/Source/CustomizableObjectEditor/Private/MuCOE/Nodes/CustomizableObjectNodeStaticMesh.h
--- a/Engine/Plugins/Mutable/Source/CustomizableObjectEditor/Private/MuCOE/Nodes/CustomizableObjectNodeStaticMesh.h
+++ b/Engine/Plugins/Mutable/Source/CustomizableObjectEditor/Private/MuCOE/Nodes/CustomizableObjectNodeStaticMesh.h
@@ -148,6 +148,7 @@ public:
 	UE_API virtual bool HasPinViewer() const override;
 	UE_API virtual bool IsNodeOutDatedAndNeedsRefresh() override;
 	UE_API virtual FString GetRefreshMessage() const override;
+	UE_API virtual void BackwardsCompatibleFixup(int32 CustomizableObjectCustomVersion) override;
 
 	// UCustomizableObjectNodeMesh interface
 	UE_API virtual UTexture2D* FindTextureForPin(const UEdGraphPin* Pin) const override;
