[mutable] Report the sections of the generated resource name so they can be easily manipulated outside the actual baking operation. Also improved texture reusability during the bake.

UE-266716
- Updated commandlet calls as a new baking method has been added that does not require the input of a map.
- Removed dead code that was trying to make unique an object name that was already known to be unique.

UE-309993
- Now the baking operation will be able to reuse mutable textures when exporting all the resources. This was not possible before because the name was never the same between updates.
- The baking operation will now be able to reuse the Skeleton and Physics assets.

#rnx

diff --git a/Engine/Plugins/Mutable/Source/CustomizableObject/Public/MuCO/CustomizableObjectInstance.h b/Engine/Plugins/Mutable/Source/CustomizableObject/Public/MuCO/CustomizableObjectInstance.h
--- a/Engine/Plugins/Mutable/Source/CustomizableObject/Public/MuCO/CustomizableObjectInstance.h
+++ b/Engine/Plugins/Mutable/Source/CustomizableObject/Public/MuCO/CustomizableObjectInstance.h
@@ -70,9 +70,49 @@ struct FBakedResourceData
 
 	/**
 	 * The path used by the saved package.
+	 * EX : "/Game/baking_CO/BakeDir_Char/Part3/T_CharacterBake_Body_EyeBrown_Normal_2"
 	 */
 	UPROPERTY(BlueprintReadOnly, Category = CustomizableObjectInstanceBaker)
 	FString AssetPath;
+
+	/**
+	 * The section of the asset name that represents the asset type based prefix.
+	 * EX : "T_", "M_", "SKEL_", "" ...
+	 */
+	UPROPERTY(BlueprintReadOnly, Category = CustomizableObjectInstanceBaker)
+	FString Prefix;
+
+	/**
+	 * The section of the asset name that represents the name the user has requested for all the baked resources to have.
+	 * EX : "CharacterBake_", "CO_ABC_", "MyTestBake_", "" ...
+	 */
+	UPROPERTY(BlueprintReadOnly, Category = CustomizableObjectInstanceBaker)
+	FString UserGivenName;
+
+	/**
+	 * The section of the asset name that represents the name of the component of the source asset that was baked.
+	 * It can be empty in the case the resource is shared between multiple components.
+	 * EX : "Body_", "A_", "B_", "" ...
+	 */
+	UPROPERTY(BlueprintReadOnly, Category = CustomizableObjectInstanceBaker)
+	FString ComponentName;
+
+	/**
+	 * The section of the asset name that represents the name of the source resource that was baked.
+	 * EX : "EyeBrown_Normal", "BodyTexture_4849745649878/", "AlbedoEyes_4756489*65", "BodyTexture" ...
+	 * @note The substring is composed by the source object name and, in some cases, a suffix that contains the mutable resource generation hash used
+	 * to be able to reuse baked mutable resources between baking runs.
+	 */
+	UPROPERTY(BlueprintReadOnly, Category = CustomizableObjectInstanceBaker)
+	FString ResourceName;
+
+	/**
+	 * The section of the asset name that represents the suffix added in some cases to prevent name collisions during the bake.
+	 * This is most likely to be empty.
+	 * EX : "_2", "_1" , "" ...
+	 */
+	UPROPERTY(BlueprintReadOnly, Category = CustomizableObjectInstanceBaker)
+	FString UniqueAssetNameSuffix;
 };
 
 
diff --git a/Engine/Plugins/Mutable/Source/CustomizableObjectEditor/Internal/MuCOE/CustomizableObjectInstanceBakingUtils.h b/Engine/Plugins/Mutable/Source/CustomizableObjectEditor/Internal/MuCOE/CustomizableObjectInstanceBakingUtils.h
--- a/Engine/Plugins/Mutable/Source/CustomizableObjectEditor/Internal/MuCOE/CustomizableObjectInstanceBakingUtils.h
+++ b/Engine/Plugins/Mutable/Source/CustomizableObjectEditor/Internal/MuCOE/CustomizableObjectInstanceBakingUtils.h
@@ -16,6 +16,173 @@
  */
 CUSTOMIZABLEOBJECTEDITOR_API void ScheduleCOCompilationForBaking(UCustomizableObjectInstance& InTargetInstance, const FCompileNativeDelegate& InCustomizableObjectCompilationDelegate, const bool bPerformPartialCompilation = false);
 
+/*
+ * Data required by an asset to be able to be baked during the runtime of the baking operation.
+ * It also serves as an interface to get the different substrings that form the actual baked package name to ease name manipulation once the baking
+ * operation has completed
+ */
+class FResourceBakingData
+{
+
+public:
+	FResourceBakingData(const FString& InPrefix, const FString& inUserGivenName, const FString& InRawResourceName, const FString& InComponentName, const FString& InAssetPath):
+		Prefix(InPrefix),
+		ResourceName(InRawResourceName),
+		ComponentName(InComponentName),
+		UserGivenName(inUserGivenName),
+		AssetPath(InAssetPath)
+	{
+		check(!AssetPath.IsEmpty());
+	}
+
+	/**
+	 * Ensure that the data contained in this object is valid and well-formed
+	 * @return true If it is ok, false otherwise
+	 */
+	bool Validate() const;
+
+	/**
+	 * Generate and get the Composed Resource Name to be used when duplicating (baking) the object represented by this data structure
+	 * @return The string composed like so : Prefix_ObjectName_ResourceName(_UniqueAssetSuffix)
+	 */
+	FString GetResourceComposedName();
+	
+	/**
+	* Get the complete path to be used when duplicating the package to decide where and how the package is named
+	* @return The AssetPath/ComposedResourceName
+	 */
+	FString GetResourceComposedFullPath();
+	
+	/**
+	 * Given the ComposedName looks for it in the provided array and updates this object data so the next time the ComposedName is generated is unique
+	 * @param InCachedResources Array of resources we want to make sure the new name does not collide with.
+	 * @note This does not check the actual path of the objects, just the name as this is intended to be used during the runtime of the bake to prevent collisions
+	 * within the same bake
+	 */
+	void MakeComposedResourceNameUnique(const TArray<UObject*>& InCachedResources);
+
+	/**
+	 * Check if the file with the AssetPath + Composed name already exists, and based on that determines the save type to be performed
+	 */
+	void ComputePackageSaveResolutionType();
+	
+	/**
+	 * Get the section of the ComposedResourceName that represents the prefix.
+	 * @return The prefix string.
+	 */
+	FString GetPrefixSection() const;
+
+	/**
+	 * Get the section of the ComposedResourceName that represents the Component Name
+	 * @return The Component name section.
+	 */
+	FString GetComponentNameSection() const;
+
+	/**
+	 * Get the section of the ComposedResourceName that represents the User provided name
+	 * @return The User given name section. It can be empty.
+	 */
+	FString GetUserGivenNameSection() const;
+
+/**
+	 * Get the section of the ComposedResourceName that represents the actual Resource Name.
+	 * @return The Resource name section.
+	 */
+	FString GetResourceNameSection() const;
+
+	/***
+	 * Get the section of the ComposedResourceName that represents the uniqueness suffix added to make the name unique.
+	 * @return The suffix section.
+	 */
+	FString GetUniqueAssetNameSuffixSection() const;
+	
+private:
+
+	/**
+	 * Given a collection of UObjects return the index of the one that matches the name provided
+	 * @param InName The name of the object we are looking for
+	 * @param Objects The collection of objects we want to check for an element with the given name
+	 * @return The index of the element whose name matches the one provided or INDEX_NONE if the element could not be found
+	 */
+	int32 GetObjectWithNameIndex(const FString& InName, const TArray<UObject*>& Objects ) const;
+
+	/**
+	 * Get the package our data points to. This is based on the AssetPath and the ComposedResourceName
+	 * @return True if the asset is already at path
+	 */
+	bool DoesPackageExist(const FString& InPackageFullName, const FString& InAssetName) const;
+	
+	/**
+	 * Get the last generated Composed Resource Name.
+	 * @note As the ComposedResourceName changes over time as we update the parts that form it this method is only safe to call once the asset has been baked
+	 * as no more changes to the CachedComposedResourceName are to be expected
+	 * @return The last generated ComposedResourceName.
+	 */
+	FString GetCachedComposedResourceName() const;
+	
+	
+public:
+	
+	/**
+	 * The prefix (type based) of the resource. It is expected to include a "_" at the end like "T_" or "SKM_".
+	 */
+	FString Prefix;
+	
+	/**
+	 * The name of the object that it is being baked. This string is expected to be externally modified during the bake runtime
+	 */
+	FString ResourceName;
+
+	/**
+	 * A suffix string used to make the name of the object Unique.
+	 * Can be empty if no name collisions are found at the moment of bake
+	 */
+	FString UniqueAssetNameValueSuffix;
+
+	/**
+	 * The type of saving operation was performed. Does not affect the behaviour of the bake, but it is useful info for the caller
+	 */
+	EPackageSaveResolutionType SaveResolutionType = EPackageSaveResolutionType::None;
+
+	
+private:
+
+	/**
+	 * The component this object is part of.
+	 */
+	const FString ComponentName;
+	
+	/**
+	 * The custom name given by the user to differentiate the objects of this bake from others.
+	 * Can be empty.
+	 */
+	const FString UserGivenName;
+
+	/**
+	 * The path where to store the asset we are baking.
+	 */
+	const FString AssetPath;
+
+	/**
+	 * Cached version of the ComposedResourceName. Useful for the const consultation of the ComposedResourceName.
+	 */
+	FString CachedComposedResourceName;
+	
+
+	// String sections used to ease the removal of parts of the generated resource name once the baking has completed
+	
+	FString PrefixSection;
+	
+	FString ComponentNameSection;
+	
+	FString UserGivenNameSection;
+	
+	FString ResourceNameSection;
+	
+	FString UniqueAssetNameSuffixSection;
+	
+};
+
 
 /**
  * Schedules the async update of the target instance. It will clear previous update data.
@@ -30,7 +197,7 @@ CUSTOMIZABLEOBJECTEDITOR_API void ScheduleInstanceUpdateForBaking(UCustomizableO
  * Serializes onto disk the resources used by the targeted Customizable Object Instance. The operation can be configured in the FInstanceBakingSettings settings object.
  * @param InInstance The mutable COI instance whose resources we want to bake onto disk
  * @param Configuration The baking configuration object with all the settings to be used for this baking operation.
- * @param bIsUnattendedExecution Determines if we want to run this method and show promps for user interaction or if we want to automatically chose the more sensible option (without the need of user interaction)
+ * @param bIsUnattendedExecution Determines if we want to run this method and show prompts for user interaction or if we want to automatically chose the more sensible option (without the need of user interaction)
  * @param OutSavedPackages A collection with the packages marked for save.
  * @return True if the baking operation could be completed without issues and false otherwise
  */
@@ -38,4 +205,17 @@ CUSTOMIZABLEOBJECTEDITOR_API bool BakeCustomizableObjectInstance(
 	UCustomizableObjectInstance& InInstance,
 	const FBakingConfiguration& Configuration,
 	bool bIsUnattendedExecution,
-	TMap<UPackage*,EPackageSaveResolutionType>& OutSavedPackages);
\ No newline at end of file
+	TMap<UPackage*,const FResourceBakingData>& OutSavedPackages);
+
+
+/**
+ * Serializes onto disk the resources used by the targeted Customizable Object Instance. The operation can be configured in the FInstanceBakingSettings settings object.
+ * @param InInstance The mutable COI instance whose resources we want to bake onto disk
+ * @param Configuration The baking configuration object with all the settings to be used for this baking operation.
+ * @param bIsUnattendedExecution Determines if we want to run this method and show prompts for user interaction or if we want to automatically chose the more sensible option (without the need of user interaction)
+ * @return True if the baking operation could be completed without issues and false otherwise
+ */
+CUSTOMIZABLEOBJECTEDITOR_API bool BakeCustomizableObjectInstance(
+	UCustomizableObjectInstance& InInstance,
+	const FBakingConfiguration& Configuration,
+	bool bIsUnattendedExecution);
\ No newline at end of file
diff --git a/Engine/Plugins/Mutable/Source/CustomizableObjectEditor/Private/MuCOE/CustomizableObjectEditorViewportClient.cpp b/Engine/Plugins/Mutable/Source/CustomizableObjectEditor/Private/MuCOE/CustomizableObjectEditorViewportClient.cpp
--- a/Engine/Plugins/Mutable/Source/CustomizableObjectEditor/Private/MuCOE/CustomizableObjectEditorViewportClient.cpp
+++ b/Engine/Plugins/Mutable/Source/CustomizableObjectEditor/Private/MuCOE/CustomizableObjectEditorViewportClient.cpp
@@ -1686,13 +1686,11 @@ void FCustomizableObjectEditorViewportClient::OnInstanceForBakingUpdate(const FU
 			Configuration.bGenerateConstantMaterialInstancesOnBake = FolderDlg->GetGenerateConstantMaterialInstances();
 
 			// TODO: UE-266716 Add a way to setup the prefixes from the editor dialog
-		
-			TMap<UPackage*, EPackageSaveResolutionType> SavedPackages;
+			
 			BakeCustomizableObjectInstance(
 				*BakeTempInstance,
 				Configuration,
-				false,
-				SavedPackages);
+				false);
 		}
 	
 		BakeTempInstance = nullptr;
diff --git a/Engine/Plugins/Mutable/Source/CustomizableObjectEditor/Private/MuCOE/CustomizableObjectInstanceBaker.cpp b/Engine/Plugins/Mutable/Source/CustomizableObjectEditor/Private/MuCOE/CustomizableObjectInstanceBaker.cpp
--- a/Engine/Plugins/Mutable/Source/CustomizableObjectEditor/Private/MuCOE/CustomizableObjectInstanceBaker.cpp
+++ b/Engine/Plugins/Mutable/Source/CustomizableObjectEditor/Private/MuCOE/CustomizableObjectInstanceBaker.cpp
@@ -34,7 +34,7 @@ void UCustomizableObjectInstanceBaker::BakeInstance(UCustomizableObjectInstance*
 		FinishBakingOperation(false, nullptr);
 		return;
 	}
-	
+
 	if (!InTargetInstance->GetCustomizableObject())
 	{
 		UE_LOG(LogMutable,Error,TEXT("The provided instance does not have a Customizable Object defined."))
@@ -44,7 +44,7 @@ void UCustomizableObjectInstanceBaker::BakeInstance(UCustomizableObjectInstance*
 	
 	// We clone the InTargetInstance to ensure that the baked it is not open in a COI editor.
 	InstanceToBake = InTargetInstance->Clone();
-
+	
 	if (ensure(InstanceToBake))
 	{
 		// Request the async compilation of the CO so we can later perform the update safely
@@ -83,8 +83,8 @@ void UCustomizableObjectInstanceBaker::UpdateInstance()
 	// Once the update finishes successfully the baking operation will engage
 	if (ensure(InstanceToBake))
 	{
-		ScheduleInstanceUpdateForBaking(*InstanceToBake, UpdateDelegate);
-	}
+	ScheduleInstanceUpdateForBaking(*InstanceToBake, UpdateDelegate);
+}
 	else
 	{
 		FinishBakingOperation(false, nullptr);
@@ -117,8 +117,8 @@ void UCustomizableObjectInstanceBaker::Bake() const
 	GIsRunningUnattendedScript = true;
 	
 	// Array with all the packages that did get saved. Packages that failed the saving procedure will not appear in this array.
-	TMap<UPackage*, EPackageSaveResolutionType> SavedPackages;
-	
+	TMap<UPackage*, const FResourceBakingData> SavedPackages;
+
 	bool bBakeWasSuccessful = false;
 	if (ensure(InstanceToBake))
 	{
@@ -138,11 +138,18 @@ void UCustomizableObjectInstanceBaker::Bake() const
 		// Display a list with the paths of the assets saved
 		UE_LOG(LogMutable, Log, TEXT("Saved assets (%d) :"), SavedPackages.Num());
 		SavedPackagesData.Reserve(SavedPackages.Num());
-		for (const TPair<UPackage*, EPackageSaveResolutionType>& SavedPackagePair : SavedPackages)
+		for (const TPair<UPackage*, const FResourceBakingData>& SavedPackagePair : SavedPackages)
 		{
 			FBakedResourceData BakeData;
 			BakeData.AssetPath = SavedPackagePair.Key->GetPathName();
-			BakeData.SaveType = SavedPackagePair.Value;
+			
+			const FResourceBakingData& PackageData = SavedPackagePair.Value;
+			BakeData.SaveType = PackageData.SaveResolutionType;
+			BakeData.Prefix = PackageData.GetPrefixSection();
+			BakeData.UserGivenName = PackageData.GetUserGivenNameSection();
+			BakeData.ComponentName = PackageData.GetComponentNameSection();
+			BakeData.ResourceName = PackageData.GetResourceNameSection();
+			BakeData.UniqueAssetNameSuffix = PackageData.GetUniqueAssetNameSuffixSection();
 
 			// Sanity checks
 			check(BakeData.SaveType != EPackageSaveResolutionType::None);
diff --git a/Engine/Plugins/Mutable/Source/CustomizableObjectEditor/Private/MuCOE/CustomizableObjectInstanceBakingUtils.cpp b/Engine/Plugins/Mutable/Source/CustomizableObjectEditor/Private/MuCOE/CustomizableObjectInstanceBakingUtils.cpp
--- a/Engine/Plugins/Mutable/Source/CustomizableObjectEditor/Private/MuCOE/CustomizableObjectInstanceBakingUtils.cpp
+++ b/Engine/Plugins/Mutable/Source/CustomizableObjectEditor/Private/MuCOE/CustomizableObjectInstanceBakingUtils.cpp
@@ -27,62 +27,202 @@
 #define LOCTEXT_NAMESPACE "CustomizableObjectEditor" 
 
 
-/**
- * Removes the "_n" string found at the end of the material and skeletal mesh name.
- * @warning It does not check if that part is there or not.
- * @param MaterialName The name of the material that you want not to be unique.
- */
-void MakeAssetNameNotUnique(FString& MaterialName)
+FString FResourceBakingData::GetCachedComposedResourceName() const
 {
-	// This should be a material generated by mutable so try to remove the value that makes it unique
-	int32 UniquenessValueSeparatorIndex = INDEX_NONE;
-	MaterialName.FindLastChar(TEXT('_'),UniquenessValueSeparatorIndex);
-	if (UniquenessValueSeparatorIndex != INDEX_NONE)
+	return CachedComposedResourceName;
+}
+
+FString FResourceBakingData::GetResourceComposedName()
+{
+	ensure(!Prefix.IsEmpty());
+	ensure(!ResourceName.IsEmpty());
+	
+	// Prepare the strings that will form the composed string. Cache their values for external reference
+	PrefixSection = Prefix;																			// Ex : "T_" or "SKM_"
+	ComponentNameSection = ComponentName.IsEmpty() ? ComponentName : ComponentName + TEXT("_");		// Ex : "Head_" or "Body1_"			
+	UserGivenNameSection = UserGivenName.IsEmpty() ? UserGivenName : UserGivenName + TEXT("_");		// Ex : "MyObject_" or ""
+	ResourceNameSection =  ResourceName;															// Ex : "Texture_0000000..."
+	UniqueAssetNameSuffixSection = UniqueAssetNameValueSuffix;										// Ex : "_0"
+	
+	// const FString ObjectBaseName = UserGivenNameSection.IsEmpty() ? ComponentNameSection : UserGivenNameSection + ComponentNameSection;
+	CachedComposedResourceName = PrefixSection + UserGivenNameSection + ComponentNameSection + ResourceNameSection + UniqueAssetNameSuffixSection;
+	
+	return CachedComposedResourceName;
+}
+
+
+FString FResourceBakingData::GetResourceComposedFullPath()
+{
+	return AssetPath + TEXT("/") + GetResourceComposedName();
+}
+
+
+bool FResourceBakingData::Validate() const
+{
+	if (ResourceName.IsEmpty() || Prefix.IsEmpty() || AssetPath.IsEmpty())
 	{
-		MaterialName = MaterialName.Left(UniquenessValueSeparatorIndex);
+		return false;
 	}
-						
-	// Ensure that the new name contains numerical values representing the hash (speculative)
-	int32 HashValueSeparatorIndex = INDEX_NONE;
-	MaterialName.FindLastChar(TEXT('_'),HashValueSeparatorIndex);
-	if (ensure(HashValueSeparatorIndex != INDEX_NONE))
+	if (GetCachedComposedResourceName().IsEmpty())
 	{
-		const FString HashValue = MaterialName.RightChop(HashValueSeparatorIndex+1);
-		// If this fails it means that the name of the material did not contain or the hash or the value that makes it unique.
-		check(HashValue.IsNumeric())
+		return false;
 	}
+	if (SaveResolutionType == EPackageSaveResolutionType::None)
+	{
+		return false;
+	}
+
+	return true;
 }
 
 
-/**
- * Remove the prefix from the provided string.
- * @param InOutString The string to update
- * @param ToCheckPrefixes The list of possible prefixes that, if present, will get removed from the target string
- */
-bool RemovePrefix(FString& InOutString, const TArray<FString>& ToCheckPrefixes)
+int32 FResourceBakingData::GetObjectWithNameIndex(const FString& InName, const TArray<UObject*>& Objects) const
 {
-	for (const FString& PrefixToRemove : ToCheckPrefixes)
+	int32 Index = 0;
+	for (UObject* CachedResource : Objects)
 	{
-		if (InOutString.Find(PrefixToRemove, ESearchCase::CaseSensitive) == 0)
+		if (CachedResource && CachedResource->GetName() == InName)
 		{
-			InOutString = InOutString.RightChop(PrefixToRemove.Len());
-			return true;
+			return Index;
 		}
+		Index++;
 	}
 
-	return false;
+	return INDEX_NONE;
+}
+
+
+void FResourceBakingData::MakeComposedResourceNameUnique(const TArray<UObject*>& InCachedResources)
+{
+	// FString Output;
+	const FString TempComposedResourceName = GetResourceComposedName();
+	
+	// Look for the resource name provided to see if we have already worked with it.
+	int32 FindResult = GetObjectWithNameIndex(TempComposedResourceName, InCachedResources);
+	if (FindResult != INDEX_NONE)
+	{
+		// Add an integer suffix to create the unique name
+		uint32 Count = 0;
+		while (FindResult != INDEX_NONE)
+		{
+			FindResult = GetObjectWithNameIndex(TempComposedResourceName + "_" + FString::FromInt(Count), InCachedResources);
+			Count++;
+		}
+
+		UniqueAssetNameValueSuffix =  "_" + FString::FromInt(--Count);
+	}
+}
+
+
+bool FResourceBakingData::DoesPackageExist(const FString& InPackageFullName, const FString& InAssetName) const
+{
+	// Check if the package already exists
+	UPackage* ExistingPackage = FindPackage(nullptr, *InPackageFullName);
+	if (!ExistingPackage)
+	{
+		const FString PackageFilePath = InPackageFullName + "." + InAssetName;
+		
+		FString PackageFileName;
+		if (FPackageName::DoesPackageExist(PackageFilePath, &PackageFileName))
+		{
+			ExistingPackage = LoadPackage(nullptr, *PackageFileName, LOAD_EditorOnly);
+		}
+	}
+
+	return ExistingPackage != nullptr;
+}
+
+
+void FResourceBakingData::ComputePackageSaveResolutionType()
+{
+	SaveResolutionType = EPackageSaveResolutionType::None;
+	if (DoesPackageExist(GetResourceComposedFullPath(), GetResourceComposedName()))
+	{
+		// File found, instead of creating a new package use this one and update the contents it has
+		SaveResolutionType = EPackageSaveResolutionType::ReusedFile;
+	}
+	else
+	{
+		// The package could not be found. It is safe to create a package at this path and with the given name
+		SaveResolutionType = EPackageSaveResolutionType::NewFile;
+	}
+
+	// Just to be sure this case never happens even with future changes
+	check(SaveResolutionType == EPackageSaveResolutionType::ReusedFile || SaveResolutionType == EPackageSaveResolutionType::NewFile);
+
+	const UEnum* OptimizationLevelEnum = StaticEnum<EPackageSaveResolutionType>();
+	check(OptimizationLevelEnum);
+}
+
+
+FString FResourceBakingData::GetPrefixSection() const
+{
+	const FString ComposedResourceName = GetCachedComposedResourceName();
+	ensure(ComposedResourceName.StartsWith(PrefixSection) && PrefixSection.EndsWith(TEXT("_")));
+	return PrefixSection;
+}
+
+
+FString FResourceBakingData::GetComponentNameSection() const
+{
+	const FString ComposedResourceName = GetCachedComposedResourceName();
+	ensure(ComposedResourceName.Contains(ComponentNameSection));
+	return ComponentNameSection;
+}
+
+
+FString FResourceBakingData::GetUserGivenNameSection() const
+{
+	const FString ComposedResourceName = GetCachedComposedResourceName();
+	ensure(ComposedResourceName.Contains(UserGivenNameSection));
+	return UserGivenNameSection;
+}
+
+
+FString FResourceBakingData::GetResourceNameSection() const
+{
+	const FString ComposedResourceName = GetCachedComposedResourceName();
+	ensure(ComposedResourceName.Contains(ResourceNameSection));
+	return ResourceNameSection;
+}
+
+
+FString FResourceBakingData::GetUniqueAssetNameSuffixSection() const
+{
+	if (!UniqueAssetNameSuffixSection.IsEmpty())
+	{
+		const FString ComposedResourceName = GetCachedComposedResourceName();
+		ensure(ComposedResourceName.EndsWith(UniqueAssetNameSuffixSection));
+	}
+
+	return UniqueAssetNameSuffixSection;
 }
 
 
 /**
- * Compose the resource name to be used by a resource.
- * @param InPrefix The prefix we want to use for a given resource.
- * @param InObjectName The name of the object this resource is part of.
- * @param InOutResourceName The name of the resource itself. This string will be the target of the concatenations.
+ * Removes the "_n" string found at the end of the material and skeletal mesh name.
+ * @warning It does not check if that part is there or not.
+ * @param InMaterialName The name of the material that you want not to be unique.
  */
-void ComposeResourceName( const FString& InPrefix, const FString& InObjectName, FString& InOutResourceName)
+void MakeAssetNameNotUnique(FString& InMaterialName)
 {
-	InOutResourceName = InPrefix + InObjectName + InOutResourceName;
+	// This should be a material generated by mutable so try to remove the value that makes it unique
+	int32 UniquenessValueSeparatorIndex = INDEX_NONE;
+	InMaterialName.FindLastChar(TEXT('_'),UniquenessValueSeparatorIndex);
+	if (UniquenessValueSeparatorIndex != INDEX_NONE)
+	{
+		InMaterialName = InMaterialName.Left(UniquenessValueSeparatorIndex);
+	}
+						
+	// Ensure that the new name contains numerical values representing the hash (speculative)
+	int32 HashValueSeparatorIndex = INDEX_NONE;
+	InMaterialName.FindLastChar(TEXT('_'),HashValueSeparatorIndex);
+	if (ensure(HashValueSeparatorIndex != INDEX_NONE))
+	{
+		const FString HashValue = InMaterialName.RightChop(HashValueSeparatorIndex+1);
+		// If this fails it means that the name of the material did not contain or the hash or the value that makes it unique.
+		check(HashValue.IsNumeric())
+	}
 }
 
 
@@ -198,84 +338,18 @@ int32 GetObjectWithNameIndex(const FString& InName, const TArray<UObject*>& Obje
 }
 
 
-/**
- * Outputs a string that we know it is unique.
- * @param ResourceName The name of the resource we have provided. This should have the name of the current resource and will have the unique name for the resource once the method exits
- * @param CachedResources Collection with all the already processed resources name's
- */
-void MakeResourceNameUnique(FString& ResourceName, const TArray<UObject*>& CachedResources)
+bool TryRemovePrefix (FString& InOutString, const TArray<FString>& ToCheckPrefixes)
 {
-	// Look for the resource name provided to see if we have already worked with it.
-	int32 FindResult = GetObjectWithNameIndex(ResourceName, CachedResources);
-	if (FindResult != INDEX_NONE)
-	{
-		// Add an integer suffix to create the unique name
-		uint32 Count = 0;
-		while (FindResult != INDEX_NONE)
-		{
-			FindResult = GetObjectWithNameIndex(ResourceName + "_" + FString::FromInt(Count), CachedResources);
-			Count++;
-		}
-
-		ResourceName += "_" + FString::FromInt(--Count);
-	}
-}
-
-
-/**
- * Get the package at a given path and with a given name
- * @param InPackageSavePath The path were we plan to save the package
- * @param InObjName The name used by the package
- * @return True if the asset is already at path
- */
-UPackage* GetPackage(const FString& InPackageSavePath, const FString& InObjName)
-{
-	// Check if the package already exists
-	const FString PackagePath = InPackageSavePath + "/" + InObjName;
-	UPackage* ExistingPackage = FindPackage(nullptr, *PackagePath);
-	if (!ExistingPackage)
+	for (const FString& PrefixToRemove : ToCheckPrefixes)
 	{
-		const FString PackageFilePath = PackagePath + "." + InObjName;
-		
-		FString PackageFileName;
-		if (FPackageName::DoesPackageExist(PackageFilePath, &PackageFileName))
+		if (InOutString.Find(PrefixToRemove, ESearchCase::CaseSensitive) == 0)
 		{
-			ExistingPackage = LoadPackage(nullptr, *PackageFileName, LOAD_EditorOnly);
+			InOutString = InOutString.RightChop(PrefixToRemove.Len());
+			return true;
 		}
 	}
 
-	return ExistingPackage;
-}
-
-
-/**
- * Get the save resolution type for a given package. Will return or New or Reuse type depending if the package already exists on disk or not
- * @param InPackageSavePath The path where to save the package
- * @param InPackageName The name to be given to the package that we want to save or reuse
- * @return 
- */
-EPackageSaveResolutionType GetPackageSaveResolution(const FString& InPackageSavePath, const FString& InPackageName)
-{
-	EPackageSaveResolutionType SaveResolutionType = EPackageSaveResolutionType::None;
-	if (GetPackage(InPackageSavePath, InPackageName))
-	{
-		// File found, instead of creating a new package use this one and update the contents it has
-		SaveResolutionType = EPackageSaveResolutionType::ReusedFile;
-	}
-	else
-	{
-		// The package could not be found. It is safe to create a package at this path and with the given name
-		SaveResolutionType = EPackageSaveResolutionType::NewFile;
-	}
-
-	// Just to be sure this case never happens even with future changes
-	check(SaveResolutionType == EPackageSaveResolutionType::ReusedFile || SaveResolutionType == EPackageSaveResolutionType::NewFile);
-
-	const UEnum* OptimizationLevelEnum = StaticEnum<EPackageSaveResolutionType>();
-	check(OptimizationLevelEnum);
-	UE_LOG(LogMutable, Verbose, TEXT("SaveResolution for the \"%s\" package is \"%s\""), *InPackageSavePath, *OptimizationLevelEnum->GetValueAsName(SaveResolutionType).ToString())
-
-	return SaveResolutionType;
+	return false;
 }
 
 
@@ -399,7 +473,7 @@ bool IsAMutableTexture (const UTexture2D* InTexture)
  * @param NewObjectPair The Pair of data we want to add to the map. Only if the package is not present in the map the action will get performed.
  * @param OutSavedPackages The map you want to add data into.
  */
-void AddUniqueToSavePackage(TPair<UPackage*, EPackageSaveResolutionType> NewObjectPair, TMap<UPackage*,EPackageSaveResolutionType>& OutSavedPackages)
+void AddUniqueToSavePackage(const TPair<UPackage*, const FResourceBakingData>& NewObjectPair, TMap<UPackage*,const FResourceBakingData>& OutSavedPackages)
 {
 	if (!OutSavedPackages.Contains(NewObjectPair.Key))
 	{
@@ -412,12 +486,12 @@ bool BakeCustomizableObjectInstance(
 	UCustomizableObjectInstance& InInstance,
 	const FBakingConfiguration& Configuration,
 	bool bIsUnattendedExecution,
-	TMap<UPackage*,EPackageSaveResolutionType>& OutSavedPackages)
+	TMap<UPackage*,const FResourceBakingData>& OutSavedPackages)
 {
 	OutSavedPackages.Reset();
 
 	// Extract some baking settings from the configuration file provided
-	const FString& FileName = Configuration.OutputFilesBaseName;
+	const FString& UserGivenFilename = Configuration.OutputFilesBaseName;
 	const FString& AssetPath = Configuration.OutputPath;
 	
 	// Resource prefixes to be used. If an invalid prefix has been provided by the configuration object use the UE default one instead
@@ -459,13 +533,13 @@ bool BakeCustomizableObjectInstance(
 		return false;
 	}
 	
-	if (!ValidateProvidedFileName(FileName))
+	if (!ValidateProvidedFileName(UserGivenFilename))
 	{
 		UE_LOG(LogMutable, Error, TEXT("The FileName for the instance baking is not valid."));
 		return false;
 	}
 
-	if (!ValidateProvidedAssetPath(FileName,AssetPath,InstanceCO))
+	if (!ValidateProvidedAssetPath(UserGivenFilename,AssetPath,InstanceCO))
 	{
 		UE_LOG(LogMutable, Error, TEXT("The AssetPath for the instance baking is not valid."));
 		return false;
@@ -524,25 +598,12 @@ bool BakeCustomizableObjectInstance(
 		{
 			continue;
 		}
-
-		// Append the component name to the ObjectFileName
-		FString ObjectBaseName = FileName;
-		if (FileName.IsEmpty())
-		{
-			ObjectBaseName = ComponentName.ToString();
-		}
-		else
-		{
-			ObjectBaseName = FileName + TEXT("_") + ComponentName.ToString();
-		}
-		ObjectBaseName += TEXT("_");
 		
 		TMap<UObject*, UObject*> ReplacementMap;
 
 		if (Configuration.bExportAllResourcesOnBake)
 		{
 			UMaterial* Material;
-			FString PackageName;
 
 			// Each element of the array corresponds to a material of the mesh
 			// The key is the parameter index and the value the texture to be used there
@@ -578,39 +639,34 @@ bool BakeCustomizableObjectInstance(
 							}
 							
 							// Update the name adding to it the name of the component alongside other data
-							FString TextureName = SourceTexture->GetName();
-							RemovePrefix(TextureName, {FBakingConfiguration::BakedResourcePrefix});
-							RemovePrefix(TextureName, {TextureAssetPrefix});
-							ComposeResourceName(TextureAssetPrefix, ObjectBaseName, TextureName);
-
+							FResourceBakingData TextureData {TextureAssetPrefix, UserGivenFilename, SourceTexture->GetName(), FString{}, AssetPath};
+							TryRemovePrefix(TextureData.ResourceName, {FBakingConfiguration::BakedResourcePrefix});
+							TryRemovePrefix(TextureData.ResourceName, {TextureAssetPrefix});
+							
 							// Experimental
-							// If the texture is not from mutable and we know it has already been processed...
-							if (!IsAMutableTexture(SourceTexture) && HandledSourceObjects.Contains(SourceTexture))
+							// If we know the resource has already been processed...
+							if (HandledSourceObjects.Contains(SourceTexture))
 							{
-								// Instead of just creating a new name for the texture to be stored in disk, a better aproach could be using
+								// Instead of just creating a new name for the texture to be stored in disk, a better approach could be using
 								// the already duplicated file and assigning it to the material we are working with.
 								// The issue this has is that the texture may be named as from one component and now will be used for multiple ones
 								for (UObject* DuplicatedObject : DuplicatedObjects)
 								{
 									if (DuplicatedObject &&
 										DuplicatedObject->IsA<UTexture>() &&
-										DuplicatedObject->GetName() == TextureName)
+										DuplicatedObject->GetName() == TextureData.GetResourceComposedName())
 									{
-										// TODO: Allow reusing a texture of one component in another one (DuplicatedObject->GetName() == TextureName)
-										//		to be able to do this I will need to stop using the name of the texture as the name has the component embedded on it
-										//		or, instead, remove the component name from the full name and then check
-											
 										UTexture* DupTexture = Cast<UTexture>(DuplicatedObject);
 										TextureReplacementMaps[MaterialIndex].Add(ParameterInfoIndex, DupTexture);
 
-										UE_LOG(LogMutable, Verbose, TEXT("Matching texture found with name %s in component %s"), *TextureName, *ComponentName.ToString())
+										UE_LOG(LogMutable, Verbose, TEXT("Matching texture found with name %s in component %s"), *TextureData.GetResourceComposedName(), *ComponentName.ToString())
 										break;
 									}
 								}
 							}
 							
 							// Here we will know if the asset already exists (ReusedFile) or if it is new (NewFile)
-							const EPackageSaveResolutionType SaveResolution = GetPackageSaveResolution(AssetPath, TextureName);
+							TextureData.ComputePackageSaveResolutionType();
 							
 							// Duplicating mutable generated textures
 							if (IsAMutableTexture(SourceTexture))
@@ -621,13 +677,12 @@ bool BakeCustomizableObjectInstance(
 									UTexture* OriginalTexture = nullptr;
 									MaterialInstance->Parent->GetTextureParameterValue(FName(*ParameterName.GetPlainNameString()), OriginalTexture);
 
-									PackageName = AssetPath + FString("/") + TextureName;
-									UTexture2D* DuplicatedTexture = FUnrealBakeHelpers::BakeHelper_CreateAssetTexture(SourceTexture, TextureName, PackageName, OriginalTexture, nullptr, SaveResolution);
+									UTexture2D* DuplicatedTexture = FUnrealBakeHelpers::BakeHelper_CreateAssetTexture(SourceTexture, OriginalTexture, nullptr, TextureData);
 									
 									HandledSourceObjects.Add(SourceTexture);
 									DuplicatedObjects.Add(DuplicatedTexture);
 									
-									AddUniqueToSavePackage({DuplicatedTexture->GetPackage(), SaveResolution},OutSavedPackages);
+									AddUniqueToSavePackage({DuplicatedTexture->GetPackage(), TextureData},OutSavedPackages);
 
 									if (OriginalTexture != nullptr)
 									{
@@ -638,14 +693,12 @@ bool BakeCustomizableObjectInstance(
 							else
 							{
 								// Duplicate the non-mutable textures of the Material instance (pass-through textures)
-
-								PackageName = AssetPath + FString("/") + TextureName;
-								UObject* DuplicatedTexture = FUnrealBakeHelpers::BakeHelper_DuplicateAsset(Texture, TextureName, PackageName, nullptr, false, SaveResolution);
+								UObject* DuplicatedTexture = FUnrealBakeHelpers::BakeHelper_DuplicateAsset(Texture, nullptr, false, TextureData);
 
 								HandledSourceObjects.Add(Texture);
 								DuplicatedObjects.Add(DuplicatedTexture);
 								
-								AddUniqueToSavePackage({DuplicatedTexture->GetPackage(), SaveResolution},OutSavedPackages);
+								AddUniqueToSavePackage({DuplicatedTexture->GetPackage(), TextureData},OutSavedPackages);
 
 								UTexture* DupTexture = Cast<UTexture>(DuplicatedTexture);
 								TextureReplacementMaps[MaterialIndex].Add(ParameterInfoIndex, DupTexture);
@@ -677,27 +730,24 @@ bool BakeCustomizableObjectInstance(
 						check(ReplacementMap.Contains(Interface));	
 						continue;
 					}
+
+					FResourceBakingData MaterialData {MaterialInstanceAssetPrefix, UserGivenFilename, Interface->GetName(), ComponentName.ToString(), AssetPath};
+				 	// If this is a material marked with the BAKE prefix then it should include a mark to make it unique. Remove it for the baking
+				 	// operation
+				 	if (TryRemovePrefix(MaterialData.ResourceName,{FBakingConfiguration::BakedResourcePrefix}))
+				 	{
+				 		MakeAssetNameNotUnique(MaterialData.ResourceName);
+				 	}
+					TryRemovePrefix (MaterialData.ResourceName,MaterialResourcePrefixes);
 					
-					FString MaterialName = Interface->GetName();
-					// If this is a material marked with the BAKE prefix then it should include a mark to make it unique. Remove it for the baking
-					// operation
-					if (RemovePrefix(MaterialName, {FBakingConfiguration::BakedResourcePrefix}))
-					{
-						MakeAssetNameNotUnique(MaterialName);
-					}
-					RemovePrefix(MaterialName, MaterialResourcePrefixes);
-					ComposeResourceName(MaterialAssetPrefix, ObjectBaseName, MaterialName);
-					
-					// Give it an unique name or may happen that another material already stored on disk during this bake will get updated
+					// Give it a unique name or may happen that another material already stored on disk during this bake will get updated
 					// with data form this material. We want to avoid that from happening as we do not want to require the name to always be a hash of
 					// the contents (they are not always produced by mutable in this case)
-					MakeResourceNameUnique(MaterialName, DuplicatedObjects);
+					MaterialData.MakeComposedResourceNameUnique(DuplicatedObjects);
+					MaterialData.ComputePackageSaveResolutionType();
 					
-					const EPackageSaveResolutionType SaveResolution = GetPackageSaveResolution(AssetPath, MaterialName);
-
-					PackageName = AssetPath + FString("/") + MaterialName;
-					UObject* DuplicatedObject = FUnrealBakeHelpers::BakeHelper_DuplicateAsset(Material, MaterialName, PackageName, 
-						 nullptr, Configuration.bGenerateConstantMaterialInstancesOnBake, SaveResolution);
+					UObject* DuplicatedObject = FUnrealBakeHelpers::BakeHelper_DuplicateAsset(Material, nullptr,
+						Configuration.bGenerateConstantMaterialInstancesOnBake, MaterialData);
 
 					HandledSourceObjects.Add(Interface);
 					DuplicatedObjects.Add(DuplicatedObject);
@@ -707,7 +757,7 @@ bool BakeCustomizableObjectInstance(
 					// Tell the system that the Interface object will be replaced by the Duplicated one
 					ReplacementMap.Add(Interface, DuplicatedObject);
 					
-					AddUniqueToSavePackage({DuplicatedObject->GetPackage(), SaveResolution},OutSavedPackages);
+					AddUniqueToSavePackage({DuplicatedObject->GetPackage(), MaterialData},OutSavedPackages);
 
 					// Copy the texture parameters from the interface to the material
 					if (UMaterial* DuplicatedMaterial = Cast<UMaterial>(DuplicatedObject))
@@ -737,35 +787,35 @@ bool BakeCustomizableObjectInstance(
 				{
 					continue;
 				}
-
-				FString MaterialName;
+				
+				// Use the name of the interface as it is unique because of the hash it contains
+				FResourceBakingData MaterialData {FString{}, UserGivenFilename, Interface->GetName(), ComponentName.ToString(), AssetPath};
 				{
-					MaterialName = Interface->GetName();		// Use the name of the interface as it is unique because of the hash it contains
-					if (RemovePrefix(MaterialName, {FBakingConfiguration::BakedResourcePrefix}))
+					if (TryRemovePrefix(MaterialData.ResourceName, {FBakingConfiguration::BakedResourcePrefix}))
 					{
-						MakeAssetNameNotUnique(MaterialName);
+						MakeAssetNameNotUnique(MaterialData.ResourceName);
 					}
-					RemovePrefix(MaterialName, MaterialResourcePrefixes);
+					TryRemovePrefix(MaterialData.ResourceName,MaterialResourcePrefixes);
 					
 					if (Configuration.bGenerateConstantMaterialInstancesOnBake && Interface->IsA<UMaterialInstance>())
 					{
 						// Change the prefix to the Material Constant Instance since in this situation the new asset based on the Interface
 						// will be of the Material Constant Instance type
-						ComposeResourceName(MaterialConstantInstanceAssetPrefix, ObjectBaseName, MaterialName);
+						MaterialData.Prefix = MaterialConstantInstanceAssetPrefix;
 					}
 					else
 					{
 						if (Interface->IsA<UMaterial>())
 						{
-							ComposeResourceName(MaterialAssetPrefix, ObjectBaseName, MaterialName);
+							MaterialData.Prefix = MaterialAssetPrefix;
 						}
 						else if (Interface->IsA<UMaterialInstanceConstant>())
 						{
-							ComposeResourceName(MaterialConstantInstanceAssetPrefix, ObjectBaseName, MaterialName);
+							MaterialData.Prefix = MaterialConstantInstanceAssetPrefix;
 						}
 						else if (Interface->IsA<UMaterialInstanceDynamic>())
 						{
-							ComposeResourceName(MaterialDynamicInstanceAssetPrefix, ObjectBaseName, MaterialName);
+							MaterialData.Prefix = MaterialDynamicInstanceAssetPrefix;
 						}
 						else
 						{
@@ -776,17 +826,14 @@ bool BakeCustomizableObjectInstance(
 
 				// One material could be used by multiple material instances so, to be sure, make sure the name is unique so we do not reuse it during the bake
 				// We could be more specific in which cases we do this but I left it this way to keep the complexity as low as possible
-				MakeResourceNameUnique(MaterialName, DuplicatedObjects);
-				
-				const EPackageSaveResolutionType SaveResolution = GetPackageSaveResolution(AssetPath, MaterialName);
+				MaterialData.MakeComposedResourceNameUnique(DuplicatedObjects);
+				MaterialData.ComputePackageSaveResolutionType();
 				
-				FString MatPkgName = AssetPath + FString("/") + MaterialName;
-				UObject* DuplicatedMaterial = FUnrealBakeHelpers::BakeHelper_DuplicateAsset(Interface, MaterialName, 
-					MatPkgName, &ReplacementMap, Configuration.bGenerateConstantMaterialInstancesOnBake, SaveResolution);
+				UObject* DuplicatedMaterial = FUnrealBakeHelpers::BakeHelper_DuplicateAsset(Interface, &ReplacementMap, Configuration.bGenerateConstantMaterialInstancesOnBake, MaterialData);
 				HandledSourceObjects.Add(Interface);
 				DuplicatedObjects.Add(DuplicatedMaterial);
 				
-				AddUniqueToSavePackage({DuplicatedMaterial->GetPackage(), SaveResolution},OutSavedPackages);
+				AddUniqueToSavePackage({DuplicatedMaterial->GetPackage(), MaterialData},OutSavedPackages);
 
 				// Only need to duplicate the generate textures if the original material is a dynamic instance
 				// If the material has Mutable textures, then it will be a dynamic material instance for sure
@@ -830,23 +877,20 @@ bool BakeCustomizableObjectInstance(
 										
 										// The source texture has not yet been processed so duplicate it
 										
-										// Generate a new name for the duplicated texture so it does not collide with the original one
-										FString TextureName = SourceTexture->GetName();
-										RemovePrefix(TextureName, {FBakingConfiguration::BakedResourcePrefix});
-										RemovePrefix(TextureName, {TextureAssetPrefix});
-										ComposeResourceName(TextureAssetPrefix, ObjectBaseName, TextureName);
+										FResourceBakingData TextureData {TextureAssetPrefix, UserGivenFilename, SourceTexture->GetName(), FString{}, AssetPath};
+										TryRemovePrefix(TextureData.ResourceName ,{FBakingConfiguration::BakedResourcePrefix});
+										TryRemovePrefix(TextureData.ResourceName , {TextureAssetPrefix});
 										
 										// Check if the file is already in disk, so we can later decide what we do with that file
-										const EPackageSaveResolutionType TextureSaveResolution = GetPackageSaveResolution(AssetPath, TextureName);
-
-										FString TexPkgName = AssetPath + FString("/") + TextureName;
-										UTexture2D* DuplicatedTexture = FUnrealBakeHelpers::BakeHelper_CreateAssetTexture(SourceTexture, TextureName, TexPkgName, nullptr, nullptr, SaveResolution);
+										TextureData.ComputePackageSaveResolutionType();
+										
+										UTexture2D* DuplicatedTexture = FUnrealBakeHelpers::BakeHelper_CreateAssetTexture(SourceTexture, nullptr, nullptr, TextureData);
 
 										// Cache the file we just have handled and it's duplicate
 										HandledSourceObjects.Add(SourceTexture);
 										DuplicatedObjects.Add(DuplicatedTexture);
 										
-										AddUniqueToSavePackage({DuplicatedTexture->GetPackage(), TextureSaveResolution},OutSavedPackages);
+										AddUniqueToSavePackage({DuplicatedTexture->GetPackage(), TextureData},OutSavedPackages);
 
 										if (InstDynamic)
 										{
@@ -881,56 +925,65 @@ bool BakeCustomizableObjectInstance(
 				}
 			}
 		}
-
-		// TODO: As it happens with the textures we may be able to re-use the other resources generated for one component into the baked data of the other
-		//		for now we will keep them separated for the assets being processed bellow
-		
-		const FString CustomizableObjectName = InInstance.GetCustomizableObject()->GetName() + TEXT("_");
 		
-		// Get the clean name of the Skeletal mesh (without the embedded CO name)
-		FString SkeletalMeshName = Mesh->GetName();
+
+		FResourceBakingData SkeletalMeshData {SkeletalMeshAssetPrefix, UserGivenFilename, Mesh->GetName(), ComponentName.ToString(), AssetPath};
 		if (Mesh->HasAnyFlags(RF_Transient))
 		{
-			SkeletalMeshName = SkeletalMeshName.Replace(*CustomizableObjectName,TEXT(""));
-			if (RemovePrefix(SkeletalMeshName, {FBakingConfiguration::BakedResourcePrefix}))
+			// Remove the CO name from the name of the resource
+			const FString CustomizableObjectName = InInstance.GetCustomizableObject()->GetName();
+			ensureMsgf(SkeletalMeshData.ResourceName.Contains(CustomizableObjectName), TEXT("It is expected for the mutable generated meshes for baking to contain the name of the CO in their name. Has this changed?"));
+			const FString CustomizableObjectNameSection = CustomizableObjectName + TEXT("_");
+			SkeletalMeshData.ResourceName = SkeletalMeshData.ResourceName.Replace(*CustomizableObjectNameSection,TEXT(""));
+
+			// Remove the Component name from the name of the resource
+			ensureMsgf(SkeletalMeshData.ResourceName.Contains(ComponentName.ToString()), TEXT("It is expected for the mutable generated meshes for baking to contain the name of the component in their name. Has this changed?"));
+			const FString ComponentNameSection = ComponentName.ToString() + TEXT("_");
+			SkeletalMeshData.ResourceName = SkeletalMeshData.ResourceName.Replace( *ComponentNameSection,TEXT(""));
+
+			if (TryRemovePrefix(SkeletalMeshData.ResourceName, {FBakingConfiguration::BakedResourcePrefix}))
 			{
-				MakeAssetNameNotUnique(SkeletalMeshName);
-		}
+				MakeAssetNameNotUnique(SkeletalMeshData.ResourceName);
+			}
 		}
-		RemovePrefix(SkeletalMeshName, {SkeletalMeshAssetPrefix});
-		ComposeResourceName(SkeletalMeshAssetPrefix, ObjectBaseName, SkeletalMeshName);
-		MakeResourceNameUnique(SkeletalMeshName, DuplicatedObjects);
-		
+		TryRemovePrefix(SkeletalMeshData.ResourceName,{SkeletalMeshAssetPrefix});
+		SkeletalMeshData.MakeComposedResourceNameUnique(DuplicatedObjects);			// This ensures we get only one SKM per component
+
 		// Skeletal Mesh's Skeleton
 		if (USkeleton* Skeleton = Mesh->GetSkeleton())
 		{
 			const bool bTransient = Skeleton->GetPackage() == GetTransientPackage();
-
+			
 			// Duplicate only if transient or export all assets.
 			if (bTransient || Configuration.bExportAllResourcesOnBake)
 			{
-				FString SkeletonName = Mesh->GetName();
-				if (Mesh->HasAnyFlags(RF_Transient))
+				// Use the ResourceName already curated for the Skeletal Mesh
+				FResourceBakingData SkeletonData {SkeletonAssetPrefix ,UserGivenFilename, SkeletalMeshData.ResourceName, FString{}, AssetPath};
+
+				// If the skeleton was already duplicated use the already duplicated one
+				if (HandledSourceObjects.Contains(Skeleton))
 				{
-					SkeletonName = SkeletonName.Replace(*CustomizableObjectName,TEXT(""));
-					if (RemovePrefix(SkeletonName, {FBakingConfiguration::BakedResourcePrefix}))
+					for (UObject* DuplicatedObject : DuplicatedObjects)
 					{
-						MakeAssetNameNotUnique(SkeletonName);
+						if (DuplicatedObject &&
+							DuplicatedObject->IsA<USkeleton>() &&
+							DuplicatedObject->GetName() == SkeletonData.GetResourceComposedName())
+						{
+							ReplacementMap.Add(Skeleton, DuplicatedObject);
+
+							UE_LOG(LogMutable, Verbose, TEXT("Matching skeleton found with name %s in component %s"), *SkeletonData.GetResourceComposedName(), *ComponentName.ToString())
+							break;
+						}
 					}
 				}
-				RemovePrefix(SkeletonName, {SkeletalMeshAssetPrefix});
-				ComposeResourceName(SkeletonAssetPrefix, ObjectBaseName, SkeletonName);
-				MakeResourceNameUnique(SkeletonName, DuplicatedObjects);
-				
-				const EPackageSaveResolutionType SkeletonSaveResolution = GetPackageSaveResolution(AssetPath, SkeletonName);
-				FString SkeletonPkgName = AssetPath + FString("/") + SkeletonName;
-				UObject* DuplicatedSkeleton = FUnrealBakeHelpers::BakeHelper_DuplicateAsset(Skeleton, SkeletonName, 
-					SkeletonPkgName, &ReplacementMap, false, SkeletonSaveResolution);
+
+				SkeletonData.ComputePackageSaveResolutionType();
+				UObject* DuplicatedSkeleton = FUnrealBakeHelpers::BakeHelper_DuplicateAsset(Skeleton, &ReplacementMap, false, SkeletonData);
 
 				HandledSourceObjects.Add(Skeleton);
 				DuplicatedObjects.Add(DuplicatedSkeleton);
 
-				AddUniqueToSavePackage({DuplicatedSkeleton->GetPackage(), SkeletonSaveResolution},OutSavedPackages);
+				AddUniqueToSavePackage({DuplicatedSkeleton->GetPackage(), SkeletonData},OutSavedPackages);
 
 				ReplacementMap.Add(Skeleton, DuplicatedSkeleton);
 			}
@@ -945,28 +998,33 @@ bool BakeCustomizableObjectInstance(
 			// Duplicate only if transient or export all assets.
 			if (bTransient || Configuration.bExportAllResourcesOnBake)
 			{
-				FString PhysicsAssetName = Mesh->GetName();
-				if (Mesh->HasAnyFlags(RF_Transient))
+				// Use the ResourceName already curated for the Skeletal Mesh
+				FResourceBakingData PhysicsAssetData {PhysicsAssetPrefix, UserGivenFilename, SkeletalMeshData.ResourceName, FString{}, AssetPath};
+
+				// If the PhysicsAsset was already duplicated use the already duplicated one
+				if (HandledSourceObjects.Contains(PhysicsAsset))
 				{
-					PhysicsAssetName = PhysicsAssetName.Replace(*CustomizableObjectName,TEXT(""));
-					if (RemovePrefix(PhysicsAssetName, {FBakingConfiguration::BakedResourcePrefix}))
+					for (UObject* DuplicatedObject : DuplicatedObjects)
 					{
-						MakeAssetNameNotUnique(PhysicsAssetName);
+						if (DuplicatedObject &&
+							DuplicatedObject->IsA<UPhysicsAsset>() &&
+							DuplicatedObject->GetName() == PhysicsAssetData.GetResourceComposedName())
+						{
+							ReplacementMap.Add(PhysicsAsset, DuplicatedObject);
+
+							UE_LOG(LogMutable, Verbose, TEXT("Matching Physics Asset found with name %s in component %s"), *PhysicsAssetData.GetResourceComposedName(), *ComponentName.ToString())
+							break;
+						}
 					}
 				}
-				RemovePrefix(PhysicsAssetName, {SkeletalMeshAssetPrefix});
-				ComposeResourceName(PhysicsAssetPrefix, ObjectBaseName, PhysicsAssetName);
-				MakeResourceNameUnique(PhysicsAssetName, DuplicatedObjects);
-				
-				const EPackageSaveResolutionType PhysicsAssetSaveResolution = GetPackageSaveResolution(AssetPath, PhysicsAssetName);
-				FString PhysicsAssetPkgName = AssetPath + FString("/") + PhysicsAssetName;
-				UObject* DuplicatedPhysicsAsset = FUnrealBakeHelpers::BakeHelper_DuplicateAsset(PhysicsAsset, PhysicsAssetName, 
-					PhysicsAssetPkgName, &ReplacementMap, false, PhysicsAssetSaveResolution);
+
+				PhysicsAssetData.ComputePackageSaveResolutionType();
+				UObject* DuplicatedPhysicsAsset = FUnrealBakeHelpers::BakeHelper_DuplicateAsset(PhysicsAsset, &ReplacementMap, false, PhysicsAssetData);
 
 				HandledSourceObjects.Add(PhysicsAsset);
 				DuplicatedObjects.Add(DuplicatedPhysicsAsset);
 
-				AddUniqueToSavePackage({DuplicatedPhysicsAsset->GetPackage(), PhysicsAssetSaveResolution},OutSavedPackages);
+				AddUniqueToSavePackage({DuplicatedPhysicsAsset->GetPackage(), PhysicsAssetData},OutSavedPackages);
 
 				ReplacementMap.Add(PhysicsAsset, DuplicatedPhysicsAsset);
 
@@ -974,16 +1032,13 @@ bool BakeCustomizableObjectInstance(
 			}
 		}
 
-
-		// Skeletal Mesh
-		const EPackageSaveResolutionType MeshSaveResolution = GetPackageSaveResolution(AssetPath, SkeletalMeshName);
-		FString PkgName = AssetPath + FString("/") + SkeletalMeshName;
-		UObject* DuplicatedMesh = FUnrealBakeHelpers::BakeHelper_DuplicateAsset(Mesh, SkeletalMeshName, PkgName, 
-			&ReplacementMap, false, MeshSaveResolution);
+		SkeletalMeshData.ComputePackageSaveResolutionType();
+		UObject* DuplicatedMesh = FUnrealBakeHelpers::BakeHelper_DuplicateAsset(Mesh, &ReplacementMap, false, SkeletalMeshData);
+		
 		HandledSourceObjects.Add(Mesh);
 		DuplicatedObjects.Add(DuplicatedMesh);
 		
-		AddUniqueToSavePackage({DuplicatedMesh->GetPackage(), MeshSaveResolution},OutSavedPackages);
+		AddUniqueToSavePackage({DuplicatedMesh->GetPackage(), SkeletalMeshData},OutSavedPackages);
 
 		Mesh->Build();
 
@@ -1071,10 +1126,11 @@ bool BakeCustomizableObjectInstance(
 		// Prepare the list of assets we want to provide to "PromptForCheckoutAndSave" for saving
 		TArray<UPackage*> PackagesToSaveProxy;
 		PackagesToSaveProxy.Reserve(OutSavedPackages.Num());
-		for (TPair<UPackage*, EPackageSaveResolutionType> DataToSave : OutSavedPackages)
+		for (TPair<UPackage*, FResourceBakingData> DataToSave : OutSavedPackages)
 		{
 			// Ensure we have no duplicates here
 			check(!PackagesToSaveProxy.Contains(DataToSave.Key));
+			check(DataToSave.Value.Validate());
 			
 			PackagesToSaveProxy.Push(DataToSave.Key);
 		}
@@ -1101,4 +1157,11 @@ bool BakeCustomizableObjectInstance(
 	return false;
 }
 
+
+bool BakeCustomizableObjectInstance(UCustomizableObjectInstance& InInstance, const FBakingConfiguration& Configuration, bool bIsUnattendedExecution)
+{
+	TMap<UPackage*,const FResourceBakingData> DummyMap = TMap<UPackage*,const FResourceBakingData>{};
+	return BakeCustomizableObjectInstance(InInstance, Configuration, bIsUnattendedExecution, DummyMap);
+}
+
 #undef LOCTEXT_NAMESPACE 
diff --git a/Engine/Plugins/Mutable/Source/CustomizableObjectEditor/Private/MuCOE/UnrealBakeHelpers.cpp b/Engine/Plugins/Mutable/Source/CustomizableObjectEditor/Private/MuCOE/UnrealBakeHelpers.cpp
--- a/Engine/Plugins/Mutable/Source/CustomizableObjectEditor/Private/MuCOE/UnrealBakeHelpers.cpp
+++ b/Engine/Plugins/Mutable/Source/CustomizableObjectEditor/Private/MuCOE/UnrealBakeHelpers.cpp
@@ -10,40 +10,23 @@
 #include "Engine/Texture2D.h"
 
 #include "AssetRegistry/AssetRegistryModule.h"
-#include "AssetToolsModule.h"
 #include "IAssetTools.h"
 #include "Serialization/ArchiveReplaceObjectRef.h"
-#include "Materials/Material.h"
 #include "Materials/MaterialInstanceConstant.h"
-#include "Materials/MaterialInstanceDynamic.h"
 #include "Factories/MaterialInstanceConstantFactoryNew.h"
 
 #include "TextureResource.h"
 
-UObject* FUnrealBakeHelpers::BakeHelper_DuplicateAsset(UObject* Object, const FString& ObjName, const FString& PkgName,
-											  			TMap<UObject*, UObject*>* ReplacementMap, const bool bGenerateConstantMaterialInstances,
-											  			const EPackageSaveResolutionType SaveResolutionType)
-{
-	FString FinalObjectName = ObjName;
-	FString FinalPackageName = PkgName;
-	
-	if (SaveResolutionType == EPackageSaveResolutionType::NewFile)
-	{
-		// To prevent our file hash from being modified we will append a "_" so the GetModuleChecked do not mess up our number.
-		// If not done the hash will get re-hashed as a way of making the asset name unique. With the char added we make the method append a value at the end instead
-		FinalObjectName += "_";
-		FinalPackageName += "_";
-		
-		FAssetToolsModule& AssetToolsModule = FModuleManager::GetModuleChecked<FAssetToolsModule>("AssetTools");
-		AssetToolsModule.Get().CreateUniqueAssetName(FinalPackageName, "", FinalPackageName, FinalObjectName);
 
-		// If the last char is our safety char, remove it, otherwise this asset has got a new suffix added to make the name unique
-		FinalPackageName.RemoveFromEnd("_");
-		FinalObjectName.RemoveFromEnd("_");
-	}
 
+UObject* FUnrealBakeHelpers::BakeHelper_DuplicateAsset(UObject* Object, TMap<UObject*, UObject*>* ReplacementMap, const bool bGenerateConstantMaterialInstances,
+											  			FResourceBakingData& InOutPackageData)
+{
+	const FString FinalPackageFullName = InOutPackageData.GetResourceComposedFullPath();
+	const FString FinalPackageShortName = InOutPackageData.GetResourceComposedName();
+	
 	// Create or LOAD the package (if said package exists)
-	UPackage* Package = CreatePackage(*FinalPackageName);
+	UPackage* Package = CreatePackage(*FinalPackageFullName);
 	check(Package);
 	Package->FullyLoad();
 	
@@ -58,7 +41,7 @@ UObject* FUnrealBakeHelpers::BakeHelper_DuplicateAsset(UObject* Object, const FS
 		UMaterialInstanceConstantFactoryNew* MaterialFactory = NewObject<UMaterialInstanceConstantFactoryNew>();
 		MaterialFactory->InitialParent = MatInstance->Parent;
 
-		UMaterialInstanceConstant* MatInstanceConst = CastChecked<UMaterialInstanceConstant>(MaterialFactory->FactoryCreateNew(UMaterialInstanceConstant::StaticClass(), Package, FName(FinalObjectName), RF_NoFlags, nullptr, GWarn));
+		UMaterialInstanceConstant* MatInstanceConst = CastChecked<UMaterialInstanceConstant>(MaterialFactory->FactoryCreateNew(UMaterialInstanceConstant::StaticClass(), Package, FName(FinalPackageShortName), RF_NoFlags, nullptr, GWarn));
 
 		TMap<int32, UTexture*> EmptyTextureReplacementMap;
 		CopyAllMaterialParameters<UMaterialInstanceConstant>(*MatInstanceConst, *MatInstance, EmptyTextureReplacementMap);
@@ -67,7 +50,7 @@ UObject* FUnrealBakeHelpers::BakeHelper_DuplicateAsset(UObject* Object, const FS
 	}
 	else
 	{
-		FObjectDuplicationParameters Params = InitStaticDuplicateObjectParams(Object, Package, *FinalObjectName, RF_AllFlags, nullptr, EDuplicateMode::Normal);
+		FObjectDuplicationParameters Params = InitStaticDuplicateObjectParams(Object, Package, *FinalPackageShortName, RF_AllFlags, nullptr, EDuplicateMode::Normal);
 		DupObject = StaticDuplicateObjectEx(Params);
 	}
 
@@ -122,13 +105,13 @@ namespace
 }
 
 
-UTexture2D* FUnrealBakeHelpers::BakeHelper_CreateAssetTexture(UTexture2D* SourceTexture, const FString& TexObjName, const FString& TexPkgName, const UTexture* OrgTex,
-																TMap<UObject*, UObject*>* ReplacementMap, const EPackageSaveResolutionType SaveResolutionType)
+UTexture2D* FUnrealBakeHelpers::BakeHelper_CreateAssetTexture(UTexture2D* SourceTexture, const UTexture* OrgTex,
+																TMap<UObject*, UObject*>* ReplacementMap, FResourceBakingData& InOutPackageData)
 {
 	const bool bIsMutableTexture = !SourceTexture->Source.IsValid();
 	if (!bIsMutableTexture)
 	{
-		return Cast<UTexture2D>(BakeHelper_DuplicateAsset(SourceTexture, TexObjName, TexPkgName, ReplacementMap,false, SaveResolutionType));
+		return Cast<UTexture2D>(BakeHelper_DuplicateAsset(SourceTexture, ReplacementMap,false, InOutPackageData));
 	}
 
 	int32 SourceTextureSizeX = SourceTexture->GetPlatformData()->SizeX;
@@ -137,33 +120,14 @@ UTexture2D* FUnrealBakeHelpers::BakeHelper_CreateAssetTexture(UTexture2D* Source
 	ETextureSourceFormat PixelFormat = (SourceTexturePixelFormat == PF_BC4 || SourceTexturePixelFormat == PF_G8) ? TSF_G8 : TSF_BGRA8;
 
 	// Begin duplicate texture
-	FString FinalObjectName = TexObjName;
-	FString FinalPackageName = TexPkgName;
-
-	if (SaveResolutionType == EPackageSaveResolutionType::NewFile)
-	{
-		// As a safety measure that prevents the update of the hashed integer at the end, add a final char to prevent the the
-		// CreateUniqueAssetName from changing the value for our resource hash.
-		
-		// To prevent our hash from being modified
-		FinalObjectName += "_";
-		FinalPackageName += "_";
-		
-		// This is destroying the hash at the end of the TexObjectName....
-		// note: We are adding a sufix so it gets added to the FinalPkgName to prevent the method from changing the value of the hash
-		FAssetToolsModule& AssetToolsModule = FModuleManager::GetModuleChecked<FAssetToolsModule>("AssetTools");
-		AssetToolsModule.Get().CreateUniqueAssetName(FinalPackageName, TEXT(""), FinalPackageName, FinalObjectName);
-
-		// If the last char is our safety char, remove it, otherwise this asset has got a new sufix added to make the name unique
-		FinalPackageName.RemoveFromEnd("_");
-		FinalObjectName.RemoveFromEnd("_");
-	}
+	const FString FinalPackageFullName = InOutPackageData.GetResourceComposedFullPath();
+	const FString FinalPackageShortName = InOutPackageData.GetResourceComposedName();
 
-	UPackage* Package = CreatePackage(*FinalPackageName);
+	UPackage* Package = CreatePackage(*FinalPackageFullName);
 	check(Package);
 	Package->FullyLoad();
 
-	UTexture2D* DupTex = NewObject<UTexture2D>(Package, *FinalObjectName, RF_Public | RF_Standalone);
+	UTexture2D* DupTex = NewObject<UTexture2D>(Package, *FinalPackageShortName, RF_Public | RF_Standalone);
 
 	// The garbage collector is called in the middle of the bake process, and this can destroy this temporary object. 
 	// We add them to the garbage root to prevent this. This will avoid them being unloaded while the editor is running, but this
diff --git a/Engine/Plugins/Mutable/Source/CustomizableObjectEditor/Private/MuCOE/UnrealBakeHelpers.h b/Engine/Plugins/Mutable/Source/CustomizableObjectEditor/Private/MuCOE/UnrealBakeHelpers.h
--- a/Engine/Plugins/Mutable/Source/CustomizableObjectEditor/Private/MuCOE/UnrealBakeHelpers.h
+++ b/Engine/Plugins/Mutable/Source/CustomizableObjectEditor/Private/MuCOE/UnrealBakeHelpers.h
@@ -9,6 +9,7 @@
 #include "Materials/MaterialExpressionTextureBase.h"
 #include "Materials/MaterialInstanceConstant.h"
 #include "Materials/MaterialInterface.h"
+#include "MuCOE/CustomizableObjectInstanceBakingUtils.h"
 
 #define UE_API CUSTOMIZABLEOBJECTEDITOR_API
 
@@ -25,17 +26,16 @@ class FUnrealBakeHelpers
 {
 public:
 
-	static UE_API UObject* BakeHelper_DuplicateAsset(UObject* Object, const FString& ObjName, const FString& PkgName,
-											  TMap<UObject*, UObject*>* ReplacementMap, const bool bGenerateConstantMaterialInstances,
-											  const EPackageSaveResolutionType SaveResolutionType);
+	static UE_API UObject* BakeHelper_DuplicateAsset(UObject* Object,
+											  TMap<UObject*, UObject*>* ReplacementMap, const bool bGenerateConstantMaterialInstances, FResourceBakingData& InOutPackageData);
 
 	/**
 	 * Duplicates a texture asset. Duplicates Mutable and non Mutable textures.
 	 *
 	 * @param OrgTex Original source texture from which a Mutable texture has been generated. Only required when SrcTex is a Mutable texture.
 	 */
-	static UE_API UTexture2D* BakeHelper_CreateAssetTexture(UTexture2D* SourceTexture, const FString& TexObjName, const FString& TexPkgName, const UTexture* OrgTex,
-													 TMap<UObject*, UObject*>* ReplacementMap, const EPackageSaveResolutionType SaveResolutionType);
+	static UE_API UTexture2D* BakeHelper_CreateAssetTexture(UTexture2D* SourceTexture, const UTexture* OrgTex,
+													 TMap<UObject*, UObject*>* ReplacementMap, FResourceBakingData& InOutPackageData);
 
 	template<typename MaterialType>
 	static void CopyAllMaterialParameters(MaterialType& DestMaterial, UMaterialInterface& OriginMaterial, const TMap<int32, UTexture*>& TextureReplacementMap)
diff --git a/Engine/Plugins/Mutable/Source/MutableValidation/Private/MuV/COIBakingTestCommandlet.cpp b/Engine/Plugins/Mutable/Source/MutableValidation/Private/MuV/COIBakingTestCommandlet.cpp
--- a/Engine/Plugins/Mutable/Source/MutableValidation/Private/MuV/COIBakingTestCommandlet.cpp
+++ b/Engine/Plugins/Mutable/Source/MutableValidation/Private/MuV/COIBakingTestCommandlet.cpp
@@ -213,7 +213,7 @@ int32 UCOIBakingTestCommandlet::Main(const FString& Params)
 				Configuration.bExportAllResourcesOnBake = true;
 				Configuration.bGenerateConstantMaterialInstancesOnBake = true;
 
-				TMap<UPackage*, EPackageSaveResolutionType> SavedPackages;
+				TMap<UPackage*,const FResourceBakingData> SavedPackages;
 				bWasBakingSuccessful = BakeCustomizableObjectInstance(
 					*TargetInstance,
 					Configuration,
@@ -225,12 +225,12 @@ int32 UCOIBakingTestCommandlet::Main(const FString& Params)
 				const UEnum* Enum = StaticEnum<EPackageSaveResolutionType>();
 				check(Enum);
 				
-				for (const TPair<UPackage*, EPackageSaveResolutionType>& SavedPackagePair : SavedPackages)
+				for (const TPair<UPackage*,const FResourceBakingData>& SavedPackagePair : SavedPackages)
 				{
 					const UPackage* SavedPackage = SavedPackagePair.Key;
 					if (ensureMsgf(SavedPackage, TEXT("One entry of the \"SavedPackages\" map is a null pointer. All entries are expected to point to non null objects.")))
 					{
-						const FString SaveTypeString = Enum->GetNameByIndex(static_cast<uint8>(SavedPackagePair.Value)).ToString();
+						const FString SaveTypeString = Enum->GetNameByIndex(static_cast<uint8>(SavedPackagePair.Value.SaveResolutionType)).ToString();
 						UE_LOG(LogMutable, Display, TEXT("Stored package \"%s\" with \"%s\" as save resolution type."), *SavedPackage->GetName(), *SaveTypeString)
 					}
 				}
