[Mutable] Add checkf to failure case.

#FailedValidations Preflight

diff --git a/Engine/Plugins/Mutable/Source/CustomizableObjectEditor/Private/MuCOE/GenerateMutableSource/GenerateMutableSourceMesh.cpp b/Engine/Plugins/Mutable/Source/CustomizableObjectEditor/Private/MuCOE/GenerateMutableSource/GenerateMutableSourceMesh.cpp
--- a/Engine/Plugins/Mutable/Source/CustomizableObjectEditor/Private/MuCOE/GenerateMutableSource/GenerateMutableSourceMesh.cpp
+++ b/Engine/Plugins/Mutable/Source/CustomizableObjectEditor/Private/MuCOE/GenerateMutableSource/GenerateMutableSourceMesh.cpp
@@ -1181,7 +1181,8 @@ UE::Tasks::FTask ConvertSkeletalMeshToMutable(TSharedRef<FMeshConversionContext>
 	const uint32 MeshId = CityHash32(reinterpret_cast<const char*>(*MeshName), MeshName.Len() * sizeof(FString::ElementType));
 
 	const USkeletalMesh* SkeletalMesh = Cast<USkeletalMesh>(Source.Mesh.Get());
-	check(SkeletalMesh);
+
+	checkf(SkeletalMesh, TEXT("SkeletalMesh cast failed for mesh '%s' (MeshId: %u)"), *MeshName, MeshId);
 
 	const FSkeletalMeshModel* ImportedModel = SkeletalMesh->GetImportedModel();
 	if (!ImportedModel)
