[Mutable] Replaced ResourceID with TPassthroughObjectPtr.

- ResourceIDs are no longer stored in constants. Now are stored only in Args.
- NONE is now evaluated to false in all cases.
- UObjects to be loaded are now discovered on the BeginUpdate.
- Fixed multiple crashes with the CodeViewer.

#FailedValidations Preflight

diff --git a/Engine/Plugins/Mutable/Source/CustomizableObject/Internal/MuCO/BoneNames.h b/Engine/Plugins/Mutable/Source/CustomizableObject/Internal/MuCO/BoneNames.h
--- a/Engine/Plugins/Mutable/Source/CustomizableObject/Internal/MuCO/BoneNames.h
+++ b/Engine/Plugins/Mutable/Source/CustomizableObject/Internal/MuCO/BoneNames.h
@@ -10,6 +10,8 @@ class UModelResources;
 class FBoneNames
 {
 public:
+	CUSTOMIZABLEOBJECT_API FBoneNames();
+
 	CUSTOMIZABLEOBJECT_API FBoneNames(const UModelResources& ModelResources);
 
 	UE::Mutable::Private::FBoneName* Find(const FName& BoneName);
diff --git a/Engine/Plugins/Mutable/Source/CustomizableObject/Internal/MuCO/CustomizableObjectInstancePrivate.h b/Engine/Plugins/Mutable/Source/CustomizableObject/Internal/MuCO/CustomizableObjectInstancePrivate.h
--- a/Engine/Plugins/Mutable/Source/CustomizableObject/Internal/MuCO/CustomizableObjectInstancePrivate.h
+++ b/Engine/Plugins/Mutable/Source/CustomizableObject/Internal/MuCO/CustomizableObjectInstancePrivate.h
@@ -369,12 +369,6 @@ public:
 	// To be indexed with object component index
 	UPROPERTY(Transient)
 	TArray<FCustomizableInstanceComponentData> ComponentsData;
-
-	UPROPERTY(Transient)
-	TArray< TObjectPtr<UMaterialInterface> > ReferencedMaterials;
-
-	// Converts a ReferencedMaterials index from the CustomizableObject to an index in the ReferencedMaterials in the Instance
-	TMap<int32, uint32> ObjectToInstanceIndexMap;
 	
 	UPROPERTY(Transient)
 	TArray< TObjectPtr<UPhysicsAsset> > ClothingPhysicsAssets;
@@ -391,10 +385,6 @@ public:
 
 	UPROPERTY(Transient)
 	TArray<FExtensionInstanceData> ExtensionInstanceData;
-	
-	// The pass-through assets that will be loaded during an update
-	TArray<TSoftObjectPtr<const UTexture>> PassThroughTexturesToLoad;
-	TArray<TSoftObjectPtr<const UStreamableRenderAsset>> PassThroughMeshesToLoad;
 
 private:
 	ECOInstanceFlags InstanceFlagsPrivate = ECOInstanceFlags::ECONone;
diff --git a/Engine/Plugins/Mutable/Source/CustomizableObject/Internal/MuCO/CustomizableObjectPrivate.h b/Engine/Plugins/Mutable/Source/CustomizableObject/Internal/MuCO/CustomizableObjectPrivate.h
--- a/Engine/Plugins/Mutable/Source/CustomizableObject/Internal/MuCO/CustomizableObjectPrivate.h
+++ b/Engine/Plugins/Mutable/Source/CustomizableObject/Internal/MuCO/CustomizableObjectPrivate.h
@@ -1237,29 +1237,14 @@ public:
 	/** Skeletons used by the compiled UE::Mutable::Private::FModel. */
 	UPROPERTY()
 	TArray<TSoftObjectPtr<USkeleton>> Skeletons;
-
-	/** Materials used by the compiled UE::Mutable::Private::FModel */
-	UPROPERTY()
-	TArray<TSoftObjectPtr<UMaterialInterface>> Materials;
-
-	/** PassThrough textures used by the UE::Mutable::Private::FModel. */
-	UPROPERTY()
-	TArray<TSoftObjectPtr<UTexture>> PassThroughTextures;
-
-	/** PassThrough meshes used by the UE::Mutable::Private::FModel. */
+	
 	UPROPERTY()
-	TArray<TSoftObjectPtr<UStreamableRenderAsset>> PassThroughMeshes;
-
+	TMap<uint32, TSoftObjectPtr<UObject>> PassthroughObjects;
+	
 #if WITH_EDITORONLY_DATA
 	/** Runtime referenced textures used by the UE::Mutable::Private::FModel. */
 	UPROPERTY()
 	TArray<TSoftObjectPtr<UTexture2D>> RuntimeReferencedTextures;
-	
-	/** Runtime referenced meshes used by the UE::Mutable::Private::FModel.
-	* TODO: Move to FMutableSourceMeshData when runtime is really implemented.
-	*/
-	UPROPERTY()
-	TArray<TSoftObjectPtr<const UStreamableRenderAsset>> RuntimeReferencedMeshes;
 #endif
 	
 	/** Physics assets gathered from the SkeletalMeshes, to be used in mesh generation in-game */
@@ -1737,6 +1722,8 @@ enum class ECustomizableObjectVersions : int32
 
 	FixSkinWeightProfilesConversion,
 
+	MaterialIDMovedToArgs,
+
 	// -----<new versions can be added above this line>--------
 	LastCustomizableObjectVersion
 };
diff --git a/Engine/Plugins/Mutable/Source/CustomizableObject/Internal/MuCO/CustomizableObjectSystemPrivate.h b/Engine/Plugins/Mutable/Source/CustomizableObject/Internal/MuCO/CustomizableObjectSystemPrivate.h
--- a/Engine/Plugins/Mutable/Source/CustomizableObject/Internal/MuCO/CustomizableObjectSystemPrivate.h
+++ b/Engine/Plugins/Mutable/Source/CustomizableObject/Internal/MuCO/CustomizableObjectSystemPrivate.h
@@ -703,6 +703,9 @@ public:
 
 	/** If a InstanceUsage is in this set it means that its AttachParent has been modified (USkeletalMesh changed, UMaterial changed...). */
 	TSet<TWeakObjectPtr<UCustomizableObjectInstanceUsage>> AttachedParentUpdated;
+
+	/** Manage loading Passthrough Objects discovered on the StartUpdate. */
+	TSharedPtr<UE::Mutable::Private::FPassthroughObjectLoader> PassthroughObjectLoader;
 	
 	/** Hard references to objects. Avoids GC to collect them. */
 	TArray<TStrongObjectPtr<const UObject>> Objects;
@@ -710,10 +713,6 @@ public:
 	// Used during an update to prevent the pass-through textures loaded by LoadAdditionalAssetsAsync() from being unloaded by GC
 	// between AdditionalAssetsAsyncLoaded() and their setting into the generated materials in BuildMaterials()
 	TArray<TStrongObjectPtr<const UTexture>> LoadedPassThroughTexturesPendingSetMaterial;
-
-	// Used during an update to prevent the pass-through meshes loaded by LoadAdditionalAssetsAsync() from being unloaded by GC
-	// between AdditionalAssetsAsyncLoaded() and their setting into the generated materials in BuildMaterials()
-	TArray<TStrongObjectPtr<const UStreamableRenderAsset>> LoadedPassThroughMeshesPendingSetMaterial;
 	
 	TArray<UE::Mutable::Private::FImageId> CachedTextures;
 };
diff --git a/Engine/Plugins/Mutable/Source/CustomizableObject/Private/MuCO/BoneNames.cpp b/Engine/Plugins/Mutable/Source/CustomizableObject/Private/MuCO/BoneNames.cpp
--- a/Engine/Plugins/Mutable/Source/CustomizableObject/Private/MuCO/BoneNames.cpp
+++ b/Engine/Plugins/Mutable/Source/CustomizableObject/Private/MuCO/BoneNames.cpp
@@ -5,6 +5,11 @@
 #include "MuCO/CustomizableObjectPrivate.h"
 
 
+FBoneNames::FBoneNames()
+{
+}
+
+
 FBoneNames::FBoneNames(const UModelResources& ModelResources)
 {
 	BoneNamesMap.Reserve(ModelResources.BoneNamesMap.Num());
diff --git a/Engine/Plugins/Mutable/Source/CustomizableObject/Private/MuCO/CustomizableObject.cpp b/Engine/Plugins/Mutable/Source/CustomizableObject/Private/MuCO/CustomizableObject.cpp
--- a/Engine/Plugins/Mutable/Source/CustomizableObject/Private/MuCO/CustomizableObject.cpp
+++ b/Engine/Plugins/Mutable/Source/CustomizableObject/Private/MuCO/CustomizableObject.cpp
@@ -3378,9 +3378,7 @@ void UCustomizableObjectPrivate::LogMemory()
 			{
 				Arch << ModelResources->ReferenceSkeletalMeshesData;
 				Arch << ModelResources->Skeletons;
-				Arch << ModelResources->Materials;
-				Arch << ModelResources->PassThroughTextures;
-				Arch << ModelResources->PassThroughMeshes;
+				Arch << ModelResources->PassthroughObjects;
 				Arch << ModelResources->PhysicsAssets;
 				Arch << ModelResources->AnimBPs;
 				Arch << ModelResources->AnimBpOverridePhysiscAssetsInfo;
diff --git a/Engine/Plugins/Mutable/Source/CustomizableObject/Private/MuCO/CustomizableObjectInstance.cpp b/Engine/Plugins/Mutable/Source/CustomizableObject/Private/MuCO/CustomizableObjectInstance.cpp
--- a/Engine/Plugins/Mutable/Source/CustomizableObject/Private/MuCO/CustomizableObjectInstance.cpp
+++ b/Engine/Plugins/Mutable/Source/CustomizableObject/Private/MuCO/CustomizableObjectInstance.cpp
@@ -2332,26 +2332,11 @@ bool UCustomizableInstancePrivate::UpdateSkeletalMesh_PostBeginUpdate0(UCustomiz
 		// If it is a referenced resource, only the first LOD is relevant.
 		if (ComponentMesh->IsReference())
         {
-        	int32 ReferenceID = ComponentMesh->GetReferencedMesh();
-        	TSoftObjectPtr<UStreamableRenderAsset> Ref = ModelResources.PassThroughMeshes[ReferenceID];
-
-        	if (!Ref.IsValid())
-        	{
-        		// This shouldn't happen here synchronosuly. It should have been requested as an async load.
-        		UE_LOG(LogMutable, Error, TEXT("Referenced mesh [%s] was not pre-loaded. It will be sync-loaded probably causing a hitch. CO [%s]"), *Ref.ToString(), *GetNameSafe(CustomizableObject));
-        	}
-
-			UStreamableRenderAsset* Asset = UE::Mutable::Private::LoadObject(Ref);
-			USkeletalMesh* SkeletalMesh = Cast<USkeletalMesh>(Asset);
-			if (SkeletalMesh)
+	        if (USkeletalMesh* SkeletalMesh = ComponentMesh->PassthroughObject.Get())
 			{
 				SkeletalMeshes.Add(ComponentName, SkeletalMesh);
 			}
-			else
-			{
-				// Pass-through static meshes not implemented yet.
-				UE_LOG(LogMutable, Error, TEXT("Referenced static meshes [%s] are not supported yet. CO [%s]"), *Ref.ToString(), *GetNameSafe(CustomizableObject));
-			}
+			
         	continue;
         }
 
@@ -4958,10 +4943,7 @@ UE::Tasks::FTask LoadAdditionalAssetsAndData(const TSharedRef<FUpdateContextPriv
 	TArray<FSoftObjectPath> AssetsToStream;
 
 	TArray<FInstanceUpdateData::FComponent>& Components = Context->InstanceUpdateData.Components;
-
-	InstancePrivate->ObjectToInstanceIndexMap.Empty();
-	InstancePrivate->ReferencedMaterials.Empty();
-
+	
 	const int32 NumClothingAssets = ModelResources.ClothingAssetsData.Num();
 	InstancePrivate->ClothingPhysicsAssets.Reset(NumClothingAssets);
 	InstancePrivate->ClothingPhysicsAssets.SetNum(NumClothingAssets);
@@ -4969,67 +4951,6 @@ UE::Tasks::FTask LoadAdditionalAssetsAndData(const TSharedRef<FUpdateContextPriv
 	InstancePrivate->GatheredAnimBPs.Empty();
 	InstancePrivate->AnimBPGameplayTags.Reset();
 	InstancePrivate->AnimBpPhysicsAssets.Reset();
-	
-	for (const FInstanceUpdateData::FSurface& Surface : Context->InstanceUpdateData.Surfaces)
-	{
-		if (!Surface.Material)
-		{
-			continue;
-		}
-
-		if (Surface.Material->ReferenceID == INDEX_NONE)
-		{
-			continue;
-		}
-		
-		const int32 MaterialIndex = Surface.Material->ReferenceID;
-		if (InstancePrivate->ObjectToInstanceIndexMap.Contains(MaterialIndex))
-		{
-			continue;
-		}
-
-		TSoftObjectPtr<UMaterialInterface> AssetPtr = ModelResources.Materials.IsValidIndex(MaterialIndex) ? ModelResources.Materials[MaterialIndex] : nullptr;
-		UMaterialInterface* LoadedMaterial = AssetPtr.Get();
-
-		const int32 ReferencedMaterialsIndex = InstancePrivate->ReferencedMaterials.Add(LoadedMaterial);
-		InstancePrivate->ObjectToInstanceIndexMap.Add(MaterialIndex, ReferencedMaterialsIndex);
-
-		if (!LoadedMaterial && !AssetPtr.IsNull())
-		{
-			AssetsToStream.Add(AssetPtr.ToSoftObjectPath());
-		}
-	}
-
-	for (const FInstanceUpdateData::FComponent& Component : Context->InstanceUpdateData.Components)
-	{
-		if (!Component.OverlayMaterial)
-		{
-			continue;
-		}
-
-		if (Component.OverlayMaterial->ReferenceID == INDEX_NONE)
-		{
-			continue;
-		}
-		
-		const int32 MaterialIndex = Component.OverlayMaterial->ReferenceID;
-		if (InstancePrivate->ObjectToInstanceIndexMap.Contains(Component.OverlayMaterial->ReferenceID))
-		{
-			continue;
-		}
-
-		TSoftObjectPtr<UMaterialInterface> AssetPtr = ModelResources.Materials.IsValidIndex(MaterialIndex) ? ModelResources.Materials[Component.OverlayMaterial->ReferenceID] : nullptr;
-		UMaterialInterface* LoadedMaterial = AssetPtr.Get();
-
-		const int32 ReferencedMaterialsIndex = InstancePrivate->ReferencedMaterials.Add(LoadedMaterial);
-		InstancePrivate->ObjectToInstanceIndexMap.Add(MaterialIndex, ReferencedMaterialsIndex);
-
-		if (!LoadedMaterial && !AssetPtr.IsNull())
-		{
-			AssetsToStream.Add(AssetPtr.ToSoftObjectPath());
-		}
-	}
-
 
 	// Load Skeletons required by the SubMeshes of the newly generated Mesh, will be merged later
 	for (int32 InstanceComponentIndex = 0; InstanceComponentIndex < Context->NumInstanceComponents; ++InstanceComponentIndex)
@@ -5389,16 +5310,8 @@ UE::Tasks::FTask LoadAdditionalAssetsAndData(const TSharedRef<FUpdateContextPriv
 		}
 	}
 
-	for (TSoftObjectPtr<const UTexture>& TextureRef : InstancePrivate->PassThroughTexturesToLoad)
-	{
-		AssetsToStream.Add(TextureRef.ToSoftObjectPath());
-	}
-
-	for (TSoftObjectPtr<const UStreamableRenderAsset>& MeshRef : InstancePrivate->PassThroughMeshesToLoad)
-	{
-		AssetsToStream.Add(MeshRef.ToSoftObjectPath());
-	}
-
+	Context->PassthroughObjectLoader->Load(ModelResources.PassthroughObjects, AssetsToStream);
+	
 	// Copy FExtensionData Object node input from the Instance to the InstanceUpdateData
 	for (int32 ExtensionDataIndex = 0; ExtensionDataIndex < Context->MutableInstance->GetExtensionDataCount(); ExtensionDataIndex++)
 	{
@@ -5548,6 +5461,8 @@ void AdditionalAssetsAsyncLoaded(const TSharedRef<FUpdateContextPrivate> Context
 
 	UModelResources& ModelResources = *CustomizableObjectPrivate->GetModelResources();
 
+	Context->PassthroughObjectLoader->Resolve(ModelResources.PassthroughObjects);
+	
 	for (int32 ResourceIndex : Context->StreamedResourceIndex)
 	{
 		ModelResources.StreamedResourceData[ResourceIndex].Hold();
@@ -5557,33 +5472,6 @@ void AdditionalAssetsAsyncLoaded(const TSharedRef<FUpdateContextPrivate> Context
 	{
 		ModelResources.StreamedExtensionData[ResourceIndex].Hold();
 	}
-	
-	// Loaded Materials
-	check(InstancePrivate->ObjectToInstanceIndexMap.Num() == InstancePrivate->ReferencedMaterials.Num());
-
-	for (TPair<uint32, uint32> Pair : InstancePrivate->ObjectToInstanceIndexMap)
-	{
-		const TSoftObjectPtr<UMaterialInterface>& AssetPtr = ModelResources.Materials.IsValidIndex(Pair.Key) ? ModelResources.Materials[Pair.Key] : nullptr;
-		InstancePrivate->ReferencedMaterials[Pair.Value] = AssetPtr.Get();
-
-#if WITH_EDITOR
-		if (!InstancePrivate->ReferencedMaterials[Pair.Value])
-		{
-			if (!AssetPtr.IsNull())
-			{
-				FString ErrorMsg = FString::Printf(TEXT("Mutable couldn't load the material [%s] and won't be rendered. If it has been deleted or renamed, please recompile all the mutable objects that use it."), *AssetPtr.GetAssetName());
-				UE_LOG(LogMutable, Error, TEXT("%s"), *ErrorMsg);
-
-				FMessageLog MessageLog("Mutable");
-				MessageLog.Notify(FText::FromString(ErrorMsg), EMessageSeverity::Error, true);
-			}
-			else
-			{
-				ensure(false); // Couldn't load the material, and we don't know which material
-			}
-		}
-#endif
-	}
 
 	for (FCustomizableInstanceComponentData& ComponentData : InstancePrivate->ComponentsData)
 	{
@@ -5692,26 +5580,6 @@ void AdditionalAssetsAsyncLoaded(const TSharedRef<FUpdateContextPrivate> Context
 		}
 		ComponentData.PhysicsAssets.AdditionalPhysicsAssetsToLoad.Empty();
 	}
-
-	Context->LoadedPassThroughTexturesPendingSetMaterial.Empty(InstancePrivate->PassThroughTexturesToLoad.Num());
-
-	for (TSoftObjectPtr<const UTexture>& TextureRef : InstancePrivate->PassThroughTexturesToLoad)
-	{
-		ensure(TextureRef.IsValid());
-		Context->LoadedPassThroughTexturesPendingSetMaterial.Add(TStrongObjectPtr(TextureRef.Get()));
-	}
-
-	InstancePrivate->PassThroughTexturesToLoad.Empty();
-
-	Context->LoadedPassThroughMeshesPendingSetMaterial.Empty(InstancePrivate->PassThroughMeshesToLoad.Num());
-
-	for (TSoftObjectPtr<const UStreamableRenderAsset>& MeshRef : InstancePrivate->PassThroughMeshesToLoad)
-	{
-		ensure(MeshRef.IsValid());
-		Context->LoadedPassThroughMeshesPendingSetMaterial.Add(TStrongObjectPtr(MeshRef.Get()));
-	}
-
-	InstancePrivate->PassThroughMeshesToLoad.Empty();
 	
 #if WITH_EDITOR
 	// TODO: Remove with UE-217665 when the underlying bug in the ColorPicker is solved
@@ -5894,22 +5762,8 @@ void UCustomizableInstancePrivate::BuildMaterials(const TSharedRef<FUpdateContex
 
 			if (Component.OverlayMaterial)
 			{
-				if (Component.OverlayMaterial->ReferenceID != INDEX_NONE)
-				{
-					if (const uint32* ReferencedMaterialIndex = ObjectToInstanceIndexMap.Find(Component.OverlayMaterial->ReferenceID))
-					{
-						if (ReferencedMaterials.IsValidIndex(*ReferencedMaterialIndex))
-						{
-							ComponentData->OverlayMaterial = ReferencedMaterials[*ReferencedMaterialIndex];
-							OverlayMaterial = ComponentData->OverlayMaterial;
-						}
-					}
-				}
-				else
-				{
-					ComponentData->OverlayMaterial = Component.OverlayMaterial->Material.Get();
-					OverlayMaterial = ComponentData->OverlayMaterial;
-				}
+				ComponentData->OverlayMaterial = Component.OverlayMaterial->PassthroughObject.Get();
+				OverlayMaterial = ComponentData->OverlayMaterial;
 			}
 		}
 
@@ -5976,17 +5830,7 @@ void UCustomizableInstancePrivate::BuildMaterials(const TSharedRef<FUpdateContex
 					continue;
 				}
 
-				UMaterialInterface* MaterialTemplate = nullptr;
-
-				if (Surface.Material->ReferenceID == INDEX_NONE)
-				{
-					MaterialTemplate = Surface.Material->Material.Get();
-				}
-				else if (const uint32* ReferencedMaterialIndex = ObjectToInstanceIndexMap.Find(Surface.Material->ReferenceID))
-				{
-					MaterialTemplate = ReferencedMaterials[*ReferencedMaterialIndex];
-				}
-				
+				UMaterialInterface* MaterialTemplate = Surface.Material->PassthroughObject.Get();
 				if (!MaterialTemplate)
 				{
 					// Missing MaterialTemplate. Use DefaultMaterial instead. 
@@ -6178,29 +6022,10 @@ void UCustomizableInstancePrivate::BuildMaterials(const TSharedRef<FUpdateContex
 						{
 							check(MutableImage->IsReference());
 
-							uint32 ReferenceID = MutableImage->GetReferencedTexture();
-							if (ModelResources.PassThroughTextures.IsValidIndex(ReferenceID))
-							{
-								TSoftObjectPtr<UTexture> Ref = ModelResources.PassThroughTextures[ReferenceID];
-
-								// The texture should have been loaded by now by LoadAdditionalAssetsAsync()
-								PassThroughTexture = Ref.Get();
-
-								if (!PassThroughTexture)
-								{
-									// The texture should be loaded, something went wrong, possibly a bug in LoadAdditionalAssetsAsync()
-									UE_LOG(LogMutable, Error,
-										TEXT("Pass-through texture with name %s hasn't been loaded yet in BuildMaterials(). Forcing sync load."),
-										*Ref.ToSoftObjectPath().ToString());
-									ensure(false);
-									PassThroughTexture = UE::Mutable::Private::LoadObject(Ref);
-								}
-							}
-
+							PassThroughTexture = MutableImage->PassthroughObject.Get();
+							
 							if (!PassThroughTexture)
 							{
-								// Internal error.
-								UE_LOG(LogMutable, Error, TEXT("Missing referenced image [%d]."), ReferenceID);
 								continue;
 							}
 						}
diff --git a/Engine/Plugins/Mutable/Source/CustomizableObject/Private/MuCO/CustomizableObjectMeshUpdate.cpp b/Engine/Plugins/Mutable/Source/CustomizableObject/Private/MuCO/CustomizableObjectMeshUpdate.cpp
--- a/Engine/Plugins/Mutable/Source/CustomizableObject/Private/MuCO/CustomizableObjectMeshUpdate.cpp
+++ b/Engine/Plugins/Mutable/Source/CustomizableObject/Private/MuCO/CustomizableObjectMeshUpdate.cpp
@@ -242,7 +242,7 @@ namespace impl
 #endif
 
 		// For now, we are forcing the recreation of mutable-side instances with every update.
-		UE::Mutable::Private::FInstance::FID InstanceID = System->NewInstance(Model, OperationData->ExternalResourceProvider);
+		UE::Mutable::Private::FInstance::FID InstanceID = System->NewInstance(Model, OperationData->ExternalResourceProvider, nullptr);
 		UE_LOG(LogMutable, Verbose, TEXT("Creating Mutable instance with id [%d] for a mesh update"), InstanceID);
 
 		// LOD mask, set to all ones to build all LODs
diff --git a/Engine/Plugins/Mutable/Source/CustomizableObject/Private/MuCO/CustomizableObjectMipDataProvider.cpp b/Engine/Plugins/Mutable/Source/CustomizableObject/Private/MuCO/CustomizableObjectMipDataProvider.cpp
--- a/Engine/Plugins/Mutable/Source/CustomizableObject/Private/MuCO/CustomizableObjectMipDataProvider.cpp
+++ b/Engine/Plugins/Mutable/Source/CustomizableObject/Private/MuCO/CustomizableObjectMipDataProvider.cpp
@@ -131,7 +131,7 @@ namespace Impl
 #endif
 
 		// For now, we are forcing the recreation of mutable-side instances with every update.
-		UE::Mutable::Private::FInstance::FID InstanceID = System->NewInstance(Model, OperationData->UpdateContext->ExternalResourceProvider);
+		UE::Mutable::Private::FInstance::FID InstanceID = System->NewInstance(Model, OperationData->UpdateContext->ExternalResourceProvider, nullptr);
 		UE_LOG(LogMutable, Verbose, TEXT("Creating Mutable instance with id [%d] for a single UpdateImage"), InstanceID)
 
 		TSharedPtr<const UE::Mutable::Private::FInstance> Instance;
diff --git a/Engine/Plugins/Mutable/Source/CustomizableObject/Private/MuCO/CustomizableObjectSystem.cpp b/Engine/Plugins/Mutable/Source/CustomizableObject/Private/MuCO/CustomizableObjectSystem.cpp
--- a/Engine/Plugins/Mutable/Source/CustomizableObject/Private/MuCO/CustomizableObjectSystem.cpp
+++ b/Engine/Plugins/Mutable/Source/CustomizableObject/Private/MuCO/CustomizableObjectSystem.cpp
@@ -2349,20 +2349,7 @@ namespace impl
 
 		Image.Image = GetImageTask.GetResult();
 		check(Image.Image->IsReference());
-
-		const uint32 ReferenceID = Image.Image->GetReferencedTexture();
-
-		if (ModelResources.PassThroughTextures.IsValidIndex(ReferenceID))
-		{
-			const TSoftObjectPtr<const UTexture> Ref = ModelResources.PassThroughTextures[ReferenceID];
-			Instance->GetPrivate()->PassThroughTexturesToLoad.Add(Ref);
-		}
-		else
-		{
-			// internal error.
-			UE_LOG(LogMutable, Error, TEXT("Referenced image [%d] was not stored in the resource array."), ReferenceID);
-		}
-			
+		
 		Task_Mutable_GetMeshes_GetImage_Loop(OperationData, StartTime, GetImagesData, ++GetImageIndex);
 	}
 
@@ -2615,28 +2602,7 @@ namespace impl
 		FInstanceUpdateData::FLOD& LOD = OperationData->InstanceUpdateData.LODs[LODIndex];
 
 		LOD.Mesh = GetMeshTask.GetResult();
-
-		if (LOD.Mesh &&
-			LOD.Mesh->IsReference())
-		{
-			const UCustomizableObjectInstance* Instance = OperationData->Instance.Get();
-			const UCustomizableObject* CustomizableObject = Instance->GetCustomizableObject();
-			const UModelResources& ModelResources = *CustomizableObject->GetPrivate()->GetModelResources();
-
-			uint32 ReferenceID = LOD.Mesh->GetReferencedMesh();
-
-			if (ModelResources.PassThroughMeshes.IsValidIndex(ReferenceID))
-			{
-				TSoftObjectPtr<const UStreamableRenderAsset> Ref = ModelResources.PassThroughMeshes[ReferenceID];
-				Instance->GetPrivate()->PassThroughMeshesToLoad.Add(Ref);
-			}
-			else
-			{
-				// internal error.
-				UE_LOG(LogMutable, Error, TEXT("Referenced mesh [%d] was not stored in the resource array."), ReferenceID);
-			}
-		}
-			
+		
 		Task_Mutable_GetMeshes_GetMesh_Loop(OperationData, StartTime, GetMeshesData, ++GetMeshIndex);
 	}
 	
@@ -2687,10 +2653,7 @@ namespace impl
 			check(UCustomizableObjectSystem::GetInstance()->GetPrivate() != nullptr);
 
 			UCustomizableInstancePrivate* CustomizableObjectInstancePrivateData = OperationData->Instance->GetPrivate();
-
-			CustomizableObjectInstancePrivateData->PassThroughTexturesToLoad.Empty();
-			CustomizableObjectInstancePrivateData->PassThroughMeshesToLoad.Empty();
-
+			
 			if (OperationData->PixelFormatOverride)
 			{
 				OperationData->MutableSystem->SetImagePixelConversionOverride( OperationData->PixelFormatOverride );
@@ -3672,7 +3635,7 @@ namespace impl
 
 			Operation->SetFirstRequestedLOD(RequestedLODs);
 		}
-
+		
 		// CreateMutableInstance
 		{
 			if (FLogBenchmarkUtil::IsBenchmarkingReportingEnabled())
@@ -3696,19 +3659,20 @@ namespace impl
 			Operation->BoneNames = BoneNames;
 			
 			Operation->ExternalResourceProvider = MakeShared<FUnrealMutableResourceProvider>(BoneNames);
-			
+			Operation->PassthroughObjectLoader = MakeShared<UE::Mutable::Private::FPassthroughObjectLoader>();
+
 			if (Operation->bLiveUpdateMode)
 			{
 				if (Operation->InstanceID == 0)
 				{
 					// It's the first update since the instance was put in LiveUpdate Mode, this ID will be reused from now on
-					Operation->InstanceID = MutableSystem->NewInstance(Operation->Model, Operation->ExternalResourceProvider);
+					Operation->InstanceID = MutableSystem->NewInstance(Operation->Model, Operation->ExternalResourceProvider, Operation->PassthroughObjectLoader);
 					UE_LOG(LogMutable, Verbose, TEXT("Creating Mutable instance with id [%d] for reuse "), Operation->InstanceID);
 				}
 				else
 				{
 					// The instance was already in LiveUpdate Mode, the ID is reused
-					MutableSystem->ReuseInstance(Operation->InstanceID, Operation->ExternalResourceProvider);
+					MutableSystem->ReuseInstance(Operation->InstanceID, Operation->ExternalResourceProvider, Operation->PassthroughObjectLoader);
 					UE_LOG(LogMutable, Verbose, TEXT("Reusing Mutable instance with id [%d] "), Operation->InstanceID);
 				}
 			}
@@ -3716,10 +3680,10 @@ namespace impl
 			{
 				// In non-LiveUpdate mode, we are forcing the recreation of mutable-side instances with every update.
 				check(Operation->InstanceID == 0);
-				Operation->InstanceID = MutableSystem->NewInstance(Operation->Model, Operation->ExternalResourceProvider);
+				Operation->InstanceID = MutableSystem->NewInstance(Operation->Model, Operation->ExternalResourceProvider, Operation->PassthroughObjectLoader);
 				UE_LOG(LogMutable, Verbose, TEXT("Creating Mutable instance with id [%d] "), Operation->InstanceID);
 			}
-
+			
 			Operation->MutableInstance = MutableSystem->BeginUpdate_GameThread(Operation->InstanceID,
 				Operation->Parameters,
 				Operation->MeshIdRegistry,
diff --git a/Engine/Plugins/Mutable/Source/CustomizableObjectEditor/Private/MuCOE/CustomizableObjectCompiler.cpp b/Engine/Plugins/Mutable/Source/CustomizableObjectEditor/Private/MuCOE/CustomizableObjectCompiler.cpp
--- a/Engine/Plugins/Mutable/Source/CustomizableObjectEditor/Private/MuCOE/CustomizableObjectCompiler.cpp
+++ b/Engine/Plugins/Mutable/Source/CustomizableObjectEditor/Private/MuCOE/CustomizableObjectCompiler.cpp
@@ -953,37 +953,14 @@ void FCustomizableObjectCompiler::CompileInternal(bool bAsync)
 		ModelResources->CodeVersion = GetECustomizableObjectVersionEnumHash();
 
 		ModelResources->ReferenceSkeletalMeshesData = MoveTemp(GenerationContext.ReferenceSkeletalMeshesData);
-
-		ModelResources->Materials.Reserve(GenerationContext.ReferencedMaterials.Num());
-		for (const UMaterialInterface* Material : GenerationContext.ReferencedMaterials)
-		{
-			ModelResources->Materials.Emplace(const_cast<UMaterialInterface*>(Material));
-		}
-
-		for (const TPair<TSoftObjectPtr<UStreamableRenderAsset>, FMutableGraphGenerationContext::FGeneratedReferencedMesh>& Pair : GenerationContext.PassthroughMeshMap)
-		{
-			check(Pair.Value.ID == ModelResources->PassThroughMeshes.Num());
-			ModelResources->PassThroughMeshes.Add(Pair.Key);
-		}
-
-		for (const TPair<FMutableSourceMeshData, FMutableGraphGenerationContext::FGeneratedReferencedMesh>& Pair : GenerationContext.RuntimeReferencedMeshMap)
-		{
-			check(Pair.Value.ID == ModelResources->RuntimeReferencedMeshes.Num());
-			ModelResources->RuntimeReferencedMeshes.Add(Pair.Key.Mesh);
-		}
-
-		for (const TPair<TSoftObjectPtr<UTexture>, FMutableGraphGenerationContext::FGeneratedReferencedTexture>& Pair : GenerationContext.PassthroughTextureMap)
-		{
-			check(Pair.Value.ID == ModelResources->PassThroughTextures.Num());
-			ModelResources->PassThroughTextures.Add(Pair.Key);
-		}
-
+		
 		for (const TPair<TSoftObjectPtr<UTexture2D>, FMutableGraphGenerationContext::FGeneratedReferencedTexture>& Pair : GenerationContext.RuntimeReferencedTextureMap)
 		{
 			check(Pair.Value.ID == ModelResources->RuntimeReferencedTextures.Num());
 			ModelResources->RuntimeReferencedTextures.Add(Pair.Key);
 		}
 
+		ModelResources->PassthroughObjects = MoveTemp(GenerationContext.PassthroughObjects);
 		ModelResources->TextureParameterDefaultValues = MoveTemp(GenerationContext.TextureParameterDefaultValues);
 		ModelResources->SkeletalMeshParameterDefaultValues = MoveTemp(GenerationContext.SkeletalMeshParameterDefaultValues);
 		ModelResources->MaterialParameterDefaultValues = MoveTemp(GenerationContext.MaterialParameterDefaultValues);
@@ -1657,14 +1634,6 @@ UE::Mutable::Private::Ptr<UE::Mutable::Private::Node> FCustomizableObjectCompile
 		OutCompilerReferencedTextures.Add(Tex);
 	}
 
-	// Pass out the referenced meshes
-	OutRuntimeReferencedMeshes.Empty();
-	for (const TPair<FMutableSourceMeshData, FMutableGraphGenerationContext::FGeneratedReferencedMesh>& Pair : GenerationContext.RuntimeReferencedMeshMap)
-	{
-		check(Pair.Value.ID == OutRuntimeReferencedMeshes.Num());
-		OutRuntimeReferencedMeshes.Add(Pair.Key.Mesh);
-	}
-
 	OutCompilerReferencedMeshes.Empty();
 	for (const TPair<FMutableSourceMeshData, FMutableGraphGenerationContext::FGeneratedReferencedMesh>& Pair : GenerationContext.CompileTimeMeshMap)
 	{
diff --git a/Engine/Plugins/Mutable/Source/CustomizableObjectEditor/Private/MuCOE/CustomizableObjectEditor.cpp b/Engine/Plugins/Mutable/Source/CustomizableObjectEditor/Private/MuCOE/CustomizableObjectEditor.cpp
--- a/Engine/Plugins/Mutable/Source/CustomizableObjectEditor/Private/MuCOE/CustomizableObjectEditor.cpp
+++ b/Engine/Plugins/Mutable/Source/CustomizableObjectEditor/Private/MuCOE/CustomizableObjectEditor.cpp
@@ -1495,7 +1495,7 @@ void FCustomizableObjectEditor::CodeViewer() const
 
 	UE::Mutable::Private::Ptr<UE::Mutable::Private::Node> RootNode = DebuggerCompile(Model, RuntimeTextures, CompilerTextures, RuntimeMeshes, CompilerMeshes);
 	
-	const TSharedRef<FBoneNames> BoneNames = MakeShared<FBoneNames>(*CustomizableObject->GetPrivate()->GetModelResources());
+	const TSharedRef<FBoneNames> BoneNames = MakeShared<FBoneNames>(); // TODO DebuggerCompile does not create a UModelResource
 	TSharedPtr<FUnrealMutableResourceProvider> ExternalResourceProvider = MakeShared<FUnrealMutableResourceProvider>(BoneNames);
 	
 	const TSharedPtr<SDockTab> NewMutableObjectTab = SNew(SDockTab)
diff --git a/Engine/Plugins/Mutable/Source/CustomizableObjectEditor/Private/MuCOE/GenerateMutableSource/GenerateMutableSource.cpp b/Engine/Plugins/Mutable/Source/CustomizableObjectEditor/Private/MuCOE/GenerateMutableSource/GenerateMutableSource.cpp
--- a/Engine/Plugins/Mutable/Source/CustomizableObjectEditor/Private/MuCOE/GenerateMutableSource/GenerateMutableSource.cpp
+++ b/Engine/Plugins/Mutable/Source/CustomizableObjectEditor/Private/MuCOE/GenerateMutableSource/GenerateMutableSource.cpp
@@ -1707,35 +1707,45 @@ TSharedPtr<UE::Mutable::Private::FImage> GenerateImageConstant(UTexture* Texture
 	// For now it is not a problem because passthrough textures shouldn't mix with any other operations.
 	UE::Mutable::Private::FImageDesc ImageDesc = GenerateImageDescriptor(Texture);
 
-	FMutableGraphGenerationContext::FGeneratedReferencedTexture InvalidEntry;
-	InvalidEntry.ID = TNumericLimits<uint32>::Max();
+	UE::Mutable::Private::PASSTHROUGH_ID Id;
 
-	FMutableGraphGenerationContext::FGeneratedReferencedTexture* Entry;
-	int32 Num = 0;
-	
 	if (bIsReference)
 	{
-		Entry = &GenerationContext.PassthroughTextureMap.FindOrAdd(Texture, InvalidEntry);
-		Num = GenerationContext.PassthroughTextureMap.Num();
+		Id = GenerationContext.AddPassthroughObject(*Texture);
 	}
 	else if (bIsCompileTime)
 	{
-        Entry = &GenerationContext.CompileTimeTextureMap.FindOrAdd(CastChecked<UTexture2D>(Texture), InvalidEntry);
-		Num = GenerationContext.CompileTimeTextureMap.Num();
+		FMutableGraphGenerationContext::FGeneratedReferencedTexture InvalidEntry;
+		InvalidEntry.ID = UE::Mutable::Private::PASSTHROUGH_ID_INVALID;
+
+        FMutableGraphGenerationContext::FGeneratedReferencedTexture* Entry = &GenerationContext.CompileTimeTextureMap.FindOrAdd(CastChecked<UTexture2D>(Texture), InvalidEntry);
+		int32 Num = GenerationContext.CompileTimeTextureMap.Num();
+
+		if (Entry->ID == UE::Mutable::Private::PASSTHROUGH_ID_INVALID)
+		{
+			Entry->ID = Num - 1;
+		}
+
+		Id = Entry->ID;
 	}
 	else
 	{
-		Entry = &GenerationContext.RuntimeReferencedTextureMap.FindOrAdd(CastChecked<UTexture2D>(Texture), InvalidEntry);
-		Num = GenerationContext.RuntimeReferencedTextureMap.Num();
-	}
+		FMutableGraphGenerationContext::FGeneratedReferencedTexture InvalidEntry;
+		InvalidEntry.ID = UE::Mutable::Private::PASSTHROUGH_ID_INVALID;
 
-	if (Entry->ID == TNumericLimits<uint32>::Max())
-	{
-		Entry->ID = Num - 1;
+		FMutableGraphGenerationContext::FGeneratedReferencedTexture* Entry = &GenerationContext.RuntimeReferencedTextureMap.FindOrAdd(CastChecked<UTexture2D>(Texture), InvalidEntry);
+		int32 Num = GenerationContext.RuntimeReferencedTextureMap.Num();
+
+		if (Entry->ID == UE::Mutable::Private::PASSTHROUGH_ID_INVALID)
+		{
+			Entry->ID = Num - 1;
+		}
+
+		Id = Entry->ID;
 	}
 	
 	// Compile-time references that are left should be resolved immediately (should only happen in editor).
-	TSharedPtr<UE::Mutable::Private::FImage> Result = UE::Mutable::Private::FImage::CreateAsReference(Entry->ID, ImageDesc, bForceLoad);
+	TSharedPtr<UE::Mutable::Private::FImage> Result = UE::Mutable::Private::FImage::CreateAsReference(UE::Mutable::Private::TPassthroughObjectPtr<UTexture>(Id), ImageDesc, bForceLoad); // TODO GMT If bIsCompileTime it should not create a TPassthroughObjectPtr since it is not a Passthrough
 	return Result;
 }
 
@@ -1768,35 +1778,34 @@ TSharedPtr<UE::Mutable::Private::FMesh> GenerateMeshConstant(
 		}
 	}
 
-
-	FMutableGraphGenerationContext::FGeneratedReferencedMesh InvalidEntry;
-	InvalidEntry.ID = TNumericLimits<uint32>::Max();
-
-	FMutableGraphGenerationContext::FGeneratedReferencedMesh* Entry = nullptr;
-	int32 Num = 0;
+	UE::Mutable::Private::PASSTHROUGH_ID Id;
 
 	if (Source.bIsPassthrough)
 	{
-		Entry = &GenerationContext.PassthroughMeshMap.FindOrAdd(Source.Mesh, InvalidEntry);
-		Num = GenerationContext.PassthroughMeshMap.Num();
+		Id = GenerationContext.AddPassthroughObject(*Source.Mesh);
 	}
 	else if (bIsCompileTime)
 	{
-		Entry = &GenerationContext.CompileTimeMeshMap.FindOrAdd(Source, InvalidEntry);
-		Num = GenerationContext.CompileTimeMeshMap.Num();
+		FMutableGraphGenerationContext::FGeneratedReferencedMesh InvalidEntry;
+		InvalidEntry.ID = UE::Mutable::Private::PASSTHROUGH_ID_INVALID;
+
+		FMutableGraphGenerationContext::FGeneratedReferencedMesh* Entry = &GenerationContext.CompileTimeMeshMap.FindOrAdd(Source, InvalidEntry);
+		int32 Num = GenerationContext.CompileTimeMeshMap.Num();
+
+		if (Entry->ID == UE::Mutable::Private::PASSTHROUGH_ID_INVALID)
+		{
+			Entry->ID = Num - 1;
+		}
+
+		Id = Entry->ID;
 	}
 	else
 	{
-		Entry = &GenerationContext.RuntimeReferencedMeshMap.FindOrAdd(Source, InvalidEntry);
-		Num = GenerationContext.RuntimeReferencedMeshMap.Num();
+		Id = UE::Mutable::Private::PASSTHROUGH_ID_INVALID;
+		unimplemented();
 	}
 
-	if (Entry->ID == TNumericLimits<uint32>::Max())
-	{
-		Entry->ID = Num - 1;
-	}
-
-	TSharedPtr<UE::Mutable::Private::FMesh> Result = UE::Mutable::Private::FMesh::CreateAsReference(Entry->ID, bForceLoad);
+	TSharedPtr<UE::Mutable::Private::FMesh> Result = UE::Mutable::Private::FMesh::CreateAsReference(UE::Mutable::Private::TPassthroughObjectPtr<USkeletalMesh>(Id), bForceLoad); // TODO GMT If bIsCompileTime it should not create a TPassthroughObjectPtr since it is not a Passthrough
 	return Result;
 }
 
@@ -1920,3 +1929,26 @@ uint32 GetTypeHash(const FSharedSurfaceKey& Key)
 {
 	return HashCombine(GetTypeHash(Key.MaterialGuid), GetTypeHash(Key.MacroContextGuid));
 }
+
+
+UE::Mutable::Private::PASSTHROUGH_ID FMutableGraphGenerationContext::AddPassthroughObject(UObject& Object)
+{
+	TSoftObjectPtr<UObject> SoftObjectPtr(&Object);
+		
+	UE::Mutable::Private::PASSTHROUGH_ID Id = UE::Mutable::Private::GetPassthroughId(SoftObjectPtr);
+	while (TSoftObjectPtr<UObject>* PassthroughObject = PassthroughObjects.Find(Id))
+	{
+		if (*PassthroughObject == &Object)
+		{
+			return Id;
+		}
+		else
+		{
+			++Id;
+		}
+	}
+	
+	PassthroughObjects.Add(Id, SoftObjectPtr);
+
+	return Id;
+}
\ No newline at end of file
diff --git a/Engine/Plugins/Mutable/Source/CustomizableObjectEditor/Private/MuCOE/GenerateMutableSource/GenerateMutableSource.h b/Engine/Plugins/Mutable/Source/CustomizableObjectEditor/Private/MuCOE/GenerateMutableSource/GenerateMutableSource.h
--- a/Engine/Plugins/Mutable/Source/CustomizableObjectEditor/Private/MuCOE/GenerateMutableSource/GenerateMutableSource.h
+++ b/Engine/Plugins/Mutable/Source/CustomizableObjectEditor/Private/MuCOE/GenerateMutableSource/GenerateMutableSource.h
@@ -674,14 +674,14 @@ struct FMutableGraphGenerationContext
 	/** Data stored per-generated passthrough texture. */
 	struct FGeneratedReferencedTexture
 	{
-		uint32 ID;
+		UE::Mutable::Private::PASSTHROUGH_ID ID;
 		//UE::Mutable::Private::FImageDesc ImageDesc;
 	};
 
 	/** Data stored per-generated passthrough mesh. */
 	struct FGeneratedReferencedMesh
 	{
-		uint32 ID;
+		UE::Mutable::Private::PASSTHROUGH_ID ID;
 	};
 
 	struct FParamInfo
@@ -691,23 +691,11 @@ struct FMutableGraphGenerationContext
 
 		FParamInfo(const FString& InParamName, bool bInIsToggle) : ParamName(InParamName), bIsToggle(bInIsToggle) {}
 	};
-
-	// Cache of runtime pass-through meshes and their IDs used in the core to identify them.
-	// These meshes will remain as external references even in optimized models.
-	TMap<TSoftObjectPtr<UStreamableRenderAsset>, FGeneratedReferencedMesh> PassthroughMeshMap;
-
-	// Cache of runtime meshes and their IDs used in the core to identify them.
-	// These meshes will remain as external references even in optimized models.
-	TMap<FMutableSourceMeshData, FGeneratedReferencedMesh> RuntimeReferencedMeshMap;
-
+	
 	// Cache of runtime pass-through meshes and their IDs used in the core to identify them
 	// These meshes will become mutable meshes in the compiled model.
 	TMap<FMutableSourceMeshData, FGeneratedReferencedMesh> CompileTimeMeshMap;
 
-	// Cache of runtime pass-through images and their IDs used in the core to identify them.
-	// These textures will remain as external references even in optimized models.
-	TMap<TSoftObjectPtr<UTexture>, FGeneratedReferencedTexture> PassthroughTextureMap;
-
 	// Cache of runtime images and their IDs used in the core to identify them.
 	// These textures will remain as external references even in optimized models.
 	TMap<TSoftObjectPtr<UTexture2D>, FGeneratedReferencedTexture> RuntimeReferencedTextureMap;
@@ -716,6 +704,8 @@ struct FMutableGraphGenerationContext
 	// These textures will become mutable images in the compiled model.
 	TMap<TSoftObjectPtr<UTexture2D>, FGeneratedReferencedTexture> CompileTimeTextureMap;
 
+	TMap<UE::Mutable::Private::PASSTHROUGH_ID, TSoftObjectPtr<UObject>> PassthroughObjects;
+	
 	// Mutable meshes already build for source UStaticMesh or USkeletalMesh.
 	struct FGeneratedMeshData
 	{
@@ -800,6 +790,8 @@ struct FMutableGraphGenerationContext
 	/** Get the reference skeletal mesh associated to the current mesh component being generated */
 	FMutableComponentInfo* GetCurrentComponentInfo();
 
+	UE::Mutable::Private::PASSTHROUGH_ID AddPassthroughObject(UObject& Object);
+	
 	UObject* LoadObject(const FSoftObjectPtr& SoftObject);
 	
 	template<typename T>
@@ -819,7 +811,6 @@ struct FMutableGraphGenerationContext
 	
 	TArray<FMutableRefSkeletalMeshData> ReferenceSkeletalMeshesData;
 	
-	TArray<UMaterialInterface*> ReferencedMaterials;
 	TArray<FName> ReferencedMaterialSlotNames;
 	TMap<FGeneratedImagePropertiesKey, FGeneratedImageProperties> ImageProperties;
 	TArray<const UCustomizableObjectNode*> NoNameNodeObjectArray;
@@ -911,7 +902,7 @@ struct FMutableGraphGenerationContext
 	TObjectPtr<UObject> RootVersionBridge;
 
 	/** Index of the Referenced Material being generated */
-	int32 CurrentReferencedMaterialIndex = -1;
+	TStrongObjectPtr<UMaterialInterface> CurrentReferencedMaterial;
 	
 	// This should be automatized (create a struct that each time that a Macro is stacked, it adds 1 to the count.
 	// it also substracts 1 each time we remove a Macro from the stack)
diff --git a/Engine/Plugins/Mutable/Source/CustomizableObjectEditor/Private/MuCOE/GenerateMutableSource/GenerateMutableSourceComponent.cpp b/Engine/Plugins/Mutable/Source/CustomizableObjectEditor/Private/MuCOE/GenerateMutableSource/GenerateMutableSourceComponent.cpp
--- a/Engine/Plugins/Mutable/Source/CustomizableObjectEditor/Private/MuCOE/GenerateMutableSource/GenerateMutableSourceComponent.cpp
+++ b/Engine/Plugins/Mutable/Source/CustomizableObjectEditor/Private/MuCOE/GenerateMutableSource/GenerateMutableSourceComponent.cpp
@@ -160,7 +160,7 @@ UE::Mutable::Private::Ptr<UE::Mutable::Private::NodeComponent> GenerateMutableSo
 		else if(UMaterialInterface* OverlayMaterial = GenerationContext.LoadObject(TypedComponentMesh->GetOverlayMaterial()))
 		{
 			UE::Mutable::Private::Ptr<UE::Mutable::Private::NodeMaterialConstant> MaterialNode = new UE::Mutable::Private::NodeMaterialConstant();
-			MaterialNode->MaterialId = GenerationContext.ReferencedMaterials.AddUnique(OverlayMaterial);
+			MaterialNode->MaterialId = GenerationContext.AddPassthroughObject(*OverlayMaterial);
 			NodeComponentNew->OverlayMaterial = MaterialNode;
 		}
 
diff --git a/Engine/Plugins/Mutable/Source/CustomizableObjectEditor/Private/MuCOE/GenerateMutableSource/GenerateMutableSourceMaterial.cpp b/Engine/Plugins/Mutable/Source/CustomizableObjectEditor/Private/MuCOE/GenerateMutableSource/GenerateMutableSourceMaterial.cpp
--- a/Engine/Plugins/Mutable/Source/CustomizableObjectEditor/Private/MuCOE/GenerateMutableSource/GenerateMutableSourceMaterial.cpp
+++ b/Engine/Plugins/Mutable/Source/CustomizableObjectEditor/Private/MuCOE/GenerateMutableSource/GenerateMutableSourceMaterial.cpp
@@ -84,8 +84,8 @@ UE::Mutable::Private::Ptr<UE::Mutable::Private::NodeMaterial> GenerateMutableSou
 			{
 				MaterialNode = new UE::Mutable::Private::NodeMaterialConstant();
 
-				GenerationContext.CurrentReferencedMaterialIndex = GenerationContext.ReferencedMaterials.AddUnique(Material);
-				MaterialNode->MaterialId = GenerationContext.CurrentReferencedMaterialIndex;
+				GenerationContext.CurrentReferencedMaterial = TStrongObjectPtr(Material);
+				MaterialNode->MaterialId = GenerationContext.AddPassthroughObject(*Material);;
 
 				GenerationContext.MaterialConstantNodesCache.Add(TypedNodeMaterialConstant, MaterialNode);
 			}
@@ -367,7 +367,7 @@ UE::Mutable::Private::Ptr<UE::Mutable::Private::NodeMaterial> GenerateMutableSou
 	{
 		//This node will add a default value in case of error
 		UE::Mutable::Private::Ptr<UE::Mutable::Private::NodeMaterialConstant> ConstantValue = new UE::Mutable::Private::NodeMaterialConstant();
-		ConstantValue->MaterialId = GenerationContext.CurrentReferencedMaterialIndex;
+		ConstantValue->MaterialId = GenerationContext.AddPassthroughObject(*GenerationContext.CurrentReferencedMaterial.Get());
 
 		Result = ConstantValue;
 
@@ -397,8 +397,7 @@ UE::Mutable::Private::Ptr<UE::Mutable::Private::NodeMaterial> GenerateMutableSou
 			{
 				if (UMaterialInstance* TableMaterial = TypedNodeTable->GetColumnDefaultAssetByType<UMaterialInstance>(Pin))
 				{
-					check(GenerationContext.ReferencedMaterials.IsValidIndex(GenerationContext.CurrentReferencedMaterialIndex));
-					UMaterialInterface* BaseMaterial = GenerationContext.ReferencedMaterials[GenerationContext.CurrentReferencedMaterialIndex];
+					UMaterialInterface* BaseMaterial = GenerationContext.CurrentReferencedMaterial.Get();
 
 					const bool bTableMaterialCheckDisabled = GenerationContext.CompilationContext->Object->GetPrivate()->IsTableMaterialsParentCheckDisabled();
 					const bool bMaterialParentMismatch = !bTableMaterialCheckDisabled && BaseMaterial
diff --git a/Engine/Plugins/Mutable/Source/CustomizableObjectEditor/Private/MuCOE/GenerateMutableSource/GenerateMutableSourceSurface.cpp b/Engine/Plugins/Mutable/Source/CustomizableObjectEditor/Private/MuCOE/GenerateMutableSource/GenerateMutableSourceSurface.cpp
--- a/Engine/Plugins/Mutable/Source/CustomizableObjectEditor/Private/MuCOE/GenerateMutableSource/GenerateMutableSourceSurface.cpp
+++ b/Engine/Plugins/Mutable/Source/CustomizableObjectEditor/Private/MuCOE/GenerateMutableSource/GenerateMutableSourceSurface.cpp
@@ -139,7 +139,7 @@ UE::Mutable::Private::Ptr<UE::Mutable::Private::NodeSurface> GenerateMutableSour
 
 		SurfNode->SurfaceGuid = SurfaceGuid;
 
-		GenerationContext.CurrentReferencedMaterialIndex = GenerationContext.ReferencedMaterials.AddUnique(Material);
+		GenerationContext.CurrentReferencedMaterial = TStrongObjectPtr(Material);
 
 		// Find reference mesh used to generate the surface metadata for this fragment.
 		FMutableSourceMeshData MeshData;
@@ -240,7 +240,7 @@ UE::Mutable::Private::Ptr<UE::Mutable::Private::NodeSurface> GenerateMutableSour
 		if (!SurfNode->Material)
 		{
 			UE::Mutable::Private::Ptr<UE::Mutable::Private::NodeMaterialConstant> MaterialNode = new UE::Mutable::Private::NodeMaterialConstant();
-			MaterialNode->MaterialId = GenerationContext.CurrentReferencedMaterialIndex;
+			MaterialNode->MaterialId = GenerationContext.AddPassthroughObject(*GenerationContext.CurrentReferencedMaterial.Get());
 			SurfNode->Material = MaterialNode;
 		}
 
diff --git a/Engine/Plugins/Mutable/Source/CustomizableObjectEditor/Private/MuCOE/GenerateMutableSource/GenerateMutableSourceTable.cpp b/Engine/Plugins/Mutable/Source/CustomizableObjectEditor/Private/MuCOE/GenerateMutableSource/GenerateMutableSourceTable.cpp
--- a/Engine/Plugins/Mutable/Source/CustomizableObjectEditor/Private/MuCOE/GenerateMutableSource/GenerateMutableSourceTable.cpp
+++ b/Engine/Plugins/Mutable/Source/CustomizableObjectEditor/Private/MuCOE/GenerateMutableSource/GenerateMutableSourceTable.cpp
@@ -400,7 +400,7 @@ bool FillTableColumn(const UCustomizableObjectNodeTable* TableNode, UE::Mutable:
 					CurrentColumn = MutableTable->AddColumn(ColumnName, UE::Mutable::Private::ETableColumnType::Material);
 				}
 
-				int32 ReferenceMaterialId = GenerationContext.ReferencedMaterials.AddUnique(MaterialInstance);
+				int32 ReferenceMaterialId = GenerationContext.AddPassthroughObject(*MaterialInstance);
 				MutableTable->SetCell(CurrentColumn, RowId, ReferenceMaterialId);
 
 				return true;
@@ -1197,7 +1197,8 @@ bool GenerateNoneRow(const UCustomizableObjectNodeTable& TableNode, const UEdGra
 			ColumnId = MutableTable->AddColumn(ColumnName, UE::Mutable::Private::ETableColumnType::Material);
 		}
 
-		MutableTable->SetCell(ColumnId, RowId, GenerationContext.CurrentReferencedMaterialIndex);
+		int32 MaterialId = GenerationContext.AddPassthroughObject(*GenerationContext.CurrentReferencedMaterial.Get());
+		MutableTable->SetCell(ColumnId, RowId, MaterialId);
 	}
 	else if (Pin->PinType.PinCategory == Schema->PC_Color)
 	{
diff --git a/Engine/Plugins/Mutable/Source/CustomizableObjectEditor/Private/MuCOE/SMutableCodeViewer.cpp b/Engine/Plugins/Mutable/Source/CustomizableObjectEditor/Private/MuCOE/SMutableCodeViewer.cpp
--- a/Engine/Plugins/Mutable/Source/CustomizableObjectEditor/Private/MuCOE/SMutableCodeViewer.cpp
+++ b/Engine/Plugins/Mutable/Source/CustomizableObjectEditor/Private/MuCOE/SMutableCodeViewer.cpp
@@ -293,8 +293,13 @@ void SMutableCodeViewer::Construct(const FArguments& InArgs, UCustomizableObject
 
 	MutableModel = InMutableModel;
 	ExternalResourceProvider = InArgs._ExternalResourceProvider;
-	PreviewParameters = Instance->GetPrivate()->GetDescriptor().GetParameters();
 
+	PreviewParameters = Instance->GetPrivate()->GetDescriptor().GetParameters();
+	if (!PreviewParameters)
+	{
+		PreviewParameters = UE::Mutable::Private::FModel::NewParameters(MutableModel);
+	}
+	
 	RootNodes.Empty();
 	RootNodeAddresses.Empty();
 	ItemCache.Empty();
@@ -657,6 +662,11 @@ void SMutableCodeViewer::Construct(const FArguments& InArgs, UCustomizableObject
 FReply SMutableCodeViewer::UpdateParameters()
 {
 	PreviewParameters = Instance->GetPrivate()->GetDescriptor().GetParameters();
+	if (!PreviewParameters)
+	{
+		PreviewParameters = UE::Mutable::Private::FModel::NewParameters(MutableModel);
+	}
+
 	bIsPreviewPendingUpdate = true;
 
 	return FReply::Handled();
@@ -2824,7 +2834,6 @@ void SMutableCodeViewer::Tick(const FGeometry& AllottedGeometry, const double In
 		}
 	}
 	
-	
 	if (!bIsPreviewPendingUpdate)
 	{
 		return;
@@ -2837,9 +2846,9 @@ void SMutableCodeViewer::Tick(const FGeometry& AllottedGeometry, const double In
 
 	UE::Mutable::Private::FSettings Settings;
 	const TSharedPtr<UE::Mutable::Private::FSystem> System = MakeShared<UE::Mutable::Private::FSystem>(Settings);
-	
-	System->GetPrivate()->BeginBuild(MutableModel, Object->GetPrivate()->MeshIdRegistry, Object->GetPrivate()->ImageIdRegistry, Object->GetPrivate()->MaterialIdRegistry);
 
+	System->GetPrivate()->BeginBuild(MutableModel, Object->GetPrivate()->MeshIdRegistry, Object->GetPrivate()->ImageIdRegistry, Object->GetPrivate()->MaterialIdRegistry);
+	
 	switch (OperationDataType)
 	{
 	case UE::Mutable::Private::EDataType::Layout:
diff --git a/Engine/Plugins/Mutable/Source/CustomizableObjectEditor/Private/MuCOE/SMutableInstanceViewer.cpp b/Engine/Plugins/Mutable/Source/CustomizableObjectEditor/Private/MuCOE/SMutableInstanceViewer.cpp
--- a/Engine/Plugins/Mutable/Source/CustomizableObjectEditor/Private/MuCOE/SMutableInstanceViewer.cpp
+++ b/Engine/Plugins/Mutable/Source/CustomizableObjectEditor/Private/MuCOE/SMutableInstanceViewer.cpp
@@ -280,6 +280,21 @@ void SMutableInstanceViewer::Construct(const FArguments& InArgs, const TSharedPt
 }
 
 
+FString ToString(const UE::Mutable::Private::TPassthroughObjectPtr<UObject>& PassthroughObject)
+{
+	if (PassthroughObject.IsResolved())
+	{
+		FSoftObjectPath SoftObjectPath(PassthroughObject.Get());
+		return SoftObjectPath.ToString();
+	}
+	else
+	{
+		return FString::FromInt(PassthroughObject.GetId());
+	}
+}
+
+
+
 void SMutableInstanceViewer::SetInstance(const TSharedPtr<const UE::Mutable::Private::FInstance>& InInstance, const TSharedPtr<const UE::Mutable::Private::FModel>& Model, const TSharedPtr<const UE::Mutable::Private::FParameters>& Parameters, const UE::Mutable::Private::FSystem& System)
 {
 	if (InInstance != MutableInstance)
@@ -325,8 +340,7 @@ void SMutableInstanceViewer::SetInstance(const TSharedPtr<const UE::Mutable::Pri
 
 						if (MutableMesh && MutableMesh->IsReference())
 						{
-							NewLOD->MeshId = FText::Format(
-								LOCTEXT("MeshReference", "Id [{0}]  Ref [{1}] "), FText::FromString(MeshId.ToString()), MutableMesh->GetReferencedMesh());
+							NewLOD->MeshId = FText::Format(LOCTEXT("MeshReference", "Id [{0}]  PassthroughObject [{1}] "), FText::FromString(MeshId.ToString()), FText::FromString(::ToString(MutableMesh->PassthroughObject)));
 						}
 					}
 
diff --git a/Engine/Plugins/Mutable/Source/MutableRuntime/Internal/MuR/Image.h b/Engine/Plugins/Mutable/Source/MutableRuntime/Internal/MuR/Image.h
--- a/Engine/Plugins/Mutable/Source/MutableRuntime/Internal/MuR/Image.h
+++ b/Engine/Plugins/Mutable/Source/MutableRuntime/Internal/MuR/Image.h
@@ -2,6 +2,7 @@
 
 #pragma once
 
+#include "MuR/PassthroughObject.h"
 #include "MuR/ImageTypes.h"
 
 #include "MuR/MutableMath.h"
@@ -45,7 +46,7 @@ namespace UE::Mutable::Private
 		UE_API FImage(uint32 SizeX, uint32 SizeY, uint32 Lods, EImageFormat Format, EInitializationType InitType);
 
 		/** Create a new empty image that repreents an external resource image. */
-		static UE_API TSharedPtr<FImage> CreateAsReference(uint32 ID, const FImageDesc& Desc, bool bForceLoad);
+		static UE_API TSharedPtr<FImage> CreateAsReference(const TPassthroughObjectPtr<UTexture>& ExternalObject, const FImageDesc& Desc, bool bForceLoad);
 
 		//! Serialisation
 		static UE_API void Serialise( const FImage*, FOutputArchive& );
@@ -69,10 +70,7 @@ namespace UE::Mutable::Private
 
 		/** If true, this is a reference that must be resolved at compile time. */
 		UE_API bool IsForceLoad() const;
-
-		/** Return the id of the engine referenced texture. Only valid if IsReference. */
-		UE_API uint32 GetReferencedTexture() const;
-
+    	
 		//! Return the width of the image.
         UE_API uint16 GetSizeX() const;
 
@@ -113,10 +111,7 @@ namespace UE::Mutable::Private
 
 			// If this is set, the image has an updated relevancy map. This flag is not persisent.
 			IF_HAS_RELEVANCY_MAP = 1 << 3,
-
-			/** If this is set, this is a reference to an external image, and the ReferenceID is valid. */
-			IF_IS_REFERENCE = 1 << 4,
-
+			
 			/** For reference images, this indicates that they should be loaded into full images as soon as they are generated. */
 			IF_IS_FORCELOAD = 1 << 5
 		} EImageFlags;
@@ -128,10 +123,9 @@ namespace UE::Mutable::Private
 		uint16 RelevancyMinY = 0;
 		uint16 RelevancyMaxY = 0;
 
-		TStrongObjectPtr<UTexture> Texture;
+		/** Passthrough Id or TStrongObjectPtr. */
+    	TPassthroughObjectPtr<UTexture> PassthroughObject;
 
-		/** Only valid if the right flags are set, this identifies a referenced image. */
-		uint32 ReferenceID = 0;
     	
 		/** Pixel data for all lods. */
 		FImageDataStorage DataStorage;
@@ -151,7 +145,7 @@ namespace UE::Mutable::Private
 			Result->RelevancyMinY = RelevancyMinY;
 			Result->RelevancyMaxY = RelevancyMaxY;
 			Result->DataStorage = DataStorage;
-			Result->ReferenceID = ReferenceID;
+			Result->PassthroughObject = PassthroughObject;
 
 			return Result;
 		}
@@ -171,7 +165,7 @@ namespace UE::Mutable::Private
 			RelevancyMinY = Other->RelevancyMinY;
 			RelevancyMaxY = Other->RelevancyMaxY;
 			DataStorage = Other->DataStorage;
-			ReferenceID = Other->ReferenceID;
+			PassthroughObject = Other->PassthroughObject;
 		}
 
 		void CopyMove(FImage* Other)
@@ -187,7 +181,7 @@ namespace UE::Mutable::Private
 			Flags = Other->Flags;
 			RelevancyMinY = Other->RelevancyMinY;
 			RelevancyMaxY = Other->RelevancyMaxY;
-			ReferenceID = Other->ReferenceID;
+			PassthroughObject = Other->PassthroughObject;
 
 			DataStorage = MoveTemp(Other->DataStorage);
 		}
@@ -203,8 +197,8 @@ namespace UE::Mutable::Private
 		inline bool operator==(const FImage& Other) const
 		{
 			return 
-				(DataStorage == Other.DataStorage) &&
-				(ReferenceID == Other.ReferenceID);
+				DataStorage == Other.DataStorage &&
+				PassthroughObject == Other.PassthroughObject;
 		}
 
 		//-----------------------------------------------------------------------------------------
diff --git a/Engine/Plugins/Mutable/Source/MutableRuntime/Internal/MuR/Material.h b/Engine/Plugins/Mutable/Source/MutableRuntime/Internal/MuR/Material.h
--- a/Engine/Plugins/Mutable/Source/MutableRuntime/Internal/MuR/Material.h
+++ b/Engine/Plugins/Mutable/Source/MutableRuntime/Internal/MuR/Material.h
@@ -10,6 +10,7 @@
 
 #include "UObject/StrongObjectPtr.h"
 #include "Materials/MaterialInterface.h"
+#include "MuR/PassthroughObject.h"
 
 #define UE_API MUTABLERUNTIME_API
 
@@ -40,17 +41,11 @@ namespace UE::Mutable::Private
 		// Own interface
 		//-----------------------------------------------------------------------------------------
 		UE_API bool operator==(const FMaterial& Other) const;
-
-
 		UE_API void Serialise(FOutputArchive&) const;
 		UE_API void Unserialise(FInputArchive&);
 
-
-	public:
-
-		TStrongObjectPtr<UMaterialInterface> Material;
-
-		int32 ReferenceID = INDEX_NONE;
+		/** Passthrough Id or TStrongObjectPtr. */
+		TPassthroughObjectPtr<UMaterialInterface> PassthroughObject;
 
 		TMap<FName, TVariant<OP::ADDRESS, TSharedPtr<const FImage>>> ImageParameters;
 		TMap<FName, FVector4f> ColorParameters;
diff --git a/Engine/Plugins/Mutable/Source/MutableRuntime/Internal/MuR/Mesh.h b/Engine/Plugins/Mutable/Source/MutableRuntime/Internal/MuR/Mesh.h
--- a/Engine/Plugins/Mutable/Source/MutableRuntime/Internal/MuR/Mesh.h
+++ b/Engine/Plugins/Mutable/Source/MutableRuntime/Internal/MuR/Mesh.h
@@ -16,6 +16,7 @@
 #include "MuR/RefCounted.h"
 #include "MuR/Serialisation.h"
 #include "MuR/Skeleton.h"
+#include "MuR/PassthroughObject.h"
 #include "Templates/Tuple.h"
 
 #include "UObject/StrongObjectPtr.h"
@@ -214,9 +215,6 @@ namespace UE::Mutable::Private
 		/** The mesh is formatted to be used for wrapping projection */
 		ProjectWrappingFormat = 1 << 1,
 
-		/** The mesh is a reference to an external resource mesh. */
-		IsResourceReference = 1 << 2,
-
 		/** The mesh is a reference to an external resource mesh and must be loaded when first referenced. */
 		IsResourceForceLoad = 1 << 3,
 
@@ -241,7 +239,7 @@ namespace UE::Mutable::Private
         //-----------------------------------------------------------------------------------------
 				
 		/** Create a new empty mesh that repreents an external resource mesh. */
-		static UE_API TSharedPtr<FMesh> CreateAsReference(uint32 ID, bool bForceLoad);
+		static UE_API TSharedPtr<FMesh> CreateAsReference(const TPassthroughObjectPtr<USkeletalMesh>& PasstroughObjectPtr, bool bForceLoad);
 
         /** Deep clone this mesh. */
         UE_API TSharedPtr<FMesh> Clone() const;
@@ -268,10 +266,7 @@ namespace UE::Mutable::Private
 
 		/** If true, this is a reference that must be resolved at compile time. */
 		UE_API bool IsForceLoad() const;
-
-		/** Return the id of the engine referenced mesh. Only valid if IsReference. */
-		UE_API uint32 GetReferencedMesh() const;
-
+    	
 		/** Mesh references can actually reference a morphed mesh. In that case you can set and get the morph with these functions. */
 		UE_API void SetReferencedMorph(const FString& MorphName);
 		UE_API const FString& GetReferencedMorph() const;
@@ -445,8 +440,8 @@ namespace UE::Mutable::Private
 		/** Skeletal Meshes involved in the generation of the mesh. Meshes from Mesh Parameters.*/
 		TArray<TStrongObjectPtr<USkeletalMesh>> SkeletalMeshes;
 
-		/** Only valid if the right flags are set, this identifies a referenced mesh. */
-		uint32 ReferenceID = 0;
+		/** Passthrough Id or TStrongObjectPtr. */
+		TPassthroughObjectPtr<USkeletalMesh> PassthroughObject;
 
 		/** If the mesh is a reference the referenced morph name is stored here. Otherwise it is an empty string. */
 		FString ReferencedMorph;
@@ -531,7 +526,7 @@ namespace UE::Mutable::Private
 		{
 			bool bEqual = true;
 
-			if (bEqual) bEqual = (ReferenceID == Other.ReferenceID);
+			if (bEqual) bEqual = (PassthroughObject == Other.PassthroughObject);
 			if (bEqual) bEqual = (MeshIDPrefix == Other.MeshIDPrefix);
 			if (bEqual) bEqual = (IndexBuffers == Other.IndexBuffers);
 			if (bEqual) bEqual = (VertexBuffers == Other.VertexBuffers);
diff --git a/Engine/Plugins/Mutable/Source/MutableRuntime/Internal/MuR/Operations.h b/Engine/Plugins/Mutable/Source/MutableRuntime/Internal/MuR/Operations.h
--- a/Engine/Plugins/Mutable/Source/MutableRuntime/Internal/MuR/Operations.h
+++ b/Engine/Plugins/Mutable/Source/MutableRuntime/Internal/MuR/Operations.h
@@ -2,6 +2,7 @@
 
 #pragma once
 
+#include "PassthroughObject.h"
 #include "MuR/ImageTypes.h"
 #include "MuR/Types.h"
 #include "MuR/SerialisationPrivate.h"
@@ -447,6 +448,13 @@ namespace UE::Mutable::Private
             int32 Skeleton;
         };
 
+    	struct MaterialConstantArgs
+    	{
+    		ADDRESS Value;
+    		PASSTHROUGH_ID ID;
+    	};
+
+    	
 		struct ParameterArgs
 		{
 			ADDRESS variable;
@@ -478,7 +486,7 @@ namespace UE::Mutable::Private
 		struct ResourceReferenceArgs
 		{
 			FImageDesc ImageDesc;
-			int32 ID;
+			PASSTHROUGH_ID ID;
 			int8 ForceLoad;
 		};
 
diff --git a/Engine/Plugins/Mutable/Source/MutableRuntime/Internal/MuR/ParametersPrivate.h b/Engine/Plugins/Mutable/Source/MutableRuntime/Internal/MuR/ParametersPrivate.h
--- a/Engine/Plugins/Mutable/Source/MutableRuntime/Internal/MuR/ParametersPrivate.h
+++ b/Engine/Plugins/Mutable/Source/MutableRuntime/Internal/MuR/ParametersPrivate.h
@@ -6,6 +6,7 @@
 #include "MuR/TVariant.h"
 #include "Math/MathFwd.h"
 
+#include "MuR/TVariant.h"
 #include "MuR/Parameters.h"
 #include "MuR/SerialisationPrivate.h"
 #include "MuR/Operations.h"
diff --git a/Engine/Plugins/Mutable/Source/MutableRuntime/Internal/MuR/PassthroughObject.h b/Engine/Plugins/Mutable/Source/MutableRuntime/Internal/MuR/PassthroughObject.h
new file mode 100644
--- /dev/null
+++ b/Engine/Plugins/Mutable/Source/MutableRuntime/Internal/MuR/PassthroughObject.h
@@ -0,0 +1,125 @@
+// Copyright Epic Games, Inc. All Rights Reserved.
+
+#pragma once
+
+#include "MuR/TVariant.h"
+#include "MuR/Serialisation.h"
+
+
+namespace UE::Mutable::Private
+{
+	class FPassthroughObjectLoader;
+	
+	using PASSTHROUGH_ID = uint32;
+
+	constexpr PASSTHROUGH_ID PASSTHROUGH_ID_INVALID = TNumericLimits<PASSTHROUGH_ID>::Max();
+	
+	MUTABLERUNTIME_API PASSTHROUGH_ID GetPassthroughId(const TSoftObjectPtr<UObject>& SoftObjectPtr);
+
+	template<typename T>
+	class TPassthroughObjectPtr
+	{
+		friend FPassthroughObjectLoader;
+		template<class U> friend class TPassthroughObjectPtr;
+
+	public:
+		TPassthroughObjectPtr() : Resource(MakeShared<TVariant<TStrongObjectPtr<UObject>, PASSTHROUGH_ID>>())
+		{
+			Resource->Set<PASSTHROUGH_ID>(PASSTHROUGH_ID_INVALID);
+		}
+
+		explicit TPassthroughObjectPtr(PASSTHROUGH_ID Id) : Resource(MakeShared<TVariant<TStrongObjectPtr<UObject>, PASSTHROUGH_ID>>())
+		{
+			Resource->Set<PASSTHROUGH_ID>(Id);
+		}
+
+		explicit TPassthroughObjectPtr(T* Object) : Resource(MakeShared<TVariant<TStrongObjectPtr<UObject>, PASSTHROUGH_ID>>())
+		{
+			Resource->Set<TStrongObjectPtr<UObject>>(TStrongObjectPtr<UObject>(Object));
+		}
+
+		TPassthroughObjectPtr(const TPassthroughObjectPtr& Other) : Resource(Other.Resource)
+		{
+		}
+
+		template<typename U>
+		TPassthroughObjectPtr(const TPassthroughObjectPtr<U>& Other) : Resource(Other.Resource)
+		{
+		}
+
+		TPassthroughObjectPtr& operator=(const TPassthroughObjectPtr& Other)
+		{
+			Resource = Other.Resource;
+			return *this;
+		}
+
+		bool operator==(const TPassthroughObjectPtr& Other) const
+		{
+			return Resource.Get() == Other.Resource.Get();
+		}
+
+		bool IsResolved() const
+		{
+			return Resource->IsType<TStrongObjectPtr<UObject>>();
+		}
+		
+		T* Get() const
+		{
+			if (UObject* Object = Resource->Get<TStrongObjectPtr<UObject>>().Get())
+			{
+				return CastChecked<T>(Object);
+			}
+
+			return nullptr;
+		}
+
+		PASSTHROUGH_ID GetId() const
+		{
+			return Resource->Get<PASSTHROUGH_ID>(); 
+		}
+
+		bool IsValid() const
+		{
+			return !Resource->IsType<PASSTHROUGH_ID>() || Resource->Get<PASSTHROUGH_ID>() != PASSTHROUGH_ID_INVALID;
+		}
+		
+		explicit operator bool() const
+		{
+			return IsValid();
+		}
+	
+	private:
+		TSharedRef<TVariant<TStrongObjectPtr<UObject>, PASSTHROUGH_ID>> Resource;
+	};	
+	
+	
+	class FPassthroughObjectLoader
+	{
+	public:
+		template<typename T>
+		TPassthroughObjectPtr<T> Find(PASSTHROUGH_ID Id)
+		{
+			if (TPassthroughObjectPtr<UObject>* Result = PassthroughObjects.Find(Id))
+			{
+				return TPassthroughObjectPtr<T>(*Result);
+			}
+			else
+			{
+				return {};
+			}
+		}
+		
+		template<typename T>
+		void Add(PASSTHROUGH_ID Id)
+		{
+			PassthroughObjects.Add(Id, TPassthroughObjectPtr<T>(Id));
+		}
+
+		MUTABLERUNTIME_API void Load(const TMap<PASSTHROUGH_ID, TSoftObjectPtr<UObject>>& ResolutionMap, TArray<FSoftObjectPath>& AssetsToStream) const;
+
+		MUTABLERUNTIME_API void Resolve(TMap<PASSTHROUGH_ID, TSoftObjectPtr<UObject>>& ResolutionMap);
+
+	private:
+		TMap<PASSTHROUGH_ID, TPassthroughObjectPtr<UObject>> PassthroughObjects;
+	};
+}
diff --git a/Engine/Plugins/Mutable/Source/MutableRuntime/Internal/MuR/System.h b/Engine/Plugins/Mutable/Source/MutableRuntime/Internal/MuR/System.h
--- a/Engine/Plugins/Mutable/Source/MutableRuntime/Internal/MuR/System.h
+++ b/Engine/Plugins/Mutable/Source/MutableRuntime/Internal/MuR/System.h
@@ -131,9 +131,13 @@ namespace UE::Mutable::Private
         //! A call to NewInstance must be paired to a call to ReleasesInstance when the instance is
         //! no longer needed.
         //! \return An identifier that is always bigger than 0.
-        UE_API FInstance::FID NewInstance(const TSharedPtr<const FModel>& Model, const TSharedPtr<FExternalResourceProvider>& ExternalResourceProvider);
+        UE_API FInstance::FID NewInstance(const TSharedPtr<const FModel>& Model,
+        	const TSharedPtr<FExternalResourceProvider>& ExternalResourceProvider,
+        	const TSharedPtr<FPassthroughObjectLoader>& PassthroughObjectLoader);
 
-		UE_API void ReuseInstance(FInstance::FID InstanceId, const TSharedPtr<FExternalResourceProvider>& ExternalResourceProvider);
+		UE_API void ReuseInstance(FInstance::FID InstanceId,
+			const TSharedPtr<FExternalResourceProvider>& ExternalResourceProvider,
+			const TSharedPtr<FPassthroughObjectLoader>& PassthroughObjectLoader);
 		
         //! \brief Update an instance with a new parameter set and/or state.
         //!
diff --git a/Engine/Plugins/Mutable/Source/MutableRuntime/Internal/MuR/SystemPrivate.h b/Engine/Plugins/Mutable/Source/MutableRuntime/Internal/MuR/SystemPrivate.h
--- a/Engine/Plugins/Mutable/Source/MutableRuntime/Internal/MuR/SystemPrivate.h
+++ b/Engine/Plugins/Mutable/Source/MutableRuntime/Internal/MuR/SystemPrivate.h
@@ -1234,6 +1234,7 @@ namespace UE::Mutable::Private
 		TSharedPtr<FProgramCache> Cache;
 
 		TSharedPtr<FExternalResourceProvider> ExternalResourceProvider;
+		TSharedPtr<FPassthroughObjectLoader> PassthroughObjectLoader;
 		
 		~FLiveInstance()
 		{
@@ -1302,6 +1303,9 @@ namespace UE::Mutable::Private
 		TSharedPtr<FImageIdRegistry> CurrentImageIdRegistry;
 
     	TSharedPtr<FMaterialIdRegistry> CurrentMaterialIdRegistry;
+
+    	/** Only valid in initial generation. We can not load UObjects while streaming. */
+    	TSharedPtr<FPassthroughObjectLoader> CurrentPassthroughObjectLoader;
     	
 		/** Resources that have been used in the past, but haven't been deallocated because they still fitted the memory budget and they could be reused. */
 		TMemoryTrackedArray<TSharedPtr<FImage>> PooledImages;
diff --git a/Engine/Plugins/Mutable/Source/MutableRuntime/Private/MuR/CodeRunner.cpp b/Engine/Plugins/Mutable/Source/MutableRuntime/Private/MuR/CodeRunner.cpp
--- a/Engine/Plugins/Mutable/Source/MutableRuntime/Private/MuR/CodeRunner.cpp
+++ b/Engine/Plugins/Mutable/Source/MutableRuntime/Private/MuR/CodeRunner.cpp
@@ -232,7 +232,7 @@ namespace UE::Mutable::Private
 		
 		if (ExternalResourceProvider)
 		{
-			if (Id.ReferenceResourceId < 0)
+			if (Id.ReferenceResourceId == PASSTHROUGH_ID_INVALID)
 			{
 				// It's a parameter image
 				return ExternalResourceProvider->GetImageAsync(Id.ImageParameter, MipmapsToSkip, ResultCallback);
@@ -259,7 +259,7 @@ namespace UE::Mutable::Private
 
 		if (ExternalResourceProvider)
 		{
-			if (Id.ReferenceResourceId < 0)
+			if (Id.ReferenceResourceId == PASSTHROUGH_ID_INVALID)
 			{
 				// It's a parameter mesh
 				return ExternalResourceProvider->GetMeshAsync(Id.MeshParameter, LODIndex, SectionIndex, ResultCallback);
@@ -1186,11 +1186,15 @@ namespace UE::Mutable::Private
 
 		case EOpType::MI_CONSTANT:
 		{
-			OP::ResourceConstantArgs Args = Program.GetOpArgs<OP::ResourceConstantArgs>(Item.At);
-			TSharedPtr<const FMaterial> SourceConst = Program.ConstantMaterials[Args.value];
-			check(SourceConst);
+			OP::MaterialConstantArgs Args = Program.GetOpArgs<OP::MaterialConstantArgs>(Item.At);
 
-			StoreMaterial(Item, SourceConst);
+        	TSharedPtr<const FMaterial> SourceConst = Program.ConstantMaterials[Args.Value];
+       		check(SourceConst);
+        		
+			TSharedPtr<FMaterial> Source = UE::Mutable::Private::CloneOrTakeOver<FMaterial>(SourceConst);
+        	Source->PassthroughObject = FindPassthrough<UMaterialInterface>(Args.ID);
+        		
+			StoreMaterial(Item, Source);
 			break;
 		}
 
@@ -1235,7 +1239,8 @@ namespace UE::Mutable::Private
 				}
 				else
 				{
-					Result = FMesh::CreateAsReference(Args.ID, false);
+					TPassthroughObjectPtr<USkeletalMesh> ExternalObject = FindPassthrough<USkeletalMesh>(Args.ID);
+					Result = FMesh::CreateAsReference(ExternalObject, false);
 				}
 				StoreMesh(Item, Result);
 				break;
@@ -4660,7 +4665,8 @@ namespace UE::Mutable::Private
 				}
 				else
 				{
-					Result = FImage::CreateAsReference(Args.ID, Args.ImageDesc, false);
+					TPassthroughObjectPtr<UTexture> ExternalObject = FindPassthrough<UTexture>(Args.ID);
+					Result = FImage::CreateAsReference(ExternalObject, Args.ImageDesc, false);
 				}
 				StoreImage(Item, Result);
 				break;
@@ -5675,7 +5681,7 @@ namespace UE::Mutable::Private
 
             case 1:
             {
-                bool a = Args.A ? LoadBool( FCacheAddress(Args.A,Item) ) : true;
+                bool a = LoadBool(FCacheAddress(Args.A, Item));
                 if (!a)
                 {
                     StoreBool( Item, false );
@@ -5691,7 +5697,7 @@ namespace UE::Mutable::Private
             case 2:
             {
                 // We arrived here because a is true
-                bool b = Args.B ? LoadBool( FCacheAddress(Args.B,Item) ) : true;
+                bool b = LoadBool(FCacheAddress(Args.B, Item));
                 StoreBool( Item, b );
                 break;
             }
@@ -6435,10 +6441,12 @@ namespace UE::Mutable::Private
 			OP::MaterialParameterArgs Args = Program.GetOpArgs<OP::MaterialParameterArgs >(Item.At);
 			TSharedPtr<FRangeIndex> Index = BuildCurrentOpRangeIndex(Item, Args.Variable);
 
-			TSharedPtr<FMaterial> Material = MakeShared<FMaterial>();
-			Material->Material = TStrongObjectPtr(InParams->GetMaterialValue(Args.Variable, Index.Get()));
+			UMaterialInterface* MaterialInterface = InParams->GetMaterialValue(Args.Variable, Index.Get());
 
-			if (Material->Material)
+			TSharedPtr<FMaterial> Material = MakeShared<FMaterial>();
+			Material->PassthroughObject = TPassthroughObjectPtr(MaterialInterface);
+ 
+			if (MaterialInterface)
 			{
 				// Images
 				const TArray<uint32>& ImageParameterNames = Program.ConstantUInt32Lists[Args.ImageParameterNames];
@@ -6449,7 +6457,7 @@ namespace UE::Mutable::Private
 					const FString& ParameterName = Program.ConstantStrings[ImageParameterNames[ParameterIndex]];
 
 					UTexture* Texture;
-					bool bParameterFound = Material->Material->GetTextureParameterValue(FName(ParameterName), Texture);
+					bool bParameterFound = MaterialInterface->GetTextureParameterValue(FName(ParameterName), Texture);
 
 					if (bParameterFound)
 					{
@@ -6468,7 +6476,7 @@ namespace UE::Mutable::Private
 					const FString& ParameterName = Program.ConstantStrings[ScalarParameterNames[ParameterIndex]];
 
 					float ScalarValue;
-					bool bParameterFound = Material->Material->GetScalarParameterValue(FName(ParameterName), ScalarValue);
+					bool bParameterFound = MaterialInterface->GetScalarParameterValue(FName(ParameterName), ScalarValue);
 					
 					if (bParameterFound)
 					{
@@ -6484,7 +6492,7 @@ namespace UE::Mutable::Private
 					const FString& ParameterName = Program.ConstantStrings[ColorParameterNames[ParameterIndex]];
 
 					FLinearColor VectorValue;
-					bool bParameterFound = Material->Material->GetVectorParameterValue(FName(ParameterName), VectorValue);
+					bool bParameterFound = MaterialInterface->GetVectorParameterValue(FName(ParameterName), VectorValue);
 					
 					if (bParameterFound)
 					{
@@ -8100,8 +8108,7 @@ namespace UE::Mutable::Private
 			TSharedPtr<FRangeIndex> Index = BuildCurrentOpRangeIndex(Item, Args.MaterialParameter);
 
 			// Get material parameter from the array of parameters
-			TSharedPtr<FMaterial> Material = MakeShared<FMaterial>();
-			Material->Material = TStrongObjectPtr(pParams->GetMaterialValue(Args.MaterialParameter, Index.Get()));
+			UMaterialInterface* MaterialInterface = pParams->GetMaterialValue(Args.MaterialParameter, Index.Get());
 
 			// Get the texture parameter name
 			check(Args.ParameterName < (uint32)Program.ConstantStrings.Num());
@@ -8110,9 +8117,9 @@ namespace UE::Mutable::Private
 			UTexture* Image = nullptr;
 
 			// Get the parameter texture from the UMaterial
-			if (Material->Material)
+			if (MaterialInterface)
 			{
-				Material->Material->GetTextureParameterValue(FName(ParameterName), Image);
+				MaterialInterface->GetTextureParameterValue(FName(ParameterName), Image);
 			}
 
 			FExtendedImageDesc& Result = ImageDescResults.FindOrAdd(Item.At);
diff --git a/Engine/Plugins/Mutable/Source/MutableRuntime/Private/MuR/CodeRunner.h b/Engine/Plugins/Mutable/Source/MutableRuntime/Private/MuR/CodeRunner.h
--- a/Engine/Plugins/Mutable/Source/MutableRuntime/Private/MuR/CodeRunner.h
+++ b/Engine/Plugins/Mutable/Source/MutableRuntime/Private/MuR/CodeRunner.h
@@ -147,7 +147,7 @@ namespace  UE::Mutable::Private
 		struct FExternalResourceId
 		{
 			/** If it is an image or mesh reference. */
-			int32 ReferenceResourceId = -1;
+			PASSTHROUGH_ID ReferenceResourceId = PASSTHROUGH_ID_INVALID;
 
 			/** If it is an image or mesh parameter.*/
 			UTexture* ImageParameter = nullptr; 
@@ -636,6 +636,22 @@ namespace  UE::Mutable::Private
 		{
 			System->WorkingMemoryManager.Release(Resource);
 		}
+
+    	template<typename T>
+		TPassthroughObjectPtr<T> FindPassthrough(PASSTHROUGH_ID Id)
+		{
+			if (System->WorkingMemoryManager.CurrentPassthroughObjectLoader)
+			{
+				TPassthroughObjectPtr<T> PassthroughObject = System->WorkingMemoryManager.CurrentPassthroughObjectLoader->Find<T>(Id);
+				check(PassthroughObject); // Object not discovered in the CodeRunnerBegin.
+				
+				return PassthroughObject;
+			}
+			else
+			{
+				return TPassthroughObjectPtr<T>(Id);
+			}
+		}
     	
 		/** This flag is turned on when a streaming error or similar happens. Results are not usable.
 		* This should only happen in-editor.
@@ -730,7 +746,6 @@ namespace  UE::Mutable::Private
 
 		/** */
 		void LaunchIssuedTask(const TSharedPtr<FIssuedTask>& TaskToIssue, bool& bOutFailed);
-
     };
 
 
diff --git a/Engine/Plugins/Mutable/Source/MutableRuntime/Private/MuR/CodeRunnerBegin.cpp b/Engine/Plugins/Mutable/Source/MutableRuntime/Private/MuR/CodeRunnerBegin.cpp
--- a/Engine/Plugins/Mutable/Source/MutableRuntime/Private/MuR/CodeRunnerBegin.cpp
+++ b/Engine/Plugins/Mutable/Source/MutableRuntime/Private/MuR/CodeRunnerBegin.cpp
@@ -88,9 +88,9 @@ namespace UE::Mutable::Private
 	}
 	
 
-	FProgramCache& CodeRunnerBegin::GetMemory() const
+	FWorkingMemoryManager& CodeRunnerBegin::GetMemory() const
 	{
-		return *System.GetPrivate()->WorkingMemoryManager.CurrentInstanceCache;
+		return System.GetPrivate()->WorkingMemoryManager;
 	}
 
 
@@ -730,7 +730,7 @@ namespace UE::Mutable::Private
 							TSharedPtr<FInstance> Instance;
 							if (InInstance)
 							{
-								Instance = UE::Mutable::Private::CloneOrTakeOver<FInstance>(InInstance);
+								Instance = UE::Mutable::Private::CloneOrTakeOver<FInstance>(InInstance); // TODO GMT It will always clone due to a second reference being in the ResultsInstance
 							}
 							else
 							{
@@ -1271,7 +1271,7 @@ namespace UE::Mutable::Private
 					Op.At = Item.At;
 					Op.ExecutionIndex = Item.ExecutionIndex;
         		
-					TSharedPtr<FRangeIndex> Index = BuildCurrentOpRangeIndex(Op, *Params, *Model.Get(), GetMemory(), Args.variable);
+					TSharedPtr<FRangeIndex> Index = BuildCurrentOpRangeIndex(Op, *Params, *Model.Get(), *GetMemory().CurrentInstanceCache, Args.variable);
 					const bool Value = Params->GetBoolValue(Args.variable, Index.Get());
 					
 					StoreBool(Item, Value);
@@ -1423,7 +1423,7 @@ namespace UE::Mutable::Private
 					Op.At = Item.At;
 					Op.ExecutionIndex = Item.ExecutionIndex;
 					
-					TSharedPtr<FRangeIndex> Index = BuildCurrentOpRangeIndex(Op, *Params, *Model.Get(), GetMemory(), Args.variable);
+					TSharedPtr<FRangeIndex> Index = BuildCurrentOpRangeIndex(Op, *Params, *Model.Get(), *GetMemory().CurrentInstanceCache, Args.variable);
 					const float Result = Params->GetFloatValue( Args.variable, Index.Get());
 
 					StoreFloat(Item, Result);
@@ -1550,7 +1550,7 @@ namespace UE::Mutable::Private
 					Op.At = Item.At;
 					Op.ExecutionIndex = Item.ExecutionIndex;
 				
-					TSharedPtr<FRangeIndex> Index = BuildCurrentOpRangeIndex(Op, *Params, *Model.Get(), GetMemory(), Args.variable);
+					TSharedPtr<FRangeIndex> Index = BuildCurrentOpRangeIndex(Op, *Params, *Model.Get(), *GetMemory().CurrentInstanceCache, Args.variable);
 
 					FString Result;
 					Params->GetStringValue(Args.variable, Result, Index.Get());
@@ -1621,7 +1621,7 @@ namespace UE::Mutable::Private
 			case EOpType::ME_CONSTANT:
 				{
 					OP::MeshConstantArgs Args = Program.GetOpArgs<OP::MeshConstantArgs>(Item.At);
-					ME_SKELETON_ID.Push(Args.Skeleton);
+					GetMemory().CurrentPassthroughObjectLoader->Add<USkeletalMesh>(Args.Skeleton);
 
 					StoreNone(Item);
 					break;
@@ -1630,7 +1630,7 @@ namespace UE::Mutable::Private
 			case EOpType::ME_REFERENCE:
 				{
 					OP::ResourceReferenceArgs Args = Program.GetOpArgs<OP::ResourceReferenceArgs>(Item.At);
-					ME_REFERENCE_ID.Push(Args.ID);
+					GetMemory().CurrentPassthroughObjectLoader->Add<USkeletalMesh>(Args.ID);
 
 					StoreNone(Item);
 					break;
@@ -1945,7 +1945,10 @@ namespace UE::Mutable::Private
 					{
 						if (Program.ConstantUInt32Lists.IsValidIndex(OpArgs.SkeletonIds.ListAddress))
 						{
-							ME_SKELETON_ID.Append(Program.ConstantUInt32Lists[OpArgs.SkeletonIds.ListAddress]);
+							for (const uint32 SkeletonId : Program.ConstantUInt32Lists[OpArgs.SkeletonIds.ListAddress])
+							{
+								GetMemory().CurrentPassthroughObjectLoader->Add<USkeleton>(SkeletonId);
+							}
 						}
 						else
 						{
@@ -1954,7 +1957,7 @@ namespace UE::Mutable::Private
 					}
 					else
 					{
-						ME_SKELETON_ID.Push(OpArgs.SkeletonIds.SkeletonId);
+						GetMemory().CurrentPassthroughObjectLoader->Add<USkeleton>(OpArgs.SkeletonIds.SkeletonId);
 					}
 
 					StoreNone(Item);
@@ -2002,7 +2005,7 @@ namespace UE::Mutable::Private
 			case EOpType::IM_REFERENCE:
 				{
 					OP::ResourceReferenceArgs Args = Program.GetOpArgs<OP::ResourceReferenceArgs>(Item.At);
-					IM_REFERENCE_ID.Push(Args.ID);
+					GetMemory().CurrentPassthroughObjectLoader->Add<UTexture>(Args.ID);
 
 					StoreNone(Item);
 					break;
@@ -2062,11 +2065,11 @@ namespace UE::Mutable::Private
 
 							for (int32 IterationIndex = 0; IterationIndex < NumIterations; ++IterationIndex)
 							{
-								ExecutionIndex Index = GetMemory().GetRangeIndex(Item.ExecutionIndex);
+								ExecutionIndex Index = GetMemory().CurrentInstanceCache->GetRangeIndex(Item.ExecutionIndex);
 								Index.SetFromModelRangeIndex(Args.rangeId, IterationIndex);
 
 								FScheduledOpInline ItemCopy = Item;
-								ItemCopy.ExecutionIndex = GetMemory().GetRangeIndexIndex(Index);
+								ItemCopy.ExecutionIndex = GetMemory().CurrentInstanceCache->GetRangeIndexIndex(Index);
 
 								PushOp(FScheduledOpInline(Args.mask, Item));
 								PushOp(FScheduledOpInline(Args.blended, Item));
@@ -2347,9 +2350,8 @@ namespace UE::Mutable::Private
 
 			case EOpType::MI_CONSTANT:
 				{
-					OP::ResourceConstantArgs Args = Program.GetOpArgs<OP::ResourceConstantArgs>(Item.At);
-					TSharedPtr<const FMaterial> SourceConst = Program.ConstantMaterials[Args.value];
-					MI_REFERENCE_ID.Add(SourceConst->ReferenceID);
+					OP::MaterialConstantArgs Args = Program.GetOpArgs<OP::MaterialConstantArgs>(Item.At);
+					GetMemory().CurrentPassthroughObjectLoader->Add<UMaterialInterface>(Args.ID);
 
 					StoreNone(Item);
 					break;
diff --git a/Engine/Plugins/Mutable/Source/MutableRuntime/Private/MuR/CodeRunnerBegin.h b/Engine/Plugins/Mutable/Source/MutableRuntime/Private/MuR/CodeRunnerBegin.h
--- a/Engine/Plugins/Mutable/Source/MutableRuntime/Private/MuR/CodeRunnerBegin.h
+++ b/Engine/Plugins/Mutable/Source/MutableRuntime/Private/MuR/CodeRunnerBegin.h
@@ -102,7 +102,7 @@ namespace UE::Mutable::Private
 
 		TSharedPtr<const FInstance> RunCode(const FScheduledOpInline& Root);
 
-		FProgramCache& GetMemory() const;
+		FWorkingMemoryManager& GetMemory() const;
 		
 		FScheduledOpInline PopOp();
 		
@@ -154,11 +154,5 @@ namespace UE::Mutable::Private
 		TMap<FCacheAddressInline, TSharedPtr<const FInstance>> ResultsInstance;
 		TMap<FCacheAddressInline, TSharedPtr<const String>> ResultsString;
 		TMap<FCacheAddressInline, TSharedPtr<const FExtensionData>> ResultsExtensionData;
-		
-		TArray<int32> MI_REFERENCE_ID;
-		TArray<int32> IM_REFERENCE_ID;
-		TArray<int32> ME_REFERENCE_ID;
-    	
-		TArray<int32> ME_SKELETON_ID;
 	};
 }
diff --git a/Engine/Plugins/Mutable/Source/MutableRuntime/Private/MuR/CodeRunnerTasked.cpp b/Engine/Plugins/Mutable/Source/MutableRuntime/Private/MuR/CodeRunnerTasked.cpp
--- a/Engine/Plugins/Mutable/Source/MutableRuntime/Private/MuR/CodeRunnerTasked.cpp
+++ b/Engine/Plugins/Mutable/Source/MutableRuntime/Private/MuR/CodeRunnerTasked.cpp
@@ -2616,18 +2616,17 @@ namespace UE::Mutable::Private
 			TSharedPtr<FRangeIndex> Index = BuildCurrentOpRangeIndex(Item, Args.MaterialParameter);
 
 			// Get material parameter from the array of parameters
-			TSharedPtr<FMaterial> Material = MakeShared<FMaterial>();
-			Material->Material = TStrongObjectPtr(Params->GetMaterialValue(Args.MaterialParameter, Index.Get()));
+			UMaterialInterface* Material = Params->GetMaterialValue(Args.MaterialParameter, Index.Get());
 
 			// Get the texture parameter name
 			check(Args.ParameterName < (uint32)Program.ConstantStrings.Num());
 			const FString& ParameterName = Program.ConstantStrings[Args.ParameterName];
 
 			// Get the parameter texture from the UMaterial
-			if (Material->Material)
+			if (Material)
 			{
 				UTexture* Texture;
-				bool bParameterFound = Material->Material->GetTextureParameterValue(FName(ParameterName), Texture);
+				bool bParameterFound = Material->GetTextureParameterValue(FName(ParameterName), Texture);
 
 				if (bParameterFound && Texture)
 				{
diff --git a/Engine/Plugins/Mutable/Source/MutableRuntime/Private/MuR/Image.cpp b/Engine/Plugins/Mutable/Source/MutableRuntime/Private/MuR/Image.cpp
--- a/Engine/Plugins/Mutable/Source/MutableRuntime/Private/MuR/Image.cpp
+++ b/Engine/Plugins/Mutable/Source/MutableRuntime/Private/MuR/Image.cpp
@@ -43,14 +43,13 @@ namespace UE::Mutable::Private
 	}
 
 
-	TSharedPtr<FImage> FImage::CreateAsReference(uint32 ID, const FImageDesc& Desc, bool bForceLoad)
+	TSharedPtr<FImage> FImage::CreateAsReference(const TPassthroughObjectPtr<UTexture>& ExternalObject, const FImageDesc& Desc, bool bForceLoad)
 	{
 		TSharedPtr<FImage> Result = MakeShared<FImage>();
-		Result->ReferenceID = ID;
+		Result->PassthroughObject = ExternalObject;
 	
 		Result->DataStorage.InitVoid(Desc);
 
-		Result->Flags = EImageFlags::IF_IS_REFERENCE;
 		if (bForceLoad)
 		{
 			Result->Flags = Result->Flags | EImageFlags::IF_IS_FORCELOAD;
@@ -214,21 +213,14 @@ namespace UE::Mutable::Private
 
 	bool FImage::IsReference() const
 	{
-		return Flags & EImageFlags::IF_IS_REFERENCE;
+		return PassthroughObject.IsValid();
 	}
 
 	bool FImage::IsForceLoad() const
 	{
 		return Flags & EImageFlags::IF_IS_FORCELOAD;
 	}
-
-	//---------------------------------------------------------------------------------------------
-	uint32 FImage::GetReferencedTexture() const
-	{
-		ensure(IsReference());
-		return ReferenceID;
-	}
-
+	
 	int32 FImage::CalculateDataSize(int32 SizeX, int32 SizeY, int32 LODCount, EImageFormat Format)
 	{
 		int32 Result = 0;
diff --git a/Engine/Plugins/Mutable/Source/MutableRuntime/Private/MuR/Material.cpp b/Engine/Plugins/Mutable/Source/MutableRuntime/Private/MuR/Material.cpp
--- a/Engine/Plugins/Mutable/Source/MutableRuntime/Private/MuR/Material.cpp
+++ b/Engine/Plugins/Mutable/Source/MutableRuntime/Private/MuR/Material.cpp
@@ -32,7 +32,6 @@ namespace UE::Mutable::Private
 
 	void FMaterial::Serialise(FOutputArchive& Arch) const
 	{
-		Arch << ReferenceID;
 		Arch << ImageParameters;
 		Arch << ColorParameters;
 		Arch << ScalarParameters;
@@ -41,7 +40,6 @@ namespace UE::Mutable::Private
 
 	void FMaterial::Unserialise(FInputArchive& Arch)
 	{
-		Arch >>ReferenceID;
 		Arch >> ImageParameters;
 		Arch >> ColorParameters;
 		Arch >> ScalarParameters;
@@ -52,8 +50,7 @@ namespace UE::Mutable::Private
 	{
 		TSharedPtr<FMaterial> Result = MakeShared<FMaterial>();
 
-		Result->Material = Material;
-		Result->ReferenceID = ReferenceID;
+		Result->PassthroughObject = PassthroughObject;
 		Result->ImageParameters = ImageParameters;
 		Result->ColorParameters = ColorParameters;
 		Result->ScalarParameters = ScalarParameters;
@@ -64,7 +61,9 @@ namespace UE::Mutable::Private
 
 	bool FMaterial::operator==(const FMaterial& Other) const
 	{
-		return ReferenceID == Other.ReferenceID && Material == Other.Material	&&
-			ImageParameters == Other.ImageParameters &&	ColorParameters == Other.ColorParameters && ScalarParameters == Other.ScalarParameters;
+		return PassthroughObject == Other.PassthroughObject &&
+			ImageParameters == Other.ImageParameters &&
+			ColorParameters == Other.ColorParameters &&
+			ScalarParameters == Other.ScalarParameters;
 	};
 }
diff --git a/Engine/Plugins/Mutable/Source/MutableRuntime/Private/MuR/Mesh.cpp b/Engine/Plugins/Mutable/Source/MutableRuntime/Private/MuR/Mesh.cpp
--- a/Engine/Plugins/Mutable/Source/MutableRuntime/Private/MuR/Mesh.cpp
+++ b/Engine/Plugins/Mutable/Source/MutableRuntime/Private/MuR/Mesh.cpp
@@ -43,12 +43,11 @@ TSharedPtr<FMesh> FMesh::StaticUnserialise(FInputArchive& Arch)
 }
 
 
-TSharedPtr<FMesh> FMesh::CreateAsReference(uint32 ID,  bool bForceLoad)
+TSharedPtr<FMesh> FMesh::CreateAsReference(const TPassthroughObjectPtr<USkeletalMesh>& PasstroughObjectPtr,  bool bForceLoad)
 {
 	TSharedPtr<FMesh> Result = MakeShared<FMesh>();
-	Result->ReferenceID = ID;
+	Result->PassthroughObject = PasstroughObjectPtr;
 
-	EnumAddFlags(Result->Flags, EMeshFlags::IsResourceReference);
 	if (bForceLoad)
 	{
 		EnumAddFlags(Result->Flags, EMeshFlags::IsResourceForceLoad);
@@ -60,7 +59,7 @@ TSharedPtr<FMesh> FMesh::CreateAsReference(uint32 ID,  bool bForceLoad)
 
 bool FMesh::IsReference() const
 {
-	return EnumHasAnyFlags(Flags, EMeshFlags::IsResourceReference);
+	return PassthroughObject.IsValid();
 }
 
 
@@ -68,14 +67,7 @@ bool FMesh::IsForceLoad() const
 {
 	return EnumHasAnyFlags(Flags, EMeshFlags::IsResourceForceLoad);
 }
-
-
-uint32 FMesh::GetReferencedMesh() const
-{
-	ensure(IsReference());
-	return ReferenceID;
-}
-
+	
 
 void FMesh::SetReferencedMorph(const FString& MorphName)
 {
@@ -115,7 +107,7 @@ void FMesh::CopyFrom(const FMesh& From, EMeshCopyFlags InFlags)
 
     InternalId = From.InternalId;
 	Flags = From.Flags;
-	ReferenceID = From.ReferenceID;
+	PassthroughObject = From.PassthroughObject;
 
 	MeshIDPrefix = From.MeshIDPrefix;
 	SkeletalMeshes = From.SkeletalMeshes;
@@ -920,6 +912,8 @@ void FMesh::FBonePose::Unserialise(FInputArchive& Arch)
 
 void FMesh::Serialise(FOutputArchive& Arch) const
 {
+	check(!IsReference());
+	
 	Arch << IndexBuffers;
 	Arch << VertexBuffers;
 	Arch << AdditionalBuffers;
@@ -942,12 +936,6 @@ void FMesh::Serialise(FOutputArchive& Arch) const
 	Arch << AdditionalPhysicsBodies;
 
 	Arch << MeshIDPrefix;
-
-	if ( IsReference() )
-	{
-		Arch << ReferenceID;
-		Arch << ReferencedMorph;
-	}
 }
 
 
@@ -978,12 +966,6 @@ void FMesh::Unserialise(FInputArchive& Arch)
 	Arch >> AdditionalPhysicsBodies;
 
 	Arch >> MeshIDPrefix;
-
-	if (IsReference())
-	{
-		Arch >> ReferenceID;
-		Arch >> ReferencedMorph;
-	}
 }
 
 
@@ -997,7 +979,7 @@ bool FMesh::IsSimilar(const FMesh& Other) const
 	}
 
 	bool bEqual = IndexBuffers == Other.IndexBuffers;
-	bEqual = bEqual && (ReferenceID == Other.ReferenceID);
+	bEqual = bEqual && (PassthroughObject == Other.PassthroughObject);
 
 	if (bEqual && Skeleton != Other.Skeleton)
 	{
diff --git a/Engine/Plugins/Mutable/Source/MutableRuntime/Private/MuR/PassthroughObject.cpp b/Engine/Plugins/Mutable/Source/MutableRuntime/Private/MuR/PassthroughObject.cpp
new file mode 100644
--- /dev/null
+++ b/Engine/Plugins/Mutable/Source/MutableRuntime/Private/MuR/PassthroughObject.cpp
@@ -0,0 +1,91 @@
+﻿// Copyright Epic Games, Inc. All Rights Reserved.
+
+#include "MuR/PassthroughObject.h"
+
+#include "MuR/MutableRuntimeModule.h"
+
+
+namespace UE::Mutable::Private
+{
+	PASSTHROUGH_ID GetPassthroughId(const TSoftObjectPtr<UObject>& SoftObjectPtr)
+	{
+		PASSTHROUGH_ID Id = GetTypeHash(SoftObjectPtr.ToSoftObjectPath());
+
+		if (Id == PASSTHROUGH_ID_INVALID)
+		{
+			++Id;
+		}
+		
+		return Id;
+	}
+
+
+	void FPassthroughObjectLoader::Load(const TMap<PASSTHROUGH_ID, TSoftObjectPtr<UObject>>& ResolutionMap, TArray<FSoftObjectPath>& AssetsToStream) const
+	{
+		for (const TPair<PASSTHROUGH_ID, TPassthroughObjectPtr<UObject>>& Pair:  PassthroughObjects)
+		{
+			const TPassthroughObjectPtr<UObject>& ExternalObject = Pair.Value;
+				
+			if (!ExternalObject.Resource->IsType<PASSTHROUGH_ID>()) 
+			{
+				// Already resolved.
+				continue;
+			}
+
+			const PASSTHROUGH_ID ID = ExternalObject.Resource->Get<PASSTHROUGH_ID>();
+			if (ID == PASSTHROUGH_ID_INVALID)
+			{
+				continue;
+			}
+
+			const TSoftObjectPtr<UObject>* Result = ResolutionMap.Find(ID);
+			if (!Result)
+			{
+				UE_LOG(LogMutableCore, Error, TEXT("Unable to find Passthrough Object with id: %u"), Pair.Key);
+				continue;
+			}
+			
+			AssetsToStream.Add(Result->ToSoftObjectPath());
+		}
+	}
+
+
+	void FPassthroughObjectLoader::Resolve(TMap<PASSTHROUGH_ID, TSoftObjectPtr<UObject>>& ResolutionMap)
+	{
+		for (const TPair<PASSTHROUGH_ID, TPassthroughObjectPtr<UObject>>& Pair:  PassthroughObjects)
+		{
+			const TPassthroughObjectPtr<UObject>& ExternalObject = Pair.Value;
+
+			if (!ExternalObject.Resource->IsType<PASSTHROUGH_ID>())
+			{
+				// Already resolved.
+				continue;
+			}
+
+			const PASSTHROUGH_ID ID = ExternalObject.Resource->Get<PASSTHROUGH_ID>();
+			if (ID == PASSTHROUGH_ID_INVALID)
+			{
+				ExternalObject.Resource->Set<TStrongObjectPtr<UObject>>(nullptr);
+				continue;
+			}
+			
+			const TSoftObjectPtr<UObject>* Result = ResolutionMap.Find(ID);
+			if (!Result)
+			{
+				// Error already thrown in Load.
+				ExternalObject.Resource->Set<TStrongObjectPtr<UObject>>(nullptr);
+				continue;
+			}
+
+			UObject* Object = Result->Get();
+			if (!Object)
+			{
+				UE_LOG(LogMutableCore, Error, TEXT("Unable to load Passthrough object: %s"), *Result->ToSoftObjectPath().ToString());
+				ExternalObject.Resource->Set<TStrongObjectPtr<UObject>>(nullptr);
+				continue;
+			}
+
+			ExternalObject.Resource->Set<TStrongObjectPtr<UObject>>(TStrongObjectPtr(Object));
+		}
+	}
+}
diff --git a/Engine/Plugins/Mutable/Source/MutableRuntime/Private/MuR/System.cpp b/Engine/Plugins/Mutable/Source/MutableRuntime/Private/MuR/System.cpp
--- a/Engine/Plugins/Mutable/Source/MutableRuntime/Private/MuR/System.cpp
+++ b/Engine/Plugins/Mutable/Source/MutableRuntime/Private/MuR/System.cpp
@@ -197,7 +197,9 @@ namespace UE::Mutable::Private
 
 
     //---------------------------------------------------------------------------------------------
-    FInstance::FID FSystem::NewInstance( const TSharedPtr<const FModel>& InModel, const TSharedPtr<FExternalResourceProvider>& ExternalResourceProvider)
+    FInstance::FID FSystem::NewInstance( const TSharedPtr<const FModel>& InModel,
+    	const TSharedPtr<FExternalResourceProvider>& ExternalResourceProvider,
+    	const TSharedPtr<FPassthroughObjectLoader>& PassthroughObjectLoader)
     {
 		LLM_SCOPE_BYNAME(TEXT("MutableRuntime"));
 		MUTABLE_CPUPROFILER_SCOPE(NewInstance);
@@ -209,6 +211,7 @@ namespace UE::Mutable::Private
 		InstanceData.State = -1;
 		InstanceData.Cache = MakeShared<FProgramCache>();
     	InstanceData.ExternalResourceProvider = ExternalResourceProvider;
+    	InstanceData.PassthroughObjectLoader = PassthroughObjectLoader;
 		m_pD->WorkingMemoryManager.LiveInstances.Add(InstanceData);
 
 		TRACE_COUNTER_SET(MutableRuntime_LiveInstances, m_pD->WorkingMemoryManager.LiveInstances.Num());
@@ -217,12 +220,15 @@ namespace UE::Mutable::Private
 	}
 
 
-    void FSystem::ReuseInstance(FInstance::FID InstanceId, const TSharedPtr<FExternalResourceProvider>& ExternalResourceProvider)
+    void FSystem::ReuseInstance(FInstance::FID InstanceId,
+    	const TSharedPtr<FExternalResourceProvider>& ExternalResourceProvider,
+    	const TSharedPtr<FPassthroughObjectLoader>& PassthroughObjectLoader)
     {
     	FLiveInstance* LiveInstance = m_pD->FindLiveInstance(InstanceId);
     	check(LiveInstance);
 
     	LiveInstance->ExternalResourceProvider = ExternalResourceProvider;
+    	LiveInstance->PassthroughObjectLoader = PassthroughObjectLoader;
     }
 
 
@@ -232,7 +238,7 @@ namespace UE::Mutable::Private
                                      const TSharedPtr<FMeshIdRegistry>& InMeshIdRegistry,
 									 const TSharedPtr<FImageIdRegistry>& InImageIdRegistry,
 									 const TSharedPtr<FMaterialIdRegistry>& InMaterialIdRegistry,
-                                     int32 InStateIndex,
+									 int32 InStateIndex,
                                      uint32 InLodMask)
     {
 		LLM_SCOPE_BYNAME(TEXT("MutableRuntime"));
@@ -252,11 +258,13 @@ namespace UE::Mutable::Private
 		}
 
     	m_pD->WorkingMemoryManager.CurrentInstanceCache = LiveInstance->Cache;
+    	m_pD->WorkingMemoryManager.CurrentPassthroughObjectLoader = LiveInstance->PassthroughObjectLoader;
 
 		m_pD->WorkingMemoryManager.CurrentMeshIdRegistry = InMeshIdRegistry;
     	m_pD->WorkingMemoryManager.CurrentImageIdRegistry = InImageIdRegistry;
     	m_pD->WorkingMemoryManager.CurrentMaterialIdRegistry = InMaterialIdRegistry;
     	
+    	
 		OP::ADDRESS RootAt = LiveInstance->Model->GetPrivate()->Program.States[InStateIndex].Root;
 
     	CodeRunnerBegin Runner(*this, LiveInstance->Model, InParams, InLodMask);
@@ -267,6 +275,13 @@ namespace UE::Mutable::Private
 			// In case of failure return an empty instance, to prevent following code to have to check it every time
 			Result = MakeShared<FInstance>();
 		}
+
+    	m_pD->WorkingMemoryManager.CurrentInstanceCache = nullptr;
+    	m_pD->WorkingMemoryManager.CurrentPassthroughObjectLoader = nullptr;
+
+    	m_pD->WorkingMemoryManager.CurrentMeshIdRegistry = nullptr;
+    	m_pD->WorkingMemoryManager.CurrentImageIdRegistry = nullptr;
+    	m_pD->WorkingMemoryManager.CurrentMaterialIdRegistry = nullptr;
     	
 		return Result;
 	}
@@ -298,6 +313,7 @@ namespace UE::Mutable::Private
 		}
 
 		m_pD->WorkingMemoryManager.CurrentInstanceCache = LiveInstance->Cache;
+    	m_pD->WorkingMemoryManager.CurrentPassthroughObjectLoader = LiveInstance->PassthroughObjectLoader;
 
     	m_pD->WorkingMemoryManager.CurrentMeshIdRegistry = InMeshIdRegistry;
     	m_pD->WorkingMemoryManager.CurrentImageIdRegistry = InImageIdRegistry;
@@ -361,7 +377,12 @@ namespace UE::Mutable::Private
 		m_pD->WorkingMemoryManager.EndRunnerThread();
 
 		m_pD->WorkingMemoryManager.CurrentInstanceCache = nullptr;
+    	m_pD->WorkingMemoryManager.CurrentPassthroughObjectLoader = nullptr;
 
+    	m_pD->WorkingMemoryManager.CurrentMeshIdRegistry = nullptr;
+    	m_pD->WorkingMemoryManager.CurrentImageIdRegistry = nullptr;
+    	m_pD->WorkingMemoryManager.CurrentMaterialIdRegistry = nullptr;
+    	
 		return Result;
 	}
 
@@ -388,6 +409,7 @@ namespace UE::Mutable::Private
 		FLiveInstance* LiveInstance = m_pD->FindLiveInstance(instanceID);
 		check(LiveInstance);
 		m_pD->WorkingMemoryManager.CurrentInstanceCache = LiveInstance->Cache;
+		m_pD->WorkingMemoryManager.CurrentPassthroughObjectLoader = LiveInstance->PassthroughObjectLoader;
 
 		OP::ADDRESS RootAddress = ImageId.GetKey()->Address;
 		pResult = m_pD->BuildImage(LiveInstance->ExternalResourceProvider, LiveInstance->Model, LiveInstance->OldParameters.Get(), RootAddress, MipsToSkip, InImageLOD);
@@ -399,7 +421,8 @@ namespace UE::Mutable::Private
 		}
 
 		m_pD->WorkingMemoryManager.CurrentInstanceCache = nullptr;
-
+    	m_pD->WorkingMemoryManager.CurrentPassthroughObjectLoader = nullptr;
+    	
 		return pResult;
 	}
 
@@ -417,6 +440,7 @@ namespace UE::Mutable::Private
 		FLiveInstance* LiveInstance = m_pD->FindLiveInstance(instanceID);
 		check(LiveInstance);
 		m_pD->WorkingMemoryManager.CurrentInstanceCache = LiveInstance->Cache;
+		m_pD->WorkingMemoryManager.CurrentPassthroughObjectLoader = LiveInstance->PassthroughObjectLoader;
 
 		OP::ADDRESS RootAddress = ImageId.GetKey()->Address;
 		
@@ -427,7 +451,8 @@ namespace UE::Mutable::Private
 		{
 			m_pD->WorkingMemoryManager.EndRunnerThread();
 			m_pD->WorkingMemoryManager.CurrentInstanceCache = nullptr;
-
+			m_pD->WorkingMemoryManager.CurrentPassthroughObjectLoader = nullptr;
+			
 			return UE::Tasks::MakeCompletedTask<TSharedPtr<const FImage>>(BlackImage());
 		}
 		
@@ -464,6 +489,7 @@ namespace UE::Mutable::Private
 
 					SystemPrivate->WorkingMemoryManager.EndRunnerThread();
 					SystemPrivate->WorkingMemoryManager.CurrentInstanceCache = nullptr;
+					SystemPrivate->WorkingMemoryManager.CurrentPassthroughObjectLoader = nullptr;
 					
 					return Result;
 				},
@@ -497,6 +523,7 @@ namespace UE::Mutable::Private
 		FLiveInstance* LiveInstance = m_pD->FindLiveInstance(instanceID);
 		check(LiveInstance);
 		m_pD->WorkingMemoryManager.CurrentInstanceCache = LiveInstance->Cache;
+		m_pD->WorkingMemoryManager.CurrentPassthroughObjectLoader = LiveInstance->PassthroughObjectLoader;
 
 		OP::ADDRESS RootAddress = ImageId.GetKey()->Address;
 
@@ -540,6 +567,7 @@ namespace UE::Mutable::Private
 		}
 
 		m_pD->WorkingMemoryManager.CurrentInstanceCache = nullptr;
+		m_pD->WorkingMemoryManager.CurrentPassthroughObjectLoader = nullptr;
 
 		return Result;
 	}
@@ -559,6 +587,7 @@ namespace UE::Mutable::Private
 		FLiveInstance* LiveInstance = m_pD->FindLiveInstance(instanceID);
 		check(LiveInstance);
 		m_pD->WorkingMemoryManager.CurrentInstanceCache = LiveInstance->Cache;
+		m_pD->WorkingMemoryManager.CurrentPassthroughObjectLoader = LiveInstance->PassthroughObjectLoader;
 
 		OP::ADDRESS RootAddress = ImageId.GetKey()->Address;
 
@@ -578,6 +607,7 @@ namespace UE::Mutable::Private
 		{
 			m_pD->WorkingMemoryManager.EndRunnerThread();
 			m_pD->WorkingMemoryManager.CurrentInstanceCache = nullptr;
+			m_pD->WorkingMemoryManager.CurrentPassthroughObjectLoader = nullptr;
 
 			return UE::Tasks::MakeCompletedTask<FExtendedImageDesc>(); 
 		}
@@ -607,6 +637,7 @@ namespace UE::Mutable::Private
 
 					SystemPrivate->WorkingMemoryManager.EndRunnerThread();
 					SystemPrivate->WorkingMemoryManager.CurrentInstanceCache = nullptr;
+					SystemPrivate->WorkingMemoryManager.CurrentPassthroughObjectLoader = nullptr;
 
 					return Result;
 				},
@@ -632,6 +663,7 @@ namespace UE::Mutable::Private
 		FLiveInstance* LiveInstance = m_pD->FindLiveInstance(instanceID);
 		check(LiveInstance);
 		m_pD->WorkingMemoryManager.CurrentInstanceCache = LiveInstance->Cache;
+		m_pD->WorkingMemoryManager.CurrentPassthroughObjectLoader = LiveInstance->PassthroughObjectLoader;
 
 		OP::ADDRESS RootAddress = MeshId.GetKey()->Address;
 		Result = m_pD->BuildMesh(
@@ -648,6 +680,8 @@ namespace UE::Mutable::Private
 		}
 
 		m_pD->WorkingMemoryManager.CurrentInstanceCache = nullptr;
+		m_pD->WorkingMemoryManager.CurrentPassthroughObjectLoader = nullptr;
+		
 		return Result;
 	}
 	
@@ -666,6 +700,7 @@ namespace UE::Mutable::Private
 		FLiveInstance* LiveInstance = m_pD->FindLiveInstance(instanceID);
 		check(LiveInstance);
 		m_pD->WorkingMemoryManager.CurrentInstanceCache = LiveInstance->Cache;
+		m_pD->WorkingMemoryManager.CurrentPassthroughObjectLoader = LiveInstance->PassthroughObjectLoader;
 		m_pD->WorkingMemoryManager.BeginRunnerThread();
 		
 		OP::ADDRESS RootAddress = MeshId.GetKey()->Address;
@@ -675,6 +710,7 @@ namespace UE::Mutable::Private
 		{
 			m_pD->WorkingMemoryManager.EndRunnerThread();
 			m_pD->WorkingMemoryManager.CurrentInstanceCache = nullptr;
+			m_pD->WorkingMemoryManager.CurrentPassthroughObjectLoader = nullptr;
 
 			return UE::Tasks::MakeCompletedTask<TSharedPtr<const FMesh>>(new FMesh());
 		}
@@ -709,6 +745,7 @@ namespace UE::Mutable::Private
 					
 					SystemPrivate->WorkingMemoryManager.EndRunnerThread();
 					SystemPrivate->WorkingMemoryManager.CurrentInstanceCache = nullptr;
+					SystemPrivate->WorkingMemoryManager.CurrentPassthroughObjectLoader = nullptr;
 					
 					return Result;
 				},
@@ -732,6 +769,7 @@ namespace UE::Mutable::Private
 		FLiveInstance* LiveInstance = m_pD->FindLiveInstance(InstanceID);
 		check(LiveInstance);
 		m_pD->WorkingMemoryManager.CurrentInstanceCache = LiveInstance->Cache;
+		m_pD->WorkingMemoryManager.CurrentPassthroughObjectLoader = LiveInstance->PassthroughObjectLoader;
 		m_pD->WorkingMemoryManager.BeginRunnerThread();
 		
 		OP::ADDRESS RootAddress = MaterialId.GetKey()->Address;
@@ -741,6 +779,7 @@ namespace UE::Mutable::Private
 		{
 			m_pD->WorkingMemoryManager.EndRunnerThread();
 			m_pD->WorkingMemoryManager.CurrentInstanceCache = nullptr;
+			m_pD->WorkingMemoryManager.CurrentPassthroughObjectLoader = nullptr;
 
 			return UE::Tasks::MakeCompletedTask<TSharedPtr<const FMaterial>>(new FMaterial());
 		}
@@ -774,6 +813,7 @@ namespace UE::Mutable::Private
 					
 					SystemPrivate->WorkingMemoryManager.EndRunnerThread();
 					SystemPrivate->WorkingMemoryManager.CurrentInstanceCache = nullptr;
+					SystemPrivate->WorkingMemoryManager.CurrentPassthroughObjectLoader = nullptr;
 					
 					return Result;
 				},
@@ -797,6 +837,7 @@ namespace UE::Mutable::Private
 			LiveInstance->Cache->CheckHitCountsCleared();
 
 			m_pD->WorkingMemoryManager.CurrentInstanceCache = LiveInstance->Cache;
+			m_pD->WorkingMemoryManager.CurrentPassthroughObjectLoader = LiveInstance->PassthroughObjectLoader;
 
 			// We don't want to clear the cache layer 1 because it contains data that can be useful for a 
 			// future update (same states, just runtime parameters changed).
@@ -807,6 +848,7 @@ namespace UE::Mutable::Private
 			m_pD->WorkingMemoryManager.ClearCacheLayer0();
 
 			m_pD->WorkingMemoryManager.CurrentInstanceCache = nullptr;
+			m_pD->WorkingMemoryManager.CurrentPassthroughObjectLoader = nullptr;
 		}
 
 		// Reduce the cache until it fits the limit.
@@ -941,10 +983,12 @@ namespace UE::Mutable::Private
 		FLiveInstance* LiveInstance = m_pD->FindLiveInstance(InstanceID);
 		check(LiveInstance);
 		m_pD->WorkingMemoryManager.CurrentInstanceCache = LiveInstance->Cache;
+		m_pD->WorkingMemoryManager.CurrentPassthroughObjectLoader = LiveInstance->PassthroughObjectLoader;
 		
 		RelevantParameterVisitor visitor( m_pD, LiveInstance->Model, FParameters, Flags );
 
 		m_pD->WorkingMemoryManager.CurrentInstanceCache = nullptr;
+		m_pD->WorkingMemoryManager.CurrentPassthroughObjectLoader = nullptr;
     }
 
 
@@ -1019,6 +1063,8 @@ namespace UE::Mutable::Private
 		WorkingMemoryManager.CurrentInstanceCache = MakeShared<FProgramCache>();
 		WorkingMemoryManager.CurrentInstanceCache->Init(InModel->GetPrivate()->Program.OpAddress.Num());
 
+    	WorkingMemoryManager.CurrentPassthroughObjectLoader = nullptr;
+    	
     	WorkingMemoryManager.CurrentMeshIdRegistry = InMeshIdRegistry;
     	WorkingMemoryManager.CurrentImageIdRegistry = InImageIdRegistry;
     	WorkingMemoryManager.CurrentMaterialIdRegistry = InMaterialIdRegistry;
@@ -1034,6 +1080,10 @@ namespace UE::Mutable::Private
 	void FSystem::Private::EndBuild()
 	{
 		WorkingMemoryManager.CurrentInstanceCache = nullptr;
+		WorkingMemoryManager.CurrentPassthroughObjectLoader = nullptr;
+    	WorkingMemoryManager.CurrentMeshIdRegistry = nullptr;
+    	WorkingMemoryManager.CurrentImageIdRegistry = nullptr;
+    	WorkingMemoryManager.CurrentMaterialIdRegistry = nullptr;
 	}
 
 
diff --git a/Engine/Plugins/Mutable/Source/MutableTools/Internal/MuT/Compiler.h b/Engine/Plugins/Mutable/Source/MutableTools/Internal/MuT/Compiler.h
--- a/Engine/Plugins/Mutable/Source/MutableTools/Internal/MuT/Compiler.h
+++ b/Engine/Plugins/Mutable/Source/MutableTools/Internal/MuT/Compiler.h
@@ -28,7 +28,7 @@ namespace UE::Mutable::Private
     class NodeTransformedObject;
 
 	/** Arguments are Texture ID, output generated texture and a "run immediately" flag. */
-	typedef TFunction<UE::Tasks::FTask(int32, TSharedPtr<TSharedPtr<FImage>>, bool)> FReferencedImageResourceFunc;
+	typedef TFunction<UE::Tasks::FTask(PASSTHROUGH_ID, TSharedPtr<TSharedPtr<FImage>>, bool)> FReferencedImageResourceFunc;
 
 	/** Arguments are 
 	* - Mesh ID, 
@@ -36,7 +36,7 @@ namespace UE::Mutable::Private
 	* - output generated mesh 
 	* - "run immediately" flag. 
 	*/
-	typedef TFunction<UE::Tasks::FTask(int32, const FString&, TSharedPtr<TSharedPtr<FMesh>>, bool)> FReferencedMeshResourceFunc;
+	typedef TFunction<UE::Tasks::FTask(PASSTHROUGH_ID, const FString&, TSharedPtr<TSharedPtr<FMesh>>, bool)> FReferencedMeshResourceFunc;
 
     /** Options used to compile the models with a compiler. */
     class CompilerOptions : public RefCounted
diff --git a/Engine/Plugins/Mutable/Source/MutableTools/Internal/MuT/NodeMaterialConstant.h b/Engine/Plugins/Mutable/Source/MutableTools/Internal/MuT/NodeMaterialConstant.h
--- a/Engine/Plugins/Mutable/Source/MutableTools/Internal/MuT/NodeMaterialConstant.h
+++ b/Engine/Plugins/Mutable/Source/MutableTools/Internal/MuT/NodeMaterialConstant.h
@@ -5,6 +5,7 @@
 #include "HAL/Platform.h"
 #include "MuR/Ptr.h"
 #include "MuR/RefCounted.h"
+#include "MuR/PassthroughObject.h"
 #include "MuT/NodeMaterial.h"
 #include "MuT/NodeImage.h"
 #include "MuR/Types.h"
@@ -23,7 +24,7 @@ namespace UE::Mutable::Private
 		virtual const FNodeType* GetType() const override { return &StaticType; }
 		static const FNodeType* GetStaticType() { return &StaticType; }
 
-		int32 MaterialId = INDEX_NONE;
+		PASSTHROUGH_ID MaterialId = PASSTHROUGH_ID_INVALID;
 
 		TMap<FName, FVector4f> ColorValues;
 		TMap<FName, float> ScalarValues;
diff --git a/Engine/Plugins/Mutable/Source/MutableTools/Private/MuT/ASTOpConstantResource.cpp b/Engine/Plugins/Mutable/Source/MutableTools/Private/MuT/ASTOpConstantResource.cpp
--- a/Engine/Plugins/Mutable/Source/MutableTools/Private/MuT/ASTOpConstantResource.cpp
+++ b/Engine/Plugins/Mutable/Source/MutableTools/Private/MuT/ASTOpConstantResource.cpp
@@ -614,68 +614,92 @@ namespace UE::Mutable::Private
 
 		if (!linkedAddress && !bLinkedAndNull)
 		{
-			if (Type == EOpType::ME_CONSTANT)
+			switch (Type)
 			{
-				OP::MeshConstantArgs Args;
-				FMemory::Memzero(Args);
+			case EOpType::ME_CONSTANT:
+				{
+					OP::MeshConstantArgs Args;
+					FMemory::Memzero(Args);
 
-				TSharedPtr<FMesh> MeshValue = StaticCastSharedPtr<const FMesh>(GetValue())->Clone();
-				check(MeshValue);
+					TSharedPtr<FMesh> MeshValue = StaticCastSharedPtr<const FMesh>(GetValue())->Clone();
+					check(MeshValue);
 
-				Args.Skeleton = -1;
-				if (TSharedPtr<const FSkeleton> MeshSkeleton = MeshValue->GetSkeleton())
-				{
-					Args.Skeleton = Program.AddConstant(MeshSkeleton);
-					MeshValue->SetSkeleton(nullptr);
-				}
+					Args.Skeleton = -1;
+					if (TSharedPtr<const FSkeleton> MeshSkeleton = MeshValue->GetSkeleton())
+					{
+						Args.Skeleton = Program.AddConstant(MeshSkeleton);
+						MeshValue->SetSkeleton(nullptr);
+					}
 
-				Args.Value = AddConstantMesh(Program, MeshValue, *Options);
-				int32 DataDescIndex = Options->AdditionalData.SourceMeshPerConstant.Add(SourceDataDescriptor);
-				check(DataDescIndex == Args.Value);
+					check(!MeshValue->PassthroughObject);
+					
+					Args.Value = AddConstantMesh(Program, MeshValue, *Options);
+					int32 DataDescIndex = Options->AdditionalData.SourceMeshPerConstant.Add(SourceDataDescriptor);
+					check(DataDescIndex == Args.Value);
 
-				linkedAddress = (OP::ADDRESS)Program.OpAddress.Num();
-				Program.OpAddress.Add((uint32)Program.ByteCode.Num());
-				AppendCode(Program.ByteCode, Type);
-				AppendCode(Program.ByteCode, Args);
-			}
-			else
-			{
-				OP::ResourceConstantArgs Args;
-				FMemory::Memzero(Args);
+					linkedAddress = (OP::ADDRESS)Program.OpAddress.Num();
+					Program.OpAddress.Add((uint32)Program.ByteCode.Num());
+					AppendCode(Program.ByteCode, Type);
+					AppendCode(Program.ByteCode, Args);
 
-				bool bValidData = true;
+					break;
+				}
 
-				switch (Type)
-				{
-				case EOpType::IM_CONSTANT:
+			case EOpType::IM_CONSTANT:
 				{
+					OP::ResourceConstantArgs Args;
+					FMemory::Memzero(Args);
+					
 					TSharedPtr<const FImage> pTyped = StaticCastSharedPtr<const FImage>(GetValue());
 					check(pTyped);
 
 					if (pTyped->GetSizeX() * pTyped->GetSizeY() == 0)
 					{
 						// It's an empty or degenerated image, return a null operation.
-						bValidData = false;
+						// Null op
+						linkedAddress = 0;
+						bLinkedAndNull = true;
 					}
 					else
 					{
-						Args.value = AddConstantImage( Program, pTyped, *Options);
+						check(!pTyped->PassthroughObject);
 
+						Args.value = AddConstantImage( Program, pTyped, *Options);
+						
 						int32 DataDescIndex = Options->AdditionalData.SourceImagePerConstant.Add(SourceDataDescriptor);
 						check(DataDescIndex == Args.value);
+						
+						linkedAddress = (OP::ADDRESS)Program.OpAddress.Num();
+						Program.OpAddress.Add((uint32)Program.ByteCode.Num());
+						AppendCode(Program.ByteCode, Type);
+						AppendCode(Program.ByteCode, Args);
 					}
 
 					break;
 				}
-				case EOpType::LA_CONSTANT:
+				
+			case EOpType::LA_CONSTANT:
 				{
+					OP::ResourceConstantArgs Args;
+					FMemory::Memzero(Args);
+					
 					TSharedPtr<const FLayout> pTyped = StaticCastSharedPtr<const FLayout>(GetValue());
 					check(pTyped);
 					Args.value = Program.AddConstant(pTyped);
+
+					linkedAddress = (OP::ADDRESS)Program.OpAddress.Num();
+					Program.OpAddress.Add((uint32)Program.ByteCode.Num());
+					AppendCode(Program.ByteCode, Type);
+					AppendCode(Program.ByteCode, Args);
+					
 					break;
 				}
-				case EOpType::MI_CONSTANT:
+
+			case EOpType::MI_CONSTANT:
 				{
+					OP::MaterialConstantArgs Args;
+					FMemory::Memzero(Args);
+
 					TSharedPtr<FMaterial> MaterialValue = StaticCastSharedPtr<const FMaterial>(GetValue())->Clone();
 					check(MaterialValue);
 
@@ -688,29 +712,22 @@ namespace UE::Mutable::Private
 						MaterialValue->ImageParameters.Add(FName(ImageOperation.Key), NewImage);
 					}
 
-					Args.value = Program.AddConstant(MaterialValue);
+					Args.ID = MaterialValue->PassthroughObject.GetId();
+					MaterialValue->PassthroughObject = {}; // Not serialized but compared. Set to invalid to avoid duplicating it.
+					
+					Args.Value = Program.AddConstant(MaterialValue);
 					
-					break;
-				}
-				default:
-					check(false);
-				}
-
-				if (bValidData)
-				{
 					linkedAddress = (OP::ADDRESS)Program.OpAddress.Num();
 					Program.OpAddress.Add((uint32)Program.ByteCode.Num());
 					AppendCode(Program.ByteCode, Type);
 					AppendCode(Program.ByteCode, Args);
+					
+					break;
 				}
-				else
-				{
-					// Null op
-					linkedAddress = 0;
-					bLinkedAndNull = true;
-				}
+			default:
+				unimplemented();
 			}
-
+			
 			// Clear stored value to reduce memory usage.
 			LoadedValue = nullptr;
 			Proxy = nullptr;
diff --git a/Engine/Plugins/Mutable/Source/MutableTools/Private/MuT/ASTOpConstantResource.h b/Engine/Plugins/Mutable/Source/MutableTools/Private/MuT/ASTOpConstantResource.h
--- a/Engine/Plugins/Mutable/Source/MutableTools/Private/MuT/ASTOpConstantResource.h
+++ b/Engine/Plugins/Mutable/Source/MutableTools/Private/MuT/ASTOpConstantResource.h
@@ -38,6 +38,8 @@ namespace UE::Mutable::Private
 		/** Source data descriptor. */
 		FSourceDataDescriptor SourceDataDescriptor;
 
+		PASSTHROUGH_ID PassthroughId = PASSTHROUGH_ID_INVALID;
+		
 		/** Image Parameters are evaluated at runtime. */
 		TMap<FName, ASTChild> ImageOperations;
 
diff --git a/Engine/Plugins/Mutable/Source/MutableTools/Private/MuT/ASTOpReferenceResource.h b/Engine/Plugins/Mutable/Source/MutableTools/Private/MuT/ASTOpReferenceResource.h
--- a/Engine/Plugins/Mutable/Source/MutableTools/Private/MuT/ASTOpReferenceResource.h
+++ b/Engine/Plugins/Mutable/Source/MutableTools/Private/MuT/ASTOpReferenceResource.h
@@ -30,7 +30,7 @@ namespace UE::Mutable::Private
 		bool bForceLoad = false;
 
 		/** Externally provided ID to identify the resource. */
-		uint32 ID = 0;
+		PASSTHROUGH_ID ID = PASSTHROUGH_ID_INVALID;
 
 		/** */
 		FImageDesc ImageDesc;
diff --git a/Engine/Plugins/Mutable/Source/MutableTools/Private/MuT/CodeGenerator.cpp b/Engine/Plugins/Mutable/Source/MutableTools/Private/MuT/CodeGenerator.cpp
--- a/Engine/Plugins/Mutable/Source/MutableTools/Private/MuT/CodeGenerator.cpp
+++ b/Engine/Plugins/Mutable/Source/MutableTools/Private/MuT/CodeGenerator.cpp
@@ -2448,7 +2448,7 @@ namespace UE::Mutable::Private
 				{
 					FMesh* OriginalMesh = static_cast<const NodeMeshConstant*>(OriginalMeshNode)->Value.Get();
 					check(OriginalMesh->IsReference());
-					uint32 OriginalMeshID = OriginalMesh->GetReferencedMesh();
+					PASSTHROUGH_ID OriginalMeshID = OriginalMesh->PassthroughObject.GetId();
 
 					bool bRunImmediatlyIfPossible = IsInGameThread();
 					TargetMeshTask = CompilerOptions->OptimisationOptions.ReferencedMeshResourceProvider(
@@ -2639,7 +2639,7 @@ namespace UE::Mutable::Private
 				{
 					FMesh* OriginalMesh = static_cast<const NodeMeshConstant*>(OriginalMeshNode)->Value.Get();
 					check(OriginalMesh->IsReference());
-					uint32 OriginalMeshID = OriginalMesh->GetReferencedMesh();
+					PASSTHROUGH_ID OriginalMeshID = OriginalMesh->PassthroughObject.GetId();
 
 					bool bRunImmediatlyIfPossible = IsInGameThread();
 					FString NoMorph;
diff --git a/Engine/Plugins/Mutable/Source/MutableTools/Private/MuT/CodeGenerator_Image.cpp b/Engine/Plugins/Mutable/Source/MutableTools/Private/MuT/CodeGenerator_Image.cpp
--- a/Engine/Plugins/Mutable/Source/MutableTools/Private/MuT/CodeGenerator_Image.cpp
+++ b/Engine/Plugins/Mutable/Source/MutableTools/Private/MuT/CodeGenerator_Image.cpp
@@ -240,7 +240,7 @@ namespace UE::Mutable::Private
 		{
 			Ptr<ASTOpReferenceResource> ReferenceOp = new ASTOpReferenceResource();
 			ReferenceOp->Type = EOpType::IM_REFERENCE;
-			ReferenceOp->ID = pImage->GetReferencedTexture();
+			ReferenceOp->ID = pImage->PassthroughObject.GetId();
 			ReferenceOp->bForceLoad = pImage->IsForceLoad();
 			ReferenceOp->SourceDataDescriptor = InNode->SourceDataDescriptor;
 
diff --git a/Engine/Plugins/Mutable/Source/MutableTools/Private/MuT/CodeGenerator_Material.cpp b/Engine/Plugins/Mutable/Source/MutableTools/Private/MuT/CodeGenerator_Material.cpp
--- a/Engine/Plugins/Mutable/Source/MutableTools/Private/MuT/CodeGenerator_Material.cpp
+++ b/Engine/Plugins/Mutable/Source/MutableTools/Private/MuT/CodeGenerator_Material.cpp
@@ -95,7 +95,7 @@ namespace UE::Mutable::Private
 	{
 		// Constant Material Resource
 		TSharedPtr<FMaterial> MaterialResource = MakeShared<FMaterial>();
-		MaterialResource->ReferenceID = Node->MaterialId;
+		MaterialResource->PassthroughObject = TPassthroughObjectPtr<UMaterialInterface>(Node->MaterialId);
 		MaterialResource->ColorParameters = Node->ColorValues;
 		MaterialResource->ScalarParameters = Node->ScalarValues;
 
@@ -252,7 +252,7 @@ namespace UE::Mutable::Private
 					// TODO(Max): UE-314401
 					// Constant Material Resource
 					TSharedPtr<FMaterial> MaterialResource = MakeShared<FMaterial>();
-					MaterialResource->ReferenceID = MaterialId;
+					MaterialResource->PassthroughObject = TPassthroughObjectPtr<UMaterialInterface>(MaterialId);
 
 					// Constant Material Operation
 					Ptr<ASTOpConstantResource> ConstantMaterial = new ASTOpConstantResource();
@@ -348,7 +348,7 @@ namespace UE::Mutable::Private
 
 		// Make a checkered debug image
 		TSharedPtr<FMaterial> GenerateMaterial = MakeShared<FMaterial>();
-		GenerateMaterial->ReferenceID = INDEX_NONE;
+		GenerateMaterial->PassthroughObject = {};
 
 		// Store the constant material
 		ConstantMaterial->SetValue(GenerateMaterial, CompilerOptions->OptimisationOptions.DiskCacheContext);
diff --git a/Engine/Plugins/Mutable/Source/MutableTools/Private/MuT/CodeGenerator_Mesh.cpp b/Engine/Plugins/Mutable/Source/MutableTools/Private/MuT/CodeGenerator_Mesh.cpp
--- a/Engine/Plugins/Mutable/Source/MutableTools/Private/MuT/CodeGenerator_Mesh.cpp
+++ b/Engine/Plugins/Mutable/Source/MutableTools/Private/MuT/CodeGenerator_Mesh.cpp
@@ -755,7 +755,7 @@ namespace UE::Mutable::Private
 		{
 			Ptr<ASTOpReferenceResource> ReferenceOp = new ASTOpReferenceResource();
 			ReferenceOp->Type = EOpType::ME_REFERENCE;
-			ReferenceOp->ID = Mesh->GetReferencedMesh();
+			ReferenceOp->ID = Mesh->PassthroughObject.GetId();
 			ReferenceOp->bForceLoad = false;
 
 			FMeshGenerationResult Result;
@@ -796,7 +796,7 @@ namespace UE::Mutable::Private
 		// Compile-time reference? Add task to resolve as requisite.
 		if (Mesh->IsReference() && Mesh->IsForceLoad())
 		{
-			uint32 MeshID = Mesh->GetReferencedMesh();
+			PASSTHROUGH_ID MeshID = Mesh->PassthroughObject.GetId();
 
 			bool bRunImmediatlyIfPossible = IsInGameThread();
 			FString Morph = Mesh->GetReferencedMorph();
@@ -1389,7 +1389,7 @@ namespace UE::Mutable::Private
 		{
 			if (Mesh->IsReference() && Mesh->IsForceLoad())
 			{
-				uint32 MeshID = Mesh->GetReferencedMesh();
+				PASSTHROUGH_ID MeshID = Mesh->PassthroughObject.GetId();
 
 				bool bRunImmediatlyIfPossible = IsInGameThread();
 				FString Morph = Mesh->GetReferencedMorph();
diff --git a/Engine/Plugins/Mutable/Source/MutableTools/Private/MuT/CodeGenerator_SecondPass.cpp b/Engine/Plugins/Mutable/Source/MutableTools/Private/MuT/CodeGenerator_SecondPass.cpp
--- a/Engine/Plugins/Mutable/Source/MutableTools/Private/MuT/CodeGenerator_SecondPass.cpp
+++ b/Engine/Plugins/Mutable/Source/MutableTools/Private/MuT/CodeGenerator_SecondPass.cpp
@@ -398,9 +398,9 @@ namespace UE::Mutable::Private
 				Surface.NegativeTags,
 				Empty, Empty, Empty, Empty);
 
-			Ptr<ASTOpBoolAnd> ConditionOp = new ASTOpBoolAnd();
-			ConditionOp->A = Surface.ObjectCondition;
-			ConditionOp->B = c;
+				Ptr<ASTOpBoolAnd> ConditionOp = new ASTOpBoolAnd();
+				ConditionOp->A = Surface.ObjectCondition ? Surface.ObjectCondition : new ASTOpConstantBool(true);
+				ConditionOp->B = c ? c : new ASTOpConstantBool(true);
 
 			Surface.FinalCondition = ConditionOp;
 			
diff --git a/Engine/Plugins/Mutable/Source/MutableTools/Private/MuT/CodeOptimiser.cpp b/Engine/Plugins/Mutable/Source/MutableTools/Private/MuT/CodeOptimiser.cpp
--- a/Engine/Plugins/Mutable/Source/MutableTools/Private/MuT/CodeOptimiser.cpp
+++ b/Engine/Plugins/Mutable/Source/MutableTools/Private/MuT/CodeOptimiser.cpp
@@ -735,7 +735,7 @@ namespace UE::Mutable::Private
 						||
 						SubgraphType == EOpType::IM_CONSTANT)
 					{
-						uint32 ImageID = 0;
+						PASSTHROUGH_ID ImageID = PASSTHROUGH_ID_INVALID;
 						if (SubgraphType == EOpType::IM_REFERENCE)
 						{
 							bIsReference = true;
@@ -749,7 +749,7 @@ namespace UE::Mutable::Private
 							bIsReference = Value->IsReference();
 							if (bIsReference)
 							{
-								ImageID = Value->GetReferencedTexture();
+								ImageID = Value->PassthroughObject.GetId();
 							}
 						}
 
@@ -792,7 +792,7 @@ namespace UE::Mutable::Private
 						||
 						SubgraphType == EOpType::ME_CONSTANT)
 					{
-						uint32 MeshID = 0;
+						PASSTHROUGH_ID MeshID = PASSTHROUGH_ID_INVALID;
 						FString MeshMorph;
 						if (SubgraphType == EOpType::ME_REFERENCE)
 						{
@@ -807,7 +807,7 @@ namespace UE::Mutable::Private
 							bIsReference = Value->IsReference();
 							if (bIsReference)
 							{
-								MeshID = Value->GetReferencedMesh();
+								MeshID = Value->PassthroughObject.GetId();
 								MeshMorph = Value->GetReferencedMorph();
 							}
 						}
