[mutable] Added support into the MutableDebugger for the LinearToSRGB mutable operation type.

#rnx

diff --git a/Engine/Plugins/Mutable/Source/CustomizableObjectEditor/Private/MuCOE/SMutableCodeViewer.cpp b/Engine/Plugins/Mutable/Source/CustomizableObjectEditor/Private/MuCOE/SMutableCodeViewer.cpp
--- a/Engine/Plugins/Mutable/Source/CustomizableObjectEditor/Private/MuCOE/SMutableCodeViewer.cpp
+++ b/Engine/Plugins/Mutable/Source/CustomizableObjectEditor/Private/MuCOE/SMutableCodeViewer.cpp
@@ -1669,6 +1669,13 @@ void SMutableCodeViewer::GenerateElementRecursive(const int32& InStateIndex, UE:
 		AddOpFunc(Args.Y, TEXT("y "));
 		break;
 	}
+
+	case UE::Mutable::Private::EOpType::CO_LINEARTOSRGB:
+	{
+		UE::Mutable::Private::OP::ColorArgs Args = InProgram.GetOpArgs<UE::Mutable::Private::OP::ColorArgs>(InParentAddress);
+		AddOpFunc(Args.Color, TEXT("Color "));
+		break;
+	}
 		
 	case UE::Mutable::Private::EOpType::IM_RESIZELIKE:
 	{
