[Mutable] Fix "Unable to load Passthrough object" errors.

diff --git a/Engine/Plugins/Mutable/Source/MutableRuntime/Private/MuR/CodeRunnerBegin.cpp b/Engine/Plugins/Mutable/Source/MutableRuntime/Private/MuR/CodeRunnerBegin.cpp
--- a/Engine/Plugins/Mutable/Source/MutableRuntime/Private/MuR/CodeRunnerBegin.cpp
+++ b/Engine/Plugins/Mutable/Source/MutableRuntime/Private/MuR/CodeRunnerBegin.cpp
@@ -1621,7 +1621,6 @@ namespace UE::Mutable::Private
 			case EOpType::ME_CONSTANT:
 				{
 					OP::MeshConstantArgs Args = Program.GetOpArgs<OP::MeshConstantArgs>(Item.At);
-					GetMemory().CurrentPassthroughObjectLoader->Add<USkeletalMesh>(Args.Skeleton);
 
 					StoreNone(Item);
 					break;
@@ -1630,7 +1629,11 @@ namespace UE::Mutable::Private
 			case EOpType::ME_REFERENCE:
 				{
 					OP::ResourceReferenceArgs Args = Program.GetOpArgs<OP::ResourceReferenceArgs>(Item.At);
-					GetMemory().CurrentPassthroughObjectLoader->Add<USkeletalMesh>(Args.ID);
+
+					if (!Args.ForceLoad) // Not a passthrough. Only happens with None optimization.
+					{
+						GetMemory().CurrentPassthroughObjectLoader->Add<USkeletalMesh>(Args.ID);
+					}
 
 					StoreNone(Item);
 					break;
@@ -1941,25 +1944,6 @@ namespace UE::Mutable::Private
 					using OpEnumFlags = OP::MeshAddMetadataArgs::EnumFlags;
 					const bool bIsSkeletonList = EnumHasAnyFlags(OpArgs.Flags, OpEnumFlags::IsSkeletonList);
 
-					if (bIsSkeletonList)
-					{
-						if (Program.ConstantUInt32Lists.IsValidIndex(OpArgs.SkeletonIds.ListAddress))
-						{
-							for (const uint32 SkeletonId : Program.ConstantUInt32Lists[OpArgs.SkeletonIds.ListAddress])
-							{
-								GetMemory().CurrentPassthroughObjectLoader->Add<USkeleton>(SkeletonId);
-							}
-						}
-						else
-						{
-							check(false);
-						}
-					}
-					else
-					{
-						GetMemory().CurrentPassthroughObjectLoader->Add<USkeleton>(OpArgs.SkeletonIds.SkeletonId);
-					}
-
 					StoreNone(Item);
 					break;
 				}
@@ -2005,7 +1989,11 @@ namespace UE::Mutable::Private
 			case EOpType::IM_REFERENCE:
 				{
 					OP::ResourceReferenceArgs Args = Program.GetOpArgs<OP::ResourceReferenceArgs>(Item.At);
-					GetMemory().CurrentPassthroughObjectLoader->Add<UTexture>(Args.ID);
+
+					if (!Args.ForceLoad) // Not a passthrough. Only happens with None optimization.
+					{
+						GetMemory().CurrentPassthroughObjectLoader->Add<UTexture>(Args.ID);
+					}
 
 					StoreNone(Item);
 					break;
