[Mutable] Fix cook indeterminism with Passthrough objects.

#FailedValidations Preflight

diff --git a/Engine/Plugins/Mutable/Source/MutableRuntime/Private/MuR/PassthroughObject.cpp b/Engine/Plugins/Mutable/Source/MutableRuntime/Private/MuR/PassthroughObject.cpp
--- a/Engine/Plugins/Mutable/Source/MutableRuntime/Private/MuR/PassthroughObject.cpp
+++ b/Engine/Plugins/Mutable/Source/MutableRuntime/Private/MuR/PassthroughObject.cpp
@@ -9,7 +9,7 @@ namespace UE::Mutable::Private
 {
 	PASSTHROUGH_ID GetPassthroughId(const TSoftObjectPtr<UObject>& SoftObjectPtr)
 	{
-		PASSTHROUGH_ID Id = GetTypeHash(SoftObjectPtr.ToSoftObjectPath());
+		PASSTHROUGH_ID Id = GetTypeHash(SoftObjectPtr.ToSoftObjectPath().ToString().ToLower());
 
 		if (Id == PASSTHROUGH_ID_INVALID)
 		{
