[Mutable] Fix crash when stripping section morph data from a parametric mesh.

- The case with an empty batch was not handled correctly when selecting the batches affecting a section.

#rnx

diff --git a/Engine/Plugins/Mutable/Source/CustomizableObject/Private/MuCO/UnrealConversionUtils.cpp b/Engine/Plugins/Mutable/Source/CustomizableObject/Private/MuCO/UnrealConversionUtils.cpp
--- a/Engine/Plugins/Mutable/Source/CustomizableObject/Private/MuCO/UnrealConversionUtils.cpp
+++ b/Engine/Plugins/Mutable/Source/CustomizableObject/Private/MuCO/UnrealConversionUtils.cpp
@@ -1451,6 +1451,11 @@ namespace UnrealConversionUtils
 						const int32 SourceMorphNumBatches = Morph.BatchesPerMorph[MorphIndex];
 						const int32 SourceBatchOffset = Morph.BatchStartOffsetPerMorph[MorphIndex]; 
 
+						if (SourceMorphNumBatches == 0)
+						{
+							continue;
+						}
+
 						int32 BatchIndex = 0;
 						for ( ; BatchIndex < SourceMorphNumBatches; ++BatchIndex)
 						{
