- Fixed up plugins that have both Base and Default ini files, and one plugin (WebSocketNetworking) that had a just wrong named file)

#FailedValidations Preflight

diff --git a/Engine/Plugins/Enterprise/GLTFExporter/Config/BaseGLTFExporter.ini b/Engine/Plugins/Enterprise/GLTFExporter/Config/BaseGLTFExporter.ini
deleted file mode 100644
diff --git a/Engine/Plugins/Experimental/AnimToTexture/Config/BaseAnimToTexture.ini b/Engine/Plugins/Experimental/AnimToTexture/Config/BaseAnimToTexture.ini
deleted file mode 100644
diff --git a/Engine/Plugins/Experimental/AnimToTexture/Config/DefaultAnimToTexture.ini b/Engine/Plugins/Experimental/AnimToTexture/Config/DefaultAnimToTexture.ini
--- a/Engine/Plugins/Experimental/AnimToTexture/Config/DefaultAnimToTexture.ini
+++ b/Engine/Plugins/Experimental/AnimToTexture/Config/DefaultAnimToTexture.ini
@@ -1,3 +1,5 @@
 [CoreRedirects]
++StructRedirects=(OldName="/Script/AnimToTexture.AnimSequenceInfo",NewName="/Script/AnimToTexture.AnimToTextureAnimSequenceInfo")
++StructRedirects=(OldName="/Script/AnimToTexture.AnimInfo",NewName="/Script/AnimToTexture.AnimToTextureAnimInfo")
 +PropertyRedirects==(OldName="AnimToTextureDataAsset.MasterSkeletalMesh",NewName="LeaderSkeletalMesh")
-+FunctionRedirects=(OldName="AnimToTextureDataAsset.BP_GetMasterSkeletalMesh", NewName="BP_GetLeaderSkeletalMesh")
\ No newline at end of file
++FunctionRedirects=(OldName="AnimToTextureDataAsset.BP_GetMasterSkeletalMesh", NewName="BP_GetLeaderSkeletalMesh")
diff --git a/Engine/Plugins/Experimental/WebSocketNetworking/Config/BaseWebSocketNetDriver.ini b/Engine/Plugins/Experimental/WebSocketNetworking/Config/BaseWebSocketNetDriver.ini
deleted file mode 100644
diff --git a/Engine/Plugins/Experimental/WebSocketNetworking/Config/Engine.ini b/Engine/Plugins/Experimental/WebSocketNetworking/Config/Engine.ini
new file mode 100644
--- /dev/null
+++ b/Engine/Plugins/Experimental/WebSocketNetworking/Config/Engine.ini
@@ -0,0 +1,25 @@
+;[/Script/Engine.GameEngine]
+;!NetDriverDefinitions=ClearArray
+;+NetDriverDefinitions=(DefName="GameNetDriver", DriverClassName="/Script/WebSocketNetworking.WebSocketNetDriver", DriverClassNameFallback="/Script/WebSocketNetworking.WebSocketNetDriver")
+;[/Script/Engine.Engine]
+;+NetDriverDefinitions=(DefName="GameNetDriver", DriverClassName="/Script/WebSocketNetworking.WebSocketNetDriver", DriverClassNameFallback="/Script/WebSocketsNetworking.IpNetDriver")
+
+[/Script/WebSocketNetworking.WebSocketNetDriver]
+AllowPeerConnections=False
+AllowPeerVoice=False
+ConnectionTimeout=60.0
+InitialConnectTimeout=120.0
+RecentlyDisconnectedTrackingTime=180
+AckTimeout=10.0
+KeepAliveTime=20.2
+MaxClientRate=15000
+MaxInternetClientRate=10000
+RelevantTimeout=5.0
+SpawnPrioritySeconds=1.0
+ServerTravelPause=4.0
+NetServerMaxTickRate=30
+LanServerMaxTickRate=35
+WebSocketPort=8889
+NetConnectionClassName="/Script/WebSocketNetworking.WebSocketConnection"
+MaxPortCountToTry=512
+
diff --git a/Engine/Plugins/Mutable/Config/BaseMutable.ini b/Engine/Plugins/Mutable/Config/BaseMutable.ini
deleted file mode 100644
diff --git a/Engine/Plugins/Mutable/Config/DefaultMutable.ini b/Engine/Plugins/Mutable/Config/DefaultMutable.ini
--- a/Engine/Plugins/Mutable/Config/DefaultMutable.ini
+++ b/Engine/Plugins/Mutable/Config/DefaultMutable.ini
@@ -1,3 +1,34 @@
+[CoreRedirects]
+
++StructRedirects=(OldName="/Script/CustomizableObject.ParameterUIData",NewName="/Script/CustomizableObject.MutableParameterData")
++StructRedirects=(OldName="/Script/CustomizableObject.CompilationOptions",NewName="/Script/CustomizableObject.CompilationOptions_DEPRECATED")
++StructRedirects=(OldName="/Script/CustomizableObject.FCustomizableObjectAssetParameterValue",NewName="/Script/CustomizableObject.FCustomizableObjectTextureParameterValue")
+
++PropertyRedirects=(OldName="GroupProjectorParameterImage.OptionImage", NewName="OptionTexture")
++PropertyRedirects=(OldName="CustomizableObjectNodeGroupProjectorParameter.OptionImages", NewName="OptionTextures")
++PropertyRedirects=(OldName="CustomizableObjectNodeGroupProjectorParameter.OptionImagesDataTable", NewName="OptionTexturesDataTable")
+
+; Node renames
++ClassRedirects=(OldName="/Script/CustomizableObjectEditor.CustomizableObjectNodeMeshClipDeform", NewName="/Script/CustomizableObjectEditor.CustomizableObjectNodeModifierClipDeform")
++ClassRedirects=(OldName="/Script/CustomizableObjectEditor.CustomizableObjectNodeMeshClipWithMesh", NewName="/Script/CustomizableObjectEditor.CustomizableObjectNodeModifierClipWithMesh")
++ClassRedirects=(OldName="/Script/CustomizableObjectEditor.CustomizableObjectNodeMeshClipMorph", NewName="/Script/CustomizableObjectEditor.CustomizableObjectNodeModifierClipMorph")
++ClassRedirects=(OldName="/Script/CustomizableObjectEditor.CustomizableObjectNodeExtendMaterial", NewName="/Script/CustomizableObjectEditor.CustomizableObjectNodeModifierExtendMeshSection")
++ClassRedirects=(OldName="/Script/CustomizableObjectEditor.CustomizableObjectNodeEditMaterial", NewName="/Script/CustomizableObjectEditor.CustomizableObjectNodeModifierEditMeshSection")
++ClassRedirects=(OldName="/Script/CustomizableObjectEditor.CustomizableObjectNodeMorphMaterial", NewName="/Script/CustomizableObjectEditor.CustomizableObjectNodeModifierMorphMeshSection")
++ClassRedirects=(OldName="/Script/CustomizableObjectEditor.CustomizableObjectNodeRemoveMesh", NewName="/Script/CustomizableObjectEditor.CustomizableObjectNodeModifierRemoveMesh")
++ClassRedirects=(OldName="/Script/CustomizableObjectEditor.CustomizableObjectNodeRemoveMeshBlocks", NewName="/Script/CustomizableObjectEditor.CustomizableObjectNodeModifierRemoveMeshBlocks")
++ClassRedirects=(OldName="/Script/CustomizableObjectEditor.CustomizableObjectNodeEditMaterialBase", NewName="/Script/CustomizableObjectEditor.CustomizableObjectNodeModifierEditMeshSectionBase")
++ClassRedirects=(OldName="/Script/CustomizableObjectEditor.CustomizableObjectNodeEditLayoutBlocks", NewName="/Script/CustomizableObjectEditor.CustomizableObjectNodeModifierEditLayoutBlocks")
+
++FunctionRedirects=(OldName="/Script/CustomizableObject.CustomizableObject.GetIntParameterOptionUIMetadata",NewName="/Script/CustomizableObject.CustomizableObject.GetEnumParameterValueUIMetadata")
++FunctionRedirects=(OldName="/Script/CustomizableObject.CustomizableObject.GetIntParameterGroupType",NewName="/Script/CustomizableObject.CustomizableObject.GetEnumParameterGroupType")
++FunctionRedirects=(OldName="/Script/CustomizableObject.CustomizableObject.GetIntParameterOptionDataTable",NewName="/Script/CustomizableObject.CustomizableObject.GetEnumParameterValueDataTable")
++FunctionRedirects=(OldName="/Script/CustomizableObject.CustomizableObject.GetIntParameterDefaultValue",NewName="/Script/CustomizableObject.CustomizableObject.GetEnumParameterDefaultValue")
++PropertyRedirects=(OldName="/Script/CustomizableObjectEditor.CustomizableObjectNodeComponentPassthroughMesh.Mesh",NewName="/Script/CustomizableObjectEditor.CustomizableObjectNodeComponentPassthroughMesh.SkeletalMesh")
++PropertyRedirects=(OldName="/Script/CustomizableObject.CustomizableObject.bEnableAnimBpPhysicsAssetsManipualtion",NewName="/Script/CustomizableObject.CustomizableObject.bEnableAnimBpPhysicsAssetsManipulation")
++PropertyRedirects=(OldName="/Script/CustomizableObject.CompileParams.bGatherReferneces",NewName="/Script/CustomizableObject.CompileParams.bGatherReferences")
++PropertyRedirects=(OldName="/Script/CustomizableObject.CompileParams.bSkipIfOutOfDate",NewName="/Script/CustomizableObject.CompileParams.bSkipIfNotOutOfDate")
+
 ; Settings for the Mutable plugin
 [Features]
 
diff --git a/Engine/Plugins/Runtime/Metasound/Config/BaseMetasound.ini b/Engine/Plugins/Runtime/Metasound/Config/BaseMetasound.ini
deleted file mode 100644
diff --git a/Engine/Plugins/Runtime/Metasound/Config/DefaultMetaSound.ini b/Engine/Plugins/Runtime/Metasound/Config/DefaultMetaSound.ini
--- a/Engine/Plugins/Runtime/Metasound/Config/DefaultMetaSound.ini
+++ b/Engine/Plugins/Runtime/Metasound/Config/DefaultMetaSound.ini
@@ -0,0 +1,28 @@
+[CoreRedirects]
+
+; 5.0
++PropertyRedirects=(OldName="MetaSoundSettings.AutoUpdateBlacklist",NewName="AutoUpdateDenylist")
++PropertyRedirects=(OldName="MetaSoundSettings.AutoUpdateAssetBlacklist",NewName="AutoUpdateDenylist")
+
++StructRedirects=(OldName="MetasoundEditorGraphInputBoolRef",NewName="/Script/MetasoundEditor.MetasoundEditorGraphMemberDefaultBoolRef")
++StructRedirects=(OldName="MetasoundEditorGraphInputIntRef",NewName="/Script/MetasoundEditor.MetasoundEditorGraphMemberDefaultIntRef")
++StructRedirects=(OldName="MetasoundEditorGraphInputObjectRef",NewName="/Script/MetasoundEditor.MetasoundEditorGraphMemberDefaultObjectRef")
+ 
++ClassRedirects=(OldName="MetasoundEditorGraphInputLiteral",NewName="/Script/MetasoundEditor.MetasoundEditorGraphMemberDefaultLiteral")
++ClassRedirects=(OldName="MetasoundEditorGraphInputBool",NewName="/Script/MetasoundEditor.MetasoundEditorGraphMemberDefaultBool")
++ClassRedirects=(OldName="MetasoundEditorGraphInputBoolArray",NewName="/Script/MetasoundEditor.MetasoundEditorGraphMemberDefaultBoolArray")
++ClassRedirects=(OldName="MetasoundEditorGraphInputInt",NewName="/Script/MetasoundEditor.MetasoundEditorGraphMemberDefaultInt")
++ClassRedirects=(OldName="MetasoundEditorGraphInputIntArray",NewName="/Script/MetasoundEditor.MetasoundEditorGraphMemberDefaultIntArray")
++ClassRedirects=(OldName="MetasoundEditorGraphInputFloat",NewName="/Script/MetasoundEditor.MetasoundEditorGraphMemberDefaultFloat")
++ClassRedirects=(OldName="MetasoundEditorGraphInputFloatArray",NewName="/Script/MetasoundEditor.MetasoundEditorGraphMemberDefaultFloatArray")
++ClassRedirects=(OldName="MetasoundEditorGraphInputString",NewName="/Script/MetasoundEditor.MetasoundEditorGraphMemberDefaultString")
++ClassRedirects=(OldName="MetasoundEditorGraphInputStringArray",NewName="/Script/MetasoundEditor.MetasoundEditorGraphMemberDefaultStringArray")
++ClassRedirects=(OldName="MetasoundEditorGraphInputObject",NewName="/Script/MetasoundEditor.MetasoundEditorGraphMemberDefaultObject")
++ClassRedirects=(OldName="MetasoundEditorGraphInputObjectArray",NewName="/Script/MetasoundEditor.MetasoundEditorGraphMemberDefaultObjectArray")
+
+; 5.1
++ClassRedirects=(OldName="MetaSound",NewName="/Script/MetasoundEngine.MetaSoundPatch")
++PropertyRedirects=(OldName="MetasoundEditorSettings.bPinMetaSoundInAssetMenu",NewName="bPinMetaSoundPatchInAssetMenu")
+
+; 5.3
++EnumRedirects=(OldName="MetasoundSourceAudioFormat",NewName="/Script/MetasoundEngine.MetaSoundOutputAudioFormat")
