[Mutable] Removed warning "Invalid Mesh Parameter. Nullptr".

diff --git a/Engine/Plugins/Mutable/Source/CustomizableObject/Private/MuCO/UnrealMutableImageProvider.cpp b/Engine/Plugins/Mutable/Source/CustomizableObject/Private/MuCO/UnrealMutableImageProvider.cpp
--- a/Engine/Plugins/Mutable/Source/CustomizableObject/Private/MuCO/UnrealMutableImageProvider.cpp
+++ b/Engine/Plugins/Mutable/Source/CustomizableObject/Private/MuCO/UnrealMutableImageProvider.cpp
@@ -471,7 +471,6 @@ TTuple<UE::Tasks::FTask, TFunction<void()>> FUnrealMutableResourceProvider::GetM
 
 	if (!SkeletalMesh)
 	{
-		UE_LOG(LogMutable, Warning, TEXT("Invalid Mesh Parameter. Nullptr."));
 		ResultCallback(Result);
 		return Invoke(TrivialReturn);
 	}
