[Mutable] Passthrough Texture Parameter.

diff --git a/Engine/Plugins/Mutable/Source/CustomizableObject/Internal/MuCO/CustomizableObjectCustomVersion.h b/Engine/Plugins/Mutable/Source/CustomizableObject/Internal/MuCO/CustomizableObjectCustomVersion.h
--- a/Engine/Plugins/Mutable/Source/CustomizableObject/Internal/MuCO/CustomizableObjectCustomVersion.h
+++ b/Engine/Plugins/Mutable/Source/CustomizableObject/Internal/MuCO/CustomizableObjectCustomVersion.h
@@ -206,6 +206,8 @@ struct FCustomizableObjectCustomVersion
 		
 		SetDisplayNamePropertyAsPinNameOfTableNodes2,			// Added in order to apply the same fix for cases where the version was skipped
 
+		NodeTextureParameterPassthroughPin,
+
 		// -----<new versions can be added above this line>-------------------------------------------------
 		VersionPlusOne,
 		LatestVersion = VersionPlusOne - 1
diff --git a/Engine/Plugins/Mutable/Source/CustomizableObjectEditor/Private/MuCOE/GenerateMutableSource/GenerateMutableSourceImage.cpp b/Engine/Plugins/Mutable/Source/CustomizableObjectEditor/Private/MuCOE/GenerateMutableSource/GenerateMutableSourceImage.cpp
--- a/Engine/Plugins/Mutable/Source/CustomizableObjectEditor/Private/MuCOE/GenerateMutableSource/GenerateMutableSourceImage.cpp
+++ b/Engine/Plugins/Mutable/Source/CustomizableObjectEditor/Private/MuCOE/GenerateMutableSource/GenerateMutableSourceImage.cpp
@@ -1295,6 +1295,7 @@ UE::Mutable::Private::Ptr<UE::Mutable::Private::NodeImageParameter> GenerateMuta
 
 		TextureNode->Name = TypedNodeParam->GetParameterName(&GenerationContext.MacroNodesStack);
 		TextureNode->UID = GenerationContext.GetNodeIdUnique(ParameterNode).ToString();
+		TextureNode->bPassthrough = Pin == TypedNodeParam->PassthroughPin.Get();
 
 		if (TypedNodeParam->DefaultValue)
 		{
diff --git a/Engine/Plugins/Mutable/Source/CustomizableObjectEditor/Private/MuCOE/Nodes/CustomizableObjectNodeTextureParameter.cpp b/Engine/Plugins/Mutable/Source/CustomizableObjectEditor/Private/MuCOE/Nodes/CustomizableObjectNodeTextureParameter.cpp
--- a/Engine/Plugins/Mutable/Source/CustomizableObjectEditor/Private/MuCOE/Nodes/CustomizableObjectNodeTextureParameter.cpp
+++ b/Engine/Plugins/Mutable/Source/CustomizableObjectEditor/Private/MuCOE/Nodes/CustomizableObjectNodeTextureParameter.cpp
@@ -19,6 +19,22 @@ bool UCustomizableObjectNodeTextureParameter::IsExperimental() const
 }
 
 
+void UCustomizableObjectNodeTextureParameter::AllocateDefaultPins(UCustomizableObjectNodeRemapPins* RemapPins)
+{
+	Super::AllocateDefaultPins(RemapPins);
+
+	const FName TypePassthrough = UEdGraphSchema_CustomizableObject::PC_PassthroughTexture;
+	
+	const FName PinName = UEdGraphSchema_CustomizableObject::GetPinCategoryName(TypePassthrough);
+	const FText PinFriendlyName = UEdGraphSchema_CustomizableObject::GetPinCategoryFriendlyName(TypePassthrough);
+	
+	UEdGraphPin* Pin = CustomCreatePin(EGPD_Output, TypePassthrough, PinName);
+	Pin->PinFriendlyName = PinFriendlyName;
+
+	PassthroughPin = Pin;
+}
+
+
 void UCustomizableObjectNodeTextureParameter::BackwardsCompatibleFixup(int32 CustomizableObjectCustomVersion)
 {
 	Super::BackwardsCompatibleFixup(CustomizableObjectCustomVersion);
@@ -36,6 +52,18 @@ void UCustomizableObjectNodeTextureParameter::BackwardsCompatibleFixup(int32 Cus
 			Pin->PinFriendlyName = LOCTEXT("Image_Pin_Category", "Texture");
 		}
 	}
+	else if (CustomizableObjectCustomVersion == FCustomizableObjectCustomVersion::NodeTextureParameterPassthroughPin)
+	{
+		const FName TypePassthrough = UEdGraphSchema_CustomizableObject::PC_PassthroughTexture;
+	
+		const FName PinName = UEdGraphSchema_CustomizableObject::GetPinCategoryName(TypePassthrough);
+		const FText PinFriendlyName = UEdGraphSchema_CustomizableObject::GetPinCategoryFriendlyName(TypePassthrough);
+	
+		UEdGraphPin* Pin = CustomCreatePin(EGPD_Output, TypePassthrough, PinName);
+		Pin->PinFriendlyName = PinFriendlyName;
+
+		PassthroughPin = Pin;
+	}
 }
 
 
diff --git a/Engine/Plugins/Mutable/Source/CustomizableObjectEditor/Private/MuCOE/Nodes/CustomizableObjectNodeTextureParameter.h b/Engine/Plugins/Mutable/Source/CustomizableObjectEditor/Private/MuCOE/Nodes/CustomizableObjectNodeTextureParameter.h
--- a/Engine/Plugins/Mutable/Source/CustomizableObjectEditor/Private/MuCOE/Nodes/CustomizableObjectNodeTextureParameter.h
+++ b/Engine/Plugins/Mutable/Source/CustomizableObjectEditor/Private/MuCOE/Nodes/CustomizableObjectNodeTextureParameter.h
@@ -37,12 +37,16 @@ public:
 	UPROPERTY(EditAnywhere, Category = CustomizableObject)
 	int32 TextureSizeY = 0;
 
+	UPROPERTY()
+	FEdGraphPinReference PassthroughPin;
+	
 	// UCustomizableObjectNode interface
-	UE_API virtual void BackwardsCompatibleFixup(int32 CustomizableObjectCustomVersion) override;
-	UE_API virtual bool IsExperimental() const override;
+	virtual void AllocateDefaultPins(UCustomizableObjectNodeRemapPins* RemapPins) override;
+	virtual void BackwardsCompatibleFixup(int32 CustomizableObjectCustomVersion) override;
+	virtual bool IsExperimental() const override;
 
 	// CustomizableObjectNodeParameter interface
-	UE_API virtual FName GetCategory() const override;
+	virtual FName GetCategory() const override;
 };
 
 #undef UE_API
diff --git a/Engine/Plugins/Mutable/Source/MutableRuntime/Internal/MuR/Operations.h b/Engine/Plugins/Mutable/Source/MutableRuntime/Internal/MuR/Operations.h
--- a/Engine/Plugins/Mutable/Source/MutableRuntime/Internal/MuR/Operations.h
+++ b/Engine/Plugins/Mutable/Source/MutableRuntime/Internal/MuR/Operations.h
@@ -59,6 +59,8 @@ namespace UE::Mutable::Private
 		IM_REFERENCE,
 		ME_REFERENCE,
 
+    	IM_PARAMETER_REFERENCE,
+
         //! Select one value or the other depending on a boolean input
         NU_CONDITIONAL,
         SC_CONDITIONAL,
diff --git a/Engine/Plugins/Mutable/Source/MutableRuntime/Private/MuR/CodeRunner.cpp b/Engine/Plugins/Mutable/Source/MutableRuntime/Private/MuR/CodeRunner.cpp
--- a/Engine/Plugins/Mutable/Source/MutableRuntime/Private/MuR/CodeRunner.cpp
+++ b/Engine/Plugins/Mutable/Source/MutableRuntime/Private/MuR/CodeRunner.cpp
@@ -4679,6 +4679,19 @@ namespace UE::Mutable::Private
 			break;
 		}
 
+		case EOpType::IM_PARAMETER_REFERENCE:
+        {
+			OP::ParameterArgs Args = Program.GetOpArgs<OP::ParameterArgs>(Item.At);
+			TSharedPtr<FRangeIndex> Index = BuildCurrentOpRangeIndex(Item, Args.variable);
+			UTexture* Image = pParams->GetImageValue(Args.variable, Index.Get());
+
+       		TPassthroughObjectPtr<UTexture> ExternalObject(Image);
+			TSharedPtr<FImage> Result = FImage::CreateAsReference(ExternalObject, FImageDesc(), false);
+        		
+			StoreImage(Item, Result);
+			break;
+		}
+
         case EOpType::IM_CROP:
         {
 			OP::ImageCropArgs Args = Program.GetOpArgs<OP::ImageCropArgs>(Item.At);
@@ -7043,6 +7056,12 @@ namespace UE::Mutable::Private
 			break;
 		}
 
+		case EOpType::IM_PARAMETER_REFERENCE:
+		{
+			check(false); // Passthrough. No image descriptor required.
+			break;
+		}
+
 		case EOpType::IM_CONDITIONAL:
 		{
 			MUTABLE_CPUPROFILER_SCOPE(RunCodeImageDesc_IM_CONDITIONAL);
diff --git a/Engine/Plugins/Mutable/Source/MutableRuntime/Private/MuR/CodeRunnerBegin.cpp b/Engine/Plugins/Mutable/Source/MutableRuntime/Private/MuR/CodeRunnerBegin.cpp
--- a/Engine/Plugins/Mutable/Source/MutableRuntime/Private/MuR/CodeRunnerBegin.cpp
+++ b/Engine/Plugins/Mutable/Source/MutableRuntime/Private/MuR/CodeRunnerBegin.cpp
@@ -111,7 +111,7 @@ namespace UE::Mutable::Private
 
 		Items.Push(Item);
 	}
-
+ 
 
 	void CodeRunnerBegin::StoreInt(const FCacheAddressInline& To, int32 Value)
 	{
@@ -392,6 +392,7 @@ namespace UE::Mutable::Private
 			case EOpType::IM_PARAMETER:
 			case EOpType::MI_PARAMETER:
 			case EOpType::IM_PARAMETER_FROM_MATERIAL:
+			case EOpType::IM_PARAMETER_REFERENCE:
 				{
 					StoreNone(Item);
 					break;
diff --git a/Engine/Plugins/Mutable/Source/MutableRuntime/Private/MuR/Operations.cpp b/Engine/Plugins/Mutable/Source/MutableRuntime/Private/MuR/Operations.cpp
--- a/Engine/Plugins/Mutable/Source/MutableRuntime/Private/MuR/Operations.cpp
+++ b/Engine/Plugins/Mutable/Source/MutableRuntime/Private/MuR/Operations.cpp
@@ -46,6 +46,8 @@ namespace UE::Mutable::Private
 		{ EDataType::Image,			},	// IM_REFERENCE
 		{ EDataType::Mesh,			},	// ME_REFERENCE
 
+		{ EDataType::Image,			},	// IM_PARAMETER_REFERENCE
+    	
 		{ EDataType::Int,			},	// NU_CONDITIONAL
 		{ EDataType::Scalar,		},	// SC_CONDITIONAL
 		{ EDataType::Color,			},	// CO_CONDITIONAL
@@ -198,6 +200,7 @@ namespace UE::Mutable::Private
 		case EOpType::IM_REFERENCE:
 		case EOpType::ME_REFERENCE:
 		case EOpType::IM_PARAMETER_FROM_MATERIAL:
+		case EOpType::IM_PARAMETER_REFERENCE:
 			break;
 
         case EOpType::SC_CURVE:
diff --git a/Engine/Plugins/Mutable/Source/MutableRuntime/Private/MuR/System.cpp b/Engine/Plugins/Mutable/Source/MutableRuntime/Private/MuR/System.cpp
--- a/Engine/Plugins/Mutable/Source/MutableRuntime/Private/MuR/System.cpp
+++ b/Engine/Plugins/Mutable/Source/MutableRuntime/Private/MuR/System.cpp
@@ -944,6 +944,7 @@ namespace UE::Mutable::Private
 			case EOpType::IM_PARAMETER:
 			case EOpType::MA_PARAMETER:
 			case EOpType::MI_PARAMETER:
+            case EOpType::IM_PARAMETER_REFERENCE:
             {
 				OP::ParameterArgs args = Program.GetOpArgs<OP::ParameterArgs>(at);
 				OP::ADDRESS ParamIndex = args.variable;
diff --git a/Engine/Plugins/Mutable/Source/MutableTools/Internal/MuT/NodeImageParameter.h b/Engine/Plugins/Mutable/Source/MutableTools/Internal/MuT/NodeImageParameter.h
--- a/Engine/Plugins/Mutable/Source/MutableTools/Internal/MuT/NodeImageParameter.h
+++ b/Engine/Plugins/Mutable/Source/MutableTools/Internal/MuT/NodeImageParameter.h
@@ -26,6 +26,8 @@ namespace UE::Mutable::Private
 		/** Ranges for the parameter in case it is multidimensional. */
 		TArray<Ptr<NodeRange>> Ranges;
 
+    	bool bPassthrough = false;
+
 	public:
 
 		// Node Interface
diff --git a/Engine/Plugins/Mutable/Source/MutableTools/Private/MuT/CodeGenerator_Image.cpp b/Engine/Plugins/Mutable/Source/MutableTools/Private/MuT/CodeGenerator_Image.cpp
--- a/Engine/Plugins/Mutable/Source/MutableTools/Private/MuT/CodeGenerator_Image.cpp
+++ b/Engine/Plugins/Mutable/Source/MutableTools/Private/MuT/CodeGenerator_Image.cpp
@@ -277,7 +277,7 @@ namespace UE::Mutable::Private
 			if (!Found)
 			{
 				op = new ASTOpParameter();
-				op->Type = EOpType::IM_PARAMETER;
+				op->Type = Node->bPassthrough ? EOpType::IM_PARAMETER_REFERENCE : EOpType::IM_PARAMETER;
 
 				op->Parameter.Name = Node->Name;
 				bool bParseOk = FGuid::Parse(Node->UID, op->Parameter.UID);
diff --git a/Engine/Plugins/Mutable/Source/MutableTools/Private/MuT/CodeOptimiser.cpp b/Engine/Plugins/Mutable/Source/MutableTools/Private/MuT/CodeOptimiser.cpp
--- a/Engine/Plugins/Mutable/Source/MutableTools/Private/MuT/CodeOptimiser.cpp
+++ b/Engine/Plugins/Mutable/Source/MutableTools/Private/MuT/CodeOptimiser.cpp
@@ -552,6 +552,7 @@ namespace UE::Mutable::Private
 					case EOpType::MA_PARAMETER:
 					case EOpType::MI_PARAMETER:
 					case EOpType::IM_PARAMETER_FROM_MATERIAL:
+					case EOpType::IM_PARAMETER_REFERENCE:
 						bIsConstantSubgraph = false;
 						break;
 					default:
diff --git a/Engine/Plugins/Mutable/Source/MutableTools/Private/MuT/CodeOptimiser_Parameters.cpp b/Engine/Plugins/Mutable/Source/MutableTools/Private/MuT/CodeOptimiser_Parameters.cpp
--- a/Engine/Plugins/Mutable/Source/MutableTools/Private/MuT/CodeOptimiser_Parameters.cpp
+++ b/Engine/Plugins/Mutable/Source/MutableTools/Private/MuT/CodeOptimiser_Parameters.cpp
@@ -1582,6 +1582,7 @@ namespace UE::Mutable::Private
 				case EOpType::IM_PARAMETER:
 				case EOpType::ME_PARAMETER:
 				case EOpType::MA_PARAMETER:
+				case EOpType::IM_PARAMETER_REFERENCE:
 				{
 					const ASTOpParameter* TypedOp = static_cast<const ASTOpParameter*>(At.get());
                     Parameters.Add(TypedOp->Parameter.Name);
diff --git a/Engine/Plugins/Mutable/Source/MutableTools/Private/MuT/Compiler.cpp b/Engine/Plugins/Mutable/Source/MutableTools/Private/MuT/Compiler.cpp
--- a/Engine/Plugins/Mutable/Source/MutableTools/Private/MuT/Compiler.cpp
+++ b/Engine/Plugins/Mutable/Source/MutableTools/Private/MuT/Compiler.cpp
@@ -75,7 +75,7 @@ namespace UE::Mutable::Private
 		{ false,	{ 0, 0, 0, 0, 0, 0, 0, 0, 0, 0, 0, 0, 0, 0, 0, 0, 0, 0, 0, 0, 0, 0, 0, 0, 0, 0, 0, 0, 0 }		},	// CO_PARAMETER
 		{ false,	{ 0, 0, 0, 0, 0, 0, 0, 0, 0, 0, 0, 0, 0, 0, 0, 0, 0, 0, 0, 0, 0, 0, 0, 0, 0, 0, 0, 0, 0 }		},	// PR_PARAMETER
 		{ false,	{ 0, 0, 0, 0, 0, 0, 0, 0, 0, 0, 0, 0, 0, 0, 0, 0, 0, 0, 0, 0, 0, 0, 0, 0, 0, 0, 0, 0, 0 }		},	// IM_PARAMETER
-		{ true,		{ 0, 0, 0, 0, 0, 0, 0, 0, 0, 0, 0, 0, 0, 0, 0, 0, 0, 0, 0, 0, 0, 0, 0, 0, 0, 0, 0, 0, 0 }		},	// ME_PARAMETER
+		{ false,	{ 0, 0, 0, 0, 0, 0, 0, 0, 0, 0, 0, 0, 0, 0, 0, 0, 0, 0, 0, 0, 0, 0, 0, 0, 0, 0, 0, 0, 0 }		},	// ME_PARAMETER
 		{ false,	{ 0, 0, 0, 0, 0, 0, 0, 0, 0, 0, 0, 0, 0, 0, 0, 0, 0, 0, 0, 0, 0, 0, 0, 0, 0, 0, 0, 0, 0 }		},	// ST_PARAMETER
 		{ false,	{ 0, 0, 0, 0, 0, 0, 0, 0, 0, 0, 0, 0, 0, 0, 0, 0, 0, 0, 0, 0, 0, 0, 0, 0, 0, 0, 0, 0, 0 }		},	// MA_PARAMETER
 		{ false,	{ 0, 0, 0, 0, 0, 0, 0, 0, 0, 0, 0, 0, 0, 0, 0, 0, 0, 0, 0, 0, 0, 0, 0, 0, 0, 0, 0, 0, 0 }		},	// MI_PARAMETER
@@ -83,6 +83,8 @@ namespace UE::Mutable::Private
 		{ true,		{ 0, 0, 0, 0, 0, 0, 0, 0, 0, 0, 0, 0, 0, 0, 0, 0, 0, 0, 0, 0, 0, 0, 0, 0, 0, 0, 0, 0, 0 }		},	// IM_REFERENCE
 		{ true,		{ 0, 0, 0, 0, 0, 0, 0, 0, 0, 0, 0, 0, 0, 0, 0, 0, 0, 0, 0, 0, 0, 0, 0, 0, 0, 0, 0, 0, 0 }		},	// ME_REFERENCE
 
+    	{ false,	{ 0, 0, 0, 0, 0, 0, 0, 0, 0, 0, 0, 0, 0, 0, 0, 0, 0, 0, 0, 0, 0, 0, 0, 0, 0, 0, 0, 0, 0 }		},	// IM_PARAMETER_REFERENCE
+
 		{ false,	{ 0, 0, 0, 0, 0, 0, 0, 0, 0, 0, 0, 0, 0, 0, 0, 0, 0, 0, 0, 0, 0, 0, 0, 0, 0, 0, 0, 0, 0 }		},	// NU_CONDITIONAL
 		{ false,	{ 0, 0, 0, 0, 0, 0, 0, 0, 0, 0, 0, 0, 0, 0, 0, 0, 0, 0, 0, 0, 0, 0, 0, 0, 0, 0, 0, 0, 0 }		},	// SC_CONDITIONAL
 		{ false,	{ 0, 0, 0, 0, 0, 0, 0, 0, 0, 0, 0, 0, 0, 0, 0, 0, 0, 0, 0, 0, 0, 0, 0, 0, 0, 0, 0, 0, 0 }		},	// CO_CONDITIONAL
diff --git a/Engine/Plugins/Mutable/Source/MutableTools/Private/MuT/ErrorLog.cpp b/Engine/Plugins/Mutable/Source/MutableTools/Private/MuT/ErrorLog.cpp
--- a/Engine/Plugins/Mutable/Source/MutableTools/Private/MuT/ErrorLog.cpp
+++ b/Engine/Plugins/Mutable/Source/MutableTools/Private/MuT/ErrorLog.cpp
@@ -228,6 +228,8 @@ namespace UE::Mutable::Private
 		TEXT("IM_REFERENCE     "),
 		TEXT("ME_REFERENCE     "),
 
+    	TEXT("IM_PARAMETER_REFERENCE"),
+    	
 		TEXT("NU_CONDITIONAL   "),
 		TEXT("SC_CONDITIONAL   "),
 		TEXT("CO_CONDITIONAL   "),
