[Mutable] Fix modifiers not being applied correctly with optimization None.

diff --git a/Engine/Plugins/Mutable/Source/CustomizableObject/Internal/MuCO/CustomizableObjectPrivate.h b/Engine/Plugins/Mutable/Source/CustomizableObject/Internal/MuCO/CustomizableObjectPrivate.h
--- a/Engine/Plugins/Mutable/Source/CustomizableObject/Internal/MuCO/CustomizableObjectPrivate.h
+++ b/Engine/Plugins/Mutable/Source/CustomizableObject/Internal/MuCO/CustomizableObjectPrivate.h
@@ -1734,6 +1734,8 @@ enum class ECustomizableObjectVersions : int32
 
 	FixSocketPriority,
 
+	FixDefaultBehaviourAndsWithoutOperants,
+
 	// -----<new versions can be added above this line>--------
 	LastCustomizableObjectVersion
 };
diff --git a/Engine/Plugins/Mutable/Source/MutableTools/Private/MuT/CodeGenerator.cpp b/Engine/Plugins/Mutable/Source/MutableTools/Private/MuT/CodeGenerator.cpp
--- a/Engine/Plugins/Mutable/Source/MutableTools/Private/MuT/CodeGenerator.cpp
+++ b/Engine/Plugins/Mutable/Source/MutableTools/Private/MuT/CodeGenerator.cpp
@@ -647,8 +647,8 @@ namespace UE::Mutable::Private
 					{
 						// TODO: This could be done earlier?
 						Ptr<ASTOpBoolAnd> ConditionOp = new ASTOpBoolAnd();
-						ConditionOp->A = Component.ObjectCondition;
-						ConditionOp->B = Component.ComponentCondition;
+						ConditionOp->A = Component.ObjectCondition ? Component.ObjectCondition : new ASTOpConstantBool(true);
+						ConditionOp->B = Component.ComponentCondition ? Component.ComponentCondition : new ASTOpConstantBool(true);
 
 						Ptr<ASTOpConditional> IfOp = new ASTOpConditional();
 						IfOp->type = EOpType::IN_CONDITIONAL;
diff --git a/Engine/Plugins/Mutable/Source/MutableTools/Private/MuT/CodeGenerator_FirstPass.cpp b/Engine/Plugins/Mutable/Source/MutableTools/Private/MuT/CodeGenerator_FirstPass.cpp
--- a/Engine/Plugins/Mutable/Source/MutableTools/Private/MuT/CodeGenerator_FirstPass.cpp
+++ b/Engine/Plugins/Mutable/Source/MutableTools/Private/MuT/CodeGenerator_FirstPass.cpp
@@ -664,7 +664,7 @@ namespace UE::Mutable::Private
             if ( UE::Mutable::Private::Ptr<const NodeObject> ChildNode = Node->Children[t] )
             {
                 // Overwrite the implicit condition
-                Ptr<ASTOp> paramOp = 0;
+                Ptr<ASTOp> paramOp = new ASTOpConstantBool(true);
                 switch (Node->Type )
                 {
                     case NodeObjectGroup::CS_TOGGLE_EACH:
diff --git a/Engine/Plugins/Mutable/Source/MutableTools/Private/MuT/CodeGenerator_SecondPass.cpp b/Engine/Plugins/Mutable/Source/MutableTools/Private/MuT/CodeGenerator_SecondPass.cpp
--- a/Engine/Plugins/Mutable/Source/MutableTools/Private/MuT/CodeGenerator_SecondPass.cpp
+++ b/Engine/Plugins/Mutable/Source/MutableTools/Private/MuT/CodeGenerator_SecondPass.cpp
@@ -67,7 +67,7 @@ namespace UE::Mutable::Private
 			return *Found;
 		}
 		
-		Ptr<ASTOp> c;
+		Ptr<ASTOp> c = new ASTOpConstantBool(false);
 
 		// Condition expression for all the surfaces that activate the tag
 		for (int32 surfIndex : t.Surfaces)
@@ -113,7 +113,7 @@ namespace UE::Mutable::Private
 			if (surfCondition)
 			{
 				Ptr<ASTOpBoolAnd> f = new ASTOpBoolAnd;
-				f->A = surface.ObjectCondition;
+				f->A = surface.ObjectCondition ? surface.ObjectCondition : new ASTOpConstantBool(true);
 				f->B = surfCondition;
 				fullCondition = OpPool.Add(f);
 			}
@@ -121,19 +121,11 @@ namespace UE::Mutable::Private
 			{
 				fullCondition = OpPool.Add(surface.ObjectCondition);
 			}
-
-
-			if (!c)
-			{
-				c = fullCondition;
-			}
-			else
-			{
-				Ptr<ASTOpBoolOr> o = new ASTOpBoolOr();
-				o->A = fullCondition;
-				o->B = c;
-				c = OpPool.Add(o);
-			}
+			
+			Ptr<ASTOpBoolOr> o = new ASTOpBoolOr();
+			o->A = fullCondition;
+			o->B = c;
+			c = OpPool.Add(o);
 
 			// Optimise the condition now.
 			//PartialOptimise( c, CompilerOptions->OptimisationOptions );
@@ -165,7 +157,7 @@ namespace UE::Mutable::Private
 			return OpPool.Add(new ASTOpConstantBool(false));
 		}
 
-		Ptr<ASTOp> c;
+		Ptr<ASTOp> c = new ASTOpConstantBool(true);
 
 		for (const FString& PositiveTag : PositiveTags)
 		{
@@ -206,17 +198,10 @@ namespace UE::Mutable::Private
 
 			if (!isConstant)
 			{
-				if (!c)
-				{
-					c = tagCondition;
-				}
-				else
-				{
-					Ptr<ASTOpBoolAnd> o = new ASTOpBoolAnd;
-					o->A = tagCondition;
-					o->B = c;
-					c = OpPool.Add(o);
-				}
+				Ptr<ASTOpBoolAnd> o = new ASTOpBoolAnd;
+				o->A = tagCondition;
+				o->B = c;
+				c = OpPool.Add(o);
 			}
 			else if (constantValue == true)
 			{
@@ -278,17 +263,10 @@ namespace UE::Mutable::Private
 				Ptr<ASTOpBoolNot> n = new ASTOpBoolNot;
 				n->A = TagCondition;
 
-				if (!c)
-				{
-					c = OpPool.Add(n);
-				}
-				else
-				{
-					Ptr<ASTOpBoolAnd> o = new ASTOpBoolAnd;
-					o->A = n;
-					o->B = c;
-					c = OpPool.Add(o);
-				}
+				Ptr<ASTOpBoolAnd> o = new ASTOpBoolAnd;
+				o->A = n;
+				o->B = c;
+				c = OpPool.Add(o);
 			}
 			else if (bIsConstant && bIsConstantValue == true)
 			{
@@ -398,9 +376,9 @@ namespace UE::Mutable::Private
 				Surface.NegativeTags,
 				Empty, Empty, Empty, Empty);
 
-				Ptr<ASTOpBoolAnd> ConditionOp = new ASTOpBoolAnd();
-				ConditionOp->A = Surface.ObjectCondition ? Surface.ObjectCondition : new ASTOpConstantBool(true);
-				ConditionOp->B = c ? c : new ASTOpConstantBool(true);
+			Ptr<ASTOpBoolAnd> ConditionOp = new ASTOpBoolAnd();
+			ConditionOp->A = Surface.ObjectCondition ? Surface.ObjectCondition : new ASTOpConstantBool(true);
+			ConditionOp->B = c;
 
 			Surface.FinalCondition = ConditionOp;
 			
@@ -431,9 +409,9 @@ namespace UE::Mutable::Private
 				FirstPass->Modifiers[ModifierIndex].PositiveTags,
 				FirstPass->Modifiers[ModifierIndex].NegativeTags,
 				Empty, Empty, Empty, Empty);
-			
-			Ptr<ASTOpBoolAnd> ConditionOp = new ASTOpBoolAnd;
-			ConditionOp->A = FirstPass->Modifiers[ModifierIndex].ObjectCondition;
+
+			Ptr<ASTOpBoolAnd> ConditionOp = new ASTOpBoolAnd();
+			ConditionOp->A = FirstPass->Modifiers[ModifierIndex].ObjectCondition ? FirstPass->Modifiers[ModifierIndex].ObjectCondition : new ASTOpConstantBool(true);
 			ConditionOp->B = c;
 
 			FirstPass->Modifiers[ModifierIndex].FinalCondition = ConditionOp;
