[Mutable] Bump Customizable Object version to fix DDC crash

#rnx

diff --git a/Engine/Plugins/Mutable/Source/CustomizableObject/Internal/MuCO/CustomizableObjectPrivate.h b/Engine/Plugins/Mutable/Source/CustomizableObject/Internal/MuCO/CustomizableObjectPrivate.h
--- a/Engine/Plugins/Mutable/Source/CustomizableObject/Internal/MuCO/CustomizableObjectPrivate.h
+++ b/Engine/Plugins/Mutable/Source/CustomizableObject/Internal/MuCO/CustomizableObjectPrivate.h
@@ -1736,6 +1736,8 @@ enum class ECustomizableObjectVersions : int32
 
 	FixDefaultBehaviourAndsWithoutOperants,
 
+	BumpVersionDDCCrash,
+
 	// -----<new versions can be added above this line>--------
 	LastCustomizableObjectVersion
 };
