[Mutable] Allow Texture Parameters compile both Passthrough and Non-Passthrouh pins at the same time.

IM_PARAMETER now always outputs a FImage in Passthrough Parameter mode.

diff --git a/Engine/Plugins/Mutable/Source/CustomizableObject/Internal/MuCO/CustomizableObjectPrivate.h b/Engine/Plugins/Mutable/Source/CustomizableObject/Internal/MuCO/CustomizableObjectPrivate.h
--- a/Engine/Plugins/Mutable/Source/CustomizableObject/Internal/MuCO/CustomizableObjectPrivate.h
+++ b/Engine/Plugins/Mutable/Source/CustomizableObject/Internal/MuCO/CustomizableObjectPrivate.h
@@ -1738,6 +1738,8 @@ enum class ECustomizableObjectVersions : int32
 
 	BumpVersionDDCCrash,
 
+	ImageParameterConvert,
+
 	// -----<new versions can be added above this line>--------
 	LastCustomizableObjectVersion
 };
diff --git a/Engine/Plugins/Mutable/Source/CustomizableObjectEditor/Private/MuCOE/GenerateMutableSource/GenerateMutableSourceImage.cpp b/Engine/Plugins/Mutable/Source/CustomizableObjectEditor/Private/MuCOE/GenerateMutableSource/GenerateMutableSourceImage.cpp
--- a/Engine/Plugins/Mutable/Source/CustomizableObjectEditor/Private/MuCOE/GenerateMutableSource/GenerateMutableSourceImage.cpp
+++ b/Engine/Plugins/Mutable/Source/CustomizableObjectEditor/Private/MuCOE/GenerateMutableSource/GenerateMutableSourceImage.cpp
@@ -63,6 +63,7 @@
 #include "MuT/NodeImageVariation.h"
 #include "MuT/NodeImageConstant.h"
 #include "MuT/NodeImageExternal.h"
+#include "MuT/NodeImageParameterConvert.h"
 #include "MuT/NodeScalarConstant.h"
 
 #define LOCTEXT_NAMESPACE "CustomizableObjectEditor"
@@ -165,52 +166,63 @@ UE::Mutable::Private::Ptr<UE::Mutable::Private::NodeImage> GenerateMutableSource
 	{
 		UE::Mutable::Private::Ptr<UE::Mutable::Private::NodeImageParameter> ImageParameterNode = GenerateMutableSourceImageParameter(Pin, GenerationContext);
 
-		// Force the same format that the default texture if any.
-		UE::Mutable::Private::Ptr<UE::Mutable::Private::NodeImageFormat> FormatNode = new UE::Mutable::Private::NodeImageFormat();
-		// \TODO: Take it from default?
-		//if (TypedNodeParam->DefaultValue)
-		//{
-		//	FormatNode->SetFormat(UE::Mutable::Private::EImageFormat::RGBA_UByte);
-		//}
-		//else
+		const bool bPassthrough = Pin == TypedNodeParam->PassthroughPin.Get();
+		if (bPassthrough)
 		{
-			// Force an "easy format" on the texture.
-			FormatNode->Format = UE::Mutable::Private::EImageFormat::RGBA_UByte;
+			Result = ImageParameterNode;
 		}
-		FormatNode->Source = ImageParameterNode;
+		else
+		{
+			UE::Mutable::Private::Ptr<UE::Mutable::Private::NodeImageParameterConvert> ImagePassthroughConvertNode = new UE::Mutable::Private::NodeImageParameterConvert();
+			ImagePassthroughConvertNode->ImageParameter = ImageParameterNode;
+			
+			// Force the same format that the default texture if any.
+			UE::Mutable::Private::Ptr<UE::Mutable::Private::NodeImageFormat> FormatNode = new UE::Mutable::Private::NodeImageFormat();
+			// \TODO: Take it from default?
+			//if (TypedNodeParam->DefaultValue)
+			//{
+			//	FormatNode->SetFormat(UE::Mutable::Private::EImageFormat::RGBA_UByte);
+			//}
+			//else
+			{
+				// Force an "easy format" on the texture.
+				FormatNode->Format = UE::Mutable::Private::EImageFormat::RGBA_UByte;
+			}
+			FormatNode->Source = ImagePassthroughConvertNode;
 
-		UE::Mutable::Private::Ptr<UE::Mutable::Private::NodeImageResize> ResizeNode = new UE::Mutable::Private::NodeImageResize();
-		ResizeNode->Base = FormatNode;
-		ResizeNode->bRelative = false;
+			UE::Mutable::Private::Ptr<UE::Mutable::Private::NodeImageResize> ResizeNode = new UE::Mutable::Private::NodeImageResize();
+			ResizeNode->Base = FormatNode;
+			ResizeNode->bRelative = false;
 
-		FUintVector2 TextureSize(TypedNodeParam->TextureSizeX, TypedNodeParam->TextureSizeY);
+			FUintVector2 TextureSize(TypedNodeParam->TextureSizeX, TypedNodeParam->TextureSizeY);
 
-		if (const UTexture2D* ReferenceTexture = TypedNodeParam->ReferenceValue)
-		{
-			const uint32 LODBias = ComputeLODBiasForTexture(GenerationContext, *TypedNodeParam->ReferenceValue, ReferenceTexture, ReferenceTextureSize);
-			TextureSize.X = FMath::Max(ReferenceTexture->Source.GetSizeX() >> LODBias, 1);
-			TextureSize.Y = FMath::Max(ReferenceTexture->Source.GetSizeY() >> LODBias, 1);
-		}
-		else
-		{
-			const int32 MaxNodeTextureSize = FMath::Max(TypedNodeParam->TextureSizeX, TypedNodeParam->TextureSizeY);
-			if (MaxNodeTextureSize <= 0)
+			if (const UTexture2D* ReferenceTexture = TypedNodeParam->ReferenceValue)
 			{
-				TextureSize.X = TextureSize.Y = 1;
-				GenerationContext.Log(LOCTEXT("TextureParameterSize0", "Texture size not specified. Add a reference texture or set a valid value to the Texture Size variables."), Node);
+				const uint32 LODBias = ComputeLODBiasForTexture(GenerationContext, *TypedNodeParam->ReferenceValue, ReferenceTexture, ReferenceTextureSize);
+				TextureSize.X = FMath::Max(ReferenceTexture->Source.GetSizeX() >> LODBias, 1);
+				TextureSize.Y = FMath::Max(ReferenceTexture->Source.GetSizeY() >> LODBias, 1);
 			}
-			else if (ReferenceTextureSize > 0 && ReferenceTextureSize < MaxNodeTextureSize)
+			else
 			{
-				const uint32 MipsToSkip = FMath::CeilLogTwo(MaxNodeTextureSize) - FMath::CeilLogTwo(ReferenceTextureSize);
-				TextureSize.X = FMath::Max(TextureSize.X >> MipsToSkip, (uint32)1);
-				TextureSize.Y = FMath::Max(TextureSize.Y >> MipsToSkip, (uint32)1);
+				const int32 MaxNodeTextureSize = FMath::Max(TypedNodeParam->TextureSizeX, TypedNodeParam->TextureSizeY);
+				if (MaxNodeTextureSize <= 0)
+				{
+					TextureSize.X = TextureSize.Y = 1;
+					GenerationContext.Log(LOCTEXT("TextureParameterSize0", "Texture size not specified. Add a reference texture or set a valid value to the Texture Size variables."), Node);
+				}
+				else if (ReferenceTextureSize > 0 && ReferenceTextureSize < MaxNodeTextureSize)
+				{
+					const uint32 MipsToSkip = FMath::CeilLogTwo(MaxNodeTextureSize) - FMath::CeilLogTwo(ReferenceTextureSize);
+					TextureSize.X = FMath::Max(TextureSize.X >> MipsToSkip, (uint32)1);
+					TextureSize.Y = FMath::Max(TextureSize.Y >> MipsToSkip, (uint32)1);
+				}
 			}
-		}
 
-		ResizeNode->SizeX = TextureSize.X;
-		ResizeNode->SizeY = TextureSize.Y;
+			ResizeNode->SizeX = TextureSize.X;
+			ResizeNode->SizeY = TextureSize.Y;
 
-		Result = ResizeNode;
+			Result = ResizeNode;
+		}
 	}
 
 	else if (const ICustomizableObjectNodeMeshInterface* TypedNodeMesh = Cast<ICustomizableObjectNodeMeshInterface>(Node))
@@ -1295,8 +1307,7 @@ UE::Mutable::Private::Ptr<UE::Mutable::Private::NodeImageParameter> GenerateMuta
 
 		TextureNode->Name = TypedNodeParam->GetParameterName(&GenerationContext.MacroNodesStack);
 		TextureNode->UID = GenerationContext.GetNodeIdUnique(ParameterNode).ToString();
-		TextureNode->bPassthrough = Pin == TypedNodeParam->PassthroughPin.Get();
-
+		
 		if (TypedNodeParam->DefaultValue)
 		{
 			GenerationContext.TextureParameterDefaultValues.Add(FName(TypedNodeParam->GetParameterName()), TypedNodeParam->DefaultValue);
diff --git a/Engine/Plugins/Mutable/Source/CustomizableObjectEditor/Private/MuCOE/SMutableCodeViewer.cpp b/Engine/Plugins/Mutable/Source/CustomizableObjectEditor/Private/MuCOE/SMutableCodeViewer.cpp
--- a/Engine/Plugins/Mutable/Source/CustomizableObjectEditor/Private/MuCOE/SMutableCodeViewer.cpp
+++ b/Engine/Plugins/Mutable/Source/CustomizableObjectEditor/Private/MuCOE/SMutableCodeViewer.cpp
@@ -3189,7 +3189,6 @@ void FMutableCodeTreeElement::GenerateLabelText()
 	case UE::Mutable::Private::EOpType::SC_PARAMETER:
 	case UE::Mutable::Private::EOpType::CO_PARAMETER:
 	case UE::Mutable::Private::EOpType::PR_PARAMETER:
-	case UE::Mutable::Private::EOpType::IM_PARAMETER:
 	case UE::Mutable::Private::EOpType::ST_PARAMETER:
 	{
 		UE::Mutable::Private::OP::ParameterArgs Args = Program.GetOpArgs<UE::Mutable::Private::OP::ParameterArgs>(MutableOperation);
diff --git a/Engine/Plugins/Mutable/Source/CustomizableObjectEditor/Private/MuCOE/SMutableGraphViewer.cpp b/Engine/Plugins/Mutable/Source/CustomizableObjectEditor/Private/MuCOE/SMutableGraphViewer.cpp
--- a/Engine/Plugins/Mutable/Source/CustomizableObjectEditor/Private/MuCOE/SMutableGraphViewer.cpp
+++ b/Engine/Plugins/Mutable/Source/CustomizableObjectEditor/Private/MuCOE/SMutableGraphViewer.cpp
@@ -132,6 +132,7 @@ static const char* MutableNodeNames[] =
 	"ImageTransform",
 	"ImageFromMaterialParameter",
 	"ImageExternal",
+	"ImageParameterConvert",
 
 	"Bool",
 	"BoolConstant",
diff --git a/Engine/Plugins/Mutable/Source/MutableRuntime/Internal/MuR/Operations.h b/Engine/Plugins/Mutable/Source/MutableRuntime/Internal/MuR/Operations.h
--- a/Engine/Plugins/Mutable/Source/MutableRuntime/Internal/MuR/Operations.h
+++ b/Engine/Plugins/Mutable/Source/MutableRuntime/Internal/MuR/Operations.h
@@ -59,8 +59,6 @@ namespace UE::Mutable::Private
 		IM_REFERENCE,
 		ME_REFERENCE,
 
-    	IM_PARAMETER_REFERENCE,
-
         //! Select one value or the other depending on a boolean input
         NU_CONDITIONAL,
         SC_CONDITIONAL,
@@ -218,6 +216,8 @@ namespace UE::Mutable::Private
 		//! Apply linear transform to Image content. Resulting samples outside the original image are tiled.
 		IM_TRANSFORM,
 
+    	IM_PARAMETER_CONVERT,
+
         //-----------------------------------------------------------------------------------------
         // Mesh operations
         //-----------------------------------------------------------------------------------------
@@ -1151,6 +1151,11 @@ namespace UE::Mutable::Private
     	{
     		int32 NumOperants;
     	};
+
+    	struct ImageParameterConvertArgs
+    	{
+    		ADDRESS ImageParameter;
+    	};
     };
 
 
diff --git a/Engine/Plugins/Mutable/Source/MutableRuntime/Private/MuR/CodeRunner.cpp b/Engine/Plugins/Mutable/Source/MutableRuntime/Private/MuR/CodeRunner.cpp
--- a/Engine/Plugins/Mutable/Source/MutableRuntime/Private/MuR/CodeRunner.cpp
+++ b/Engine/Plugins/Mutable/Source/MutableRuntime/Private/MuR/CodeRunner.cpp
@@ -4679,7 +4679,7 @@ namespace UE::Mutable::Private
 			break;
 		}
 
-		case EOpType::IM_PARAMETER_REFERENCE:
+		case EOpType::IM_PARAMETER:
         {
 			OP::ParameterArgs Args = Program.GetOpArgs<OP::ParameterArgs>(Item.At);
 			TSharedPtr<FRangeIndex> Index = BuildCurrentOpRangeIndex(Item, Args.variable);
@@ -5581,6 +5581,28 @@ namespace UE::Mutable::Private
 			break;
 		}
 
+		case EOpType::IM_PARAMETER_CONVERT:
+        {
+        	OP::ImageParameterConvertArgs Args = Program.GetOpArgs<OP::ImageParameterConvertArgs>(Item.At);
+        	switch (Item.Stage)
+        	{
+        	case 0:
+        		AddOp(FScheduledOp(Item.At, Item, 1), FScheduledOp(Args.ImageParameter, Item));
+        		break;
+
+        	case 1:
+        	{
+        		// This has been moved to a task. It should have been intercepted in IssueOp.
+        		check(false);
+        	}
+
+        	default:
+        		check(false);
+        	}
+
+	        break;
+        }
+			
         default:
             if (type!=EOpType::NONE)
             {
@@ -7027,16 +7049,33 @@ namespace UE::Mutable::Private
 			break;
 		}
 
-		case EOpType::IM_PARAMETER:
+		case EOpType::IM_PARAMETER_CONVERT:
 		{
-			MUTABLE_CPUPROFILER_SCOPE(RunCodeImageDesc_IM_PARAMETER);
-			check(Item.Stage == 0);
+			MUTABLE_CPUPROFILER_SCOPE(RunCodeImageDesc_IM_PARAMETER_CONVERT);
 
-			OP::ParameterArgs Args = Program.GetOpArgs<OP::ParameterArgs>(Item.At);
-			UTexture* Image = pParams->GetImageValue(Args.variable);
+			OP::ImageParameterConvertArgs Args = Program.GetOpArgs<OP::ImageParameterConvertArgs>(Item.At);
+			switch (Item.Stage)
+			{
+			case 0:
+				{
+					AddOp(FScheduledOp(Item.At, Item, 1), 
+						FScheduledOp(Args.ImageParameter, Item));
+					
+					break;
+				}
 
-			FExtendedImageDesc& Result = ImageDescResults.FindOrAdd(Item.At);
-			Result = GetExternalImageDesc(Image);
+			case 1:
+				{
+					FExtendedImageDesc& Result = ImageDescResults.FindOrAdd(Item.At);
+					Result = ImageDescResults.FindChecked(Args.ImageParameter);
+
+					StoreValidDesc(Item); 
+					break;
+				}
+
+			default:
+				unimplemented();
+			}
 
 			StoreValidDesc(Item);
 			break;
@@ -7056,9 +7095,18 @@ namespace UE::Mutable::Private
 			break;
 		}
 
-		case EOpType::IM_PARAMETER_REFERENCE:
+		case EOpType::IM_PARAMETER:
 		{
-			check(false); // Passthrough. No image descriptor required.
+			MUTABLE_CPUPROFILER_SCOPE(RunCodeImageDesc_IM_PARAMETER);
+			check(Item.Stage == 0);
+
+			OP::ParameterArgs Args = Program.GetOpArgs<OP::ParameterArgs>(Item.At);
+			UTexture* Image = pParams->GetImageValue(Args.variable);
+
+			FExtendedImageDesc& Result = ImageDescResults.FindOrAdd(Item.At);
+			Result = GetExternalImageDesc(Image);
+
+			StoreValidDesc(Item);
 			break;
 		}
 
diff --git a/Engine/Plugins/Mutable/Source/MutableRuntime/Private/MuR/CodeRunnerBegin.cpp b/Engine/Plugins/Mutable/Source/MutableRuntime/Private/MuR/CodeRunnerBegin.cpp
--- a/Engine/Plugins/Mutable/Source/MutableRuntime/Private/MuR/CodeRunnerBegin.cpp
+++ b/Engine/Plugins/Mutable/Source/MutableRuntime/Private/MuR/CodeRunnerBegin.cpp
@@ -392,7 +392,6 @@ namespace UE::Mutable::Private
 			case EOpType::IM_PARAMETER:
 			case EOpType::MI_PARAMETER:
 			case EOpType::IM_PARAMETER_FROM_MATERIAL:
-			case EOpType::IM_PARAMETER_REFERENCE:
 				{
 					StoreNone(Item);
 					break;
@@ -2290,7 +2289,16 @@ namespace UE::Mutable::Private
 					StoreNone(Item);
 					break;
 				}
-	
+
+			case EOpType::IM_PARAMETER_CONVERT:
+				{
+					OP::ImageParameterConvertArgs Args = Program.GetOpArgs<OP::ImageParameterConvertArgs>(Item.At);
+					PushOp(FScheduledOpInline(Args.ImageParameter, Item));
+
+					StoreNone(Item);
+					break;
+				}
+				
 			case EOpType::LA_MERGE:
 				{
 					OP::LayoutMergeArgs Args = Program.GetOpArgs<OP::LayoutMergeArgs>(Item.At);
diff --git a/Engine/Plugins/Mutable/Source/MutableRuntime/Private/MuR/CodeRunnerTasked.cpp b/Engine/Plugins/Mutable/Source/MutableRuntime/Private/MuR/CodeRunnerTasked.cpp
--- a/Engine/Plugins/Mutable/Source/MutableRuntime/Private/MuR/CodeRunnerTasked.cpp
+++ b/Engine/Plugins/Mutable/Source/MutableRuntime/Private/MuR/CodeRunnerTasked.cpp
@@ -2591,21 +2591,25 @@ namespace UE::Mutable::Private
 			break;
 		}		
 
-		case EOpType::IM_PARAMETER:
+		case EOpType::IM_PARAMETER_CONVERT:
 		{
-			OP::ParameterArgs args = Program.GetOpArgs<OP::ParameterArgs>(Item.At);
-			TSharedPtr<FRangeIndex> Index = BuildCurrentOpRangeIndex(Item, args.variable);
-
-			UTexture* Image = Params->GetImageValue(args.variable, Index.Get());
+			if (Item.Stage == 1)
+			{
+				OP::ImageParameterConvertArgs Args = Program.GetOpArgs<OP::ImageParameterConvertArgs>(Item.At);
+			
+				TSharedPtr<const FImage> Image = LoadImage(FCacheAddress(Args.ImageParameter, Item));
+				
+				check(ImageLOD < TNumericLimits<uint8>::Max() && ImageLOD >= 0);
+				check(ImageLOD + static_cast<int32>(Item.ExecutionOptions) < TNumericLimits<uint8>::Max());
 
-			check(ImageLOD < TNumericLimits<uint8>::Max() && ImageLOD >= 0);
-			check(ImageLOD + static_cast<int32>(Item.ExecutionOptions) < TNumericLimits<uint8>::Max());
+				const uint8 MipmapsToSkip = Item.ExecutionOptions + static_cast<uint8>(ImageLOD);
 
-			const uint8 MipmapsToSkip = Item.ExecutionOptions + static_cast<uint8>(ImageLOD);
+				FExternalResourceId FullId;
+				FullId.ImageParameter = Image->PassthroughObject.Get();
+				Issued = MakeShared<FImageExternalLoadTask>(Item, MipmapsToSkip, FullId);
 
-			CodeRunner::FExternalResourceId FullId;
-			FullId.ImageParameter = Image;
-			Issued = MakeShared<FImageExternalLoadTask>(Item, MipmapsToSkip, FullId);
+				Release(Image);
+			}
 
 			break;
 		}
diff --git a/Engine/Plugins/Mutable/Source/MutableRuntime/Private/MuR/Operations.cpp b/Engine/Plugins/Mutable/Source/MutableRuntime/Private/MuR/Operations.cpp
--- a/Engine/Plugins/Mutable/Source/MutableRuntime/Private/MuR/Operations.cpp
+++ b/Engine/Plugins/Mutable/Source/MutableRuntime/Private/MuR/Operations.cpp
@@ -21,54 +21,52 @@ namespace UE::Mutable::Private
 
 		{ EDataType::Bool,			},	// BO_CONSTANT
 		{ EDataType::Int,			},	// NU_CONSTANT
-		{ EDataType::Scalar,		},	// SC_CONSTANT
+		{ EDataType::Scalar,			},	// SC_CONSTANT
 		{ EDataType::Color,			},	// CO_CONSTANT
 		{ EDataType::Image,			},	// IM_CONSTANT
 		{ EDataType::Mesh,			},	// ME_CONSTANT
-		{ EDataType::Layout,		},	// LA_CONSTANT
+		{ EDataType::Layout,			},	// LA_CONSTANT
 		{ EDataType::Projector,		},	// PR_CONSTANT
-		{ EDataType::String,		},	// ST_CONSTANT
+		{ EDataType::String,			},	// ST_CONSTANT
 		{ EDataType::ExtensionData,	},	// ED_CONSTANT
-		{ EDataType::Matrix,		},	// MA_CONSTANT
+		{ EDataType::Matrix,			},	// MA_CONSTANT
 		{ EDataType::Material,		},	// MI_CONSTANT
 
 		{ EDataType::Bool,			},	// BO_PARAMETER
 		{ EDataType::Int,			},	// NU_PARAMETER
-		{ EDataType::Scalar,		},	// SC_PARAMETER
+		{ EDataType::Scalar,			},	// SC_PARAMETER
 		{ EDataType::Color,			},	// CO_PARAMETER
 		{ EDataType::Projector,		},	// PR_PARAMETER
 		{ EDataType::Image,			},	// IM_PARAMETER
 		{ EDataType::Mesh,			},	// ME_PARAMETER
-		{ EDataType::String,		},	// ST_PARAMETER
-		{ EDataType::Matrix,		},	// MA_PARAMETER
+		{ EDataType::String,			},	// ST_PARAMETER
+		{ EDataType::Matrix,			},	// MA_PARAMETER
 		{ EDataType::Material,		},	// MI_PARAMETER
 
 		{ EDataType::Image,			},	// IM_REFERENCE
 		{ EDataType::Mesh,			},	// ME_REFERENCE
 
-		{ EDataType::Image,			},	// IM_PARAMETER_REFERENCE
-    	
 		{ EDataType::Int,			},	// NU_CONDITIONAL
-		{ EDataType::Scalar,		},	// SC_CONDITIONAL
+		{ EDataType::Scalar,			},	// SC_CONDITIONAL
 		{ EDataType::Color,			},	// CO_CONDITIONAL
 		{ EDataType::Image,			},	// IM_CONDITIONAL
 		{ EDataType::Mesh,			},	// ME_CONDITIONAL
-		{ EDataType::Layout,		},	// LA_CONDITIONAL
+		{ EDataType::Layout,			},	// LA_CONDITIONAL
 		{ EDataType::Instance,		},	// IN_CONDITIONAL
 		{ EDataType::ExtensionData,	},	// ED_CONDITIONAL
 		{ EDataType::Material,		},	// MI_CONDITIONAL
 
 		{ EDataType::Int,			},	// NU_SWITCH
-		{ EDataType::Scalar,		},	// SC_SWITCH
+		{ EDataType::Scalar,			},	// SC_SWITCH
 		{ EDataType::Color,			},	// CO_SWITCH
 		{ EDataType::Image,			},	// IM_SWITCH
 		{ EDataType::Mesh,			},	// ME_SWITCH
-		{ EDataType::Layout,		},	// LA_SWITCH
+		{ EDataType::Layout,			},	// LA_SWITCH
 		{ EDataType::Instance,		},	// IN_SWITCH
 		{ EDataType::ExtensionData,	},	// ED_SWITCH
 		{ EDataType::Material,		},	// MI_SWITCH
 
-		{ EDataType::Scalar,		},	// SC_MATERIAL_BREAK
+		{ EDataType::Scalar,			},	// SC_MATERIAL_BREAK
 		{ EDataType::Color,			},	// CO_MATERIAL_BREAK
 		{ EDataType::Image,			},	// IM_MATERIAL_BREAK
 		{ EDataType::Image,			},	// IM_PARAMETER_FROM_MATERIAL
@@ -78,8 +76,8 @@ namespace UE::Mutable::Private
 		{ EDataType::Bool,			},	// BO_OR
 		{ EDataType::Bool,			},	// BO_NOT
 
-		{ EDataType::Scalar,		},	// SC_ARITHMETIC
-		{ EDataType::Scalar,		},	// SC_CURVE
+		{ EDataType::Scalar,			},	// SC_ARITHMETIC
+		{ EDataType::Scalar,			},	// SC_CURVE
 
 		{ EDataType::Color,			},	// CO_SAMPLEIMAGE
 		{ EDataType::Color,			},	// CO_SWIZZLE
@@ -112,6 +110,7 @@ namespace UE::Mutable::Private
 		{ EDataType::Image,			},	// IM_INVERT
 		{ EDataType::Image,			},	// IM_NORMALCOMPOSITE
 		{ EDataType::Image,			},	// IM_TRANSFORM
+		{ EDataType::Image,			},	// IM_PARAMETER_CONVERT
 
 		{ EDataType::Mesh,			},	// ME_APPLYLAYOUT
 		{ EDataType::Mesh,			},	// ME_PREPARELAYOUT
@@ -200,7 +199,6 @@ namespace UE::Mutable::Private
 		case EOpType::IM_REFERENCE:
 		case EOpType::ME_REFERENCE:
 		case EOpType::IM_PARAMETER_FROM_MATERIAL:
-		case EOpType::IM_PARAMETER_REFERENCE:
 			break;
 
         case EOpType::SC_CURVE:
@@ -590,6 +588,13 @@ namespace UE::Mutable::Private
 			break;
 		}
 
+        case EOpType::IM_PARAMETER_CONVERT:
+        {
+        	OP::ImageParameterConvertArgs Args = program.GetOpArgs<OP::ImageParameterConvertArgs>(at);
+        	f(Args.ImageParameter);
+        	break;
+        }
+
         //-------------------------------------------------------------------------------------
         case EOpType::ME_APPLYLAYOUT:
         {
diff --git a/Engine/Plugins/Mutable/Source/MutableRuntime/Private/MuR/System.cpp b/Engine/Plugins/Mutable/Source/MutableRuntime/Private/MuR/System.cpp
--- a/Engine/Plugins/Mutable/Source/MutableRuntime/Private/MuR/System.cpp
+++ b/Engine/Plugins/Mutable/Source/MutableRuntime/Private/MuR/System.cpp
@@ -944,7 +944,6 @@ namespace UE::Mutable::Private
 			case EOpType::IM_PARAMETER:
 			case EOpType::MA_PARAMETER:
 			case EOpType::MI_PARAMETER:
-            case EOpType::IM_PARAMETER_REFERENCE:
             {
 				OP::ParameterArgs args = Program.GetOpArgs<OP::ParameterArgs>(at);
 				OP::ADDRESS ParamIndex = args.variable;
diff --git a/Engine/Plugins/Mutable/Source/MutableTools/Internal/MuT/Node.h b/Engine/Plugins/Mutable/Source/MutableTools/Internal/MuT/Node.h
--- a/Engine/Plugins/Mutable/Source/MutableTools/Internal/MuT/Node.h
+++ b/Engine/Plugins/Mutable/Source/MutableTools/Internal/MuT/Node.h
@@ -78,6 +78,7 @@ namespace UE::Mutable::Private
 			ImageTransform,
 			ImageFromMaterialParameter,
 			ImageExternal,
+			ImageParameterConvert,
 
 			Bool,
 			BoolConstant,
diff --git a/Engine/Plugins/Mutable/Source/MutableTools/Internal/MuT/NodeImageParameter.h b/Engine/Plugins/Mutable/Source/MutableTools/Internal/MuT/NodeImageParameter.h
--- a/Engine/Plugins/Mutable/Source/MutableTools/Internal/MuT/NodeImageParameter.h
+++ b/Engine/Plugins/Mutable/Source/MutableTools/Internal/MuT/NodeImageParameter.h
@@ -3,7 +3,6 @@
 #pragma once
 
 #include "MuR/Ptr.h"
-#include "MuR/RefCounted.h"
 #include "MuT/Node.h"
 #include "MuT/NodeImage.h"
 #include "MuT/NodeRange.h"
@@ -26,26 +25,17 @@ namespace UE::Mutable::Private
 		/** Ranges for the parameter in case it is multidimensional. */
 		TArray<Ptr<NodeRange>> Ranges;
 
-    	bool bPassthrough = false;
-
-	public:
-
 		// Node Interface
 		virtual const FNodeType* GetType() const override { return &StaticType; }
 		static const FNodeType* GetStaticType() { return &StaticType; }
 
 	protected:
-
 		/** Forbidden. Manage with the Ptr<> template. */
-		~NodeImageParameter() {}
+		virtual ~NodeImageParameter() override {}
 
 	private:
-
 		static UE_API FNodeType StaticType;
-
 	};
-
-
 }
 
 #undef UE_API
diff --git a/Engine/Plugins/Mutable/Source/MutableTools/Internal/MuT/NodeImageParameterConvert.h b/Engine/Plugins/Mutable/Source/MutableTools/Internal/MuT/NodeImageParameterConvert.h
new file mode 100644
--- /dev/null
+++ b/Engine/Plugins/Mutable/Source/MutableTools/Internal/MuT/NodeImageParameterConvert.h
@@ -0,0 +1,32 @@
+// Copyright Epic Games, Inc. All Rights Reserved.
+
+#pragma once
+
+#include "MuR/Ptr.h"
+#include "MuT/Node.h"
+#include "MuT/NodeImage.h"
+
+#define UE_API MUTABLETOOLS_API
+
+
+namespace UE::Mutable::Private
+{
+	class NodeImageParameterConvert : public NodeImage
+	{
+	public:
+		Ptr<NodeImage> ImageParameter;
+		
+		// Node Interface
+		virtual const FNodeType* GetType() const override { return &StaticType; }
+		static const FNodeType* GetStaticType() { return &StaticType; }
+
+	protected:
+		virtual ~NodeImageParameterConvert() override {}
+
+	private:
+		static UE_API FNodeType StaticType;
+	};
+}
+
+
+#undef UE_API
diff --git a/Engine/Plugins/Mutable/Source/MutableTools/Private/MuT/ASTOpImageParameterConvert.cpp b/Engine/Plugins/Mutable/Source/MutableTools/Private/MuT/ASTOpImageParameterConvert.cpp
new file mode 100644
--- /dev/null
+++ b/Engine/Plugins/Mutable/Source/MutableTools/Private/MuT/ASTOpImageParameterConvert.cpp
@@ -0,0 +1,91 @@
+// Copyright Epic Games, Inc. All Rights Reserved.
+
+#include "MuT/ASTOpImageParameterConvert.h"
+
+#include "MuT/ASTOpBoolEqualIntConst.h"
+#include "MuR/ModelPrivate.h"
+
+
+namespace UE::Mutable::Private
+{
+	ASTOpImageParameterConvert::ASTOpImageParameterConvert() :
+		ImageParameter(this)
+	{
+	}
+
+	ASTOpImageParameterConvert::~ASTOpImageParameterConvert()
+	{
+		// Explicit call needed to avoid recursive destruction
+		RemoveChildren();
+	}
+
+
+	void ASTOpImageParameterConvert::ForEachChild(const TFunctionRef<void(ASTChild&)> Func)
+	{
+		Func(ImageParameter);
+	}
+
+
+	bool ASTOpImageParameterConvert::IsEqual(const ASTOp& OtherUntyped) const
+	{
+		if (OtherUntyped.GetOpType() == GetOpType())
+		{
+			const ASTOpImageParameterConvert* Other = static_cast<const ASTOpImageParameterConvert*>(&OtherUntyped);
+			return Type == Other->Type &&
+				ImageParameter == Other->ImageParameter;
+		}
+		
+		return false;
+	}
+
+
+	Ptr<ASTOp> ASTOpImageParameterConvert::Clone(MapChildFuncRef MapChild) const
+	{
+		Ptr<ASTOpImageParameterConvert> New = new ASTOpImageParameterConvert();
+		New->Type = Type;
+		New->ImageParameter =  ASTChild(New, MapChild(ImageParameter.child()));
+
+		return New;
+	}
+
+
+	EOpType ASTOpImageParameterConvert::GetOpType() const 
+	{ 
+		return Type; 
+	}
+
+
+	uint64 ASTOpImageParameterConvert::Hash() const
+	{
+		uint64 res = std::hash<uint64>()(uint64(Type));
+
+		hash_combine(res, ImageParameter.child().get());
+
+		return res;
+	}
+
+
+	void ASTOpImageParameterConvert::Link(FProgram& Program, FLinkerOptions*)
+	{
+		if (!linkedAddress)
+		{
+			linkedAddress = static_cast<OP::ADDRESS>(Program.OpAddress.Num());
+			Program.OpAddress.Add(static_cast<uint32>(Program.ByteCode.Num()));
+			AppendCode(Program.ByteCode, Type);
+
+			OP::ImageParameterConvertArgs Args;
+			FMemory::Memzero(Args);
+
+			check(ImageParameter->linkedAddress);
+			Args.ImageParameter = ImageParameter->linkedAddress;
+		
+			AppendCode(Program.ByteCode, Args);
+		}
+	}
+
+	
+	FImageDesc ASTOpImageParameterConvert::GetImageDesc(bool, FGetImageDescContext*) const
+	{
+		return ImageParameter->GetImageDesc();
+	}
+}
diff --git a/Engine/Plugins/Mutable/Source/MutableTools/Private/MuT/ASTOpImageParameterConvert.h b/Engine/Plugins/Mutable/Source/MutableTools/Private/MuT/ASTOpImageParameterConvert.h
new file mode 100644
--- /dev/null
+++ b/Engine/Plugins/Mutable/Source/MutableTools/Private/MuT/ASTOpImageParameterConvert.h
@@ -0,0 +1,39 @@
+// Copyright Epic Games, Inc. All Rights Reserved.
+
+#pragma once
+
+#include "HAL/Platform.h"
+#include "MuR/Image.h"
+#include "MuR/Operations.h"
+#include "MuR/ParametersPrivate.h"
+#include "MuR/Ptr.h"
+#include "MuT/AST.h"
+
+
+namespace UE::Mutable::Private
+{
+	struct FProgram;
+	
+	class ASTOpImageParameterConvert final : public ASTOp
+	{
+	public:
+		EOpType Type = EOpType::IM_PARAMETER_CONVERT;
+
+		ASTChild ImageParameter;
+		
+		int32 LinkedParameterIndex = -1;
+
+		ASTOpImageParameterConvert();
+		virtual ~ASTOpImageParameterConvert() override;
+
+		// ASTOp interface.
+		virtual EOpType GetOpType() const override;
+		virtual uint64 Hash() const override;
+		virtual void ForEachChild(const TFunctionRef<void(ASTChild&)>) override;
+		virtual bool IsEqual(const ASTOp&) const override;
+		virtual Ptr<ASTOp> Clone(MapChildFuncRef ) const override;
+		virtual void Link(FProgram&, FLinkerOptions*) override;
+		virtual FImageDesc GetImageDesc(bool, FGetImageDescContext*) const override;
+	};
+}
+
diff --git a/Engine/Plugins/Mutable/Source/MutableTools/Private/MuT/CodeGenerator.h b/Engine/Plugins/Mutable/Source/MutableTools/Private/MuT/CodeGenerator.h
--- a/Engine/Plugins/Mutable/Source/MutableTools/Private/MuT/CodeGenerator.h
+++ b/Engine/Plugins/Mutable/Source/MutableTools/Private/MuT/CodeGenerator.h
@@ -46,6 +46,7 @@
 
 namespace UE::Mutable::Private
 {
+	class NodeImageParameterConvert;
 	class NodeMaterialExternal;
 	class NodeScalarExternal;
 }
@@ -552,6 +553,7 @@ namespace UE::Mutable::Private
 		void GenerateImage_MaterialBreak(const FImageGenerationOptions&, FImageGenerationResult&, const NodeImageMaterialBreak*);
 		void GenerateImage_FromMaterialParameter(const FImageGenerationOptions&, FImageGenerationResult&, const NodeImageFromMaterialParameter*);
 		void GenerateImage_External(const FImageGenerationOptions&, FImageGenerationResult&, const NodeImageExternal*);
+		void GenerateImage_ParameterConvert(const FImageGenerationOptions&, FImageGenerationResult&, const NodeImageParameterConvert*);
 
     	void GenerateCrop(const FImageGenerationOptions& Options, FImageGenerationResult& Result, const NodeImage&);
 
diff --git a/Engine/Plugins/Mutable/Source/MutableTools/Private/MuT/CodeGenerator_Image.cpp b/Engine/Plugins/Mutable/Source/MutableTools/Private/MuT/CodeGenerator_Image.cpp
--- a/Engine/Plugins/Mutable/Source/MutableTools/Private/MuT/CodeGenerator_Image.cpp
+++ b/Engine/Plugins/Mutable/Source/MutableTools/Private/MuT/CodeGenerator_Image.cpp
@@ -1,5 +1,6 @@
 // Copyright Epic Games, Inc. All Rights Reserved.
 
+#include "ASTOpImageParameterConvert.h"
 #include "HAL/PlatformMath.h"
 #include "Misc/AssertionMacros.h"
 #include "MuR/Image.h"
@@ -73,6 +74,7 @@
 #include "MuT/NodeImageMultiLayer.h"
 #include "MuT/NodeImageNormalComposite.h"
 #include "MuT/NodeImageParameter.h"
+#include "MuT/NodeImageParameterConvert.h"
 #include "MuT/NodeImagePlainColour.h"
 #include "MuT/NodeImageProject.h"
 #include "MuT/NodeImageResize.h"
@@ -142,6 +144,7 @@ namespace UE::Mutable::Private
 		case Node::EType::ImageMaterialBreak: GenerateImage_MaterialBreak(Options, Result, static_cast<const NodeImageMaterialBreak*>(Node)); break;
 		case Node::EType::ImageFromMaterialParameter: GenerateImage_FromMaterialParameter(Options, Result, static_cast<const NodeImageFromMaterialParameter*>(Node)); break;
 		case Node::EType::ImageExternal: GenerateImage_External(Options, Result, static_cast<const NodeImageExternal*>(Node)); break;
+		case Node::EType::ImageParameterConvert: GenerateImage_ParameterConvert(Options, Result, static_cast<const NodeImageParameterConvert*>(Node)); break;
 		default: check(false);
 		}
 
@@ -277,7 +280,7 @@ namespace UE::Mutable::Private
 			if (!Found)
 			{
 				op = new ASTOpParameter();
-				op->Type = Node->bPassthrough ? EOpType::IM_PARAMETER_REFERENCE : EOpType::IM_PARAMETER;
+				op->Type = EOpType::IM_PARAMETER;
 
 				op->Parameter.Name = Node->Name;
 				bool bParseOk = FGuid::Parse(Node->UID, op->Parameter.UID);
@@ -1934,6 +1937,20 @@ namespace UE::Mutable::Private
 		Result.op = ExtensionResult.Op;
 	}
 
+	
+	void CodeGenerator::GenerateImage_ParameterConvert(const FImageGenerationOptions& InOptions, FImageGenerationResult& Result, const NodeImageParameterConvert* InNode)
+	{
+		Ptr<ASTOpImageParameterConvert> ReferenceConvert = new ASTOpImageParameterConvert();
+
+		FImageGenerationResult PassthroughImageResult;
+		GenerateImage(InOptions, PassthroughImageResult, InNode->ImageParameter);
+
+		ReferenceConvert->ImageParameter = PassthroughImageResult.op;
+		
+		Result.op = ReferenceConvert;
+		Result.Parameter = PassthroughImageResult.Parameter;
+	}
+
 
 	TSharedPtr<FImage> CodeGenerator::GenerateMissingImage(EImageFormat Format)
 	{
diff --git a/Engine/Plugins/Mutable/Source/MutableTools/Private/MuT/CodeOptimiser.cpp b/Engine/Plugins/Mutable/Source/MutableTools/Private/MuT/CodeOptimiser.cpp
--- a/Engine/Plugins/Mutable/Source/MutableTools/Private/MuT/CodeOptimiser.cpp
+++ b/Engine/Plugins/Mutable/Source/MutableTools/Private/MuT/CodeOptimiser.cpp
@@ -552,7 +552,6 @@ namespace UE::Mutable::Private
 					case EOpType::MA_PARAMETER:
 					case EOpType::MI_PARAMETER:
 					case EOpType::IM_PARAMETER_FROM_MATERIAL:
-					case EOpType::IM_PARAMETER_REFERENCE:
 						bIsConstantSubgraph = false;
 						break;
 					default:
diff --git a/Engine/Plugins/Mutable/Source/MutableTools/Private/MuT/CodeOptimiser_Parameters.cpp b/Engine/Plugins/Mutable/Source/MutableTools/Private/MuT/CodeOptimiser_Parameters.cpp
--- a/Engine/Plugins/Mutable/Source/MutableTools/Private/MuT/CodeOptimiser_Parameters.cpp
+++ b/Engine/Plugins/Mutable/Source/MutableTools/Private/MuT/CodeOptimiser_Parameters.cpp
@@ -1582,7 +1582,6 @@ namespace UE::Mutable::Private
 				case EOpType::IM_PARAMETER:
 				case EOpType::ME_PARAMETER:
 				case EOpType::MA_PARAMETER:
-				case EOpType::IM_PARAMETER_REFERENCE:
 				{
 					const ASTOpParameter* TypedOp = static_cast<const ASTOpParameter*>(At.get());
                     Parameters.Add(TypedOp->Parameter.Name);
diff --git a/Engine/Plugins/Mutable/Source/MutableTools/Private/MuT/Compiler.cpp b/Engine/Plugins/Mutable/Source/MutableTools/Private/MuT/Compiler.cpp
--- a/Engine/Plugins/Mutable/Source/MutableTools/Private/MuT/Compiler.cpp
+++ b/Engine/Plugins/Mutable/Source/MutableTools/Private/MuT/Compiler.cpp
@@ -83,8 +83,6 @@ namespace UE::Mutable::Private
 		{ true,		{ 0, 0, 0, 0, 0, 0, 0, 0, 0, 0, 0, 0, 0, 0, 0, 0, 0, 0, 0, 0, 0, 0, 0, 0, 0, 0, 0, 0, 0 }		},	// IM_REFERENCE
 		{ true,		{ 0, 0, 0, 0, 0, 0, 0, 0, 0, 0, 0, 0, 0, 0, 0, 0, 0, 0, 0, 0, 0, 0, 0, 0, 0, 0, 0, 0, 0 }		},	// ME_REFERENCE
 
-    	{ false,	{ 0, 0, 0, 0, 0, 0, 0, 0, 0, 0, 0, 0, 0, 0, 0, 0, 0, 0, 0, 0, 0, 0, 0, 0, 0, 0, 0, 0, 0 }		},	// IM_PARAMETER_REFERENCE
-
 		{ false,	{ 0, 0, 0, 0, 0, 0, 0, 0, 0, 0, 0, 0, 0, 0, 0, 0, 0, 0, 0, 0, 0, 0, 0, 0, 0, 0, 0, 0, 0 }		},	// NU_CONDITIONAL
 		{ false,	{ 0, 0, 0, 0, 0, 0, 0, 0, 0, 0, 0, 0, 0, 0, 0, 0, 0, 0, 0, 0, 0, 0, 0, 0, 0, 0, 0, 0, 0 }		},	// SC_CONDITIONAL
 		{ false,	{ 0, 0, 0, 0, 0, 0, 0, 0, 0, 0, 0, 0, 0, 0, 0, 0, 0, 0, 0, 0, 0, 0, 0, 0, 0, 0, 0, 0, 0 }		},	// CO_CONDITIONAL
@@ -149,6 +147,7 @@ namespace UE::Mutable::Private
 		{ true,		{ 0, 1, 1, 1, 0, 0, 0, 0, 0, 0, 0, 0, 0, 0, 0, 0, 0, 0, 0, 0, 0, 0, 0, 0, 0, 0, 0, 0, 0 }		},	// IM_INVERT
 		{ true,		{ 0, 1, 1, 1, 0, 0, 0, 0, 0, 0, 0, 0, 0, 0, 0, 0, 0, 0, 0, 0, 0, 0, 0, 0, 0, 0, 0, 0, 0 }		},	// IM_NORMALCOMPOSITE
 		{ true,		{ 0, 1, 1, 1, 0, 0, 0, 0, 0, 0, 0, 0, 0, 0, 0, 0, 0, 0, 0, 0, 0, 0, 0, 0, 0, 0, 0, 0, 0 }		},	// IM_TRANSFORM
+		{ true,		{ 0, 0, 0, 0, 0, 0, 0, 0, 0, 0, 0, 0, 0, 0, 0, 0, 0, 0, 0, 0, 0, 0, 0, 0, 0, 0, 0, 0, 0 }		},	// IM_PARAMETER_CONVERT
 
 		{ true,		{ 0, 0, 0, 0, 0, 0, 0, 0, 0, 0, 0, 0, 0, 0, 0, 0, 0, 0, 0, 0, 0, 0, 0, 0, 0, 0, 0, 0, 0 }		},	// ME_APPLYLAYOUT
 		{ true,		{ 0, 0, 0, 0, 0, 0, 0, 0, 0, 0, 0, 0, 0, 0, 0, 0, 0, 0, 0, 0, 0, 0, 0, 0, 0, 0, 0, 0, 0 }		},	// ME_PREPARELAYOUT
diff --git a/Engine/Plugins/Mutable/Source/MutableTools/Private/MuT/ErrorLog.cpp b/Engine/Plugins/Mutable/Source/MutableTools/Private/MuT/ErrorLog.cpp
--- a/Engine/Plugins/Mutable/Source/MutableTools/Private/MuT/ErrorLog.cpp
+++ b/Engine/Plugins/Mutable/Source/MutableTools/Private/MuT/ErrorLog.cpp
@@ -227,8 +227,6 @@ namespace UE::Mutable::Private
 		
 		TEXT("IM_REFERENCE     "),
 		TEXT("ME_REFERENCE     "),
-
-    	TEXT("IM_PARAMETER_REFERENCE"),
     	
 		TEXT("NU_CONDITIONAL   "),
 		TEXT("SC_CONDITIONAL   "),
@@ -294,6 +292,7 @@ namespace UE::Mutable::Private
 		TEXT("IM_INVERT        "),
 		TEXT("IM_NORMAL_COMPO  "),
 		TEXT("IM_TRANSFORM     "),
+		TEXT("IM_PARAMETER_CONVERT"),
 
 		TEXT("ME_APPLYLAYOUT   "),
 		TEXT("ME_PREPARELAYOUT "),
diff --git a/Engine/Plugins/Mutable/Source/MutableTools/Private/MuT/Node.cpp b/Engine/Plugins/Mutable/Source/MutableTools/Private/MuT/Node.cpp
--- a/Engine/Plugins/Mutable/Source/MutableTools/Private/MuT/Node.cpp
+++ b/Engine/Plugins/Mutable/Source/MutableTools/Private/MuT/Node.cpp
@@ -102,6 +102,7 @@
 #include "MuT/NodeRangeFromScalar.h"
 #include "MuT/NodeMaterial.h"
 #include "MuT/NodeImageMaterialBreak.h"
+#include "MuT/NodeImageParameterConvert.h"
 #include "MuT/NodeMaterialConstant.h"
 #include "MuT/NodeMaterialExternal.h"
 #include "MuT/NodeMaterialTable.h"
@@ -182,6 +183,7 @@ namespace UE::Mutable::Private
 	FNodeType NodeImageColourMap::StaticType = FNodeType(Node::EType::ImageColorMap, NodeImage::GetStaticType());
 	FNodeType NodeImageFromMaterialParameter::StaticType = FNodeType(Node::EType::ImageFromMaterialParameter, NodeImage::GetStaticType());
 	FNodeType NodeImageExternal::StaticType = FNodeType(Node::EType::ImageExternal, NodeImage::GetStaticType());
+	FNodeType NodeImageParameterConvert::StaticType = FNodeType(Node::EType::ImageParameterConvert, NodeImage::GetStaticType());
 
 	FNodeType NodeColour::StaticType = FNodeType(Node::EType::Color, Node::GetStaticType());
 	FNodeType NodeColourConstant::StaticType = FNodeType(Node::EType::ColorConstant, NodeColour::GetStaticType());
