[Mutable] InstancedStruct Parameter.

- Removed UN (Unknown) type. Now all UN_EXTERNAL are typed.
- FInstancedStructs now are properly hashed (FBitWriterResourceKey).
- Forbid External Types to have UObject.
- Renamed UCONodeExternal to UCONodeExternalOperation.
- Fixed pins issue with UCONodeExternalOperation.
- Added IsLoaded to UCustomizableObjectNode. If not loaded and error will be shown on the node.
- Renamed CO and COI Color Parameter to Vector Parameter.
- FParameter UObject/FInstancedStruct defaults now are initialized on SetModel.
- Removed InstanceDescriptro docs. No longer public, docs already on COI.
- Fixed issue where Material, Mesh and Texture parameter defaults names were stored as FName.
- Renamed types with Extension to External.
- Encpasulated External Types FInstancedStruct to FValue abd FValueConst.

diff --git a/Engine/Plugins/Experimental/MutableExtensionExample/Source/MutableExtensionExample/Private/MeshClipSphere.cpp b/Engine/Plugins/Experimental/MutableExtensionExample/Source/MutableExtensionExample/Private/MeshClipSphere.cpp
new file mode 100644
--- /dev/null
+++ b/Engine/Plugins/Experimental/MutableExtensionExample/Source/MutableExtensionExample/Private/MeshClipSphere.cpp
@@ -0,0 +1,66 @@
+﻿// Copyright Epic Games, Inc. All Rights Reserved.
+
+#include "MeshClipSphere.h"
+
+#include "PrimitiveSphere.h"
+#include "MuR/External/MeshAdapter.h"
+
+
+const FName FMeshClipSphere::NameInputMesh = TEXT("Mesh");
+
+const FName FMeshClipSphere::NameInputPrimitiveSphere = TEXT("Sphere");
+
+
+FName FMeshClipSphere::GetName() const
+{
+	return TEXT("Mesh Clip With Sphere");
+}
+
+
+TArray<TPair<FName, const UScriptStruct*>> FMeshClipSphere::GetInputs() const
+{
+	TArray<TPair<FName, const UScriptStruct*>> Inputs;
+
+	Inputs.Emplace(NameInputMesh, UE::Mutable::FMeshAdapter::StaticStruct());
+	Inputs.Emplace(NameInputPrimitiveSphere, FPrimitiveSphere::StaticStruct());
+	
+	return Inputs;
+}
+
+
+TPair<FName, const UScriptStruct*> FMeshClipSphere::GetOutput() const
+{
+	return MakeTuple(TEXT("Mesh"), UE::Mutable::FMeshAdapter::StaticStruct());
+}
+
+
+void FMeshClipSphere::Evaluate(UE::Mutable::FContext& Context) const
+{
+	// Since we are going to edit the Mesh input, we need to make it non-constant using the CopyOrMove.
+	UE::Mutable::FValue InputMesh = UE::Mutable::CopyOrMove(Context.GetInput(NameInputMesh));
+
+	// Cast the input to a Mesh Adapter.
+	UE::Mutable::FMeshAdapter& Mesh = InputMesh.Get<UE::Mutable::FMeshAdapter>();
+
+	// Get the Primitive input
+	UE::Mutable::FValueConst InputSphere = Context.GetInput(NameInputPrimitiveSphere);
+	const FPrimitiveSphere& PrimitiveSphere = InputSphere.Get<FPrimitiveSphere>();
+	
+	TBitArray<> VerticesToClip;
+	VerticesToClip.Init(false, Mesh.GetNumVertices());
+
+	for (int32 Index = 0; Index < Mesh.GetNumVertices(); ++Index)
+	{
+		FVector3f Vertex = Mesh.GetVertex(Index);
+
+		if (PrimitiveSphere.Sphere.IsInside(static_cast<FVector>(Vertex)))
+		{
+			VerticesToClip[Index] = true;
+		}
+	}
+
+	Mesh.RemoveVertices(VerticesToClip);
+	
+	// Output the modified Mesh Adapter.
+	Context.SetOutput(MoveTemp(InputMesh));
+}
diff --git a/Engine/Plugins/Experimental/MutableExtensionExample/Source/MutableExtensionExample/Private/MeshIdentity.cpp b/Engine/Plugins/Experimental/MutableExtensionExample/Source/MutableExtensionExample/Private/MeshIdentity.cpp
--- a/Engine/Plugins/Experimental/MutableExtensionExample/Source/MutableExtensionExample/Private/MeshIdentity.cpp
+++ b/Engine/Plugins/Experimental/MutableExtensionExample/Source/MutableExtensionExample/Private/MeshIdentity.cpp
@@ -5,7 +5,7 @@
 #include "MuR/External/MeshAdapter.h"
 
 
-const FName FMeshIdentity::InputMesh = TEXT("Mesh");
+const FName FMeshIdentity::NameInputMesh = TEXT("Mesh");
 
 
 FName FMeshIdentity::GetName() const
@@ -18,7 +18,7 @@ TArray<TPair<FName, const UScriptStruct*>> FMeshIdentity::GetInputs() const
 {
 	TArray<TPair<FName, const UScriptStruct*>> Inputs;
 
-	Inputs.Emplace(InputMesh, UE::Mutable::FMeshAdapter::StaticStruct());
+	Inputs.Emplace(NameInputMesh, UE::Mutable::FMeshAdapter::StaticStruct());
 	
 	return Inputs;
 }
@@ -32,5 +32,5 @@ TPair<FName, const UScriptStruct*> FMeshIdentity::GetOutput() const
 
 void FMeshIdentity::Evaluate(UE::Mutable::FContext& Context) const
 {
-	Context.GetOutput() = Context.GetInput(InputMesh);
+	Context.SetOutput(Context.GetInput(NameInputMesh));
 }
diff --git a/Engine/Plugins/Experimental/MutableExtensionExample/Source/MutableExtensionExample/Private/MeshNoise.cpp b/Engine/Plugins/Experimental/MutableExtensionExample/Source/MutableExtensionExample/Private/MeshNoise.cpp
--- a/Engine/Plugins/Experimental/MutableExtensionExample/Source/MutableExtensionExample/Private/MeshNoise.cpp
+++ b/Engine/Plugins/Experimental/MutableExtensionExample/Source/MutableExtensionExample/Private/MeshNoise.cpp
@@ -6,9 +6,9 @@
 #include "MuR/External/MeshAdapter.h"
 
 
-const FName FMeshNoise::InputMesh = TEXT("Mesh");
+const FName FMeshNoise::NameInputMesh = TEXT("Mesh");
 
-const FName FMeshNoise::InputFactor = TEXT("Factor");
+const FName FMeshNoise::NameInputFactor = TEXT("Factor");
 
 
 FName FMeshNoise::GetName() const
@@ -21,8 +21,8 @@ TArray<TPair<FName, const UScriptStruct*>> FMeshNoise::GetInputs() const
 {
 	TArray<TPair<FName, const UScriptStruct*>> Inputs;
 
-	Inputs.Emplace(InputMesh, UE::Mutable::FMeshAdapter::StaticStruct());
-	Inputs.Emplace(InputFactor, UE::Mutable::FFloatAdapter::StaticStruct());
+	Inputs.Emplace(NameInputMesh, UE::Mutable::FMeshAdapter::StaticStruct());
+	Inputs.Emplace(NameInputFactor, UE::Mutable::FFloatAdapter::StaticStruct());
 	
 	return Inputs;
 }
@@ -37,15 +37,15 @@ TPair<FName, const UScriptStruct*> FMeshNoise::GetOutput() const
 void FMeshNoise::Evaluate(UE::Mutable::FContext& Context) const
 {
 	// Since we are going to edit the Mesh input, we need to make it non-constant using the CopyOrMove.
-	TSharedPtr<FInstancedStruct> Input = UE::Mutable::CopyOrMove(Context.GetInput(InputMesh));
+	UE::Mutable::FValue Input = UE::Mutable::CopyOrMove(Context.GetInput(NameInputMesh));
 
 	// Cast the input to a Mesh Adapter.
-	UE::Mutable::FMeshAdapter& Mesh = Input->GetMutable<UE::Mutable::FMeshAdapter>();
+	UE::Mutable::FMeshAdapter& Mesh = Input.Get<UE::Mutable::FMeshAdapter>();
 
 	// Get the Factor input
-	const float Factor = Context.GetInput(InputFactor)->Get<UE::Mutable::FFloatAdapter>().GetValue();
+	const float Factor = Context.GetInput(NameInputFactor).Get<UE::Mutable::FFloatAdapter>().GetValue();
 	
-	// Move each vertex some random noise.
+	// Move each vertex some random noise. This random is allowed since it is deterministic.
 	FRandomStream Random(Seed);
 	
 	for (int32 Index = 0; Index < Mesh.GetNumVertices(); ++Index)
@@ -58,5 +58,5 @@ void FMeshNoise::Evaluate(UE::Mutable::FContext& Context) const
 	}
 	
 	// Output the modified Mesh Adapter.
-	Context.GetOutput() = Input;
+	Context.SetOutput(MoveTemp(Input));
 }
diff --git a/Engine/Plugins/Experimental/MutableExtensionExample/Source/MutableExtensionExample/Private/MutableExtensionExampleModule.cpp b/Engine/Plugins/Experimental/MutableExtensionExample/Source/MutableExtensionExample/Private/MutableExtensionExampleModule.cpp
--- a/Engine/Plugins/Experimental/MutableExtensionExample/Source/MutableExtensionExample/Private/MutableExtensionExampleModule.cpp
+++ b/Engine/Plugins/Experimental/MutableExtensionExample/Source/MutableExtensionExample/Private/MutableExtensionExampleModule.cpp
@@ -2,6 +2,7 @@
 
 #include "MutableExtensionExampleModule.h"
 
+#include "MeshClipSphere.h"
 #include "MeshIdentity.h"
 #include "MeshNoise.h"
 #include "MuCO/ICustomizableObjectModule.h"
@@ -16,6 +17,7 @@ void FMutableExtensionExampleModule::StartupModule()
 
 	Module.RegisterExternalOperation(FMeshIdentity::StaticStruct());
 	Module.RegisterExternalOperation(FMeshNoise::StaticStruct());
+	Module.RegisterExternalOperation(FMeshClipSphere::StaticStruct());
 }
 
 
@@ -25,4 +27,5 @@ void FMutableExtensionExampleModule::ShutdownModule()
 
 	Module.UnregisterExternalOperation(FMeshIdentity::StaticStruct());
 	Module.UnregisterExternalOperation(FMeshNoise::StaticStruct());
+	Module.UnregisterExternalOperation(FMeshClipSphere::StaticStruct());
 }
diff --git a/Engine/Plugins/Experimental/MutableExtensionExample/Source/MutableExtensionExample/Public/MeshClipSphere.h b/Engine/Plugins/Experimental/MutableExtensionExample/Source/MutableExtensionExample/Public/MeshClipSphere.h
new file mode 100644
--- /dev/null
+++ b/Engine/Plugins/Experimental/MutableExtensionExample/Source/MutableExtensionExample/Public/MeshClipSphere.h
@@ -0,0 +1,24 @@
+﻿// Copyright Epic Games, Inc. All Rights Reserved.
+
+#pragma once
+
+#include "MuR/External/Operation.h"
+
+#include "MeshClipSphere.generated.h"
+
+
+USTRUCT()
+struct FMeshClipSphere : public UE::Mutable::FExternalOperation
+{
+	GENERATED_BODY()
+	
+	virtual ~FMeshClipSphere() override = default;
+	virtual FName GetName() const override;
+	virtual TArray<TPair<FName, const UScriptStruct*>> GetInputs() const override;
+	virtual TPair<FName, const UScriptStruct*> GetOutput() const override;
+	virtual void Evaluate(UE::Mutable::FContext& Context) const override;
+
+	const static FName NameInputMesh;
+
+	const static FName NameInputPrimitiveSphere;
+};
\ No newline at end of file
diff --git a/Engine/Plugins/Experimental/MutableExtensionExample/Source/MutableExtensionExample/Public/MeshIdentity.h b/Engine/Plugins/Experimental/MutableExtensionExample/Source/MutableExtensionExample/Public/MeshIdentity.h
--- a/Engine/Plugins/Experimental/MutableExtensionExample/Source/MutableExtensionExample/Public/MeshIdentity.h
+++ b/Engine/Plugins/Experimental/MutableExtensionExample/Source/MutableExtensionExample/Public/MeshIdentity.h
@@ -18,5 +18,5 @@ struct FMeshIdentity : public UE::Mutable::FExternalOperation
 	virtual TPair<FName, const UScriptStruct*> GetOutput() const override;
 	virtual void Evaluate(UE::Mutable::FContext& Context) const override;
 
-	const static FName InputMesh;
+	const static FName NameInputMesh;
 };
\ No newline at end of file
diff --git a/Engine/Plugins/Experimental/MutableExtensionExample/Source/MutableExtensionExample/Public/MeshNoise.h b/Engine/Plugins/Experimental/MutableExtensionExample/Source/MutableExtensionExample/Public/MeshNoise.h
--- a/Engine/Plugins/Experimental/MutableExtensionExample/Source/MutableExtensionExample/Public/MeshNoise.h
+++ b/Engine/Plugins/Experimental/MutableExtensionExample/Source/MutableExtensionExample/Public/MeshNoise.h
@@ -18,9 +18,9 @@ struct FMeshNoise : public UE::Mutable::FExternalOperation
 	virtual TPair<FName, const UScriptStruct*> GetOutput() const override;
 	virtual void Evaluate(UE::Mutable::FContext& Context) const override;
 
-	const static FName InputMesh;
+	const static FName NameInputMesh;
 
-	const static FName InputFactor;
+	const static FName NameInputFactor;
 
 	/** Operation constant. This value will be copied when compiling a Customizable Object, hence it can not be changed at runtime.
 	 *
diff --git a/Engine/Plugins/Experimental/MutableExtensionExample/Source/MutableExtensionExample/Public/PrimitiveSphere.h b/Engine/Plugins/Experimental/MutableExtensionExample/Source/MutableExtensionExample/Public/PrimitiveSphere.h
new file mode 100644
--- /dev/null
+++ b/Engine/Plugins/Experimental/MutableExtensionExample/Source/MutableExtensionExample/Public/PrimitiveSphere.h
@@ -0,0 +1,15 @@
+﻿// Copyright Epic Games, Inc. All Rights Reserved.
+
+#pragma once
+
+#include "PrimitiveSphere.generated.h"
+
+
+USTRUCT(DisplayName = "Sphere")
+struct FPrimitiveSphere
+{
+	GENERATED_BODY()
+
+	UPROPERTY(EditAnywhere, Category = Primitive)
+	FSphere3d Sphere;
+};
\ No newline at end of file
diff --git a/Engine/Plugins/Experimental/MutablePopulation/Source/CustomizableObjectPopulation/Private/MuCOP/CustomizableObjectPopulationGenerator.cpp b/Engine/Plugins/Experimental/MutablePopulation/Source/CustomizableObjectPopulation/Private/MuCOP/CustomizableObjectPopulationGenerator.cpp
--- a/Engine/Plugins/Experimental/MutablePopulation/Source/CustomizableObjectPopulation/Private/MuCOP/CustomizableObjectPopulationGenerator.cpp
+++ b/Engine/Plugins/Experimental/MutablePopulation/Source/CustomizableObjectPopulation/Private/MuCOP/CustomizableObjectPopulationGenerator.cpp
@@ -102,7 +102,7 @@ void FCustomizableObjectPopulationClassGenerator::GenerateParameters(
 				const FLinearColor SampledValue = UniformColorSamplers[SamplerIndex.SamplerIndex].Sample(Rand);
 				if (Instance->ContainsVectorParameter(ParamName))
 				{
-					Instance->SetColorParameterSelectedOption(ParamName, SampledValue);
+					Instance->SetVectorParameterSelectedOption(ParamName, SampledValue);
 				}
 				break;
 			}
@@ -110,7 +110,7 @@ void FCustomizableObjectPopulationClassGenerator::GenerateParameters(
 			{
 				if (Instance->ContainsVectorParameter(ParamName))
 				{
-					Instance->SetColorParameterSelectedOption(ParamName, ConstantColors[SamplerIndex.SamplerIndex]);
+					Instance->SetVectorParameterSelectedOption(ParamName, ConstantColors[SamplerIndex.SamplerIndex]);
 				}
 				break;
 			}
diff --git a/Engine/Plugins/Mutable/Source/CustomizableObject/Internal/MuCO/CustomizableObjectPrivate.h b/Engine/Plugins/Mutable/Source/CustomizableObject/Internal/MuCO/CustomizableObjectPrivate.h
--- a/Engine/Plugins/Mutable/Source/CustomizableObject/Internal/MuCO/CustomizableObjectPrivate.h
+++ b/Engine/Plugins/Mutable/Source/CustomizableObject/Internal/MuCO/CustomizableObjectPrivate.h
@@ -1394,13 +1394,19 @@ public:
 	TMap<FName, FPerQualityLevelInt> MinQualityLevelLODPerComponent;
 
 	UPROPERTY()
-	TMap<FName, TObjectPtr<UTexture>> TextureParameterDefaultValues;
+	TMap<FString, TObjectPtr<UTexture>> TextureParameterDefaultValues;
 
 	UPROPERTY()
-	TMap<FName, TObjectPtr<USkeletalMesh>> SkeletalMeshParameterDefaultValues;
+	TMap<FString, TObjectPtr<USkeletalMesh>> SkeletalMeshParameterDefaultValues;
 
 	UPROPERTY()
-	TMap<FName, TObjectPtr<UMaterialInterface>> MaterialParameterDefaultValues;
+	TMap<FString, TObjectPtr<UMaterialInterface>> MaterialParameterDefaultValues;
+
+	UPROPERTY()
+	TMap<FString, FInstancedStruct> ExternalTypeParameterDefaultValues;
+	
+	UPROPERTY()
+	TMap<FString, TObjectPtr<const UScriptStruct>> ExternalTypeParameterTypes;
 	
 	UPROPERTY()
 	FString ReleaseVersion;
@@ -1740,6 +1746,8 @@ enum class ECustomizableObjectVersions : int32
 
 	ImageParameterConvert,
 
+	InstancedStructType,
+
 	// -----<new versions can be added above this line>--------
 	LastCustomizableObjectVersion
 };
diff --git a/Engine/Plugins/Mutable/Source/CustomizableObject/Internal/MuCO/ICustomizableObjectModulePrivate.h b/Engine/Plugins/Mutable/Source/CustomizableObject/Internal/MuCO/ICustomizableObjectModulePrivate.h
--- a/Engine/Plugins/Mutable/Source/CustomizableObject/Internal/MuCO/ICustomizableObjectModulePrivate.h
+++ b/Engine/Plugins/Mutable/Source/CustomizableObject/Internal/MuCO/ICustomizableObjectModulePrivate.h
@@ -2,12 +2,16 @@
 
 #pragma once
 
-#include "Modules/ModuleManager.h"
 #include "MuCO/ICustomizableObjectModule.h"
 
-/**
- * The public interface of the CustomizableObject module
- */
+#include "Modules/ModuleManager.h"
+
+namespace UE::Mutable
+{
+	struct FExternalOperation;
+}
+
+
 class ICustomizableObjectModulePrivate : public ICustomizableObjectModule
 {
 public:
@@ -19,3 +23,6 @@ public:
 	virtual TArray<TStrongObjectPtr<const UScriptStruct>> GetRegisteredExternalOperations() const = 0;
 };
 
+
+CUSTOMIZABLEOBJECT_API bool ContainsUObject(const UE::Mutable::FExternalOperation& Operation, const UScriptStruct*& OutExternalType);
+
diff --git a/Engine/Plugins/Mutable/Source/CustomizableObject/Internal/MuCO/UnrealToMutableTextureConversionUtils.h b/Engine/Plugins/Mutable/Source/CustomizableObject/Internal/MuCO/UnrealToMutableTextureConversionUtils.h
--- a/Engine/Plugins/Mutable/Source/CustomizableObject/Internal/MuCO/UnrealToMutableTextureConversionUtils.h
+++ b/Engine/Plugins/Mutable/Source/CustomizableObject/Internal/MuCO/UnrealToMutableTextureConversionUtils.h
@@ -141,7 +141,7 @@ inline uint32 GetTypeHash(const FMutableSourceMeshData& Key)
 }
 
 
-struct FSourceExtensionOptions
+struct FSourceExternalOptions
 {
 	// GenerateMutableSourceMesh
 	FMutableSourceMeshData MeshOptions;
@@ -151,7 +151,7 @@ struct FSourceExtensionOptions
 	// GenerateMutableSourceImage
 	int32 ReferenceTextureSize = 0;
 
-	bool operator==(const FSourceExtensionOptions& Other) const = default;
+	bool operator==(const FSourceExternalOptions& Other) const = default;
 };
 
 
diff --git a/Engine/Plugins/Mutable/Source/CustomizableObject/Private/MuCO/CustomizableObject.cpp b/Engine/Plugins/Mutable/Source/CustomizableObject/Private/MuCO/CustomizableObject.cpp
--- a/Engine/Plugins/Mutable/Source/CustomizableObject/Private/MuCO/CustomizableObject.cpp
+++ b/Engine/Plugins/Mutable/Source/CustomizableObject/Private/MuCO/CustomizableObject.cpp
@@ -1481,6 +1481,12 @@ void UCustomizableObjectPrivate::UpdateParameterPropertiesFromModel(const TShare
 				break;
 			}
 
+			case UE::Mutable::Private::EParameterType::InstancedStruct:
+			{
+				Data.Type = EMutableParameterType::InstancedStruct;
+				break;
+			}
+
 			default:
 				// Unhandled type?
 				check(false);
@@ -1760,10 +1766,16 @@ bool UCustomizableObject::GetBoolParameterDefaultValue(const FString& InParamete
 
 FLinearColor UCustomizableObject::GetColorParameterDefaultValue(const FString& InParameterName) const
 {
-	const int32 ParameterIndex = GetPrivate()->FindParameter(InParameterName);
+	return GetVectorParameterDefaultValue(InParameterName);
+}
+
+
+FVector4f UCustomizableObject::GetVectorParameterDefaultValue(const FString& ParameterName) const
+{
+	const int32 ParameterIndex = GetPrivate()->FindParameter(ParameterName);
 	if (ParameterIndex == INDEX_NONE)
 	{
-		UE_LOG(LogMutable, Error, TEXT("Tried to access the default value of the nonexistent color parameter [%s] in the CustomizableObject [%s]."), *InParameterName, *GetName());
+		UE_LOG(LogMutable, Error, TEXT("Tried to access the default value of the nonexistent color parameter [%s] in the CustomizableObject [%s]."), *ParameterName, *GetName());
 		return FCustomizableObjectVectorParameterValue::DEFAULT_PARAMETER_VALUE;
 	}
 
@@ -1830,25 +1842,46 @@ FCustomizableObjectProjector UCustomizableObject::GetProjectorParameterDefaultVa
 
 UTexture* UCustomizableObject::GetTextureParameterDefaultValue(const FString& ParameterName) const
 {
-	const TObjectPtr<UTexture>* Result = GetPrivate()->GetModelResources()->TextureParameterDefaultValues.Find(FName(ParameterName));
+	const TObjectPtr<UTexture>* Result = GetPrivate()->GetModelResources()->TextureParameterDefaultValues.Find(ParameterName);
 	return Result ? *Result : nullptr;
 }
  
 
 UMaterialInterface* UCustomizableObject::GetMaterialParameterDefaultValue(const FString& ParameterName) const
 {
-	const TObjectPtr<UMaterialInterface>* Result = GetPrivate()->GetModelResources()->MaterialParameterDefaultValues.Find(FName(ParameterName));
+	const TObjectPtr<UMaterialInterface>* Result = GetPrivate()->GetModelResources()->MaterialParameterDefaultValues.Find(ParameterName);
 	return Result ? *Result : nullptr;
 }
 
 
 USkeletalMesh* UCustomizableObject::GetSkeletalMeshParameterDefaultValue(const FString& ParameterName) const
 {
-	const TObjectPtr<USkeletalMesh>* Result = GetPrivate()->GetModelResources()->SkeletalMeshParameterDefaultValues.Find(FName(ParameterName));
+	const TObjectPtr<USkeletalMesh>* Result = GetPrivate()->GetModelResources()->SkeletalMeshParameterDefaultValues.Find(ParameterName);
 	return Result ? *Result : nullptr;
 }
 
 
+FInstancedStruct UCustomizableObject::GetExternalTypeParameterDefaultValue(const FString& ParameterName) const
+{
+	const UModelResources& ModelResources = GetPrivate()->GetModelResourcesChecked();
+	
+	if (const FInstancedStruct* Result = ModelResources.ExternalTypeParameterDefaultValues.Find(ParameterName))
+	{
+		return *Result;
+	}
+
+	if (const TObjectPtr<const UScriptStruct>* Result = ModelResources.ExternalTypeParameterTypes.Find(ParameterName))
+	{
+		FInstancedStruct InstancedStruct;
+		InstancedStruct.InitializeAs(Result->Get());
+		
+		return InstancedStruct;
+	}
+	
+	return FCustomizableObjectExternalTypeParameterValue::DEFAULT_PARAMETER_VALUE;
+}
+
+
 bool UCustomizableObject::IsParameterMultidimensional(const FString& InParameterName) const
 {
 	const int32 ParameterIndex = GetPrivate()->FindParameter(InParameterName);
@@ -2217,9 +2250,44 @@ void UCustomizableObjectPrivate::SetModel(const TSharedPtr<UE::Mutable::Private:
 #else 
 		LocalModelResources = GetModelResources();
 #endif
-
 		check(LocalModelResources);
-		MutableModel->GetPrivate()->Program.ExternalOperationProvider = MakeShared<UE::Mutable::Private::FExternalOperationProvider>(*LocalModelResources);
+		
+		UE::Mutable::Private::FProgram& Program = MutableModel->GetPrivate()->Program;
+		Program.ExternalOperationProvider = MakeShared<UE::Mutable::Private::FExternalOperationProvider>(*LocalModelResources);
+
+		for (UE::Mutable::Private::FParameterDesc& Parameter : Program.Parameters)
+		{
+			switch (Parameter.Type)
+			{
+			case UE::Mutable::Private::EParameterType::Image:
+				Parameter.DefaultValue.Set<UE::Mutable::Private::FParamTextureType>(TStrongObjectPtr<UTexture>(LocalModelResources->TextureParameterDefaultValues[Parameter.Name]));
+				break;
+				
+			case UE::Mutable::Private::EParameterType::Mesh:
+				Parameter.DefaultValue.Set<UE::Mutable::Private::FParamSkeletalMeshType>(TStrongObjectPtr<USkeletalMesh>(LocalModelResources->SkeletalMeshParameterDefaultValues[Parameter.Name]));
+				break;
+				
+			case UE::Mutable::Private::EParameterType::Material:
+				Parameter.DefaultValue.Set<UE::Mutable::Private::FParamMaterialType>(TStrongObjectPtr<UMaterialInterface>(LocalModelResources->MaterialParameterDefaultValues[Parameter.Name]));
+				break;
+				
+			case UE::Mutable::Private::EParameterType::InstancedStruct:
+				Parameter.DefaultValue.Set<UE::Mutable::Private::FParamInstancedStructType>(MakeShared<FInstancedStruct>(LocalModelResources->ExternalTypeParameterDefaultValues[Parameter.Name]));
+				break;
+				
+			case UE::Mutable::Private::EParameterType::Matrix:
+			case UE::Mutable::Private::EParameterType::String:
+			case UE::Mutable::Private::EParameterType::None:
+			case UE::Mutable::Private::EParameterType::Bool:
+			case UE::Mutable::Private::EParameterType::Int:
+			case UE::Mutable::Private::EParameterType::Float:
+			case UE::Mutable::Private::EParameterType::Color:
+			case UE::Mutable::Private::EParameterType::Projector:
+				break;
+			default:
+				unimplemented();
+			}
+		}
 	}
 	
 	// Create parameter properties
diff --git a/Engine/Plugins/Mutable/Source/CustomizableObject/Private/MuCO/CustomizableObjectInstance.cpp b/Engine/Plugins/Mutable/Source/CustomizableObject/Private/MuCO/CustomizableObjectInstance.cpp
--- a/Engine/Plugins/Mutable/Source/CustomizableObject/Private/MuCO/CustomizableObjectInstance.cpp
+++ b/Engine/Plugins/Mutable/Source/CustomizableObject/Private/MuCO/CustomizableObjectInstance.cpp
@@ -2676,15 +2676,27 @@ void UCustomizableObjectInstance::SetMaterialParameterSelectedOption(const FStri
 }
 
 
+FInstancedStruct UCustomizableObjectInstance::GetExternalTypeParameterSelectedOption(const FString& ExternalTypeParamName, int32 RangeIndex) const
+{
+	return Descriptor.GetExternalTypeParameterSelectedOption(ExternalTypeParamName, RangeIndex);
+}
+
+
+void UCustomizableObjectInstance::SetExternalTypeParameterSelectedOption(const FString& ExternalTypeParamName, const FInstancedStruct& ExternalTypeValue, int32 RangeIndex)
+{
+	Descriptor.SetExternalTypeParameterSelectedOption(ExternalTypeParamName, ExternalTypeValue, RangeIndex);
+}
+
+
 FLinearColor UCustomizableObjectInstance::GetColorParameterSelectedOption(const FString& ColorParamName) const
 {
-	return Descriptor.GetColorParameterSelectedOption(ColorParamName);
+	return Descriptor.GetVectorParameterSelectedOption(ColorParamName);
 }
 
 
 void UCustomizableObjectInstance::SetColorParameterSelectedOption(const FString & ColorParamName, const FLinearColor& ColorValue)
 {
-	Descriptor.SetColorParameterSelectedOption(ColorParamName, ColorValue);
+	Descriptor.SetVectorParameterSelectedOption(ColorParamName, ColorValue);
 }
 
 
@@ -2700,7 +2712,13 @@ void UCustomizableObjectInstance::SetBoolParameterSelectedOption(const FString&
 }
 
 
-void UCustomizableObjectInstance::SetVectorParameterSelectedOption(const FString& VectorParamName, const FLinearColor& VectorValue)
+FVector4f UCustomizableObjectInstance::GetVectorParameterSelectedOption(const FString& VectorParamName)
+{
+	return Descriptor.GetVectorParameterSelectedOption(VectorParamName);
+}
+
+
+void UCustomizableObjectInstance::SetVectorParameterSelectedOption(const FString& VectorParamName, const FVector4f& VectorValue)
 {
 	Descriptor.SetVectorParameterSelectedOption(VectorParamName, VectorValue);
 }
@@ -2851,12 +2869,24 @@ bool UCustomizableObjectInstance::ContainsMaterialParameter(const FString& Param
 }
 
 
+bool UCustomizableObjectInstance::ContainsExternalTypeParameter(const FString& ParameterName) const
+{
+	return Descriptor.FindTypedParameterIndex(ParameterName, EMutableParameterType::InstancedStruct) != INDEX_NONE;
+}
+
+
 bool UCustomizableObjectInstance::ContainsBoolParameter(const FString& ParameterName) const
 {
 	return Descriptor.FindTypedParameterIndex(ParameterName, EMutableParameterType::Bool) != INDEX_NONE;
 }
 
 
+bool UCustomizableObjectInstance::ContainsColorParameter(const FString& ParameterName) const
+{
+	return Descriptor.FindTypedParameterIndex(ParameterName, EMutableParameterType::Color) != INDEX_NONE;
+}
+
+
 bool UCustomizableObjectInstance::ContainsVectorParameter(const FString& ParameterName) const
 {
 	return Descriptor.FindTypedParameterIndex(ParameterName, EMutableParameterType::Color) != INDEX_NONE;
diff --git a/Engine/Plugins/Mutable/Source/CustomizableObject/Private/MuCO/CustomizableObjectInstanceDescriptor.cpp b/Engine/Plugins/Mutable/Source/CustomizableObject/Private/MuCO/CustomizableObjectInstanceDescriptor.cpp
--- a/Engine/Plugins/Mutable/Source/CustomizableObject/Private/MuCO/CustomizableObjectInstanceDescriptor.cpp
+++ b/Engine/Plugins/Mutable/Source/CustomizableObject/Private/MuCO/CustomizableObjectInstanceDescriptor.cpp
@@ -592,11 +592,7 @@ TSharedPtr<UE::Mutable::Private::FParameters> FCustomizableObjectInstanceDescrip
  
 	const UModelResources& ModelResources = CustomizableObject->GetPrivate()->GetModelResourcesChecked();
 	
-	TSharedPtr<UE::Mutable::Private::FParameters> MutableParameters = UE::Mutable::Private::FModel::NewParameters(CustomizableObject->GetPrivate()->GetModel(),
-		nullptr,
-		&ModelResources.TextureParameterDefaultValues,
-		&ModelResources.SkeletalMeshParameterDefaultValues,
-		&ModelResources.MaterialParameterDefaultValues);
+	TSharedPtr<UE::Mutable::Private::FParameters> MutableParameters = UE::Mutable::Private::FModel::NewParameters(CustomizableObject->GetPrivate()->GetModel(), nullptr);
  
 	const int32 ParamCount = MutableParameters->GetCount();
 	for (int32 ParamIndex = 0; ParamIndex < ParamCount; ++ParamIndex)
@@ -843,6 +839,32 @@ TSharedPtr<UE::Mutable::Private::FParameters> FCustomizableObjectInstanceDescrip
 
 			break;
 		}
+
+		case UE::Mutable::Private::EParameterType::InstancedStruct:
+		{
+			for (const FCustomizableObjectExternalTypeParameterValue& ExternalTypeParameter : ExternalTypeParameters)
+			{
+				if (ExternalTypeParameter.ParameterName == Name || (Uid.IsValid() && ExternalTypeParameter.Id == Uid))
+				{
+					if (TSharedPtr<UE::Mutable::Private::FRangeIndex> RangeIdxPtr = MutableParameters->NewRangeIndex(ParamIndex))
+					{
+						for (int32 RangeIndex = 0; RangeIndex < ExternalTypeParameter.ParameterRangeValues.Num(); ++RangeIndex)
+						{
+							RangeIdxPtr->SetPosition(0, RangeIndex);
+							MutableParameters->SetInstancedStructValue(ParamIndex, MakeShared<const FInstancedStruct>(ExternalTypeParameter.ParameterRangeValues[RangeIndex]), RangeIdxPtr.Get());
+						}
+					}
+					else
+					{
+						MutableParameters->SetInstancedStructValue(ParamIndex, MakeShared<const FInstancedStruct>(ExternalTypeParameter.ParameterValue));
+					}
+
+					break;
+				}
+			}
+
+			break;
+		}
  
 		default:
 			check(false); // Missing case.
@@ -893,6 +915,7 @@ void FCustomizableObjectInstanceDescriptor::ReloadParameters()
 	TArray<FCustomizableObjectTextureParameterValue> OldTextureParameters = TextureParameters;
 	TArray<FCustomizableObjectSkeletalMeshParameterValue> OldSkeletalMeshParameters = SkeletalMeshParameters;
 	TArray<FCustomizableObjectMaterialParameterValue> OldMaterialParameters = MaterialParameters;
+	TArray<FCustomizableObjectExternalTypeParameterValue> OldExternalTypeParameters = ExternalTypeParameters;
 	TArray<FCustomizableObjectVectorParameterValue> OldVectorParameters = VectorParameters;
 	TArray<FCustomizableObjectProjectorParameterValue> OldProjectorParameters = ProjectorParameters;
 	TArray<FCustomizableObjectTransformParameterValue> OldTransformParameters = TransformParameters;
@@ -903,6 +926,7 @@ void FCustomizableObjectInstanceDescriptor::ReloadParameters()
 	TextureParameters.Reset();
 	SkeletalMeshParameters.Reset();
 	MaterialParameters.Reset();
+	ExternalTypeParameters.Reset();
 	VectorParameters.Reset();
 	ProjectorParameters.Reset();
 	TransformParameters.Reset();
@@ -921,12 +945,7 @@ void FCustomizableObjectInstanceDescriptor::ReloadParameters()
 		return;
 	}
 	
-	TSharedPtr<UE::Mutable::Private::FParameters> MutableParameters = UE::Mutable::Private::FModel::NewParameters(
-		CustomizableObject->GetPrivate()->GetModel(),
-		nullptr,
-		&ModelResources->TextureParameterDefaultValues,
-		&ModelResources->SkeletalMeshParameterDefaultValues,
-		&ModelResources->MaterialParameterDefaultValues);
+	TSharedPtr<UE::Mutable::Private::FParameters> MutableParameters = UE::Mutable::Private::FModel::NewParameters(CustomizableObject->GetPrivate()->GetModel(), nullptr);
  
 	int32 ParamCount = MutableParameters->GetCount();
 	for (int32 ParamIndex = 0; ParamIndex < ParamCount; ++ParamIndex)
@@ -1423,6 +1442,65 @@ void FCustomizableObjectInstanceDescriptor::ReloadParameters()
 
 			break;
 		}
+
+		case UE::Mutable::Private::EParameterType::InstancedStruct:
+		{
+			FCustomizableObjectExternalTypeParameterValue Param;
+			Param.ParameterName = Name;
+			Param.Id = Uid;
+
+			auto FindByNameAndUid = [&](const FCustomizableObjectExternalTypeParameterValue& P)
+				{
+					return P.ParameterName == Name || (Uid.IsValid() && P.Id == Uid);
+				};
+
+			if (FCustomizableObjectExternalTypeParameterValue* Result = OldExternalTypeParameters.FindByPredicate(FindByNameAndUid))
+			{
+				if (TSharedPtr<UE::Mutable::Private::FRangeIndex> RangeIdxPtr = MutableParameters->NewRangeIndex(ParamIndex))
+				{
+					Param.ParameterRangeValues = Result->ParameterRangeValues;
+				}
+				else
+				{
+					Param.ParameterValue = Result->ParameterValue;
+				}
+			}
+			else // Not found in Instance Parameters. Use Mutable Parameters.
+			{
+				if (TSharedPtr<UE::Mutable::Private::FRangeIndex> RangeIdxPtr = MutableParameters->NewRangeIndex(ParamIndex))
+				{
+					const int32 ValueCount = MutableParameters->GetValueCount(ParamIndex);
+
+					for (int32 ValueIndex = 0; ValueIndex < ValueCount; ++ValueIndex)
+					{
+						TSharedPtr<UE::Mutable::Private::FRangeIndex> RangeValueIdxPtr = MutableParameters->GetValueIndex(ParamIndex, ValueIndex);
+						int32 RangeIndex = RangeValueIdxPtr->GetPosition(0);
+
+						if (!Param.ParameterRangeValues.IsValidIndex(RangeIndex))
+						{
+							const int32 PreviousNum = Param.ParameterRangeValues.Num();
+							Param.ParameterRangeValues.AddUninitialized(RangeIndex + 1 - Param.ParameterRangeValues.Num());
+
+							for (int32 Index = PreviousNum; Index < Param.ParameterRangeValues.Num(); ++Index)
+							{
+								Param.ParameterRangeValues[Index] = FInstancedStruct();
+							}
+						}
+
+						Param.ParameterRangeValues[RangeIndex] = *MutableParameters->GetInstancedStructValue(ParamIndex, RangeValueIdxPtr.Get());
+					}
+				}
+				else
+				{
+					Param.ParameterValue = *MutableParameters->GetInstancedStructValue(ParamIndex);
+				}
+			}
+
+			ExternalTypeParameters.Add(Param);
+
+			break;
+		}
+			
 		default:
 			check(false); // Missing case.
 			break;
@@ -1676,7 +1754,6 @@ void FCustomizableObjectInstanceDescriptor::SetFloatParameterSelectedOption(cons
 		check(FloatParameters[ParameterIndexInInstance].ParameterRangeValues.IsValidIndex(RangeIndex));
 		FloatParameters[ParameterIndexInInstance].ParameterRangeValues[RangeIndex] = FloatValue;
 	}
- 
 }
  
  
@@ -1911,7 +1988,91 @@ void FCustomizableObjectInstanceDescriptor::SetMaterialParameterSelectedOption(c
 }
 
 
-FLinearColor FCustomizableObjectInstanceDescriptor::GetColorParameterSelectedOption(const FString& ColorParamName) const
+FInstancedStruct FCustomizableObjectInstanceDescriptor::GetExternalTypeParameterSelectedOption(const FString& ExternalTypeParameterName, int32 RangeIndex) const
+{
+	if (!CustomizableObject)
+	{
+		CustomizableObjectNullErrorMessage();
+		return FCustomizableObjectExternalTypeParameterValue::DEFAULT_PARAMETER_VALUE;
+	}
+ 
+	const int32 ParameterIndexInObject = CustomizableObject->GetPrivate()->FindParameter(ExternalTypeParameterName);
+	const int32 ExternalTypeParamIndex = FindTypedParameterIndex(ExternalTypeParameterName, EMutableParameterType::InstancedStruct);
+ 
+	if (ParameterIndexInObject >= 0 && ExternalTypeParameters.IsValidIndex(ExternalTypeParamIndex))
+	{
+		// Due to optimizations a parameter may lose its multidimensionality (if it becomes constant). 
+		// In that case it means there is only one possible option so it is ok to set it as if wasn't multidimensional.
+		bool bIsMultidimensional = CustomizableObject->GetPrivate()->IsParameterMultidimensional(ParameterIndexInObject);
+		if (!bIsMultidimensional || RangeIndex == INDEX_NONE)
+		{
+			return ExternalTypeParameters[ExternalTypeParamIndex].ParameterValue;
+		}
+		else
+		{
+			if (ExternalTypeParameters[ExternalTypeParamIndex].ParameterRangeValues.IsValidIndex(RangeIndex))
+			{
+				return ExternalTypeParameters[ExternalTypeParamIndex].ParameterRangeValues[RangeIndex];
+			}
+		}
+	}
+ 
+	LogParameterNotFoundWarning(ExternalTypeParameterName, ParameterIndexInObject, ExternalTypeParamIndex, CustomizableObject, __FUNCTION__);
+	
+	return FCustomizableObjectExternalTypeParameterValue::DEFAULT_PARAMETER_VALUE; 
+}
+
+
+void FCustomizableObjectInstanceDescriptor::SetExternalTypeParameterSelectedOption(const FString& ExternalTypeParamName, const FInstancedStruct& ExternalTypeValue, int32 RangeIndex)
+{
+	if (!CustomizableObject)
+	{
+		CustomizableObjectNullErrorMessage();
+		return;
+	}
+ 
+	RETURN_ON_UNCOMPILED_CO(CustomizableObject, TEXT("Error: Cannot set Int parameter "));
+ 
+	const int32 ParameterIndexInObject = CustomizableObject->GetPrivate()->FindParameter(ExternalTypeParamName);
+	const int32 ParameterIndexInInstance = FindTypedParameterIndex(ExternalTypeParamName, EMutableParameterType::InstancedStruct);
+ 
+	if (ParameterIndexInObject < 0 || ParameterIndexInInstance < 0)
+	{
+		// Warn and early out since we could not find the parameter to set.
+		LogParameterNotFoundWarning(ExternalTypeParamName, ParameterIndexInObject, ParameterIndexInInstance, CustomizableObject, __FUNCTION__);
+		return;
+	}
+
+	const UModelResources& ModelResources = CustomizableObject->GetPrivate()->GetModelResourcesChecked();
+	if (!ExternalTypeValue.GetScriptStruct()->IsChildOf(ModelResources.ExternalTypeParameterDefaultValues[ExternalTypeParamName].GetScriptStruct()))
+	{
+		UE_LOG(LogMutable, Error, TEXT("External Parameter %s incorrect type."), *ExternalTypeParamName);
+		return;
+	}
+	
+	// Due to optimizations a parameter may lose its multidimensionality (if it becomes constant). 
+	// In that case it means there is only one possible option so it is ok to set it as if wasn't multidimensional.
+	bool bIsMultidimensional = CustomizableObject->GetPrivate()->IsParameterMultidimensional(ParameterIndexInObject);
+	if (!bIsMultidimensional || RangeIndex == INDEX_NONE)
+	{
+		ExternalTypeParameters[ParameterIndexInInstance].ParameterValue = ExternalTypeValue;
+	}
+	else
+	{
+		if (!ExternalTypeParameters[ParameterIndexInInstance].ParameterRangeValues.IsValidIndex(RangeIndex))
+		{
+			const int32 InsertionIndex = ExternalTypeParameters[ParameterIndexInInstance].ParameterRangeValues.Num();
+			const int32 NumInsertedElements = RangeIndex + 1 - ExternalTypeParameters[ParameterIndexInInstance].ParameterRangeValues.Num();
+			FloatParameters[ParameterIndexInInstance].ParameterRangeValues.InsertDefaulted(InsertionIndex, NumInsertedElements);
+		}
+ 
+		check(ExternalTypeParameters[ParameterIndexInInstance].ParameterRangeValues.IsValidIndex(RangeIndex));
+		ExternalTypeParameters[ParameterIndexInInstance].ParameterRangeValues[RangeIndex] = ExternalTypeValue;
+	}
+}
+
+
+FVector4f FCustomizableObjectInstanceDescriptor::GetVectorParameterSelectedOption(const FString& VectorParamName) const
 {
 	if (!CustomizableObject)
 	{
@@ -1919,21 +2080,21 @@ FLinearColor FCustomizableObjectInstanceDescriptor::GetColorParameterSelectedOpt
 		return FCustomizableObjectVectorParameterValue::DEFAULT_PARAMETER_VALUE;
 	}
  
-	const int32 ParameterIndexInObject = CustomizableObject->GetPrivate()->FindParameter(ColorParamName);
-	const int32 ColorParamIndex = FindTypedParameterIndex(ColorParamName, EMutableParameterType::Color);
+	const int32 ParameterIndexInObject = CustomizableObject->GetPrivate()->FindParameter(VectorParamName);
+	const int32 VectorParamIndex = FindTypedParameterIndex(VectorParamName, EMutableParameterType::Color);
  
-	if (ColorParamIndex == INDEX_NONE)
+	if (VectorParamIndex == INDEX_NONE)
 	{
-		LogParameterNotFoundWarning(ColorParamName, ParameterIndexInObject, ColorParamIndex, CustomizableObject, __FUNCTION__);
+		LogParameterNotFoundWarning(VectorParamName, ParameterIndexInObject, VectorParamIndex, CustomizableObject, __FUNCTION__);
 		return FCustomizableObjectVectorParameterValue::DEFAULT_PARAMETER_VALUE;
 	}
  
 	// TODO: Multidimensional parameter support
-	return VectorParameters.IsValidIndex(ColorParamIndex) ? VectorParameters[ColorParamIndex].ParameterValue : FLinearColor();
+	return VectorParameters.IsValidIndex(VectorParamIndex) ? VectorParameters[VectorParamIndex].ParameterValue : FCustomizableObjectVectorParameterValue::DEFAULT_PARAMETER_VALUE;
 }
- 
- 
-void FCustomizableObjectInstanceDescriptor::SetColorParameterSelectedOption(const FString& ColorParamName, const FLinearColor& ColorValue)
+
+
+void FCustomizableObjectInstanceDescriptor::SetVectorParameterSelectedOption(const FString& VectorParamName, const FVector4f& VectorValue)
 {
 	if (!CustomizableObject)
 	{
@@ -1943,10 +2104,20 @@ void FCustomizableObjectInstanceDescriptor::SetColorParameterSelectedOption(cons
  
 	RETURN_ON_UNCOMPILED_CO(CustomizableObject, TEXT("Error: Cannot set Int parameter "));
  
-	// TODO: Multidimensional parameter support
-	SetVectorParameterSelectedOption(ColorParamName, ColorValue);
-}
+	const int32 ParameterIndexInObject = CustomizableObject->GetPrivate()->FindParameter(VectorParamName);
+	const int32 ParameterIndexInInstance = FindTypedParameterIndex(VectorParamName, EMutableParameterType::Color);
+ 
+	if (ParameterIndexInObject < 0 || ParameterIndexInInstance < 0)
+	{
+		// Early out since we could not find the parameter to set.
+		LogParameterNotFoundWarning(VectorParamName, ParameterIndexInObject, ParameterIndexInInstance, CustomizableObject, __FUNCTION__);
+		return;
+	}
  
+	VectorParameters[ParameterIndexInInstance].ParameterValue = VectorValue;
+}
+
+
 FTransform FCustomizableObjectInstanceDescriptor::GetTransformParameterSelectedOption(const FString& TransformParamName) const
 {
 	if (!CustomizableObject)
@@ -2036,31 +2207,7 @@ void FCustomizableObjectInstanceDescriptor::SetBoolParameterSelectedOption(const
  
 	BoolParameters[ParameterIndexInInstance].ParameterValue = BoolValue;
 }
- 
- 
-void FCustomizableObjectInstanceDescriptor::SetVectorParameterSelectedOption(const FString& VectorParamName, const FLinearColor& VectorValue)
-{
-	if (!CustomizableObject)
-	{
-		CustomizableObjectNullErrorMessage();
-		return;
-	}
- 
-	RETURN_ON_UNCOMPILED_CO(CustomizableObject, TEXT("Error: Cannot set Int parameter "));
- 
-	const int32 ParameterIndexInObject = CustomizableObject->GetPrivate()->FindParameter(VectorParamName);
-	const int32 ParameterIndexInInstance = FindTypedParameterIndex(VectorParamName, EMutableParameterType::Color);
- 
-	if (ParameterIndexInObject < 0 || ParameterIndexInInstance < 0)
-	{
-		// Early out since we could not find the parameter to set.
-		LogParameterNotFoundWarning(VectorParamName, ParameterIndexInObject, ParameterIndexInInstance, CustomizableObject, __FUNCTION__);
-		return;
-	}
- 
-	VectorParameters[ParameterIndexInInstance].ParameterValue = VectorValue;
-}
- 
+
  
 void FCustomizableObjectInstanceDescriptor::SetProjectorValue(const FString& ProjectorParamName,
 	const FVector& Pos, const FVector& Direction, const FVector& Up, const FVector& Scale,
@@ -2845,7 +2992,7 @@ void FCustomizableObjectInstanceDescriptor::SetDefaultValue(int32 ParamIndex)
  
 	case EMutableParameterType::Color:
 	{
-		const FLinearColor DefaultValue = CustomizableObject->GetColorParameterDefaultValue(ParamName);
+		const FLinearColor DefaultValue = CustomizableObject->GetVectorParameterDefaultValue(ParamName);
 		VectorParameters[TypedIndex].ParameterValue = DefaultValue;
 		break;
 	}
@@ -2905,6 +3052,19 @@ void FCustomizableObjectInstanceDescriptor::SetDefaultValue(int32 ParamIndex)
 		}
 		break;
 	}
+
+	case EMutableParameterType::InstancedStruct:
+	{
+		FInstancedStruct DefaultValue = CustomizableObject->GetExternalTypeParameterDefaultValue(ParamName);
+
+		if (ensure(ExternalTypeParameters.IsValidIndex(TypedIndex)))
+		{
+			FCustomizableObjectExternalTypeParameterValue& ExternalTypeParameter = ExternalTypeParameters[TypedIndex];
+			ExternalTypeParameter.ParameterValue = DefaultValue;
+			ExternalTypeParameter.ParameterRangeValues.Empty();
+		}
+		break;
+	}
 		
 	default:
 		// Parameter type replication not implemented.
diff --git a/Engine/Plugins/Mutable/Source/CustomizableObject/Private/MuCO/CustomizableObjectModule.cpp b/Engine/Plugins/Mutable/Source/CustomizableObject/Private/MuCO/CustomizableObjectModule.cpp
--- a/Engine/Plugins/Mutable/Source/CustomizableObject/Private/MuCO/CustomizableObjectModule.cpp
+++ b/Engine/Plugins/Mutable/Source/CustomizableObject/Private/MuCO/CustomizableObjectModule.cpp
@@ -187,6 +187,60 @@ TArrayView<const FRegisteredObjectNodeInputPin> FCustomizableObjectModule::GetAd
 }
 
 
+bool ContainsUObject(const UScriptStruct* ScriptStruct)
+{
+	if (!ScriptStruct)
+	{
+		return false;
+	}
+
+	for (TFieldIterator<FProperty> It(ScriptStruct); It; ++It)
+	{
+		const FProperty* Prop = *It;
+
+		if (Prop->HasAnyPropertyFlags(CPF_TObjectPtr | CPF_UObjectWrapper))
+		{
+			return true;
+		}
+
+		if (const FStructProperty* StructProp = CastField<FStructProperty>(Prop))
+		{
+			if (ContainsUObject(StructProp->Struct))
+			{
+				return true;
+			}
+		}
+	}
+
+	return false;
+}
+
+
+bool ContainsUObject(const UE::Mutable::FExternalOperation& Operation, const UScriptStruct*& OutExternalType)
+{
+	TArray<const UScriptStruct*> Types;
+	
+	for (const TPair<FName, const UScriptStruct*>& Input : Operation.GetInputs())
+	{
+		Types.AddUnique(Input.Value);
+	}
+
+	Types.AddUnique((Operation.GetOutput().Value));
+
+	for (const UScriptStruct* Type : Types)
+	{
+		if (ContainsUObject(Type))
+		{
+			OutExternalType = Type;
+			return true;
+		}
+	}
+
+	OutExternalType = nullptr;
+	return false;
+}
+
+
 void FCustomizableObjectModule::RegisterExternalOperation(const UScriptStruct* Operation)
 {
 	if (!Operation)
@@ -199,6 +253,16 @@ void FCustomizableObjectModule::RegisterExternalOperation(const UScriptStruct* O
 		UE_LOG(LogMutable, Error, TEXT("Unable to register operation. Operation that does not inherit from FMutableOperation: %s"), *Operation->GetName());
 		return;
 	}
+
+	FInstancedStruct InstancedStruct;
+	InstancedStruct.InitializeAs(Operation);
+	
+	const UScriptStruct* ExternalType;
+	if (ContainsUObject(InstancedStruct.Get<UE::Mutable::FExternalOperation>(), ExternalType))
+	{
+		UE_LOG(LogMutable, Error, TEXT("Unable to register operation. Operation contains UObject properties. Operation: %s Type: %s"), *Operation->GetName(), *ExternalType->GetName());
+		return;
+	}
 	
 	RegisteredOperations.AddUnique(TStrongObjectPtr(Operation));
 }
diff --git a/Engine/Plugins/Mutable/Source/CustomizableObject/Private/MuCO/CustomizableObjectParameterTypeDefinitions.cpp b/Engine/Plugins/Mutable/Source/CustomizableObject/Private/MuCO/CustomizableObjectParameterTypeDefinitions.cpp
new file mode 100644
--- /dev/null
+++ b/Engine/Plugins/Mutable/Source/CustomizableObject/Private/MuCO/CustomizableObjectParameterTypeDefinitions.cpp
@@ -0,0 +1,20 @@
+// Copyright Epic Games, Inc. All Rights Reserved.
+
+#include "MuCO/CustomizableObjectParameterTypeDefinitions.h"
+
+#include "MuR/Types.h"
+
+
+uint32 GetTypeHash(const FCustomizableObjectExternalTypeParameterValue& Key)
+{
+	uint32 Hash = GetTypeHash(Key.ParameterName);
+	
+	Hash = HashCombine(Hash, UE::Mutable::Private::GetTypeHash(Key.ParameterValue));
+	
+	for (const FInstancedStruct& Value : Key.ParameterRangeValues)
+	{
+		Hash = HashCombine(Hash, UE::Mutable::Private::GetTypeHash(Key.ParameterValue));
+	}
+	
+	return Hash;
+}
diff --git a/Engine/Plugins/Mutable/Source/CustomizableObject/Private/MuCO/DescriptorHash.cpp b/Engine/Plugins/Mutable/Source/CustomizableObject/Private/MuCO/DescriptorHash.cpp
--- a/Engine/Plugins/Mutable/Source/CustomizableObject/Private/MuCO/DescriptorHash.cpp
+++ b/Engine/Plugins/Mutable/Source/CustomizableObject/Private/MuCO/DescriptorHash.cpp
@@ -53,6 +53,11 @@ FDescriptorHash::FDescriptorHash(const FCustomizableObjectInstanceDescriptor& De
 	{
 		Hash = HashCombine(Hash, GetTypeHash(Value));
 	}
+	
+	for (const FCustomizableObjectExternalTypeParameterValue& Value : Descriptor.ExternalTypeParameters)
+	{
+		Hash = HashCombine(Hash, GetTypeHash(Value));
+	}
 
 	Hash = HashCombine(Hash, GetTypeHash(Descriptor.State));
 	Hash = HashCombine(Hash, GetTypeHash(Descriptor.GetBuildParameterRelevancy()));
diff --git a/Engine/Plugins/Mutable/Source/CustomizableObject/Private/MuCO/ExternalOperationProvider.cpp b/Engine/Plugins/Mutable/Source/CustomizableObject/Private/MuCO/ExternalOperationProvider.cpp
--- a/Engine/Plugins/Mutable/Source/CustomizableObject/Private/MuCO/ExternalOperationProvider.cpp
+++ b/Engine/Plugins/Mutable/Source/CustomizableObject/Private/MuCO/ExternalOperationProvider.cpp
@@ -17,7 +17,10 @@ namespace UE::Mutable::Private
 	{
 		for (TPair<OP::ADDRESS, FInstancedStruct>& Pair : ExternalOperations)
 		{
-			Collector.AddPropertyReferences(Pair.Value.GetScriptStruct(), Pair.Value.GetMutableMemory());
+			if (const UScriptStruct* ScriptStruct = Pair.Value.GetScriptStruct())
+			{
+				Collector.AddPropertyReferences(ScriptStruct, Pair.Value.GetMutableMemory());
+			}
 		}
 	}
 
diff --git a/Engine/Plugins/Mutable/Source/CustomizableObject/Public/MuCO/CustomizableObject.h b/Engine/Plugins/Mutable/Source/CustomizableObject/Public/MuCO/CustomizableObject.h
--- a/Engine/Plugins/Mutable/Source/CustomizableObject/Public/MuCO/CustomizableObject.h
+++ b/Engine/Plugins/Mutable/Source/CustomizableObject/Public/MuCO/CustomizableObject.h
@@ -397,6 +397,12 @@ public:
 	  * @return The default value of the provided parameter name. */
 	UFUNCTION(BlueprintCallable, Category = CustomizableObject)
 	CUSTOMIZABLEOBJECT_API FLinearColor GetColorParameterDefaultValue(const FString& ParameterName) const;
+
+	/** Get the default value of a parameter of type Vector.
+	  * @param ParameterName The name of the Vector parameter to get the default value of.
+	  * @return The default value of the provided parameter name. */
+	UFUNCTION(BlueprintCallable, Category = CustomizableObject)
+	CUSTOMIZABLEOBJECT_API FVector4f GetVectorParameterDefaultValue(const FString& ParameterName) const;
 	
 	/** Get the default value of a parameter of type Transform.
 	  * @param ParameterName The name of the Transform parameter to get the default value of.
@@ -425,6 +431,11 @@ public:
 	UFUNCTION(BlueprintCallable, Category = CustomizableObject)
 	CUSTOMIZABLEOBJECT_API UMaterialInterface* GetMaterialParameterDefaultValue(const FString& ParameterName) const;
 
+	/** Get the default value of a parameter of type External Type.
+	  * @param ParameterName The name of the External Type parameter to get the default value of.
+	  * @return The default value of the provided parameter name. */
+	CUSTOMIZABLEOBJECT_API FInstancedStruct GetExternalTypeParameterDefaultValue(const FString& ParameterName) const;
+	
 	/** Return true if the parameter at the index provided is multidimensional.
 	  * @param InParameterName The name of the parameter to check.
 	  * @return True if the parameter is multidimensional and false if it is not. */
diff --git a/Engine/Plugins/Mutable/Source/CustomizableObject/Public/MuCO/CustomizableObjectInstance.h b/Engine/Plugins/Mutable/Source/CustomizableObject/Public/MuCO/CustomizableObjectInstance.h
--- a/Engine/Plugins/Mutable/Source/CustomizableObject/Public/MuCO/CustomizableObjectInstance.h
+++ b/Engine/Plugins/Mutable/Source/CustomizableObject/Public/MuCO/CustomizableObjectInstance.h
@@ -512,6 +512,14 @@ public:
 	UFUNCTION(BlueprintCallable, Category = CustomizableObjectInstance)
 	CUSTOMIZABLEOBJECT_API void SetMaterialParameterSelectedOption(const FString& MaterialParamName, UMaterialInterface* MaterialValue, int32 RangeIndex = -1);
 	
+	/** Gets the value of an External Type parameter with name "ExternalTypeParamName". */
+	UFUNCTION(BlueprintCallable, Category = CustomizableObjectInstance)
+	CUSTOMIZABLEOBJECT_API FInstancedStruct GetExternalTypeParameterSelectedOption(const FString& ExternalTypeParamName,  int32 RangeIndex = -1) const;
+
+	/** Sets the External Type value "ExternalTypeValue" of an External Type parameter with name "ExternalTypeParamName". */
+	UFUNCTION(BlueprintCallable, Category = CustomizableObjectInstance)
+	CUSTOMIZABLEOBJECT_API void SetExternalTypeParameterSelectedOption(const FString& ExternalTypeParamName, const FInstancedStruct& ExternalTypeValue, int32 RangeIndex = -1);
+
 	/** Gets the value of a color parameter with name "ColorParamName". */
 	UFUNCTION(BlueprintCallable, Category = CustomizableObjectInstance)
 	CUSTOMIZABLEOBJECT_API FLinearColor GetColorParameterSelectedOption(const FString& ColorParamName) const;
@@ -520,7 +528,7 @@ public:
 	UFUNCTION(BlueprintCallable, Category = CustomizableObjectInstance)
 	CUSTOMIZABLEOBJECT_API void SetColorParameterSelectedOption(const FString& ColorParamName, const FLinearColor& ColorValue);
 
-	/** Gets the bool value "BoolValue" of a bool parameter with name "BoolParamName". */
+	/** Gets the bool value of a bool parameter with name "BoolParamName". */
 	UFUNCTION(BlueprintCallable, Category = CustomizableObjectInstance)
 	CUSTOMIZABLEOBJECT_API bool GetBoolParameterSelectedOption(const FString& BoolParamName) const;
 
@@ -528,9 +536,13 @@ public:
 	UFUNCTION(BlueprintCallable, Category = CustomizableObjectInstance)
 	CUSTOMIZABLEOBJECT_API void SetBoolParameterSelectedOption(const FString& BoolParamName, bool BoolValue);
 
-	/** Sets the vector value "VectorValue" of a vector parameter with index "VectorParamIndex". */
+	/** Gets the Vector value of a Vector parameter with name "VectorParamName". */
+	UFUNCTION(BlueprintCallable, Category = CustomizableObjectInstance)
+	CUSTOMIZABLEOBJECT_API FVector4f GetVectorParameterSelectedOption(const FString& VectorParamName);
+	
+	/** Sets the Vector value "VectorValue" of a Vector parameter with name "VectorParamName". */
 	UFUNCTION(BlueprintCallable, Category = CustomizableObjectInstance)
-	CUSTOMIZABLEOBJECT_API void SetVectorParameterSelectedOption(const FString& VectorParamName, const FLinearColor& VectorValue);
+	CUSTOMIZABLEOBJECT_API void SetVectorParameterSelectedOption(const FString& VectorParamName, const FVector4f& VectorValue);
 
 	/** Gets the value of a transform parameter with name "TransformParamName". */
 	UFUNCTION(BlueprintCallable, Category = CustomizableObjectInstance)
@@ -629,11 +641,19 @@ public:
 	/** Return true if the Material Parameter exists. */
 	UFUNCTION(BlueprintCallable, Category = CustomizableObjectInstance)
 	CUSTOMIZABLEOBJECT_API bool ContainsMaterialParameter(const FString& ParameterName) const;
+
+	/** Return true if the External Type Parameter exists. */
+	UFUNCTION(BlueprintCallable, Category = CustomizableObjectInstance)
+	CUSTOMIZABLEOBJECT_API bool ContainsExternalTypeParameter(const FString& ParameterName) const;
 	
 	/** Return true if the Bool Parameter exists. */
 	UFUNCTION(BlueprintCallable, Category = CustomizableObjectInstance)
 	CUSTOMIZABLEOBJECT_API bool ContainsBoolParameter(const FString& ParameterName) const;
 
+	/** Return true if the Color Parameter exists. */
+	UFUNCTION(BlueprintCallable, Category = CustomizableObjectInstance)
+	CUSTOMIZABLEOBJECT_API bool ContainsColorParameter(const FString& ParameterName) const;
+	
 	/** Return true if the Vector Parameter exists. */
 	UFUNCTION(BlueprintCallable, Category = CustomizableObjectInstance)
 	CUSTOMIZABLEOBJECT_API bool ContainsVectorParameter(const FString& ParameterName) const;
@@ -645,7 +665,6 @@ public:
 	/** Return true if the Transform Parameter exists. */
 	UFUNCTION(BlueprintCallable, Category = CustomizableObjectInstance)
 	CUSTOMIZABLEOBJECT_API bool ContainsTransformParameter(const FString& ParameterName) const;
-	
 
 	/** Increases the range of values of the integer with ParamName and returns the index of the new integer value, -1 otherwise.
 	  * The added value is initialized with the first integer option and is the last one of the range. */
diff --git a/Engine/Plugins/Mutable/Source/CustomizableObject/Public/MuCO/CustomizableObjectInstanceDescriptor.h b/Engine/Plugins/Mutable/Source/CustomizableObject/Public/MuCO/CustomizableObjectInstanceDescriptor.h
--- a/Engine/Plugins/Mutable/Source/CustomizableObject/Public/MuCO/CustomizableObjectInstanceDescriptor.h
+++ b/Engine/Plugins/Mutable/Source/CustomizableObject/Public/MuCO/CustomizableObjectInstanceDescriptor.h
@@ -84,215 +84,139 @@ struct FCustomizableObjectInstanceDescriptor
 	// FParameters
 	// ------------------------------------------------------------
 	
-	/** Return true if there are any parameters. */
 	UE_API bool HasAnyParameters() const;
 
-	/** Gets the value of the int parameter with name "ParamName". */
 	UE_API const FString& GetIntParameterSelectedOption(const FString& ParamName, int32 RangeIndex = -1) const;
 
-	/** Sets the selected option of an int parameter by the option's name. */
 	UE_API void SetIntParameterSelectedOption(int32 IntParamIndex, const FString& SelectedOption, int32 RangeIndex = -1);
 
-	/** Sets the selected option of an int parameter, by the option's name */
 	UE_API void SetIntParameterSelectedOption(const FString& ParamName, const FString& SelectedOptionName, int32 RangeIndex = -1);
 
-	/** Gets the value of a float parameter with name "FloatParamName". */
 	UE_API float GetFloatParameterSelectedOption(const FString& FloatParamName, int32 RangeIndex = -1) const;
 
-	/** Sets the float value "FloatValue" of a float parameter with index "FloatParamIndex". */
 	UE_API void SetFloatParameterSelectedOption(const FString& FloatParamName, float FloatValue, int32 RangeIndex = -1);
 
-	/** Gets the value of a texture parameter with name "TextureParamName". */
 	UE_API UTexture* GetTextureParameterSelectedOption(const FString& TextureParamName, int32 RangeIndex) const;
 
-	/** Sets the texture value "TextureValue" of a texture parameter with index "TextureParamIndex". */
 	UE_API void SetTextureParameterSelectedOption(const FString& TextureParamName, UTexture* TextureValue, int32 RangeIndex);
 	
-	/** Gets the value of a Skeletal Mesh parameter with name "SkeletalMeshParamName". */
 	UE_API USkeletalMesh* GetSkeletalMeshParameterSelectedOption(const FString& SkeletalMeshParamName, int32 RangeIndex = -1) const;
 
-	/** Sets the skeletalMesh value "SkeletalMeshValue" of a skeletalMesh parameter with index "SkeletalMeshParamIndex". */
 	UE_API void SetSkeletalMeshParameterSelectedOption(const FString& SkeletalMeshParamName, USkeletalMesh* SkeletalMeshValue, int32 RangeIndex = -1);
 
-	/** Gets the value of a Material parameter with name "MaterialMeshParamName". */
 	UE_API UMaterialInterface* GetMaterialParameterSelectedOption(const FString& MaterialMeshParamName, int32 RangeIndex = -1) const;
 
-	/** Sets the material value "MaterialParamName" of a Material parameter with index "MaterialMeshParamName". */
-	UE_API void SetMaterialParameterSelectedOption(const FString& MaterialParamName, UMaterialInterface* MaterialValue, int32 RangeIndex = -1);
+	UE_API FInstancedStruct GetExternalTypeParameterSelectedOption(const FString& ExternalTypeParamName, int32 RangeIndex = -1) const;
 	
-	/** Gets the value of a color parameter with name "ColorParamName". */
-	UE_API FLinearColor GetColorParameterSelectedOption(const FString& ColorParamName) const;
-
-	/** Sets the color value "ColorValue" of a color parameter with index "ColorParamIndex". */
-	UE_API void SetColorParameterSelectedOption(const FString& ColorParamName, const FLinearColor& ColorValue);
+	UE_API void SetExternalTypeParameterSelectedOption(const FString& ParamName, const FInstancedStruct& Value, int32 RangeIndex = -1);
+	
+	UE_API void SetMaterialParameterSelectedOption(const FString& MaterialParamName, UMaterialInterface* MaterialValue, int32 RangeIndex = -1);
 
-	/** Gets the value of a transform parameter with name "TransformParamName". */
 	UE_API FTransform GetTransformParameterSelectedOption(const FString& TransformParamName) const;
 
-	/** Sets the transform value "TransformValue" of a transform parameter with name "TransformParamName". */
 	UE_API void SetTransformParameterSelectedOption(const FString& TransformParamName, const FTransform& TransformValue);
 
-	/** Gets the value of the bool parameter with name "BoolParamName". */
 	UE_API bool GetBoolParameterSelectedOption(const FString& BoolParamName) const;
 
-	/** Sets the bool value "BoolValue" of a bool parameter with name "BoolParamName". */
 	UE_API void SetBoolParameterSelectedOption(const FString& BoolParamName, bool BoolValue);
 
-	/** Sets the vector value "VectorValue" of a bool parameter with index "VectorParamIndex". */
-	UE_API void SetVectorParameterSelectedOption(const FString& VectorParamName, const FLinearColor& VectorValue);
-
-	/** Sets the projector values of a projector parameter with index "ProjectorParamIndex". */
+	UE_API FVector4f GetVectorParameterSelectedOption(const FString& BoolParamName) const;
+	
+	UE_API void SetVectorParameterSelectedOption(const FString& VectorParamName, const FVector4f& VectorValue);
+	
 	UE_API void SetProjectorValue(const FString& ProjectorParamName,
 		const FVector& Pos, const FVector& Direction, const FVector& Up, const FVector& Scale,
 		float Angle,
 		int32 RangeIndex = -1);
 
-	/** Set only the projector position. */
 	UE_API void SetProjectorPosition(const FString& ProjectorParamName, const FVector& Pos, int32 RangeIndex = -1);
 
-	/** Set only the projector direction. */
 	UE_API void SetProjectorDirection(const FString& ProjectorParamName, const FVector& Direction, int32 RangeIndex = -1);
 	
-	/** Set only the projector up vector. */
 	UE_API void SetProjectorUp(const FString& ProjectorParamName, const FVector& Up, int32 RangeIndex = -1);
 
-	/** Set only the projector scale. */
 	UE_API void SetProjectorScale(const FString& ProjectorParamName, const FVector& Scale, int32 RangeIndex = -1);
 
-	/** Set only the cylindrical projector angle. */
 	UE_API void SetProjectorAngle(const FString& ProjectorParamName, float Angle, int32 RangeIndex = -1);
 	
-	/** Get the projector values of a projector parameter with index "ProjectorParamIndex". */
 	UE_API void GetProjectorValue(const FString& ProjectorParamName,
 		FVector& OutPos, FVector& OutDirection, FVector& OutUp, FVector& OutScale,
 		float& OutAngle, ECustomizableObjectProjectorType& OutType,
 		int32 RangeIndex = -1) const;
 
-	/** Float version. See GetProjectorValue. */
 	UE_API void GetProjectorValueF(const FString& ProjectorParamName,
 		FVector3f& OutPos, FVector3f& OutDirection, FVector3f& OutUp, FVector3f& OutScale,
 		float& OutAngle, ECustomizableObjectProjectorType& OutType,
 		int32 RangeIndex = -1) const;
 
-	/** Get the current projector position for the parameter with the given name. */
 	UE_API FVector GetProjectorPosition(const FString & ParamName, int32 RangeIndex = -1) const;
 
-	/** Get the current projector direction vector for the parameter with the given name. */
 	UE_API FVector GetProjectorDirection(const FString & ParamName, int32 RangeIndex = -1) const;
 
-	/** Get the current projector up vector for the parameter with the given name. */
 	UE_API FVector GetProjectorUp(const FString & ParamName, int32 RangeIndex = -1) const;
 
-	/** Get the current projector scale for the parameter with the given name. */
 	UE_API FVector GetProjectorScale(const FString & ParamName, int32 RangeIndex = -1) const;
 
-	/** Get the current cylindrical projector angle for the parameter with the given name. */
 	UE_API float GetProjectorAngle(const FString& ParamName, int32 RangeIndex = -1) const;
 
-	/** Get the current projector type for the parameter with the given name. */
 	UE_API ECustomizableObjectProjectorType GetProjectorParameterType(const FString& ParamName, int32 RangeIndex = -1) const;
 
-	/** Get the current projector for the parameter with the given name. */
 	UE_API FCustomizableObjectProjector GetProjector(const FString& ParamName, int32 RangeIndex) const;
 	
-	/** Finds the parameter with name ParamName in the array of its type, returns the index if found, INDEX_NONE otherwise. */
 	UE_API int32 FindTypedParameterIndex(const FString& ParamName, EMutableParameterType Type) const;
 
-	// Parameter Ranges
-
-	/** Gets the range of values of the projector with ParamName, returns -1 if the parameter does not exist. */
 	UE_API int32 GetProjectorValueRange(const FString& ParamName) const;
 
-	/** Gets the range of values of the int with ParamName, returns -1 if the parameter does not exist. */
 	UE_API int32 GetIntValueRange(const FString& ParamName) const;
 
-	/** Gets the range of values of the float with ParamName, returns -1 if the parameter does not exist. */
 	UE_API int32 GetFloatValueRange(const FString& ParamName) const;
 
-	/** Gets the range of values of the texture with ParamName, returns -1 if the parameter does not exist. */
 	UE_API int32 GetTextureValueRange(const FString& ParamName) const;
 
-	/** Increases the range of values of the integer with ParamName, returns the index of the new integer value, -1 otherwise.
-	 * The added value is initialized with the first integer option and is the last one of the range. */
 	UE_API int32 AddValueToIntRange(const FString& ParamName);
 
-	/** Increases the range of values of the float with ParamName, returns the index of the new float value, -1 otherwise.
-	 * The added value is initialized with 0.5f and is the last one of the range. */
 	UE_API int32 AddValueToFloatRange(const FString& ParamName);
 
-	/** Increases the range of values of the float with ParamName, returns the index of the new float value, -1 otherwise. 
-	 * The added value is not initialized. */
 	UE_API int32 AddValueToTextureRange(const FString& ParamName);
 
-	/** Increases the range of values of the projector with ParamName, returns the index of the new projector value, -1 otherwise.
-	 * The added value is initialized with the default projector as set up in the editor and is the last one of the range. */
 	UE_API int32 AddValueToProjectorRange(const FString& ParamName);
 
-	/** Remove the FRangeIndex element of the integer range of values from the parameter ParamName. If Range index is -1, removes the last element.
-		Returns the index of the last valid integer, -1 if no values left. */
 	UE_API int32 RemoveValueFromIntRange(const FString& ParamName, int32 RangeIndex = -1);
 
-	/** Remove the FRangeIndex element of the float range of values from the parameter ParamName. If Range index is -1, removes the last element.
-		Returns the index of the last valid float, -1 if no values left. */
 	UE_API int32 RemoveValueFromFloatRange(const FString& ParamName, int32 RangeIndex = -1);
 
-	/** Remove the last of the texture range of values from the parameter ParamName, returns the index of the last valid float, -1 if no values left. */
 	UE_API int32 RemoveValueFromTextureRange(const FString& ParamName);
 
-	/** Remove the FRangeIndex element of the texture range of values from the parameter ParamName, returns the index of the last valid float, -1 if no values left. */
 	UE_API int32 RemoveValueFromTextureRange(const FString& ParamName, int32 RangeIndex);
-
-	/** Remove the FRangeIndex element of the projector range of values from the parameter ParamName. If Range index is -1, removes the last element.
-		Returns the index of the last valid projector, -1 if no values left.
-	*/
+	
 	UE_API int32 RemoveValueFromProjectorRange(const FString& ParamName, int32 RangeIndex = -1);
 
-	// ------------------------------------------------------------
-   	// States
-   	// ------------------------------------------------------------
-
-	/** Get the current optimization state. */
 	UE_API int32 GetState() const;
 
-	/** Get the current optimization state. */	
 	UE_API FString GetCurrentState() const;
 
-	/** Set the current optimization state. */
 	UE_API void SetState(int32 InState);
 
-	/** Set the current optimization state. */
 	UE_API void SetCurrentState(const FString& StateName);
 
-	// ------------------------------------------------------------
-
 	UE_API void SetRandomValues();
 	
 	UE_API void SetRandomValuesFromStream(const FRandomStream& InStream);
 
 	UE_API void SetDefaultValue(int32 ParamIndex);
+
 	UE_API void SetDefaultValues();
 	
-	// ------------------------------------------------------------
-	// Multilayer Projectors
-	// ------------------------------------------------------------
-
-	/** @return true if ParamName belongs to a multilayer projector parameter. */
 	UE_API bool IsMultilayerProjector(const FString& ParamName) const;
 
-	// Layers
-	/** @return number of layers of the projector with name ParamName,-1 if invalid or not found. */
 	UE_API int32 NumProjectorLayers(const FName& ParamName) const;
 
-	/** Creates a new layer for the multilayer projector with name ParamName. */
 	UE_API void CreateLayer(const FName& ParamName, int32 Index);
 
-	/** Removes the layer at Index from the multilayer projector with name ParamName. */
 	UE_API void RemoveLayerAt(const FName& ParamName, int32 Index);
 
-	/** @return copy of the layer at Index for the multilayer projector with name ParamName. */
 	UE_API FMultilayerProjectorLayer GetLayer(const FName& ParamName, int32 Index) const;
 
-	/** Updates the parameters of the layer at Index from the multilayer projector with name ParamName. */
 	UE_API void UpdateLayer(const FName& ParamName, int32 Index, const FMultilayerProjectorLayer& Layer);
 
 	/** Return a Mutable Core object containing all parameters. */
@@ -330,6 +254,9 @@ struct FCustomizableObjectInstanceDescriptor
 	UPROPERTY()
 	TArray<FCustomizableObjectTransformParameterValue> TransformParameters;
 
+	UPROPERTY()
+	TArray<FCustomizableObjectExternalTypeParameterValue> ExternalTypeParameters;
+
 	/** Mutable parameters optimization state. Transient UProperty to make it translatable. */
 	UPROPERTY(Transient)
 	int32 State = 0;
diff --git a/Engine/Plugins/Mutable/Source/CustomizableObject/Public/MuCO/CustomizableObjectParameterTypeDefinitions.h b/Engine/Plugins/Mutable/Source/CustomizableObject/Public/MuCO/CustomizableObjectParameterTypeDefinitions.h
--- a/Engine/Plugins/Mutable/Source/CustomizableObject/Public/MuCO/CustomizableObjectParameterTypeDefinitions.h
+++ b/Engine/Plugins/Mutable/Source/CustomizableObject/Public/MuCO/CustomizableObjectParameterTypeDefinitions.h
@@ -2,6 +2,8 @@
 
 #pragma once
 
+#include "StructUtils/InstancedStruct.h"
+
 #include "CustomizableObjectParameterTypeDefinitions.generated.h"
 
 class UTexture;
@@ -22,6 +24,7 @@ enum class EMutableParameterType : uint8
 	Transform		UMETA(DisplayName = "Transform"),
 	SkeletalMesh	UMETA(DisplayName = "SkeletalMesh"),
 	Material		UMETA(DisplayName = "Material"),
+	InstancedStruct UMETA(DisplayName = "InstancedStruct"),
 };
 
 
@@ -407,3 +410,27 @@ inline uint32 GetTypeHash(const FCustomizableObjectProjectorParameterValue& Key)
 	
 	return Hash;
 }
+
+
+USTRUCT(BlueprintType)
+struct FCustomizableObjectExternalTypeParameterValue
+{
+	GENERATED_USTRUCT_BODY()
+
+	inline static const FInstancedStruct DEFAULT_PARAMETER_VALUE = {};
+
+	UPROPERTY(Category = CustomizableObjectVectorParameterValue, VisibleAnywhere)
+	FString ParameterName;
+
+	UPROPERTY(Category = CustomizableObjectVectorParameterValue, VisibleAnywhere)
+	FInstancedStruct ParameterValue;
+
+	UPROPERTY(Category = CustomizableObjectVectorParameterValue, VisibleAnywhere)
+	FGuid Id;
+
+	UPROPERTY(Category = CustomizableObjectVectorParameterValue, VisibleAnywhere)
+	TArray<FInstancedStruct> ParameterRangeValues;
+};
+
+
+uint32 GetTypeHash(const FCustomizableObjectExternalTypeParameterValue& Key);
diff --git a/Engine/Plugins/Mutable/Source/CustomizableObjectEditor/CustomizableObjectEditor.Build.cs b/Engine/Plugins/Mutable/Source/CustomizableObjectEditor/CustomizableObjectEditor.Build.cs
--- a/Engine/Plugins/Mutable/Source/CustomizableObjectEditor/CustomizableObjectEditor.Build.cs
+++ b/Engine/Plugins/Mutable/Source/CustomizableObjectEditor/CustomizableObjectEditor.Build.cs
@@ -70,8 +70,8 @@ public class CustomizableObjectEditor : ModuleRules
 				"TextureBuildUtilities",
 				"ImageCore",
 				"KismetWidgets",
-
-				"MeshDescription"
+				"MeshDescription",
+				"StructUtilsEditor",
 			}
 		);
 
diff --git a/Engine/Plugins/Mutable/Source/CustomizableObjectEditor/Private/MuCOE/CustomizableInstanceDetails.cpp b/Engine/Plugins/Mutable/Source/CustomizableObjectEditor/Private/MuCOE/CustomizableInstanceDetails.cpp
--- a/Engine/Plugins/Mutable/Source/CustomizableObjectEditor/Private/MuCOE/CustomizableInstanceDetails.cpp
+++ b/Engine/Plugins/Mutable/Source/CustomizableObjectEditor/Private/MuCOE/CustomizableInstanceDetails.cpp
@@ -21,6 +21,7 @@
 #include "MuCOE/CustomizableObjectEditor.h"
 #include "MuCOE/UnrealEditorPortabilityHelpers.h"
 #include "AssetThumbnail.h"
+#include "IStructureDetailsView.h"
 #include "Slate/DeferredCleanupSlateBrush.h"
 #include "SSearchableComboBox.h"
 #include "MuCO/LoadUtils.h"
@@ -39,6 +40,8 @@
 #include "Widgets/Layout/SUniformGridPanel.h"
 #include "Engine/Texture.h"
 #include "Engine/SkeletalMesh.h"
+#include "Nodes/ExternalTypeCustomization.h"
+#include "UObject/StructOnScope.h"
 
 class UObject;
 
@@ -1084,6 +1087,31 @@ TSharedRef<SWidget> FCustomizableInstanceDetails::GenerateParameterWidget(const
 				.OnObjectChanged(this, &FCustomizableInstanceDetails::OnMaterialParameterSelected, ParamName)
 				.ObjectPath(ObjectPath);
 		}
+		case EMutableParameterType::InstancedStruct:
+		{
+			FInstancedStruct InstancedStruct = CustomInstance->GetExternalTypeParameterSelectedOption(ParamName);
+
+			const UScriptStruct* ScriptStruct = InstancedStruct.GetScriptStruct();
+				
+			TSharedPtr<FStructOnScope> StructOnScope = MakeShared<FStructOnScope>(ScriptStruct);
+
+			ScriptStruct->CopyScriptStruct(StructOnScope->GetStructMemory(), InstancedStruct.GetMemory());
+
+			FPropertyEditorModule& PropertyEditorModule = FModuleManager::LoadModuleChecked<FPropertyEditorModule>("PropertyEditor");
+
+			FDetailsViewArgs DetailViewArgs;
+			DetailViewArgs.bAllowSearch = false;
+			DetailViewArgs.bShowScrollBar = true;
+			DetailViewArgs.NameAreaSettings = FDetailsViewArgs::HideNameArea;
+
+			FStructureDetailsViewArgs StructureDetailsViewArgs;
+			StructureDetailsViewArgs.bShowObjects = false;
+			
+			TSharedRef<IStructureDetailsView> Detailsview = PropertyEditorModule.CreateStructureDetailView(DetailViewArgs, StructureDetailsViewArgs, StructOnScope);
+			Detailsview->GetOnFinishedChangingPropertiesDelegate().AddSP(this, &FCustomizableInstanceDetails::OnExternalTypeParameterChanged, StructOnScope, ParamName, -1);
+
+			return Detailsview->GetWidget().ToSharedRef();
+		}
 		case EMutableParameterType::Transform:
 		{
 			return GenerateTransformWidget(ParamName);
@@ -1456,7 +1484,7 @@ TSharedRef<SWidget> FCustomizableInstanceDetails::GenerateColorWidget(const FStr
 
 FLinearColor FCustomizableInstanceDetails::GetColorParameterValue(FString ParamName) const
 {
-	return CustomInstance->GetColorParameterSelectedOption(ParamName);
+	return CustomInstance->GetVectorParameterSelectedOption(ParamName);
 }
 
 
@@ -1485,7 +1513,7 @@ void FCustomizableInstanceDetails::OnSetColorFromColorPicker(FLinearColor NewCol
 {
 	FScopedTransaction LocalTransaction(FText::Format(LOCTEXT("SetColorParameter", "Set Color Parameter: {0}"), FText::FromString(PickerParamName)));
 	CustomInstance->Modify();
-	CustomInstance->SetColorParameterSelectedOption(PickerParamName, NewColor);
+	CustomInstance->SetVectorParameterSelectedOption(PickerParamName, NewColor);
 	UpdateInstance();
 }
 
@@ -2319,6 +2347,23 @@ TSharedRef<SWidget> FCustomizableInstanceDetails::OnGenerateWidgetProjectorParam
 }
 
 
+void FCustomizableInstanceDetails::OnExternalTypeParameterChanged(const FPropertyChangedEvent& PropertyChangedEvent, TSharedPtr<FStructOnScope> StructOnScope, FString ParamName, const int32 RangeIndex)
+{
+	if (const UScriptStruct* ScriptStruct = Cast<UScriptStruct>(StructOnScope->GetStructPtr()))
+	{
+		CustomInstance->Modify();
+
+		FInstancedStruct InstancedStruct;
+		InstancedStruct.InitializeAs(ScriptStruct);
+		ScriptStruct->CopyScriptStruct(InstancedStruct.GetMutableMemory(), StructOnScope->GetStructMemory());
+
+		CustomInstance->SetExternalTypeParameterSelectedOption(ParamName, InstancedStruct, RangeIndex);
+
+		UpdateInstance();
+	}
+}
+
+
 FReply FCustomizableInstanceDetails::OnProjectorLayerAdded(FString ParamName)
 {
 	FScopedTransaction LocalTransaction(LOCTEXT("AddProjectorLayer", "Add Projector Layer"));
diff --git a/Engine/Plugins/Mutable/Source/CustomizableObjectEditor/Private/MuCOE/CustomizableInstanceDetails.h b/Engine/Plugins/Mutable/Source/CustomizableObjectEditor/Private/MuCOE/CustomizableInstanceDetails.h
--- a/Engine/Plugins/Mutable/Source/CustomizableObjectEditor/Private/MuCOE/CustomizableInstanceDetails.h
+++ b/Engine/Plugins/Mutable/Source/CustomizableObjectEditor/Private/MuCOE/CustomizableInstanceDetails.h
@@ -6,6 +6,7 @@
 #include "UObject/ObjectPtr.h"
 #include "UObject/WeakObjectPtr.h"
 #include "GameplayTagContainer.h"
+#include "Templates/SharedPointer.h"
 
 class FDetailWidgetRow;
 class FReply;
@@ -18,6 +19,7 @@ class IDetailLayoutBuilder;
 class SWidget;
 class UCustomizableObject;
 class UCustomizableObjectInstance;
+class FStructOnScope;
 
 enum class ECheckBoxState : uint8;
 namespace ESelectInfo { enum Type : int; }
@@ -137,7 +139,6 @@ private:
 	void OnTransformParameterChanged(FTransform NewTransform, const FString ParamName);
 	void OnTransformParameterCommitted(FTransform NewTransform, ETextCommit::Type Type, const FString ParamName);
 
-
 	// Bool FParameters Functions
 	TSharedRef<SWidget> GenerateBoolWidget(const FString& ParamName);
 	ECheckBoxState GetBoolParameterValue(const FString ParamName) const;
@@ -157,6 +158,9 @@ private:
 	TSharedRef<SWidget> MakeTextureComboEntryWidget(TSharedPtr<FString> InItem) const;
 	TSharedRef<SWidget> OnGenerateWidgetProjectorParameter(TSharedPtr<FString> InItem) const;
 
+	// External Type FParmeter Functions
+	void OnExternalTypeParameterChanged(const FPropertyChangedEvent& PropertyChangedEvent, TSharedPtr<FStructOnScope> StructOnScope, FString ParamName, const int32 RangeIndex);
+	
 	// Parameter Functions
 	FReply OnCopyAllParameters();
 	FReply OnPasteAllParameters();
diff --git a/Engine/Plugins/Mutable/Source/CustomizableObjectEditor/Private/MuCOE/CustomizableObjectCompiler.cpp b/Engine/Plugins/Mutable/Source/CustomizableObjectEditor/Private/MuCOE/CustomizableObjectCompiler.cpp
--- a/Engine/Plugins/Mutable/Source/CustomizableObjectEditor/Private/MuCOE/CustomizableObjectCompiler.cpp
+++ b/Engine/Plugins/Mutable/Source/CustomizableObjectEditor/Private/MuCOE/CustomizableObjectCompiler.cpp
@@ -964,6 +964,9 @@ void FCustomizableObjectCompiler::CompileInternal(bool bAsync)
 		ModelResources->TextureParameterDefaultValues = MoveTemp(GenerationContext.TextureParameterDefaultValues);
 		ModelResources->SkeletalMeshParameterDefaultValues = MoveTemp(GenerationContext.SkeletalMeshParameterDefaultValues);
 		ModelResources->MaterialParameterDefaultValues = MoveTemp(GenerationContext.MaterialParameterDefaultValues);
+		ModelResources->ExternalTypeParameterDefaultValues = MoveTemp(GenerationContext.ExternalTypeParameterDefaultValues);
+
+		ModelResources->ExternalTypeParameterTypes = MoveTemp(GenerationContext.ExternalTypeParameterTypes);
 
 		ModelResources->AnimBPs = MoveTemp(GenerationContext.AnimBPAssets);
 
diff --git a/Engine/Plugins/Mutable/Source/CustomizableObjectEditor/Private/MuCOE/CustomizableObjectEditorModule.cpp b/Engine/Plugins/Mutable/Source/CustomizableObjectEditor/Private/MuCOE/CustomizableObjectEditorModule.cpp
--- a/Engine/Plugins/Mutable/Source/CustomizableObjectEditor/Private/MuCOE/CustomizableObjectEditorModule.cpp
+++ b/Engine/Plugins/Mutable/Source/CustomizableObjectEditor/Private/MuCOE/CustomizableObjectEditorModule.cpp
@@ -111,9 +111,11 @@
 #include "Nodes/CustomizableObjectNodeComponentMeshDetails.h"
 #include "MuCOE/GenerateMutableSource/GenerateMutableSourceTable.h"
 #include "MuR/External/Operation.h"
-#include "Nodes/CONodeExternal.h"
-#include "Nodes/CONodeExternalDetails.h"
+#include "Nodes/CONodeExternalOperation.h"
+#include "Nodes/CONodeExternalOperationDetails.h"
+#include "Nodes/CONodeExternalTypeParameter.h"
 #include "Nodes/CustomizableObjectNodeTextureParameter.h"
+#include "Nodes/ExternalTypeCustomization.h"
 #include "PhysicsEngine/PhysicsAsset.h"
 
 
@@ -280,7 +282,7 @@ void FCustomizableObjectEditorModule::StartupModule()
 	RegisterCustomDetails(PropertyModule, UCustomizableObjectNodeMacroInstance::StaticClass(), FOnGetDetailCustomizationInstance::CreateStatic(&FCustomizableObjectNodeMacroInstanceDetails::MakeInstance));
 	RegisterCustomDetails(PropertyModule, UCustomizableObjectNodeParameter::StaticClass(), FOnGetDetailCustomizationInstance::CreateStatic(&FCustomizableObjectNodeParameterDetails::MakeInstance));
 	RegisterCustomDetails(PropertyModule, UCONodeMaterialBreak::StaticClass(), FOnGetDetailCustomizationInstance::CreateStatic(&FCustomizableObjectNodeDetails::MakeInstance));
-	RegisterCustomDetails(PropertyModule, UCONodeExternal::StaticClass(), FOnGetDetailCustomizationInstance::CreateStatic(&FCONodeExternalDetails::MakeInstance));
+	RegisterCustomDetails(PropertyModule, UCONodeExternalOperation::StaticClass(), FOnGetDetailCustomizationInstance::CreateStatic(&FCONodeExternalOperationDetails::MakeInstance));
 
 	// Other Objects
 	RegisterCustomDetails(PropertyModule, UCustomizableObject::StaticClass(), FOnGetDetailCustomizationInstance::CreateStatic(&FCustomizableObjectDetails::MakeInstance));
@@ -302,6 +304,7 @@ void FCustomizableObjectEditorModule::StartupModule()
 	PropertyModule.RegisterCustomPropertyTypeLayout(FCustomizableObjectTextureVariation::StaticStruct()->GetFName(), FOnGetPropertyTypeCustomizationInstance::CreateStatic(&FCustomizableObjectVariationCustomization::MakeInstance));
 	PropertyModule.RegisterCustomPropertyTypeLayout(FTableNodeCompilationFilter::StaticStruct()->GetFName(), FOnGetPropertyTypeCustomizationInstance::CreateStatic(&FCustomizableObjectNodeTableCompilationFilterEditor::MakeInstance));
 	PropertyModule.RegisterCustomPropertyTypeLayout(COVariableUStructName, FOnGetPropertyTypeCustomizationInstance::CreateStatic(&FCOVariableCustomzation::MakeInstance));
+	PropertyModule.RegisterCustomPropertyTypeLayout("InstancedStruct", FOnGetPropertyTypeCustomizationInstance::CreateStatic(&FExternalTypeCustomization::MakeInstance), MakeShared<FExternalTypeTypeIdentifier>());
 
 	PropertyModule.NotifyCustomizationModuleChanged();
 
diff --git a/Engine/Plugins/Mutable/Source/CustomizableObjectEditor/Private/MuCOE/EdGraphSchema_CustomizableObject.cpp b/Engine/Plugins/Mutable/Source/CustomizableObjectEditor/Private/MuCOE/EdGraphSchema_CustomizableObject.cpp
--- a/Engine/Plugins/Mutable/Source/CustomizableObjectEditor/Private/MuCOE/EdGraphSchema_CustomizableObject.cpp
+++ b/Engine/Plugins/Mutable/Source/CustomizableObjectEditor/Private/MuCOE/EdGraphSchema_CustomizableObject.cpp
@@ -117,8 +117,11 @@
 #include "MuR/External/Operation.h"
 #include "MuR/External/TextureAdapter.h"
 #include "MuR/External/VectorAdapter.h"
-#include "Nodes/CONodeExternal.h"
+#include "Nodes/CONodeExternalOperation.h"
+#include "Nodes/CONodeExternalTypeParameter.h"
+#include "Nodes/ExternalTypeCustomization.h"
 #include "Nodes/CONodeSchema.h"
+#include "UObject/Class.h"
 
 #include UE_INLINE_GENERATED_CPP_BY_NAME(EdGraphSchema_CustomizableObject)
 
@@ -657,6 +660,8 @@ void UEdGraphSchema_CustomizableObject::GetGraphContextActions(FGraphContextMenu
 	}
 	
 	// Extensions
+	TArray<const UScriptStruct*> ExternalTypes;
+
 	{
 		for (const TStrongObjectPtr<const UScriptStruct>& OperationType : ICustomizableObjectModulePrivate::Get().GetRegisteredExternalOperations())
 		{
@@ -665,11 +670,40 @@ void UEdGraphSchema_CustomizableObject::GetGraphContextActions(FGraphContextMenu
 
 			if (const UE::Mutable::FExternalOperation* Operation = OperationInstancedStruct.GetPtr<UE::Mutable::FExternalOperation>())
 			{
-				UCONodeExternal* Node = ContextMenuBuilder.CreateTemplateNode<UCONodeExternal>();
-				Node->OperationInstancedStruct = OperationInstancedStruct;
+				// External Operation
+				{
+					UCONodeExternalOperation* Node = ContextMenuBuilder.CreateTemplateNode<UCONodeExternalOperation>();
+					Node->OperationInstancedStruct = OperationInstancedStruct;
+
+					AddNewNodeActionFiltered(Node, ContextMenuBuilder, GetPinCategoryName(Operation->GetOutput().Value).ToString(), GeneralGrouping, Filter);
+
+					ExternalTypes.AddUnique(Operation->GetOutput().Value);
+
+					for (const TPair<FName, const UScriptStruct*>& Pair : Operation->GetInputs())
+					{
+						ExternalTypes.AddUnique(Pair.Value);
+					}
+				}
+			}
+		}
+
+		// Parameter
+		for (const UScriptStruct* ScriptStruct : ExternalTypes)
+		{
+			if (IsTypeNative(ScriptStruct))
+			{
+				continue;
+			}
 
-				AddNewNodeActionFiltered(Node, ContextMenuBuilder, GetPinCategoryName(Operation->GetOutput().Value).ToString(), GeneralGrouping, Filter);
+			if (ScriptStruct->GetCppStructOps()->GetCapabilities().IsAbstract)
+			{
+				continue;
 			}
+				
+			UCONodeExternalTypeParameter* Node = ContextMenuBuilder.CreateTemplateNode<UCONodeExternalTypeParameter>();
+			Node->DefaultValue.InitializeAs(ScriptStruct);
+
+			AddNewNodeActionFiltered(Node, ContextMenuBuilder, GetPinCategoryName(ScriptStruct).ToString(), GeneralGrouping, Filter);
 		}
 	}
 	
@@ -677,33 +711,18 @@ void UEdGraphSchema_CustomizableObject::GetGraphContextActions(FGraphContextMenu
 	{
 		TArray<FName> PinTypes({ PC_MeshSection, PC_Modifier, PC_Mesh, PC_Texture, PC_Projector, PC_GroupProjector, PC_Color, PC_Float, PC_Bool, PC_Enum, PC_Transform, PC_Stack, PC_PassthroughTexture, PC_Material, PC_PoseAsset, PC_Component });
 
-		// Add pin types from extensions
+		// Extensions
 		for (const FRegisteredCustomizableObjectPinType& PinType : ICustomizableObjectModule::Get().GetExtendedPinTypes())
 		{
 			PinTypes.AddUnique(PinType.PinType.Name);
 		}
 
-		for (const TStrongObjectPtr<const UScriptStruct>& OperationType : ICustomizableObjectModulePrivate::Get().GetRegisteredExternalOperations())
+		// External Types
+		for (const UScriptStruct* ScriptStruct : ExternalTypes)
 		{
+			const FName Category = GetPinCategoryName(ScriptStruct);
+			PinTypes.AddUnique(Category);
 			FInstancedStruct OperationInstancedStruct;
-			OperationInstancedStruct.InitializeAs(OperationType.Get());
-			
-			if (const UE::Mutable::FExternalOperation* Operation = OperationInstancedStruct.GetPtr<UE::Mutable::FExternalOperation>())
-			{
-				// Inputs
-				for (TPair<FName, const UScriptStruct*> Pair : Operation->GetInputs())
-				{
-					const FName Category = UEdGraphSchema_CustomizableObject::GetPinCategoryName(Pair.Value);
-					PinTypes.AddUnique(Category);
-				}
-
-				// Output
-				{
-					TPair<FName, const UScriptStruct*> Pair = Operation->GetOutput();
-					const FName Category = UEdGraphSchema_CustomizableObject::GetPinCategoryName(Pair.Value);
-					PinTypes.AddUnique(Category);
-				}
-			}
 		}
 
 		for (const FName& PinCategory : PinTypes)
@@ -910,7 +929,7 @@ FLinearColor UEdGraphSchema_CustomizableObject::GetPinTypeColor(const FName& Typ
 		}
 	}
 
-	const uint32 Hash = GetTypeHash(TypeString);
+	const uint32 Hash = GetTypeHash(TypeString.ToString());
 	return FLinearColor::MakeRandomSeededColor(Hash);
 }
 
@@ -1550,9 +1569,7 @@ FName UEdGraphSchema_CustomizableObject::GetPinCategoryName(const FName& PinCate
 			}
 		}
 
-		// Need to fail gracefully here in case a plugin that was active when this graph was
-		// created is no longer loaded.
-		return TEXT("Unknown");
+		return PinCategory;
 	}
 }
 
@@ -1682,9 +1699,20 @@ FText UEdGraphSchema_CustomizableObject::GetPinCategoryFriendlyName(const FName&
 			}
 		}
 
-		// Need to fail gracefully here in case a plugin that was active when this graph was
-		// created is no longer loaded.
-		return LOCTEXT("Unknown_Pin_Category", "Unknown");
+		return {};
+	}
+}
+
+
+FText UEdGraphSchema_CustomizableObject::GetPinCategoryFriendlyName(const UScriptStruct* Type)
+{
+	if (IsTypeNative(Type))
+	{
+		return GetPinCategoryFriendlyName(GetPinCategoryName(Type));
+	}
+	else
+	{
+		return Type->GetDisplayNameText();
 	}
 }
 
@@ -1713,4 +1741,15 @@ TSharedPtr<IAssetReferenceFilter> UEdGraphSchema_CustomizableObject::MakeAssetRe
 	return {};
 }
 
+bool UEdGraphSchema_CustomizableObject::IsTypeNative(const UScriptStruct* Type)
+{
+	return Type == UE::Mutable::FMeshAdapter::StaticStruct() ||
+		Type == UE::Mutable::FTextureAdapter::StaticStruct() ||
+		Type == UE::Mutable::FTextureAdapter::StaticStruct() ||
+		Type == UE::Mutable::FFloatAdapter::StaticStruct() ||
+		Type == UE::Mutable::FVectorAdapter::StaticStruct() ||
+		Type == UE::Mutable::FMaterialAdapter::StaticStruct();
+}
+
+
 #undef LOCTEXT_NAMESPACE
diff --git a/Engine/Plugins/Mutable/Source/CustomizableObjectEditor/Private/MuCOE/GenerateMutableSource/GenerateMutableSource.cpp b/Engine/Plugins/Mutable/Source/CustomizableObjectEditor/Private/MuCOE/GenerateMutableSource/GenerateMutableSource.cpp
--- a/Engine/Plugins/Mutable/Source/CustomizableObjectEditor/Private/MuCOE/GenerateMutableSource/GenerateMutableSource.cpp
+++ b/Engine/Plugins/Mutable/Source/CustomizableObjectEditor/Private/MuCOE/GenerateMutableSource/GenerateMutableSource.cpp
@@ -421,7 +421,7 @@ uint32 GetTypeHash(const FGeneratedKey& Key)
 	return Hash;
 }
 
-uint32 GetTypeHash(const FGeneratedSourceExtensionKey& Key)
+uint32 GetTypeHash(const FGeneratedSourceExternalKey& Key)
 {
 	uint32 Hash = GetTypeHash(Key.Pin);
 
diff --git a/Engine/Plugins/Mutable/Source/CustomizableObjectEditor/Private/MuCOE/GenerateMutableSource/GenerateMutableSource.h b/Engine/Plugins/Mutable/Source/CustomizableObjectEditor/Private/MuCOE/GenerateMutableSource/GenerateMutableSource.h
--- a/Engine/Plugins/Mutable/Source/CustomizableObjectEditor/Private/MuCOE/GenerateMutableSource/GenerateMutableSource.h
+++ b/Engine/Plugins/Mutable/Source/CustomizableObjectEditor/Private/MuCOE/GenerateMutableSource/GenerateMutableSource.h
@@ -180,16 +180,16 @@ public:
 uint32 GetTypeHash(const FGeneratedKey& Key);
 
 
-struct FGeneratedSourceExtensionKey
+struct FGeneratedSourceExternalKey
 {
 	const UEdGraphPin* Pin = nullptr;
-	FSourceExtensionOptions Options;
+	FSourceExternalOptions Options;
 
-	bool operator==(const FGeneratedSourceExtensionKey& Other) const = default;
+	bool operator==(const FGeneratedSourceExternalKey& Other) const = default;
 };
 
 
-uint32 GetTypeHash(const FGeneratedSourceExtensionKey& Key);
+uint32 GetTypeHash(const FGeneratedSourceExternalKey& Key);
 
 
 struct FImageKey
@@ -264,7 +264,7 @@ struct FGeneratedData
 };
 
 
-struct FGeneratedSourceExtensionData
+struct FGeneratedSourceExternalData
 {
 	UE::Mutable::Private::Ptr<UE::Mutable::Private::NodeExternal> Node;
 };
@@ -617,7 +617,7 @@ struct FMutableGraphGenerationContext
 	
 	// Cache of generated pins per LOD
 	TMap<FGeneratedKey, FGeneratedData> Generated;
-	TMap<FGeneratedSourceExtensionKey, FGeneratedSourceExtensionData> GeneratedSourceExtensions;
+	TMap<FGeneratedSourceExternalKey, FGeneratedSourceExternalData> GeneratedSourceExternals;
 
 	/** Set of all generated nodes. */
 	TSet<UCustomizableObjectNode*> GeneratedNodes;
@@ -907,11 +907,14 @@ struct FMutableGraphGenerationContext
 	// This should be automatized (create a struct that each time that a Macro is stacked, it adds 1 to the count.
 	// it also substracts 1 each time we remove a Macro from the stack)
 	TArray<const UCustomizableObjectNodeMacroInstance*> MacroNodesStack;
-	
-	TMap<FName, TObjectPtr<UTexture>> TextureParameterDefaultValues;
-	TMap<FName, TObjectPtr<USkeletalMesh>> SkeletalMeshParameterDefaultValues;
-	TMap<FName, TObjectPtr<UMaterialInterface>> MaterialParameterDefaultValues;
 
+	// TODO GMT Get Default Values from Core. We only need compiled ones
+	TMap<FString, TObjectPtr<UTexture>> TextureParameterDefaultValues;
+	TMap<FString, TObjectPtr<USkeletalMesh>> SkeletalMeshParameterDefaultValues;
+	TMap<FString, TObjectPtr<UMaterialInterface>> MaterialParameterDefaultValues;
+	TMap<FString, FInstancedStruct> ExternalTypeParameterDefaultValues;
+
+	TMap<FString, TObjectPtr<const UScriptStruct>> ExternalTypeParameterTypes;
 	
 	struct FMaterialBreakParameter {
 		FName ParameterName;
diff --git a/Engine/Plugins/Mutable/Source/CustomizableObjectEditor/Private/MuCOE/GenerateMutableSource/GenerateMutableSourceColor.cpp b/Engine/Plugins/Mutable/Source/CustomizableObjectEditor/Private/MuCOE/GenerateMutableSource/GenerateMutableSourceColor.cpp
--- a/Engine/Plugins/Mutable/Source/CustomizableObjectEditor/Private/MuCOE/GenerateMutableSource/GenerateMutableSourceColor.cpp
+++ b/Engine/Plugins/Mutable/Source/CustomizableObjectEditor/Private/MuCOE/GenerateMutableSource/GenerateMutableSourceColor.cpp
@@ -13,7 +13,7 @@
 #include "MuCOE/GenerateMutableSource/GenerateMutableSourceTable.h"
 #include "MuCOE/GenerateMutableSource/GenerateMutableSource.h"
 #include "MuCOE/GraphTraversal.h"
-#include "MuCOE/Nodes/CONodeExternal.h"
+#include "MuCOE/Nodes/CONodeExternalOperation.h"
 #include "MuCOE/Nodes/CustomizableObjectNodeColorArithmeticOp.h"
 #include "MuCOE/Nodes/CustomizableObjectNodeColorConstant.h"
 #include "MuCOE/Nodes/CustomizableObjectNodeColorFromFloats.h"
@@ -363,12 +363,12 @@ UE::Mutable::Private::Ptr<UE::Mutable::Private::NodeColour> GenerateMutableSourc
 		}
 	}
 	
-	else if (Cast<UCONodeExternal>(Node))
+	else if (Cast<UCONodeExternalOperation>(Node))
 	{
 		UE::Mutable::Private::Ptr<UE::Mutable::Private::NodeColourExternal> MeshExtension = new UE::Mutable::Private::NodeColourExternal();
 
-		FSourceExtensionOptions Options;
-		MeshExtension->Node = GenerateMutableSourceExtension(Pin, GenerationContext, Options);
+		FSourceExternalOptions Options;
+		MeshExtension->Node = GenerateMutableSourceExternal(Pin, GenerationContext, Options);
 		
 		Result = MeshExtension;
 
diff --git a/Engine/Plugins/Mutable/Source/CustomizableObjectEditor/Private/MuCOE/GenerateMutableSource/GenerateMutableSourceExternal.cpp b/Engine/Plugins/Mutable/Source/CustomizableObjectEditor/Private/MuCOE/GenerateMutableSource/GenerateMutableSourceExternal.cpp
--- a/Engine/Plugins/Mutable/Source/CustomizableObjectEditor/Private/MuCOE/GenerateMutableSource/GenerateMutableSourceExternal.cpp
+++ b/Engine/Plugins/Mutable/Source/CustomizableObjectEditor/Private/MuCOE/GenerateMutableSource/GenerateMutableSourceExternal.cpp
@@ -7,40 +7,46 @@
 #include "GenerateMutableSourceFloat.h"
 #include "GenerateMutableSourceImage.h"
 #include "GenerateMutableSourceMacro.h"
+#include "GenerateMutableSourceMaterial.h"
 #include "GenerateMutableSourceMesh.h"
+#include "MuCO/ICustomizableObjectModulePrivate.h"
 #include "MuCOE/GraphTraversal.h"
-#include "MuCOE/Nodes/CONodeExternal.h"
+#include "MuCOE/Nodes/CONodeExternalOperation.h"
+#include "MuCOE/Nodes/CONodeExternalTypeParameter.h"
+#include "MuCOE/Nodes/ExternalTypeCustomization.h"
 #include "MuR/External/FloatAdapter.h"
 #include "MuR/External/MaterialAdapter.h"
 #include "MuR/External/MeshAdapter.h"
 #include "MuR/External/Operation.h"
 #include "MuR/External/TextureAdapter.h"
 #include "MuR/External/VectorAdapter.h"
+#include "MuT/NodeExternalOperation.h"
+#include "MuT/NodeExternalParameter.h"
 
 #define LOCTEXT_NAMESPACE "CustomizableObjectEditor"
 
 
-UE::Mutable::Private::Ptr<UE::Mutable::Private::NodeExternal> GenerateMutableSourceExtension(const UEdGraphPin* Pin, FMutableGraphGenerationContext& GenerationContext, const FSourceExtensionOptions& Options)
+UE::Mutable::Private::Ptr<UE::Mutable::Private::NodeExternal> GenerateMutableSourceExternal(const UEdGraphPin* Pin, FMutableGraphGenerationContext& GenerationContext, const FSourceExternalOptions& Options)
 {
 	check(Pin)
 	RETURN_ON_CYCLE(*Pin, GenerationContext)
 
 	CheckNumOutputs(*Pin, GenerationContext);
 	
-	UCustomizableObjectNode* InNode = CastChecked<UCustomizableObjectNode>(Pin->GetOwningNode());
+	UCustomizableObjectNode* Node = CastChecked<UCustomizableObjectNode>(Pin->GetOwningNode());
 
-	FGeneratedSourceExtensionKey Key;
+	FGeneratedSourceExternalKey Key;
 	Key.Pin = Pin;
 	Key.Options = Options;
 	
-	if (const FGeneratedSourceExtensionData* Generated = GenerationContext.GeneratedSourceExtensions.Find(Key))
+	if (const FGeneratedSourceExternalData* Generated = GenerationContext.GeneratedSourceExternals.Find(Key))
 	{
 		return Generated->Node;
 	}
 
-	if (InNode->IsNodeOutDatedAndNeedsRefresh())
+	if (Node->IsNodeOutDatedAndNeedsRefresh())
 	{
-		InNode->SetRefreshNodeWarning();
+		Node->SetRefreshNodeWarning();
 	}
 	
 	UCustomizableObjectNode* EditorNode = CastChecked<UCustomizableObjectNode>(Pin->GetOwningNode());
@@ -48,44 +54,52 @@ UE::Mutable::Private::Ptr<UE::Mutable::Private::NodeExternal> GenerateMutableSou
 	UE::Mutable::Private::Ptr<UE::Mutable::Private::NodeExternal> Result;
 	bool bCacheNode = true;
 	
-	if (const UCONodeExternal* TypedNodeExtension = Cast<UCONodeExternal>(EditorNode))
+	if (const UCONodeExternalOperation* TypedNodeExternal = Cast<UCONodeExternalOperation>(EditorNode))
 	{
 		[&]()
 		{
-			const UE::Mutable::FExternalOperation* Operation = TypedNodeExtension->OperationInstancedStruct.GetPtr<UE::Mutable::FExternalOperation>();
+			const UE::Mutable::FExternalOperation* Operation = TypedNodeExternal->OperationInstancedStruct.GetPtr<UE::Mutable::FExternalOperation>();
 			if (!Operation)
 			{
-				FText Msg = FText::Format(LOCTEXT("ExtensionNotLoaded", "Node is using an external operation which has not been loaded: {0}"), FText::FromName(TypedNodeExtension->CachedOperationName));
-				GenerationContext.Log(Msg, TypedNodeExtension, EMessageSeverity::Error);
+				FText Msg = FText::Format(LOCTEXT("ExternalNotLoaded", "Node is using an external operation which has not been loaded: {0}"), FText::FromName(TypedNodeExternal->CachedOperationName));
+				GenerationContext.Log(Msg, TypedNodeExternal, EMessageSeverity::Error);
 				return;
 			}
-			
-			UE::Mutable::Private::Ptr<UE::Mutable::Private::NodeExternal> ExtensionNode = new UE::Mutable::Private::NodeExternal();
+
+			const UScriptStruct* ExternalType;
+			if (ContainsUObject(*Operation, ExternalType))
+			{
+				FText Msg = FText::Format(LOCTEXT("TypeUObject", "Operation contains UObject properties. Operation: {0} Type: {1}"), FText::FromName(TypedNodeExternal->CachedOperationName), FText::FromName(ExternalType->GetFName()));
+				GenerationContext.Log(Msg, TypedNodeExternal, EMessageSeverity::Error);
+				return;
+			}
+				
+			UE::Mutable::Private::Ptr<UE::Mutable::Private::NodeExternalOperation> ExternalNode = new UE::Mutable::Private::NodeExternalOperation();
             				
-            ExtensionNode->OperationInstancedStruct = TypedNodeExtension->OperationInstancedStruct;
+            ExternalNode->OperationInstancedStruct = TypedNodeExternal->OperationInstancedStruct;
 
             TArray<TPair<FName, const UScriptStruct*>> Inputs = Operation->GetInputs();
             for (const TPair<FName, const UScriptStruct*>& Pair : Inputs)
             {
             	const FName& InputName = Pair.Key;
             
-            	const FEdGraphPinReference* InputResult = TypedNodeExtension->InputPins.FindByPredicate([&](const FEdGraphPinReference& Element)
+            	const FEdGraphPinReference* InputResult = TypedNodeExternal->InputPins.FindByPredicate([&](const FEdGraphPinReference& Element)
             	{
-            		return TypedNodeExtension->GetPinData<UCONodeExtensionPinData>(*Element.Get()).Name == InputName;
+            		return TypedNodeExternal->GetPinData<UCONodeExtenalPinData>(*Element.Get()).Name == InputName;
             	});
         
             	if (!InputResult)
             	{
-            		FText Msg = FText::Format(LOCTEXT("ExtensionInputNotFound", "Can not find input in external operation: {0}"), FText::FromName(InputName));
-            		GenerationContext.Log(Msg, TypedNodeExtension, EMessageSeverity::Error);
+            		FText Msg = FText::Format(LOCTEXT("ExternalInputNotFound", "Can not find input in external operation: {0}"), FText::FromName(InputName));
+            		GenerationContext.Log(Msg, TypedNodeExternal, EMessageSeverity::Error);
             		continue;
             	}
 
             	const UEdGraphPin* ConnectedPin = FollowInputPin(*InputResult->Get());
             	if (!ConnectedPin)
             	{
-            		FText Msg = FText::Format(LOCTEXT("ExtensionInputNotConnected", "External operation input not connected: {0}"), FText::FromName(InputName));
-            		GenerationContext.Log(Msg, TypedNodeExtension, EMessageSeverity::Error);
+            		FText Msg = FText::Format(LOCTEXT("ExternalInputNotConnected", "External operation input not connected: {0}"), FText::FromName(InputName));
+            		GenerationContext.Log(Msg, TypedNodeExternal, EMessageSeverity::Error);
             		continue;
             	}
 
@@ -94,65 +108,79 @@ UE::Mutable::Private::Ptr<UE::Mutable::Private::NodeExternal> GenerateMutableSou
             	if (InputType == UE::Mutable::FMeshAdapter::StaticStruct())
             	{
             		UE::Mutable::Private::Ptr<UE::Mutable::Private::NodeMesh> SourceNode = GenerateMutableSourceMesh(ConnectedPin, GenerationContext, Options.MeshOptions, Options.bLinkedToExtendMaterial, Options.bOnlyConnectedLOD);
-            		ExtensionNode->Inputs.Add(SourceNode);
+            		ExternalNode->Inputs.Add(SourceNode);
             	}
             	else if (InputType == UE::Mutable::FTextureAdapter::StaticStruct())
             	{
             		UE::Mutable::Private::Ptr<UE::Mutable::Private::NodeImage> SourceNode = GenerateMutableSourceImage(ConnectedPin, GenerationContext, Options.ReferenceTextureSize);
-            		ExtensionNode->Inputs.Add(SourceNode);
+            		ExternalNode->Inputs.Add(SourceNode);
             	}
             	else if (InputType == UE::Mutable::FFloatAdapter::StaticStruct())
             	{
             		UE::Mutable::Private::Ptr<UE::Mutable::Private::NodeScalar> SourceNode = GenerateMutableSourceFloat(ConnectedPin, GenerationContext);
-            		ExtensionNode->Inputs.Add(SourceNode);
+            		ExternalNode->Inputs.Add(SourceNode);
             	}
             	else if (InputType == UE::Mutable::FVectorAdapter::StaticStruct())
             	{
             		UE::Mutable::Private::Ptr<UE::Mutable::Private::NodeColour> SourceNode = GenerateMutableSourceColor(ConnectedPin, GenerationContext);
-            		ExtensionNode->Inputs.Add(SourceNode);
+            		ExternalNode->Inputs.Add(SourceNode);
             	}
             	else if (InputType == UE::Mutable::FMaterialAdapter::StaticStruct())
             	{
-            		UE::Mutable::Private::Ptr<UE::Mutable::Private::NodeColour> SourceNode = GenerateMutableSourceColor(ConnectedPin, GenerationContext);
-            		ExtensionNode->Inputs.Add(SourceNode);
+            		UE::Mutable::Private::Ptr<UE::Mutable::Private::NodeMaterial> SourceNode = GenerateMutableSourceMaterial(ConnectedPin, GenerationContext);
+            		ExternalNode->Inputs.Add(SourceNode);
             	}
             	else
             	{
-            		UE::Mutable::Private::Ptr<UE::Mutable::Private::Node> SourceNode = GenerateMutableSourceExtension(ConnectedPin, GenerationContext, Options);
-            		ExtensionNode->Inputs.Add(SourceNode);
+            		UE::Mutable::Private::Ptr<UE::Mutable::Private::Node> SourceNode = GenerateMutableSourceExternal(ConnectedPin, GenerationContext, Options);
+            		ExternalNode->Inputs.Add(SourceNode);
             	}	
             }
 
-			Result = ExtensionNode;
+			Result = ExternalNode;
 		}();		
 	}
 
-	else if (const UCustomizableObjectNodeMacroInstance* TypedNodeMacro = Cast<UCustomizableObjectNodeMacroInstance>(InNode))
+	if (const UCONodeExternalTypeParameter* TypeNodeExternalParameter = Cast<UCONodeExternalTypeParameter>(EditorNode))
+	{
+		UE::Mutable::Private::Ptr<UE::Mutable::Private::NodeExternalParameter> ExternalParameterNode = new UE::Mutable::Private::NodeExternalParameter();
+
+		ExternalParameterNode->Name = TypeNodeExternalParameter->GetParameterName(&GenerationContext.MacroNodesStack);
+		ExternalParameterNode->UID = GenerationContext.GetNodeIdUnique(Node).ToString();
+		ExternalParameterNode->DefaultValue = MakeShared<FInstancedStruct>(TypeNodeExternalParameter->DefaultValue);
+
+		GenerationContext.ExternalTypeParameterDefaultValues.Add(TypeNodeExternalParameter->GetParameterName(), TypeNodeExternalParameter->DefaultValue);
+		GenerationContext.ExternalTypeParameterTypes.Add(TypeNodeExternalParameter->GetParameterName(), TypeNodeExternalParameter->DefaultValue.GetScriptStruct());
+		
+		Result = ExternalParameterNode;
+	}
+	
+	else if (const UCustomizableObjectNodeMacroInstance* TypedNodeMacro = Cast<UCustomizableObjectNodeMacroInstance>(Node))
 	{
 		bCacheNode = false;
 
-		Result = GenerateMutableSourceMacro<UE::Mutable::Private::NodeExternal>(*Pin, GenerationContext, GenerateMutableSourceExtension, Options);
+		Result = GenerateMutableSourceMacro<UE::Mutable::Private::NodeExternal>(*Pin, GenerationContext, GenerateMutableSourceExternal, Options);
 	}
 
-	else if (const UCustomizableObjectNodeTunnel* TypedNodeTunnel = Cast<UCustomizableObjectNodeTunnel>(InNode))
+	else if (const UCustomizableObjectNodeTunnel* TypedNodeTunnel = Cast<UCustomizableObjectNodeTunnel>(Node))
 	{
 		bCacheNode = false;
 
-		Result = GenerateMutableSourceMacro<UE::Mutable::Private::NodeExternal>(*Pin, GenerationContext, GenerateMutableSourceExtension, Options);
+		Result = GenerateMutableSourceMacro<UE::Mutable::Private::NodeExternal>(*Pin, GenerationContext, GenerateMutableSourceExternal, Options);
 	}
 	
 	if (Result)
 	{
-		Result->SetMessageContext(InNode);
+		Result->SetMessageContext(Node);
 	}
 
 	if (bCacheNode)
 	{
-		FGeneratedSourceExtensionData CacheData;
+		FGeneratedSourceExternalData CacheData;
 		CacheData.Node = Result;
 	
-		GenerationContext.GeneratedSourceExtensions.Add(Key, CacheData);
-		GenerationContext.GeneratedNodes.Add(InNode);	
+		GenerationContext.GeneratedSourceExternals.Add(Key, CacheData);
+		GenerationContext.GeneratedNodes.Add(Node);	
 	}
 
 	return Result;
diff --git a/Engine/Plugins/Mutable/Source/CustomizableObjectEditor/Private/MuCOE/GenerateMutableSource/GenerateMutableSourceExternal.h b/Engine/Plugins/Mutable/Source/CustomizableObjectEditor/Private/MuCOE/GenerateMutableSource/GenerateMutableSourceExternal.h
--- a/Engine/Plugins/Mutable/Source/CustomizableObjectEditor/Private/MuCOE/GenerateMutableSource/GenerateMutableSourceExternal.h
+++ b/Engine/Plugins/Mutable/Source/CustomizableObjectEditor/Private/MuCOE/GenerateMutableSource/GenerateMutableSourceExternal.h
@@ -5,9 +5,9 @@
 #include "MuT/NodeExternal.h"
 
 class UEdGraphPin;
-struct FSourceExtensionOptions;
+struct FSourceExternalOptions;
 struct FMutableGraphGenerationContext;
 
 
-UE::Mutable::Private::Ptr<UE::Mutable::Private::NodeExternal> GenerateMutableSourceExtension(const UEdGraphPin* Pin, FMutableGraphGenerationContext& GenerationContext, const FSourceExtensionOptions& Options);
+UE::Mutable::Private::Ptr<UE::Mutable::Private::NodeExternal> GenerateMutableSourceExternal(const UEdGraphPin* Pin, FMutableGraphGenerationContext& GenerationContext, const FSourceExternalOptions& Options);
 
diff --git a/Engine/Plugins/Mutable/Source/CustomizableObjectEditor/Private/MuCOE/GenerateMutableSource/GenerateMutableSourceFloat.cpp b/Engine/Plugins/Mutable/Source/CustomizableObjectEditor/Private/MuCOE/GenerateMutableSource/GenerateMutableSourceFloat.cpp
--- a/Engine/Plugins/Mutable/Source/CustomizableObjectEditor/Private/MuCOE/GenerateMutableSource/GenerateMutableSourceFloat.cpp
+++ b/Engine/Plugins/Mutable/Source/CustomizableObjectEditor/Private/MuCOE/GenerateMutableSource/GenerateMutableSourceFloat.cpp
@@ -14,7 +14,7 @@
 #include "MuCOE/GenerateMutableSource/GenerateMutableSourceMaterial.h"
 #include "MuCOE/GenerateMutableSource/GenerateMutableSourceTable.h"
 #include "MuCOE/GraphTraversal.h"
-#include "MuCOE/Nodes/CONodeExternal.h"
+#include "MuCOE/Nodes/CONodeExternalOperation.h"
 #include "MuCOE/Nodes/CONodeMaterialBreak.h"
 #include "MuCOE/Nodes/CustomizableObjectNodeCurve.h"
 #include "MuCOE/Nodes/CustomizableObjectNodeEnumParameter.h"
@@ -390,12 +390,12 @@ UE::Mutable::Private::Ptr<UE::Mutable::Private::NodeScalar> GenerateMutableSourc
 		}
 	}
 
-	else if (Cast<UCONodeExternal>(Node))
+	else if (Cast<UCONodeExternalOperation>(Node))
 	{
 		UE::Mutable::Private::Ptr<UE::Mutable::Private::NodeScalarExternal> MeshExtension = new UE::Mutable::Private::NodeScalarExternal();
 
-		FSourceExtensionOptions Options;
-		MeshExtension->Node = GenerateMutableSourceExtension(Pin, GenerationContext, Options);
+		FSourceExternalOptions Options;
+		MeshExtension->Node = GenerateMutableSourceExternal(Pin, GenerationContext, Options);
 		
 		Result = MeshExtension;
 
diff --git a/Engine/Plugins/Mutable/Source/CustomizableObjectEditor/Private/MuCOE/GenerateMutableSource/GenerateMutableSourceImage.cpp b/Engine/Plugins/Mutable/Source/CustomizableObjectEditor/Private/MuCOE/GenerateMutableSource/GenerateMutableSourceImage.cpp
--- a/Engine/Plugins/Mutable/Source/CustomizableObjectEditor/Private/MuCOE/GenerateMutableSource/GenerateMutableSourceImage.cpp
+++ b/Engine/Plugins/Mutable/Source/CustomizableObjectEditor/Private/MuCOE/GenerateMutableSource/GenerateMutableSourceImage.cpp
@@ -42,7 +42,7 @@
 #include "MuCOE/Nodes/CustomizableObjectNodeTunnel.h"
 #include "MuCOE/UnrealEditorPortabilityHelpers.h"
 #include "MuCO/UnrealToMutableTextureConversionUtils.h"
-#include "MuCOE/Nodes/CONodeExternal.h"
+#include "MuCOE/Nodes/CONodeExternalOperation.h"
 #include "MuT/NodeImageBinarise.h"
 #include "MuT/NodeImageColourMap.h"
 #include "MuT/NodeImageFormat.h"
@@ -1163,13 +1163,13 @@ UE::Mutable::Private::Ptr<UE::Mutable::Private::NodeImage> GenerateMutableSource
 		}
 	}
 
-	else if (Cast<UCONodeExternal>(Node))
+	else if (Cast<UCONodeExternalOperation>(Node))
 	{
 		UE::Mutable::Private::Ptr<UE::Mutable::Private::NodeImageExternal> MeshExtension = new UE::Mutable::Private::NodeImageExternal();
 
-		FSourceExtensionOptions Options;
+		FSourceExternalOptions Options;
 		Options.ReferenceTextureSize = ReferenceTextureSize;
-		MeshExtension->Node = GenerateMutableSourceExtension(Pin, GenerationContext, Options);
+		MeshExtension->Node = GenerateMutableSourceExternal(Pin, GenerationContext, Options);
 		
 		Result = MeshExtension;
 	}
@@ -1308,10 +1308,7 @@ UE::Mutable::Private::Ptr<UE::Mutable::Private::NodeImageParameter> GenerateMuta
 		TextureNode->Name = TypedNodeParam->GetParameterName(&GenerationContext.MacroNodesStack);
 		TextureNode->UID = GenerationContext.GetNodeIdUnique(ParameterNode).ToString();
 		
-		if (TypedNodeParam->DefaultValue)
-		{
-			GenerationContext.TextureParameterDefaultValues.Add(FName(TypedNodeParam->GetParameterName()), TypedNodeParam->DefaultValue);
-		}
+		GenerationContext.TextureParameterDefaultValues.Add(TypedNodeParam->GetParameterName(), TypedNodeParam->DefaultValue);
 
 		GenerationContext.ParameterUIDataMap.Add(TypedNodeParam->GetParameterName(&GenerationContext.MacroNodesStack), FMutableParameterData(
 			TypedNodeParam->ParamUIMetadata,
diff --git a/Engine/Plugins/Mutable/Source/CustomizableObjectEditor/Private/MuCOE/GenerateMutableSource/GenerateMutableSourceMaterial.cpp b/Engine/Plugins/Mutable/Source/CustomizableObjectEditor/Private/MuCOE/GenerateMutableSource/GenerateMutableSourceMaterial.cpp
--- a/Engine/Plugins/Mutable/Source/CustomizableObjectEditor/Private/MuCOE/GenerateMutableSource/GenerateMutableSourceMaterial.cpp
+++ b/Engine/Plugins/Mutable/Source/CustomizableObjectEditor/Private/MuCOE/GenerateMutableSource/GenerateMutableSourceMaterial.cpp
@@ -34,7 +34,7 @@
 
 #include "Materials/MaterialInterface.h"
 #include "Materials/MaterialInstance.h"
-#include "MuCOE/Nodes/CONodeExternal.h"
+#include "MuCOE/Nodes/CONodeExternalOperation.h"
 #include "MuT/NodeMaterialExternal.h"
 
 
@@ -282,10 +282,7 @@ UE::Mutable::Private::Ptr<UE::Mutable::Private::NodeMaterial> GenerateMutableSou
 			MaterialNode->Name = MaterialParameterNode->GetParameterName(&GenerationContext.MacroNodesStack);
 			MaterialNode->UID = GenerationContext.GetNodeIdUnique(Node).ToString();
 
-			if (MaterialParameterNode->DefaultValue)
-			{
-				GenerationContext.MaterialParameterDefaultValues.Add(FName(MaterialParameterNode->GetParameterName()), UE::Mutable::Private::LoadObject(MaterialParameterNode->DefaultValue));
-			}
+			GenerationContext.MaterialParameterDefaultValues.Add(MaterialParameterNode->GetParameterName(), UE::Mutable::Private::LoadObject(MaterialParameterNode->DefaultValue));
 
 			GenerationContext.ParameterUIDataMap.Add(MaterialParameterNode->GetParameterName(&GenerationContext.MacroNodesStack), FMutableParameterData(
 				MaterialParameterNode->ParamUIMetadata,
@@ -469,12 +466,12 @@ UE::Mutable::Private::Ptr<UE::Mutable::Private::NodeMaterial> GenerateMutableSou
 		}
 	}
 
-	else if (Cast<UCONodeExternal>(Node))
+	else if (Cast<UCONodeExternalOperation>(Node))
 	{
 		UE::Mutable::Private::Ptr<UE::Mutable::Private::NodeMaterialExternal> MeshExtension = new UE::Mutable::Private::NodeMaterialExternal();
 
-		FSourceExtensionOptions Options;
-		MeshExtension->Node = GenerateMutableSourceExtension(Pin, GenerationContext, Options);
+		FSourceExternalOptions Options;
+		MeshExtension->Node = GenerateMutableSourceExternal(Pin, GenerationContext, Options);
 		
 		Result = MeshExtension;
 
diff --git a/Engine/Plugins/Mutable/Source/CustomizableObjectEditor/Private/MuCOE/GenerateMutableSource/GenerateMutableSourceMesh.cpp b/Engine/Plugins/Mutable/Source/CustomizableObjectEditor/Private/MuCOE/GenerateMutableSource/GenerateMutableSourceMesh.cpp
--- a/Engine/Plugins/Mutable/Source/CustomizableObjectEditor/Private/MuCOE/GenerateMutableSource/GenerateMutableSourceMesh.cpp
+++ b/Engine/Plugins/Mutable/Source/CustomizableObjectEditor/Private/MuCOE/GenerateMutableSource/GenerateMutableSourceMesh.cpp
@@ -63,7 +63,7 @@
 #include "AssetRegistry/AssetRegistryModule.h"
 #include "Hash/CityHash.h"
 #include "PackedNormal.h"
-#include "MuCOE/Nodes/CONodeExternal.h"
+#include "MuCOE/Nodes/CONodeExternalOperation.h"
 #include "MuR/External/MeshAdapter.h"
 #include "MuR/External/Operation.h"
 #include "MuR/External/TextureAdapter.h"
@@ -4595,10 +4595,7 @@ UE::Mutable::Private::Ptr<UE::Mutable::Private::NodeMesh> GenerateMutableSourceM
 		ParameterNode->Name = TypedNodeParam->GetParameterName();
 		ParameterNode->UID = GenerationContext.GetNodeIdUnique(Node).ToString();
 
-		if (TypedNodeParam->DefaultValue)
-		{
-			GenerationContext.SkeletalMeshParameterDefaultValues.Add(FName(TypedNodeParam->GetParameterName()), UE::Mutable::Private::LoadObject(TypedNodeParam->DefaultValue));
-		}
+		GenerationContext.SkeletalMeshParameterDefaultValues.Add(TypedNodeParam->GetParameterName(), UE::Mutable::Private::LoadObject(TypedNodeParam->DefaultValue));
 
 		int32 BaseLODIndex = 0;
 		int32 BaseSectionIndex = 0;
@@ -4684,15 +4681,15 @@ UE::Mutable::Private::Ptr<UE::Mutable::Private::NodeMesh> GenerateMutableSourceM
 		}
 	}
 
-	else if (Cast<UCONodeExternal>(Node))
+	else if (Cast<UCONodeExternalOperation>(Node))
 	{
 		UE::Mutable::Private::Ptr<UE::Mutable::Private::NodeMeshExternal> MeshExtension = new UE::Mutable::Private::NodeMeshExternal();
 
-		FSourceExtensionOptions Options;
+		FSourceExternalOptions Options;
 		Options.MeshOptions = BaseMeshData;
 		Options.bLinkedToExtendMaterial = bLinkedToExtendMaterial;
 		Options.bOnlyConnectedLOD = bOnlyConnectedLOD;
-		MeshExtension->Node = GenerateMutableSourceExtension(Pin, GenerationContext, Options);
+		MeshExtension->Node = GenerateMutableSourceExternal(Pin, GenerationContext, Options);
 		
 		Result = MeshExtension;
 	}
diff --git a/Engine/Plugins/Mutable/Source/CustomizableObjectEditor/Private/MuCOE/GraphTraversal.cpp b/Engine/Plugins/Mutable/Source/CustomizableObjectEditor/Private/MuCOE/GraphTraversal.cpp
--- a/Engine/Plugins/Mutable/Source/CustomizableObjectEditor/Private/MuCOE/GraphTraversal.cpp
+++ b/Engine/Plugins/Mutable/Source/CustomizableObjectEditor/Private/MuCOE/GraphTraversal.cpp
@@ -29,7 +29,7 @@
 #include "MuCOE/Nodes/CustomizableObjectNodeAnimationPose.h"
 #include "MuCOE/Nodes/CustomizableObjectNodeObjectGroup.h"
 #include "MuCOE/Nodes/CustomizableObjectNodeReroute.h"
-#include "Nodes/CONodeExternal.h"
+#include "Nodes/CONodeExternalOperation.h"
 #include "UObject/UObjectIterator.h"
 
 
@@ -750,7 +750,7 @@ const UEdGraphPin* FindMeshBaseSource(const UEdGraphPin& Pin, const bool bOnlyLo
 		}
 	}
 
-	else if (Cast<UCONodeExternal>(Node))
+	else if (Cast<UCONodeExternalOperation>(Node))
 	{
 		if (!bOnlyLookForStaticMesh)
 		{
diff --git a/Engine/Plugins/Mutable/Source/CustomizableObjectEditor/Private/MuCOE/Nodes/CONodeExternal.cpp b/Engine/Plugins/Mutable/Source/CustomizableObjectEditor/Private/MuCOE/Nodes/CONodeExternal.cpp
deleted file mode 100644
diff --git a/Engine/Plugins/Mutable/Source/CustomizableObjectEditor/Private/MuCOE/Nodes/CONodeExternal.h b/Engine/Plugins/Mutable/Source/CustomizableObjectEditor/Private/MuCOE/Nodes/CONodeExternal.h
deleted file mode 100644
diff --git a/Engine/Plugins/Mutable/Source/CustomizableObjectEditor/Private/MuCOE/Nodes/CONodeExternalDetails.cpp b/Engine/Plugins/Mutable/Source/CustomizableObjectEditor/Private/MuCOE/Nodes/CONodeExternalDetails.cpp
deleted file mode 100644
diff --git a/Engine/Plugins/Mutable/Source/CustomizableObjectEditor/Private/MuCOE/Nodes/CONodeExternalDetails.h b/Engine/Plugins/Mutable/Source/CustomizableObjectEditor/Private/MuCOE/Nodes/CONodeExternalDetails.h
deleted file mode 100644
diff --git a/Engine/Plugins/Mutable/Source/CustomizableObjectEditor/Private/MuCOE/Nodes/CONodeExternalOperation.cpp b/Engine/Plugins/Mutable/Source/CustomizableObjectEditor/Private/MuCOE/Nodes/CONodeExternalOperation.cpp
new file mode 100644
--- /dev/null
+++ b/Engine/Plugins/Mutable/Source/CustomizableObjectEditor/Private/MuCOE/Nodes/CONodeExternalOperation.cpp
@@ -0,0 +1,79 @@
+﻿// Copyright Epic Games, Inc. All Rights Reserved.
+
+#include "MuCOE/Nodes/CONodeExternalOperation.h"
+
+#include "MuCOE/EdGraphSchema_CustomizableObject.h"
+#include "MuR/External/Operation.h"
+
+#define LOCTEXT_NAMESPACE "CustomizableObjectEditor"
+
+
+FText UCONodeExternalOperation::GetNodeTitle(ENodeTitleType::Type TitleType) const
+{
+	const UE::Mutable::FExternalOperation* Operation = OperationInstancedStruct.GetPtr<UE::Mutable::FExternalOperation>();
+	const FName Name = Operation ? Operation->GetName() : CachedOperationName;
+	
+	if (TitleType == ENodeTitleType::ListView)
+	{
+		return FText::FromName(Name);
+	}
+	else
+	{
+		return FText::Format(LOCTEXT("External_NodeTitle", "{0}\nExternal"), FText::FromName(Name));
+	}
+}
+
+
+FLinearColor UCONodeExternalOperation::GetNodeTitleColor() const
+{
+	const FName Category = UEdGraphSchema_CustomizableObject::GetPinCategoryName(OutputPin.Get()->PinType.PinCategory);
+	return UEdGraphSchema_CustomizableObject::GetPinTypeColor(Category);
+}
+
+
+bool UCONodeExternalOperation::IsLoaded() const
+{
+	return OperationInstancedStruct.IsValid();
+}
+
+
+void UCONodeExternalOperation::AllocateDefaultPins(UCustomizableObjectNodeRemapPins* RemapPins)
+{
+	Super::AllocateDefaultPins(RemapPins);
+
+	const UE::Mutable::FExternalOperation* Operation = OperationInstancedStruct.GetPtr<UE::Mutable::FExternalOperation>();
+	if (!Operation)
+	{
+		return;
+	}
+	
+	// Inputs
+	InputPins.Empty();
+	
+	for (TPair<FName, const UScriptStruct*> Pair : Operation->GetInputs())
+	{
+		UCONodeExtenalPinData* PinData = NewObject<UCONodeExtenalPinData>(this);
+		PinData->Name = Pair.Key;
+		
+		const FName Category = UEdGraphSchema_CustomizableObject::GetPinCategoryName(Pair.Value);
+
+		UEdGraphPin* Pin = CustomCreatePin(EGPD_Input, Category, Pair.Key, PinData);
+		InputPins.Add(Pin);
+	}
+
+	// Output
+	{
+		const TPair<FName, const UScriptStruct*> Pair = Operation->GetOutput();
+
+		UCONodeExtenalPinData* PinData = NewObject<UCONodeExtenalPinData>(this);
+		PinData->Name = Pair.Key;
+
+		const FName Category = UEdGraphSchema_CustomizableObject::GetPinCategoryName(Pair.Value);
+		OutputPin = CustomCreatePin(EGPD_Output, Category, Pair.Key, PinData);
+	}
+
+	CachedOperationName = Operation->GetName();
+}
+
+
+#undef LOCTEXT_NAMESPACE
diff --git a/Engine/Plugins/Mutable/Source/CustomizableObjectEditor/Private/MuCOE/Nodes/CONodeExternalOperation.h b/Engine/Plugins/Mutable/Source/CustomizableObjectEditor/Private/MuCOE/Nodes/CONodeExternalOperation.h
new file mode 100644
--- /dev/null
+++ b/Engine/Plugins/Mutable/Source/CustomizableObjectEditor/Private/MuCOE/Nodes/CONodeExternalOperation.h
@@ -0,0 +1,47 @@
+﻿// Copyright Epic Games, Inc. All Rights Reserved.
+
+#pragma once
+
+#include "MuCOE/Nodes/CustomizableObjectNode.h"
+
+#include "CONodeExternalOperation.generated.h"
+
+
+UCLASS()
+class UCONodeExtenalPinData : public UCustomizableObjectNodePinData
+{
+	GENERATED_BODY()
+
+public:
+	UPROPERTY()
+	FName Name;
+};
+
+
+UCLASS()
+class UCONodeExternalOperation : public UCustomizableObjectNode
+{
+	GENERATED_BODY()
+
+public:
+	// EdGraphNode interface
+	virtual FText GetNodeTitle(ENodeTitleType::Type TitleType) const override;
+	virtual FLinearColor GetNodeTitleColor() const override;
+	
+	// UCustomizableObjectNode interface
+	virtual void AllocateDefaultPins(UCustomizableObjectNodeRemapPins* RemapPins) override;
+	virtual bool IsLoaded() const override;
+	
+	UPROPERTY()
+	FInstancedStruct OperationInstancedStruct;
+
+	UPROPERTY()
+	TArray<FEdGraphPinReference> InputPins;
+
+	UPROPERTY()
+	FEdGraphPinReference OutputPin;
+	
+	/** In case of the operation not being loaded, the node uses this cached name so we can at least show the user which operation it was. */
+	UPROPERTY()
+	FName CachedOperationName;
+};
diff --git a/Engine/Plugins/Mutable/Source/CustomizableObjectEditor/Private/MuCOE/Nodes/CONodeExternalOperationDetails.cpp b/Engine/Plugins/Mutable/Source/CustomizableObjectEditor/Private/MuCOE/Nodes/CONodeExternalOperationDetails.cpp
new file mode 100644
--- /dev/null
+++ b/Engine/Plugins/Mutable/Source/CustomizableObjectEditor/Private/MuCOE/Nodes/CONodeExternalOperationDetails.cpp
@@ -0,0 +1,83 @@
+﻿// Copyright Epic Games, Inc. All Rights Reserved.
+
+#include "MuCOE/Nodes/CONodeExternalOperationDetails.h"
+
+#include "Editor.h"
+#include "CONodeExternalOperation.h"
+#include "DetailCategoryBuilder.h"
+#include "DetailLayoutBuilder.h"
+#include "IStructureDetailsView.h"
+
+#define LOCTEXT_NAMESPACE "CustomizableObjectEditor"
+
+
+TSharedRef<IDetailCustomization> FCONodeExternalOperationDetails::MakeInstance()
+{
+	return MakeShareable(new FCONodeExternalOperationDetails);
+}
+
+
+void FCONodeExternalOperationDetails::CustomizeDetails(IDetailLayoutBuilder& DetailBuilder)
+{
+	TSharedPtr<const IDetailsView> DetailsView = DetailBuilder.GetDetailsViewSharedPtr();
+	if (!DetailsView || !DetailsView->GetSelectedObjects().Num())
+	{
+		return;
+	}
+
+	UCONodeExternalOperation* Node = Cast<UCONodeExternalOperation>(DetailsView->GetSelectedObjects()[0]);
+	if (!Node)
+	{
+		return;
+	}
+
+	const UStruct* Struct = Node->OperationInstancedStruct.GetScriptStruct();
+	if (!Struct)
+	{
+		return;
+	}
+
+	TSharedRef<FStructOnScope> StructData = MakeShared<FStructOnScope>(Struct, Node->OperationInstancedStruct.GetMutableMemory());
+
+	for (FProperty* Property = Struct->PropertyLink; Property; Property = Property->PropertyLinkNext)
+	{
+		FName PropertyName = Property->GetFName();
+		TSharedPtr<IPropertyHandle> PropertyHandle = DetailBuilder.AddStructurePropertyData(StructData, PropertyName);
+		IDetailPropertyRow& DetailPropertyRow = DetailBuilder.AddPropertyToCategory(PropertyHandle.ToSharedRef());
+
+		TSharedPtr<IPropertyHandle> PropertyRowHandle = DetailPropertyRow.GetPropertyHandle();
+		
+		PropertyRowHandle->SetOnPropertyValuePreChange(FSimpleDelegate::CreateLambda([Node]()
+		{
+			GEditor->BeginTransaction(LOCTEXT("EditProperty", "Edited property"));
+			Node->Modify();
+		}));
+		
+		PropertyRowHandle->SetOnPropertyValueChangedWithData(TDelegate<void(const FPropertyChangedEvent&)>::CreateLambda([this](const FPropertyChangedEvent& InEvent)
+		{
+			GEditor->EndTransaction();
+		}));
+
+		PropertyRowHandle->SetOnChildPropertyValuePreChange(FSimpleDelegate::CreateLambda([Node]()
+		{
+			GEditor->BeginTransaction(LOCTEXT("EditChildProperty", "Edited child property"))	;
+			Node->Modify();
+		}));
+
+		PropertyRowHandle->SetOnChildPropertyValueChangedWithData(TDelegate<void(const FPropertyChangedEvent&)>::CreateLambda([this](const FPropertyChangedEvent& InEvent)
+		{
+			GEditor->EndTransaction();
+		}));
+	}
+}
+
+
+void FCONodeExternalOperationDetails::CustomizeDetails(const TSharedPtr<IDetailLayoutBuilder>& DetailBuilder)
+{
+	IDetailCustomization::CustomizeDetails(DetailBuilder);
+
+	DetailBuilderPtr = DetailBuilder;
+}
+
+
+#undef LOCTEXT_NAMESPACE
diff --git a/Engine/Plugins/Mutable/Source/CustomizableObjectEditor/Private/MuCOE/Nodes/CONodeExternalOperationDetails.h b/Engine/Plugins/Mutable/Source/CustomizableObjectEditor/Private/MuCOE/Nodes/CONodeExternalOperationDetails.h
new file mode 100644
--- /dev/null
+++ b/Engine/Plugins/Mutable/Source/CustomizableObjectEditor/Private/MuCOE/Nodes/CONodeExternalOperationDetails.h
@@ -0,0 +1,23 @@
+﻿// Copyright Epic Games, Inc. All Rights Reserved.
+
+#pragma once
+
+#include "IDetailCustomization.h"
+
+class FReply;
+class IDetailLayoutBuilder;
+class UCustomizableObjectNode;
+
+
+class FCONodeExternalOperationDetails : public IDetailCustomization
+{
+public:
+	static TSharedRef<IDetailCustomization> MakeInstance();
+	
+	// IDetailCustomization interface
+	virtual void CustomizeDetails(IDetailLayoutBuilder& DetailBuilder) override;
+	virtual void CustomizeDetails(const TSharedPtr<IDetailLayoutBuilder>& DetailBuilder) override;
+
+private: 
+	TSharedPtr<IDetailLayoutBuilder> DetailBuilderPtr;
+};
diff --git a/Engine/Plugins/Mutable/Source/CustomizableObjectEditor/Private/MuCOE/Nodes/CONodeExternalTypeParameter.cpp b/Engine/Plugins/Mutable/Source/CustomizableObjectEditor/Private/MuCOE/Nodes/CONodeExternalTypeParameter.cpp
new file mode 100644
--- /dev/null
+++ b/Engine/Plugins/Mutable/Source/CustomizableObjectEditor/Private/MuCOE/Nodes/CONodeExternalTypeParameter.cpp
@@ -0,0 +1,50 @@
+﻿// Copyright Epic Games, Inc. All Rights Reserved.
+
+#include "MuCOE/Nodes/CONodeExternalTypeParameter.h"
+
+#include "MuCOE/EdGraphSchema_CustomizableObject.h"
+
+#define LOCTEXT_NAMESPACE "CustomizableObjectEditor"
+
+
+FText UCONodeExternalTypeParameter::GetNodeTitle(ENodeTitleType::Type TitleType) const
+{
+	const FText Title = Super::GetNodeTitle(TitleType);
+	
+	if (TitleType == ENodeTitleType::ListView)
+	{
+		return Title;
+	}
+	else
+	{
+		return FText::Format(LOCTEXT("ExternalParameter_NodeTitle", "{0}\nExternal"), Title);
+	}
+}
+
+
+FName UCONodeExternalTypeParameter::GetCategory() const
+{
+	if (DefaultValue.IsValid())
+	{
+		CachedType = UEdGraphSchema_CustomizableObject::GetPinCategoryName(DefaultValue.GetScriptStruct());
+	}
+	
+	return CachedType;
+}
+
+
+FText UCONodeExternalTypeParameter::GetCategoryFriendlyName() const
+{
+	if (DefaultValue.IsValid())
+	{
+		CachedFriendlyName = UEdGraphSchema_CustomizableObject::GetPinCategoryFriendlyName(DefaultValue.GetScriptStruct());
+	}
+
+	return CachedFriendlyName;
+}
+
+
+bool UCONodeExternalTypeParameter::IsLoaded() const
+{
+	return DefaultValue.IsValid();
+}
diff --git a/Engine/Plugins/Mutable/Source/CustomizableObjectEditor/Private/MuCOE/Nodes/CONodeExternalTypeParameter.h b/Engine/Plugins/Mutable/Source/CustomizableObjectEditor/Private/MuCOE/Nodes/CONodeExternalTypeParameter.h
new file mode 100644
--- /dev/null
+++ b/Engine/Plugins/Mutable/Source/CustomizableObjectEditor/Private/MuCOE/Nodes/CONodeExternalTypeParameter.h
@@ -0,0 +1,39 @@
+﻿// Copyright Epic Games, Inc. All Rights Reserved.
+
+#pragma once
+
+#include "CustomizableObjectNodeParameter.h"
+
+#include "MuCOE/Nodes/CustomizableObjectNode.h"
+
+#include "CONodeExternalTypeParameter.generated.h"
+
+class UScriptStruct;
+
+
+UCLASS()
+class UCONodeExternalTypeParameter : public UCustomizableObjectNodeParameter
+{
+	GENERATED_BODY()
+
+public:
+	// EdGraphNode interface
+	virtual FText GetNodeTitle(ENodeTitleType::Type TitleType) const override;
+
+	// UCustomizableObjectNode interface
+	virtual bool IsLoaded() const override;
+	
+	// UCustomizableObjectNodeParameter interface
+	virtual FName GetCategory() const override;
+	virtual FText GetCategoryFriendlyName() const override;
+
+	UPROPERTY(EditAnywhere, Category = CustomizableObject, meta = (MutableExternalType))
+	FInstancedStruct DefaultValue;
+
+	/** In case of the operation not being loaded, the node uses this cached name so we can at least show the user which type it was. */
+	UPROPERTY()
+	mutable FName CachedType;
+
+	UPROPERTY()
+	mutable FText CachedFriendlyName;
+};
diff --git a/Engine/Plugins/Mutable/Source/CustomizableObjectEditor/Private/MuCOE/Nodes/CustomizableObjectNode.cpp b/Engine/Plugins/Mutable/Source/CustomizableObjectEditor/Private/MuCOE/Nodes/CustomizableObjectNode.cpp
--- a/Engine/Plugins/Mutable/Source/CustomizableObjectEditor/Private/MuCOE/Nodes/CustomizableObjectNode.cpp
+++ b/Engine/Plugins/Mutable/Source/CustomizableObjectEditor/Private/MuCOE/Nodes/CustomizableObjectNode.cpp
@@ -619,6 +619,12 @@ bool UCustomizableObjectNode::IsExperimental() const
 }
 
 
+bool UCustomizableObjectNode::IsLoaded() const
+{
+	return true;
+}
+
+
 TArray<UEdGraphPin*> UCustomizableObjectNode::GetAllOrphanPins() const
 {
 	TArray<UEdGraphPin*> OrphanPins;
diff --git a/Engine/Plugins/Mutable/Source/CustomizableObjectEditor/Private/MuCOE/Nodes/CustomizableObjectNodeParameter.cpp b/Engine/Plugins/Mutable/Source/CustomizableObjectEditor/Private/MuCOE/Nodes/CustomizableObjectNodeParameter.cpp
--- a/Engine/Plugins/Mutable/Source/CustomizableObjectEditor/Private/MuCOE/Nodes/CustomizableObjectNodeParameter.cpp
+++ b/Engine/Plugins/Mutable/Source/CustomizableObjectEditor/Private/MuCOE/Nodes/CustomizableObjectNodeParameter.cpp
@@ -20,7 +20,7 @@ void UCustomizableObjectNodeParameter::AllocateDefaultPins(UCustomizableObjectNo
 {
 	const FName Type = GetCategory();
 	const FName PinName = UEdGraphSchema_CustomizableObject::GetPinCategoryName(Type);
-	const FText PinFriendlyName = UEdGraphSchema_CustomizableObject::GetPinCategoryFriendlyName(Type);
+	const FText PinFriendlyName = GetCategoryFriendlyName();
 	
 	UEdGraphPin* ValuePin = CustomCreatePin(EGPD_Output, Type, PinName);
 	ValuePin->PinFriendlyName = PinFriendlyName;
@@ -70,10 +70,7 @@ void UCustomizableObjectNodeParameter::SetPinEditableName(const UEdGraphPin& Pin
 
 FText UCustomizableObjectNodeParameter::GetNodeTitle(ENodeTitleType::Type TitleType) const
 {
-	FFormatNamedArguments Args;
-	Args.Add(TEXT("Type"), UEdGraphSchema_CustomizableObject::GetPinCategoryFriendlyName(GetCategory()));
-
-	return FText::Format(LOCTEXT("ParameterTitle_ListView", "{Type} Parameter"), Args);
+	return FText::Format(LOCTEXT("ParameterTitle_ListView", "{0} Parameter"), GetCategoryFriendlyName());
 }
 
 
@@ -107,6 +104,12 @@ void UCustomizableObjectNodeParameter::PinConnectionListChanged(UEdGraphPin* Pin
 }
 
 
+FText UCustomizableObjectNodeParameter::GetCategoryFriendlyName() const
+{
+	return UEdGraphSchema_CustomizableObject::GetPinCategoryFriendlyName(GetCategory());
+}
+
+
 FString UCustomizableObjectNodeParameter::GetParameterName(TArray<const UCustomizableObjectNodeMacroInstance*>* MacroContext) const
 {
 	if (NamePin.Get())
diff --git a/Engine/Plugins/Mutable/Source/CustomizableObjectEditor/Private/MuCOE/Nodes/CustomizableObjectNodeParameter.h b/Engine/Plugins/Mutable/Source/CustomizableObjectEditor/Private/MuCOE/Nodes/CustomizableObjectNodeParameter.h
--- a/Engine/Plugins/Mutable/Source/CustomizableObjectEditor/Private/MuCOE/Nodes/CustomizableObjectNodeParameter.h
+++ b/Engine/Plugins/Mutable/Source/CustomizableObjectEditor/Private/MuCOE/Nodes/CustomizableObjectNodeParameter.h
@@ -51,6 +51,7 @@ public:
 	// Own interface
 	/** Return the pin category of this node. */
 	virtual FName GetCategory() const PURE_VIRTUAL(UCustomizableObjectNodeParameter::GetCategory, return {}; );
+	virtual FText GetCategoryFriendlyName() const;
 	UE_API virtual FString GetParameterName(TArray<const UCustomizableObjectNodeMacroInstance*>* MacroContext = nullptr) const;
 	UE_API virtual void SetParameterName(const FString& Name);
 };
diff --git a/Engine/Plugins/Mutable/Source/CustomizableObjectEditor/Private/MuCOE/Nodes/ExternalTypeCustomization.cpp b/Engine/Plugins/Mutable/Source/CustomizableObjectEditor/Private/MuCOE/Nodes/ExternalTypeCustomization.cpp
new file mode 100644
--- /dev/null
+++ b/Engine/Plugins/Mutable/Source/CustomizableObjectEditor/Private/MuCOE/Nodes/ExternalTypeCustomization.cpp
@@ -0,0 +1,40 @@
+// Copyright Epic Games, Inc. All Rights Reserved.
+
+// MyConfigCustomization.cpp
+
+#include "ExternalTypeCustomization.h"
+
+#include "DetailWidgetRow.h"
+#include "InstancedStructDetails.h"
+#include "StructUtilsMetadata.h"
+#include "Widgets/Layout/SBox.h"
+
+
+bool FExternalTypeTypeIdentifier::IsPropertyTypeCustomized(const IPropertyHandle& InPropertyHandle) const
+{
+	return InPropertyHandle.HasMetaData(TEXT("MutableExternalType"));
+}
+
+TSharedRef<IPropertyTypeCustomization> FExternalTypeCustomization::MakeInstance()
+{
+	return MakeShareable(new FExternalTypeCustomization());
+}
+
+void FExternalTypeCustomization::CustomizeHeader(TSharedRef<IPropertyHandle> PropertyHandle, FDetailWidgetRow& HeaderRow, IPropertyTypeCustomizationUtils& CustomizationUtils)
+{
+	// Make the type constant.
+	PropertyHandle->SetInstanceMetaData(UE::StructUtils::Metadata::StructTypeConstName, TEXT(""));
+
+	FInstancedStructDetails::CustomizeHeader(PropertyHandle, HeaderRow, CustomizationUtils);
+
+	// Disable reset button since it allows to remove the type.
+	HeaderRow.OverrideResetToDefault(FResetToDefaultOverride::Create(TAttribute<bool>(false), false));
+
+	// Remove type picker.
+	HeaderRow.ValueContent()
+		.VAlign(VAlign_Center)
+		[
+			SNew(SBox)
+		];
+}
+
diff --git a/Engine/Plugins/Mutable/Source/CustomizableObjectEditor/Private/MuCOE/Nodes/ExternalTypeCustomization.h b/Engine/Plugins/Mutable/Source/CustomizableObjectEditor/Private/MuCOE/Nodes/ExternalTypeCustomization.h
new file mode 100644
--- /dev/null
+++ b/Engine/Plugins/Mutable/Source/CustomizableObjectEditor/Private/MuCOE/Nodes/ExternalTypeCustomization.h
@@ -0,0 +1,24 @@
+// Copyright Epic Games, Inc. All Rights Reserved.
+
+#pragma once
+
+#include "PropertyEditorModule.h"
+
+#include "InstancedStructDetails.h"
+#include "IPropertyTypeCustomization.h"
+
+
+class FExternalTypeTypeIdentifier : public IPropertyTypeIdentifier 
+{
+public:
+	virtual bool IsPropertyTypeCustomized(const IPropertyHandle& InPropertyHandle) const override;
+};
+
+
+class FExternalTypeCustomization : public FInstancedStructDetails
+{
+public:
+	static TSharedRef<IPropertyTypeCustomization> MakeInstance();
+	virtual void CustomizeHeader(TSharedRef<IPropertyHandle> PropertyHandle, FDetailWidgetRow& HeaderRow, IPropertyTypeCustomizationUtils& CustomizationUtils) override;
+};
+
diff --git a/Engine/Plugins/Mutable/Source/CustomizableObjectEditor/Private/MuCOE/Nodes/SCONodeExternal.cpp b/Engine/Plugins/Mutable/Source/CustomizableObjectEditor/Private/MuCOE/Nodes/SCONodeExternal.cpp
deleted file mode 100644
diff --git a/Engine/Plugins/Mutable/Source/CustomizableObjectEditor/Private/MuCOE/Nodes/SCONodeExternal.h b/Engine/Plugins/Mutable/Source/CustomizableObjectEditor/Private/MuCOE/Nodes/SCONodeExternal.h
deleted file mode 100644
diff --git a/Engine/Plugins/Mutable/Source/CustomizableObjectEditor/Private/MuCOE/Nodes/SCustomizableObjectNode.cpp b/Engine/Plugins/Mutable/Source/CustomizableObjectEditor/Private/MuCOE/Nodes/SCustomizableObjectNode.cpp
--- a/Engine/Plugins/Mutable/Source/CustomizableObjectEditor/Private/MuCOE/Nodes/SCustomizableObjectNode.cpp
+++ b/Engine/Plugins/Mutable/Source/CustomizableObjectEditor/Private/MuCOE/Nodes/SCustomizableObjectNode.cpp
@@ -43,6 +43,18 @@ void SCustomizableObjectNode::CreateBelowPinControls(TSharedPtr<SVerticalBox> Ma
 				.ErrorText(LOCTEXT("SCONode_Experimental", "EXPERIMENTAL"))
 			];
 		}
+		
+		if (!Node->IsLoaded())
+		{
+			MainBox->AddSlot()
+			.Padding(FMargin(0.0, 2.0, 0.0, 0.0))
+			.AutoHeight()
+			[
+				SNew(SErrorText)
+				.BackgroundColor(FAppStyle::GetColor("ErrorReporting.BackgroundColor"))
+				.ErrorText(LOCTEXT("SCONodeExternal_NotLoaded", "NOT LOADED"))
+			];
+		}
 	}
 }
 
diff --git a/Engine/Plugins/Mutable/Source/CustomizableObjectEditor/Private/MuCOE/SMutableGraphViewer.cpp b/Engine/Plugins/Mutable/Source/CustomizableObjectEditor/Private/MuCOE/SMutableGraphViewer.cpp
--- a/Engine/Plugins/Mutable/Source/CustomizableObjectEditor/Private/MuCOE/SMutableGraphViewer.cpp
+++ b/Engine/Plugins/Mutable/Source/CustomizableObjectEditor/Private/MuCOE/SMutableGraphViewer.cpp
@@ -206,7 +206,9 @@ static const char* MutableNodeNames[] =
 	"ModifierTransformInMesh",
 	"ModifierTransformWithBone",
 
-	"Extension",
+	"External",
+	"ExternalOperation",
+	"ExternalParameter",
 	
 	"ExtensionData",
 	"ExtensionDataConstant",
diff --git a/Engine/Plugins/Mutable/Source/CustomizableObjectEditor/Public/MuCOE/EdGraphSchema_CustomizableObject.h b/Engine/Plugins/Mutable/Source/CustomizableObjectEditor/Public/MuCOE/EdGraphSchema_CustomizableObject.h
--- a/Engine/Plugins/Mutable/Source/CustomizableObjectEditor/Public/MuCOE/EdGraphSchema_CustomizableObject.h
+++ b/Engine/Plugins/Mutable/Source/CustomizableObjectEditor/Public/MuCOE/EdGraphSchema_CustomizableObject.h
@@ -76,7 +76,10 @@ public:
 	static UE_API FName GetPinCategoryName(const FName& PinCategory);
 	static UE_API FName GetPinCategoryName(const UScriptStruct* Type);
 	static UE_API FText GetPinCategoryFriendlyName(const FName& PinCategory);
+	static UE_API FText GetPinCategoryFriendlyName(const UScriptStruct* Type);
 	static UE_API bool IsPassthrough(const FName& PinCategory);
+	static UE_API bool IsTypeNative(const UScriptStruct* Type);
+
 
 private:
 	static TSharedPtr<IAssetReferenceFilter> MakeAssetReferenceFilter(const UEdGraph* Graph);
diff --git a/Engine/Plugins/Mutable/Source/CustomizableObjectEditor/Public/MuCOE/Nodes/CustomizableObjectNode.h b/Engine/Plugins/Mutable/Source/CustomizableObjectEditor/Public/MuCOE/Nodes/CustomizableObjectNode.h
--- a/Engine/Plugins/Mutable/Source/CustomizableObjectEditor/Public/MuCOE/Nodes/CustomizableObjectNode.h
+++ b/Engine/Plugins/Mutable/Source/CustomizableObjectEditor/Public/MuCOE/Nodes/CustomizableObjectNode.h
@@ -375,7 +375,9 @@ protected:
 public:
 	/** Return true to mark this node as experimental. */
 	UE_API virtual bool IsExperimental() const;
-	
+
+	UE_API virtual bool IsLoaded() const;
+
 private:
 	/** Data attached to a given pin. Not all pins contain data.
 	 * Data must be eliminated once the pins is removed. This is done automatically on CustomRemovePin.
diff --git a/Engine/Plugins/Mutable/Source/MutableRuntime/Internal/MuR/Model.h b/Engine/Plugins/Mutable/Source/MutableRuntime/Internal/MuR/Model.h
--- a/Engine/Plugins/Mutable/Source/MutableRuntime/Internal/MuR/Model.h
+++ b/Engine/Plugins/Mutable/Source/MutableRuntime/Internal/MuR/Model.h
@@ -69,11 +69,7 @@ namespace UE::Mutable::Private
 		//! Create a set of new parameters of the model with the default values.
 		//! If old parameters are provided, they will be reused when possible instead of the
 		//! default values.
-        static UE_API TSharedPtr<FParameters> NewParameters(TSharedPtr<const FModel> Model,
-        	const FParameters* OldParameters = nullptr,
-        	const TMap<FName, TObjectPtr<UTexture>>* ImageDefaults = nullptr,
-        	const TMap<FName, TObjectPtr<USkeletalMesh>>* MeshDefaults = nullptr,
-        	const TMap<FName, TObjectPtr<UMaterialInterface>>* MaterialDefaults = nullptr);
+        static UE_API TSharedPtr<FParameters> NewParameters(TSharedPtr<const FModel> Model, const FParameters* OldParameters = nullptr);
 
 		/** Return true if the parameter is multi-dimensional */
 		UE_API bool IsParameterMultidimensional(int32 ParameterIndex) const;
diff --git a/Engine/Plugins/Mutable/Source/MutableRuntime/Internal/MuR/Operations.h b/Engine/Plugins/Mutable/Source/MutableRuntime/Internal/MuR/Operations.h
--- a/Engine/Plugins/Mutable/Source/MutableRuntime/Internal/MuR/Operations.h
+++ b/Engine/Plugins/Mutable/Source/MutableRuntime/Internal/MuR/Operations.h
@@ -54,6 +54,7 @@ namespace UE::Mutable::Private
 		ST_PARAMETER,
     	MA_PARAMETER,
     	MI_PARAMETER,
+    	IS_PARAMETER,
 
 		//! A referenced, but opaque engine resource
 		IM_REFERENCE,
@@ -115,6 +116,9 @@ namespace UE::Mutable::Private
         //! Get a scalar value from a curve
         SC_CURVE,
 
+    	//! External operation
+		SC_EXTERNAL,
+
         //-----------------------------------------------------------------------------------------
         // Colour operations. Colours are sometimes used as generic vectors.
         //-----------------------------------------------------------------------------------------
@@ -134,6 +138,9 @@ namespace UE::Mutable::Private
 		//! Apply a Linear to sRGB color transformation on a given color vector.
 		CO_LINEARTOSRGB,
 
+    	//! External operation
+		CO_EXTERNAL,
+    	
         //-----------------------------------------------------------------------------------------
         // Image operations
         //-----------------------------------------------------------------------------------------
@@ -218,6 +225,9 @@ namespace UE::Mutable::Private
 
     	IM_PARAMETER_CONVERT,
 
+    	//! External operation
+		IM_EXTERNAL,
+    	
         //-----------------------------------------------------------------------------------------
         // Mesh operations
         //-----------------------------------------------------------------------------------------
@@ -300,6 +310,9 @@ namespace UE::Mutable::Private
 		//! Transform with a 4x4 matrix the geometry channels of a mesh that are skinned to a bone or hierarchy
 		ME_TRANSFORMWITHBONE,
 
+    	//! External operation
+    	ME_EXTERNAL,
+
         //-----------------------------------------------------------------------------------------
         // Instance operations
         //-----------------------------------------------------------------------------------------
@@ -355,11 +368,18 @@ namespace UE::Mutable::Private
 		LA_FROMMESH,
 
     	//-----------------------------------------------------------------------------------------
-		// External operations
+	    // Material operations
+	    //-----------------------------------------------------------------------------------------
+
+    	//! External operation
+		MI_EXTERNAL,
+
+    	//-----------------------------------------------------------------------------------------
+		// FInstancedStruct operations
 		//-----------------------------------------------------------------------------------------
 
-    	//! External operation. Returns unknown type.
-		UN_EXTERNAL, // TODO GMT Instruction for each type
+		//! External operation
+		IS_EXTERNAL,
     	
         //-----------------------------------------------------------------------------------------
         // Utility values
@@ -1191,7 +1211,6 @@ namespace UE::Mutable::Private
 		Skeleton,
 		PhysicsAsset,
 
-		Unknown, // Not a type by itself. Unknown at compile time but at runtime it has a valid type. 
 		Count
 	};
 
diff --git a/Engine/Plugins/Mutable/Source/MutableRuntime/Internal/MuR/Parameters.h b/Engine/Plugins/Mutable/Source/MutableRuntime/Internal/MuR/Parameters.h
--- a/Engine/Plugins/Mutable/Source/MutableRuntime/Internal/MuR/Parameters.h
+++ b/Engine/Plugins/Mutable/Source/MutableRuntime/Internal/MuR/Parameters.h
@@ -8,6 +8,7 @@
 #include "Templates/SharedPointer.h"
 #include "VectorTypes.h"
 
+struct FInstancedStruct;
 class UTexture;
 class USkeletalMesh;
 class UMaterialInterface;
@@ -61,6 +62,8 @@ namespace UE::Mutable::Private
     	/** A 4x4 matrix. */
     	Matrix,
 
+    	InstancedStruct,
+
         /** Utility enumeration value, not really a parameter type. */
         Count
 
@@ -360,14 +363,27 @@ namespace UE::Mutable::Private
 		//! \pre The parameter specified by index is a T_MATERIAL.
 		//! \param index Index of the parameter from 0 to GetCount()-1
 		//! \param pos Only for multidimensional parameters: relevant position to set in the ranges
-		//! \return The Material interface specified when setting the image value (\see SetImageValue)
+		//! \return The Material interface specified when setting the Material value (\see SetMaterialValue)
 		UE_API UMaterialInterface* GetMaterialValue(int32 Index, const FRangeIndex* RangeIndex = nullptr) const;
 
 		//! If the parameter is of the material type, set its value.
 		//! \param index Index of the parameter from 0 to GetCount()-1
-		//! \param ExternalResource Material interface to set as value.
+		//! \param Value Material interface to set as value.
 		//! \param pos Only for multidimensional parameters: relevant position to set in the ranges
-		UE_API void SetMaterialValue(int32 Index, UMaterialInterface* ExternalResource, const FRangeIndex* RangeIndex = nullptr);
+		UE_API void SetMaterialValue(int32 Index, UMaterialInterface* Value, const FRangeIndex* RangeIndex = nullptr);
+
+		/** Return the value of a Instanced Struct parameter.
+		* @pre The parameter specified by index is a T_INSTANCEDSTRUCT.
+		* @param Index Index of the parameter from 0 to GetCount()-1
+		* @param RangeIndex Only for multidimensional parameters: relevant position to set in the ranges
+		* @return The Instanced Struct interface specified when setting the Instanced Struct value (\see SetInstancedStructValue) */
+		UE_API TSharedPtr<const FInstancedStruct> GetInstancedStructValue(int32 Index, const FRangeIndex* RangeIndex = nullptr) const;
+
+		//! If the parameter is of the Instanced Struct type, set its value.
+		//! \param Index Index of the parameter from 0 to GetCount()-1
+		//! \param Value Instanced Struct interface to set as value.
+		//! \param RangeIndex Only for multidimensional parameters: relevant position to set in the ranges
+		UE_API void SetInstancedStructValue(int32 Index, const TSharedPtr<const FInstancedStruct>& Value, const FRangeIndex* RangeIndex = nullptr);
 
         //! Utility method to compare the values of a specific parameter with the values of another
         //! FParameters object. It returns false if type or values are different.
diff --git a/Engine/Plugins/Mutable/Source/MutableRuntime/Internal/MuR/ParametersPrivate.h b/Engine/Plugins/Mutable/Source/MutableRuntime/Internal/MuR/ParametersPrivate.h
--- a/Engine/Plugins/Mutable/Source/MutableRuntime/Internal/MuR/ParametersPrivate.h
+++ b/Engine/Plugins/Mutable/Source/MutableRuntime/Internal/MuR/ParametersPrivate.h
@@ -19,6 +19,7 @@
 
 #define UE_API MUTABLERUNTIME_API
 
+
 namespace UE::Mutable::Private
 {
 
@@ -262,9 +263,20 @@ namespace Private
 	using FParamStringType = Private::TIndirectObject<FString>;
 	using FParamMatrixType = Private::TIndirectObject<FMatrix44f>;
 	using FParamMaterialType = TStrongObjectPtr<UMaterialInterface>;
+	using FParamInstancedStructType = TSharedPtr<const FInstancedStruct>; // TSharedPtr always valid and FInstancedStruct always initialized to its type.
 	
 	using FParameterValue = TVariant<
-            FParamBoolType, FParamIntType, FParamFloatType, FParamColorType, FParamProjectorType, FParamTextureType, FParamSkeletalMeshType, FParamStringType, FParamMatrixType, FParamMaterialType>;
+		FParamBoolType,
+		FParamIntType,
+		FParamFloatType,
+		FParamColorType,
+		FParamProjectorType,
+		FParamTextureType,
+		FParamSkeletalMeshType,
+		FParamStringType,
+		FParamMatrixType,
+		FParamMaterialType,
+		FParamInstancedStructType>;
 
     // static_assert to track ParameterValue size changes. It is ok to change if needed.
     static_assert(sizeof(FParameterValue) == 8*4, "ParameterValue size has changed.");
diff --git a/Engine/Plugins/Mutable/Source/MutableRuntime/Internal/MuR/Serialisation.h b/Engine/Plugins/Mutable/Source/MutableRuntime/Internal/MuR/Serialisation.h
--- a/Engine/Plugins/Mutable/Source/MutableRuntime/Internal/MuR/Serialisation.h
+++ b/Engine/Plugins/Mutable/Source/MutableRuntime/Internal/MuR/Serialisation.h
@@ -14,6 +14,7 @@
 #include "Containers/Array.h"
 #include "Containers/Map.h"
 #include "Containers/ContainerAllocationPolicies.h"
+#include "StructUtils/InstancedStruct.h"
 #include "Templates/SharedPointer.h"
 #include "UObject/StrongObjectPtr.h"
 
@@ -411,8 +412,7 @@ namespace UE::Mutable::Private
 	MUTABLERUNTIME_API void operator<<(FOutputArchive&, const bool&);
 	MUTABLERUNTIME_API void operator>>(FInputArchive&, bool&);
 
-
-	// Do not serialize pointers to UObjects.
+	
 	template<typename T>
 	void operator<<(FOutputArchive& Arch, const TStrongObjectPtr<T>& V)
 	{
@@ -424,6 +424,9 @@ namespace UE::Mutable::Private
 	{
     	// Do Nothing. UObjects values are saved in the CO.	
 	}
+
+	MUTABLERUNTIME_API void operator<<(FOutputArchive& Arch, const TSharedPtr<const FInstancedStruct>& V);
+	MUTABLERUNTIME_API void operator>>(FInputArchive& Arch, TSharedPtr<const FInstancedStruct>& V);
 }
 
 
diff --git a/Engine/Plugins/Mutable/Source/MutableRuntime/Internal/MuR/SystemPrivate.h b/Engine/Plugins/Mutable/Source/MutableRuntime/Internal/MuR/SystemPrivate.h
--- a/Engine/Plugins/Mutable/Source/MutableRuntime/Internal/MuR/SystemPrivate.h
+++ b/Engine/Plugins/Mutable/Source/MutableRuntime/Internal/MuR/SystemPrivate.h
@@ -565,6 +565,7 @@ namespace UE::Mutable::Private
 				MatrixResults.Add(FMatrix44f());
 				ExtensionDataResults.Add(nullptr);
 				MaterialResults.Add(nullptr);
+				InstancedStructResults.Add({});
 			}
 		}
 
@@ -595,6 +596,10 @@ namespace UE::Mutable::Private
 					ExtensionDataResults[DataTypeIndex] = nullptr;
 					break;
 
+				case EDataType::InstancedStruct:
+					InstancedStructResults[DataTypeIndex] = {};
+					break;
+
 				default:
 					break;
 				}
@@ -888,12 +893,17 @@ namespace UE::Mutable::Private
 		
 		TSharedPtr<const FInstancedStruct> GetInstancedStruct(FCacheAddress At)
 		{
-			if (!At.At) return nullptr;
+			check(At.At);
 			FOpExecutionData* Data = OpExecutionData.get_ptr(At);
-			if (!Data) return nullptr;
+			check(Data);
 
 			check(Data->DataType == EDataType::InstancedStruct);
-			return InstancedStructResults[Data->DataTypeIndex];
+			TSharedPtr<const FInstancedStruct> Result = InstancedStructResults[Data->DataTypeIndex];
+
+			check(Result);
+			check(Result->IsValid());
+
+			return Result;
 		}
 		
 		
@@ -1055,7 +1065,7 @@ namespace UE::Mutable::Private
 			check(Data->DataTypeIndex != 0);
 		}
 
-		void SetInstancedStruct(FCacheAddress At, TSharedPtr<const FInstancedStruct> Value)
+		void SetInstancedStruct(FCacheAddress At, const TSharedPtr<const FInstancedStruct>& Value)
 		{
 			check(At.At < OpExecutionData.size_code());
 			FOpExecutionData* Data = OpExecutionData.get_ptr(At);
diff --git a/Engine/Plugins/Mutable/Source/MutableRuntime/Internal/MuR/Types.h b/Engine/Plugins/Mutable/Source/MutableRuntime/Internal/MuR/Types.h
--- a/Engine/Plugins/Mutable/Source/MutableRuntime/Internal/MuR/Types.h
+++ b/Engine/Plugins/Mutable/Source/MutableRuntime/Internal/MuR/Types.h
@@ -6,6 +6,7 @@
 #include "Containers/Map.h"
 #include "Containers/Array.h"
 #include "Containers/UnrealString.h"
+#include "StructUtils/InstancedStruct.h"
 
 #include "Types.generated.h"
 
@@ -65,3 +66,18 @@ inline bool operator==(const TMap<K,V>& A, const TMap<K,V>& B )
 }
 
 
+namespace UE::Mutable::Private
+{
+	inline uint32 GetTypeHash(const FInstancedStruct& InstancedStruct)
+	{
+		if (!InstancedStruct.IsValid())
+		{
+			return 0;
+		}
+
+		const UScriptStruct* ScriptedStruct = InstancedStruct.GetScriptStruct();
+		const uint32 ValueHash = ScriptedStruct->GetStructTypeHash(InstancedStruct.GetMemory());
+
+		return HashCombine(ValueHash, PointerHash(ScriptedStruct));
+	}
+}
diff --git a/Engine/Plugins/Mutable/Source/MutableRuntime/Private/MuR/BitWriterResourceKey.cpp b/Engine/Plugins/Mutable/Source/MutableRuntime/Private/MuR/BitWriterResourceKey.cpp
new file mode 100644
--- /dev/null
+++ b/Engine/Plugins/Mutable/Source/MutableRuntime/Private/MuR/BitWriterResourceKey.cpp
@@ -0,0 +1,17 @@
+// Copyright Epic Games, Inc. All Rights Reserved.
+
+#include "BitWriterResourceKey.h"
+
+#include "UObject/ObjectPtr.h"
+
+
+FBitWriterResourceKey::FBitWriterResourceKey(int64 InMaxBits, bool AllowResize): FBitWriter(InMaxBits, AllowResize)
+{
+}
+
+
+FArchive& FBitWriterResourceKey::operator<<(FObjectPtr& Value)
+{
+	Serialize(Value.Get(), sizeof(UObject*));
+	return *this;
+}
diff --git a/Engine/Plugins/Mutable/Source/MutableRuntime/Private/MuR/BitWriterResourceKey.h b/Engine/Plugins/Mutable/Source/MutableRuntime/Private/MuR/BitWriterResourceKey.h
new file mode 100644
--- /dev/null
+++ b/Engine/Plugins/Mutable/Source/MutableRuntime/Private/MuR/BitWriterResourceKey.h
@@ -0,0 +1,14 @@
+// Copyright Epic Games, Inc. All Rights Reserved.
+
+#pragma once
+
+#include "Serialization/BitWriter.h"
+
+
+class FBitWriterResourceKey : public FBitWriter
+{
+public:
+	FBitWriterResourceKey(int64 InMaxBits, bool AllowResize = false);
+
+	virtual FArchive& operator<<(FObjectPtr& Value) override;
+};
diff --git a/Engine/Plugins/Mutable/Source/MutableRuntime/Private/MuR/CodeRunner.cpp b/Engine/Plugins/Mutable/Source/MutableRuntime/Private/MuR/CodeRunner.cpp
--- a/Engine/Plugins/Mutable/Source/MutableRuntime/Private/MuR/CodeRunner.cpp
+++ b/Engine/Plugins/Mutable/Source/MutableRuntime/Private/MuR/CodeRunner.cpp
@@ -6780,7 +6780,12 @@ namespace UE::Mutable::Private
             RunCode_InstanceAddResource(Item, InModel, pParams);
             break;
 
-		case EOpType::UN_EXTERNAL:
+		case EOpType::SC_EXTERNAL:
+		case EOpType::CO_EXTERNAL:
+		case EOpType::IM_EXTERNAL:
+		case EOpType::ME_EXTERNAL:
+		case EOpType::MI_EXTERNAL:
+		case EOpType::IS_EXTERNAL:
 			{
 				const uint8* Data = Program.GetOpArgsPointer(Item.At);
 
@@ -6796,6 +6801,14 @@ namespace UE::Mutable::Private
 				}
 					
 				const FExternalOperation* Operation = OperationInstancedStruct.GetPtr<FExternalOperation>();
+
+				const TArray<TPair<FName, const UScriptStruct*>>& Inputs = Operation->GetInputs();
+				if (Inputs.Num() != Args.NumOperants)
+				{
+					UE_LOG(LogMutableCore, Error, TEXT("External Operation inputs do not match compiled data: %s"), *OperationInstancedStruct.GetScriptStruct()->GetName());
+					bOutFailed = true;
+					return;
+				}
 				
 				switch (Item.Stage)
 				{
@@ -6820,10 +6833,7 @@ namespace UE::Mutable::Private
 				case 1:
 					{
 						FContext Context;
-
-						const TArray<TPair<FName, const UScriptStruct*>>& Inputs = Operation->GetInputs();
-						check(Inputs.Num() == Args.NumOperants);
-
+						
 						// Inputs
 						for (int32 Index = 0; Index < Args.NumOperants; ++Index)
 						{
@@ -6833,88 +6843,127 @@ namespace UE::Mutable::Private
 
 							const TPair<FName, const UScriptStruct*>& Input = Inputs[Index];
 
-							const UScriptStruct* Type = Input.Value;
-							if (Type == FMeshAdapter::StaticStruct())
-							{
-								TSharedPtr<const FMesh> Value = LoadMesh(FCacheAddress(OperantAt, Item));
-	
-								TSharedPtr<FInstancedStruct> InstancedStruct = MakeShared<FInstancedStruct>();
-								InstancedStruct->InitializeAs<FMeshAdapter>().Mesh = ConstCastSharedPtr<FMesh>(Value);
-
-								Context.Inputs.Add(Input.Key, MoveTemp(InstancedStruct));
-							}
-							else if (Type == FTextureAdapter::StaticStruct())
-							{
-								TSharedPtr<const FImage> Value = LoadImage(FCacheAddress(OperantAt, Item));
-
-								TSharedPtr<FInstancedStruct> InstancedStruct = MakeShared<FInstancedStruct>();
-								InstancedStruct->InitializeAs<FTextureAdapter>().Image = ConstCastSharedPtr<FImage>(Value);
-
-								Context.Inputs.Add(Input.Key, MoveTemp(InstancedStruct));
-							}
-							else if (Type == FFloatAdapter::StaticStruct())
+							const EOpType InputOpType = Program.GetOpType(OperantAt);
+							const EDataType InputDataType = GetOpDataType(InputOpType);
+							switch (InputDataType)
 							{
-								float Value = LoadScalar(FCacheAddress(OperantAt, Item));
+							case EDataType::Color:
+								{
+									const FVector4f Vector = LoadColor(FCacheAddress(OperantAt, Item));
+									
+									FValue Value;
+									Value.Ptr = MakeShared<FInstancedStruct>();
+									Value.Ptr->InitializeAs<FVectorAdapter>().Value = Vector;
+
+									Context.Inputs.Add(Input.Key, MoveTemp(Value));
+									break;
+								}
+								
+							case EDataType::Image:
+								{
+									const TSharedPtr<const FImage> Image = LoadImage(FCacheAddress(OperantAt, Item));
+									
+									FValue Value;
+									Value.Ptr = MakeShared<FInstancedStruct>();
+									Value.Ptr->InitializeAs<FTextureAdapter>().Image = ConstCastSharedPtr<FImage>(Image);
+									
+									Context.Inputs.Add(Input.Key, MoveTemp(Value));
+									break;
+								}
 
-								TSharedPtr<FInstancedStruct> InstancedStruct = MakeShared<FInstancedStruct>();
-								InstancedStruct->InitializeAs<FFloatAdapter>().Value = Value;
+							case EDataType::Mesh:
+								{
+									const TSharedPtr<const FMesh> Mesh = LoadMesh(FCacheAddress(OperantAt, Item));
 
-								Context.Inputs.Add(Input.Key, MoveTemp(InstancedStruct));
-							}
-							else if (Type == FVectorAdapter::StaticStruct())
-							{
-								FVector4f Value = LoadColor(FCacheAddress(OperantAt, Item));
+									FValue Value;
+									Value.Ptr = MakeShared<FInstancedStruct>();
+									Value.Ptr->InitializeAs<FMeshAdapter>().Mesh = ConstCastSharedPtr<FMesh>(Mesh);
 
-								TSharedPtr<FInstancedStruct> InstancedStruct = MakeShared<FInstancedStruct>();
-								InstancedStruct->InitializeAs<FVectorAdapter>().Value = Value;
+									Context.Inputs.Add(Input.Key, MoveTemp(Value));
+									break;
+								}
+								
+							case EDataType::Material:
+								{
+									const TSharedPtr<const FMaterial> Material = LoadMaterial(FCacheAddress(OperantAt, Item));
 
-								Context.Inputs.Add(Input.Key, MoveTemp(InstancedStruct));
-							}
-							else if (Type == FMaterialAdapter::StaticStruct())
-							{
-								TSharedPtr<const FMaterial> Value = LoadMaterial(FCacheAddress(OperantAt, Item));
+									FValue Value;
+									Value.Ptr = MakeShared<FInstancedStruct>();
+									Value.Ptr->InitializeAs<FMaterialAdapter>().Material = ConstCastSharedPtr<FMaterial>(Material);
 
-								TSharedPtr<FInstancedStruct> InstancedStruct = MakeShared<FInstancedStruct>();
-								InstancedStruct->InitializeAs<FMaterialAdapter>().Material = ConstCastSharedPtr<FMaterial>(Value);
+									Context.Inputs.Add(Input.Key, MoveTemp(Value));
+									break;
+								}
+								
+							case EDataType::InstancedStruct:
+								{
+									const TSharedPtr<const FInstancedStruct> InstancedStruct = LoadInstancedStruct(FCacheAddress(OperantAt, Item));
+									
+									FValue Value;
+									Value.Ptr = MakeShared<FInstancedStruct>();
+									Value.Ptr = ConstCastSharedPtr<FInstancedStruct>(InstancedStruct);
+									
+									Context.Inputs.Add(Input.Key, MoveTemp(Value));
+									break;
+								}
+									
+							case EDataType::Scalar:
+								{
+									const float Float = LoadScalar(FCacheAddress(OperantAt, Item));
 
-								Context.Inputs.Add(Input.Key, MoveTemp(InstancedStruct));
-							}
-							else
-							{
-								TSharedPtr<const FInstancedStruct> Value = LoadInstancedStruct(FCacheAddress(OperantAt, Item));
+									FValue Value;
+									Value.Ptr = MakeShared<FInstancedStruct>();
+									Value.Ptr->InitializeAs<FFloatAdapter>().Value = Float;
 
-								Context.Inputs.Add(Input.Key, Value);
+									Context.Inputs.Add(Input.Key, MoveTemp(Value));
+									break;
+								}
+									
+							default: 
+								unimplemented();
 							}
 						}
 
 						// Evaluate
 						Operation->Evaluate(Context);
 
-						// Store output
-						const UStruct* Type = Context.Output->GetScriptStruct();
-						if (Type == FMeshAdapter::StaticStruct())
-						{
-							StoreMesh(FCacheAddress(Item.At, Item), Context.Output.Get()->Get<FMeshAdapter>().Mesh);
-						}
-						else if (Type == FTextureAdapter::StaticStruct())
-						{
-							StoreImage(FCacheAddress(Item.At, Item), Context.Output.Get()->Get<FTextureAdapter>().Image);
-						}
-						else if (Type == FFloatAdapter::StaticStruct())
-						{
-							StoreScalar(FCacheAddress(Item.At, Item), Context.Output.Get()->Get<FFloatAdapter>().Value);
-						}
-						else if (Type == FVectorAdapter::StaticStruct())
-						{
-							StoreColor(FCacheAddress(Item.At, Item), Context.Output.Get()->Get<FVectorAdapter>().Value);
-						}
-						else if (Type == FMaterialAdapter::StaticStruct())
+						if (!Context.Output.Ptr)
 						{
-							StoreMaterial(FCacheAddress(Item.At, Item), Context.Output.Get()->Get<FMaterialAdapter>().Material);
+							UE_LOG(LogMutableCore, Error, TEXT("External Operation output not set: %s"), *OperationInstancedStruct.GetScriptStruct()->GetName());
+							bOutFailed = true;
+							return;
 						}
-						else
+						
+						// Store output
+						const EDataType OuputDataType = GetOpDataType(type);
+						switch (OuputDataType)
 						{
-							StoreInstancedStruct(FCacheAddress(Item.At, Item), Context.Output);
+						case EDataType::Color:
+							StoreColor(FCacheAddress(Item.At, Item), Context.Output.Get<FVectorAdapter>().Value);
+							break;
+							
+						case EDataType::Image:
+							StoreImage(FCacheAddress(Item.At, Item), Context.Output.Get<FTextureAdapter>().Image);
+							break;
+							
+						case EDataType::Mesh:
+							StoreMesh(FCacheAddress(Item.At, Item), Context.Output.Get<FMeshAdapter>().Mesh);
+							break;
+							
+						case EDataType::Material:
+							StoreMaterial(FCacheAddress(Item.At, Item), Context.Output.Get<FMaterialAdapter>().Material);
+							break;
+							
+						case EDataType::InstancedStruct:
+							StoreInstancedStruct(FCacheAddress(Item.At, Item), Context.Output.Ptr);
+							break;
+									
+						case EDataType::Scalar:
+							StoreScalar(FCacheAddress(Item.At, Item), Context.Output.Get<FFloatAdapter>().Value);
+							break;
+									
+						default: 
+							unimplemented();
 						}
 
 						break;
@@ -6926,6 +6975,16 @@ namespace UE::Mutable::Private
 
 				break;	
 			}
+
+		case EOpType::IS_PARAMETER:
+			{
+				OP::ParameterArgs Args = Program.GetOpArgs<OP::ParameterArgs>(Item.At);
+				TSharedPtr<FRangeIndex> Index = BuildCurrentOpRangeIndex(Item, Args.variable);
+				TSharedPtr<const FInstancedStruct> Result = pParams->GetInstancedStructValue(Args.variable, Index.Get());
+				StoreInstancedStruct(Item, Result);
+
+				break;
+			}
 			
 		default:
 		{
@@ -6979,6 +7038,7 @@ namespace UE::Mutable::Private
 			case EDataType::Material:
 				RunCode_Material(Item, pParams, InModel.Get());
 				break;
+				
 			default:
 				check(false);
 				break;
diff --git a/Engine/Plugins/Mutable/Source/MutableRuntime/Private/MuR/CodeRunnerBegin.cpp b/Engine/Plugins/Mutable/Source/MutableRuntime/Private/MuR/CodeRunnerBegin.cpp
--- a/Engine/Plugins/Mutable/Source/MutableRuntime/Private/MuR/CodeRunnerBegin.cpp
+++ b/Engine/Plugins/Mutable/Source/MutableRuntime/Private/MuR/CodeRunnerBegin.cpp
@@ -368,7 +368,7 @@ namespace UE::Mutable::Private
 				case EDataType::Image:		
 				case EDataType::Layout:		
 				case EDataType::Material:		
-				case EDataType::Unknown:		
+				case EDataType::InstancedStruct:		
 					break;
 
 				default:
@@ -392,6 +392,7 @@ namespace UE::Mutable::Private
 			case EOpType::IM_PARAMETER:
 			case EOpType::MI_PARAMETER:
 			case EOpType::IM_PARAMETER_FROM_MATERIAL:
+			case EOpType::IS_PARAMETER:
 				{
 					StoreNone(Item);
 					break;
@@ -2365,7 +2366,12 @@ namespace UE::Mutable::Private
 					break;
 				}
 
-			case EOpType::UN_EXTERNAL:
+			case EOpType::SC_EXTERNAL:
+			case EOpType::CO_EXTERNAL:
+			case EOpType::IM_EXTERNAL:
+			case EOpType::ME_EXTERNAL:
+			case EOpType::MI_EXTERNAL:
+			case EOpType::IS_EXTERNAL:
 				{
 					const uint8* Data = Program.GetOpArgsPointer(Item.At);
 
@@ -2379,6 +2385,15 @@ namespace UE::Mutable::Private
 						break;
 					}
 
+					const FExternalOperation* Operation = OperationInstancedStruct.GetPtr<FExternalOperation>();
+
+					const TArray<TPair<FName, const UScriptStruct*>>& Inputs = Operation->GetInputs();
+					if (Inputs.Num() != Args.NumOperants)
+					{
+						StoreNone(Item);
+						break;
+					}
+
 					switch(Item.Stage)
 					{
 					case 0:
@@ -2399,44 +2414,30 @@ namespace UE::Mutable::Private
 
 					case 1:
 						{
-							const FExternalOperation* Operation = OperationInstancedStruct.GetPtr<FExternalOperation>();
-
-							const TArray<TPair<FName, const UScriptStruct*>>& Inputs = Operation->GetInputs();
-							check(Inputs.Num() == Args.NumOperants);
-							
 							for (int32 Index = 0; Index < Args.NumOperants; ++Index)
 							{
 								OP::ADDRESS OperantAt;
 								FMemory::Memcpy(&OperantAt, Data, sizeof(OP::ADDRESS));
 								Data += sizeof(OP::ADDRESS);
 
-								const TPair<FName, const UScriptStruct*>& Input = Inputs[Index];
-
-								const UScriptStruct* InputType = Input.Value;
-								
-								if (InputType == FMeshAdapter::StaticStruct())
-								{
-									// Nothing.
-								}
-								else if (InputType == FTextureAdapter::StaticStruct())
-								{
-									// Nothing.
-								}
-								else if (InputType == FFloatAdapter::StaticStruct())
+								const EOpType InputOpType = Program.GetOpType(OperantAt);
+								const EDataType InputDataType = GetOpDataType(InputOpType);
+								switch (InputDataType)
 								{
+								case EDataType::Color:
+								case EDataType::Image:
+								case EDataType::Mesh:
+								case EDataType::Material:
+								case EDataType::InstancedStruct:
+									// Nothing
+									break;
+									
+								case EDataType::Scalar:
 									LoadFloat(OperantAt);
-								}
-								else if (InputType == FVectorAdapter::StaticStruct())
-								{
-									// Nothing.
-								}
-								else if (InputType == FMaterialAdapter::StaticStruct())
-								{
-									// Nothing.
-								}
-								else
-								{
-									// Nothing.
+									break;
+									
+								default: 
+									unimplemented();
 								}
 							}
 
diff --git a/Engine/Plugins/Mutable/Source/MutableRuntime/Private/MuR/External/MaterialAdapter.cpp b/Engine/Plugins/Mutable/Source/MutableRuntime/Private/MuR/External/MaterialAdapter.cpp
--- a/Engine/Plugins/Mutable/Source/MutableRuntime/Private/MuR/External/MaterialAdapter.cpp
+++ b/Engine/Plugins/Mutable/Source/MutableRuntime/Private/MuR/External/MaterialAdapter.cpp
@@ -3,6 +3,7 @@
 #include "MuR/External/MaterialAdapter.h"
 
 #include "MuR/Material.h"
+#include "MuR/Image.h"
 #include "MuR/External/FloatAdapter.h"
 #include "MuR/External/TextureAdapter.h"
 #include "MuR/External/VectorAdapter.h"
@@ -18,6 +19,7 @@ namespace UE::Mutable
 
 	FMaterialAdapter::FMaterialAdapter(const FMaterialAdapter& Other)
 	{
+		Material = MakeShared<Private::FMaterial>();
 		*Material = *Other.Material;
 	}
 
@@ -29,90 +31,123 @@ namespace UE::Mutable
 	}
 
 
-	TSharedPtr<const FFloatAdapter> FMaterialAdapter::GetFloat(const FName& ParameterName) const
+	TOptional<TValueConst<FFloatAdapter>> FMaterialAdapter::GetFloat(const FName& ParameterName) const
 	{
 		float* Result = Material->ScalarParameters.Find(ParameterName);
 		if (!Result)
 		{
-			return nullptr;
+			return {};
 		}
 
-		TSharedPtr<FFloatAdapter> Adapter = MakeShared<FFloatAdapter>();
-		Adapter->Value = *Result;
+		TValue<FFloatAdapter> Value;
+		Value.Ptr = MakeShared<FInstancedStruct>();
+		Value.Ptr->InitializeAs(FFloatAdapter::StaticStruct());
+		Value.Ptr->GetMutable<FFloatAdapter>().Value = *Result;
 		
-		return Adapter;
+		return TOptional<TValueConst<FFloatAdapter>>(MoveTemp(Value));
 	}
 
 
-	bool FMaterialAdapter::SetFloat(const FName& ParameterName, const TSharedPtr<const FFloatAdapter>& Adapter)
+	bool FMaterialAdapter::SetFloat(const FName& ParameterName, const TValueConst<FFloatAdapter>* Value)
 	{
-		if (!Adapter)
+		if (!Material->ScalarParameters.Contains(ParameterName))
 		{
 			return false;
 		}
 
-		Material->ScalarParameters.Add(ParameterName, Adapter->GetValue());
+		if (Value)
+		{
+			Material->ScalarParameters.Add(ParameterName, Value->Get().GetValue());
+		}
+		else
+		{
+			Material->ScalarParameters.Remove(ParameterName);
+		}
+		
 		return true;
 	}
 
 
-	TSharedPtr<const FVectorAdapter> FMaterialAdapter::GetVector(const FName& ParameterName) const
+	TOptional<TValueConst<FVectorAdapter>> FMaterialAdapter::GetVector(const FName& ParameterName) const
 	{
 		FVector4f* Result = Material->ColorParameters.Find(ParameterName);
 		if (!Result)
 		{
-			return nullptr;
+			return {};
 		}
 
-		TSharedPtr<FVectorAdapter> Adapter = MakeShared<FVectorAdapter>();
-		Adapter->Value = *Result;
+		TValue<FVectorAdapter> Value;
+		Value.Ptr = MakeShared<FInstancedStruct>();
+		Value.Ptr->InitializeAs(FVectorAdapter::StaticStruct());
+		Value.Ptr->GetMutable<FVectorAdapter>().Value = *Result;
 		
-		return Adapter;
+		return TOptional<TValueConst<FVectorAdapter>>(MoveTemp(Value));
 	}
 
 
-	bool FMaterialAdapter::SetVector(const FName& ParameterName, const TSharedPtr<const FVectorAdapter>& Adapter)
+	bool FMaterialAdapter::SetVector(const FName& ParameterName, const TValueConst<FVectorAdapter>* Value)
 	{
-		if (!Adapter)
+		if (!Material->ColorParameters.Contains(ParameterName))
 		{
 			return false;
 		}
+
+		if (Value)
+		{
+			Material->ColorParameters.Add(ParameterName, Value->Get().GetValue());
+		}
+		else
+		{
+			Material->ColorParameters.Remove(ParameterName);
+		}
 		
-		Material->ColorParameters.Add(ParameterName, Adapter->GetValue());
 		return true;
 	}
 
 
-	TSharedPtr<const FTextureAdapter> FMaterialAdapter::GetTexture(const FName& ParameterName) const
+	TOptional<TValueConst<FTextureAdapter>> FMaterialAdapter::GetTexture(const FName& ParameterName) const
 	{
 		TVariant<Private::OP::ADDRESS, TSharedPtr<const Private::FImage>>* Result = Material->ImageParameters.Find(ParameterName);
 		if (!Result)
 		{
-			return nullptr;
+			return {};
 		}
 
 		if (ensure(Result->IsType<TSharedPtr<const Private::FImage>>()))
 		{
-			TSharedPtr<FTextureAdapter> Adapter = MakeShared<FTextureAdapter>(); 
-			Adapter->Image = ConstCastSharedPtr<Private::FImage>(Result->Get<TSharedPtr<const Private::FImage>>()); // TODO GMT If the image is lazy what should we do? Force eval before calling an UN_EXTERNAL?
-
-			return Adapter;
+			TValue<FTextureAdapter> Value;
+			Value.Ptr = MakeShared<FInstancedStruct>();
+			Value.Ptr->InitializeAs(FTextureAdapter::StaticStruct());
+			Value.Ptr->GetMutable<FTextureAdapter>().Image = ConstCastSharedPtr<Private::FImage>(Result->Get<TSharedPtr<const Private::FImage>>()); // TODO GMT If lazy wait
+			
+			return TOptional<TValueConst<FTextureAdapter>>(MoveTemp(Value));
 		}
 		else
 		{
-			return nullptr;
+			return {};
 		}
 	}
 
 	
-	void FMaterialAdapter::SetTexture(const FName& ParameterName, const TSharedPtr<FTextureAdapter>& Adapter)
+	bool FMaterialAdapter::SetTexture(const FName& ParameterName, const TValueConst<FTextureAdapter>* Value)
 	{
-		if (Adapter)
+		if (!Material->ImageParameters.Contains(ParameterName))
+		{
+			return false;
+		}
+
+		if (Value)
 		{
 			TVariant<Private::OP::ADDRESS, TSharedPtr<const Private::FImage>> Image;
-			Image.Set<TSharedPtr<const Private::FImage>>(Adapter->Image);
+			Image.Set<TSharedPtr<const Private::FImage>>(Value->Get().Image);
 			
 			Material->ImageParameters.Add(ParameterName, Image);
 		}
+		else
+		{
+			Material->ImageParameters.Remove(ParameterName);
+		}
+		
+		return true;
 	}
 }
diff --git a/Engine/Plugins/Mutable/Source/MutableRuntime/Private/MuR/External/MeshAdapter.cpp b/Engine/Plugins/Mutable/Source/MutableRuntime/Private/MuR/External/MeshAdapter.cpp
--- a/Engine/Plugins/Mutable/Source/MutableRuntime/Private/MuR/External/MeshAdapter.cpp
+++ b/Engine/Plugins/Mutable/Source/MutableRuntime/Private/MuR/External/MeshAdapter.cpp
@@ -4,6 +4,7 @@
 
 #include "MuR/Mesh.h"
 #include "MuR/MeshPrivate.h"
+#include "MuR/OpMeshRemove.h"
 
 
 namespace UE::Mutable
@@ -16,6 +17,7 @@ namespace UE::Mutable
 
 	FMeshAdapter::FMeshAdapter(const FMeshAdapter& Other)
 	{
+		Mesh = MakeShared<Private::FMesh>();
 		Mesh->CopyFrom(*Other.Mesh);
 	}
 
@@ -48,4 +50,10 @@ namespace UE::Mutable
 
 		BasePositionIter.SetFromVec3f(Vertex);
 	}
+	
+
+	void FMeshAdapter::RemoveVertices(const TBitArray<>& Vertices)
+	{
+		Private::MeshRemoveVerticesWithCullSet(Mesh.Get(), Vertices, false);
+	}
 }
diff --git a/Engine/Plugins/Mutable/Source/MutableRuntime/Private/MuR/External/Operation.cpp b/Engine/Plugins/Mutable/Source/MutableRuntime/Private/MuR/External/Operation.cpp
--- a/Engine/Plugins/Mutable/Source/MutableRuntime/Private/MuR/External/Operation.cpp
+++ b/Engine/Plugins/Mutable/Source/MutableRuntime/Private/MuR/External/Operation.cpp
@@ -2,19 +2,16 @@
 
 #include "MuR/External/Operation.h"
 
-#include "StructUtils/InstancedStruct.h"
-
 
 namespace UE::Mutable
 {
-	TSharedPtr<const FInstancedStruct> FContext::GetInput(const FName& OperantName)
+	FValueConst FContext::GetInput(const FName& Name)
 	{
-		return Inputs.FindAndRemoveChecked(OperantName); // If not found it means that the input was already requested.
+		return Inputs.FindAndRemoveChecked(Name);
 	}
 
-
-	TSharedPtr<const FInstancedStruct>& FContext::GetOutput()
+	void FContext::SetOutput(FValueConst&& Value)
 	{
-		return Output;
+		Output = MoveTemp(Value);
 	}
 }
diff --git a/Engine/Plugins/Mutable/Source/MutableRuntime/Private/MuR/External/TextureAdapter.cpp b/Engine/Plugins/Mutable/Source/MutableRuntime/Private/MuR/External/TextureAdapter.cpp
--- a/Engine/Plugins/Mutable/Source/MutableRuntime/Private/MuR/External/TextureAdapter.cpp
+++ b/Engine/Plugins/Mutable/Source/MutableRuntime/Private/MuR/External/TextureAdapter.cpp
@@ -15,6 +15,7 @@ namespace UE::Mutable
 
 	FTextureAdapter::FTextureAdapter(const FTextureAdapter& Other)
 	{
+		Image = MakeShared<Private::FImage>();
 		Image->Copy(Other.Image.Get());
 	}
 
diff --git a/Engine/Plugins/Mutable/Source/MutableRuntime/Private/MuR/External/Value.cpp b/Engine/Plugins/Mutable/Source/MutableRuntime/Private/MuR/External/Value.cpp
new file mode 100644
--- /dev/null
+++ b/Engine/Plugins/Mutable/Source/MutableRuntime/Private/MuR/External/Value.cpp
@@ -0,0 +1,33 @@
+// Copyright Epic Games, Inc. All Rights Reserved.
+
+#include "MuR/External/Value.h"
+
+#include "StructUtils/InstancedStruct.h"
+
+
+namespace UE::Mutable
+{
+	FValue CopyOrMove(FValueConst&& ValueConst)
+	{
+		if (ValueConst.Ptr.IsUnique())
+		{
+			FValue Value;
+			Value.Ptr = ConstCastSharedPtr<FInstancedStruct>(ValueConst.Ptr); 
+
+			return MoveTemp(Value);
+		}
+		else
+		{
+			FValue Value;
+			Value.Ptr = MakeShared<FInstancedStruct>(*ValueConst.Ptr); 
+
+			return MoveTemp(Value);
+		}
+	}
+
+	
+	FValueConst::FValueConst(FValue&& Value)
+	{
+		Ptr = MoveTemp(Value.Ptr);
+	}
+}
diff --git a/Engine/Plugins/Mutable/Source/MutableRuntime/Private/MuR/Model.cpp b/Engine/Plugins/Mutable/Source/MutableRuntime/Private/MuR/Model.cpp
--- a/Engine/Plugins/Mutable/Source/MutableRuntime/Private/MuR/Model.cpp
+++ b/Engine/Plugins/Mutable/Source/MutableRuntime/Private/MuR/Model.cpp
@@ -419,11 +419,7 @@ namespace UE::Mutable::Private
     }
 
 
-	TSharedPtr<FParameters> FModel::NewParameters(TSharedPtr<const FModel> Model,
-		const FParameters* OldParameters,
-		const TMap<FName, TObjectPtr<UTexture>>* ImageDefaults,
-		const TMap<FName, TObjectPtr<USkeletalMesh>>* MeshDefaults,
-		const TMap<FName, TObjectPtr<UMaterialInterface>>* MaterialDefaults)
+	TSharedPtr<FParameters> FModel::NewParameters(TSharedPtr<const FModel> Model, const FParameters* OldParameters)
 	{
 		LLM_SCOPE_BYNAME(TEXT("MutableRuntime"));
 
@@ -482,6 +478,18 @@ namespace UE::Mutable::Private
                         case EParameterType::Image:
                             pRes->SetImageValue( thisP, OldParameters->GetImageValue(p) );
                             break;
+                        	
+                        case EParameterType::Mesh:
+                        	pRes->SetMeshValue( thisP, OldParameters->GetMeshValue(p) );
+                        	break;
+                        	
+	        			case EParameterType::Material:
+	        				pRes->SetMaterialValue( thisP, OldParameters->GetMaterialValue(p) );
+	        				break;
+                        	
+	        			case EParameterType::InstancedStruct:
+	        				pRes->SetInstancedStructValue( thisP, OldParameters->GetInstancedStructValue(p) );
+	        				break;
 
                         default:
                             check(false);
@@ -492,45 +500,6 @@ namespace UE::Mutable::Private
             }
         }
 
-    	if (ImageDefaults)
-    	{
-    		for (const TPair<FName, TObjectPtr<UTexture>>& Pair : *ImageDefaults)
-    		{
-    			int32 Index = pRes->GetPrivate()->Find(Pair.Key.ToString());
-    			if (Index != INDEX_NONE)
-    			{
-    				check(pRes->GetType(Index) == EParameterType::Image);
-    				pRes->SetImageValue(Index, Pair.Value);
-    			}
-    		}
-    	}
-    	
-    	if (MeshDefaults)
-    	{
-    		for (const TPair<FName, TObjectPtr<USkeletalMesh>>& Pair : *MeshDefaults)
-    		{
-    			int32 Index = pRes->GetPrivate()->Find(Pair.Key.ToString());
-    			if (Index != INDEX_NONE)
-    			{
-    				check(pRes->GetType(Index) == EParameterType::Mesh);
-    				pRes->SetMeshValue(Index, Pair.Value);
-    			}
-    		}
-    	}
-
-		if (MaterialDefaults)
-		{
-			for (const TPair<FName, TObjectPtr<UMaterialInterface>>& Pair : *MaterialDefaults)
-			{
-				int32 Index = pRes->GetPrivate()->Find(Pair.Key.ToString());
-				if (Index != INDEX_NONE)
-				{
-					check(pRes->GetType(Index) == EParameterType::Material);
-					pRes->SetMaterialValue(Index, Pair.Value);
-				}
-			}
-		}
-
         return pRes;
     }
 
diff --git a/Engine/Plugins/Mutable/Source/MutableRuntime/Private/MuR/Operations.cpp b/Engine/Plugins/Mutable/Source/MutableRuntime/Private/MuR/Operations.cpp
--- a/Engine/Plugins/Mutable/Source/MutableRuntime/Private/MuR/Operations.cpp
+++ b/Engine/Plugins/Mutable/Source/MutableRuntime/Private/MuR/Operations.cpp
@@ -42,6 +42,7 @@ namespace UE::Mutable::Private
 		{ EDataType::String,			},	// ST_PARAMETER
 		{ EDataType::Matrix,			},	// MA_PARAMETER
 		{ EDataType::Material,		},	// MI_PARAMETER
+		{ EDataType::InstancedStruct,},	// MI_PARAMETER
 
 		{ EDataType::Image,			},	// IM_REFERENCE
 		{ EDataType::Mesh,			},	// ME_REFERENCE
@@ -78,12 +79,14 @@ namespace UE::Mutable::Private
 
 		{ EDataType::Scalar,			},	// SC_ARITHMETIC
 		{ EDataType::Scalar,			},	// SC_CURVE
+		{ EDataType::Scalar,			},	// SC_EXTERNAL
 
 		{ EDataType::Color,			},	// CO_SAMPLEIMAGE
 		{ EDataType::Color,			},	// CO_SWIZZLE
 		{ EDataType::Color,			},	// CO_FROMSCALARS
 		{ EDataType::Color,			},	// CO_ARITHMETIC
 		{ EDataType::Color,			},	// CO_LINEARTOSRGB
+		{ EDataType::Color,			},	// CO_EXTERNAL
 
 		{ EDataType::Image,			},	// IM_LAYER
 		{ EDataType::Image,			},	// IM_LAYERCOLOUR
@@ -111,6 +114,7 @@ namespace UE::Mutable::Private
 		{ EDataType::Image,			},	// IM_NORMALCOMPOSITE
 		{ EDataType::Image,			},	// IM_TRANSFORM
 		{ EDataType::Image,			},	// IM_PARAMETER_CONVERT
+		{ EDataType::Image,			},	// IM_EXTERNAL
 
 		{ EDataType::Mesh,			},	// ME_APPLYLAYOUT
 		{ EDataType::Mesh,			},	// ME_PREPARELAYOUT
@@ -137,6 +141,7 @@ namespace UE::Mutable::Private
 		{ EDataType::Mesh,			},	// ME_ADDMETADATA
 		{ EDataType::Mesh,			},	// ME_TRANSFORMWITHMESH
 		{ EDataType::Mesh,			},	// ME_TRANSFORMWITHBONE
+		{ EDataType::Mesh,			},	// ME_EXTERNAL
 
 		{ EDataType::Instance,		},	// IN_ADDMESH
 		{ EDataType::Instance,		},	// IN_ADDIMAGE
@@ -150,12 +155,14 @@ namespace UE::Mutable::Private
 		{ EDataType::Instance,		},	// IN_ADDOVERLAYMATERIAL
 		{ EDataType::Instance,		},	// IN_ADDMATERIAL
 
-		{ EDataType::Layout,		},	// LA_PACK
-		{ EDataType::Layout,		},	// LA_MERGE
-		{ EDataType::Layout,		},	// LA_REMOVEBLOCKS
-		{ EDataType::Layout,		},	// LA_FROMMESH
+		{ EDataType::Layout,			},	// LA_PACK
+		{ EDataType::Layout,			},	// LA_MERGE
+		{ EDataType::Layout,			},	// LA_REMOVEBLOCKS
+		{ EDataType::Layout,			},	// LA_FROMMESH
 
-		{ EDataType::Unknown,	},	// UN_EXTERNAL
+		{ EDataType::Material,		},	// IM_EXTERNAL
+
+    	{ EDataType::InstancedStruct,},	// IS_EXTERNAL
 	};
 
     // clang-format on
@@ -196,6 +203,7 @@ namespace UE::Mutable::Private
 		case EOpType::ME_PARAMETER:
 		case EOpType::MA_PARAMETER:
 		case EOpType::MI_PARAMETER:
+		case EOpType::IS_PARAMETER:
 		case EOpType::IM_REFERENCE:
 		case EOpType::ME_REFERENCE:
 		case EOpType::IM_PARAMETER_FROM_MATERIAL:
@@ -917,7 +925,12 @@ namespace UE::Mutable::Private
 			break;
 		}
         
-        case EOpType::UN_EXTERNAL:
+        case EOpType::SC_EXTERNAL:
+        case EOpType::CO_EXTERNAL:
+        case EOpType::IM_EXTERNAL:
+        case EOpType::ME_EXTERNAL:
+        case EOpType::MI_EXTERNAL:
+        case EOpType::IS_EXTERNAL:
         {
         	const uint8* Data = program.GetOpArgsPointer(at);
 
diff --git a/Engine/Plugins/Mutable/Source/MutableRuntime/Private/MuR/Parameters.cpp b/Engine/Plugins/Mutable/Source/MutableRuntime/Private/MuR/Parameters.cpp
--- a/Engine/Plugins/Mutable/Source/MutableRuntime/Private/MuR/Parameters.cpp
+++ b/Engine/Plugins/Mutable/Source/MutableRuntime/Private/MuR/Parameters.cpp
@@ -973,7 +973,87 @@ namespace UE::Mutable::Private
 	}
 
 
-    FProjector FParameters::Private::GetProjectorValue( int32 Index, const FRangeIndex* Pos ) const
+	TSharedPtr<const FInstancedStruct> FParameters::GetInstancedStructValue(int32 Index, const FRangeIndex* RangeIndex) const
+	{
+    	check(Index >= 0 && Index < (int32)m_pD->Values.Num());
+    	check(GetType(Index) == EParameterType::InstancedStruct);
+
+    	// Early out in case of invalid parameters
+    	if (Index < 0 || Index >= (int32)m_pD->Values.Num() ||
+			GetType(Index) != EParameterType::InstancedStruct)
+    	{
+    		return {};
+    	}
+
+    	// Single value case
+    	if (!RangeIndex)
+    	{
+    		// Return the single value
+    		return m_pD->Values[Index].Get<FParamInstancedStructType>();
+    	}
+
+    	// Multivalue case
+    	check(RangeIndex->Parameter == Index);
+
+    	if (Index < int32(m_pD->MultiValues.Num()))
+    	{
+    		const TMap< TArray<int32>, FParameterValue >& m = m_pD->MultiValues[Index];
+    		const FParameterValue* it = m.Find(RangeIndex->Values);
+    		if (it)
+    		{
+    			return it->Get<FParamInstancedStructType>();
+    		}
+    	}
+
+    	// Multivalue parameter, but no multivalue set. Return single value.
+    	return m_pD->Values[Index].Get<FParamInstancedStructType>();
+	}
+
+
+	void FParameters::SetInstancedStructValue(int32 Index, const TSharedPtr<const FInstancedStruct>& Value, const FRangeIndex* RangeIndex)
+	{
+    	LLM_SCOPE_BYNAME(TEXT("MutableRuntime"));
+
+    	check(Index >= 0 && Index < (int32)m_pD->Values.Num());
+    	check(GetType(Index) == EParameterType::InstancedStruct);
+
+    	// Early out in case of invalid parameters
+    	if (Index < 0 || Index >= (int32)m_pD->Values.Num() ||
+			GetType(Index) != EParameterType::InstancedStruct)
+    	{
+    		return;
+    	}
+
+    	// Single value case
+    	if (!RangeIndex)
+    	{
+    		// Clear multivalue, if set.
+    		if (Index < int32(m_pD->MultiValues.Num()))
+    		{
+    			m_pD->MultiValues[Index].Empty();
+    		}
+
+    		m_pD->Values[Index].Set<FParamInstancedStructType>(Value);
+    	}
+
+    	// Multivalue case
+    	else
+    	{
+    		check(RangeIndex->Parameter == Index);
+
+    		if (Index >= int32(m_pD->MultiValues.Num()))
+    		{
+    			m_pD->MultiValues.SetNum(Index + 1);
+    		}
+
+    		TMap< TArray<int32>, FParameterValue >& m = m_pD->MultiValues[Index];
+    		FParameterValue& it = m.FindOrAdd(RangeIndex->Values);
+    		it.Set<FParamInstancedStructType>(Value);
+    	}
+	}
+
+
+	FProjector FParameters::Private::GetProjectorValue( int32 Index, const FRangeIndex* Pos ) const
     {
         const FProgram& Program = Model->GetPrivate()->Program;
 
diff --git a/Engine/Plugins/Mutable/Source/MutableRuntime/Private/MuR/Serialisation.cpp b/Engine/Plugins/Mutable/Source/MutableRuntime/Private/MuR/Serialisation.cpp
--- a/Engine/Plugins/Mutable/Source/MutableRuntime/Private/MuR/Serialisation.cpp
+++ b/Engine/Plugins/Mutable/Source/MutableRuntime/Private/MuR/Serialisation.cpp
@@ -95,7 +95,19 @@ namespace UE::Mutable::Private
     	T = S != 0;
     }
 
-	
+
+	void operator<<(FOutputArchive& Arch, const TSharedPtr<const FInstancedStruct>& V)
+	{
+		// Do Nothing. UStructs values are saved in the CO.	
+	}
+
+
+	void operator>>(FInputArchive& Arch, TSharedPtr<const FInstancedStruct>& V)
+	{
+		// Do Nothing. UStructs values are saved in the CO.	
+	}
+
+
 	FInputMemoryStream::FInputMemoryStream( const void* InBuffer, uint64 InSize )
     {
         Buffer = InBuffer;
diff --git a/Engine/Plugins/Mutable/Source/MutableRuntime/Private/MuR/System.cpp b/Engine/Plugins/Mutable/Source/MutableRuntime/Private/MuR/System.cpp
--- a/Engine/Plugins/Mutable/Source/MutableRuntime/Private/MuR/System.cpp
+++ b/Engine/Plugins/Mutable/Source/MutableRuntime/Private/MuR/System.cpp
@@ -3,6 +3,7 @@
 #include "MuR/System.h"
 
 #include "CodeRunnerBegin.h"
+#include "BitWriterResourceKey.h"
 #include "Containers/Array.h"
 #include "Containers/Map.h"
 #include "Containers/Set.h"
@@ -944,6 +945,7 @@ namespace UE::Mutable::Private
 			case EOpType::IM_PARAMETER:
 			case EOpType::MA_PARAMETER:
 			case EOpType::MI_PARAMETER:
+			case EOpType::IS_PARAMETER:
             {
 				OP::ParameterArgs args = Program.GetOpArgs<OP::ParameterArgs>(at);
 				OP::ADDRESS ParamIndex = args.variable;
@@ -1979,7 +1981,7 @@ namespace UE::Mutable::Private
 		}
 
 		// Generate the relevant parameters blob
-		FBitWriter Blob( 2048*8, true );
+		FBitWriterResourceKey Blob(2048*8, true);
 
 		const TArray<FParameterDesc>& ParamDescs = Params->GetPrivate()->Model->GetPrivate()->Program.Parameters;
 
@@ -2199,6 +2201,29 @@ namespace UE::Mutable::Private
 				break;
 			}
 
+			case EParameterType::InstancedStruct:
+			{
+				const FParamInstancedStructType& Value = Params->GetPrivate()->Values[ParamIndex].Get<FParamInstancedStructType>();
+
+				FInstancedStruct* ValueInstancedStruct = const_cast<FInstancedStruct*>(Value.Get());
+				ValueInstancedStruct->Serialize(Blob);
+
+				// Multi-values
+				if (Params->GetPrivate()->HasMultipleValues(ParamIndex))
+				{
+					const TMap<TArray<int32>, FParameterValue>& Multi = Params->GetPrivate()->MultiValues[ParamIndex];
+					AddMultiValueKeys(Blob, Multi);
+					for (const TPair<TArray<int32>, FParameterValue>& MultiValuePair : Multi)
+					{
+						const FParamInstancedStructType& MultiValue = MultiValuePair.Value.Get<FParamInstancedStructType>();
+
+						FInstancedStruct* MultiValueInstancedStruct = const_cast<FInstancedStruct*>(MultiValue.Get());
+						MultiValueInstancedStruct->Serialize(Blob);
+					}
+				}
+				break;
+			}
+				
 			default:
 				// unsupported parameter type
 				check(false);
diff --git a/Engine/Plugins/Mutable/Source/MutableRuntime/Public/MuR/External/FloatAdapter.h b/Engine/Plugins/Mutable/Source/MutableRuntime/Public/MuR/External/FloatAdapter.h
--- a/Engine/Plugins/Mutable/Source/MutableRuntime/Public/MuR/External/FloatAdapter.h
+++ b/Engine/Plugins/Mutable/Source/MutableRuntime/Public/MuR/External/FloatAdapter.h
@@ -35,7 +35,6 @@ namespace UE::Mutable
 		
 		FFloatAdapter& operator=(FFloatAdapter&& Other) = delete;
 
-		// Mesh interface
 		MUTABLERUNTIME_API float GetValue() const;
 
 		MUTABLERUNTIME_API void SetValue(float Value);
diff --git a/Engine/Plugins/Mutable/Source/MutableRuntime/Public/MuR/External/MaterialAdapter.h b/Engine/Plugins/Mutable/Source/MutableRuntime/Public/MuR/External/MaterialAdapter.h
--- a/Engine/Plugins/Mutable/Source/MutableRuntime/Public/MuR/External/MaterialAdapter.h
+++ b/Engine/Plugins/Mutable/Source/MutableRuntime/Public/MuR/External/MaterialAdapter.h
@@ -2,6 +2,7 @@
 
 #pragma once
 
+#include "Value.h"
 #include "MuR/Image.h"
 
 #include "MaterialAdapter.generated.h"
@@ -35,18 +36,18 @@ namespace UE::Mutable
 		FMaterialAdapter& operator=(const FMaterialAdapter& Other);
 		
 		FMaterialAdapter& operator=(FMaterialAdapter&& Other) = delete;
-		
-		TSharedPtr<const FFloatAdapter> GetFloat(const FName& ParameterName) const;
 
-		bool SetFloat(const FName& ParameterName, const TSharedPtr<const FFloatAdapter>& Adapter);
+		MUTABLERUNTIME_API TOptional<TValueConst<FFloatAdapter>> GetFloat(const FName& ParameterName) const;
+
+		MUTABLERUNTIME_API bool SetFloat(const FName& ParameterName, const TValueConst<FFloatAdapter>* Value);
 		
-		TSharedPtr<const FVectorAdapter> GetVector(const FName& ParameterName) const;
+		MUTABLERUNTIME_API TOptional<TValueConst<FVectorAdapter>> GetVector(const FName& ParameterName) const;
 		
-		bool SetVector(const FName& ParameterName, const TSharedPtr<const FVectorAdapter>& Adapter);
+		MUTABLERUNTIME_API bool SetVector(const FName& ParameterName, const TValueConst<FVectorAdapter>* Value);
 		
-		TSharedPtr<const FTextureAdapter> GetTexture(const FName& ParameterName) const;
+		MUTABLERUNTIME_API TOptional<TValueConst<FTextureAdapter>> GetTexture(const FName& ParameterName) const;
 
-		void SetTexture(const FName& ParameterName, const TSharedPtr<FTextureAdapter>& Adapter);
+		MUTABLERUNTIME_API bool SetTexture(const FName& ParameterName, const TValueConst<FTextureAdapter>* Value);
 	
 	private:
 		TSharedPtr<Private::FMaterial> Material;
diff --git a/Engine/Plugins/Mutable/Source/MutableRuntime/Public/MuR/External/MeshAdapter.h b/Engine/Plugins/Mutable/Source/MutableRuntime/Public/MuR/External/MeshAdapter.h
--- a/Engine/Plugins/Mutable/Source/MutableRuntime/Public/MuR/External/MeshAdapter.h
+++ b/Engine/Plugins/Mutable/Source/MutableRuntime/Public/MuR/External/MeshAdapter.h
@@ -38,7 +38,9 @@ namespace UE::Mutable
 		MUTABLERUNTIME_API FVector3f GetVertex(int32 Index) const;
 
 		MUTABLERUNTIME_API void SetVertex(int32 Index, const FVector3f& Vertex);
-
+		
+		MUTABLERUNTIME_API void RemoveVertices(const TBitArray<>& Vertices);
+		
 	private:
 		TSharedPtr<Private::FMesh> Mesh;
 	};
diff --git a/Engine/Plugins/Mutable/Source/MutableRuntime/Public/MuR/External/Operation.h b/Engine/Plugins/Mutable/Source/MutableRuntime/Public/MuR/External/Operation.h
--- a/Engine/Plugins/Mutable/Source/MutableRuntime/Public/MuR/External/Operation.h
+++ b/Engine/Plugins/Mutable/Source/MutableRuntime/Public/MuR/External/Operation.h
@@ -2,38 +2,20 @@
 
 #pragma once
 
-#include "StructUtils/InstancedStruct.h"
+#include "Value.h"
 
 #include "Operation.generated.h"
 
 
 namespace UE::Mutable
 {
+	class FContext;
+	
 	namespace Private
 	{
 		class CodeRunner;
 	}
-
-
-	/** Copy (or move) a value to remove it's const.
-	 *
-	 * Instead of copy, a value will be moved if it is uniquely referenced. A value may not be unique due to:
-	 * - Being being used in more than one subsequent input.
-	 * - User holding additional references to it. Use MoveTemp whenever possible. */
-	template<class T>
-	TSharedPtr<T> CopyOrMove(const TSharedPtr<const T>&& Source)
-	{
-		if (Source.IsUnique())
-		{
-			return ConstCastSharedPtr<T>(Source);
-		}
-		else
-		{
-			return MakeShared<T>(*Source);
-		}
-	}
 	
-
 	/** Context of the External Operation. */
 	class FContext
 	{
@@ -43,15 +25,15 @@ namespace UE::Mutable
 		/** Returns and removes the input from the Context.
 		 *
 		 * Once called the owner is the requester. This allows CopyOrMove to move (see CopyOrMove for details). */
-		MUTABLERUNTIME_API TSharedPtr<const FInstancedStruct> GetInput(const FName& OperantName);
+		MUTABLERUNTIME_API FValueConst GetInput(const FName& Name);
 
 		/** Return the pointer to store the ouput. */
-		MUTABLERUNTIME_API TSharedPtr<const FInstancedStruct>& GetOutput();
+		MUTABLERUNTIME_API void SetOutput(FValueConst&& Value);
 		
 	private:
-		TMap<const FName, TSharedPtr<const FInstancedStruct>> Inputs;
+		TMap<const FName, FValueConst> Inputs;
 
-		TSharedPtr<const FInstancedStruct> Output = nullptr;
+		FValueConst Output;
 	};
 	
 
diff --git a/Engine/Plugins/Mutable/Source/MutableRuntime/Public/MuR/External/Value.h b/Engine/Plugins/Mutable/Source/MutableRuntime/Public/MuR/External/Value.h
new file mode 100644
--- /dev/null
+++ b/Engine/Plugins/Mutable/Source/MutableRuntime/Public/MuR/External/Value.h
@@ -0,0 +1,270 @@
+// Copyright Epic Games, Inc. All Rights Reserved.
+
+#pragma once
+
+#include "StructUtils/InstancedStruct.h"
+
+
+struct FInstancedStruct;
+
+namespace UE::Mutable
+{
+	namespace Private
+	{
+		class CodeRunner;
+	}
+	class FValueConst;
+	class FContext;
+	struct FMaterialAdapter;
+	template<typename T> class TValue;
+	template<typename T> class TValueConst;
+
+	
+	class FValue
+	{
+		friend Private::CodeRunner;
+		friend FValueConst;
+		friend MUTABLERUNTIME_API FValue CopyOrMove(FValueConst&&);
+		
+	public:
+		template<typename T>
+		static FValue New()
+		{
+			FValue Value;
+			Value.Ptr = MakeShared<FInstancedStruct>();
+			Value.Ptr->InitializeAs(T::StaticStruct());
+
+			return Value;
+		}
+		
+		FValue(const FValue& Other) = delete;
+		
+		FValue& operator=(const FValue& Other) = delete;
+
+		MUTABLERUNTIME_API FValue(FValue&& Other)
+		{
+			Ptr = MoveTemp(Other.Ptr);
+		}
+
+		MUTABLERUNTIME_API FValue& operator=(FValue&& Other)
+		{
+			Ptr = MoveTemp(Other.Ptr);
+			return *this;
+		}
+
+		MUTABLERUNTIME_API FValue(const UScriptStruct* ScriptStruct)
+		{
+			Ptr = MakeShared<FInstancedStruct>();
+			Ptr->InitializeAs(ScriptStruct);
+		}
+		
+		template <typename T>
+		T& Get()
+		{
+			return Ptr.Get()->GetMutable<T>();
+		}
+
+		template <typename T>
+		T* GetPtr()
+		{
+			return Ptr.Get()->GetMutablePtr<T>();
+		}
+		
+	private:
+		FValue() = default;
+		
+		TSharedPtr<FInstancedStruct> Ptr;
+	};
+
+
+	class FValueConst
+	{
+		friend Private::CodeRunner;
+		friend FContext;
+		friend MUTABLERUNTIME_API FValue CopyOrMove(FValueConst&&);
+
+	public:
+		template<typename T>
+		static FValueConst New()
+		{
+			FValueConst Value;
+			Value.Ptr = MakeShared<const FInstancedStruct>();
+			Value.Ptr->InitializeAs(T::StaticStruct());
+
+			return Value;
+		}
+		
+		FValueConst(const FValueConst& Other) = delete;
+		
+		FValueConst& operator=(const FValueConst& Other) = delete;
+
+		MUTABLERUNTIME_API FValueConst(FValueConst&& Other)
+		{
+			Ptr = MoveTemp(Other.Ptr);
+		}
+
+		MUTABLERUNTIME_API FValueConst& operator=(FValueConst&& Other)
+		{
+			Ptr = MoveTemp(Other.Ptr);
+			return *this;
+		}
+		
+		MUTABLERUNTIME_API FValueConst(FValue&& Value);
+
+		template <typename T>
+		const T& Get() const
+		{
+			return Ptr.Get()->Get<const T>();
+		}
+
+		template <typename T>
+		const T* GetPtr() const
+		{
+			return Ptr.Get()->GetPtr<T>();
+		}
+		
+	private:
+		FValueConst() = default;
+		
+		TSharedPtr<const FInstancedStruct> Ptr;
+	};
+
+	
+	template<typename T>
+	class TValue
+	{
+		friend Private::CodeRunner;
+		friend FMaterialAdapter;
+		friend FValueConst;
+		template<typename U> friend class TValueConst;
+		template<typename U> friend MUTABLERUNTIME_API TValue CopyOrMove(TValueConst<U>&&);
+		
+	public:
+		static TValue New()
+		{
+			TValue Value;
+			Value.Ptr = MakeShared<FInstancedStruct>();
+			Value.Ptr->InitializeAs(T::StaticStruct());
+
+			return Value;
+		}
+		
+		TValue(const TValue& Other) = delete;
+		
+		TValue& operator=(const TValue& Other) = delete;
+
+		TValue(TValue&& Other)
+		{
+			Ptr = MoveTemp(Other.Ptr);
+		}
+
+		TValue& operator=(TValue&& Other)
+		{
+			Ptr = MoveTemp(Other.Ptr);
+			return *this;
+		}
+
+		TValue(const UScriptStruct* ScriptStruct)
+		{
+			Ptr = MakeShared<FInstancedStruct>();
+			Ptr->InitializeAs(ScriptStruct);
+		}
+		
+		T& Get()
+		{
+			return Ptr.Get()->GetMutable<T>();
+		}
+
+		T* GetPtr()
+		{
+			return Ptr.Get()->GetMutablePtr<T>();
+		}
+		
+	private:
+		TValue() = default;
+		
+		TSharedPtr<FInstancedStruct> Ptr; // TODO GMT Extension TVariant Lazy	
+	};
+	
+
+	template<typename T>
+	class TValueConst
+	{
+		friend Private::CodeRunner;
+		friend FContext;
+		template<typename U> friend class TValue;
+		template<typename U> friend MUTABLERUNTIME_API TValue<U> CopyOrMove(TValueConst&&);
+		
+	public:
+		static TValueConst New()
+		{
+			TValueConst Value;
+			Value.Ptr = MakeShared<const FInstancedStruct>();
+			Value.Ptr->InitializeAs(T::StaticStruct());
+
+			return Value;
+		}
+		
+		TValueConst(const FValueConst& Other) = delete;
+		
+		TValueConst& operator=(const FValueConst& Other) = delete;
+
+		TValueConst(TValueConst&& Other)
+		{
+			Ptr = MoveTemp(Other.Ptr);
+		}
+
+		TValueConst& operator=(TValueConst&& Other)
+		{
+			Ptr = MoveTemp(Other.Ptr);
+			return *this;
+		}
+
+		TValueConst(TValue<T>&& Value)
+		{
+			Ptr = MoveTemp(Value.Ptr);
+		}
+		
+		const T& Get() const
+		{
+			return Ptr.Get()->Get<const T>();
+		}
+
+		const T* GetPtr() const
+		{
+			return Ptr.Get()->GetPtr<T>();
+		}
+		
+	private:
+		TValueConst() = default;
+		
+		TSharedPtr<const FInstancedStruct> Ptr;
+	};
+	
+	
+	/** Copy (or move) a value to remove it's const.
+	 *
+	 * Instead of copy, a value will be moved if it is uniquely referenced. A value may not be unique due to:
+	 * - Being being used in more than one subsequent input.
+	 * - User holding additional references to it. Use MoveTemp whenever possible. */
+	MUTABLERUNTIME_API FValue CopyOrMove(FValueConst&& ValueConst);
+
+	template<typename T>
+	TValue<T> CopyOrMove(TValueConst<T>&& ValueConst)
+	{
+		if (ValueConst.Ptr.IsUnique())
+		{
+			TValue<T> Value;
+			Value.Ptr = ConstCastSharedPtr<FInstancedStruct>(ValueConst.Ptr); 
+
+			return MoveTemp(Value);
+		}
+		else
+		{
+			TValue<T> Value;
+			Value.Ptr = MakeShared<FInstancedStruct>(*ValueConst.Ptr); 
+
+			return MoveTemp(Value);
+		}
+	}
+}
diff --git a/Engine/Plugins/Mutable/Source/MutableRuntime/Public/MuR/External/VectorAdapter.h b/Engine/Plugins/Mutable/Source/MutableRuntime/Public/MuR/External/VectorAdapter.h
--- a/Engine/Plugins/Mutable/Source/MutableRuntime/Public/MuR/External/VectorAdapter.h
+++ b/Engine/Plugins/Mutable/Source/MutableRuntime/Public/MuR/External/VectorAdapter.h
@@ -35,7 +35,6 @@ namespace UE::Mutable
 		
 		FVectorAdapter& operator=(FVectorAdapter&& Other) = delete;
 
-		// Mesh interface
 		MUTABLERUNTIME_API FVector4f GetValue() const;
 
 		MUTABLERUNTIME_API void SetValue(const FVector4f& Value);
diff --git a/Engine/Plugins/Mutable/Source/MutableTools/Internal/MuT/Node.h b/Engine/Plugins/Mutable/Source/MutableTools/Internal/MuT/Node.h
--- a/Engine/Plugins/Mutable/Source/MutableTools/Internal/MuT/Node.h
+++ b/Engine/Plugins/Mutable/Source/MutableTools/Internal/MuT/Node.h
@@ -153,6 +153,8 @@ namespace UE::Mutable::Private
 			ModifierTransformWithBone,
 
 			External,
+			ExternalOperation,
+			ExternalParameter,
 			
 			ExtensionData,
 			ExtensionDataConstant,
diff --git a/Engine/Plugins/Mutable/Source/MutableTools/Internal/MuT/NodeExternal.h b/Engine/Plugins/Mutable/Source/MutableTools/Internal/MuT/NodeExternal.h
--- a/Engine/Plugins/Mutable/Source/MutableTools/Internal/MuT/NodeExternal.h
+++ b/Engine/Plugins/Mutable/Source/MutableTools/Internal/MuT/NodeExternal.h
@@ -2,9 +2,7 @@
 
 #pragma once
 
-#include "MuR/Ptr.h"
 #include "MuT/Node.h"
-#include "StructUtils/InstancedStruct.h"
 
 #define UE_API MUTABLETOOLS_API
 
@@ -24,11 +22,6 @@ namespace UE::Mutable::Private
 
 	private:
 		static UE_API FNodeType StaticType;
-
-	public:
-		FInstancedStruct OperationInstancedStruct;
-
-		TArray<Ptr<Node>> Inputs;
 	};
 }
 
diff --git a/Engine/Plugins/Mutable/Source/MutableTools/Internal/MuT/NodeExternalOperation.h b/Engine/Plugins/Mutable/Source/MutableTools/Internal/MuT/NodeExternalOperation.h
new file mode 100644
--- /dev/null
+++ b/Engine/Plugins/Mutable/Source/MutableTools/Internal/MuT/NodeExternalOperation.h
@@ -0,0 +1,37 @@
+﻿// Copyright Epic Games, Inc. All Rights Reserved.
+
+#pragma once
+
+#include "MuT/NodeExternal.h"
+
+#include "MuR/Ptr.h"
+#include "StructUtils/InstancedStruct.h"
+
+#define UE_API MUTABLETOOLS_API
+
+
+namespace UE::Mutable::Private
+{
+	class NodeExternalOperation : public NodeExternal
+	{
+	public:
+		// Node interface
+		virtual const FNodeType* GetType() const override { return GetStaticType(); }
+		static const FNodeType* GetStaticType() { return &StaticType; }
+
+	protected:
+		/** Forbidden. Manage with the Ptr<> template. */
+		virtual ~NodeExternalOperation() override {}
+
+	private:
+		static UE_API FNodeType StaticType;
+
+	public:
+		FInstancedStruct OperationInstancedStruct;
+
+		TArray<Ptr<Node>> Inputs;
+	};
+}
+
+
+#undef UE_API
diff --git a/Engine/Plugins/Mutable/Source/MutableTools/Internal/MuT/NodeExternalParameter.h b/Engine/Plugins/Mutable/Source/MutableTools/Internal/MuT/NodeExternalParameter.h
new file mode 100644
--- /dev/null
+++ b/Engine/Plugins/Mutable/Source/MutableTools/Internal/MuT/NodeExternalParameter.h
@@ -0,0 +1,39 @@
+﻿// Copyright Epic Games, Inc. All Rights Reserved.
+
+#pragma once
+
+#include "MuR/Parameters.h"
+#include "MuR/Ptr.h"
+#include "MuT/Node.h"
+#include "MuT/NodeExternal.h"
+#include "MuT/NodeRange.h"
+
+#define UE_API MUTABLETOOLS_API
+
+namespace UE::Mutable::Private
+{
+	class NodeExternalParameter : public NodeExternal
+	{
+	public:
+		// Node interface
+		virtual const FNodeType* GetType() const override { return &StaticType; }
+		static const FNodeType* GetStaticType() { return &StaticType; }
+
+	protected:
+		virtual ~NodeExternalParameter() override {}
+
+	private:
+		static UE_API FNodeType StaticType;
+
+	public:
+		TSharedPtr<const FInstancedStruct> DefaultValue;
+		
+		FString Name;
+		FString UID;
+
+		TArray<Ptr<NodeRange>> Ranges;
+	};
+}
+
+
+#undef UE_API
diff --git a/Engine/Plugins/Mutable/Source/MutableTools/Private/MuT/ASTOpExternal.cpp b/Engine/Plugins/Mutable/Source/MutableTools/Private/MuT/ASTOpExternal.cpp
--- a/Engine/Plugins/Mutable/Source/MutableTools/Private/MuT/ASTOpExternal.cpp
+++ b/Engine/Plugins/Mutable/Source/MutableTools/Private/MuT/ASTOpExternal.cpp
@@ -7,13 +7,13 @@ namespace UE::Mutable::Private
 {
 	EOpType ASTOpExternal::GetOpType() const
 	{
-		return EOpType::UN_EXTERNAL;
+		return Type;
 	}
 
 
 	uint64 ASTOpExternal::Hash() const
 	{
-		uint64 Result = GetTypeHash(OperationInstancedStruct.GetScriptStruct()); // FInstancedStruct can not be hashed. At least we can hash the type.
+		uint64 Result = GetTypeHash(OperationInstancedStruct);
 		
 		for (const ASTChild& Input : Inputs)
 		{
@@ -66,6 +66,7 @@ namespace UE::Mutable::Private
 	Ptr<ASTOp> ASTOpExternal::Clone(MapChildFuncRef MapChild) const
 	{
 		Ptr<ASTOpExternal> New = new ASTOpExternal();
+		New->Type = Type;
 		New->OperationInstancedStruct = OperationInstancedStruct;
 		
 		for (const ASTChild& Input : Inputs)
@@ -88,7 +89,8 @@ namespace UE::Mutable::Private
 			linkedAddress = (OP::ADDRESS)Program.OpAddress.Num();
 			Program.OpAddress.Add((uint32_t)Program.ByteCode.Num());
 			
-			AppendCode(Program.ByteCode, EOpType::UN_EXTERNAL);
+			check(Type != EOpType::NONE);
+			AppendCode(Program.ByteCode, Type);
 			AppendCode(Program.ByteCode, Args);
 
 			for (const ASTChild& Input : Inputs)
diff --git a/Engine/Plugins/Mutable/Source/MutableTools/Private/MuT/ASTOpExternal.h b/Engine/Plugins/Mutable/Source/MutableTools/Private/MuT/ASTOpExternal.h
--- a/Engine/Plugins/Mutable/Source/MutableTools/Private/MuT/ASTOpExternal.h
+++ b/Engine/Plugins/Mutable/Source/MutableTools/Private/MuT/ASTOpExternal.h
@@ -23,6 +23,8 @@ namespace UE::Mutable::Private
 		virtual Ptr<ASTOp> Clone(MapChildFuncRef) const override;
 		virtual void Link(FProgram&, FLinkerOptions*) override;
 
+		EOpType Type = EOpType::NONE;
+		
 		FInstancedStruct OperationInstancedStruct;
 
 		TArray<ASTChild> Inputs;
diff --git a/Engine/Plugins/Mutable/Source/MutableTools/Private/MuT/CodeGenerator.h b/Engine/Plugins/Mutable/Source/MutableTools/Private/MuT/CodeGenerator.h
--- a/Engine/Plugins/Mutable/Source/MutableTools/Private/MuT/CodeGenerator.h
+++ b/Engine/Plugins/Mutable/Source/MutableTools/Private/MuT/CodeGenerator.h
@@ -47,6 +47,8 @@
 namespace UE::Mutable::Private
 {
 	class NodeImageParameterConvert;
+	class NodeExternalParameter;
+	class NodeExternalOperation;
 	class NodeMaterialExternal;
 	class NodeScalarExternal;
 }
@@ -840,8 +842,10 @@ namespace UE::Mutable::Private
 
     	//-----------------------------------------------------------------------------------------
 
-    	struct FExternalGenerationOptions
+    	struct FExternalGenerationOptions 
     	{
+    		// TODO Propagated options could be a TVariant
+    		
     		// GenerateMesh
     		TOptional<FMeshGenerationStaticOptions> MeshStaticOptions;
     		TOptional<FMeshOptionsTask> MeshDynamicOptions;
@@ -907,9 +911,16 @@ namespace UE::Mutable::Private
     	};
     	
     	typedef TMap<FGeneratedExternalCacheKey, FExternalGenerationResult> FGeneratedExternalMap;
-    	FGeneratedExternalMap GeneratedExtension;
+    	struct FGeneratedExternal
+    	{
+    		UE::FMutex Mutex;
+    		FGeneratedExternalMap Map;
+    	};
+    	FGeneratedExternal GeneratedExternal;
 
-    	void GenerateExternal(FExternalGenerationResult& OutResult, const FExternalGenerationOptions& Options, const Ptr<const NodeExternal>& ExtensionNode);
+    	void GenerateExternal(FExternalGenerationResult& OutResult, const FExternalGenerationOptions& Options, const Ptr<const NodeExternal>& ExternalNode);
+    	void GenerateExternal_Operation(FExternalGenerationResult& OutResult, const FExternalGenerationOptions& Options, const Ptr<const NodeExternalOperation>& ExternalNode);
+    	void GenerateExternal_Parameter(FExternalGenerationResult& OutResult, const FExternalGenerationOptions& Options, const Ptr<const NodeExternalParameter>& ExtensionNode);
 
 		//-----------------------------------------------------------------------------------------
     	
diff --git a/Engine/Plugins/Mutable/Source/MutableTools/Private/MuT/CodeGenerator_External.cpp b/Engine/Plugins/Mutable/Source/MutableTools/Private/MuT/CodeGenerator_External.cpp
--- a/Engine/Plugins/Mutable/Source/MutableTools/Private/MuT/CodeGenerator_External.cpp
+++ b/Engine/Plugins/Mutable/Source/MutableTools/Private/MuT/CodeGenerator_External.cpp
@@ -1,13 +1,22 @@
 // Copyright Epic Games, Inc. All Rights Reserved.
 
 #include "CodeGenerator.h"
+#include "Image/ImageBuilder.h"
+#include "MuR/External/FloatAdapter.h"
+#include "MuR/External/MaterialAdapter.h"
+#include "MuR/External/MeshAdapter.h"
+#include "MuR/External/Operation.h"
+#include "MuR/External/TextureAdapter.h"
+#include "MuR/External/VectorAdapter.h"
+#include "MuT/NodeExternalOperation.h"
+#include "MuT/NodeExternalParameter.h"
 
 
 namespace UE::Mutable::Private
 {
-	void CodeGenerator::GenerateExternal(FExternalGenerationResult& OutResult, const FExternalGenerationOptions& Options, const Ptr<const NodeExternal>& ExtensionNode)
+	void CodeGenerator::GenerateExternal(FExternalGenerationResult& OutResult, const FExternalGenerationOptions& Options, const Ptr<const NodeExternal>& ExternalNode)
 	{
-		if (!ExtensionNode)
+		if (!ExternalNode)
 		{
 			OutResult = FExternalGenerationResult();
 			return;
@@ -15,29 +24,90 @@ namespace UE::Mutable::Private
 
 		// See if it was already generated
 		FGeneratedExternalCacheKey Key;
-		Key.Node = ExtensionNode.get();
+		Key.Node = ExternalNode.get();
 		Key.Options = Options;
 		
-		FGeneratedExternalMap::ValueType* CachedResult = GeneratedExtension.Find(Key);
-		if (CachedResult)
 		{
-			OutResult = *CachedResult;
-			return;
+			UE::TUniqueLock Lock(GeneratedExternal.Mutex);
+			if (FGeneratedExternalMap::ValueType* Found = GeneratedExternal.Map.Find(Key))
+			{
+				OutResult = *Found;
+				return;
+			}
 		}
-    			
+		
+		// Generate for each different type of node
+		FMeshTask Result;
+		switch (ExternalNode->GetType()->Type)
+		{
+		case Node::EType::ExternalOperation:
+			{
+				Ptr<const NodeExternalOperation> Operation = static_pointer_cast<const NodeExternalOperation>(ExternalNode);
+				GenerateExternal_Operation(OutResult, Options, Operation);
+				break;	
+			}
+			
+		case Node::EType::ExternalParameter:
+			{
+				Ptr<const NodeExternalParameter> Parameter = static_pointer_cast<const NodeExternalParameter>(ExternalNode);
+				GenerateExternal_Parameter(OutResult, Options, Parameter);
+				break;
+			}
+			
+		default:
+			unimplemented();
+		}
+	
+		// Cache the result
+		{
+			UE::TUniqueLock Lock(GeneratedExternal.Mutex);
+			GeneratedExternal.Map.Add(Key, OutResult);
+		}
+	}
+
+
+	void CodeGenerator::GenerateExternal_Operation(FExternalGenerationResult& OutResult, const FExternalGenerationOptions& Options, const Ptr<const NodeExternalOperation>& ExternalNode)
+	{
 		Ptr<ASTOpExternal> External = new ASTOpExternal();
-		External->OperationInstancedStruct = ExtensionNode->OperationInstancedStruct;
+		External->OperationInstancedStruct = ExternalNode->OperationInstancedStruct;
 
+		const UScriptStruct* OutputType = ExternalNode->OperationInstancedStruct.GetPtr<FExternalOperation>()->GetOutput().Value;
+		
+		if (OutputType == FMeshAdapter::StaticStruct())
+		{
+			External->Type = EOpType::ME_EXTERNAL;
+		}
+		else if (OutputType == FTextureAdapter::StaticStruct())
+		{
+			External->Type = EOpType::IM_EXTERNAL;
+		}
+		else if (OutputType == FFloatAdapter::StaticStruct())
+		{
+			External->Type = EOpType::SC_EXTERNAL;
+		}
+		else if (OutputType == FVectorAdapter::StaticStruct())
+		{
+			External->Type = EOpType::CO_EXTERNAL;
+		}
+		else if (OutputType == FMaterialAdapter::StaticStruct())
+		{
+			External->Type = EOpType::MI_EXTERNAL;
+		}
+		else
+		{
+			External->Type = EOpType::IS_EXTERNAL;
+		}
+		
 		FMeshGenerationResult Result;
 		Result.MeshOp = External;
 
-		const int32 NumInputs = ExtensionNode->Inputs.Num();
+		const int32 NumInputs = ExternalNode->Inputs.Num();
     			
 		External->Inputs.Reserve(NumInputs);
 		
 		for (int32 Index = 0; Index < NumInputs; ++Index)
 		{
-			const Ptr<Node>& Input = ExtensionNode->Inputs[Index];
+			const Ptr<Node>& Input = ExternalNode->Inputs[Index];
     				
 			if (Input->GetType()->IsA(NodeMesh::GetStaticType()))
 			{
@@ -103,7 +173,7 @@ namespace UE::Mutable::Private
 				FExternalGenerationResult InputResult;
 				GenerateExternal(InputResult, Options, static_pointer_cast<NodeExternal>(Input));
 
-				External->Inputs[Index] = InputResult.Op;
+				External->Inputs.Add(ASTChild(External.get(), InputResult.Op));
 			}
 			else
 			{
@@ -114,4 +184,71 @@ namespace UE::Mutable::Private
 
 		OutResult.Op = External;
 	}
-}
\ No newline at end of file
+
+
+	void CodeGenerator::GenerateExternal_Parameter(FExternalGenerationResult& OutResult, const FExternalGenerationOptions& Options,	const Ptr<const NodeExternalParameter>& ExtensionNode)
+	{
+		Ptr<ASTOpParameter> Op;
+
+		Ptr<ASTOpParameter>* Found = nullptr;
+		{
+			UE::TUniqueLock Lock(FirstPass.ParameterNodes.Mutex);
+			Found = FirstPass.ParameterNodes.GenericParametersCache.Find(ExtensionNode.get());
+			if (!Found)
+			{
+				FParameterDesc Param;
+				Param.Name = ExtensionNode->Name;
+				bool bParseOk = FGuid::Parse(ExtensionNode->UID, Param.UID);
+				check(bParseOk);
+				Param.Type = EParameterType::InstancedStruct;
+				Param.DefaultValue.Set<FParamInstancedStructType>(ExtensionNode->DefaultValue);
+
+				Op = new ASTOpParameter();
+				Op->Type = EOpType::IS_PARAMETER;
+				Op->Parameter = Param;
+
+				FirstPass.ParameterNodes.GenericParametersCache.Add(ExtensionNode.get(), Op);
+			}
+			else
+			{
+				Op = *Found;
+			}
+		}
+
+		if (!Found)
+		{
+			// Generate the code for the ranges
+			for (int32 RangeIndex = 0; RangeIndex < ExtensionNode->Ranges.Num(); ++RangeIndex)
+			{
+				FGenericGenerationOptions GenericOptions;
+
+				if (Options.MeshStaticOptions.IsSet())
+				{
+					GenericOptions = static_cast<FGenericGenerationOptions>(Options.MeshStaticOptions.GetValue());
+				}
+				else if (Options.ColourOptions.IsSet())
+				{
+					GenericOptions = static_cast<FGenericGenerationOptions>(Options.ColourOptions.GetValue());
+				}
+				else if (Options.ImageOptions.IsSet())
+				{
+					GenericOptions = static_cast<FGenericGenerationOptions>(Options.ImageOptions.GetValue());
+				}
+				else if (Options.MaterialOptions.IsSet())
+				{
+					GenericOptions = static_cast<FGenericGenerationOptions>(Options.MaterialOptions.GetValue());
+				}
+				else if (Options.ScalarOptions.IsSet())
+				{
+					GenericOptions = static_cast<FGenericGenerationOptions>(Options.ScalarOptions.GetValue());
+				}
+				
+				FRangeGenerationResult rangeResult;
+				GenerateRange(rangeResult, GenericOptions, ExtensionNode->Ranges[RangeIndex]);
+				Op->Ranges.Emplace(Op.get(), rangeResult.sizeOp, rangeResult.rangeName, rangeResult.rangeUID);
+			}
+		}
+
+		OutResult.Op = Op;
+	}
+}
diff --git a/Engine/Plugins/Mutable/Source/MutableTools/Private/MuT/CodeOptimiser.cpp b/Engine/Plugins/Mutable/Source/MutableTools/Private/MuT/CodeOptimiser.cpp
--- a/Engine/Plugins/Mutable/Source/MutableTools/Private/MuT/CodeOptimiser.cpp
+++ b/Engine/Plugins/Mutable/Source/MutableTools/Private/MuT/CodeOptimiser.cpp
@@ -552,6 +552,7 @@ namespace UE::Mutable::Private
 					case EOpType::MA_PARAMETER:
 					case EOpType::MI_PARAMETER:
 					case EOpType::IM_PARAMETER_FROM_MATERIAL:
+					case EOpType::IS_PARAMETER:
 						bIsConstantSubgraph = false;
 						break;
 					default:
diff --git a/Engine/Plugins/Mutable/Source/MutableTools/Private/MuT/CodeOptimiser_Parameters.cpp b/Engine/Plugins/Mutable/Source/MutableTools/Private/MuT/CodeOptimiser_Parameters.cpp
--- a/Engine/Plugins/Mutable/Source/MutableTools/Private/MuT/CodeOptimiser_Parameters.cpp
+++ b/Engine/Plugins/Mutable/Source/MutableTools/Private/MuT/CodeOptimiser_Parameters.cpp
@@ -1582,6 +1582,7 @@ namespace UE::Mutable::Private
 				case EOpType::IM_PARAMETER:
 				case EOpType::ME_PARAMETER:
 				case EOpType::MA_PARAMETER:
+				case EOpType::IS_PARAMETER:
 				{
 					const ASTOpParameter* TypedOp = static_cast<const ASTOpParameter*>(At.get());
                     Parameters.Add(TypedOp->Parameter.Name);
diff --git a/Engine/Plugins/Mutable/Source/MutableTools/Private/MuT/Compiler.cpp b/Engine/Plugins/Mutable/Source/MutableTools/Private/MuT/Compiler.cpp
--- a/Engine/Plugins/Mutable/Source/MutableTools/Private/MuT/Compiler.cpp
+++ b/Engine/Plugins/Mutable/Source/MutableTools/Private/MuT/Compiler.cpp
@@ -79,6 +79,7 @@ namespace UE::Mutable::Private
 		{ false,	{ 0, 0, 0, 0, 0, 0, 0, 0, 0, 0, 0, 0, 0, 0, 0, 0, 0, 0, 0, 0, 0, 0, 0, 0, 0, 0, 0, 0, 0 }		},	// ST_PARAMETER
 		{ false,	{ 0, 0, 0, 0, 0, 0, 0, 0, 0, 0, 0, 0, 0, 0, 0, 0, 0, 0, 0, 0, 0, 0, 0, 0, 0, 0, 0, 0, 0 }		},	// MA_PARAMETER
 		{ false,	{ 0, 0, 0, 0, 0, 0, 0, 0, 0, 0, 0, 0, 0, 0, 0, 0, 0, 0, 0, 0, 0, 0, 0, 0, 0, 0, 0, 0, 0 }		},	// MI_PARAMETER
+		{ false,	{ 0, 0, 0, 0, 0, 0, 0, 0, 0, 0, 0, 0, 0, 0, 0, 0, 0, 0, 0, 0, 0, 0, 0, 0, 0, 0, 0, 0, 0 }		},	// IS_PARAMETER
 
 		{ true,		{ 0, 0, 0, 0, 0, 0, 0, 0, 0, 0, 0, 0, 0, 0, 0, 0, 0, 0, 0, 0, 0, 0, 0, 0, 0, 0, 0, 0, 0 }		},	// IM_REFERENCE
 		{ true,		{ 0, 0, 0, 0, 0, 0, 0, 0, 0, 0, 0, 0, 0, 0, 0, 0, 0, 0, 0, 0, 0, 0, 0, 0, 0, 0, 0, 0, 0 }		},	// ME_REFERENCE
@@ -115,12 +116,14 @@ namespace UE::Mutable::Private
 
 		{ true,		{ 0, 0, 0, 0, 0, 0, 0, 0, 0, 0, 0, 0, 0, 0, 0, 0, 0, 0, 0, 0, 0, 0, 0, 0, 0, 0, 0, 0, 0 }		},	// SC_ARITHMETIC
 		{ true,		{ 0, 0, 0, 0, 0, 0, 0, 0, 0, 0, 0, 0, 0, 0, 0, 0, 0, 0, 0, 0, 0, 0, 0, 0, 0, 0, 0, 0, 0 }		},	// SC_CURVE
+		{ true,		{ 0, 0, 0, 0, 0, 0, 0, 0, 0, 0, 0, 0, 0, 0, 0, 0, 0, 0, 0, 0, 0, 0, 0, 0, 0, 0, 0, 0, 0 }		},	// SC_EXTERNAL
 
 		{ false,	{ 0, 0, 0, 0, 0, 0, 0, 0, 0, 0, 0, 0, 0, 0, 0, 0, 0, 0, 0, 0, 0, 0, 0, 0, 0, 0, 0, 0, 0 }		},	// CO_SAMPLEIMAGE
 		{ false,	{ 0, 0, 0, 0, 0, 0, 0, 0, 0, 0, 0, 0, 0, 0, 0, 0, 0, 0, 0, 0, 0, 0, 0, 0, 0, 0, 0, 0, 0 }		},	// CO_SWIZZLE
 		{ false,	{ 0, 0, 0, 0, 0, 0, 0, 0, 0, 0, 0, 0, 0, 0, 0, 0, 0, 0, 0, 0, 0, 0, 0, 0, 0, 0, 0, 0, 0 }		},	// CO_FROMSCALARS
 		{ false,	{ 0, 0, 0, 0, 0, 0, 0, 0, 0, 0, 0, 0, 0, 0, 0, 0, 0, 0, 0, 0, 0, 0, 0, 0, 0, 0, 0, 0, 0 }		},	// CO_ARITHMETIC
 		{ false,	{ 0, 0, 0, 0, 0, 0, 0, 0, 0, 0, 0, 0, 0, 0, 0, 0, 0, 0, 0, 0, 0, 0, 0, 0, 0, 0, 0, 0, 0 }		},	// CO_LINEARTOSRGB
+		{ true,		{ 0, 0, 0, 0, 0, 0, 0, 0, 0, 0, 0, 0, 0, 0, 0, 0, 0, 0, 0, 0, 0, 0, 0, 0, 0, 0, 0, 0, 0 }		},	// CO_EXTERNAL
 
 		{ true,		{ 0, 1, 1, 1, 0, 0, 0, 0, 0, 0, 0, 0, 0, 0, 0, 0, 0, 0, 0, 0, 0, 0, 0, 0, 0, 0, 0, 0, 0 }		},	// IM_LAYER
 		{ true,		{ 0, 1, 1, 1, 0, 0, 0, 0, 0, 0, 0, 0, 0, 0, 0, 0, 0, 0, 0, 0, 0, 0, 0, 0, 0, 0, 0, 0, 0 }		},	// IM_LAYERCOLOUR
@@ -148,6 +151,7 @@ namespace UE::Mutable::Private
 		{ true,		{ 0, 1, 1, 1, 0, 0, 0, 0, 0, 0, 0, 0, 0, 0, 0, 0, 0, 0, 0, 0, 0, 0, 0, 0, 0, 0, 0, 0, 0 }		},	// IM_NORMALCOMPOSITE
 		{ true,		{ 0, 1, 1, 1, 0, 0, 0, 0, 0, 0, 0, 0, 0, 0, 0, 0, 0, 0, 0, 0, 0, 0, 0, 0, 0, 0, 0, 0, 0 }		},	// IM_TRANSFORM
 		{ true,		{ 0, 0, 0, 0, 0, 0, 0, 0, 0, 0, 0, 0, 0, 0, 0, 0, 0, 0, 0, 0, 0, 0, 0, 0, 0, 0, 0, 0, 0 }		},	// IM_PARAMETER_CONVERT
+		{ true,		{ 0, 0, 0, 0, 0, 0, 0, 0, 0, 0, 0, 0, 0, 0, 0, 0, 0, 0, 0, 0, 0, 0, 0, 0, 0, 0, 0, 0, 0 }		},	// IM_EXTERNAL
 
 		{ true,		{ 0, 0, 0, 0, 0, 0, 0, 0, 0, 0, 0, 0, 0, 0, 0, 0, 0, 0, 0, 0, 0, 0, 0, 0, 0, 0, 0, 0, 0 }		},	// ME_APPLYLAYOUT
 		{ true,		{ 0, 0, 0, 0, 0, 0, 0, 0, 0, 0, 0, 0, 0, 0, 0, 0, 0, 0, 0, 0, 0, 0, 0, 0, 0, 0, 0, 0, 0 }		},	// ME_PREPARELAYOUT
@@ -174,6 +178,7 @@ namespace UE::Mutable::Private
 		{ true,		{ 0, 0, 0, 0, 0, 0, 0, 0, 0, 0, 0, 0, 0, 0, 0, 0, 0, 0, 0, 0, 0, 0, 0, 0, 0, 0, 0, 0, 0 }		},	// ME_ADDMETADATA
 		{ true,		{ 0, 0, 0, 0, 0, 0, 0, 0, 0, 0, 0, 0, 0, 0, 0, 0, 0, 0, 0, 0, 0, 0, 0, 0, 0, 0, 0, 0, 0 }		},	// ME_TRANSFORMWITHMESH
 		{ true,		{ 0, 0, 0, 0, 0, 0, 0, 0, 0, 0, 0, 0, 0, 0, 0, 0, 0, 0, 0, 0, 0, 0, 0, 0, 0, 0, 0, 0, 0 }		},	// ME_TRANSFORMWITHBONE
+		{ true,		{ 0, 0, 0, 0, 0, 0, 0, 0, 0, 0, 0, 0, 0, 0, 0, 0, 0, 0, 0, 0, 0, 0, 0, 0, 0, 0, 0, 0, 0 }		},	// ME_EXTERNAL
 
 		{ false,	{ 0, 0, 0, 0, 0, 0, 0, 0, 0, 0, 0, 0, 0, 0, 0, 0, 0, 0, 0, 0, 0, 0, 0, 0, 0, 0, 0, 0, 0 }		},	// IN_ADDMESH
 		{ false,	{ 0, 0, 0, 0, 0, 0, 0, 0, 0, 0, 0, 0, 0, 0, 0, 0, 0, 0, 0, 0, 0, 0, 0, 0, 0, 0, 0, 0, 0 }		},	// IN_ADDIMAGE
@@ -192,7 +197,9 @@ namespace UE::Mutable::Private
 		{ true,		{ 0, 0, 0, 0, 0, 0, 0, 0, 0, 0, 0, 0, 0, 0, 0, 0, 0, 0, 0, 0, 0, 0, 0, 0, 0, 0, 0, 0, 0 }		},	// LA_REMOVEBLOCKS
 		{ true,		{ 0, 0, 0, 0, 0, 0, 0, 0, 0, 0, 0, 0, 0, 0, 0, 0, 0, 0, 0, 0, 0, 0, 0, 0, 0, 0, 0, 0, 0 }		},	// LA_FROMMESH
 
-		{ true,		{ 0, 0, 0, 0, 0, 0, 0, 0, 0, 0, 0, 0, 0, 0, 0, 0, 0, 0, 0, 0, 0, 0, 0, 0, 0, 0, 0, 0, 0 }		},	// UN_EXTERNAL
+		{ true,		{ 0, 0, 0, 0, 0, 0, 0, 0, 0, 0, 0, 0, 0, 0, 0, 0, 0, 0, 0, 0, 0, 0, 0, 0, 0, 0, 0, 0, 0 }		},	// MI_EXTERNAL
+
+    	{ true,		{ 0, 0, 0, 0, 0, 0, 0, 0, 0, 0, 0, 0, 0, 0, 0, 0, 0, 0, 0, 0, 0, 0, 0, 0, 0, 0, 0, 0, 0 }		},	// IS_EXTERNAL
 	};
 
 	// clang-format on
diff --git a/Engine/Plugins/Mutable/Source/MutableTools/Private/MuT/ErrorLog.cpp b/Engine/Plugins/Mutable/Source/MutableTools/Private/MuT/ErrorLog.cpp
--- a/Engine/Plugins/Mutable/Source/MutableTools/Private/MuT/ErrorLog.cpp
+++ b/Engine/Plugins/Mutable/Source/MutableTools/Private/MuT/ErrorLog.cpp
@@ -224,6 +224,7 @@ namespace UE::Mutable::Private
 		TEXT("ST_PARAMETER     "),
 		TEXT("MA_PARAMETER     "),
 		TEXT("MI_PARAMETER     "),
+		TEXT("IS_PARAMETER     "),
 		
 		TEXT("IM_REFERENCE     "),
 		TEXT("ME_REFERENCE     "),
@@ -260,12 +261,14 @@ namespace UE::Mutable::Private
 		
 		TEXT("SC_ARITHMETIC    "),
 		TEXT("SC_CURVE         "),
-		
+    	TEXT("SC_EXTERNAL      "),
+
 		TEXT("CO_SAMPLEIMAGE   "),
 		TEXT("CO_SWIZZLE       "),
 		TEXT("CO_FROMSCALARS   "),
 		TEXT("CO_ARITHMETIC    "),
 		TEXT("CO_LINEARTOSRGB  "),
+    	TEXT("CO_EXTERNAL      "),
 
 		TEXT("IM_LAYER         "),
 		TEXT("IM_LAYERCOLOUR   "),
@@ -293,6 +296,7 @@ namespace UE::Mutable::Private
 		TEXT("IM_NORMAL_COMPO  "),
 		TEXT("IM_TRANSFORM     "),
 		TEXT("IM_PARAMETER_CONVERT"),
+    	TEXT("IM_EXTERNAL      "),
 
 		TEXT("ME_APPLYLAYOUT   "),
 		TEXT("ME_PREPARELAYOUT "),
@@ -319,6 +323,7 @@ namespace UE::Mutable::Private
 		TEXT("ME_ADDMETADATA   "),
 		TEXT("ME_TRANSFORMWITHMESH"),
 		TEXT("ME_TRANSFORMWITHBONE"),
+    	TEXT("ME_EXTERNAL      "),
 
 		TEXT("IN_ADDMESH       "),
 		TEXT("IN_ADDIMAGE      "),
@@ -337,7 +342,9 @@ namespace UE::Mutable::Private
 		TEXT("LA_REMOVEBLOCKS  "),
 		TEXT("LA_FROMMESH      "),
 
-    	TEXT("UN_EXTERNAL      "),
+    	TEXT("MI_EXTERNAL      "),
+
+    	TEXT("IS_EXTERNAL      "),
 	};
 
 	static_assert(UE_ARRAY_COUNT(s_opNames) == int32(EOpType::COUNT));
diff --git a/Engine/Plugins/Mutable/Source/MutableTools/Private/MuT/Node.cpp b/Engine/Plugins/Mutable/Source/MutableTools/Private/MuT/Node.cpp
--- a/Engine/Plugins/Mutable/Source/MutableTools/Private/MuT/Node.cpp
+++ b/Engine/Plugins/Mutable/Source/MutableTools/Private/MuT/Node.cpp
@@ -96,6 +96,8 @@
 #include "MuT/NodeBool.h"
 #include "MuT/NodeColourExternal.h"
 #include "MuT/NodeExternal.h"
+#include "MuT/NodeExternalOperation.h"
+#include "MuT/NodeExternalParameter.h"
 #include "MuT/NodeImageExternal.h"
 #include "MuT/NodeProjector.h"
 #include "MuT/NodeRange.h"
@@ -152,6 +154,8 @@ namespace UE::Mutable::Private
 	FNodeType NodeSurfaceVariation::StaticType = FNodeType(Node::EType::SurfaceVariation, NodeSurface::GetStaticType());
 
 	FNodeType NodeExternal::StaticType = FNodeType(Node::EType::External, Node::GetStaticType());
+	FNodeType NodeExternalOperation::StaticType = FNodeType(Node::EType::ExternalOperation, NodeExternal::GetStaticType());
+	FNodeType NodeExternalParameter::StaticType = FNodeType(Node::EType::ExternalParameter, NodeExternal::GetStaticType());
 	
 	FNodeType NodeLOD::StaticType = FNodeType(Node::EType::LOD, Node::GetStaticType());
 	FNodeType NodeExtensionData::StaticType = FNodeType(Node::EType::ExtensionData, Node::GetStaticType());
