[Mutable] Fix "use 'template' keyword" compile error on some compilers.

#rnx

diff --git a/Engine/Plugins/Mutable/Source/CustomizableObjectEditor/Private/MuCOE/EdGraphSchema_CustomizableObject.cpp b/Engine/Plugins/Mutable/Source/CustomizableObjectEditor/Private/MuCOE/EdGraphSchema_CustomizableObject.cpp
--- a/Engine/Plugins/Mutable/Source/CustomizableObjectEditor/Private/MuCOE/EdGraphSchema_CustomizableObject.cpp
+++ b/Engine/Plugins/Mutable/Source/CustomizableObjectEditor/Private/MuCOE/EdGraphSchema_CustomizableObject.cpp
@@ -1744,7 +1744,6 @@ TSharedPtr<IAssetReferenceFilter> UEdGraphSchema_CustomizableObject::MakeAssetRe
 bool UEdGraphSchema_CustomizableObject::IsTypeNative(const UScriptStruct* Type)
 {
 	return Type == UE::Mutable::FMeshAdapter::StaticStruct() ||
-		Type == UE::Mutable::FTextureAdapter::StaticStruct() ||
 		Type == UE::Mutable::FTextureAdapter::StaticStruct() ||
 		Type == UE::Mutable::FFloatAdapter::StaticStruct() ||
 		Type == UE::Mutable::FVectorAdapter::StaticStruct() ||
diff --git a/Engine/Plugins/Mutable/Source/MutableRuntime/Public/MuR/External/Value.h b/Engine/Plugins/Mutable/Source/MutableRuntime/Public/MuR/External/Value.h
--- a/Engine/Plugins/Mutable/Source/MutableRuntime/Public/MuR/External/Value.h
+++ b/Engine/Plugins/Mutable/Source/MutableRuntime/Public/MuR/External/Value.h
@@ -172,12 +172,12 @@ namespace UE::Mutable
 		
 		T& Get()
 		{
-			return Ptr.Get()->GetMutable<T>();
+			return Ptr.Get()->template GetMutable<T>();
 		}
 
 		T* GetPtr()
 		{
-			return Ptr.Get()->GetMutablePtr<T>();
+			return Ptr.Get()->template GetMutablePtr<T>();
 		}
 		
 	private:
@@ -227,12 +227,12 @@ namespace UE::Mutable
 		
 		const T& Get() const
 		{
-			return Ptr.Get()->Get<const T>();
+			return Ptr.Get()->template Get<const T>();
 		}
 
 		const T* GetPtr() const
 		{
-			return Ptr.Get()->GetPtr<T>();
+			return Ptr.Get()->template GetPtr<T>();
 		}
 		
 	private:
