[Mutable] Fix Dataless crash when streaming inline SKM LODs.

- Changed ColorVertexBuffer serialization to respect bNeedsCPUAccess flag.

diff --git a/Engine/Plugins/Mutable/Source/CustomizableObject/Private/MuCO/UnrealConversionUtils.cpp b/Engine/Plugins/Mutable/Source/CustomizableObject/Private/MuCO/UnrealConversionUtils.cpp
--- a/Engine/Plugins/Mutable/Source/CustomizableObject/Private/MuCO/UnrealConversionUtils.cpp
+++ b/Engine/Plugins/Mutable/Source/CustomizableObject/Private/MuCO/UnrealConversionUtils.cpp
@@ -1005,6 +1005,12 @@ namespace UnrealConversionUtils
 			[Context, SkeletalMesh = TStrongObjectPtr(SkeletalMesh), LODIndex, SectionIndex, Result, BoneNames]()
 			{
 				MUTABLE_CPUPROFILER_SCOPE(MeshConversion);
+
+				if (!SkeletalMesh->NeedCPUData(LODIndex))
+				{
+					UE_LOG(LogMutable, Error, TEXT("Unabel to read LOD %d Data from the CPU. Enable bAllowCPUAccess on the Skeletal Mesh %s."), LODIndex, *SkeletalMesh->GetName());
+					return;
+				}
 					
 				if (!Context->OriginalLOD)
 				{
diff --git a/Engine/Source/Runtime/Engine/Private/SkeletalMeshLODRenderData.cpp b/Engine/Source/Runtime/Engine/Private/SkeletalMeshLODRenderData.cpp
--- a/Engine/Source/Runtime/Engine/Private/SkeletalMeshLODRenderData.cpp
+++ b/Engine/Source/Runtime/Engine/Private/SkeletalMeshLODRenderData.cpp
@@ -812,7 +812,7 @@ void FSkeletalMeshLODRenderData::SerializeStreamedData(FArchive& Ar, USkinnedAss
 
 	if (Owner && Owner->GetHasVertexColors())
 	{
-		StaticVertexBuffers.ColorVertexBuffer.Serialize(Ar, bForceKeepCPUResources);
+		StaticVertexBuffers.ColorVertexBuffer.Serialize(Ar, bNeedsCPUAccess);
 		StaticVertexBuffers.ColorVertexBuffer.SetResourceName(SkeletalMeshResourceNames::StaticColorVertexBufferName);
 	}
 
