[Mutable] Replace ME_PARAMETER with SK_PARAMETER

- Initial introduction of the Skeletal Mesh type.
- SK_PARAMETER always returns a Passthrough Parameter Skeletal Mesh.
- ME_SKELETALMESH_BREAK breaks a section of a Skeletal Mesh.
- Fix crash when UObject parameters are nullptr.
- Fix bug with FUniqueMeshIds::EnsureUnique. Ids were always incremented.

diff --git a/Engine/Plugins/Mutable/Source/CustomizableObject/Internal/MuCO/CustomizableObjectPrivate.h b/Engine/Plugins/Mutable/Source/CustomizableObject/Internal/MuCO/CustomizableObjectPrivate.h
--- a/Engine/Plugins/Mutable/Source/CustomizableObject/Internal/MuCO/CustomizableObjectPrivate.h
+++ b/Engine/Plugins/Mutable/Source/CustomizableObject/Internal/MuCO/CustomizableObjectPrivate.h
@@ -1748,6 +1748,8 @@ enum class ECustomizableObjectVersions : int32
 
 	InstancedStructType,
 
+	SkeletalMeshParameter,
+
 	// -----<new versions can be added above this line>--------
 	LastCustomizableObjectVersion
 };
diff --git a/Engine/Plugins/Mutable/Source/CustomizableObject/Private/MuCO/CustomizableObject.cpp b/Engine/Plugins/Mutable/Source/CustomizableObject/Private/MuCO/CustomizableObject.cpp
--- a/Engine/Plugins/Mutable/Source/CustomizableObject/Private/MuCO/CustomizableObject.cpp
+++ b/Engine/Plugins/Mutable/Source/CustomizableObject/Private/MuCO/CustomizableObject.cpp
@@ -1463,13 +1463,13 @@ void UCustomizableObjectPrivate::UpdateParameterPropertiesFromModel(const TShare
 				break;
 			}
 				
-			case UE::Mutable::Private::EParameterType::Image:
+			case UE::Mutable::Private::EParameterType::Texture:
 			{
 				Data.Type = EMutableParameterType::Texture;
 				break;
 			}
 
-			case UE::Mutable::Private::EParameterType::Mesh:
+			case UE::Mutable::Private::EParameterType::SkeletalMesh:
 			{
 				Data.Type = EMutableParameterType::SkeletalMesh;
 				break;
@@ -2259,11 +2259,11 @@ void UCustomizableObjectPrivate::SetModel(const TSharedPtr<UE::Mutable::Private:
 		{
 			switch (Parameter.Type)
 			{
-			case UE::Mutable::Private::EParameterType::Image:
+			case UE::Mutable::Private::EParameterType::Texture:
 				Parameter.DefaultValue.Set<UE::Mutable::Private::FParamTextureType>(TStrongObjectPtr<UTexture>(LocalModelResources->TextureParameterDefaultValues[Parameter.Name]));
 				break;
 				
-			case UE::Mutable::Private::EParameterType::Mesh:
+			case UE::Mutable::Private::EParameterType::SkeletalMesh:
 				Parameter.DefaultValue.Set<UE::Mutable::Private::FParamSkeletalMeshType>(TStrongObjectPtr<USkeletalMesh>(LocalModelResources->SkeletalMeshParameterDefaultValues[Parameter.Name]));
 				break;
 				
diff --git a/Engine/Plugins/Mutable/Source/CustomizableObject/Private/MuCO/CustomizableObjectInstanceDescriptor.cpp b/Engine/Plugins/Mutable/Source/CustomizableObject/Private/MuCO/CustomizableObjectInstanceDescriptor.cpp
--- a/Engine/Plugins/Mutable/Source/CustomizableObject/Private/MuCO/CustomizableObjectInstanceDescriptor.cpp
+++ b/Engine/Plugins/Mutable/Source/CustomizableObject/Private/MuCO/CustomizableObjectInstanceDescriptor.cpp
@@ -220,7 +220,7 @@ void FCustomizableObjectInstanceDescriptor::SaveDescriptor(FArchive& Ar, bool bU
 					RangeValues.Empty(P.ParameterRangeValues.Num());	
 					Algo::Transform(P.ParameterRangeValues, RangeValues, [](const TObjectPtr<UTexture>& Element)
 					{
-						return Element->GetPathName();
+						return Element.GetPathName();
 					});
  
 					break;
@@ -246,7 +246,7 @@ void FCustomizableObjectInstanceDescriptor::SaveDescriptor(FArchive& Ar, bool bU
 					RangeValues.Empty(P.ParameterRangeValues.Num());	
 					Algo::Transform(P.ParameterRangeValues, RangeValues, [](const TObjectPtr<USkeletalMesh>& Element)
 					{
-						return Element->GetPathName();
+						return Element.GetPathName();
 					});
 					
 					break;
@@ -258,7 +258,34 @@ void FCustomizableObjectInstanceDescriptor::SaveDescriptor(FArchive& Ar, bool bU
 				
 			break;
 		}
- 
+
+		case EMutableParameterType::Material:
+		{
+			FString Value;
+			TArray<FString> RangeValues;
+
+			for (const FCustomizableObjectMaterialParameterValue& P : MaterialParameters)
+			{
+				if (P.ParameterName == Name)
+				{
+					Value = P.ParameterValue.GetPathName();
+				
+					RangeValues.Empty(P.ParameterRangeValues.Num());	
+					Algo::Transform(P.ParameterRangeValues, RangeValues, [](const TObjectPtr<UMaterialInterface>& Element)
+					{
+						return Element.GetPathName();
+					});
+				
+					break;
+				}
+			}
+
+			Ar << Value;
+			Ar << RangeValues;
+			
+			break;
+		}
+			
 		case EMutableParameterType::Projector:
 		{
 			FCustomizableObjectProjector Value;
@@ -762,7 +789,7 @@ TSharedPtr<UE::Mutable::Private::FParameters> FCustomizableObjectInstanceDescrip
 			break;
 		}
  
-		case UE::Mutable::Private::EParameterType::Image:
+		case UE::Mutable::Private::EParameterType::Texture:
 		{
 			for (const FCustomizableObjectTextureParameterValue& TextureParameter : TextureParameters)
 			{
@@ -773,12 +800,12 @@ TSharedPtr<UE::Mutable::Private::FParameters> FCustomizableObjectInstanceDescrip
 						for (int32 RangeIndex = 0; RangeIndex < TextureParameter.ParameterRangeValues.Num(); ++RangeIndex)
 						{
 							RangeIdxPtr->SetPosition(0, RangeIndex);
-							MutableParameters->SetImageValue(ParamIndex, TextureParameter.ParameterRangeValues[RangeIndex], RangeIdxPtr.Get());
+							MutableParameters->SetTextureValue(ParamIndex, TextureParameter.ParameterRangeValues[RangeIndex], RangeIdxPtr.Get());
 						}
 					}
 					else
 					{
-						MutableParameters->SetImageValue(ParamIndex, TextureParameter.ParameterValue);
+						MutableParameters->SetTextureValue(ParamIndex, TextureParameter.ParameterValue);
 					}
  
 					break;
@@ -788,7 +815,7 @@ TSharedPtr<UE::Mutable::Private::FParameters> FCustomizableObjectInstanceDescrip
 			break;
 		}
  
-		case UE::Mutable::Private::EParameterType::Mesh:
+		case UE::Mutable::Private::EParameterType::SkeletalMesh:
 		{
 			for (const FCustomizableObjectSkeletalMeshParameterValue& MeshParameter : SkeletalMeshParameters)
 			{
@@ -799,12 +826,12 @@ TSharedPtr<UE::Mutable::Private::FParameters> FCustomizableObjectInstanceDescrip
 						for (int32 RangeIndex = 0; RangeIndex < MeshParameter.ParameterRangeValues.Num(); ++RangeIndex)
 						{
 							RangeIdxPtr->SetPosition(0, RangeIndex);
-							MutableParameters->SetMeshValue(ParamIndex, MeshParameter.ParameterRangeValues[RangeIndex], RangeIdxPtr.Get());
+							MutableParameters->SetSkeletalMeshValue(ParamIndex, MeshParameter.ParameterRangeValues[RangeIndex], RangeIdxPtr.Get());
 						}
 					}
 					else
 					{
-						MutableParameters->SetMeshValue(ParamIndex, MeshParameter.ParameterValue);
+						MutableParameters->SetSkeletalMeshValue(ParamIndex, MeshParameter.ParameterValue);
 					}
  
 					break;
@@ -1270,7 +1297,7 @@ void FCustomizableObjectInstanceDescriptor::ReloadParameters()
 			break;
 		}
  
-		case UE::Mutable::Private::EParameterType::Image:
+		case UE::Mutable::Private::EParameterType::Texture:
 		{
 			FCustomizableObjectTextureParameterValue Param;
 			Param.ParameterName = Name;
@@ -1314,12 +1341,12 @@ void FCustomizableObjectInstanceDescriptor::ReloadParameters()
 							}
 						}
  
-						Param.ParameterRangeValues[RangeIndex] = MutableParameters->GetImageValue(ParamIndex, RangeValueIdxPtr.Get());
+						Param.ParameterRangeValues[RangeIndex] = MutableParameters->GetTextureValue(ParamIndex, RangeValueIdxPtr.Get());
 					}
 				}
 				else
 				{
-					Param.ParameterValue = MutableParameters->GetImageValue(ParamIndex);
+					Param.ParameterValue = MutableParameters->GetTextureValue(ParamIndex);
 				}
 			}
  
@@ -1328,7 +1355,7 @@ void FCustomizableObjectInstanceDescriptor::ReloadParameters()
 		}
  
  
-		case UE::Mutable::Private::EParameterType::Mesh:
+		case UE::Mutable::Private::EParameterType::SkeletalMesh:
 		{
 			FCustomizableObjectSkeletalMeshParameterValue Param;
 			Param.ParameterName = Name;
@@ -1372,12 +1399,12 @@ void FCustomizableObjectInstanceDescriptor::ReloadParameters()
 							}
 						}
  
-						Param.ParameterRangeValues[RangeIndex] = MutableParameters->GetMeshValue(ParamIndex, RangeValueIdxPtr.Get());
+						Param.ParameterRangeValues[RangeIndex] = MutableParameters->GetSkeletalMeshValue(ParamIndex, RangeValueIdxPtr.Get());
 					}
 				}
 				else
 				{
-					Param.ParameterValue = MutableParameters->GetMeshValue(ParamIndex);
+					Param.ParameterValue = MutableParameters->GetSkeletalMeshValue(ParamIndex);
 				}
 			}
  
diff --git a/Engine/Plugins/Mutable/Source/CustomizableObjectEditor/Private/MuCOE/GenerateMutableSource/GenerateMutableSource.cpp b/Engine/Plugins/Mutable/Source/CustomizableObjectEditor/Private/MuCOE/GenerateMutableSource/GenerateMutableSource.cpp
--- a/Engine/Plugins/Mutable/Source/CustomizableObjectEditor/Private/MuCOE/GenerateMutableSource/GenerateMutableSource.cpp
+++ b/Engine/Plugins/Mutable/Source/CustomizableObjectEditor/Private/MuCOE/GenerateMutableSource/GenerateMutableSource.cpp
@@ -434,6 +434,12 @@ uint32 GetTypeHash(const FGeneratedSourceExternalKey& Key)
 }
 
 
+uint32 GetTypeHash(const FGeneratedSourceSkeletalMeshKey& Key)
+{
+	return GetTypeHash(Key.Pin);
+}
+
+
 FGraphCycleKey::FGraphCycleKey(const UEdGraphPin& Pin, const FString& Id, const UCustomizableObjectNodeMacroInstance* MacroContext) :
 	Pin(Pin),
 	Id(Id),
diff --git a/Engine/Plugins/Mutable/Source/CustomizableObjectEditor/Private/MuCOE/GenerateMutableSource/GenerateMutableSource.h b/Engine/Plugins/Mutable/Source/CustomizableObjectEditor/Private/MuCOE/GenerateMutableSource/GenerateMutableSource.h
--- a/Engine/Plugins/Mutable/Source/CustomizableObjectEditor/Private/MuCOE/GenerateMutableSource/GenerateMutableSource.h
+++ b/Engine/Plugins/Mutable/Source/CustomizableObjectEditor/Private/MuCOE/GenerateMutableSource/GenerateMutableSource.h
@@ -35,6 +35,8 @@
 #include "MuR/Mesh.h"
 #include "MuR/Ptr.h"
 #include "MuT/NodeImageParameter.h"
+#include "MuT/NodeSkeletalMesh.h"
+#include "MuT/NodeSkeletalMeshParameter.h"
 #include "UObject/Package.h"
 #include "UObject/StrongObjectPtr.h"
 #include "Tasks/Task.h"
@@ -192,6 +194,17 @@ struct FGeneratedSourceExternalKey
 uint32 GetTypeHash(const FGeneratedSourceExternalKey& Key);
 
 
+struct FGeneratedSourceSkeletalMeshKey
+{
+	const UEdGraphPin* Pin = nullptr;
+	
+	bool operator==(const FGeneratedSourceSkeletalMeshKey& Other) const = default;
+};
+
+
+uint32 GetTypeHash(const FGeneratedSourceSkeletalMeshKey& Key);
+
+
 struct FImageKey
 {
 	FImageKey(const UEdGraphPin* InPin)
@@ -270,6 +283,12 @@ struct FGeneratedSourceExternalData
 };
 
 
+struct FGeneratedSourceSkeletalMeshData
+{
+	UE::Mutable::Private::Ptr<UE::Mutable::Private::NodeSkeletalMesh> Node;
+};
+
+
 inline uint32 GetTypeHash(const FImageKey& Key)
 {
 	uint32 GuidHash = GetTypeHash(Key.Pin->PinId);
@@ -618,6 +637,7 @@ struct FMutableGraphGenerationContext
 	// Cache of generated pins per LOD
 	TMap<FGeneratedKey, FGeneratedData> Generated;
 	TMap<FGeneratedSourceExternalKey, FGeneratedSourceExternalData> GeneratedSourceExternals;
+	TMap<FGeneratedSourceSkeletalMeshKey, FGeneratedSourceSkeletalMeshData> GeneratedSourceSkeletalMesh;
 
 	/** Set of all generated nodes. */
 	TSet<UCustomizableObjectNode*> GeneratedNodes;
@@ -651,6 +671,10 @@ struct FMutableGraphGenerationContext
 
 	/** Key is the Node Uid. */
 	TMap<FGuid, UE::Mutable::Private::Ptr<UE::Mutable::Private::NodeImageParameter>> GeneratedImageParameters;
+
+	/** Key is the Node Uid. */
+	TMap<FGuid, UE::Mutable::Private::Ptr<UE::Mutable::Private::NodeSkeletalMeshParameter>> GeneratedSkeletalMeshParameters;
+
 	
 	struct FGeneratedCompositeDataTablesData
 	{
diff --git a/Engine/Plugins/Mutable/Source/CustomizableObjectEditor/Private/MuCOE/GenerateMutableSource/GenerateMutableSourceMesh.cpp b/Engine/Plugins/Mutable/Source/CustomizableObjectEditor/Private/MuCOE/GenerateMutableSource/GenerateMutableSourceMesh.cpp
--- a/Engine/Plugins/Mutable/Source/CustomizableObjectEditor/Private/MuCOE/GenerateMutableSource/GenerateMutableSourceMesh.cpp
+++ b/Engine/Plugins/Mutable/Source/CustomizableObjectEditor/Private/MuCOE/GenerateMutableSource/GenerateMutableSourceMesh.cpp
@@ -11,6 +11,7 @@
 #include "ClothingAsset.h"
 #include "GenerateMutableSourceExternal.h"
 #include "GenerateMutableSourceImage.h"
+#include "GenerateMutableSourceSkeletalMesh.h"
 #include "MeshUtilities.h"
 #include "Modules/ModuleManager.h"
 #include "MuCO/CustomizableObjectPrivate.h"
@@ -52,7 +53,7 @@
 #include "MuT/NodeMeshSwitch.h"
 #include "MuT/NodeMeshTable.h"
 #include "MuT/NodeMeshVariation.h"
-#include "MuT/NodeMeshParameter.h"
+#include "MuT/NodeMeshSkeletalMeshBreak.h"
 #include "MuR/Types.h"
 #include "MuR/OpMeshFormat.h"
 #include "MuR/ConvertData.h"
@@ -4590,27 +4591,21 @@ UE::Mutable::Private::Ptr<UE::Mutable::Private::NodeMesh> GenerateMutableSourceM
 
 	else if (const UCustomizableObjectNodeSkeletalMeshParameter* TypedNodeParam = Cast<UCustomizableObjectNodeSkeletalMeshParameter>(Node))
 	{
-		UE::Mutable::Private::Ptr<UE::Mutable::Private::NodeMeshParameter> ParameterNode = new UE::Mutable::Private::NodeMeshParameter();
-
-		ParameterNode->Name = TypedNodeParam->GetParameterName();
-		ParameterNode->UID = GenerationContext.GetNodeIdUnique(Node).ToString();
-
-		GenerationContext.SkeletalMeshParameterDefaultValues.Add(TypedNodeParam->GetParameterName(), UE::Mutable::Private::LoadObject(TypedNodeParam->DefaultValue));
-
+		UE::Mutable::Private::Ptr<UE::Mutable::Private::NodeSkeletalMeshParameter> ParameterNode = UE::Mutable::Private::GenerateMutableSourceSkeletalMeshParameter(Pin, GenerationContext);
+		
+		UE::Mutable::Private::Ptr<UE::Mutable::Private::NodeMeshSkeletalMeshBreak> BreakNode = new UE::Mutable::Private::NodeMeshSkeletalMeshBreak();
+		BreakNode->SkeletalMeshParameter = ParameterNode;
+		
 		int32 BaseLODIndex = 0;
 		int32 BaseSectionIndex = 0;
 
 		{
 			TypedNodeParam->GetPinSection(*Pin, BaseLODIndex, BaseSectionIndex);
-			ParameterNode->LODIndex = BaseLODIndex;
-			ParameterNode->SectionIndex = BaseSectionIndex;
-			ParameterNode->bAutomaticLODs = GenerationContext.CurrentAutoLODStrategy == ECustomizableObjectAutomaticLODStrategy::AutomaticFromMesh;
+			BreakNode->LODIndex = BaseLODIndex;
+			BreakNode->SectionIndex = BaseSectionIndex;
+			BreakNode->bAutomaticLODs = GenerationContext.CurrentAutoLODStrategy == ECustomizableObjectAutomaticLODStrategy::AutomaticFromMesh;
 		}
 		
-		GenerationContext.ParameterUIDataMap.Add(TypedNodeParam->GetParameterName(), FMutableParameterData(
-			TypedNodeParam->ParamUIMetadata,
-			EMutableParameterType::Texture));
-
 		// Layouts
 		{
 			const FLayoutGenerationFlags& LayoutFlags = GenerationContext.LayoutGenerationFlags.Last();
@@ -4618,19 +4613,19 @@ UE::Mutable::Private::Ptr<UE::Mutable::Private::NodeMesh> GenerateMutableSourceM
 			const TArray<UCustomizableObjectLayout*>& Layouts = TypedNodeParam->GetLayouts(*Pin);
 			const int32 NumLayouts = Layouts.Num();
 
-			ParameterNode->Layouts.SetNum(NumLayouts);
+			BreakNode->Layouts.SetNum(NumLayouts);
 			for (int32 LayoutIndex = 0; LayoutIndex < NumLayouts; ++LayoutIndex)
 			{
 				if (!LayoutFlags.TexturePinModes.IsValidIndex(LayoutIndex) ||
 					LayoutFlags.TexturePinModes[LayoutIndex] != EPinMode::Mutable)
 				{
-					ParameterNode->Layouts[LayoutIndex] = CreateDefaultLayout();
+					BreakNode->Layouts[LayoutIndex] = CreateDefaultLayout();
 
 					if (Layouts.IsValidIndex(LayoutIndex) && Layouts[LayoutIndex])
 					{
 						// Keep packing strategy if possible, Overlay can be valid with EPinMode == Passthrough	
-						ParameterNode->Layouts[LayoutIndex]->Strategy = ConvertLayoutStrategy(Layouts[LayoutIndex]->PackingStrategy);
-						ParameterNode->Layouts[LayoutIndex]->TexCoordsIndex = Layouts[LayoutIndex]->GetUVChannel();
+						BreakNode->Layouts[LayoutIndex]->Strategy = ConvertLayoutStrategy(Layouts[LayoutIndex]->PackingStrategy);
+						BreakNode->Layouts[LayoutIndex]->TexCoordsIndex = Layouts[LayoutIndex]->GetUVChannel();
 					}
 
 					// Ignore layout
@@ -4644,7 +4639,7 @@ UE::Mutable::Private::Ptr<UE::Mutable::Private::NodeMesh> GenerateMutableSourceM
 					UE::Mutable::Private::Ptr<UE::Mutable::Private::NodeLayout> LayoutNode = CreateMutableLayoutNode(Layout, bLinkedToExtendMaterial);
 					LayoutNode->SetMessageContext(Node);
 
-					ParameterNode->Layouts[LayoutIndex] = LayoutNode;
+					BreakNode->Layouts[LayoutIndex] = LayoutNode;
 				}
 			}
 
@@ -4662,7 +4657,7 @@ UE::Mutable::Private::Ptr<UE::Mutable::Private::NodeMesh> GenerateMutableSourceM
 				// Get the base mesh
 				FMutableSourceMeshData Source;
 				Source.Mesh = LayoutMesh;
-				ParameterNode->ReferenceMesh = ConvertSkeletalMeshToMutable(Source, BaseLODIndex, BaseSectionIndex, GenerationContext, Node);
+				BreakNode->ReferenceMesh = ConvertSkeletalMeshToMutable(Source, BaseLODIndex, BaseSectionIndex, GenerationContext, Node);
 
 				GenerationContext.MeshGenerationFlags.Pop();
 			}
@@ -4672,7 +4667,7 @@ UE::Mutable::Private::Ptr<UE::Mutable::Private::NodeMesh> GenerateMutableSourceM
 		// animbp instance?
 		// gameplay tags?
 
-		Result = ParameterNode;
+		Result = BreakNode;
 
 		// Applying Mesh Morph Node
 		if (GenerationContext.MeshMorphStack.Num())
diff --git a/Engine/Plugins/Mutable/Source/CustomizableObjectEditor/Private/MuCOE/GenerateMutableSource/GenerateMutableSourceSkeletalMesh.cpp b/Engine/Plugins/Mutable/Source/CustomizableObjectEditor/Private/MuCOE/GenerateMutableSource/GenerateMutableSourceSkeletalMesh.cpp
new file mode 100644
--- /dev/null
+++ b/Engine/Plugins/Mutable/Source/CustomizableObjectEditor/Private/MuCOE/GenerateMutableSource/GenerateMutableSourceSkeletalMesh.cpp
@@ -0,0 +1,111 @@
+// Copyright Epic Games, Inc. All Rights Reserved.
+
+#include "GenerateMutableSourceSkeletalMesh.h"
+
+#include "GenerateMutableSource.h"
+#include "MuCOE/Nodes/CustomizableObjectNodeSkeletalMeshParameter.h"
+#include "MuR/MutableTrace.h"
+#include "MuT/NodeSkeletalMeshParameter.h"
+
+#define LOCTEXT_NAMESPACE "CustomizableObjectEditor"
+
+
+namespace UE::Mutable::Private
+{
+	Ptr<NodeSkeletalMeshParameter> GenerateMutableSourceSkeletalMeshParameter(const UEdGraphPin* Pin, FMutableGraphGenerationContext& GenerationContext)
+	{
+		check(Pin)
+		RETURN_ON_CYCLE(*Pin, GenerationContext)
+
+		CheckNumOutputs(*Pin, GenerationContext);
+	
+		UCustomizableObjectNodeParameter* ParameterNode = CastChecked<UCustomizableObjectNodeParameter>(Pin->GetOwningNode());
+
+		UE::Mutable::Private::Ptr<UE::Mutable::Private::NodeSkeletalMeshParameter> Result;
+
+		const FGuid NodeGUID = ParameterNode->NodeGuid;
+		if (const UE::Mutable::Private::Ptr<UE::Mutable::Private::NodeSkeletalMeshParameter>* GeneratedSkeletalMeshParameter = GenerationContext.GeneratedSkeletalMeshParameters.Find(NodeGUID))
+		{
+			return GeneratedSkeletalMeshParameter->get();
+		}
+	
+		if (const UCustomizableObjectNodeSkeletalMeshParameter* TypedNodeParam = Cast<UCustomizableObjectNodeSkeletalMeshParameter>(ParameterNode))
+		{
+			UE::Mutable::Private::Ptr<UE::Mutable::Private::NodeSkeletalMeshParameter> SkeletalMeshNode = new UE::Mutable::Private::NodeSkeletalMeshParameter();
+
+			SkeletalMeshNode->Name = TypedNodeParam->GetParameterName(&GenerationContext.MacroNodesStack);
+			SkeletalMeshNode->UID = GenerationContext.GetNodeIdUnique(ParameterNode).ToString();
+		
+			GenerationContext.SkeletalMeshParameterDefaultValues.Add(TypedNodeParam->GetParameterName(), LoadObject(TypedNodeParam->DefaultValue));
+
+			GenerationContext.ParameterUIDataMap.Add(TypedNodeParam->GetParameterName(&GenerationContext.MacroNodesStack), FMutableParameterData(
+				TypedNodeParam->ParamUIMetadata,
+				EMutableParameterType::SkeletalMesh));
+
+			Result = SkeletalMeshNode;
+		}
+		else
+		{
+			GenerationContext.Log(LOCTEXT("UnimplementedNode", "Node type not implemented yet."), ParameterNode);
+		}
+
+		if (Result)
+		{
+			Result->SetMessageContext(ParameterNode);
+		}
+	
+		GenerationContext.GeneratedSkeletalMeshParameters.Add(NodeGUID,Result);
+		GenerationContext.GeneratedNodes.Add(ParameterNode);
+	
+		return Result;
+	}
+
+	
+	Ptr<NodeSkeletalMesh> GenerateMutableSourceSkeletalMesh(const UEdGraphPin* Pin, FMutableGraphGenerationContext& GenerationContext)
+	{
+		MUTABLE_CPUPROFILER_SCOPE(GenerateMutableSourceSkeletalMesh);
+
+		check(Pin)
+		RETURN_ON_CYCLE(*Pin, GenerationContext)
+
+		CheckNumOutputs(*Pin, GenerationContext);
+
+		UCustomizableObjectNode* Node = CastChecked<UCustomizableObjectNode>(Pin->GetOwningNode());
+		
+		FGeneratedSourceSkeletalMeshKey Key;
+		Key.Pin = Pin;
+
+		Ptr<NodeSkeletalMesh> Result;
+		
+		if (const FGeneratedSourceSkeletalMeshData* Generated = GenerationContext.GeneratedSourceSkeletalMesh.Find(Key))
+		{
+			return Generated->Node;
+		}
+
+		if (UCustomizableObjectNodeSkeletalMeshParameter* CONodeSkeletalMeshParameter = Cast<UCustomizableObjectNodeSkeletalMeshParameter>(Node))
+		{
+			Result = GenerateMutableSourceSkeletalMeshParameter(Pin, GenerationContext);
+		}
+
+		else
+		{
+			GenerationContext.Log(LOCTEXT("UnimplementedNode", "Node type not implemented yet."), Node);
+		}
+		
+		if (Result)
+		{
+			Result->SetMessageContext(Node);
+		}
+
+		FGeneratedSourceSkeletalMeshData Data;
+		Data.Node = Result;
+		
+		GenerationContext.GeneratedSourceSkeletalMesh.Add(Key, Data);
+		GenerationContext.GeneratedNodes.Add(Node);
+
+		return Result;
+	}
+}
+
+
+#undef LOCTEXT_NAMESPACE
diff --git a/Engine/Plugins/Mutable/Source/CustomizableObjectEditor/Private/MuCOE/GenerateMutableSource/GenerateMutableSourceSkeletalMesh.h b/Engine/Plugins/Mutable/Source/CustomizableObjectEditor/Private/MuCOE/GenerateMutableSource/GenerateMutableSourceSkeletalMesh.h
new file mode 100644
--- /dev/null
+++ b/Engine/Plugins/Mutable/Source/CustomizableObjectEditor/Private/MuCOE/GenerateMutableSource/GenerateMutableSourceSkeletalMesh.h
@@ -0,0 +1,17 @@
+// Copyright Epic Games, Inc. All Rights Reserved.
+
+#pragma once
+
+#include "MuR/Ptr.h"
+#include "MuT/NodeSkeletalMesh.h"
+#include "MuT/NodeSkeletalMeshParameter.h"
+
+struct FMutableGraphGenerationContext;
+class UEdGraphPin;
+
+namespace UE::Mutable::Private
+{
+	Ptr<NodeSkeletalMesh> GenerateMutableSourceSkeletalMesh(const UEdGraphPin* Pin, FMutableGraphGenerationContext& GenerationContext);
+
+	Ptr<NodeSkeletalMeshParameter> GenerateMutableSourceSkeletalMeshParameter(const UEdGraphPin* Pin, FMutableGraphGenerationContext& GenerationContext);
+}
diff --git a/Engine/Plugins/Mutable/Source/CustomizableObjectEditor/Private/MuCOE/SMutableCodeViewer.cpp b/Engine/Plugins/Mutable/Source/CustomizableObjectEditor/Private/MuCOE/SMutableCodeViewer.cpp
--- a/Engine/Plugins/Mutable/Source/CustomizableObjectEditor/Private/MuCOE/SMutableCodeViewer.cpp
+++ b/Engine/Plugins/Mutable/Source/CustomizableObjectEditor/Private/MuCOE/SMutableCodeViewer.cpp
@@ -3197,11 +3197,10 @@ void FMutableCodeTreeElement::GenerateLabelText()
 		break;
 	}
 
-	case UE::Mutable::Private::EOpType::ME_PARAMETER:
+	case UE::Mutable::Private::EOpType::ME_SKELETALMESH_BREAK:
 	{
-		UE::Mutable::Private::OP::MeshParameterArgs Args = Program.GetOpArgs<UE::Mutable::Private::OP::MeshParameterArgs>(MutableOperation);
-		OpName += FString::Printf( TEXT(" LOD %d Section %d of "), Args.LOD, Args.Section );
-		OpName += Program.Parameters[int32(Args.variable)].Name;
+		UE::Mutable::Private::OP::MeshSkeletalMeshBreakArgs Args = Program.GetOpArgs<UE::Mutable::Private::OP::MeshSkeletalMeshBreakArgs>(MutableOperation);
+		OpName += FString::Printf( TEXT(" LOD %d Section %d"), Args.LOD, Args.Section );
 		break;
 	}
 
diff --git a/Engine/Plugins/Mutable/Source/CustomizableObjectEditor/Private/MuCOE/SMutableGraphViewer.cpp b/Engine/Plugins/Mutable/Source/CustomizableObjectEditor/Private/MuCOE/SMutableGraphViewer.cpp
--- a/Engine/Plugins/Mutable/Source/CustomizableObjectEditor/Private/MuCOE/SMutableGraphViewer.cpp
+++ b/Engine/Plugins/Mutable/Source/CustomizableObjectEditor/Private/MuCOE/SMutableGraphViewer.cpp
@@ -187,6 +187,9 @@ static const char* MutableNodeNames[] =
 
 	"LOD",
 
+	"SkeletalMesh",
+	"SkeletalMeshParameter",
+	
 	"Component",
 	"ComponentNew",
 	"ComponentEdit",
diff --git a/Engine/Plugins/Mutable/Source/MutableRuntime/Internal/MuR/Operations.h b/Engine/Plugins/Mutable/Source/MutableRuntime/Internal/MuR/Operations.h
--- a/Engine/Plugins/Mutable/Source/MutableRuntime/Internal/MuR/Operations.h
+++ b/Engine/Plugins/Mutable/Source/MutableRuntime/Internal/MuR/Operations.h
@@ -50,7 +50,7 @@ namespace UE::Mutable::Private
         CO_PARAMETER,
         PR_PARAMETER,
 		IM_PARAMETER,
-		ME_PARAMETER,
+		SK_PARAMETER,
 		ST_PARAMETER,
     	MA_PARAMETER,
     	MI_PARAMETER,
@@ -223,6 +223,7 @@ namespace UE::Mutable::Private
 		//! Apply linear transform to Image content. Resulting samples outside the original image are tiled.
 		IM_TRANSFORM,
 
+        /** Convert an FIamge in Passthrough Parameter mode to an FImage with data. */
     	IM_PARAMETER_CONVERT,
 
     	//! External operation
@@ -313,6 +314,9 @@ namespace UE::Mutable::Private
     	//! External operation
     	ME_EXTERNAL,
 
+        /** Select a Mesh Section from a Skeletal Mesh. */
+    	ME_SKELETALMESH_BREAK,
+
         //-----------------------------------------------------------------------------------------
         // Instance operations
         //-----------------------------------------------------------------------------------------
@@ -492,10 +496,10 @@ namespace UE::Mutable::Private
 			LIST_UINT32_ADDRESS ImageParameterAddress;
 		};
 
-		struct MeshParameterArgs
-		{
-			ADDRESS variable;
-			uint8 LOD = 0;
+    	struct MeshSkeletalMeshBreakArgs
+    	{
+    		ADDRESS SkeletalMeshParameter;
+    		uint8 LOD = 0;
 			uint8 Section = 0;
 			uint32 MeshID = 0;
 		};
@@ -1204,6 +1208,7 @@ namespace UE::Mutable::Private
 		Matrix,
 		Material,
 		InstancedStruct,
+		SkeletalMesh,
 
 		// Supporting data types : Never returned as an actual data type for any operation.
 		Shape,
diff --git a/Engine/Plugins/Mutable/Source/MutableRuntime/Internal/MuR/Parameters.h b/Engine/Plugins/Mutable/Source/MutableRuntime/Internal/MuR/Parameters.h
--- a/Engine/Plugins/Mutable/Source/MutableRuntime/Internal/MuR/Parameters.h
+++ b/Engine/Plugins/Mutable/Source/MutableRuntime/Internal/MuR/Parameters.h
@@ -48,10 +48,10 @@ namespace UE::Mutable::Private
         Projector,
 
 		/** An externally provided image. */
-		Image,
+		Texture,
 
 		/** An externally provided mesh. */
-		Mesh,
+		SkeletalMesh,
 		
 		/** An externally provided material*/
     	Material,
@@ -316,24 +316,23 @@ namespace UE::Mutable::Private
 								float ProjectionAngle,
 								const FRangeIndex*  RangePosition=nullptr );
 
-        //! Return the value of an image parameter.
-        //! \pre The parameter specified by index is a T_IMAGE.
-        //! \param index Index of the parameter from 0 to GetCount()-1
-		//! \param pos Only for multidimensional parameters: relevant position to set in the ranges
-	    //! \return The externalId specified when setting the image value (\see SetImageValue)
-		UE_API UTexture* GetImageValue( int32 Index, const FRangeIndex* Pos = nullptr) const;
+        //! Return the value of a Texture parameter.
+        //! \pre The parameter specified by index is a Texture.
+        //! \param Index Index of the parameter from 0 to GetCount()-1
+		//! \param Pos Only for multidimensional parameters: relevant position to set in the ranges
+	    //! \return The UTexture specified when setting the image value (\see SetTextureValue)
+		UE_API UTexture* GetTextureValue(int32 Index, const FRangeIndex* Pos = nullptr) const;
 
 		/** */
-		UE_API USkeletalMesh* GetMeshValue(int32 Index, const FRangeIndex* Pos = nullptr) const;
+		UE_API USkeletalMesh* GetSkeletalMeshValue(int32 Index, const FRangeIndex* Pos = nullptr) const;
 
-        //! If the parameter is of the image type, set its value.
-        //! \param index Index of the parameter from 0 to GetCount()-1
-        //! \param externalId Application-specific id used to identify this image during replication.
-		//! \param pos Only for multidimensional parameters: relevant position to set in the ranges
-		UE_API void SetImageValue( int32 Index, UTexture* externalId, const FRangeIndex* = nullptr);
+        //! If the parameter is of the Texture type, set its value.
+        //! \param Index Index of the parameter from 0 to GetCount()-1
+        //! \param Pos Only for multidimensional parameters: relevant position to set in the ranges
+		UE_API void SetTextureValue(int32 Index, UTexture* Value, const FRangeIndex* Pos = nullptr);
 
 		/** */
-		UE_API void SetMeshValue(int32 Index, USkeletalMesh* externalId, const FRangeIndex* = nullptr);
+		UE_API void SetSkeletalMeshValue(int32 Index, USkeletalMesh* Value, const FRangeIndex* = nullptr);
 
         //! Return the value of a float parameter.
         //! \pre The parameter specified by index is a T_FLOAT.
diff --git a/Engine/Plugins/Mutable/Source/MutableRuntime/Internal/MuR/SkeletalMesh.h b/Engine/Plugins/Mutable/Source/MutableRuntime/Internal/MuR/SkeletalMesh.h
new file mode 100644
--- /dev/null
+++ b/Engine/Plugins/Mutable/Source/MutableRuntime/Internal/MuR/SkeletalMesh.h
@@ -0,0 +1,23 @@
+// Copyright Epic Games, Inc. All Rights Reserved.
+
+#pragma once
+
+#include "RefCounted.h"
+
+#include "PassthroughObject.h"
+
+class USkeletalMesh;
+
+
+namespace UE::Mutable::Private
+{
+	class FSkeletalMesh : public FResource
+	{
+	public:
+		// FResource interface
+		virtual int32 GetDataSize() const override { return 0; }
+		
+    	TPassthroughObjectPtr<USkeletalMesh> PassthroughObject;
+	};
+}
+
diff --git a/Engine/Plugins/Mutable/Source/MutableRuntime/Internal/MuR/SystemPrivate.h b/Engine/Plugins/Mutable/Source/MutableRuntime/Internal/MuR/SystemPrivate.h
--- a/Engine/Plugins/Mutable/Source/MutableRuntime/Internal/MuR/SystemPrivate.h
+++ b/Engine/Plugins/Mutable/Source/MutableRuntime/Internal/MuR/SystemPrivate.h
@@ -2,6 +2,7 @@
 
 #pragma once
 
+#include "SkeletalMesh.h"
 #include "MuR/MemoryCounters.h"
 #include "MuR/System.h"
 
@@ -498,6 +499,7 @@ namespace UE::Mutable::Private
 		TMemoryTrackedArray<FVector4f> ColorResults;
 		TMemoryTrackedArray<TResourceResult<FImage>> ImageResults;
 		TMemoryTrackedArray<TResourceResult<FMesh>> MeshResults;
+		TMemoryTrackedArray<TSharedPtr<const FSkeletalMesh>> SkeletalMeshResults;
 		TMemoryTrackedArray<TSharedPtr<const FLayout>> LayoutResults;
 		TMemoryTrackedArray<TSharedPtr<const FInstance>> InstanceResults;
 		TMemoryTrackedArray<FProjector> ProjectorResults;
@@ -559,6 +561,7 @@ namespace UE::Mutable::Private
 				ImageResults.Emplace();
 				LayoutResults.Add(nullptr);
 				MeshResults.Emplace();
+				SkeletalMeshResults.Emplace();
 				InstanceResults.Add(nullptr);
 				ProjectorResults.Add(FProjector());
 				StringResults.Add(nullptr);
@@ -834,6 +837,29 @@ namespace UE::Mutable::Private
 		}
 
 
+		TSharedPtr<const FSkeletalMesh> GetSkeletalMesh(FCacheAddress at)
+		{
+			if (!at.At) return nullptr;
+			if (at.At >= OpExecutionData.size_code()) return nullptr;
+			FOpExecutionData* Data = OpExecutionData.get_ptr(at);
+			if (!Data) return nullptr;
+
+			check(Data->DataType == EDataType::SkeletalMesh);
+			TSharedPtr<const FSkeletalMesh> Result = SkeletalMeshResults[Data->DataTypeIndex];
+
+			// We need to decrease the hit-count even if the result is null.
+			check(Data->OpHitCount > 0);
+
+			--Data->OpHitCount;
+			if (Data->OpHitCount == 0 && !Data->IsCacheLocked)
+			{
+				SetUnused(*Data);
+			}
+
+			return Result;
+		}
+		
+
 		TSharedPtr<const FMaterial> GetMaterial(FCacheAddress at, bool& bIsLastReference)
 		{
 			bIsLastReference = false;
@@ -1131,6 +1157,28 @@ namespace UE::Mutable::Private
 			UE::Mutable::Private::UpdateLLMStats();
 		}
 
+		void SetSkeletalMesh(FCacheAddress At, TSharedPtr<const FSkeletalMesh> V)
+		{
+			check(At.At < OpExecutionData.size_code());
+			FOpExecutionData* Data = OpExecutionData.get_ptr(At);
+			check(Data->DataType == EDataType::SkeletalMesh || Data->DataType == EDataType::None);
+			Data->DataType = EDataType::SkeletalMesh;
+			Data->IsValueValid = true;
+
+			if (!Data->DataTypeIndex)
+			{
+				Data->DataTypeIndex = SkeletalMeshResults.Num();
+				SkeletalMeshResults.Add(V);
+			}
+			else
+			{
+				SkeletalMeshResults[Data->DataTypeIndex] = V;
+			}
+			check(Data->DataTypeIndex != 0);
+
+			UE::Mutable::Private::UpdateLLMStats();
+		}
+
 		void SetLayout(FCacheAddress at, TSharedPtr<const FLayout> v)
 		{
 			check(at.At < OpExecutionData.size_code());
diff --git a/Engine/Plugins/Mutable/Source/MutableRuntime/Private/MuR/CodeRunner.cpp b/Engine/Plugins/Mutable/Source/MutableRuntime/Private/MuR/CodeRunner.cpp
--- a/Engine/Plugins/Mutable/Source/MutableRuntime/Private/MuR/CodeRunner.cpp
+++ b/Engine/Plugins/Mutable/Source/MutableRuntime/Private/MuR/CodeRunner.cpp
@@ -3457,6 +3457,28 @@ namespace UE::Mutable::Private
 			break;
 		}
 
+        case EOpType::ME_SKELETALMESH_BREAK:
+		{
+			OP::MeshSkeletalMeshBreakArgs Args = Program.GetOpArgs<OP::MeshSkeletalMeshBreakArgs>(Item.At);
+			switch (Item.Stage)
+			{
+			case 0:
+				AddOp(FScheduledOp(Item.At, Item, 1), FScheduledOp(Args.SkeletalMeshParameter, Item));
+				break;
+
+			case 1:
+				{
+					// This has been moved to a task. It should have been intercepted in IssueOp.
+					check(false);
+				}
+
+			default:
+				check(false);
+			}
+
+			break;
+		}
+        	
 		default:
 			if (type!=EOpType::NONE)
 			{
@@ -4683,7 +4705,7 @@ namespace UE::Mutable::Private
         {
 			OP::ParameterArgs Args = Program.GetOpArgs<OP::ParameterArgs>(Item.At);
 			TSharedPtr<FRangeIndex> Index = BuildCurrentOpRangeIndex(Item, Args.variable);
-			UTexture* Image = pParams->GetImageValue(Args.variable, Index.Get());
+			UTexture* Image = pParams->GetTextureValue(Args.variable, Index.Get());
 
        		TPassthroughObjectPtr<UTexture> ExternalObject(Image);
 			TSharedPtr<FImage> Result = FImage::CreateAsReference(ExternalObject, FImageDesc(), false);
@@ -6985,6 +7007,19 @@ namespace UE::Mutable::Private
 
 				break;
 			}
+
+		case EOpType::SK_PARAMETER:
+			{
+				OP::ParameterArgs Args = Program.GetOpArgs<OP::ParameterArgs>(Item.At);
+				TSharedPtr<FRangeIndex> Index = BuildCurrentOpRangeIndex(Item, Args.variable);
+				USkeletalMesh* SkeletalMesh = pParams->GetSkeletalMeshValue(Args.variable, Index.Get());
+
+				TSharedPtr<FSkeletalMesh> Result = MakeShared<FSkeletalMesh>();
+				Result->PassthroughObject = TPassthroughObjectPtr<USkeletalMesh>(SkeletalMesh);
+        		
+				StoreSkeletalMesh(Item, Result);
+				break;
+			}
 			
 		default:
 		{
@@ -7161,7 +7196,7 @@ namespace UE::Mutable::Private
 			check(Item.Stage == 0);
 
 			OP::ParameterArgs Args = Program.GetOpArgs<OP::ParameterArgs>(Item.At);
-			UTexture* Image = pParams->GetImageValue(Args.variable);
+			UTexture* Image = pParams->GetTextureValue(Args.variable);
 
 			FExtendedImageDesc& Result = ImageDescResults.FindOrAdd(Item.At);
 			Result = GetExternalImageDesc(Image);
diff --git a/Engine/Plugins/Mutable/Source/MutableRuntime/Private/MuR/CodeRunner.h b/Engine/Plugins/Mutable/Source/MutableRuntime/Private/MuR/CodeRunner.h
--- a/Engine/Plugins/Mutable/Source/MutableRuntime/Private/MuR/CodeRunner.h
+++ b/Engine/Plugins/Mutable/Source/MutableRuntime/Private/MuR/CodeRunner.h
@@ -488,6 +488,11 @@ namespace  UE::Mutable::Private
 			return System->WorkingMemoryManager.LoadMesh(From);
 		}
 
+    	inline TSharedPtr<const FSkeletalMesh> LoadSkeletalMesh(const FCacheAddress& From)
+		{
+			return System->WorkingMemoryManager.CurrentInstanceCache->GetSkeletalMesh(From);
+		}
+    	
 		inline TSharedPtr<const FImage> LoadImage(const FCacheAddress& From)
 		{
 			return System->WorkingMemoryManager.LoadImage(From);
@@ -563,6 +568,11 @@ namespace  UE::Mutable::Private
 			System->WorkingMemoryManager.StoreMesh(To, Resource);
 		}
 
+    	inline void StoreSkeletalMesh(const FCacheAddress& To, TSharedPtr<const FSkeletalMesh> Resource)
+		{
+			System->WorkingMemoryManager.CurrentInstanceCache->SetSkeletalMesh(To, Resource);
+		}
+    	
 		inline void StoreImage(const FCacheAddress& To, TSharedPtr<const FImage> Resource)
 		{
 			System->WorkingMemoryManager.StoreImage(To, Resource);
diff --git a/Engine/Plugins/Mutable/Source/MutableRuntime/Private/MuR/CodeRunnerBegin.cpp b/Engine/Plugins/Mutable/Source/MutableRuntime/Private/MuR/CodeRunnerBegin.cpp
--- a/Engine/Plugins/Mutable/Source/MutableRuntime/Private/MuR/CodeRunnerBegin.cpp
+++ b/Engine/Plugins/Mutable/Source/MutableRuntime/Private/MuR/CodeRunnerBegin.cpp
@@ -369,6 +369,7 @@ namespace UE::Mutable::Private
 				case EDataType::Layout:		
 				case EDataType::Material:		
 				case EDataType::InstancedStruct:		
+				case EDataType::SkeletalMesh:		
 					break;
 
 				default:
@@ -388,7 +389,7 @@ namespace UE::Mutable::Private
 			case EOpType::CO_PARAMETER:
 			case EOpType::MA_PARAMETER:
 			case EOpType::PR_PARAMETER:
-			case EOpType::ME_PARAMETER:
+			case EOpType::SK_PARAMETER:
 			case EOpType::IM_PARAMETER:
 			case EOpType::MI_PARAMETER:
 			case EOpType::IM_PARAMETER_FROM_MATERIAL:
@@ -1986,6 +1987,16 @@ namespace UE::Mutable::Private
 					StoreNone(Item);
 					break;
 				}
+
+			case EOpType::ME_SKELETALMESH_BREAK:
+				{
+					OP::MeshSkeletalMeshBreakArgs Args = Program.GetOpArgs<OP::MeshSkeletalMeshBreakArgs>(Item.At);
+					PushOp(FScheduledOpInline(Args.SkeletalMeshParameter, Item));
+
+					StoreNone(Item);
+					break;
+				}
+			
 			
 			case EOpType::IM_REFERENCE:
 				{
diff --git a/Engine/Plugins/Mutable/Source/MutableRuntime/Private/MuR/CodeRunnerTasked.cpp b/Engine/Plugins/Mutable/Source/MutableRuntime/Private/MuR/CodeRunnerTasked.cpp
--- a/Engine/Plugins/Mutable/Source/MutableRuntime/Private/MuR/CodeRunnerTasked.cpp
+++ b/Engine/Plugins/Mutable/Source/MutableRuntime/Private/MuR/CodeRunnerTasked.cpp
@@ -41,6 +41,7 @@
 #include "Trace/Detail/Channel.h"
 #include "ProfilingDebugging/CountersTrace.h"
 #include "Async/Fundamental/Scheduler.h"
+#include "MuR/SkeletalMesh.h"
 
 #if MUTABLE_DEBUG_CODERUNNER_TASK_SCHEDULE_CALLSTACK
 #include "GenericPlatform/GenericPlatformStackWalk.h"
@@ -2667,19 +2668,21 @@ namespace UE::Mutable::Private
 			break;
 		}
 
-		case EOpType::ME_PARAMETER:
+		case EOpType::ME_SKELETALMESH_BREAK:
 		{
-			OP::MeshParameterArgs Args = Program.GetOpArgs<OP::MeshParameterArgs>(Item.At);
-			TSharedPtr<FRangeIndex> Index = BuildCurrentOpRangeIndex(Item, Args.variable);
+			if (Item.Stage == 1)
+			{
+				OP::MeshSkeletalMeshBreakArgs Args = Program.GetOpArgs<OP::MeshSkeletalMeshBreakArgs>(Item.At);
 
-			USkeletalMesh* Id = Params->GetMeshValue(Args.variable, Index.Get());
+				TSharedPtr<const FSkeletalMesh> SkeletalMesh = LoadSkeletalMesh(FCacheAddress(Args.SkeletalMeshParameter, Item));
 
-			int32 LODIndex = Args.LOD;
-			int32 SectionIndex = Args.Section;
+				int32 LODIndex = Args.LOD;
+				int32 SectionIndex = Args.Section;
 
-			CodeRunner::FExternalResourceId FullId;
-			FullId.MeshParameter = Id;
-			Issued = MakeShared<FMeshExternalLoadTask>(Item, FullId, LODIndex, SectionIndex, Args.MeshID);
+				CodeRunner::FExternalResourceId FullId;
+				FullId.MeshParameter = SkeletalMesh->PassthroughObject.Get();
+				Issued = MakeShared<FMeshExternalLoadTask>(Item, FullId, LODIndex, SectionIndex, Args.MeshID);
+			}
 
 			break;
 		}
diff --git a/Engine/Plugins/Mutable/Source/MutableRuntime/Private/MuR/Model.cpp b/Engine/Plugins/Mutable/Source/MutableRuntime/Private/MuR/Model.cpp
--- a/Engine/Plugins/Mutable/Source/MutableRuntime/Private/MuR/Model.cpp
+++ b/Engine/Plugins/Mutable/Source/MutableRuntime/Private/MuR/Model.cpp
@@ -475,20 +475,20 @@ namespace UE::Mutable::Private
                             break;
                         }
 
-                        case EParameterType::Image:
-                            pRes->SetImageValue( thisP, OldParameters->GetImageValue(p) );
+                        case EParameterType::Texture:
+                            pRes->SetTextureValue(thisP, OldParameters->GetTextureValue(p));
                             break;
                         	
-                        case EParameterType::Mesh:
-                        	pRes->SetMeshValue( thisP, OldParameters->GetMeshValue(p) );
+                        case EParameterType::SkeletalMesh:
+                        	pRes->SetSkeletalMeshValue(thisP, OldParameters->GetSkeletalMeshValue(p));
                         	break;
                         	
 	        			case EParameterType::Material:
-	        				pRes->SetMaterialValue( thisP, OldParameters->GetMaterialValue(p) );
+	        				pRes->SetMaterialValue(thisP, OldParameters->GetMaterialValue(p));
 	        				break;
                         	
 	        			case EParameterType::InstancedStruct:
-	        				pRes->SetInstancedStructValue( thisP, OldParameters->GetInstancedStructValue(p) );
+	        				pRes->SetInstancedStructValue(thisP, OldParameters->GetInstancedStructValue(p));
 	        				break;
 
                         default:
diff --git a/Engine/Plugins/Mutable/Source/MutableRuntime/Private/MuR/Operations.cpp b/Engine/Plugins/Mutable/Source/MutableRuntime/Private/MuR/Operations.cpp
--- a/Engine/Plugins/Mutable/Source/MutableRuntime/Private/MuR/Operations.cpp
+++ b/Engine/Plugins/Mutable/Source/MutableRuntime/Private/MuR/Operations.cpp
@@ -38,7 +38,7 @@ namespace UE::Mutable::Private
 		{ EDataType::Color,			},	// CO_PARAMETER
 		{ EDataType::Projector,		},	// PR_PARAMETER
 		{ EDataType::Image,			},	// IM_PARAMETER
-		{ EDataType::Mesh,			},	// ME_PARAMETER
+		{ EDataType::SkeletalMesh,	},	// SK_PARAMETER
 		{ EDataType::String,			},	// ST_PARAMETER
 		{ EDataType::Matrix,			},	// MA_PARAMETER
 		{ EDataType::Material,		},	// MI_PARAMETER
@@ -142,6 +142,7 @@ namespace UE::Mutable::Private
 		{ EDataType::Mesh,			},	// ME_TRANSFORMWITHMESH
 		{ EDataType::Mesh,			},	// ME_TRANSFORMWITHBONE
 		{ EDataType::Mesh,			},	// ME_EXTERNAL
+		{ EDataType::Mesh,			},	// ME_SKELETALMESH_BREAK
 
 		{ EDataType::Instance,		},	// IN_ADDMESH
 		{ EDataType::Instance,		},	// IN_ADDIMAGE
@@ -200,7 +201,7 @@ namespace UE::Mutable::Private
         case EOpType::CO_PARAMETER:
         case EOpType::PR_PARAMETER:
 		case EOpType::IM_PARAMETER:
-		case EOpType::ME_PARAMETER:
+		case EOpType::SK_PARAMETER:
 		case EOpType::MA_PARAMETER:
 		case EOpType::MI_PARAMETER:
 		case EOpType::IS_PARAMETER:
@@ -838,7 +839,14 @@ namespace UE::Mutable::Private
 			f(args.Matrix);
 			break;
 		}
-        	
+
+        case EOpType::ME_SKELETALMESH_BREAK:
+        {
+        	OP::MeshSkeletalMeshBreakArgs args = program.GetOpArgs<OP::MeshSkeletalMeshBreakArgs>(at);
+        	f(args.SkeletalMeshParameter);
+        	break;
+        }
+        
         //-------------------------------------------------------------------------------------
         case EOpType::IN_ADDMESH:
         case EOpType::IN_ADDIMAGE:
diff --git a/Engine/Plugins/Mutable/Source/MutableRuntime/Private/MuR/Parameters.cpp b/Engine/Plugins/Mutable/Source/MutableRuntime/Private/MuR/Parameters.cpp
--- a/Engine/Plugins/Mutable/Source/MutableRuntime/Private/MuR/Parameters.cpp
+++ b/Engine/Plugins/Mutable/Source/MutableRuntime/Private/MuR/Parameters.cpp
@@ -593,17 +593,17 @@ namespace UE::Mutable::Private
     }
 
 
-	USkeletalMesh* FParameters::GetMeshValue(int32 Index, const FRangeIndex* Pos) const
+	USkeletalMesh* FParameters::GetSkeletalMeshValue(int32 Index, const FRangeIndex* Pos) const
 	{
 		check(Index >= 0 && Index < (int32)m_pD->Values.Num());
-		check(GetType(Index) == EParameterType::Mesh);
+		check(GetType(Index) == EParameterType::SkeletalMesh);
 
 		// Early out in case of invalid parameters
 		if (Index < 0
 			||
 			Index >= (int32)m_pD->Values.Num()
 			||
-			GetType(Index) != EParameterType::Mesh)
+			GetType(Index) != EParameterType::SkeletalMesh)
 		{
 			return {};
 		}
@@ -629,14 +629,13 @@ namespace UE::Mutable::Private
 
 		// Multivalue parameter, but no multivalue set. Return single value.
 		return m_pD->Values[Index].Get<FParamSkeletalMeshType>().Get();
-
 	}
 
 
-	void FParameters::SetMeshValue(int32 Index, USkeletalMesh* Value, const FRangeIndex* Pos)
+	void FParameters::SetSkeletalMeshValue(int32 Index, USkeletalMesh* Value, const FRangeIndex* Pos)
 	{
 		check(Index >= 0 && Index < (int32)m_pD->Values.Num());
-		check(GetType(Index) == EParameterType::Mesh);
+		check(GetType(Index) == EParameterType::SkeletalMesh);
 
 		// Single value case
 		if (!Pos)
@@ -661,17 +660,17 @@ namespace UE::Mutable::Private
 	}
 
 
-	UTexture* FParameters::GetImageValue( int32 Index, const FRangeIndex* Pos ) const
+	UTexture* FParameters::GetTextureValue( int32 Index, const FRangeIndex* Pos ) const
     {
         check( Index >= 0 && Index < (int32)m_pD->Values.Num() );
-        check( GetType( Index ) == EParameterType::Image );
+        check( GetType( Index ) == EParameterType::Texture );
 
 		// Early out in case of invalid parameters
         if ( Index < 0
              ||
              Index >= (int32)m_pD->Values.Num()
              ||
-             GetType(Index) != EParameterType::Image )
+             GetType(Index) != EParameterType::Texture )
         {
             return {};
         }
@@ -700,10 +699,10 @@ namespace UE::Mutable::Private
     }
 
 
-    void FParameters::SetImageValue( int32 Index, UTexture* Value, const FRangeIndex* Pos )
+    void FParameters::SetTextureValue( int32 Index, UTexture* Value, const FRangeIndex* Pos )
     {
         check( Index >= 0 && Index < (int32)m_pD->Values.Num() );
-        check( GetType( Index ) == EParameterType::Image );
+        check( GetType( Index ) == EParameterType::Texture );
 
 		// Single value case
 		if (!Pos)
diff --git a/Engine/Plugins/Mutable/Source/MutableRuntime/Private/MuR/System.cpp b/Engine/Plugins/Mutable/Source/MutableRuntime/Private/MuR/System.cpp
--- a/Engine/Plugins/Mutable/Source/MutableRuntime/Private/MuR/System.cpp
+++ b/Engine/Plugins/Mutable/Source/MutableRuntime/Private/MuR/System.cpp
@@ -943,6 +943,7 @@ namespace UE::Mutable::Private
             case EOpType::CO_PARAMETER:
             case EOpType::PR_PARAMETER:
 			case EOpType::IM_PARAMETER:
+			case EOpType::SK_PARAMETER:
 			case EOpType::MA_PARAMETER:
 			case EOpType::MI_PARAMETER:
 			case EOpType::IS_PARAMETER:
@@ -952,16 +953,8 @@ namespace UE::Mutable::Private
                 Flags[ParamIndex] = true;
                 break;
             }
-			case EOpType::ME_PARAMETER:
-			{
-				OP::MeshParameterArgs args = Program.GetOpArgs<OP::MeshParameterArgs>(at);
-				OP::ADDRESS ParamIndex = args.variable;
-				Flags[ParamIndex] = true;
-				break;
-			}
             default:
                 break;
-
             }
 
             return UniqueDiscreteCoveredCodeVisitor<>::Visit( at, Program );
@@ -2142,9 +2135,31 @@ namespace UE::Mutable::Private
 				break;
 			}
 
-			case EParameterType::Image:
+			case EParameterType::Texture:
+			{
+				TStrongObjectPtr<UTexture> Object = Params->GetPrivate()->Values[ParamIndex].Get<FParamTextureType>();
+				FString Path = Object ? Object->GetPathName() : FString();
+				Blob.Serialize(GetData(Path), Path.Len() * sizeof(FString::ElementType));
+
+				// Multi-values
+				if (Params->GetPrivate()->HasMultipleValues(ParamIndex))
+				{
+					const TMap< TArray<int32>, FParameterValue >& Multi = Params->GetPrivate()->MultiValues[ParamIndex];
+					AddMultiValueKeys(Blob, Multi);
+					for (const TPair< TArray<int32>, FParameterValue >& v : Multi)
+					{
+						Object = v.Value.Get<FParamTextureType>();
+						Path = Object ? Object->GetPathName() : FString();
+						Blob.Serialize(GetData(Path), Path.Len() * sizeof(FString::ElementType));
+					}
+				}
+				break;
+			}
+
+			case EParameterType::SkeletalMesh:
 			{
-				FString Path = Params->GetPrivate()->Values[ParamIndex].Get<FParamTextureType>()->GetPathName();
+				TStrongObjectPtr<USkeletalMesh> Object = Params->GetPrivate()->Values[ParamIndex].Get<FParamSkeletalMeshType>();
+				FString Path = Object ? Object->GetPathName() : FString();
 				Blob.Serialize(GetData(Path), Path.Len() * sizeof(FString::ElementType));
 
 				// Multi-values
@@ -2154,16 +2169,18 @@ namespace UE::Mutable::Private
 					AddMultiValueKeys(Blob, Multi);
 					for (const TPair< TArray<int32>, FParameterValue >& v : Multi)
 					{
-						Path = v.Value.Get<FParamTextureType>()->GetPathName();
+						Object = v.Value.Get<FParamSkeletalMeshType>();
+						Path = Object ? Object->GetPathName() : FString();
 						Blob.Serialize(GetData(Path), Path.Len() * sizeof(FString::ElementType));
 					}
 				}
 				break;
 			}
 
-			case EParameterType::Mesh:
+			case EParameterType::Material:
 			{
-				FString Path = Params->GetPrivate()->Values[ParamIndex].Get<FParamSkeletalMeshType>()->GetPathName();
+				TStrongObjectPtr<UMaterialInterface> Object = Params->GetPrivate()->Values[ParamIndex].Get<FParamMaterialType>();
+				FString Path = Object ? Object->GetPathName() : FString();
 				Blob.Serialize(GetData(Path), Path.Len() * sizeof(FString::ElementType));
 
 				// Multi-values
@@ -2173,7 +2190,8 @@ namespace UE::Mutable::Private
 					AddMultiValueKeys(Blob, Multi);
 					for (const TPair< TArray<int32>, FParameterValue >& v : Multi)
 					{
-						Path = v.Value.Get<FParamSkeletalMeshType>()->GetPathName();
+						Object = v.Value.Get<FParamMaterialType>();
+						Path = Object ? Object->GetPathName() : FString();
 						Blob.Serialize(GetData(Path), Path.Len() * sizeof(FString::ElementType));
 					}
 				}
diff --git a/Engine/Plugins/Mutable/Source/MutableTools/Internal/MuT/Node.h b/Engine/Plugins/Mutable/Source/MutableTools/Internal/MuT/Node.h
--- a/Engine/Plugins/Mutable/Source/MutableTools/Internal/MuT/Node.h
+++ b/Engine/Plugins/Mutable/Source/MutableTools/Internal/MuT/Node.h
@@ -48,7 +48,7 @@ namespace UE::Mutable::Private
 			MeshVariation,
 			MeshReshape,
 			MeshClipDeform,
-			MeshParameter,
+			MeshSkeletalMeshBreak,
 			MeshExternal,
 
 			Image,
@@ -133,6 +133,9 @@ namespace UE::Mutable::Private
 
 			LOD,
 
+			SkeletalMesh,
+			SkeletalMeshParameter,
+			
 			Component,
 			ComponentNew,
 			ComponentEdit,
diff --git a/Engine/Plugins/Mutable/Source/MutableTools/Internal/MuT/NodeImageParameterConvert.h b/Engine/Plugins/Mutable/Source/MutableTools/Internal/MuT/NodeImageParameterConvert.h
--- a/Engine/Plugins/Mutable/Source/MutableTools/Internal/MuT/NodeImageParameterConvert.h
+++ b/Engine/Plugins/Mutable/Source/MutableTools/Internal/MuT/NodeImageParameterConvert.h
@@ -4,7 +4,7 @@
 
 #include "MuR/Ptr.h"
 #include "MuT/Node.h"
-#include "MuT/NodeImage.h"
+#include "MuT/NodeImageParameter.h"
 
 #define UE_API MUTABLETOOLS_API
 
@@ -14,7 +14,7 @@ namespace UE::Mutable::Private
 	class NodeImageParameterConvert : public NodeImage
 	{
 	public:
-		Ptr<NodeImage> ImageParameter;
+		Ptr<NodeImageParameter> ImageParameter;
 		
 		// Node Interface
 		virtual const FNodeType* GetType() const override { return &StaticType; }
diff --git a/Engine/Plugins/Mutable/Source/MutableTools/Internal/MuT/NodeMeshParameter.h b/Engine/Plugins/Mutable/Source/MutableTools/Internal/MuT/NodeMeshParameter.h
deleted file mode 100644
diff --git a/Engine/Plugins/Mutable/Source/MutableTools/Internal/MuT/NodeMeshSkeletalMeshBreak.h b/Engine/Plugins/Mutable/Source/MutableTools/Internal/MuT/NodeMeshSkeletalMeshBreak.h
new file mode 100644
--- /dev/null
+++ b/Engine/Plugins/Mutable/Source/MutableTools/Internal/MuT/NodeMeshSkeletalMeshBreak.h
@@ -0,0 +1,44 @@
+// Copyright Epic Games, Inc. All Rights Reserved.
+
+#pragma once
+
+#include "NodeSkeletalMeshParameter.h"
+#include "MuR/Ptr.h"
+#include "MuT/Node.h"
+#include "MuT/NodeMesh.h"
+#include "MuT/NodeLayout.h"
+
+#define UE_API MUTABLETOOLS_API
+
+namespace UE::Mutable::Private
+{
+	class NodeMeshSkeletalMeshBreak : public NodeMesh
+	{
+	public:
+		TSharedPtr<FMesh> ReferenceMesh;
+
+		TArray<Ptr<NodeLayout>> Layouts;
+
+		uint8 LODIndex = 0;
+
+		uint8 SectionIndex = 0;
+
+		bool bAutomaticLODs = false;
+
+		Ptr<NodeSkeletalMeshParameter> SkeletalMeshParameter;
+		
+		// Node Interface
+		virtual const FNodeType* GetType() const override { return &StaticType; }
+		static const FNodeType* GetStaticType() { return &StaticType; }
+
+	protected:
+		/** Forbidden. Manage with the Ptr<> template. */
+		virtual ~NodeMeshSkeletalMeshBreak() override {}
+
+	private:
+		static UE_API FNodeType StaticType;
+	};
+
+}
+
+#undef UE_API
diff --git a/Engine/Plugins/Mutable/Source/MutableTools/Internal/MuT/NodeSkeletalMesh.h b/Engine/Plugins/Mutable/Source/MutableTools/Internal/MuT/NodeSkeletalMesh.h
new file mode 100644
--- /dev/null
+++ b/Engine/Plugins/Mutable/Source/MutableTools/Internal/MuT/NodeSkeletalMesh.h
@@ -0,0 +1,29 @@
+// Copyright Epic Games, Inc. All Rights Reserved.
+
+#pragma once
+
+#include "MuT/Node.h"
+
+#define UE_API MUTABLETOOLS_API
+
+
+namespace UE::Mutable::Private
+{
+	class NodeSkeletalMesh : public Node
+	{
+	public:
+		// Node interface
+		virtual const FNodeType* GetType() const override { return &StaticType; }
+		static const FNodeType* GetStaticType() { return &StaticType; }
+	
+	protected:
+		virtual ~NodeSkeletalMesh() override {}
+
+	private:
+		static UE_API FNodeType StaticType;
+	};
+}
+
+
+#undef UE_API
+
diff --git a/Engine/Plugins/Mutable/Source/MutableTools/Internal/MuT/NodeSkeletalMeshParameter.h b/Engine/Plugins/Mutable/Source/MutableTools/Internal/MuT/NodeSkeletalMeshParameter.h
new file mode 100644
--- /dev/null
+++ b/Engine/Plugins/Mutable/Source/MutableTools/Internal/MuT/NodeSkeletalMeshParameter.h
@@ -0,0 +1,37 @@
+// Copyright Epic Games, Inc. All Rights Reserved.
+
+#pragma once
+
+#include "Containers/Map.h"
+#include "HAL/Platform.h"
+#include "MuR/Ptr.h"
+#include "MuT/NodeRange.h"
+#include "MuT/NodeSkeletalMesh.h"
+
+#define UE_API MUTABLETOOLS_API
+
+
+namespace UE::Mutable::Private
+{
+	class NodeSkeletalMeshParameter : public NodeSkeletalMesh
+	{
+	public:
+		FString Name;
+
+		FString UID;
+
+		TArray<Ptr<NodeRange>> Ranges;
+		
+		// Node Interface
+		virtual const FNodeType* GetType() const override { return &StaticType; }
+		static const FNodeType* GetStaticType() { return &StaticType; }
+
+		virtual ~NodeSkeletalMeshParameter() override {}
+
+	private:
+		static UE_API FNodeType StaticType;
+	};
+}
+
+
+#undef UE_API
diff --git a/Engine/Plugins/Mutable/Source/MutableTools/Private/MuT/ASTOPMeshSkeletalMeshBreak.cpp b/Engine/Plugins/Mutable/Source/MutableTools/Private/MuT/ASTOPMeshSkeletalMeshBreak.cpp
new file mode 100644
--- /dev/null
+++ b/Engine/Plugins/Mutable/Source/MutableTools/Private/MuT/ASTOPMeshSkeletalMeshBreak.cpp
@@ -0,0 +1,104 @@
+﻿// Copyright Epic Games, Inc. All Rights Reserved.
+
+#include "MuT/ASTOpMeshSkeletalMeshBreak.h"
+
+#include "MuT/ASTOpBoolEqualIntConst.h"
+#include "MuR/ModelPrivate.h"
+
+
+namespace UE::Mutable::Private
+{
+	ASTOpMeshSkeletalMeshBreak::ASTOpMeshSkeletalMeshBreak() :
+		SkeletalMeshParameter(this)
+	{
+	}
+
+	ASTOpMeshSkeletalMeshBreak::~ASTOpMeshSkeletalMeshBreak()
+	{
+		// Explicit call needed to avoid recursive destruction
+		RemoveChildren();
+	}
+
+
+	void ASTOpMeshSkeletalMeshBreak::ForEachChild(const TFunctionRef<void(ASTChild&)> Func)
+	{
+		Func(SkeletalMeshParameter);
+	}
+
+
+	bool ASTOpMeshSkeletalMeshBreak::IsEqual(const ASTOp& OtherUntyped) const
+	{
+		if (OtherUntyped.GetOpType() == GetOpType())
+		{
+			const ASTOpMeshSkeletalMeshBreak* Other = static_cast<const ASTOpMeshSkeletalMeshBreak*>(&OtherUntyped);
+			return Type == Other->Type &&
+				LODIndex == Other->LODIndex &&
+				SectionIndex == Other->SectionIndex &&
+				MeshID == Other->MeshID &&
+				SkeletalMeshParameter == Other->SkeletalMeshParameter;
+		}
+		
+		return false;
+	}
+
+
+	Ptr<ASTOp> ASTOpMeshSkeletalMeshBreak::Clone(MapChildFuncRef MapChild) const
+	{
+		Ptr<ASTOpMeshSkeletalMeshBreak> New = new ASTOpMeshSkeletalMeshBreak();
+		New->Type = Type;
+		New->LODIndex = LODIndex;
+		New->SectionIndex = SectionIndex;
+		New->MeshID = MeshID;
+		New->SkeletalMeshParameter = ASTChild(New, MapChild(SkeletalMeshParameter.child()));
+
+		return New;
+	}
+
+
+	EOpType ASTOpMeshSkeletalMeshBreak::GetOpType() const 
+	{ 
+		return Type; 
+	}
+
+
+	uint64 ASTOpMeshSkeletalMeshBreak::Hash() const
+	{
+		uint64 res = std::hash<uint64>()(uint64(Type));
+
+		hash_combine(res, SkeletalMeshParameter.child().get());
+		hash_combine(res, LODIndex);
+		hash_combine(res, SectionIndex);
+		hash_combine(res, MeshID);
+
+		return res;
+	}
+
+
+	void ASTOpMeshSkeletalMeshBreak::Assert()
+	{
+		ASTOp::Assert();
+	}
+
+
+	void ASTOpMeshSkeletalMeshBreak::Link(FProgram& Program, FLinkerOptions*)
+	{
+		if (!linkedAddress)
+		{
+			linkedAddress = static_cast<OP::ADDRESS>(Program.OpAddress.Num());
+			Program.OpAddress.Add(static_cast<uint32>(Program.ByteCode.Num()));
+			AppendCode(Program.ByteCode, Type);
+
+			OP::MeshSkeletalMeshBreakArgs Args;
+			FMemory::Memzero(Args);
+
+			check(SkeletalMeshParameter->linkedAddress);
+			Args.SkeletalMeshParameter = SkeletalMeshParameter->linkedAddress;
+
+			Args.LOD = LODIndex;
+			Args.Section = SectionIndex;
+			Args.MeshID = MeshID;
+
+			AppendCode(Program.ByteCode, Args);
+		}
+	}
+}
diff --git a/Engine/Plugins/Mutable/Source/MutableTools/Private/MuT/ASTOPMeshSkeletalMeshBreak.h b/Engine/Plugins/Mutable/Source/MutableTools/Private/MuT/ASTOPMeshSkeletalMeshBreak.h
new file mode 100644
--- /dev/null
+++ b/Engine/Plugins/Mutable/Source/MutableTools/Private/MuT/ASTOPMeshSkeletalMeshBreak.h
@@ -0,0 +1,44 @@
+﻿// Copyright Epic Games, Inc. All Rights Reserved.
+
+#pragma once
+
+#include "HAL/Platform.h"
+#include "MuR/Image.h"
+#include "MuR/Operations.h"
+#include "MuR/ParametersPrivate.h"
+#include "MuR/Ptr.h"
+#include "MuT/AST.h"
+
+
+namespace UE::Mutable::Private
+{
+	struct FProgram;
+
+	class ASTOpMeshSkeletalMeshBreak final : public ASTOp
+	{
+	public:
+		EOpType Type = EOpType::ME_SKELETALMESH_BREAK;
+
+		/** Used by some parameter types (Mesh) to specify which subdata from the actual parameter value should be used in the operation. */
+		int32 LODIndex = 0;
+		int32 SectionIndex = 0;
+		uint32 MeshID = 0;
+
+		ASTChild SkeletalMeshParameter;
+
+		int32 LinkedParameterIndex = -1;
+
+		ASTOpMeshSkeletalMeshBreak();
+		virtual ~ASTOpMeshSkeletalMeshBreak() override;
+
+		// ASTOp interface
+		virtual EOpType GetOpType() const override;
+		virtual uint64 Hash() const override;
+		virtual void ForEachChild(const TFunctionRef<void(ASTChild&)>) override;
+		virtual bool IsEqual(const ASTOp&) const override;
+		virtual Ptr<ASTOp> Clone(MapChildFuncRef ) const override;
+		virtual void Assert() override;
+		virtual void Link(FProgram&, FLinkerOptions*) override;
+	};
+}
+
diff --git a/Engine/Plugins/Mutable/Source/MutableTools/Private/MuT/ASTOpParameter.cpp b/Engine/Plugins/Mutable/Source/MutableTools/Private/MuT/ASTOpParameter.cpp
--- a/Engine/Plugins/Mutable/Source/MutableTools/Private/MuT/ASTOpParameter.cpp
+++ b/Engine/Plugins/Mutable/Source/MutableTools/Private/MuT/ASTOpParameter.cpp
@@ -3,7 +3,6 @@
 #include "MuT/ASTOpParameter.h"
 
 #include "MuT/ASTOpBoolEqualIntConst.h"
-#include "HAL/PlatformMath.h"
 #include "Misc/AssertionMacros.h"
 #include "MuR/ModelPrivate.h"
 #include "MuR/Parameters.h"
@@ -13,8 +12,6 @@
 
 namespace UE::Mutable::Private
 {
-
-
 	ASTOpParameter::~ASTOpParameter()
 	{
 		// Explicit call needed to avoid recursive destruction
@@ -43,9 +40,6 @@ namespace UE::Mutable::Private
 			const ASTOpParameter* Other = static_cast<const ASTOpParameter*>(&OtherUntyped);
 			return Type == Other->Type &&
 				Parameter == Other->Parameter &&
-				LODIndex == Other->LODIndex &&
-				SectionIndex == Other->SectionIndex &&
-				MeshID == Other->MeshID &&
 				Ranges == Other->Ranges &&
 				ColorParameterNames == Other->ColorParameterNames &&
 				ScalarParameterNames == Other->ScalarParameterNames &&
@@ -60,9 +54,7 @@ namespace UE::Mutable::Private
 		Ptr<ASTOpParameter> n = new ASTOpParameter();
 		n->Type = Type;
 		n->Parameter = Parameter;
-		n->LODIndex = LODIndex;
-		n->SectionIndex = SectionIndex;
-		n->MeshID = MeshID;
+
 		for (const FRangeData& c : Ranges)
 		{
 			n->Ranges.Emplace(n.get(), MapChild(c.rangeSize.child()), c.rangeName, c.rangeUID);
@@ -90,8 +82,6 @@ namespace UE::Mutable::Private
 	{
 		uint64 res = std::hash<uint64>()(uint64(Type));
 		hash_combine(res, Parameter.Type);
-		hash_combine(res, LODIndex);
-		hash_combine(res, SectionIndex);
 		hash_combine(res, Parameter.Name.Len());
 
 		return res;
@@ -119,28 +109,8 @@ namespace UE::Mutable::Private
 			linkedAddress = (OP::ADDRESS)Program.OpAddress.Num();
 			Program.OpAddress.Add((uint32)Program.ByteCode.Num());
 			AppendCode(Program.ByteCode, Type);
-
-			if (Type == EOpType::ME_PARAMETER)
-			{
-				OP::MeshParameterArgs Args;
-				FMemory::Memzero(Args);
-				Args.variable = (OP::ADDRESS)LinkedParameterIndex;
-				Args.LOD = LODIndex;
-				Args.Section = SectionIndex;
-				Args.MeshID = MeshID;
-
-				for (const FRangeData& d : Ranges)
-				{
-					OP::ADDRESS sizeAt = 0;
-					uint16 rangeId = 0;
-					LinkRange(Program, d, sizeAt, rangeId);
-					Program.Parameters[LinkedParameterIndex].Ranges.Add(rangeId);
-				}
-
-				AppendCode(Program.ByteCode, Args);
-			}
-
-			else if (Type == EOpType::MI_PARAMETER)
+		
+			if (Type == EOpType::MI_PARAMETER)
 			{
 				OP::MaterialParameterArgs Args;
 				FMemory::Memzero(Args);
@@ -250,5 +220,4 @@ namespace UE::Mutable::Private
 	{
 		return false;
 	}
-
 }
diff --git a/Engine/Plugins/Mutable/Source/MutableTools/Private/MuT/ASTOpParameter.h b/Engine/Plugins/Mutable/Source/MutableTools/Private/MuT/ASTOpParameter.h
--- a/Engine/Plugins/Mutable/Source/MutableTools/Private/MuT/ASTOpParameter.h
+++ b/Engine/Plugins/Mutable/Source/MutableTools/Private/MuT/ASTOpParameter.h
@@ -26,11 +26,6 @@ namespace UE::Mutable::Private
 		//!
 		FParameterDesc Parameter;
 
-		/** Used by some parameter types (Mesh) to specify which subdata from the actual parameter value should be used in the operation. */
-		int32 LODIndex = 0;
-		int32 SectionIndex = 0;
-		uint32 MeshID = 0;
-
 		//** Ranges adding dimensions to this parameter. */
 		TArray<FRangeData> Ranges;
 
@@ -46,9 +41,7 @@ namespace UE::Mutable::Private
 		*/
 		int32 LinkedParameterIndex = -1;
 
-	public:
-
-		~ASTOpParameter() override;
+		virtual ~ASTOpParameter() override;
 
 		// ASTOp interface.
 		virtual EOpType GetOpType() const override;
@@ -62,8 +55,6 @@ namespace UE::Mutable::Private
 		virtual FBoolEvalResult EvaluateBool(ASTOpList& Facts, FEvaluateBoolCache* = nullptr) const override;
 		virtual FImageDesc GetImageDesc(bool, FGetImageDescContext*) const override;
 		virtual bool IsColourConstant(FVector4f& OutColour) const override;
-
 	};
-
 }
 
diff --git a/Engine/Plugins/Mutable/Source/MutableTools/Private/MuT/CodeGenerator.cpp b/Engine/Plugins/Mutable/Source/MutableTools/Private/MuT/CodeGenerator.cpp
--- a/Engine/Plugins/Mutable/Source/MutableTools/Private/MuT/CodeGenerator.cpp
+++ b/Engine/Plugins/Mutable/Source/MutableTools/Private/MuT/CodeGenerator.cpp
@@ -2459,7 +2459,7 @@ namespace UE::Mutable::Private
 				}
 				else
 				{
-					check(OriginalMeshNode->GetType() == NodeMeshParameter::GetStaticType())
+					check(OriginalMeshNode->GetType() == NodeMeshSkeletalMeshBreak::GetStaticType())
 					TargetMeshTask = UE::Tasks::MakeCompletedTask<void>();
 				}
 
@@ -2530,7 +2530,7 @@ namespace UE::Mutable::Private
 							}
 							else
 							{
-								check(OriginalMeshCopy->GetType() == NodeMeshParameter::GetStaticType());
+								check(OriginalMeshCopy->GetType() == NodeMeshSkeletalMeshBreak::GetStaticType());
 
 								MorphOp = new ASTOpMeshMorph();
 								
@@ -2651,7 +2651,7 @@ namespace UE::Mutable::Private
 				}
 				else
 				{
-					check(OriginalMeshNode->GetType() == NodeMeshParameter::GetStaticType())
+					check(OriginalMeshNode->GetType() == NodeMeshSkeletalMeshBreak::GetStaticType())
 					TargetMeshTask = UE::Tasks::MakeCompletedTask<void>();
 				}
 
diff --git a/Engine/Plugins/Mutable/Source/MutableTools/Private/MuT/CodeGenerator.h b/Engine/Plugins/Mutable/Source/MutableTools/Private/MuT/CodeGenerator.h
--- a/Engine/Plugins/Mutable/Source/MutableTools/Private/MuT/CodeGenerator.h
+++ b/Engine/Plugins/Mutable/Source/MutableTools/Private/MuT/CodeGenerator.h
@@ -49,6 +49,8 @@ namespace UE::Mutable::Private
 	class NodeImageParameterConvert;
 	class NodeExternalParameter;
 	class NodeExternalOperation;
+	class NodeSkeletalMeshParameter;
+	class NodeSkeletalMesh;
 	class NodeMaterialExternal;
 	class NodeScalarExternal;
 }
@@ -112,7 +114,7 @@ namespace UE::Mutable::Private
 	class NodeMeshSwitch;
 	class NodeMeshTransform;
 	class NodeMeshVariation;
-	class NodeMeshParameter;
+	class NodeMeshSkeletalMeshBreak;
 	class NodeRange;
 	class NodeScalarArithmeticOperation;
 	class NodeScalarConstant;
@@ -832,7 +834,7 @@ namespace UE::Mutable::Private
 		FMeshTask GenerateMesh_Table(const FMeshGenerationStaticOptions&, FMeshOptionsTask, const NodeMeshTable*);
 		FMeshTask GenerateMesh_Reshape(const FMeshGenerationStaticOptions&, FMeshOptionsTask, const NodeMeshReshape*);
 		FMeshTask GenerateMesh_ClipDeform(const FMeshGenerationStaticOptions&, FMeshOptionsTask, const NodeMeshClipDeform*);
-		FMeshTask GenerateMesh_Parameter(const FMeshGenerationStaticOptions&, FMeshOptionsTask, const NodeMeshParameter*);
+		FMeshTask GenerateMesh_SkeletalMeshBreak(const FMeshGenerationStaticOptions&, FMeshOptionsTask, const NodeMeshSkeletalMeshBreak*);
     	FMeshTask GenerateMesh_External(const FMeshGenerationStaticOptions&, FMeshOptionsTask, const NodeMeshExternal*);
 
 		Ptr<ASTOp> GenerateLayoutOpsAndResult(const FMeshGenerationDynamicOptions&, Ptr<ASTOp> LastMeshOp, const TArray<Ptr<NodeLayout>>& OriginalLayouts, uint32 MeshId, FMeshGenerationResult& OutResult);
@@ -841,7 +843,27 @@ namespace UE::Mutable::Private
     	TSharedPtr<const FLayout> GenerateLayout(Ptr<const NodeLayout> SourceLayout, uint32 MeshIDPrefix);
 
     	//-----------------------------------------------------------------------------------------
+    	// Skeletal Mesh
+    	
+    	struct FSkeletalMeshGenerationResult
+    	{
+    		Ptr<ASTOp> Op;
+    	};
+
+    	typedef TMap<FGeneratedCacheKey, FSkeletalMeshGenerationResult> FGeneratedSkeletalMeshesMap;
 
+    	struct FGeneratedSkeletalMeshes
+    	{
+    		UE::FMutex Mutex;
+    		FGeneratedSkeletalMeshesMap Map;
+    	};
+    	FGeneratedSkeletalMeshes GeneratedSkeletalMeshes;
+    	
+    	void GenerateSkeletalMesh(FSkeletalMeshGenerationResult& OutResult, const FGenericGenerationOptions&, const Ptr<const NodeSkeletalMesh>&);
+    	void GenerateSkeletalMesh_Parameter(FSkeletalMeshGenerationResult& OutResult, const FGenericGenerationOptions&, const NodeSkeletalMeshParameter*);
+    	
+    	//-----------------------------------------------------------------------------------------
+    	
     	struct FExternalGenerationOptions 
     	{
     		// TODO Propagated options could be a TVariant
diff --git a/Engine/Plugins/Mutable/Source/MutableTools/Private/MuT/CodeGenerator_FirstPass.h b/Engine/Plugins/Mutable/Source/MutableTools/Private/MuT/CodeGenerator_FirstPass.h
--- a/Engine/Plugins/Mutable/Source/MutableTools/Private/MuT/CodeGenerator_FirstPass.h
+++ b/Engine/Plugins/Mutable/Source/MutableTools/Private/MuT/CodeGenerator_FirstPass.h
@@ -23,7 +23,7 @@
 #include "MuT/NodeSurfaceNew.h"
 #include "MuT/NodeSurfaceVariation.h"
 #include "MuT/NodeSurfaceSwitch.h"
-#include "MuT/NodeMeshParameter.h"
+#include "MuT/NodeMeshSkeletalMeshBreak.h"
 #include "MuT/NodeMaterialParameter.h"
 #include "MuT/ASTOpParameter.h"
 #include "MuT/ErrorLog.h"
@@ -258,9 +258,6 @@ namespace UE::Mutable::Private
 
 			/** Cache for generic parameters. */
 			TMap< Ptr<const Node>, Ptr<ASTOpParameter> > GenericParametersCache;
-
-			/** For mesh parameters we generate a different result for each LOD. */
-			TMap< Ptr<const NodeMeshParameter>, TArray<TPair<Ptr<ASTOpParameter>, FMeshGenerationResult>> > MeshParametersCache;
 		};
 		FSafeParameterNodes ParameterNodes;
 
diff --git a/Engine/Plugins/Mutable/Source/MutableTools/Private/MuT/CodeGenerator_Image.cpp b/Engine/Plugins/Mutable/Source/MutableTools/Private/MuT/CodeGenerator_Image.cpp
--- a/Engine/Plugins/Mutable/Source/MutableTools/Private/MuT/CodeGenerator_Image.cpp
+++ b/Engine/Plugins/Mutable/Source/MutableTools/Private/MuT/CodeGenerator_Image.cpp
@@ -285,7 +285,7 @@ namespace UE::Mutable::Private
 				op->Parameter.Name = Node->Name;
 				bool bParseOk = FGuid::Parse(Node->UID, op->Parameter.UID);
 				check(bParseOk);
-				op->Parameter.Type = EParameterType::Image;
+				op->Parameter.Type = EParameterType::Texture;
 				op->Parameter.DefaultValue.Set<FParamTextureType>(nullptr);
 
 				FirstPass.ParameterNodes.GenericParametersCache.Add(Node, op);
diff --git a/Engine/Plugins/Mutable/Source/MutableTools/Private/MuT/CodeGenerator_Mesh.cpp b/Engine/Plugins/Mutable/Source/MutableTools/Private/MuT/CodeGenerator_Mesh.cpp
--- a/Engine/Plugins/Mutable/Source/MutableTools/Private/MuT/CodeGenerator_Mesh.cpp
+++ b/Engine/Plugins/Mutable/Source/MutableTools/Private/MuT/CodeGenerator_Mesh.cpp
@@ -1,6 +1,7 @@
 // Copyright Epic Games, Inc. All Rights Reserved.
 
 #include "ASTOpExternal.h"
+#include "ASTOpMeshSkeletalMeshBreak.h"
 #include "Containers/Array.h"
 #include "Math/IntPoint.h"
 #include "Math/UnrealMathSSE.h"
@@ -59,7 +60,7 @@
 #include "MuT/NodeMeshTable.h"
 #include "MuT/NodeMeshTransform.h"
 #include "MuT/NodeMeshVariation.h"
-#include "MuT/NodeMeshParameter.h"
+#include "MuT/NodeMeshSkeletalMeshBreak.h"
 #include "MuT/NodeScalar.h"
 #include "MuT/Table.h"
 #include "MuT/TablePrivate.h"
@@ -97,7 +98,7 @@ namespace UE::Mutable::Private
 		case Node::EType::MeshTable: Result = GenerateMesh_Table(StaticOptions, InOptions, static_cast<const NodeMeshTable*>(Node)); break;
 		case Node::EType::MeshReshape: Result = GenerateMesh_Reshape(StaticOptions, InOptions, static_cast<const NodeMeshReshape*>(Node)); break;
 		case Node::EType::MeshClipDeform: Result = GenerateMesh_ClipDeform(StaticOptions, InOptions, static_cast<const NodeMeshClipDeform*>(Node)); break;
-		case Node::EType::MeshParameter: Result = GenerateMesh_Parameter(StaticOptions, InOptions, static_cast<const NodeMeshParameter*>(Node)); break;
+		case Node::EType::MeshSkeletalMeshBreak: Result = GenerateMesh_SkeletalMeshBreak(StaticOptions, InOptions, static_cast<const NodeMeshSkeletalMeshBreak*>(Node)); break;
 		case Node::EType::MeshExternal: Result = GenerateMesh_External(StaticOptions, InOptions, static_cast<const NodeMeshExternal*>(Node)); break;
 		default: check(false);
 		}
@@ -1359,23 +1360,26 @@ namespace UE::Mutable::Private
 	uint32 CodeGenerator::FUniqueMeshIds::EnsureUnique(uint32 Id)
 	{
 		UE::TUniqueLock Lock(Mutex);
-		bool bValid = false;
-		do
+
+		while (true)
 		{
 			bool bAlreadyPresent = false;
 			Map.FindOrAdd(Id, &bAlreadyPresent);
-			bValid = !bAlreadyPresent && Id != 0;
-			if (!bValid)
+
+			const bool bValid = !bAlreadyPresent && Id != 0;
+			if (bValid)
 			{
-				++Id;
+				break;
 			}
-		} while (bValid);
+
+			++Id;
+		}
 
 		return Id;
 	}
 
 
-	FMeshTask CodeGenerator::GenerateMesh_Parameter(const FMeshGenerationStaticOptions& StaticOptions, FMeshOptionsTask InOptions, const NodeMeshParameter* Node)
+	FMeshTask CodeGenerator::GenerateMesh_SkeletalMeshBreak(const FMeshGenerationStaticOptions& StaticOptions, FMeshOptionsTask InOptions, const NodeMeshSkeletalMeshBreak* Node)
 	{
 		TArray<UE::Tasks::FTask, TInlineAllocator<2>> Requisites;
 		Requisites.Emplace(InOptions);
@@ -1402,30 +1406,27 @@ namespace UE::Mutable::Private
 			}
 		}
 
+		FSkeletalMeshGenerationResult SkeletalMeshResult;
+		GenerateSkeletalMesh(SkeletalMeshResult, StaticOptions, Node->SkeletalMeshParameter);
+
+		Ptr<ASTOp> SkeletalMeshParameter = SkeletalMeshResult.Op;
+		
 		// Local pipe because we call GenerateRange in there.
 		FMeshTask MeshParamTask = this->LocalPipe.Launch(TEXT("MutableMeshParameter"),
-			[InOptions, Node, ReferenceMesh, StaticOptions, this]() mutable
+			[InOptions, Node, ReferenceMesh, StaticOptions, SkeletalMeshParameter, this]() mutable
 			{
 				FMeshGenerationDynamicOptions Options = InOptions.GetResult();
 
 				FMeshGenerationResult Result;
-				Ptr<ASTOpParameter> Op;
-
-				FParameterDesc param;
-				param.Name = Node->Name;
-				bool bParseOk = FGuid::Parse(Node->UID, param.UID);
-				check(bParseOk);
-				param.Type = EParameterType::Mesh;
-				param.DefaultValue.Set<FParamSkeletalMeshType>(nullptr);
-
-				Op = new ASTOpParameter();
-				Op->Type = EOpType::ME_PARAMETER;
-				Op->Parameter = param;
+				Ptr<ASTOpMeshSkeletalMeshBreak> Op;
+					
+				Op = new ASTOpMeshSkeletalMeshBreak();
+				Op->SkeletalMeshParameter = SkeletalMeshParameter;
 				Op->LODIndex = Node->LODIndex + (Node->bAutomaticLODs ? StaticOptions.LODIndex : 0);
 				Op->SectionIndex = Node->SectionIndex;
 
 				// Assign an ID prefix to the mesh. 
-				uint32 MeshID = HashCombine(GetTypeHash(Op->LODIndex), GetTypeHash(param.UID));
+				uint32 MeshID = Op->Hash();
 				MeshID = HashCombine(MeshID, GetTypeHash(Op->SectionIndex));
 				MeshID = UniqueMeshIds.EnsureUnique(MeshID);
 				Op->MeshID = MeshID;
@@ -1437,35 +1438,17 @@ namespace UE::Mutable::Private
 
 				{
 					UE::TUniqueLock Lock(FirstPass.ParameterNodes.Mutex);
-
-					TArray<TPair<Ptr<ASTOpParameter>, FMeshGenerationResult>>& ArrayFound = FirstPass.ParameterNodes.MeshParametersCache.FindOrAdd(Node);
-					if (ArrayFound.IsValidIndex(StaticOptions.LODIndex))
-					{
-						return ArrayFound[StaticOptions.LODIndex].Value;
-					}
-
-					ArrayFound.SetNum(StaticOptions.LODIndex+1);
-					ArrayFound[StaticOptions.LODIndex].Key = Op;
-
+					
 					// Fill the result structure. Some ops will be completed outside the lock.
 					Ptr<ASTOp> LastMeshOp = GenerateLayoutOpsAndResult(Options, Op, Node->Layouts, Op->MeshID, Result);
 					Result.MeshOp = LastMeshOp;
 					Result.BaseMeshOp = LastMeshOp;
-					ArrayFound[StaticOptions.LODIndex].Value = Result;
 				}
 
 				int32 LODIndex = StaticOptions.LODIndex;
 
 				FMeshGenerationDynamicOptions DynamicOptions = InOptions.GetResult();
-
-				// Generate the code for the ranges
-				for (int32 RangeIndex = 0; RangeIndex < Node->Ranges.Num(); ++RangeIndex)
-				{
-					FRangeGenerationResult RangeResult;
-					GenerateRange(RangeResult, StaticOptions, Node->Ranges[RangeIndex]);
-					Op->Ranges.Emplace(Op.get(), RangeResult.sizeOp, RangeResult.rangeName, RangeResult.rangeUID);
-				}
-
+					
 				return Result;
 			},
 			Requisites
diff --git a/Engine/Plugins/Mutable/Source/MutableTools/Private/MuT/CodeGenerator_SkeletalMesh.cpp b/Engine/Plugins/Mutable/Source/MutableTools/Private/MuT/CodeGenerator_SkeletalMesh.cpp
new file mode 100644
--- /dev/null
+++ b/Engine/Plugins/Mutable/Source/MutableTools/Private/MuT/CodeGenerator_SkeletalMesh.cpp
@@ -0,0 +1,95 @@
+﻿// Copyright Epic Games, Inc. All Rights Reserved.
+
+#include "CodeGenerator.h"
+
+#include "ASTOpAddLOD.h"
+#include "MuR/SkeletalMesh.h"
+#include "MuT/NodeSkeletalMesh.h"
+#include "MuT/NodeSkeletalMeshParameter.h"
+
+
+namespace UE::Mutable::Private
+{
+    void CodeGenerator::GenerateSkeletalMesh(FSkeletalMeshGenerationResult& Result, const FGenericGenerationOptions& Options, const Ptr<const NodeSkeletalMesh>& Node)
+	{
+    	if (!Node)
+    	{
+    		Result = FSkeletalMeshGenerationResult();
+    		return;
+    	}
+
+    	// See if it was already generated
+    	FGeneratedCacheKey Key;
+    	Key.Node = Node;
+    	Key.Options = Options;
+	    {
+    		UE::TUniqueLock Lock(GeneratedSkeletalMeshes.Mutex);
+    		FGeneratedSkeletalMeshesMap::ValueType* Found = GeneratedSkeletalMeshes.Map.Find(Key);
+    		if (Found)
+    		{
+    			Result = *Found;
+    			return;
+    		}
+	    }
+
+
+	    // Generate for each different type of node
+		switch (Node->GetType()->Type)
+		{
+		case Node::EType::SkeletalMeshParameter:
+			GenerateSkeletalMesh_Parameter(Result, Options, static_cast<const NodeSkeletalMeshParameter*>(Node.get()));
+			break;
+			
+		default:
+			unimplemented();
+		}
+    	
+    	// Cache the result
+	    {
+    		UE::TUniqueLock Lock(GeneratedSkeletalMeshes.Mutex);
+    		GeneratedSkeletalMeshes.Map.Add(Key, Result);
+	    }
+	}
+
+
+	void CodeGenerator::GenerateSkeletalMesh_Parameter(FSkeletalMeshGenerationResult& Result, const FGenericGenerationOptions& Options, const NodeSkeletalMeshParameter* Node)
+	{
+		Ptr<ASTOpParameter> op;
+
+		Ptr<ASTOpParameter>* Found = nullptr;
+		{
+			UE::TUniqueLock Lock(FirstPass.ParameterNodes.Mutex);
+			Found = FirstPass.ParameterNodes.GenericParametersCache.Find(Node);
+			if (!Found)
+			{
+				op = new ASTOpParameter();
+				op->Type = EOpType::SK_PARAMETER;
+
+				op->Parameter.Name = Node->Name;
+				bool bParseOk = FGuid::Parse(Node->UID, op->Parameter.UID);
+				check(bParseOk);
+				op->Parameter.Type = EParameterType::SkeletalMesh;
+				op->Parameter.DefaultValue.Set<FParamSkeletalMeshType>(nullptr);
+
+				FirstPass.ParameterNodes.GenericParametersCache.Add(Node, op);
+			}
+			else
+			{
+				op = *Found;
+			}
+		}
+
+		if (!Found)
+		{
+			// Generate the code for the ranges
+			for (int32 a = 0; a < Node->Ranges.Num(); ++a)
+			{
+				FRangeGenerationResult rangeResult;
+				GenerateRange(rangeResult, Options, Node->Ranges[a]);
+				op->Ranges.Emplace(op.get(), rangeResult.sizeOp, rangeResult.rangeName, rangeResult.rangeUID);
+			}
+		}
+
+		Result.Op = op;
+	}
+}
diff --git a/Engine/Plugins/Mutable/Source/MutableTools/Private/MuT/CodeOptimiser.cpp b/Engine/Plugins/Mutable/Source/MutableTools/Private/MuT/CodeOptimiser.cpp
--- a/Engine/Plugins/Mutable/Source/MutableTools/Private/MuT/CodeOptimiser.cpp
+++ b/Engine/Plugins/Mutable/Source/MutableTools/Private/MuT/CodeOptimiser.cpp
@@ -548,7 +548,7 @@ namespace UE::Mutable::Private
 					case EOpType::CO_PARAMETER:
 					case EOpType::PR_PARAMETER:
 					case EOpType::IM_PARAMETER:
-					case EOpType::ME_PARAMETER:
+					case EOpType::SK_PARAMETER:
 					case EOpType::MA_PARAMETER:
 					case EOpType::MI_PARAMETER:
 					case EOpType::IM_PARAMETER_FROM_MATERIAL:
diff --git a/Engine/Plugins/Mutable/Source/MutableTools/Private/MuT/CodeOptimiser_Parameters.cpp b/Engine/Plugins/Mutable/Source/MutableTools/Private/MuT/CodeOptimiser_Parameters.cpp
--- a/Engine/Plugins/Mutable/Source/MutableTools/Private/MuT/CodeOptimiser_Parameters.cpp
+++ b/Engine/Plugins/Mutable/Source/MutableTools/Private/MuT/CodeOptimiser_Parameters.cpp
@@ -150,7 +150,7 @@ namespace UE::Mutable::Private
                     case EOpType::CO_PARAMETER:
                     case EOpType::PR_PARAMETER:
 					case EOpType::IM_PARAMETER:
-					case EOpType::ME_PARAMETER:
+					case EOpType::SK_PARAMETER:
 					case EOpType::MA_PARAMETER:
 					{
 						const ASTOpParameter* Typed = static_cast<const ASTOpParameter*>(Op.get());
@@ -1580,7 +1580,7 @@ namespace UE::Mutable::Private
                 case EOpType::CO_PARAMETER:
                 case EOpType::PR_PARAMETER:
 				case EOpType::IM_PARAMETER:
-				case EOpType::ME_PARAMETER:
+				case EOpType::SK_PARAMETER:
 				case EOpType::MA_PARAMETER:
 				case EOpType::IS_PARAMETER:
 				{
diff --git a/Engine/Plugins/Mutable/Source/MutableTools/Private/MuT/Compiler.cpp b/Engine/Plugins/Mutable/Source/MutableTools/Private/MuT/Compiler.cpp
--- a/Engine/Plugins/Mutable/Source/MutableTools/Private/MuT/Compiler.cpp
+++ b/Engine/Plugins/Mutable/Source/MutableTools/Private/MuT/Compiler.cpp
@@ -75,7 +75,7 @@ namespace UE::Mutable::Private
 		{ false,	{ 0, 0, 0, 0, 0, 0, 0, 0, 0, 0, 0, 0, 0, 0, 0, 0, 0, 0, 0, 0, 0, 0, 0, 0, 0, 0, 0, 0, 0 }		},	// CO_PARAMETER
 		{ false,	{ 0, 0, 0, 0, 0, 0, 0, 0, 0, 0, 0, 0, 0, 0, 0, 0, 0, 0, 0, 0, 0, 0, 0, 0, 0, 0, 0, 0, 0 }		},	// PR_PARAMETER
 		{ false,	{ 0, 0, 0, 0, 0, 0, 0, 0, 0, 0, 0, 0, 0, 0, 0, 0, 0, 0, 0, 0, 0, 0, 0, 0, 0, 0, 0, 0, 0 }		},	// IM_PARAMETER
-		{ false,	{ 0, 0, 0, 0, 0, 0, 0, 0, 0, 0, 0, 0, 0, 0, 0, 0, 0, 0, 0, 0, 0, 0, 0, 0, 0, 0, 0, 0, 0 }		},	// ME_PARAMETER
+		{ false,	{ 0, 0, 0, 0, 0, 0, 0, 0, 0, 0, 0, 0, 0, 0, 0, 0, 0, 0, 0, 0, 0, 0, 0, 0, 0, 0, 0, 0, 0 }		},	// SK_PARAMETER
 		{ false,	{ 0, 0, 0, 0, 0, 0, 0, 0, 0, 0, 0, 0, 0, 0, 0, 0, 0, 0, 0, 0, 0, 0, 0, 0, 0, 0, 0, 0, 0 }		},	// ST_PARAMETER
 		{ false,	{ 0, 0, 0, 0, 0, 0, 0, 0, 0, 0, 0, 0, 0, 0, 0, 0, 0, 0, 0, 0, 0, 0, 0, 0, 0, 0, 0, 0, 0 }		},	// MA_PARAMETER
 		{ false,	{ 0, 0, 0, 0, 0, 0, 0, 0, 0, 0, 0, 0, 0, 0, 0, 0, 0, 0, 0, 0, 0, 0, 0, 0, 0, 0, 0, 0, 0 }		},	// MI_PARAMETER
@@ -179,6 +179,7 @@ namespace UE::Mutable::Private
 		{ true,		{ 0, 0, 0, 0, 0, 0, 0, 0, 0, 0, 0, 0, 0, 0, 0, 0, 0, 0, 0, 0, 0, 0, 0, 0, 0, 0, 0, 0, 0 }		},	// ME_TRANSFORMWITHMESH
 		{ true,		{ 0, 0, 0, 0, 0, 0, 0, 0, 0, 0, 0, 0, 0, 0, 0, 0, 0, 0, 0, 0, 0, 0, 0, 0, 0, 0, 0, 0, 0 }		},	// ME_TRANSFORMWITHBONE
 		{ true,		{ 0, 0, 0, 0, 0, 0, 0, 0, 0, 0, 0, 0, 0, 0, 0, 0, 0, 0, 0, 0, 0, 0, 0, 0, 0, 0, 0, 0, 0 }		},	// ME_EXTERNAL
+		{ true,		{ 0, 0, 0, 0, 0, 0, 0, 0, 0, 0, 0, 0, 0, 0, 0, 0, 0, 0, 0, 0, 0, 0, 0, 0, 0, 0, 0, 0, 0 }		},	// ME_SKELETALMESH_BREAK
 
 		{ false,	{ 0, 0, 0, 0, 0, 0, 0, 0, 0, 0, 0, 0, 0, 0, 0, 0, 0, 0, 0, 0, 0, 0, 0, 0, 0, 0, 0, 0, 0 }		},	// IN_ADDMESH
 		{ false,	{ 0, 0, 0, 0, 0, 0, 0, 0, 0, 0, 0, 0, 0, 0, 0, 0, 0, 0, 0, 0, 0, 0, 0, 0, 0, 0, 0, 0, 0 }		},	// IN_ADDIMAGE
@@ -480,26 +481,13 @@ namespace UE::Mutable::Private
 
 				// Gather the parameter list from the non-optimized data, so that we have them all even if they are optimized out
 				{
-					const int32 ParameterCount = 
-						Generator.FirstPass.ParameterNodes.GenericParametersCache.Num()
-						+
-						Generator.FirstPass.ParameterNodes.MeshParametersCache.Num();
+					const int32 ParameterCount = Generator.FirstPass.ParameterNodes.GenericParametersCache.Num();
 
 					Parameters.Reserve(ParameterCount);
 					for (const TPair<Ptr<const Node>, Ptr<ASTOpParameter>>& Entry : Generator.FirstPass.ParameterNodes.GenericParametersCache)
 					{
 						Parameters.Add(Entry.Value->Parameter);
 					}
-					for (const TPair<Ptr<const NodeMeshParameter>, TArray<TPair<Ptr<ASTOpParameter>,FMeshGenerationResult>>>& Entry : Generator.FirstPass.ParameterNodes.MeshParametersCache)
-					{
-						for (const TPair<Ptr<ASTOpParameter>, FMeshGenerationResult>& ArrayEntry : Entry.Value)
-						{
-							if (ArrayEntry.Key)
-							{
-								Parameters.AddUnique(ArrayEntry.Key->Parameter);
-							}
-						}
-					}
 				}
 
 				// Sort the parameters as deterministically as possible.
@@ -544,17 +532,6 @@ namespace UE::Mutable::Private
 										MessageContexts.AddUnique(Entry.Key->GetMessageContext());
 									}
 								}
-								for (const TPair<Ptr<const NodeMeshParameter>, TArray<TPair<Ptr<ASTOpParameter>, FMeshGenerationResult>>>& Entry : Generator.FirstPass.ParameterNodes.MeshParametersCache)
-								{
-									for (const TPair<Ptr<ASTOpParameter>, FMeshGenerationResult>& ArrayEntry : Entry.Value)
-									{
-										if (ArrayEntry.Key && ArrayEntry.Key->Parameter.Name.Compare(CurrentParameterName) == 0)
-										{
-											check(Entry.Key->GetMessageContext());
-											MessageContexts.AddUnique(Entry.Key->GetMessageContext());
-										}
-									}
-								}
 							}
 
 							// The "ensure" will fail if:
diff --git a/Engine/Plugins/Mutable/Source/MutableTools/Private/MuT/ErrorLog.cpp b/Engine/Plugins/Mutable/Source/MutableTools/Private/MuT/ErrorLog.cpp
--- a/Engine/Plugins/Mutable/Source/MutableTools/Private/MuT/ErrorLog.cpp
+++ b/Engine/Plugins/Mutable/Source/MutableTools/Private/MuT/ErrorLog.cpp
@@ -220,7 +220,7 @@ namespace UE::Mutable::Private
 		TEXT("CO_PARAMETER     "),
 		TEXT("PR_PARAMETER     "),
 		TEXT("IM_PARAMETER     "),
-		TEXT("ME_PARAMETER     "),
+		TEXT("SK_PARAMETER     "),
 		TEXT("ST_PARAMETER     "),
 		TEXT("MA_PARAMETER     "),
 		TEXT("MI_PARAMETER     "),
@@ -324,6 +324,7 @@ namespace UE::Mutable::Private
 		TEXT("ME_TRANSFORMWITHMESH"),
 		TEXT("ME_TRANSFORMWITHBONE"),
     	TEXT("ME_EXTERNAL      "),
+		TEXT("ME_SKELETALMESH_BREAK"),
 
 		TEXT("IN_ADDMESH       "),
 		TEXT("IN_ADDIMAGE      "),
diff --git a/Engine/Plugins/Mutable/Source/MutableTools/Private/MuT/Node.cpp b/Engine/Plugins/Mutable/Source/MutableTools/Private/MuT/Node.cpp
--- a/Engine/Plugins/Mutable/Source/MutableTools/Private/MuT/Node.cpp
+++ b/Engine/Plugins/Mutable/Source/MutableTools/Private/MuT/Node.cpp
@@ -69,7 +69,7 @@
 #include "MuT/NodeMeshClipMorphPlane.h"
 #include "MuT/NodeMeshClipDeform.h"
 #include "MuT/NodeMeshClipWithMesh.h"
-#include "MuT/NodeMeshParameter.h"
+#include "MuT/NodeMeshSkeletalMeshBreak.h"
 #include "MuT/NodeMeshMakeMorph.h"
 #include "MuT/NodeMeshApplyPose.h"
 #include "MuT/NodeMeshTransform.h"
@@ -113,6 +113,8 @@
 #include "MuT/NodeMaterialParameter.h"
 #include "MuT/NodeMeshExternal.h"
 #include "MuT/NodeScalarExternal.h"
+#include "MuT/NodeSkeletalMesh.h"
+#include "MuT/NodeSkeletalMeshParameter.h"
 
 
 namespace UE::Mutable::Private
@@ -158,6 +160,10 @@ namespace UE::Mutable::Private
 	FNodeType NodeExternalParameter::StaticType = FNodeType(Node::EType::ExternalParameter, NodeExternal::GetStaticType());
 	
 	FNodeType NodeLOD::StaticType = FNodeType(Node::EType::LOD, Node::GetStaticType());
+
+	FNodeType NodeSkeletalMesh::StaticType = FNodeType(Node::EType::SkeletalMesh, Node::GetStaticType());
+	FNodeType NodeSkeletalMeshParameter::StaticType = FNodeType(Node::EType::SkeletalMeshParameter, NodeSkeletalMesh::GetStaticType());
+
 	FNodeType NodeExtensionData::StaticType = FNodeType(Node::EType::ExtensionData, Node::GetStaticType());
 	FNodeType NodeExtensionDataConstant::StaticType = FNodeType(Node::EType::ExtensionDataConstant, NodeExtensionData::GetStaticType());
 
@@ -207,7 +213,7 @@ namespace UE::Mutable::Private
 	FNodeType NodeMeshClipMorphPlane::StaticType = FNodeType(Node::EType::MeshClipMorphPlane, NodeMesh::GetStaticType());
 	FNodeType NodeMeshClipDeform::StaticType = FNodeType(Node::EType::MeshClipDeform, NodeMesh::GetStaticType());
 	FNodeType NodeMeshClipWithMesh::StaticType = FNodeType(Node::EType::MeshClipWithMesh, NodeMesh::GetStaticType());
-	FNodeType NodeMeshParameter::StaticType = FNodeType(Node::EType::MeshParameter, NodeMesh::GetStaticType());
+	FNodeType NodeMeshSkeletalMeshBreak::StaticType = FNodeType(Node::EType::MeshSkeletalMeshBreak, NodeMesh::GetStaticType());
 	FNodeType NodeMeshMakeMorph::StaticType = FNodeType(Node::EType::MeshMakeMorph, NodeMesh::GetStaticType());
 	FNodeType NodeMeshApplyPose::StaticType = FNodeType(Node::EType::MeshApplyPose, NodeMesh::GetStaticType());
 	FNodeType NodeMeshTransform::StaticType = FNodeType(Node::EType::MeshTransform, NodeMesh::GetStaticType());
