[Backout] - CL46982030
#rnx
Original CL Desc
-----------------------------------------------------------------
[Mutable] Replace ME_PARAMETER with SK_PARAMETER

- Initial introduction of the Skeletal Mesh type.
- SK_PARAMETER always returns a Passthrough Parameter Skeletal Mesh.
- ME_SKELETALMESH_BREAK breaks a section of a Skeletal Mesh.
- Fix crash when UObject parameters are nullptr.
- Fix bug with FUniqueMeshIds::EnsureUnique. Ids were always incremented.

diff --git a/Engine/Plugins/Mutable/Source/CustomizableObject/Internal/MuCO/CustomizableObjectPrivate.h b/Engine/Plugins/Mutable/Source/CustomizableObject/Internal/MuCO/CustomizableObjectPrivate.h
--- a/Engine/Plugins/Mutable/Source/CustomizableObject/Internal/MuCO/CustomizableObjectPrivate.h
+++ b/Engine/Plugins/Mutable/Source/CustomizableObject/Internal/MuCO/CustomizableObjectPrivate.h
@@ -1748,8 +1748,6 @@ enum class ECustomizableObjectVersions : int32
 
 	InstancedStructType,
 
-	SkeletalMeshParameter,
-
 	// -----<new versions can be added above this line>--------
 	LastCustomizableObjectVersion
 };
diff --git a/Engine/Plugins/Mutable/Source/CustomizableObject/Private/MuCO/CustomizableObject.cpp b/Engine/Plugins/Mutable/Source/CustomizableObject/Private/MuCO/CustomizableObject.cpp
--- a/Engine/Plugins/Mutable/Source/CustomizableObject/Private/MuCO/CustomizableObject.cpp
+++ b/Engine/Plugins/Mutable/Source/CustomizableObject/Private/MuCO/CustomizableObject.cpp
@@ -1463,13 +1463,13 @@ void UCustomizableObjectPrivate::UpdateParameterPropertiesFromModel(const TShare
 				break;
 			}
 				
-			case UE::Mutable::Private::EParameterType::Texture:
+			case UE::Mutable::Private::EParameterType::Image:
 			{
 				Data.Type = EMutableParameterType::Texture;
 				break;
 			}
 
-			case UE::Mutable::Private::EParameterType::SkeletalMesh:
+			case UE::Mutable::Private::EParameterType::Mesh:
 			{
 				Data.Type = EMutableParameterType::SkeletalMesh;
 				break;
@@ -2259,11 +2259,11 @@ void UCustomizableObjectPrivate::SetModel(const TSharedPtr<UE::Mutable::Private:
 		{
 			switch (Parameter.Type)
 			{
-			case UE::Mutable::Private::EParameterType::Texture:
+			case UE::Mutable::Private::EParameterType::Image:
 				Parameter.DefaultValue.Set<UE::Mutable::Private::FParamTextureType>(TStrongObjectPtr<UTexture>(LocalModelResources->TextureParameterDefaultValues[Parameter.Name]));
 				break;
 				
-			case UE::Mutable::Private::EParameterType::SkeletalMesh:
+			case UE::Mutable::Private::EParameterType::Mesh:
 				Parameter.DefaultValue.Set<UE::Mutable::Private::FParamSkeletalMeshType>(TStrongObjectPtr<USkeletalMesh>(LocalModelResources->SkeletalMeshParameterDefaultValues[Parameter.Name]));
 				break;
 				
diff --git a/Engine/Plugins/Mutable/Source/CustomizableObject/Private/MuCO/CustomizableObjectInstanceDescriptor.cpp b/Engine/Plugins/Mutable/Source/CustomizableObject/Private/MuCO/CustomizableObjectInstanceDescriptor.cpp
--- a/Engine/Plugins/Mutable/Source/CustomizableObject/Private/MuCO/CustomizableObjectInstanceDescriptor.cpp
+++ b/Engine/Plugins/Mutable/Source/CustomizableObject/Private/MuCO/CustomizableObjectInstanceDescriptor.cpp
@@ -220,7 +220,7 @@ void FCustomizableObjectInstanceDescriptor::SaveDescriptor(FArchive& Ar, bool bU
 					RangeValues.Empty(P.ParameterRangeValues.Num());	
 					Algo::Transform(P.ParameterRangeValues, RangeValues, [](const TObjectPtr<UTexture>& Element)
 					{
-						return Element.GetPathName();
+						return Element->GetPathName();
 					});
  
 					break;
@@ -246,7 +246,7 @@ void FCustomizableObjectInstanceDescriptor::SaveDescriptor(FArchive& Ar, bool bU
 					RangeValues.Empty(P.ParameterRangeValues.Num());	
 					Algo::Transform(P.ParameterRangeValues, RangeValues, [](const TObjectPtr<USkeletalMesh>& Element)
 					{
-						return Element.GetPathName();
+						return Element->GetPathName();
 					});
 					
 					break;
@@ -258,34 +258,7 @@ void FCustomizableObjectInstanceDescriptor::SaveDescriptor(FArchive& Ar, bool bU
 				
 			break;
 		}
-
-		case EMutableParameterType::Material:
-		{
-			FString Value;
-			TArray<FString> RangeValues;
-
-			for (const FCustomizableObjectMaterialParameterValue& P : MaterialParameters)
-			{
-				if (P.ParameterName == Name)
-				{
-					Value = P.ParameterValue.GetPathName();
-				
-					RangeValues.Empty(P.ParameterRangeValues.Num());	
-					Algo::Transform(P.ParameterRangeValues, RangeValues, [](const TObjectPtr<UMaterialInterface>& Element)
-					{
-						return Element.GetPathName();
-					});
-				
-					break;
-				}
-			}
-
-			Ar << Value;
-			Ar << RangeValues;
-			
-			break;
-		}
-			
+ 
 		case EMutableParameterType::Projector:
 		{
 			FCustomizableObjectProjector Value;
@@ -789,7 +762,7 @@ TSharedPtr<UE::Mutable::Private::FParameters> FCustomizableObjectInstanceDescrip
 			break;
 		}
  
-		case UE::Mutable::Private::EParameterType::Texture:
+		case UE::Mutable::Private::EParameterType::Image:
 		{
 			for (const FCustomizableObjectTextureParameterValue& TextureParameter : TextureParameters)
 			{
@@ -800,12 +773,12 @@ TSharedPtr<UE::Mutable::Private::FParameters> FCustomizableObjectInstanceDescrip
 						for (int32 RangeIndex = 0; RangeIndex < TextureParameter.ParameterRangeValues.Num(); ++RangeIndex)
 						{
 							RangeIdxPtr->SetPosition(0, RangeIndex);
-							MutableParameters->SetTextureValue(ParamIndex, TextureParameter.ParameterRangeValues[RangeIndex], RangeIdxPtr.Get());
+							MutableParameters->SetImageValue(ParamIndex, TextureParameter.ParameterRangeValues[RangeIndex], RangeIdxPtr.Get());
 						}
 					}
 					else
 					{
-						MutableParameters->SetTextureValue(ParamIndex, TextureParameter.ParameterValue);
+						MutableParameters->SetImageValue(ParamIndex, TextureParameter.ParameterValue);
 					}
  
 					break;
@@ -815,7 +788,7 @@ TSharedPtr<UE::Mutable::Private::FParameters> FCustomizableObjectInstanceDescrip
 			break;
 		}
  
-		case UE::Mutable::Private::EParameterType::SkeletalMesh:
+		case UE::Mutable::Private::EParameterType::Mesh:
 		{
 			for (const FCustomizableObjectSkeletalMeshParameterValue& MeshParameter : SkeletalMeshParameters)
 			{
@@ -826,12 +799,12 @@ TSharedPtr<UE::Mutable::Private::FParameters> FCustomizableObjectInstanceDescrip
 						for (int32 RangeIndex = 0; RangeIndex < MeshParameter.ParameterRangeValues.Num(); ++RangeIndex)
 						{
 							RangeIdxPtr->SetPosition(0, RangeIndex);
-							MutableParameters->SetSkeletalMeshValue(ParamIndex, MeshParameter.ParameterRangeValues[RangeIndex], RangeIdxPtr.Get());
+							MutableParameters->SetMeshValue(ParamIndex, MeshParameter.ParameterRangeValues[RangeIndex], RangeIdxPtr.Get());
 						}
 					}
 					else
 					{
-						MutableParameters->SetSkeletalMeshValue(ParamIndex, MeshParameter.ParameterValue);
+						MutableParameters->SetMeshValue(ParamIndex, MeshParameter.ParameterValue);
 					}
  
 					break;
@@ -1297,7 +1270,7 @@ void FCustomizableObjectInstanceDescriptor::ReloadParameters()
 			break;
 		}
  
-		case UE::Mutable::Private::EParameterType::Texture:
+		case UE::Mutable::Private::EParameterType::Image:
 		{
 			FCustomizableObjectTextureParameterValue Param;
 			Param.ParameterName = Name;
@@ -1341,12 +1314,12 @@ void FCustomizableObjectInstanceDescriptor::ReloadParameters()
 							}
 						}
  
-						Param.ParameterRangeValues[RangeIndex] = MutableParameters->GetTextureValue(ParamIndex, RangeValueIdxPtr.Get());
+						Param.ParameterRangeValues[RangeIndex] = MutableParameters->GetImageValue(ParamIndex, RangeValueIdxPtr.Get());
 					}
 				}
 				else
 				{
-					Param.ParameterValue = MutableParameters->GetTextureValue(ParamIndex);
+					Param.ParameterValue = MutableParameters->GetImageValue(ParamIndex);
 				}
 			}
  
@@ -1355,7 +1328,7 @@ void FCustomizableObjectInstanceDescriptor::ReloadParameters()
 		}
  
  
-		case UE::Mutable::Private::EParameterType::SkeletalMesh:
+		case UE::Mutable::Private::EParameterType::Mesh:
 		{
 			FCustomizableObjectSkeletalMeshParameterValue Param;
 			Param.ParameterName = Name;
@@ -1399,12 +1372,12 @@ void FCustomizableObjectInstanceDescriptor::ReloadParameters()
 							}
 						}
  
-						Param.ParameterRangeValues[RangeIndex] = MutableParameters->GetSkeletalMeshValue(ParamIndex, RangeValueIdxPtr.Get());
+						Param.ParameterRangeValues[RangeIndex] = MutableParameters->GetMeshValue(ParamIndex, RangeValueIdxPtr.Get());
 					}
 				}
 				else
 				{
-					Param.ParameterValue = MutableParameters->GetSkeletalMeshValue(ParamIndex);
+					Param.ParameterValue = MutableParameters->GetMeshValue(ParamIndex);
 				}
 			}
  
diff --git a/Engine/Plugins/Mutable/Source/CustomizableObjectEditor/Private/MuCOE/GenerateMutableSource/GenerateMutableSource.cpp b/Engine/Plugins/Mutable/Source/CustomizableObjectEditor/Private/MuCOE/GenerateMutableSource/GenerateMutableSource.cpp
--- a/Engine/Plugins/Mutable/Source/CustomizableObjectEditor/Private/MuCOE/GenerateMutableSource/GenerateMutableSource.cpp
+++ b/Engine/Plugins/Mutable/Source/CustomizableObjectEditor/Private/MuCOE/GenerateMutableSource/GenerateMutableSource.cpp
@@ -434,12 +434,6 @@ uint32 GetTypeHash(const FGeneratedSourceExternalKey& Key)
 }
 
 
-uint32 GetTypeHash(const FGeneratedSourceSkeletalMeshKey& Key)
-{
-	return GetTypeHash(Key.Pin);
-}
-
-
 FGraphCycleKey::FGraphCycleKey(const UEdGraphPin& Pin, const FString& Id, const UCustomizableObjectNodeMacroInstance* MacroContext) :
 	Pin(Pin),
 	Id(Id),
diff --git a/Engine/Plugins/Mutable/Source/CustomizableObjectEditor/Private/MuCOE/GenerateMutableSource/GenerateMutableSource.h b/Engine/Plugins/Mutable/Source/CustomizableObjectEditor/Private/MuCOE/GenerateMutableSource/GenerateMutableSource.h
--- a/Engine/Plugins/Mutable/Source/CustomizableObjectEditor/Private/MuCOE/GenerateMutableSource/GenerateMutableSource.h
+++ b/Engine/Plugins/Mutable/Source/CustomizableObjectEditor/Private/MuCOE/GenerateMutableSource/GenerateMutableSource.h
@@ -35,8 +35,6 @@
 #include "MuR/Mesh.h"
 #include "MuR/Ptr.h"
 #include "MuT/NodeImageParameter.h"
-#include "MuT/NodeSkeletalMesh.h"
-#include "MuT/NodeSkeletalMeshParameter.h"
 #include "UObject/Package.h"
 #include "UObject/StrongObjectPtr.h"
 #include "Tasks/Task.h"
@@ -194,17 +192,6 @@ struct FGeneratedSourceExternalKey
 uint32 GetTypeHash(const FGeneratedSourceExternalKey& Key);
 
 
-struct FGeneratedSourceSkeletalMeshKey
-{
-	const UEdGraphPin* Pin = nullptr;
-	
-	bool operator==(const FGeneratedSourceSkeletalMeshKey& Other) const = default;
-};
-
-
-uint32 GetTypeHash(const FGeneratedSourceSkeletalMeshKey& Key);
-
-
 struct FImageKey
 {
 	FImageKey(const UEdGraphPin* InPin)
@@ -283,12 +270,6 @@ struct FGeneratedSourceExternalData
 };
 
 
-struct FGeneratedSourceSkeletalMeshData
-{
-	UE::Mutable::Private::Ptr<UE::Mutable::Private::NodeSkeletalMesh> Node;
-};
-
-
 inline uint32 GetTypeHash(const FImageKey& Key)
 {
 	uint32 GuidHash = GetTypeHash(Key.Pin->PinId);
@@ -637,7 +618,6 @@ struct FMutableGraphGenerationContext
 	// Cache of generated pins per LOD
 	TMap<FGeneratedKey, FGeneratedData> Generated;
 	TMap<FGeneratedSourceExternalKey, FGeneratedSourceExternalData> GeneratedSourceExternals;
-	TMap<FGeneratedSourceSkeletalMeshKey, FGeneratedSourceSkeletalMeshData> GeneratedSourceSkeletalMesh;
 
 	/** Set of all generated nodes. */
 	TSet<UCustomizableObjectNode*> GeneratedNodes;
@@ -671,10 +651,6 @@ struct FMutableGraphGenerationContext
 
 	/** Key is the Node Uid. */
 	TMap<FGuid, UE::Mutable::Private::Ptr<UE::Mutable::Private::NodeImageParameter>> GeneratedImageParameters;
-
-	/** Key is the Node Uid. */
-	TMap<FGuid, UE::Mutable::Private::Ptr<UE::Mutable::Private::NodeSkeletalMeshParameter>> GeneratedSkeletalMeshParameters;
-
 	
 	struct FGeneratedCompositeDataTablesData
 	{
diff --git a/Engine/Plugins/Mutable/Source/CustomizableObjectEditor/Private/MuCOE/GenerateMutableSource/GenerateMutableSourceMesh.cpp b/Engine/Plugins/Mutable/Source/CustomizableObjectEditor/Private/MuCOE/GenerateMutableSource/GenerateMutableSourceMesh.cpp
--- a/Engine/Plugins/Mutable/Source/CustomizableObjectEditor/Private/MuCOE/GenerateMutableSource/GenerateMutableSourceMesh.cpp
+++ b/Engine/Plugins/Mutable/Source/CustomizableObjectEditor/Private/MuCOE/GenerateMutableSource/GenerateMutableSourceMesh.cpp
@@ -11,7 +11,6 @@
 #include "ClothingAsset.h"
 #include "GenerateMutableSourceExternal.h"
 #include "GenerateMutableSourceImage.h"
-#include "GenerateMutableSourceSkeletalMesh.h"
 #include "MeshUtilities.h"
 #include "Modules/ModuleManager.h"
 #include "MuCO/CustomizableObjectPrivate.h"
@@ -53,7 +52,7 @@
 #include "MuT/NodeMeshSwitch.h"
 #include "MuT/NodeMeshTable.h"
 #include "MuT/NodeMeshVariation.h"
-#include "MuT/NodeMeshSkeletalMeshBreak.h"
+#include "MuT/NodeMeshParameter.h"
 #include "MuR/Types.h"
 #include "MuR/OpMeshFormat.h"
 #include "MuR/ConvertData.h"
@@ -4591,21 +4590,27 @@ UE::Mutable::Private::Ptr<UE::Mutable::Private::NodeMesh> GenerateMutableSourceM
 
 	else if (const UCustomizableObjectNodeSkeletalMeshParameter* TypedNodeParam = Cast<UCustomizableObjectNodeSkeletalMeshParameter>(Node))
 	{
-		UE::Mutable::Private::Ptr<UE::Mutable::Private::NodeSkeletalMeshParameter> ParameterNode = UE::Mutable::Private::GenerateMutableSourceSkeletalMeshParameter(Pin, GenerationContext);
-		
-		UE::Mutable::Private::Ptr<UE::Mutable::Private::NodeMeshSkeletalMeshBreak> BreakNode = new UE::Mutable::Private::NodeMeshSkeletalMeshBreak();
-		BreakNode->SkeletalMeshParameter = ParameterNode;
-		
+		UE::Mutable::Private::Ptr<UE::Mutable::Private::NodeMeshParameter> ParameterNode = new UE::Mutable::Private::NodeMeshParameter();
+
+		ParameterNode->Name = TypedNodeParam->GetParameterName();
+		ParameterNode->UID = GenerationContext.GetNodeIdUnique(Node).ToString();
+
+		GenerationContext.SkeletalMeshParameterDefaultValues.Add(TypedNodeParam->GetParameterName(), UE::Mutable::Private::LoadObject(TypedNodeParam->DefaultValue));
+
 		int32 BaseLODIndex = 0;
 		int32 BaseSectionIndex = 0;
 
 		{
 			TypedNodeParam->GetPinSection(*Pin, BaseLODIndex, BaseSectionIndex);
-			BreakNode->LODIndex = BaseLODIndex;
-			BreakNode->SectionIndex = BaseSectionIndex;
-			BreakNode->bAutomaticLODs = GenerationContext.CurrentAutoLODStrategy == ECustomizableObjectAutomaticLODStrategy::AutomaticFromMesh;
+			ParameterNode->LODIndex = BaseLODIndex;
+			ParameterNode->SectionIndex = BaseSectionIndex;
+			ParameterNode->bAutomaticLODs = GenerationContext.CurrentAutoLODStrategy == ECustomizableObjectAutomaticLODStrategy::AutomaticFromMesh;
 		}
 		
+		GenerationContext.ParameterUIDataMap.Add(TypedNodeParam->GetParameterName(), FMutableParameterData(
+			TypedNodeParam->ParamUIMetadata,
+			EMutableParameterType::Texture));
+
 		// Layouts
 		{
 			const FLayoutGenerationFlags& LayoutFlags = GenerationContext.LayoutGenerationFlags.Last();
@@ -4613,19 +4618,19 @@ UE::Mutable::Private::Ptr<UE::Mutable::Private::NodeMesh> GenerateMutableSourceM
 			const TArray<UCustomizableObjectLayout*>& Layouts = TypedNodeParam->GetLayouts(*Pin);
 			const int32 NumLayouts = Layouts.Num();
 
-			BreakNode->Layouts.SetNum(NumLayouts);
+			ParameterNode->Layouts.SetNum(NumLayouts);
 			for (int32 LayoutIndex = 0; LayoutIndex < NumLayouts; ++LayoutIndex)
 			{
 				if (!LayoutFlags.TexturePinModes.IsValidIndex(LayoutIndex) ||
 					LayoutFlags.TexturePinModes[LayoutIndex] != EPinMode::Mutable)
 				{
-					BreakNode->Layouts[LayoutIndex] = CreateDefaultLayout();
+					ParameterNode->Layouts[LayoutIndex] = CreateDefaultLayout();
 
 					if (Layouts.IsValidIndex(LayoutIndex) && Layouts[LayoutIndex])
 					{
 						// Keep packing strategy if possible, Overlay can be valid with EPinMode == Passthrough	
-						BreakNode->Layouts[LayoutIndex]->Strategy = ConvertLayoutStrategy(Layouts[LayoutIndex]->PackingStrategy);
-						BreakNode->Layouts[LayoutIndex]->TexCoordsIndex = Layouts[LayoutIndex]->GetUVChannel();
+						ParameterNode->Layouts[LayoutIndex]->Strategy = ConvertLayoutStrategy(Layouts[LayoutIndex]->PackingStrategy);
+						ParameterNode->Layouts[LayoutIndex]->TexCoordsIndex = Layouts[LayoutIndex]->GetUVChannel();
 					}
 
 					// Ignore layout
@@ -4639,7 +4644,7 @@ UE::Mutable::Private::Ptr<UE::Mutable::Private::NodeMesh> GenerateMutableSourceM
 					UE::Mutable::Private::Ptr<UE::Mutable::Private::NodeLayout> LayoutNode = CreateMutableLayoutNode(Layout, bLinkedToExtendMaterial);
 					LayoutNode->SetMessageContext(Node);
 
-					BreakNode->Layouts[LayoutIndex] = LayoutNode;
+					ParameterNode->Layouts[LayoutIndex] = LayoutNode;
 				}
 			}
 
@@ -4657,7 +4662,7 @@ UE::Mutable::Private::Ptr<UE::Mutable::Private::NodeMesh> GenerateMutableSourceM
 				// Get the base mesh
 				FMutableSourceMeshData Source;
 				Source.Mesh = LayoutMesh;
-				BreakNode->ReferenceMesh = ConvertSkeletalMeshToMutable(Source, BaseLODIndex, BaseSectionIndex, GenerationContext, Node);
+				ParameterNode->ReferenceMesh = ConvertSkeletalMeshToMutable(Source, BaseLODIndex, BaseSectionIndex, GenerationContext, Node);
 
 				GenerationContext.MeshGenerationFlags.Pop();
 			}
@@ -4667,7 +4672,7 @@ UE::Mutable::Private::Ptr<UE::Mutable::Private::NodeMesh> GenerateMutableSourceM
 		// animbp instance?
 		// gameplay tags?
 
-		Result = BreakNode;
+		Result = ParameterNode;
 
 		// Applying Mesh Morph Node
 		if (GenerationContext.MeshMorphStack.Num())
diff --git a/Engine/Plugins/Mutable/Source/CustomizableObjectEditor/Private/MuCOE/GenerateMutableSource/GenerateMutableSourceSkeletalMesh.cpp b/Engine/Plugins/Mutable/Source/CustomizableObjectEditor/Private/MuCOE/GenerateMutableSource/GenerateMutableSourceSkeletalMesh.cpp
deleted file mode 100644
diff --git a/Engine/Plugins/Mutable/Source/CustomizableObjectEditor/Private/MuCOE/GenerateMutableSource/GenerateMutableSourceSkeletalMesh.h b/Engine/Plugins/Mutable/Source/CustomizableObjectEditor/Private/MuCOE/GenerateMutableSource/GenerateMutableSourceSkeletalMesh.h
deleted file mode 100644
diff --git a/Engine/Plugins/Mutable/Source/CustomizableObjectEditor/Private/MuCOE/SMutableCodeViewer.cpp b/Engine/Plugins/Mutable/Source/CustomizableObjectEditor/Private/MuCOE/SMutableCodeViewer.cpp
--- a/Engine/Plugins/Mutable/Source/CustomizableObjectEditor/Private/MuCOE/SMutableCodeViewer.cpp
+++ b/Engine/Plugins/Mutable/Source/CustomizableObjectEditor/Private/MuCOE/SMutableCodeViewer.cpp
@@ -3197,10 +3197,11 @@ void FMutableCodeTreeElement::GenerateLabelText()
 		break;
 	}
 
-	case UE::Mutable::Private::EOpType::ME_SKELETALMESH_BREAK:
+	case UE::Mutable::Private::EOpType::ME_PARAMETER:
 	{
-		UE::Mutable::Private::OP::MeshSkeletalMeshBreakArgs Args = Program.GetOpArgs<UE::Mutable::Private::OP::MeshSkeletalMeshBreakArgs>(MutableOperation);
-		OpName += FString::Printf( TEXT(" LOD %d Section %d"), Args.LOD, Args.Section );
+		UE::Mutable::Private::OP::MeshParameterArgs Args = Program.GetOpArgs<UE::Mutable::Private::OP::MeshParameterArgs>(MutableOperation);
+		OpName += FString::Printf( TEXT(" LOD %d Section %d of "), Args.LOD, Args.Section );
+		OpName += Program.Parameters[int32(Args.variable)].Name;
 		break;
 	}
 
diff --git a/Engine/Plugins/Mutable/Source/CustomizableObjectEditor/Private/MuCOE/SMutableGraphViewer.cpp b/Engine/Plugins/Mutable/Source/CustomizableObjectEditor/Private/MuCOE/SMutableGraphViewer.cpp
--- a/Engine/Plugins/Mutable/Source/CustomizableObjectEditor/Private/MuCOE/SMutableGraphViewer.cpp
+++ b/Engine/Plugins/Mutable/Source/CustomizableObjectEditor/Private/MuCOE/SMutableGraphViewer.cpp
@@ -187,9 +187,6 @@ static const char* MutableNodeNames[] =
 
 	"LOD",
 
-	"SkeletalMesh",
-	"SkeletalMeshParameter",
-	
 	"Component",
 	"ComponentNew",
 	"ComponentEdit",
diff --git a/Engine/Plugins/Mutable/Source/MutableRuntime/Internal/MuR/Operations.h b/Engine/Plugins/Mutable/Source/MutableRuntime/Internal/MuR/Operations.h
--- a/Engine/Plugins/Mutable/Source/MutableRuntime/Internal/MuR/Operations.h
+++ b/Engine/Plugins/Mutable/Source/MutableRuntime/Internal/MuR/Operations.h
@@ -50,7 +50,7 @@ namespace UE::Mutable::Private
         CO_PARAMETER,
         PR_PARAMETER,
 		IM_PARAMETER,
-		SK_PARAMETER,
+		ME_PARAMETER,
 		ST_PARAMETER,
     	MA_PARAMETER,
     	MI_PARAMETER,
@@ -223,7 +223,6 @@ namespace UE::Mutable::Private
 		//! Apply linear transform to Image content. Resulting samples outside the original image are tiled.
 		IM_TRANSFORM,
 
-        /** Convert an FIamge in Passthrough Parameter mode to an FImage with data. */
     	IM_PARAMETER_CONVERT,
 
     	//! External operation
@@ -314,9 +313,6 @@ namespace UE::Mutable::Private
     	//! External operation
     	ME_EXTERNAL,
 
-        /** Select a Mesh Section from a Skeletal Mesh. */
-    	ME_SKELETALMESH_BREAK,
-
         //-----------------------------------------------------------------------------------------
         // Instance operations
         //-----------------------------------------------------------------------------------------
@@ -496,10 +492,10 @@ namespace UE::Mutable::Private
 			LIST_UINT32_ADDRESS ImageParameterAddress;
 		};
 
-    	struct MeshSkeletalMeshBreakArgs
-    	{
-    		ADDRESS SkeletalMeshParameter;
-    		uint8 LOD = 0;
+		struct MeshParameterArgs
+		{
+			ADDRESS variable;
+			uint8 LOD = 0;
 			uint8 Section = 0;
 			uint32 MeshID = 0;
 		};
@@ -1208,7 +1204,6 @@ namespace UE::Mutable::Private
 		Matrix,
 		Material,
 		InstancedStruct,
-		SkeletalMesh,
 
 		// Supporting data types : Never returned as an actual data type for any operation.
 		Shape,
diff --git a/Engine/Plugins/Mutable/Source/MutableRuntime/Internal/MuR/Parameters.h b/Engine/Plugins/Mutable/Source/MutableRuntime/Internal/MuR/Parameters.h
--- a/Engine/Plugins/Mutable/Source/MutableRuntime/Internal/MuR/Parameters.h
+++ b/Engine/Plugins/Mutable/Source/MutableRuntime/Internal/MuR/Parameters.h
@@ -48,10 +48,10 @@ namespace UE::Mutable::Private
         Projector,
 
 		/** An externally provided image. */
-		Texture,
+		Image,
 
 		/** An externally provided mesh. */
-		SkeletalMesh,
+		Mesh,
 		
 		/** An externally provided material*/
     	Material,
@@ -316,23 +316,24 @@ namespace UE::Mutable::Private
 								float ProjectionAngle,
 								const FRangeIndex*  RangePosition=nullptr );
 
-        //! Return the value of a Texture parameter.
-        //! \pre The parameter specified by index is a Texture.
-        //! \param Index Index of the parameter from 0 to GetCount()-1
-		//! \param Pos Only for multidimensional parameters: relevant position to set in the ranges
-	    //! \return The UTexture specified when setting the image value (\see SetTextureValue)
-		UE_API UTexture* GetTextureValue(int32 Index, const FRangeIndex* Pos = nullptr) const;
+        //! Return the value of an image parameter.
+        //! \pre The parameter specified by index is a T_IMAGE.
+        //! \param index Index of the parameter from 0 to GetCount()-1
+		//! \param pos Only for multidimensional parameters: relevant position to set in the ranges
+	    //! \return The externalId specified when setting the image value (\see SetImageValue)
+		UE_API UTexture* GetImageValue( int32 Index, const FRangeIndex* Pos = nullptr) const;
 
 		/** */
-		UE_API USkeletalMesh* GetSkeletalMeshValue(int32 Index, const FRangeIndex* Pos = nullptr) const;
+		UE_API USkeletalMesh* GetMeshValue(int32 Index, const FRangeIndex* Pos = nullptr) const;
 
-        //! If the parameter is of the Texture type, set its value.
-        //! \param Index Index of the parameter from 0 to GetCount()-1
-        //! \param Pos Only for multidimensional parameters: relevant position to set in the ranges
-		UE_API void SetTextureValue(int32 Index, UTexture* Value, const FRangeIndex* Pos = nullptr);
+        //! If the parameter is of the image type, set its value.
+        //! \param index Index of the parameter from 0 to GetCount()-1
+        //! \param externalId Application-specific id used to identify this image during replication.
+		//! \param pos Only for multidimensional parameters: relevant position to set in the ranges
+		UE_API void SetImageValue( int32 Index, UTexture* externalId, const FRangeIndex* = nullptr);
 
 		/** */
-		UE_API void SetSkeletalMeshValue(int32 Index, USkeletalMesh* Value, const FRangeIndex* = nullptr);
+		UE_API void SetMeshValue(int32 Index, USkeletalMesh* externalId, const FRangeIndex* = nullptr);
 
         //! Return the value of a float parameter.
         //! \pre The parameter specified by index is a T_FLOAT.
diff --git a/Engine/Plugins/Mutable/Source/MutableRuntime/Internal/MuR/SkeletalMesh.h b/Engine/Plugins/Mutable/Source/MutableRuntime/Internal/MuR/SkeletalMesh.h
deleted file mode 100644
diff --git a/Engine/Plugins/Mutable/Source/MutableRuntime/Internal/MuR/SystemPrivate.h b/Engine/Plugins/Mutable/Source/MutableRuntime/Internal/MuR/SystemPrivate.h
--- a/Engine/Plugins/Mutable/Source/MutableRuntime/Internal/MuR/SystemPrivate.h
+++ b/Engine/Plugins/Mutable/Source/MutableRuntime/Internal/MuR/SystemPrivate.h
@@ -2,7 +2,6 @@
 
 #pragma once
 
-#include "SkeletalMesh.h"
 #include "MuR/MemoryCounters.h"
 #include "MuR/System.h"
 
@@ -499,7 +498,6 @@ namespace UE::Mutable::Private
 		TMemoryTrackedArray<FVector4f> ColorResults;
 		TMemoryTrackedArray<TResourceResult<FImage>> ImageResults;
 		TMemoryTrackedArray<TResourceResult<FMesh>> MeshResults;
-		TMemoryTrackedArray<TSharedPtr<const FSkeletalMesh>> SkeletalMeshResults;
 		TMemoryTrackedArray<TSharedPtr<const FLayout>> LayoutResults;
 		TMemoryTrackedArray<TSharedPtr<const FInstance>> InstanceResults;
 		TMemoryTrackedArray<FProjector> ProjectorResults;
@@ -561,7 +559,6 @@ namespace UE::Mutable::Private
 				ImageResults.Emplace();
 				LayoutResults.Add(nullptr);
 				MeshResults.Emplace();
-				SkeletalMeshResults.Emplace();
 				InstanceResults.Add(nullptr);
 				ProjectorResults.Add(FProjector());
 				StringResults.Add(nullptr);
@@ -837,29 +834,6 @@ namespace UE::Mutable::Private
 		}
 
 
-		TSharedPtr<const FSkeletalMesh> GetSkeletalMesh(FCacheAddress at)
-		{
-			if (!at.At) return nullptr;
-			if (at.At >= OpExecutionData.size_code()) return nullptr;
-			FOpExecutionData* Data = OpExecutionData.get_ptr(at);
-			if (!Data) return nullptr;
-
-			check(Data->DataType == EDataType::SkeletalMesh);
-			TSharedPtr<const FSkeletalMesh> Result = SkeletalMeshResults[Data->DataTypeIndex];
-
-			// We need to decrease the hit-count even if the result is null.
-			check(Data->OpHitCount > 0);
-
-			--Data->OpHitCount;
-			if (Data->OpHitCount == 0 && !Data->IsCacheLocked)
-			{
-				SetUnused(*Data);
-			}
-
-			return Result;
-		}
-		
-
 		TSharedPtr<const FMaterial> GetMaterial(FCacheAddress at, bool& bIsLastReference)
 		{
 			bIsLastReference = false;
@@ -1157,28 +1131,6 @@ namespace UE::Mutable::Private
 			UE::Mutable::Private::UpdateLLMStats();
 		}
 
-		void SetSkeletalMesh(FCacheAddress At, TSharedPtr<const FSkeletalMesh> V)
-		{
-			check(At.At < OpExecutionData.size_code());
-			FOpExecutionData* Data = OpExecutionData.get_ptr(At);
-			check(Data->DataType == EDataType::SkeletalMesh || Data->DataType == EDataType::None);
-			Data->DataType = EDataType::SkeletalMesh;
-			Data->IsValueValid = true;
-
-			if (!Data->DataTypeIndex)
-			{
-				Data->DataTypeIndex = SkeletalMeshResults.Num();
-				SkeletalMeshResults.Add(V);
-			}
-			else
-			{
-				SkeletalMeshResults[Data->DataTypeIndex] = V;
-			}
-			check(Data->DataTypeIndex != 0);
-
-			UE::Mutable::Private::UpdateLLMStats();
-		}
-
 		void SetLayout(FCacheAddress at, TSharedPtr<const FLayout> v)
 		{
 			check(at.At < OpExecutionData.size_code());
diff --git a/Engine/Plugins/Mutable/Source/MutableRuntime/Private/MuR/CodeRunner.cpp b/Engine/Plugins/Mutable/Source/MutableRuntime/Private/MuR/CodeRunner.cpp
--- a/Engine/Plugins/Mutable/Source/MutableRuntime/Private/MuR/CodeRunner.cpp
+++ b/Engine/Plugins/Mutable/Source/MutableRuntime/Private/MuR/CodeRunner.cpp
@@ -3457,28 +3457,6 @@ namespace UE::Mutable::Private
 			break;
 		}
 
-        case EOpType::ME_SKELETALMESH_BREAK:
-		{
-			OP::MeshSkeletalMeshBreakArgs Args = Program.GetOpArgs<OP::MeshSkeletalMeshBreakArgs>(Item.At);
-			switch (Item.Stage)
-			{
-			case 0:
-				AddOp(FScheduledOp(Item.At, Item, 1), FScheduledOp(Args.SkeletalMeshParameter, Item));
-				break;
-
-			case 1:
-				{
-					// This has been moved to a task. It should have been intercepted in IssueOp.
-					check(false);
-				}
-
-			default:
-				check(false);
-			}
-
-			break;
-		}
-        	
 		default:
 			if (type!=EOpType::NONE)
 			{
@@ -4705,7 +4683,7 @@ namespace UE::Mutable::Private
         {
 			OP::ParameterArgs Args = Program.GetOpArgs<OP::ParameterArgs>(Item.At);
 			TSharedPtr<FRangeIndex> Index = BuildCurrentOpRangeIndex(Item, Args.variable);
-			UTexture* Image = pParams->GetTextureValue(Args.variable, Index.Get());
+			UTexture* Image = pParams->GetImageValue(Args.variable, Index.Get());
 
        		TPassthroughObjectPtr<UTexture> ExternalObject(Image);
 			TSharedPtr<FImage> Result = FImage::CreateAsReference(ExternalObject, FImageDesc(), false);
@@ -7007,19 +6985,6 @@ namespace UE::Mutable::Private
 
 				break;
 			}
-
-		case EOpType::SK_PARAMETER:
-			{
-				OP::ParameterArgs Args = Program.GetOpArgs<OP::ParameterArgs>(Item.At);
-				TSharedPtr<FRangeIndex> Index = BuildCurrentOpRangeIndex(Item, Args.variable);
-				USkeletalMesh* SkeletalMesh = pParams->GetSkeletalMeshValue(Args.variable, Index.Get());
-
-				TSharedPtr<FSkeletalMesh> Result = MakeShared<FSkeletalMesh>();
-				Result->PassthroughObject = TPassthroughObjectPtr<USkeletalMesh>(SkeletalMesh);
-        		
-				StoreSkeletalMesh(Item, Result);
-				break;
-			}
 			
 		default:
 		{
@@ -7196,7 +7161,7 @@ namespace UE::Mutable::Private
 			check(Item.Stage == 0);
 
 			OP::ParameterArgs Args = Program.GetOpArgs<OP::ParameterArgs>(Item.At);
-			UTexture* Image = pParams->GetTextureValue(Args.variable);
+			UTexture* Image = pParams->GetImageValue(Args.variable);
 
 			FExtendedImageDesc& Result = ImageDescResults.FindOrAdd(Item.At);
 			Result = GetExternalImageDesc(Image);
diff --git a/Engine/Plugins/Mutable/Source/MutableRuntime/Private/MuR/CodeRunner.h b/Engine/Plugins/Mutable/Source/MutableRuntime/Private/MuR/CodeRunner.h
--- a/Engine/Plugins/Mutable/Source/MutableRuntime/Private/MuR/CodeRunner.h
+++ b/Engine/Plugins/Mutable/Source/MutableRuntime/Private/MuR/CodeRunner.h
@@ -488,11 +488,6 @@ namespace  UE::Mutable::Private
 			return System->WorkingMemoryManager.LoadMesh(From);
 		}
 
-    	inline TSharedPtr<const FSkeletalMesh> LoadSkeletalMesh(const FCacheAddress& From)
-		{
-			return System->WorkingMemoryManager.CurrentInstanceCache->GetSkeletalMesh(From);
-		}
-    	
 		inline TSharedPtr<const FImage> LoadImage(const FCacheAddress& From)
 		{
 			return System->WorkingMemoryManager.LoadImage(From);
@@ -568,11 +563,6 @@ namespace  UE::Mutable::Private
 			System->WorkingMemoryManager.StoreMesh(To, Resource);
 		}
 
-    	inline void StoreSkeletalMesh(const FCacheAddress& To, TSharedPtr<const FSkeletalMesh> Resource)
-		{
-			System->WorkingMemoryManager.CurrentInstanceCache->SetSkeletalMesh(To, Resource);
-		}
-    	
 		inline void StoreImage(const FCacheAddress& To, TSharedPtr<const FImage> Resource)
 		{
 			System->WorkingMemoryManager.StoreImage(To, Resource);
diff --git a/Engine/Plugins/Mutable/Source/MutableRuntime/Private/MuR/CodeRunnerBegin.cpp b/Engine/Plugins/Mutable/Source/MutableRuntime/Private/MuR/CodeRunnerBegin.cpp
--- a/Engine/Plugins/Mutable/Source/MutableRuntime/Private/MuR/CodeRunnerBegin.cpp
+++ b/Engine/Plugins/Mutable/Source/MutableRuntime/Private/MuR/CodeRunnerBegin.cpp
@@ -369,7 +369,6 @@ namespace UE::Mutable::Private
 				case EDataType::Layout:		
 				case EDataType::Material:		
 				case EDataType::InstancedStruct:		
-				case EDataType::SkeletalMesh:		
 					break;
 
 				default:
@@ -389,7 +388,7 @@ namespace UE::Mutable::Private
 			case EOpType::CO_PARAMETER:
 			case EOpType::MA_PARAMETER:
 			case EOpType::PR_PARAMETER:
-			case EOpType::SK_PARAMETER:
+			case EOpType::ME_PARAMETER:
 			case EOpType::IM_PARAMETER:
 			case EOpType::MI_PARAMETER:
 			case EOpType::IM_PARAMETER_FROM_MATERIAL:
@@ -1987,16 +1986,6 @@ namespace UE::Mutable::Private
 					StoreNone(Item);
 					break;
 				}
-
-			case EOpType::ME_SKELETALMESH_BREAK:
-				{
-					OP::MeshSkeletalMeshBreakArgs Args = Program.GetOpArgs<OP::MeshSkeletalMeshBreakArgs>(Item.At);
-					PushOp(FScheduledOpInline(Args.SkeletalMeshParameter, Item));
-
-					StoreNone(Item);
-					break;
-				}
-			
 			
 			case EOpType::IM_REFERENCE:
 				{
diff --git a/Engine/Plugins/Mutable/Source/MutableRuntime/Private/MuR/CodeRunnerTasked.cpp b/Engine/Plugins/Mutable/Source/MutableRuntime/Private/MuR/CodeRunnerTasked.cpp
--- a/Engine/Plugins/Mutable/Source/MutableRuntime/Private/MuR/CodeRunnerTasked.cpp
+++ b/Engine/Plugins/Mutable/Source/MutableRuntime/Private/MuR/CodeRunnerTasked.cpp
@@ -41,7 +41,6 @@
 #include "Trace/Detail/Channel.h"
 #include "ProfilingDebugging/CountersTrace.h"
 #include "Async/Fundamental/Scheduler.h"
-#include "MuR/SkeletalMesh.h"
 
 #if MUTABLE_DEBUG_CODERUNNER_TASK_SCHEDULE_CALLSTACK
 #include "GenericPlatform/GenericPlatformStackWalk.h"
@@ -2668,21 +2667,19 @@ namespace UE::Mutable::Private
 			break;
 		}
 
-		case EOpType::ME_SKELETALMESH_BREAK:
+		case EOpType::ME_PARAMETER:
 		{
-			if (Item.Stage == 1)
-			{
-				OP::MeshSkeletalMeshBreakArgs Args = Program.GetOpArgs<OP::MeshSkeletalMeshBreakArgs>(Item.At);
+			OP::MeshParameterArgs Args = Program.GetOpArgs<OP::MeshParameterArgs>(Item.At);
+			TSharedPtr<FRangeIndex> Index = BuildCurrentOpRangeIndex(Item, Args.variable);
 
-				TSharedPtr<const FSkeletalMesh> SkeletalMesh = LoadSkeletalMesh(FCacheAddress(Args.SkeletalMeshParameter, Item));
+			USkeletalMesh* Id = Params->GetMeshValue(Args.variable, Index.Get());
 
-				int32 LODIndex = Args.LOD;
-				int32 SectionIndex = Args.Section;
+			int32 LODIndex = Args.LOD;
+			int32 SectionIndex = Args.Section;
 
-				CodeRunner::FExternalResourceId FullId;
-				FullId.MeshParameter = SkeletalMesh->PassthroughObject.Get();
-				Issued = MakeShared<FMeshExternalLoadTask>(Item, FullId, LODIndex, SectionIndex, Args.MeshID);
-			}
+			CodeRunner::FExternalResourceId FullId;
+			FullId.MeshParameter = Id;
+			Issued = MakeShared<FMeshExternalLoadTask>(Item, FullId, LODIndex, SectionIndex, Args.MeshID);
 
 			break;
 		}
diff --git a/Engine/Plugins/Mutable/Source/MutableRuntime/Private/MuR/Model.cpp b/Engine/Plugins/Mutable/Source/MutableRuntime/Private/MuR/Model.cpp
--- a/Engine/Plugins/Mutable/Source/MutableRuntime/Private/MuR/Model.cpp
+++ b/Engine/Plugins/Mutable/Source/MutableRuntime/Private/MuR/Model.cpp
@@ -475,20 +475,20 @@ namespace UE::Mutable::Private
                             break;
                         }
 
-                        case EParameterType::Texture:
-                            pRes->SetTextureValue(thisP, OldParameters->GetTextureValue(p));
+                        case EParameterType::Image:
+                            pRes->SetImageValue( thisP, OldParameters->GetImageValue(p) );
                             break;
                         	
-                        case EParameterType::SkeletalMesh:
-                        	pRes->SetSkeletalMeshValue(thisP, OldParameters->GetSkeletalMeshValue(p));
+                        case EParameterType::Mesh:
+                        	pRes->SetMeshValue( thisP, OldParameters->GetMeshValue(p) );
                         	break;
                         	
 	        			case EParameterType::Material:
-	        				pRes->SetMaterialValue(thisP, OldParameters->GetMaterialValue(p));
+	        				pRes->SetMaterialValue( thisP, OldParameters->GetMaterialValue(p) );
 	        				break;
                         	
 	        			case EParameterType::InstancedStruct:
-	        				pRes->SetInstancedStructValue(thisP, OldParameters->GetInstancedStructValue(p));
+	        				pRes->SetInstancedStructValue( thisP, OldParameters->GetInstancedStructValue(p) );
 	        				break;
 
                         default:
diff --git a/Engine/Plugins/Mutable/Source/MutableRuntime/Private/MuR/Operations.cpp b/Engine/Plugins/Mutable/Source/MutableRuntime/Private/MuR/Operations.cpp
--- a/Engine/Plugins/Mutable/Source/MutableRuntime/Private/MuR/Operations.cpp
+++ b/Engine/Plugins/Mutable/Source/MutableRuntime/Private/MuR/Operations.cpp
@@ -38,7 +38,7 @@ namespace UE::Mutable::Private
 		{ EDataType::Color,			},	// CO_PARAMETER
 		{ EDataType::Projector,		},	// PR_PARAMETER
 		{ EDataType::Image,			},	// IM_PARAMETER
-		{ EDataType::SkeletalMesh,	},	// SK_PARAMETER
+		{ EDataType::Mesh,			},	// ME_PARAMETER
 		{ EDataType::String,			},	// ST_PARAMETER
 		{ EDataType::Matrix,			},	// MA_PARAMETER
 		{ EDataType::Material,		},	// MI_PARAMETER
@@ -142,7 +142,6 @@ namespace UE::Mutable::Private
 		{ EDataType::Mesh,			},	// ME_TRANSFORMWITHMESH
 		{ EDataType::Mesh,			},	// ME_TRANSFORMWITHBONE
 		{ EDataType::Mesh,			},	// ME_EXTERNAL
-		{ EDataType::Mesh,			},	// ME_SKELETALMESH_BREAK
 
 		{ EDataType::Instance,		},	// IN_ADDMESH
 		{ EDataType::Instance,		},	// IN_ADDIMAGE
@@ -201,7 +200,7 @@ namespace UE::Mutable::Private
         case EOpType::CO_PARAMETER:
         case EOpType::PR_PARAMETER:
 		case EOpType::IM_PARAMETER:
-		case EOpType::SK_PARAMETER:
+		case EOpType::ME_PARAMETER:
 		case EOpType::MA_PARAMETER:
 		case EOpType::MI_PARAMETER:
 		case EOpType::IS_PARAMETER:
@@ -839,14 +838,7 @@ namespace UE::Mutable::Private
 			f(args.Matrix);
 			break;
 		}
-
-        case EOpType::ME_SKELETALMESH_BREAK:
-        {
-        	OP::MeshSkeletalMeshBreakArgs args = program.GetOpArgs<OP::MeshSkeletalMeshBreakArgs>(at);
-        	f(args.SkeletalMeshParameter);
-        	break;
-        }
-        
+        	
         //-------------------------------------------------------------------------------------
         case EOpType::IN_ADDMESH:
         case EOpType::IN_ADDIMAGE:
diff --git a/Engine/Plugins/Mutable/Source/MutableRuntime/Private/MuR/Parameters.cpp b/Engine/Plugins/Mutable/Source/MutableRuntime/Private/MuR/Parameters.cpp
--- a/Engine/Plugins/Mutable/Source/MutableRuntime/Private/MuR/Parameters.cpp
+++ b/Engine/Plugins/Mutable/Source/MutableRuntime/Private/MuR/Parameters.cpp
@@ -593,17 +593,17 @@ namespace UE::Mutable::Private
     }
 
 
-	USkeletalMesh* FParameters::GetSkeletalMeshValue(int32 Index, const FRangeIndex* Pos) const
+	USkeletalMesh* FParameters::GetMeshValue(int32 Index, const FRangeIndex* Pos) const
 	{
 		check(Index >= 0 && Index < (int32)m_pD->Values.Num());
-		check(GetType(Index) == EParameterType::SkeletalMesh);
+		check(GetType(Index) == EParameterType::Mesh);
 
 		// Early out in case of invalid parameters
 		if (Index < 0
 			||
 			Index >= (int32)m_pD->Values.Num()
 			||
-			GetType(Index) != EParameterType::SkeletalMesh)
+			GetType(Index) != EParameterType::Mesh)
 		{
 			return {};
 		}
@@ -629,13 +629,14 @@ namespace UE::Mutable::Private
 
 		// Multivalue parameter, but no multivalue set. Return single value.
 		return m_pD->Values[Index].Get<FParamSkeletalMeshType>().Get();
+
 	}
 
 
-	void FParameters::SetSkeletalMeshValue(int32 Index, USkeletalMesh* Value, const FRangeIndex* Pos)
+	void FParameters::SetMeshValue(int32 Index, USkeletalMesh* Value, const FRangeIndex* Pos)
 	{
 		check(Index >= 0 && Index < (int32)m_pD->Values.Num());
-		check(GetType(Index) == EParameterType::SkeletalMesh);
+		check(GetType(Index) == EParameterType::Mesh);
 
 		// Single value case
 		if (!Pos)
@@ -660,17 +661,17 @@ namespace UE::Mutable::Private
 	}
 
 
-	UTexture* FParameters::GetTextureValue( int32 Index, const FRangeIndex* Pos ) const
+	UTexture* FParameters::GetImageValue( int32 Index, const FRangeIndex* Pos ) const
     {
         check( Index >= 0 && Index < (int32)m_pD->Values.Num() );
-        check( GetType( Index ) == EParameterType::Texture );
+        check( GetType( Index ) == EParameterType::Image );
 
 		// Early out in case of invalid parameters
         if ( Index < 0
              ||
              Index >= (int32)m_pD->Values.Num()
              ||
-             GetType(Index) != EParameterType::Texture )
+             GetType(Index) != EParameterType::Image )
         {
             return {};
         }
@@ -699,10 +700,10 @@ namespace UE::Mutable::Private
     }
 
 
-    void FParameters::SetTextureValue( int32 Index, UTexture* Value, const FRangeIndex* Pos )
+    void FParameters::SetImageValue( int32 Index, UTexture* Value, const FRangeIndex* Pos )
     {
         check( Index >= 0 && Index < (int32)m_pD->Values.Num() );
-        check( GetType( Index ) == EParameterType::Texture );
+        check( GetType( Index ) == EParameterType::Image );
 
 		// Single value case
 		if (!Pos)
diff --git a/Engine/Plugins/Mutable/Source/MutableRuntime/Private/MuR/System.cpp b/Engine/Plugins/Mutable/Source/MutableRuntime/Private/MuR/System.cpp
--- a/Engine/Plugins/Mutable/Source/MutableRuntime/Private/MuR/System.cpp
+++ b/Engine/Plugins/Mutable/Source/MutableRuntime/Private/MuR/System.cpp
@@ -943,7 +943,6 @@ namespace UE::Mutable::Private
             case EOpType::CO_PARAMETER:
             case EOpType::PR_PARAMETER:
 			case EOpType::IM_PARAMETER:
-			case EOpType::SK_PARAMETER:
 			case EOpType::MA_PARAMETER:
 			case EOpType::MI_PARAMETER:
 			case EOpType::IS_PARAMETER:
@@ -953,8 +952,16 @@ namespace UE::Mutable::Private
                 Flags[ParamIndex] = true;
                 break;
             }
+			case EOpType::ME_PARAMETER:
+			{
+				OP::MeshParameterArgs args = Program.GetOpArgs<OP::MeshParameterArgs>(at);
+				OP::ADDRESS ParamIndex = args.variable;
+				Flags[ParamIndex] = true;
+				break;
+			}
             default:
                 break;
+
             }
 
             return UniqueDiscreteCoveredCodeVisitor<>::Visit( at, Program );
@@ -2135,31 +2142,9 @@ namespace UE::Mutable::Private
 				break;
 			}
 
-			case EParameterType::Texture:
-			{
-				TStrongObjectPtr<UTexture> Object = Params->GetPrivate()->Values[ParamIndex].Get<FParamTextureType>();
-				FString Path = Object ? Object->GetPathName() : FString();
-				Blob.Serialize(GetData(Path), Path.Len() * sizeof(FString::ElementType));
-
-				// Multi-values
-				if (Params->GetPrivate()->HasMultipleValues(ParamIndex))
-				{
-					const TMap< TArray<int32>, FParameterValue >& Multi = Params->GetPrivate()->MultiValues[ParamIndex];
-					AddMultiValueKeys(Blob, Multi);
-					for (const TPair< TArray<int32>, FParameterValue >& v : Multi)
-					{
-						Object = v.Value.Get<FParamTextureType>();
-						Path = Object ? Object->GetPathName() : FString();
-						Blob.Serialize(GetData(Path), Path.Len() * sizeof(FString::ElementType));
-					}
-				}
-				break;
-			}
-
-			case EParameterType::SkeletalMesh:
+			case EParameterType::Image:
 			{
-				TStrongObjectPtr<USkeletalMesh> Object = Params->GetPrivate()->Values[ParamIndex].Get<FParamSkeletalMeshType>();
-				FString Path = Object ? Object->GetPathName() : FString();
+				FString Path = Params->GetPrivate()->Values[ParamIndex].Get<FParamTextureType>()->GetPathName();
 				Blob.Serialize(GetData(Path), Path.Len() * sizeof(FString::ElementType));
 
 				// Multi-values
@@ -2169,18 +2154,16 @@ namespace UE::Mutable::Private
 					AddMultiValueKeys(Blob, Multi);
 					for (const TPair< TArray<int32>, FParameterValue >& v : Multi)
 					{
-						Object = v.Value.Get<FParamSkeletalMeshType>();
-						Path = Object ? Object->GetPathName() : FString();
+						Path = v.Value.Get<FParamTextureType>()->GetPathName();
 						Blob.Serialize(GetData(Path), Path.Len() * sizeof(FString::ElementType));
 					}
 				}
 				break;
 			}
 
-			case EParameterType::Material:
+			case EParameterType::Mesh:
 			{
-				TStrongObjectPtr<UMaterialInterface> Object = Params->GetPrivate()->Values[ParamIndex].Get<FParamMaterialType>();
-				FString Path = Object ? Object->GetPathName() : FString();
+				FString Path = Params->GetPrivate()->Values[ParamIndex].Get<FParamSkeletalMeshType>()->GetPathName();
 				Blob.Serialize(GetData(Path), Path.Len() * sizeof(FString::ElementType));
 
 				// Multi-values
@@ -2190,8 +2173,7 @@ namespace UE::Mutable::Private
 					AddMultiValueKeys(Blob, Multi);
 					for (const TPair< TArray<int32>, FParameterValue >& v : Multi)
 					{
-						Object = v.Value.Get<FParamMaterialType>();
-						Path = Object ? Object->GetPathName() : FString();
+						Path = v.Value.Get<FParamSkeletalMeshType>()->GetPathName();
 						Blob.Serialize(GetData(Path), Path.Len() * sizeof(FString::ElementType));
 					}
 				}
diff --git a/Engine/Plugins/Mutable/Source/MutableTools/Internal/MuT/Node.h b/Engine/Plugins/Mutable/Source/MutableTools/Internal/MuT/Node.h
--- a/Engine/Plugins/Mutable/Source/MutableTools/Internal/MuT/Node.h
+++ b/Engine/Plugins/Mutable/Source/MutableTools/Internal/MuT/Node.h
@@ -48,7 +48,7 @@ namespace UE::Mutable::Private
 			MeshVariation,
 			MeshReshape,
 			MeshClipDeform,
-			MeshSkeletalMeshBreak,
+			MeshParameter,
 			MeshExternal,
 
 			Image,
@@ -133,9 +133,6 @@ namespace UE::Mutable::Private
 
 			LOD,
 
-			SkeletalMesh,
-			SkeletalMeshParameter,
-			
 			Component,
 			ComponentNew,
 			ComponentEdit,
diff --git a/Engine/Plugins/Mutable/Source/MutableTools/Internal/MuT/NodeImageParameterConvert.h b/Engine/Plugins/Mutable/Source/MutableTools/Internal/MuT/NodeImageParameterConvert.h
--- a/Engine/Plugins/Mutable/Source/MutableTools/Internal/MuT/NodeImageParameterConvert.h
+++ b/Engine/Plugins/Mutable/Source/MutableTools/Internal/MuT/NodeImageParameterConvert.h
@@ -4,7 +4,7 @@
 
 #include "MuR/Ptr.h"
 #include "MuT/Node.h"
-#include "MuT/NodeImageParameter.h"
+#include "MuT/NodeImage.h"
 
 #define UE_API MUTABLETOOLS_API
 
@@ -14,7 +14,7 @@ namespace UE::Mutable::Private
 	class NodeImageParameterConvert : public NodeImage
 	{
 	public:
-		Ptr<NodeImageParameter> ImageParameter;
+		Ptr<NodeImage> ImageParameter;
 		
 		// Node Interface
 		virtual const FNodeType* GetType() const override { return &StaticType; }
diff --git a/Engine/Plugins/Mutable/Source/MutableTools/Internal/MuT/NodeMeshParameter.h b/Engine/Plugins/Mutable/Source/MutableTools/Internal/MuT/NodeMeshParameter.h
new file mode 100644
--- /dev/null
+++ b/Engine/Plugins/Mutable/Source/MutableTools/Internal/MuT/NodeMeshParameter.h
@@ -0,0 +1,60 @@
+// Copyright Epic Games, Inc. All Rights Reserved.
+
+#pragma once
+
+#include "MuR/Ptr.h"
+#include "MuR/RefCounted.h"
+#include "MuT/Node.h"
+#include "MuT/NodeMesh.h"
+#include "MuT/NodeRange.h"
+#include "MuT/NodeLayout.h"
+
+#define UE_API MUTABLETOOLS_API
+
+namespace UE::Mutable::Private
+{
+
+	/** */
+	class NodeMeshParameter : public NodeMesh
+	{
+	public:
+		
+		TSharedPtr<FMesh> ReferenceMesh;
+
+		TArray<Ptr<NodeLayout>> Layouts;
+
+		uint8 LODIndex = 0;
+
+		uint8 SectionIndex = 0;
+
+		bool bAutomaticLODs = false;
+		
+		/** Name of the parameter */
+		FString Name;
+
+		/** User provided ID to identify the parameter. */
+		FString UID;
+
+		/** Ranges for the parameter in case it is multidimensional. */
+		TArray<Ptr<NodeRange>> Ranges;
+
+	public:
+
+		// Node Interface
+		virtual const FNodeType* GetType() const override { return &StaticType; }
+		static const FNodeType* GetStaticType() { return &StaticType; }
+
+	protected:
+
+		/** Forbidden. Manage with the Ptr<> template. */
+		~NodeMeshParameter() {}
+
+	private:
+
+		static UE_API FNodeType StaticType;
+
+	};
+
+}
+
+#undef UE_API
diff --git a/Engine/Plugins/Mutable/Source/MutableTools/Internal/MuT/NodeMeshSkeletalMeshBreak.h b/Engine/Plugins/Mutable/Source/MutableTools/Internal/MuT/NodeMeshSkeletalMeshBreak.h
deleted file mode 100644
diff --git a/Engine/Plugins/Mutable/Source/MutableTools/Internal/MuT/NodeSkeletalMesh.h b/Engine/Plugins/Mutable/Source/MutableTools/Internal/MuT/NodeSkeletalMesh.h
deleted file mode 100644
diff --git a/Engine/Plugins/Mutable/Source/MutableTools/Internal/MuT/NodeSkeletalMeshParameter.h b/Engine/Plugins/Mutable/Source/MutableTools/Internal/MuT/NodeSkeletalMeshParameter.h
deleted file mode 100644
diff --git a/Engine/Plugins/Mutable/Source/MutableTools/Private/MuT/ASTOPMeshSkeletalMeshBreak.cpp b/Engine/Plugins/Mutable/Source/MutableTools/Private/MuT/ASTOPMeshSkeletalMeshBreak.cpp
deleted file mode 100644
diff --git a/Engine/Plugins/Mutable/Source/MutableTools/Private/MuT/ASTOPMeshSkeletalMeshBreak.h b/Engine/Plugins/Mutable/Source/MutableTools/Private/MuT/ASTOPMeshSkeletalMeshBreak.h
deleted file mode 100644
diff --git a/Engine/Plugins/Mutable/Source/MutableTools/Private/MuT/ASTOpParameter.cpp b/Engine/Plugins/Mutable/Source/MutableTools/Private/MuT/ASTOpParameter.cpp
--- a/Engine/Plugins/Mutable/Source/MutableTools/Private/MuT/ASTOpParameter.cpp
+++ b/Engine/Plugins/Mutable/Source/MutableTools/Private/MuT/ASTOpParameter.cpp
@@ -3,6 +3,7 @@
 #include "MuT/ASTOpParameter.h"
 
 #include "MuT/ASTOpBoolEqualIntConst.h"
+#include "HAL/PlatformMath.h"
 #include "Misc/AssertionMacros.h"
 #include "MuR/ModelPrivate.h"
 #include "MuR/Parameters.h"
@@ -12,6 +13,8 @@
 
 namespace UE::Mutable::Private
 {
+
+
 	ASTOpParameter::~ASTOpParameter()
 	{
 		// Explicit call needed to avoid recursive destruction
@@ -40,6 +43,9 @@ namespace UE::Mutable::Private
 			const ASTOpParameter* Other = static_cast<const ASTOpParameter*>(&OtherUntyped);
 			return Type == Other->Type &&
 				Parameter == Other->Parameter &&
+				LODIndex == Other->LODIndex &&
+				SectionIndex == Other->SectionIndex &&
+				MeshID == Other->MeshID &&
 				Ranges == Other->Ranges &&
 				ColorParameterNames == Other->ColorParameterNames &&
 				ScalarParameterNames == Other->ScalarParameterNames &&
@@ -54,7 +60,9 @@ namespace UE::Mutable::Private
 		Ptr<ASTOpParameter> n = new ASTOpParameter();
 		n->Type = Type;
 		n->Parameter = Parameter;
-
+		n->LODIndex = LODIndex;
+		n->SectionIndex = SectionIndex;
+		n->MeshID = MeshID;
 		for (const FRangeData& c : Ranges)
 		{
 			n->Ranges.Emplace(n.get(), MapChild(c.rangeSize.child()), c.rangeName, c.rangeUID);
@@ -82,6 +90,8 @@ namespace UE::Mutable::Private
 	{
 		uint64 res = std::hash<uint64>()(uint64(Type));
 		hash_combine(res, Parameter.Type);
+		hash_combine(res, LODIndex);
+		hash_combine(res, SectionIndex);
 		hash_combine(res, Parameter.Name.Len());
 
 		return res;
@@ -109,8 +119,28 @@ namespace UE::Mutable::Private
 			linkedAddress = (OP::ADDRESS)Program.OpAddress.Num();
 			Program.OpAddress.Add((uint32)Program.ByteCode.Num());
 			AppendCode(Program.ByteCode, Type);
-		
-			if (Type == EOpType::MI_PARAMETER)
+
+			if (Type == EOpType::ME_PARAMETER)
+			{
+				OP::MeshParameterArgs Args;
+				FMemory::Memzero(Args);
+				Args.variable = (OP::ADDRESS)LinkedParameterIndex;
+				Args.LOD = LODIndex;
+				Args.Section = SectionIndex;
+				Args.MeshID = MeshID;
+
+				for (const FRangeData& d : Ranges)
+				{
+					OP::ADDRESS sizeAt = 0;
+					uint16 rangeId = 0;
+					LinkRange(Program, d, sizeAt, rangeId);
+					Program.Parameters[LinkedParameterIndex].Ranges.Add(rangeId);
+				}
+
+				AppendCode(Program.ByteCode, Args);
+			}
+
+			else if (Type == EOpType::MI_PARAMETER)
 			{
 				OP::MaterialParameterArgs Args;
 				FMemory::Memzero(Args);
@@ -220,4 +250,5 @@ namespace UE::Mutable::Private
 	{
 		return false;
 	}
+
 }
diff --git a/Engine/Plugins/Mutable/Source/MutableTools/Private/MuT/ASTOpParameter.h b/Engine/Plugins/Mutable/Source/MutableTools/Private/MuT/ASTOpParameter.h
--- a/Engine/Plugins/Mutable/Source/MutableTools/Private/MuT/ASTOpParameter.h
+++ b/Engine/Plugins/Mutable/Source/MutableTools/Private/MuT/ASTOpParameter.h
@@ -26,6 +26,11 @@ namespace UE::Mutable::Private
 		//!
 		FParameterDesc Parameter;
 
+		/** Used by some parameter types (Mesh) to specify which subdata from the actual parameter value should be used in the operation. */
+		int32 LODIndex = 0;
+		int32 SectionIndex = 0;
+		uint32 MeshID = 0;
+
 		//** Ranges adding dimensions to this parameter. */
 		TArray<FRangeData> Ranges;
 
@@ -41,7 +46,9 @@ namespace UE::Mutable::Private
 		*/
 		int32 LinkedParameterIndex = -1;
 
-		virtual ~ASTOpParameter() override;
+	public:
+
+		~ASTOpParameter() override;
 
 		// ASTOp interface.
 		virtual EOpType GetOpType() const override;
@@ -55,6 +62,8 @@ namespace UE::Mutable::Private
 		virtual FBoolEvalResult EvaluateBool(ASTOpList& Facts, FEvaluateBoolCache* = nullptr) const override;
 		virtual FImageDesc GetImageDesc(bool, FGetImageDescContext*) const override;
 		virtual bool IsColourConstant(FVector4f& OutColour) const override;
+
 	};
+
 }
 
diff --git a/Engine/Plugins/Mutable/Source/MutableTools/Private/MuT/CodeGenerator.cpp b/Engine/Plugins/Mutable/Source/MutableTools/Private/MuT/CodeGenerator.cpp
--- a/Engine/Plugins/Mutable/Source/MutableTools/Private/MuT/CodeGenerator.cpp
+++ b/Engine/Plugins/Mutable/Source/MutableTools/Private/MuT/CodeGenerator.cpp
@@ -2459,7 +2459,7 @@ namespace UE::Mutable::Private
 				}
 				else
 				{
-					check(OriginalMeshNode->GetType() == NodeMeshSkeletalMeshBreak::GetStaticType())
+					check(OriginalMeshNode->GetType() == NodeMeshParameter::GetStaticType())
 					TargetMeshTask = UE::Tasks::MakeCompletedTask<void>();
 				}
 
@@ -2530,7 +2530,7 @@ namespace UE::Mutable::Private
 							}
 							else
 							{
-								check(OriginalMeshCopy->GetType() == NodeMeshSkeletalMeshBreak::GetStaticType());
+								check(OriginalMeshCopy->GetType() == NodeMeshParameter::GetStaticType());
 
 								MorphOp = new ASTOpMeshMorph();
 								
@@ -2651,7 +2651,7 @@ namespace UE::Mutable::Private
 				}
 				else
 				{
-					check(OriginalMeshNode->GetType() == NodeMeshSkeletalMeshBreak::GetStaticType())
+					check(OriginalMeshNode->GetType() == NodeMeshParameter::GetStaticType())
 					TargetMeshTask = UE::Tasks::MakeCompletedTask<void>();
 				}
 
diff --git a/Engine/Plugins/Mutable/Source/MutableTools/Private/MuT/CodeGenerator.h b/Engine/Plugins/Mutable/Source/MutableTools/Private/MuT/CodeGenerator.h
--- a/Engine/Plugins/Mutable/Source/MutableTools/Private/MuT/CodeGenerator.h
+++ b/Engine/Plugins/Mutable/Source/MutableTools/Private/MuT/CodeGenerator.h
@@ -49,8 +49,6 @@ namespace UE::Mutable::Private
 	class NodeImageParameterConvert;
 	class NodeExternalParameter;
 	class NodeExternalOperation;
-	class NodeSkeletalMeshParameter;
-	class NodeSkeletalMesh;
 	class NodeMaterialExternal;
 	class NodeScalarExternal;
 }
@@ -114,7 +112,7 @@ namespace UE::Mutable::Private
 	class NodeMeshSwitch;
 	class NodeMeshTransform;
 	class NodeMeshVariation;
-	class NodeMeshSkeletalMeshBreak;
+	class NodeMeshParameter;
 	class NodeRange;
 	class NodeScalarArithmeticOperation;
 	class NodeScalarConstant;
@@ -834,7 +832,7 @@ namespace UE::Mutable::Private
 		FMeshTask GenerateMesh_Table(const FMeshGenerationStaticOptions&, FMeshOptionsTask, const NodeMeshTable*);
 		FMeshTask GenerateMesh_Reshape(const FMeshGenerationStaticOptions&, FMeshOptionsTask, const NodeMeshReshape*);
 		FMeshTask GenerateMesh_ClipDeform(const FMeshGenerationStaticOptions&, FMeshOptionsTask, const NodeMeshClipDeform*);
-		FMeshTask GenerateMesh_SkeletalMeshBreak(const FMeshGenerationStaticOptions&, FMeshOptionsTask, const NodeMeshSkeletalMeshBreak*);
+		FMeshTask GenerateMesh_Parameter(const FMeshGenerationStaticOptions&, FMeshOptionsTask, const NodeMeshParameter*);
     	FMeshTask GenerateMesh_External(const FMeshGenerationStaticOptions&, FMeshOptionsTask, const NodeMeshExternal*);
 
 		Ptr<ASTOp> GenerateLayoutOpsAndResult(const FMeshGenerationDynamicOptions&, Ptr<ASTOp> LastMeshOp, const TArray<Ptr<NodeLayout>>& OriginalLayouts, uint32 MeshId, FMeshGenerationResult& OutResult);
@@ -843,27 +841,7 @@ namespace UE::Mutable::Private
     	TSharedPtr<const FLayout> GenerateLayout(Ptr<const NodeLayout> SourceLayout, uint32 MeshIDPrefix);
 
     	//-----------------------------------------------------------------------------------------
-    	// Skeletal Mesh
-    	
-    	struct FSkeletalMeshGenerationResult
-    	{
-    		Ptr<ASTOp> Op;
-    	};
-
-    	typedef TMap<FGeneratedCacheKey, FSkeletalMeshGenerationResult> FGeneratedSkeletalMeshesMap;
 
-    	struct FGeneratedSkeletalMeshes
-    	{
-    		UE::FMutex Mutex;
-    		FGeneratedSkeletalMeshesMap Map;
-    	};
-    	FGeneratedSkeletalMeshes GeneratedSkeletalMeshes;
-    	
-    	void GenerateSkeletalMesh(FSkeletalMeshGenerationResult& OutResult, const FGenericGenerationOptions&, const Ptr<const NodeSkeletalMesh>&);
-    	void GenerateSkeletalMesh_Parameter(FSkeletalMeshGenerationResult& OutResult, const FGenericGenerationOptions&, const NodeSkeletalMeshParameter*);
-    	
-    	//-----------------------------------------------------------------------------------------
-    	
     	struct FExternalGenerationOptions 
     	{
     		// TODO Propagated options could be a TVariant
diff --git a/Engine/Plugins/Mutable/Source/MutableTools/Private/MuT/CodeGenerator_FirstPass.h b/Engine/Plugins/Mutable/Source/MutableTools/Private/MuT/CodeGenerator_FirstPass.h
--- a/Engine/Plugins/Mutable/Source/MutableTools/Private/MuT/CodeGenerator_FirstPass.h
+++ b/Engine/Plugins/Mutable/Source/MutableTools/Private/MuT/CodeGenerator_FirstPass.h
@@ -23,7 +23,7 @@
 #include "MuT/NodeSurfaceNew.h"
 #include "MuT/NodeSurfaceVariation.h"
 #include "MuT/NodeSurfaceSwitch.h"
-#include "MuT/NodeMeshSkeletalMeshBreak.h"
+#include "MuT/NodeMeshParameter.h"
 #include "MuT/NodeMaterialParameter.h"
 #include "MuT/ASTOpParameter.h"
 #include "MuT/ErrorLog.h"
@@ -258,6 +258,9 @@ namespace UE::Mutable::Private
 
 			/** Cache for generic parameters. */
 			TMap< Ptr<const Node>, Ptr<ASTOpParameter> > GenericParametersCache;
+
+			/** For mesh parameters we generate a different result for each LOD. */
+			TMap< Ptr<const NodeMeshParameter>, TArray<TPair<Ptr<ASTOpParameter>, FMeshGenerationResult>> > MeshParametersCache;
 		};
 		FSafeParameterNodes ParameterNodes;
 
diff --git a/Engine/Plugins/Mutable/Source/MutableTools/Private/MuT/CodeGenerator_Image.cpp b/Engine/Plugins/Mutable/Source/MutableTools/Private/MuT/CodeGenerator_Image.cpp
--- a/Engine/Plugins/Mutable/Source/MutableTools/Private/MuT/CodeGenerator_Image.cpp
+++ b/Engine/Plugins/Mutable/Source/MutableTools/Private/MuT/CodeGenerator_Image.cpp
@@ -285,7 +285,7 @@ namespace UE::Mutable::Private
 				op->Parameter.Name = Node->Name;
 				bool bParseOk = FGuid::Parse(Node->UID, op->Parameter.UID);
 				check(bParseOk);
-				op->Parameter.Type = EParameterType::Texture;
+				op->Parameter.Type = EParameterType::Image;
 				op->Parameter.DefaultValue.Set<FParamTextureType>(nullptr);
 
 				FirstPass.ParameterNodes.GenericParametersCache.Add(Node, op);
diff --git a/Engine/Plugins/Mutable/Source/MutableTools/Private/MuT/CodeGenerator_Mesh.cpp b/Engine/Plugins/Mutable/Source/MutableTools/Private/MuT/CodeGenerator_Mesh.cpp
--- a/Engine/Plugins/Mutable/Source/MutableTools/Private/MuT/CodeGenerator_Mesh.cpp
+++ b/Engine/Plugins/Mutable/Source/MutableTools/Private/MuT/CodeGenerator_Mesh.cpp
@@ -1,7 +1,6 @@
 // Copyright Epic Games, Inc. All Rights Reserved.
 
 #include "ASTOpExternal.h"
-#include "ASTOpMeshSkeletalMeshBreak.h"
 #include "Containers/Array.h"
 #include "Math/IntPoint.h"
 #include "Math/UnrealMathSSE.h"
@@ -60,7 +59,7 @@
 #include "MuT/NodeMeshTable.h"
 #include "MuT/NodeMeshTransform.h"
 #include "MuT/NodeMeshVariation.h"
-#include "MuT/NodeMeshSkeletalMeshBreak.h"
+#include "MuT/NodeMeshParameter.h"
 #include "MuT/NodeScalar.h"
 #include "MuT/Table.h"
 #include "MuT/TablePrivate.h"
@@ -98,7 +97,7 @@ namespace UE::Mutable::Private
 		case Node::EType::MeshTable: Result = GenerateMesh_Table(StaticOptions, InOptions, static_cast<const NodeMeshTable*>(Node)); break;
 		case Node::EType::MeshReshape: Result = GenerateMesh_Reshape(StaticOptions, InOptions, static_cast<const NodeMeshReshape*>(Node)); break;
 		case Node::EType::MeshClipDeform: Result = GenerateMesh_ClipDeform(StaticOptions, InOptions, static_cast<const NodeMeshClipDeform*>(Node)); break;
-		case Node::EType::MeshSkeletalMeshBreak: Result = GenerateMesh_SkeletalMeshBreak(StaticOptions, InOptions, static_cast<const NodeMeshSkeletalMeshBreak*>(Node)); break;
+		case Node::EType::MeshParameter: Result = GenerateMesh_Parameter(StaticOptions, InOptions, static_cast<const NodeMeshParameter*>(Node)); break;
 		case Node::EType::MeshExternal: Result = GenerateMesh_External(StaticOptions, InOptions, static_cast<const NodeMeshExternal*>(Node)); break;
 		default: check(false);
 		}
@@ -1360,26 +1359,23 @@ namespace UE::Mutable::Private
 	uint32 CodeGenerator::FUniqueMeshIds::EnsureUnique(uint32 Id)
 	{
 		UE::TUniqueLock Lock(Mutex);
-
-		while (true)
+		bool bValid = false;
+		do
 		{
 			bool bAlreadyPresent = false;
 			Map.FindOrAdd(Id, &bAlreadyPresent);
-
-			const bool bValid = !bAlreadyPresent && Id != 0;
-			if (bValid)
+			bValid = !bAlreadyPresent && Id != 0;
+			if (!bValid)
 			{
-				break;
+				++Id;
 			}
-
-			++Id;
-		}
+		} while (bValid);
 
 		return Id;
 	}
 
 
-	FMeshTask CodeGenerator::GenerateMesh_SkeletalMeshBreak(const FMeshGenerationStaticOptions& StaticOptions, FMeshOptionsTask InOptions, const NodeMeshSkeletalMeshBreak* Node)
+	FMeshTask CodeGenerator::GenerateMesh_Parameter(const FMeshGenerationStaticOptions& StaticOptions, FMeshOptionsTask InOptions, const NodeMeshParameter* Node)
 	{
 		TArray<UE::Tasks::FTask, TInlineAllocator<2>> Requisites;
 		Requisites.Emplace(InOptions);
@@ -1406,27 +1402,30 @@ namespace UE::Mutable::Private
 			}
 		}
 
-		FSkeletalMeshGenerationResult SkeletalMeshResult;
-		GenerateSkeletalMesh(SkeletalMeshResult, StaticOptions, Node->SkeletalMeshParameter);
-
-		Ptr<ASTOp> SkeletalMeshParameter = SkeletalMeshResult.Op;
-		
 		// Local pipe because we call GenerateRange in there.
 		FMeshTask MeshParamTask = this->LocalPipe.Launch(TEXT("MutableMeshParameter"),
-			[InOptions, Node, ReferenceMesh, StaticOptions, SkeletalMeshParameter, this]() mutable
+			[InOptions, Node, ReferenceMesh, StaticOptions, this]() mutable
 			{
 				FMeshGenerationDynamicOptions Options = InOptions.GetResult();
 
 				FMeshGenerationResult Result;
-				Ptr<ASTOpMeshSkeletalMeshBreak> Op;
-					
-				Op = new ASTOpMeshSkeletalMeshBreak();
-				Op->SkeletalMeshParameter = SkeletalMeshParameter;
+				Ptr<ASTOpParameter> Op;
+
+				FParameterDesc param;
+				param.Name = Node->Name;
+				bool bParseOk = FGuid::Parse(Node->UID, param.UID);
+				check(bParseOk);
+				param.Type = EParameterType::Mesh;
+				param.DefaultValue.Set<FParamSkeletalMeshType>(nullptr);
+
+				Op = new ASTOpParameter();
+				Op->Type = EOpType::ME_PARAMETER;
+				Op->Parameter = param;
 				Op->LODIndex = Node->LODIndex + (Node->bAutomaticLODs ? StaticOptions.LODIndex : 0);
 				Op->SectionIndex = Node->SectionIndex;
 
 				// Assign an ID prefix to the mesh. 
-				uint32 MeshID = Op->Hash();
+				uint32 MeshID = HashCombine(GetTypeHash(Op->LODIndex), GetTypeHash(param.UID));
 				MeshID = HashCombine(MeshID, GetTypeHash(Op->SectionIndex));
 				MeshID = UniqueMeshIds.EnsureUnique(MeshID);
 				Op->MeshID = MeshID;
@@ -1438,17 +1437,35 @@ namespace UE::Mutable::Private
 
 				{
 					UE::TUniqueLock Lock(FirstPass.ParameterNodes.Mutex);
-					
+
+					TArray<TPair<Ptr<ASTOpParameter>, FMeshGenerationResult>>& ArrayFound = FirstPass.ParameterNodes.MeshParametersCache.FindOrAdd(Node);
+					if (ArrayFound.IsValidIndex(StaticOptions.LODIndex))
+					{
+						return ArrayFound[StaticOptions.LODIndex].Value;
+					}
+
+					ArrayFound.SetNum(StaticOptions.LODIndex+1);
+					ArrayFound[StaticOptions.LODIndex].Key = Op;
+
 					// Fill the result structure. Some ops will be completed outside the lock.
 					Ptr<ASTOp> LastMeshOp = GenerateLayoutOpsAndResult(Options, Op, Node->Layouts, Op->MeshID, Result);
 					Result.MeshOp = LastMeshOp;
 					Result.BaseMeshOp = LastMeshOp;
+					ArrayFound[StaticOptions.LODIndex].Value = Result;
 				}
 
 				int32 LODIndex = StaticOptions.LODIndex;
 
 				FMeshGenerationDynamicOptions DynamicOptions = InOptions.GetResult();
-					
+
+				// Generate the code for the ranges
+				for (int32 RangeIndex = 0; RangeIndex < Node->Ranges.Num(); ++RangeIndex)
+				{
+					FRangeGenerationResult RangeResult;
+					GenerateRange(RangeResult, StaticOptions, Node->Ranges[RangeIndex]);
+					Op->Ranges.Emplace(Op.get(), RangeResult.sizeOp, RangeResult.rangeName, RangeResult.rangeUID);
+				}
+
 				return Result;
 			},
 			Requisites
diff --git a/Engine/Plugins/Mutable/Source/MutableTools/Private/MuT/CodeGenerator_SkeletalMesh.cpp b/Engine/Plugins/Mutable/Source/MutableTools/Private/MuT/CodeGenerator_SkeletalMesh.cpp
deleted file mode 100644
diff --git a/Engine/Plugins/Mutable/Source/MutableTools/Private/MuT/CodeOptimiser.cpp b/Engine/Plugins/Mutable/Source/MutableTools/Private/MuT/CodeOptimiser.cpp
--- a/Engine/Plugins/Mutable/Source/MutableTools/Private/MuT/CodeOptimiser.cpp
+++ b/Engine/Plugins/Mutable/Source/MutableTools/Private/MuT/CodeOptimiser.cpp
@@ -548,7 +548,7 @@ namespace UE::Mutable::Private
 					case EOpType::CO_PARAMETER:
 					case EOpType::PR_PARAMETER:
 					case EOpType::IM_PARAMETER:
-					case EOpType::SK_PARAMETER:
+					case EOpType::ME_PARAMETER:
 					case EOpType::MA_PARAMETER:
 					case EOpType::MI_PARAMETER:
 					case EOpType::IM_PARAMETER_FROM_MATERIAL:
diff --git a/Engine/Plugins/Mutable/Source/MutableTools/Private/MuT/CodeOptimiser_Parameters.cpp b/Engine/Plugins/Mutable/Source/MutableTools/Private/MuT/CodeOptimiser_Parameters.cpp
--- a/Engine/Plugins/Mutable/Source/MutableTools/Private/MuT/CodeOptimiser_Parameters.cpp
+++ b/Engine/Plugins/Mutable/Source/MutableTools/Private/MuT/CodeOptimiser_Parameters.cpp
@@ -150,7 +150,7 @@ namespace UE::Mutable::Private
                     case EOpType::CO_PARAMETER:
                     case EOpType::PR_PARAMETER:
 					case EOpType::IM_PARAMETER:
-					case EOpType::SK_PARAMETER:
+					case EOpType::ME_PARAMETER:
 					case EOpType::MA_PARAMETER:
 					{
 						const ASTOpParameter* Typed = static_cast<const ASTOpParameter*>(Op.get());
@@ -1580,7 +1580,7 @@ namespace UE::Mutable::Private
                 case EOpType::CO_PARAMETER:
                 case EOpType::PR_PARAMETER:
 				case EOpType::IM_PARAMETER:
-				case EOpType::SK_PARAMETER:
+				case EOpType::ME_PARAMETER:
 				case EOpType::MA_PARAMETER:
 				case EOpType::IS_PARAMETER:
 				{
diff --git a/Engine/Plugins/Mutable/Source/MutableTools/Private/MuT/Compiler.cpp b/Engine/Plugins/Mutable/Source/MutableTools/Private/MuT/Compiler.cpp
--- a/Engine/Plugins/Mutable/Source/MutableTools/Private/MuT/Compiler.cpp
+++ b/Engine/Plugins/Mutable/Source/MutableTools/Private/MuT/Compiler.cpp
@@ -75,7 +75,7 @@ namespace UE::Mutable::Private
 		{ false,	{ 0, 0, 0, 0, 0, 0, 0, 0, 0, 0, 0, 0, 0, 0, 0, 0, 0, 0, 0, 0, 0, 0, 0, 0, 0, 0, 0, 0, 0 }		},	// CO_PARAMETER
 		{ false,	{ 0, 0, 0, 0, 0, 0, 0, 0, 0, 0, 0, 0, 0, 0, 0, 0, 0, 0, 0, 0, 0, 0, 0, 0, 0, 0, 0, 0, 0 }		},	// PR_PARAMETER
 		{ false,	{ 0, 0, 0, 0, 0, 0, 0, 0, 0, 0, 0, 0, 0, 0, 0, 0, 0, 0, 0, 0, 0, 0, 0, 0, 0, 0, 0, 0, 0 }		},	// IM_PARAMETER
-		{ false,	{ 0, 0, 0, 0, 0, 0, 0, 0, 0, 0, 0, 0, 0, 0, 0, 0, 0, 0, 0, 0, 0, 0, 0, 0, 0, 0, 0, 0, 0 }		},	// SK_PARAMETER
+		{ false,	{ 0, 0, 0, 0, 0, 0, 0, 0, 0, 0, 0, 0, 0, 0, 0, 0, 0, 0, 0, 0, 0, 0, 0, 0, 0, 0, 0, 0, 0 }		},	// ME_PARAMETER
 		{ false,	{ 0, 0, 0, 0, 0, 0, 0, 0, 0, 0, 0, 0, 0, 0, 0, 0, 0, 0, 0, 0, 0, 0, 0, 0, 0, 0, 0, 0, 0 }		},	// ST_PARAMETER
 		{ false,	{ 0, 0, 0, 0, 0, 0, 0, 0, 0, 0, 0, 0, 0, 0, 0, 0, 0, 0, 0, 0, 0, 0, 0, 0, 0, 0, 0, 0, 0 }		},	// MA_PARAMETER
 		{ false,	{ 0, 0, 0, 0, 0, 0, 0, 0, 0, 0, 0, 0, 0, 0, 0, 0, 0, 0, 0, 0, 0, 0, 0, 0, 0, 0, 0, 0, 0 }		},	// MI_PARAMETER
@@ -179,7 +179,6 @@ namespace UE::Mutable::Private
 		{ true,		{ 0, 0, 0, 0, 0, 0, 0, 0, 0, 0, 0, 0, 0, 0, 0, 0, 0, 0, 0, 0, 0, 0, 0, 0, 0, 0, 0, 0, 0 }		},	// ME_TRANSFORMWITHMESH
 		{ true,		{ 0, 0, 0, 0, 0, 0, 0, 0, 0, 0, 0, 0, 0, 0, 0, 0, 0, 0, 0, 0, 0, 0, 0, 0, 0, 0, 0, 0, 0 }		},	// ME_TRANSFORMWITHBONE
 		{ true,		{ 0, 0, 0, 0, 0, 0, 0, 0, 0, 0, 0, 0, 0, 0, 0, 0, 0, 0, 0, 0, 0, 0, 0, 0, 0, 0, 0, 0, 0 }		},	// ME_EXTERNAL
-		{ true,		{ 0, 0, 0, 0, 0, 0, 0, 0, 0, 0, 0, 0, 0, 0, 0, 0, 0, 0, 0, 0, 0, 0, 0, 0, 0, 0, 0, 0, 0 }		},	// ME_SKELETALMESH_BREAK
 
 		{ false,	{ 0, 0, 0, 0, 0, 0, 0, 0, 0, 0, 0, 0, 0, 0, 0, 0, 0, 0, 0, 0, 0, 0, 0, 0, 0, 0, 0, 0, 0 }		},	// IN_ADDMESH
 		{ false,	{ 0, 0, 0, 0, 0, 0, 0, 0, 0, 0, 0, 0, 0, 0, 0, 0, 0, 0, 0, 0, 0, 0, 0, 0, 0, 0, 0, 0, 0 }		},	// IN_ADDIMAGE
@@ -481,13 +480,26 @@ namespace UE::Mutable::Private
 
 				// Gather the parameter list from the non-optimized data, so that we have them all even if they are optimized out
 				{
-					const int32 ParameterCount = Generator.FirstPass.ParameterNodes.GenericParametersCache.Num();
+					const int32 ParameterCount = 
+						Generator.FirstPass.ParameterNodes.GenericParametersCache.Num()
+						+
+						Generator.FirstPass.ParameterNodes.MeshParametersCache.Num();
 
 					Parameters.Reserve(ParameterCount);
 					for (const TPair<Ptr<const Node>, Ptr<ASTOpParameter>>& Entry : Generator.FirstPass.ParameterNodes.GenericParametersCache)
 					{
 						Parameters.Add(Entry.Value->Parameter);
 					}
+					for (const TPair<Ptr<const NodeMeshParameter>, TArray<TPair<Ptr<ASTOpParameter>,FMeshGenerationResult>>>& Entry : Generator.FirstPass.ParameterNodes.MeshParametersCache)
+					{
+						for (const TPair<Ptr<ASTOpParameter>, FMeshGenerationResult>& ArrayEntry : Entry.Value)
+						{
+							if (ArrayEntry.Key)
+							{
+								Parameters.AddUnique(ArrayEntry.Key->Parameter);
+							}
+						}
+					}
 				}
 
 				// Sort the parameters as deterministically as possible.
@@ -532,6 +544,17 @@ namespace UE::Mutable::Private
 										MessageContexts.AddUnique(Entry.Key->GetMessageContext());
 									}
 								}
+								for (const TPair<Ptr<const NodeMeshParameter>, TArray<TPair<Ptr<ASTOpParameter>, FMeshGenerationResult>>>& Entry : Generator.FirstPass.ParameterNodes.MeshParametersCache)
+								{
+									for (const TPair<Ptr<ASTOpParameter>, FMeshGenerationResult>& ArrayEntry : Entry.Value)
+									{
+										if (ArrayEntry.Key && ArrayEntry.Key->Parameter.Name.Compare(CurrentParameterName) == 0)
+										{
+											check(Entry.Key->GetMessageContext());
+											MessageContexts.AddUnique(Entry.Key->GetMessageContext());
+										}
+									}
+								}
 							}
 
 							// The "ensure" will fail if:
diff --git a/Engine/Plugins/Mutable/Source/MutableTools/Private/MuT/ErrorLog.cpp b/Engine/Plugins/Mutable/Source/MutableTools/Private/MuT/ErrorLog.cpp
--- a/Engine/Plugins/Mutable/Source/MutableTools/Private/MuT/ErrorLog.cpp
+++ b/Engine/Plugins/Mutable/Source/MutableTools/Private/MuT/ErrorLog.cpp
@@ -220,7 +220,7 @@ namespace UE::Mutable::Private
 		TEXT("CO_PARAMETER     "),
 		TEXT("PR_PARAMETER     "),
 		TEXT("IM_PARAMETER     "),
-		TEXT("SK_PARAMETER     "),
+		TEXT("ME_PARAMETER     "),
 		TEXT("ST_PARAMETER     "),
 		TEXT("MA_PARAMETER     "),
 		TEXT("MI_PARAMETER     "),
@@ -324,7 +324,6 @@ namespace UE::Mutable::Private
 		TEXT("ME_TRANSFORMWITHMESH"),
 		TEXT("ME_TRANSFORMWITHBONE"),
     	TEXT("ME_EXTERNAL      "),
-		TEXT("ME_SKELETALMESH_BREAK"),
 
 		TEXT("IN_ADDMESH       "),
 		TEXT("IN_ADDIMAGE      "),
diff --git a/Engine/Plugins/Mutable/Source/MutableTools/Private/MuT/Node.cpp b/Engine/Plugins/Mutable/Source/MutableTools/Private/MuT/Node.cpp
--- a/Engine/Plugins/Mutable/Source/MutableTools/Private/MuT/Node.cpp
+++ b/Engine/Plugins/Mutable/Source/MutableTools/Private/MuT/Node.cpp
@@ -69,7 +69,7 @@
 #include "MuT/NodeMeshClipMorphPlane.h"
 #include "MuT/NodeMeshClipDeform.h"
 #include "MuT/NodeMeshClipWithMesh.h"
-#include "MuT/NodeMeshSkeletalMeshBreak.h"
+#include "MuT/NodeMeshParameter.h"
 #include "MuT/NodeMeshMakeMorph.h"
 #include "MuT/NodeMeshApplyPose.h"
 #include "MuT/NodeMeshTransform.h"
@@ -113,8 +113,6 @@
 #include "MuT/NodeMaterialParameter.h"
 #include "MuT/NodeMeshExternal.h"
 #include "MuT/NodeScalarExternal.h"
-#include "MuT/NodeSkeletalMesh.h"
-#include "MuT/NodeSkeletalMeshParameter.h"
 
 
 namespace UE::Mutable::Private
@@ -160,10 +158,6 @@ namespace UE::Mutable::Private
 	FNodeType NodeExternalParameter::StaticType = FNodeType(Node::EType::ExternalParameter, NodeExternal::GetStaticType());
 	
 	FNodeType NodeLOD::StaticType = FNodeType(Node::EType::LOD, Node::GetStaticType());
-
-	FNodeType NodeSkeletalMesh::StaticType = FNodeType(Node::EType::SkeletalMesh, Node::GetStaticType());
-	FNodeType NodeSkeletalMeshParameter::StaticType = FNodeType(Node::EType::SkeletalMeshParameter, NodeSkeletalMesh::GetStaticType());
-
 	FNodeType NodeExtensionData::StaticType = FNodeType(Node::EType::ExtensionData, Node::GetStaticType());
 	FNodeType NodeExtensionDataConstant::StaticType = FNodeType(Node::EType::ExtensionDataConstant, NodeExtensionData::GetStaticType());
 
@@ -213,7 +207,7 @@ namespace UE::Mutable::Private
 	FNodeType NodeMeshClipMorphPlane::StaticType = FNodeType(Node::EType::MeshClipMorphPlane, NodeMesh::GetStaticType());
 	FNodeType NodeMeshClipDeform::StaticType = FNodeType(Node::EType::MeshClipDeform, NodeMesh::GetStaticType());
 	FNodeType NodeMeshClipWithMesh::StaticType = FNodeType(Node::EType::MeshClipWithMesh, NodeMesh::GetStaticType());
-	FNodeType NodeMeshSkeletalMeshBreak::StaticType = FNodeType(Node::EType::MeshSkeletalMeshBreak, NodeMesh::GetStaticType());
+	FNodeType NodeMeshParameter::StaticType = FNodeType(Node::EType::MeshParameter, NodeMesh::GetStaticType());
 	FNodeType NodeMeshMakeMorph::StaticType = FNodeType(Node::EType::MeshMakeMorph, NodeMesh::GetStaticType());
 	FNodeType NodeMeshApplyPose::StaticType = FNodeType(Node::EType::MeshApplyPose, NodeMesh::GetStaticType());
 	FNodeType NodeMeshTransform::StaticType = FNodeType(Node::EType::MeshTransform, NodeMesh::GetStaticType());
