[mutable] Fixed multiple SKM crashes due to an issue with the name scheme used by mutable + the process of making resource names unique (UE).
- Simplified the "MakeNameUnique" approach as it was failing in some cases.
- All mutable generated resources will have unique names.
- Updated baking pipeline code so it takes into consideration that now the textures will have a unique suffix.
- Other small changes.

diff --git a/Engine/Plugins/Mutable/Source/CustomizableObject/Internal/MuCO/CustomizableObjectInstancePrivate.h b/Engine/Plugins/Mutable/Source/CustomizableObject/Internal/MuCO/CustomizableObjectInstancePrivate.h
--- a/Engine/Plugins/Mutable/Source/CustomizableObject/Internal/MuCO/CustomizableObjectInstancePrivate.h
+++ b/Engine/Plugins/Mutable/Source/CustomizableObject/Internal/MuCO/CustomizableObjectInstancePrivate.h
@@ -319,8 +319,7 @@ public:
 	UE_API FCustomizableInstanceComponentData* GetComponentData(FCustomizableObjectComponentIndex ObjectComponentIndex);
 	UE_API const FCustomizableInstanceComponentData* GetComponentData(FCustomizableObjectComponentIndex ObjectComponentIndex) const;
 	
-	// Create a transient texture and add it to the TextureTrackerArray
-	UE_API UTexture2D* CreateTexture(const FString& TextureName);
+	UE_API UTexture2D* CreateTexture(const FName& TextureName);
 
 	UE_API void InvalidateGeneratedData();
 
diff --git a/Engine/Plugins/Mutable/Source/CustomizableObject/Private/MuCO/CustomizableObjectInstance.cpp b/Engine/Plugins/Mutable/Source/CustomizableObject/Private/MuCO/CustomizableObjectInstance.cpp
--- a/Engine/Plugins/Mutable/Source/CustomizableObject/Private/MuCO/CustomizableObjectInstance.cpp
+++ b/Engine/Plugins/Mutable/Source/CustomizableObject/Private/MuCO/CustomizableObjectInstance.cpp
@@ -159,27 +159,6 @@ bool FMutableMaterialPlaceholder::operator==(const FMutableMaterialPlaceholder&
 }
 
 
-/**
- * Makes the provided name unique and takes care of not losing the hashed value that could be present at the end of the string
- * @warn This should only be used for mutable resources as they have the hash of their contents at the end of the string. 
- * @param ObjectName The name to be given to a UObject that was generated with data from a mu object
- * @param ResourceHash The hash of the resource used to create the UAsset you want to make its name unique.
- * @param ObjectClass The class of the asset you want to make its name unique
- */
-void MakeMutableGeneratedObjectNameUnique(FString& ObjectName, const uint32 ResourceHash ,const UClass* ObjectClass)
-{
-#if WITH_EDITOR
-	const FString HashAsString = FString::Printf(TEXT("%u"), ResourceHash);
-	ensureMsgf(ObjectName.Contains(HashAsString), TEXT("Unexpected mutable object name \"%s\" : The name was expected to end with the resource hash value of \"%u\""), *ObjectName, ResourceHash);
-#endif
-	
-	ObjectName.Append(TEXT("_"));
-	const int32 IndexOfAppendedChar = ObjectName.Len() - 1;
-	ObjectName = MakeUniqueObjectName(GetTransientPackage(), ObjectClass, *ObjectName).ToString();
-	ObjectName.RemoveAt(IndexOfAppendedChar);
-}
-
-
 uint32 GetTypeHash(const FMutableMaterialPlaceholder& PlaceHolder)
 {
 	return GetHash(PlaceHolder, false);
@@ -233,11 +212,11 @@ uint32 GetHash(const FMutableMaterialPlaceholder& PlaceHolder, bool bPersistent)
 }
 
 
-UTexture2D* UCustomizableInstancePrivate::CreateTexture(const FString& TextureName)
+UTexture2D* UCustomizableInstancePrivate::CreateTexture(const FName& TextureName)
 {
 	UTexture2D* NewTexture = NewObject<UTexture2D>(
 		GetTransientPackage(),
-		GetData(TextureName),
+		TextureName,
 		RF_Transient
 		);
 	UCustomizableObjectSystem::GetInstance()->GetPrivate()->LogBenchmarkUtil.AddTexture(*NewTexture);
@@ -2286,22 +2265,6 @@ bool UCustomizableInstancePrivate::UpdateSkeletalMesh_PostBeginUpdate0(UCustomiz
 			Context->FirstLODAvailable[ComponentName] :
 			Context->GetFirstRequestedLOD()[ComponentName];
 		
-		uint32 MeshHash = INDEX_NONE;
-		for (int32 LODIndex = FirstLOD; LODIndex < Component->LODCount && LODIndex < MAX_MESH_LOD_COUNT; ++LODIndex)
-		{
-			const UE::Mutable::Private::FMeshId& MeshId = (*MeshDescriptors)[LODIndex];
-			const uint32 MeshIDHash = UE::Mutable::Private::GetTypeHashPersistent(MeshId);
-
-			if (!Context->bBake)
-			{
-				MeshHash = HashCombineFast(MeshHash, MeshIDHash);
-			}
-			else
-			{
-				MeshHash = HashCombine(MeshHash, MeshIDHash);
-			}
-		}
-
 		if (Context->bUseMeshCache)
 		{
 			if (USkeletalMesh* CachedMesh = CustomizableObject->GetPrivate()->MeshCache.Get(*MeshDescriptors))
@@ -2350,7 +2313,23 @@ bool UCustomizableInstancePrivate::UpdateSkeletalMesh_PostBeginUpdate0(UCustomiz
 
 		USkeletalMesh* SkeletalMesh = nullptr;
 		{
-			FString SkeletalMeshName = FString::Printf (TEXT("SK_%s_%s_%lu"), *Instance->GetCustomizableObject()->GetName(), *ComponentName.ToString(), MeshHash);
+			uint32 MeshHash = INDEX_NONE;
+			for (int32 LODIndex = FirstLOD; LODIndex < Component->LODCount && LODIndex < MAX_MESH_LOD_COUNT; ++LODIndex)
+			{
+				const UE::Mutable::Private::FMeshId& MeshId = (*MeshDescriptors)[LODIndex];
+				const uint32 MeshIDHash = UE::Mutable::Private::GetTypeHashPersistent(MeshId);
+
+				if (!Context->bBake)
+				{
+					MeshHash = HashCombineFast(MeshHash, MeshIDHash);
+				}
+				else
+				{
+					MeshHash = HashCombine(MeshHash, MeshIDHash);
+				}
+			}
+
+			FString SkeletalMeshName = FString::Printf(TEXT("SK_%s_%s_h%lu"), *Instance->GetCustomizableObject()->GetName(), *ComponentName.ToString(), MeshHash);
 
 #if WITH_EDITOR
 			if (Context->bBake)
@@ -2360,16 +2339,16 @@ bool UCustomizableInstancePrivate::UpdateSkeletalMesh_PostBeginUpdate0(UCustomiz
 #endif
 
 			// Make name unique to avoid collisions with other objects.
-			MakeMutableGeneratedObjectNameUnique(SkeletalMeshName, MeshHash, USkeletalMesh::StaticClass());
+			const FName UniqueSkeletalMeshName = MakeUniqueObjectName(GetTransientPackage(), USkeletalMesh::StaticClass(), *SkeletalMeshName);
 			
 			if (!Context->bBake)
 			{
-				SkeletalMesh = NewObject<UCustomizableObjectSkeletalMesh>(GetTransientPackage(), FName(*SkeletalMeshName), RF_Transient);
+				SkeletalMesh = NewObject<UCustomizableObjectSkeletalMesh>(GetTransientPackage(), UniqueSkeletalMeshName, RF_Transient);
 			}
 #if WITH_EDITOR
 			else
 			{
-				SkeletalMesh = NewObject<USkeletalMesh>(GetTransientPackage(), FName(*SkeletalMeshName), RF_Transient);
+				SkeletalMesh = NewObject<USkeletalMesh>(GetTransientPackage(), UniqueSkeletalMeshName, RF_Transient);
 			}
 #endif
 				
@@ -6088,16 +6067,11 @@ void UCustomizableInstancePrivate::BuildMaterials(const TSharedRef<FUpdateContex
 								FString MutableTextureName;
 								{
 									const uint32 MutableImageHash = UE::Mutable::Private::GetTypeHashPersistent(Image.ImageID);
-									MutableTextureName = FString::Printf(TEXT("T_%s_%lu"), *Props.TextureParameterName, MutableImageHash);
+									MutableTextureName = FString::Printf(TEXT("T_%s_h%llu"), *Props.TextureParameterName, MutableImageHash);
 									MutableTextureName.ReplaceInline(TEXT(" "), TEXT("_"));
 
-									// If not for bake work with this resource as normal and make it unique to avoid collisions
-									if (!Context->bBake)
-									{
-										MakeMutableGeneratedObjectNameUnique(MutableTextureName, MutableImageHash, UMaterialInstanceDynamic::StaticClass());
-									}
 #if WITH_EDITOR
-									else
+									if (Context->bBake)
 									{
 										// If this is a bake operation just add the baked prefix to make the name unique for all baked instances of this resource
 										MutableTextureName = FBakingConfiguration::BakedResourcePrefix + MutableTextureName;
@@ -6105,6 +6079,8 @@ void UCustomizableInstancePrivate::BuildMaterials(const TSharedRef<FUpdateContex
 #endif
 								}
 								
+								const FName UniqueMutableTextureName = MakeUniqueObjectName(GetTransientPackage(), UTexture2D::StaticClass(), *MutableTextureName);
+
 								if (ReusedTexture && (*ReusedTexture).IsValid() && !(*ReusedTexture)->HasAnyFlags(RF_BeginDestroyed))
 								{
 									// Only uncompressed textures can be reused. This also fixes an issue in the editor where textures supposedly 
@@ -6120,7 +6096,7 @@ void UCustomizableInstancePrivate::BuildMaterials(const TSharedRef<FUpdateContex
 									else
 									{
 										ReusedTexture = nullptr;
-										MutableTexture = CreateTexture(MutableTextureName);
+										MutableTexture = CreateTexture(UniqueMutableTextureName);
 #if WITH_EDITOR
 										UE_LOG(LogMutable, Warning,
 											TEXT("Tried to reuse an uncompressed texture with name %s. Make sure the selected Mutable state disables texture compression/streaming, that one of the state's runtime parameters affects the texture and that the CO is compiled with max. optimization settings."),
@@ -6131,7 +6107,7 @@ void UCustomizableInstancePrivate::BuildMaterials(const TSharedRef<FUpdateContex
 								else
 								{
 									ReusedTexture = nullptr;
-									MutableTexture = CreateTexture(MutableTextureName);
+									MutableTexture = CreateTexture(UniqueMutableTextureName);
 								}
 
 								if (MutableTexture)
@@ -6324,9 +6300,10 @@ void UCustomizableInstancePrivate::BuildMaterials(const TSharedRef<FUpdateContex
 					if (!MaterialInstance && MutableMaterialPlaceholder.GetParams().Num() != 0)
 					{
 #if WITH_EDITOR
-						// Remove the MI_ or M_ prefixes from the material string so we use it as the name of the MID
+						
 						FString MIDName = MaterialTemplate->GetName();
 						{
+							// Remove the MI_ or M_ prefixes from the material string so we use it as the name of the MID
 							const FString MaterialPrefix = TEXT("M_");						// Material
 							const FString MaterialInstancePrefix = TEXT("MI_");				// Material Instance
 							const FString MaterialInstanceConstantPrefix = TEXT("MIC_");	// Material Instance Constant
@@ -6345,19 +6322,17 @@ void UCustomizableInstancePrivate::BuildMaterials(const TSharedRef<FUpdateContex
 							}
 
 							const uint32 MaterialPlaceHolderHash = GetHash(MutableMaterialPlaceholder, Context->bBake);
-							MIDName = FString::Printf(TEXT("%s%s_%u"),TEXT("MID_"),*MIDName, MaterialPlaceHolderHash);
+							MIDName = FString::Printf(TEXT("MID_%s_h%lu"), *MIDName, MaterialPlaceHolderHash);
 
 							// Add a "tag" to better identify this asset during the bake
 							if (Context->bBake)
 							{
 								MIDName = FBakingConfiguration::BakedResourcePrefix + MIDName;
 							}
-
-							MakeMutableGeneratedObjectNameUnique(MIDName, MaterialPlaceHolderHash, UMaterialInstanceDynamic::StaticClass());
-							check(MIDName.Contains(FString::Printf(TEXT("%u"), MaterialPlaceHolderHash)))
 						}
 						
-						MaterialInstance = UMaterialInstanceDynamic::Create(MaterialTemplate, GetTransientPackage(), FName(MIDName));
+						const FName UniqueMIDName = MakeUniqueObjectName(GetTransientPackage(), UMaterialInstanceDynamic::StaticClass(), *MIDName);
+						MaterialInstance = UMaterialInstanceDynamic::Create(MaterialTemplate, GetTransientPackage(), UniqueMIDName);
 #else
 						MaterialInstance = UMaterialInstanceDynamic::Create(MaterialTemplate, GetTransientPackage());
 #endif
diff --git a/Engine/Plugins/Mutable/Source/CustomizableObjectEditor/Private/MuCOE/CustomizableObjectInstanceBakingUtils.cpp b/Engine/Plugins/Mutable/Source/CustomizableObjectEditor/Private/MuCOE/CustomizableObjectInstanceBakingUtils.cpp
--- a/Engine/Plugins/Mutable/Source/CustomizableObjectEditor/Private/MuCOE/CustomizableObjectInstanceBakingUtils.cpp
+++ b/Engine/Plugins/Mutable/Source/CustomizableObjectEditor/Private/MuCOE/CustomizableObjectInstanceBakingUtils.cpp
@@ -202,28 +202,18 @@ FString FResourceBakingData::GetUniqueAssetNameSuffixSection() const
 
 
 /**
- * Removes the "_n" string found at the end of the material and skeletal mesh name.
+ * Removes the "_n" string found at the end of the resource and skeletal mesh name.
  * @warning It does not check if that part is there or not.
- * @param InMaterialName The name of the material that you want not to be unique.
+ * @param InMutableResourceName The name of the resource that you want not to be unique.
  */
-void MakeAssetNameNotUnique(FString& InMaterialName)
+void MakeAssetNameNotUnique(FString& InMutableResourceName)
 {
-	// This should be a material generated by mutable so try to remove the value that makes it unique
+	// This should be a resource generated by mutable so try to remove the value that makes it unique
 	int32 UniquenessValueSeparatorIndex = INDEX_NONE;
-	InMaterialName.FindLastChar(TEXT('_'),UniquenessValueSeparatorIndex);
+	InMutableResourceName.FindLastChar(TEXT('_'),UniquenessValueSeparatorIndex);
 	if (UniquenessValueSeparatorIndex != INDEX_NONE)
 	{
-		InMaterialName = InMaterialName.Left(UniquenessValueSeparatorIndex);
-	}
-						
-	// Ensure that the new name contains numerical values representing the hash (speculative)
-	int32 HashValueSeparatorIndex = INDEX_NONE;
-	InMaterialName.FindLastChar(TEXT('_'),HashValueSeparatorIndex);
-	if (ensure(HashValueSeparatorIndex != INDEX_NONE))
-	{
-		const FString HashValue = InMaterialName.RightChop(HashValueSeparatorIndex+1);
-		// If this fails it means that the name of the material did not contain or the hash or the value that makes it unique.
-		check(HashValue.IsNumeric())
+		InMutableResourceName = InMutableResourceName.Left(UniquenessValueSeparatorIndex);
 	}
 }
 
@@ -642,7 +632,10 @@ bool BakeCustomizableObjectInstance(
 							
 							// Update the name adding to it the name of the component alongside other data
 							FResourceBakingData TextureData {TextureAssetPrefix, UserGivenFilename, SourceTexture->GetName(), FString{}, AssetPath};
-							TryRemovePrefix(TextureData.ResourceName, {FBakingConfiguration::BakedResourcePrefix});
+							if (TryRemovePrefix(TextureData.ResourceName, {FBakingConfiguration::BakedResourcePrefix}))
+							{
+								MakeAssetNameNotUnique(TextureData.ResourceName);
+							}
 							TryRemovePrefix(TextureData.ResourceName, {TextureAssetPrefix});
 							
 							// Experimental
@@ -880,7 +873,10 @@ bool BakeCustomizableObjectInstance(
 										// The source texture has not yet been processed so duplicate it
 										
 										FResourceBakingData TextureData {TextureAssetPrefix, UserGivenFilename, SourceTexture->GetName(), FString{}, AssetPath};
-										TryRemovePrefix(TextureData.ResourceName ,{FBakingConfiguration::BakedResourcePrefix});
+										if (TryRemovePrefix(TextureData.ResourceName ,{FBakingConfiguration::BakedResourcePrefix}))
+										{
+											MakeAssetNameNotUnique(TextureData.ResourceName);
+										}
 										TryRemovePrefix(TextureData.ResourceName , {TextureAssetPrefix});
 										
 										// Check if the file is already in disk, so we can later decide what we do with that file
