[Mutable] Added Skeletal Mesh pin and type.

- Added IN_ADDSKELETALMESH.
- Removed :ME_REFERENCE.
- Added CodeGenerator_SkeletalMesh.
- Added GenerateMutableSourceSkeletalMesh.
- Added SKM Passthrough pin to Node Skeletal Mesh Parameter.
- Added SKM Passthrough pin to Passthrough Skeletal Mesh Component.

diff --git a/Engine/Plugins/Mutable/Source/CustomizableObject/Internal/MuCO/CustomizableObjectCustomVersion.h b/Engine/Plugins/Mutable/Source/CustomizableObject/Internal/MuCO/CustomizableObjectCustomVersion.h
--- a/Engine/Plugins/Mutable/Source/CustomizableObject/Internal/MuCO/CustomizableObjectCustomVersion.h
+++ b/Engine/Plugins/Mutable/Source/CustomizableObject/Internal/MuCO/CustomizableObjectCustomVersion.h
@@ -208,6 +208,8 @@ struct FCustomizableObjectCustomVersion
 
 		NodeTextureParameterPassthroughPin,
 
+		PinSkeletalMesh,
+
 		// -----<new versions can be added above this line>-------------------------------------------------
 		VersionPlusOne,
 		LatestVersion = VersionPlusOne - 1
diff --git a/Engine/Plugins/Mutable/Source/CustomizableObject/Internal/MuCO/CustomizableObjectInstancePrivate.h b/Engine/Plugins/Mutable/Source/CustomizableObject/Internal/MuCO/CustomizableObjectInstancePrivate.h
--- a/Engine/Plugins/Mutable/Source/CustomizableObject/Internal/MuCO/CustomizableObjectInstancePrivate.h
+++ b/Engine/Plugins/Mutable/Source/CustomizableObject/Internal/MuCO/CustomizableObjectInstancePrivate.h
@@ -144,10 +144,7 @@ struct FCustomizableInstanceComponentData
 	/** Clothing PhysicsAssets required by the current generated instance. PhysicsAssets to be loaded and merged.*/
 	TArray<TPair<int32, int32>> ClothingPhysicsAssetsToStream;
 
-	/** Array of generated MeshIds per each LOD, used to decide if the mesh should be updated or not.
-	 *  Size == NumLODsAvailable
-	 *  LODs without mesh will be set to the maximum value of FResourceID (Max_uint64). */
-	TArray<UE::Mutable::Private::FMeshId> LastMeshIdPerLOD;
+	FComponentDescriptor LastComponentDescriptor;
 
 	UPROPERTY(Transient)
 	TArray<TObjectPtr<UMaterialInterface>> OverrideMaterials;
diff --git a/Engine/Plugins/Mutable/Source/CustomizableObject/Internal/MuCO/CustomizableObjectMipDataProvider.h b/Engine/Plugins/Mutable/Source/CustomizableObject/Internal/MuCO/CustomizableObjectMipDataProvider.h
--- a/Engine/Plugins/Mutable/Source/CustomizableObject/Internal/MuCO/CustomizableObjectMipDataProvider.h
+++ b/Engine/Plugins/Mutable/Source/CustomizableObject/Internal/MuCO/CustomizableObjectMipDataProvider.h
@@ -47,6 +47,7 @@ public:
 	TSharedPtr<UE::Mutable::Private::FMeshIdRegistry> MeshIdRegistry;
 	TSharedPtr<UE::Mutable::Private::FImageIdRegistry> ImageIdRegistry;
 	TSharedPtr<UE::Mutable::Private::FMaterialIdRegistry> MaterialIdRegistry;
+	TSharedPtr<UE::Mutable::Private::FSkeletalMeshIdRegistry> SkeletalMeshIdRegistry;
 
 	TSharedPtr<FUnrealMutableResourceProvider> ExternalResourceProvider;
 
diff --git a/Engine/Plugins/Mutable/Source/CustomizableObject/Internal/MuCO/CustomizableObjectPrivate.h b/Engine/Plugins/Mutable/Source/CustomizableObject/Internal/MuCO/CustomizableObjectPrivate.h
--- a/Engine/Plugins/Mutable/Source/CustomizableObject/Internal/MuCO/CustomizableObjectPrivate.h
+++ b/Engine/Plugins/Mutable/Source/CustomizableObject/Internal/MuCO/CustomizableObjectPrivate.h
@@ -1970,6 +1970,8 @@ public:
 
 	TSharedRef<UE::Mutable::Private::FMaterialIdRegistry> MaterialIdRegistry = MakeShared<UE::Mutable::Private::FMaterialIdRegistry>();
 
+	TSharedRef<UE::Mutable::Private::FSkeletalMeshIdRegistry> SkeletalMeshIdRegistry = MakeShared<UE::Mutable::Private::FSkeletalMeshIdRegistry>();
+
 	TSharedPtr<UE::Mutable::Private::FModel> SerializedModel;
 	
 	// See UCustomizableObjectSystem::LockObject. Must only be modified from the game thread
diff --git a/Engine/Plugins/Mutable/Source/CustomizableObject/Internal/MuCO/CustomizableObjectSystemPrivate.h b/Engine/Plugins/Mutable/Source/CustomizableObject/Internal/MuCO/CustomizableObjectSystemPrivate.h
--- a/Engine/Plugins/Mutable/Source/CustomizableObject/Internal/MuCO/CustomizableObjectSystemPrivate.h
+++ b/Engine/Plugins/Mutable/Source/CustomizableObject/Internal/MuCO/CustomizableObjectSystemPrivate.h
@@ -456,6 +456,10 @@ struct FInstanceUpdateData
 		/** Overlay material. Only one per mesh, if any. */
 		UE::Mutable::Private::FMaterialId OverlayMaterialId;
 		TSharedPtr<const UE::Mutable::Private::FMaterial> OverlayMaterial;
+
+		UE::Mutable::Private::FSkeletalMeshId SkeletalMeshId;
+		
+		TSharedPtr<const UE::Mutable::Private::FSkeletalMesh> SkeletalMesh;
 	};
 
 	TArray<FComponent> Components;
@@ -522,6 +526,13 @@ struct FInstanceUpdateData
 };
 
 
+struct FComponentDescriptor
+{
+	TArray<UE::Mutable::Private::FMeshId> MeshDescriptors;
+	UE::Mutable::Private::FSkeletalMeshId SkeletalMeshDescriptor;
+};
+
+
 /** Update Context.
  *
  * Alive from the start to the end of the update (both API and LOD update). */
@@ -626,6 +637,7 @@ public:
 	TSharedPtr<UE::Mutable::Private::FMeshIdRegistry> MeshIdRegistry;
 	TSharedPtr<UE::Mutable::Private::FImageIdRegistry> ImageIdRegistry;
 	TSharedPtr<UE::Mutable::Private::FMaterialIdRegistry> MaterialIdRegistry;
+	TSharedPtr<UE::Mutable::Private::FSkeletalMeshIdRegistry> SkeletalMeshIdRegistry;
 
 	TSharedPtr<FUnrealMutableResourceProvider> ExternalResourceProvider;
 
@@ -662,15 +674,12 @@ public:
 	UE::Mutable::Private::FImageOperator::FImagePixelFormatFunc PixelFormatOverride;
 
 private:
-
-	/** Mutable Meshes required for each component. Outermost index is the object component index, inner index is the LOD. */
-	TArray<TArray<UE::Mutable::Private::FMeshId>> MeshDescriptors;
+	TArray<FComponentDescriptor> ComponentDescriptors;
 
 public:
-
-	void InitMeshDescriptors(int32 Size);
-	const TArray<TArray<UE::Mutable::Private::FMeshId>>& GetMeshDescriptors() const;
-	TArray<UE::Mutable::Private::FMeshId>* GetMeshDescriptors(FCustomizableObjectComponentIndex Index);
+	void InitComponentDescriptors(int32 Size);
+	const TArray<FComponentDescriptor>& GetComponentDescriptors() const;
+	FComponentDescriptor* GetComponentDescriptor(FCustomizableObjectComponentIndex Index);
 
 	/** Used to know if the updated instances' meshes are different from the previous ones. 
 	  * The index of the array is the instance component's index.
diff --git a/Engine/Plugins/Mutable/Source/CustomizableObject/Internal/MuCO/UnrealToMutableTextureConversionUtils.h b/Engine/Plugins/Mutable/Source/CustomizableObject/Internal/MuCO/UnrealToMutableTextureConversionUtils.h
--- a/Engine/Plugins/Mutable/Source/CustomizableObject/Internal/MuCO/UnrealToMutableTextureConversionUtils.h
+++ b/Engine/Plugins/Mutable/Source/CustomizableObject/Internal/MuCO/UnrealToMutableTextureConversionUtils.h
@@ -92,10 +92,7 @@ struct FMutableSourceMeshData
 	TSoftObjectPtr<UStreamableRenderAsset> TableReferenceSkeletalMesh;
 
 	FName Component = NAME_None;
-
-	/** */
-	bool bIsPassthrough = false;
-
+	
 	/** Required for SurfaceMetadataID */
 
 	FMutableSourceSurfaceMetadata Metadata;
diff --git a/Engine/Plugins/Mutable/Source/CustomizableObject/Private/MuCO/CustomizableObjectInstance.cpp b/Engine/Plugins/Mutable/Source/CustomizableObject/Private/MuCO/CustomizableObjectInstance.cpp
--- a/Engine/Plugins/Mutable/Source/CustomizableObject/Private/MuCO/CustomizableObjectInstance.cpp
+++ b/Engine/Plugins/Mutable/Source/CustomizableObject/Private/MuCO/CustomizableObjectInstance.cpp
@@ -236,7 +236,7 @@ void UCustomizableInstancePrivate::InvalidateGeneratedData()
 
 	// Init Component Data
 	FCustomizableInstanceComponentData TemplateComponentData;
-	TemplateComponentData.LastMeshIdPerLOD.Init({}, MAX_MESH_LOD_COUNT);
+	TemplateComponentData.LastComponentDescriptor.MeshDescriptors.Init({}, MAX_MESH_LOD_COUNT);
 	ComponentsData.Init(TemplateComponentData, ComponentsData.Num());
 
 	GeneratedMaterials.Empty();
@@ -254,7 +254,7 @@ void UCustomizableInstancePrivate::InitCustomizableObjectData(const UCustomizabl
 	
 	// Init Component Data
 	FCustomizableInstanceComponentData TemplateComponentData;
-	TemplateComponentData.LastMeshIdPerLOD.Init({}, MAX_MESH_LOD_COUNT);
+	TemplateComponentData.LastComponentDescriptor.MeshDescriptors.Init({}, MAX_MESH_LOD_COUNT);
 	ComponentsData.Init(TemplateComponentData, InCustomizableObject->GetComponentCount());
 
 	ExtensionInstanceData.Empty();
@@ -302,8 +302,8 @@ void UCustomizableInstancePrivate::ClearInstanceComponentData(const FName& Compo
 
 	if (ComponentData)
 	{
-		ComponentData->LastMeshIdPerLOD.Reset(MAX_MESH_LOD_COUNT);
-		ComponentData->LastMeshIdPerLOD.AddDefaulted(MAX_MESH_LOD_COUNT);
+		ComponentData->LastComponentDescriptor.MeshDescriptors.Reset(MAX_MESH_LOD_COUNT);
+		ComponentData->LastComponentDescriptor.MeshDescriptors.AddDefaulted(MAX_MESH_LOD_COUNT);
 
 		ComponentData->AnimSlotToBP.Empty();
 		ComponentData->AssetUserDataArray.Empty();
@@ -1982,47 +1982,58 @@ bool DoComponentsNeedUpdate(const TSharedRef<FUpdateContextPrivate>& Context, bo
 			Context->FirstLODAvailable[ComponentName] :
 			Context->GetFirstRequestedLOD()[ComponentName];
 
-		// Look for invalid meshes
-		for (int32 LODIndex = FirstLOD; LODIndex < Component.LODCount; ++LODIndex)
+		bool bMeshChanged = true;
+		
+		if (Component.SkeletalMesh->PassthroughObject)
 		{
-			const FInstanceUpdateData::FLOD& LOD = Context->InstanceUpdateData.LODs[Component.FirstLOD + LODIndex];
-
-			if (!LOD.Mesh)
+			if (const FCustomizableInstanceComponentData* ComponentData = InstancePrivate->GetComponentData(ObjectComponentIndex)) // Could be nullptr if the component has not been generated.
 			{
-				bHasInvalidMesh |= bComponentWithMesh; // Mesh will be invalid if the UE::Mutable::Private::mesh is missing on higher LODs.
-				continue;
+				bMeshChanged = Component.SkeletalMeshId != ComponentData->LastComponentDescriptor.SkeletalMeshDescriptor;
 			}
-
-			if (LOD.SurfaceCount == 0 && !LOD.Mesh->IsReference())
+		}
+		else
+		{
+			// Look for invalid meshes
+			for (int32 LODIndex = FirstLOD; LODIndex < Component.LODCount; ++LODIndex)
 			{
-				continue;
-			}
+				const FInstanceUpdateData::FLOD& LOD = Context->InstanceUpdateData.LODs[Component.FirstLOD + LODIndex];
 
-			if (!LOD.Mesh->IsReference() && LOD.Mesh->GetVertexCount() == 0)
-			{
-				UE_LOG(LogMutable, Error, TEXT("Failed to generate SkeletalMesh for CO Instance [%s]. CO [%s] has invalid geometry for LOD [%d] Component [%d]."),
-					*Instance->GetName(), *CustomizableObject->GetName(),
-					LODIndex, InstanceComponentIndex);
-				bHasInvalidMesh = true; // Unreal does not support empty LODs.
-			}
+				if (!LOD.Mesh)
+				{
+					bHasInvalidMesh |= bComponentWithMesh; // Mesh will be invalid if the UE::Mutable::Private::mesh is missing on higher LODs.
+					continue;
+				}
 
-			check(LOD.SurfaceCount > 0);
-			bComponentWithMesh = true;
-		}
+				if (LOD.SurfaceCount == 0 && !LOD.Mesh->IsReference())
+				{
+					continue;
+				}
 
-		bool bMeshChanged = true;
+				if (!LOD.Mesh->IsReference() && LOD.Mesh->GetVertexCount() == 0)
+				{
+					UE_LOG(LogMutable, Error, TEXT("Failed to generate SkeletalMesh for CO Instance [%s]. CO [%s] has invalid geometry for LOD [%d] Component [%d]."),
+						*Instance->GetName(), *CustomizableObject->GetName(),
+						LODIndex, InstanceComponentIndex);
+					bHasInvalidMesh = true; // Unreal does not support empty LODs.
+				}
 
-		// Compare Mesh IDs 
-		const TArray<UE::Mutable::Private::FMeshId>* MeshIDs = Context->GetMeshDescriptors(ObjectComponentIndex);
-		if (ensure(MeshIDs))
-		{
-			if (const FCustomizableInstanceComponentData* ComponentData = InstancePrivate->GetComponentData(ObjectComponentIndex)) // Could be nullptr if the component has not been generated.
+				check(LOD.SurfaceCount > 0);
+				bComponentWithMesh = true;
+			}
+			
+			// Compare Mesh IDs 
+			const FComponentDescriptor* ComponentDescriptor = Context->GetComponentDescriptor(ObjectComponentIndex);
+			if (ensure(ComponentDescriptor))
 			{
-				bMeshChanged = *MeshIDs != ComponentData->LastMeshIdPerLOD;
+				if (const FCustomizableInstanceComponentData* ComponentData = InstancePrivate->GetComponentData(ObjectComponentIndex)) // Could be nullptr if the component has not been generated.
+				{
+					bMeshChanged = ComponentDescriptor->MeshDescriptors != ComponentData->LastComponentDescriptor.MeshDescriptors;
+				}
 			}
-		}
 
-		Context->MeshChangedPerInstanceComponent[InstanceComponentIndex] = bMeshChanged;
+			Context->MeshChangedPerInstanceComponent[InstanceComponentIndex] = bMeshChanged;
+		}
+		
 		bNeedsUpdate |= bMeshChanged;
 	}
 
@@ -2245,11 +2256,21 @@ bool UCustomizableInstancePrivate::UpdateSkeletalMesh_PostBeginUpdate0(UCustomiz
 			return false;
 		}
 
-		TArray<UE::Mutable::Private::FMeshId>* MeshDescriptors = Context->GetMeshDescriptors(ObjectComponentIndex);
-		check(MeshDescriptors);
-		check(MeshDescriptors->Num() == MAX_MESH_LOD_COUNT);
-		ComponentData->LastMeshIdPerLOD = *MeshDescriptors;
+		FComponentDescriptor* ComponentDescriptor = Context->GetComponentDescriptor(ObjectComponentIndex);
+		check(ComponentDescriptor);
+		check(ComponentDescriptor->MeshDescriptors.Num() == MAX_MESH_LOD_COUNT);
+		ComponentData->LastComponentDescriptor = *ComponentDescriptor;
 
+		if (Component->SkeletalMesh->PassthroughObject)
+		{
+			if (USkeletalMesh* SkeletalMesh = Component->SkeletalMesh->PassthroughObject.Get())
+			{
+				SkeletalMeshes.Add(ComponentName, SkeletalMesh);
+			}
+			
+			continue;
+		}
+		
 		// If the component doesn't need an update copy the previously generated mesh.
 		if (!Context->MeshChangedPerInstanceComponent[InstanceComponentIndex.GetValue()])
 		{
@@ -2267,7 +2288,7 @@ bool UCustomizableInstancePrivate::UpdateSkeletalMesh_PostBeginUpdate0(UCustomiz
 		
 		if (Context->bUseMeshCache)
 		{
-			if (USkeletalMesh* CachedMesh = CustomizableObject->GetPrivate()->MeshCache.Get(*MeshDescriptors))
+			if (USkeletalMesh* CachedMesh = CustomizableObject->GetPrivate()->MeshCache.Get(ComponentDescriptor->MeshDescriptors))
 			{
 				SkeletalMeshes.Add(ComponentName, CachedMesh);
 				continue;
@@ -2291,17 +2312,6 @@ bool UCustomizableInstancePrivate::UpdateSkeletalMesh_PostBeginUpdate0(UCustomiz
 			continue;
 		}
 
-		// If it is a referenced resource, only the first LOD is relevant.
-		if (ComponentMesh->IsReference())
-        {
-	        if (USkeletalMesh* SkeletalMesh = ComponentMesh->PassthroughObject.Get())
-			{
-				SkeletalMeshes.Add(ComponentName, SkeletalMesh);
-			}
-			
-        	continue;
-        }
-
 		if (!ModelResources.ReferenceSkeletalMeshesData.IsValidIndex(ObjectComponentIndex.GetValue()))
 		{
 			InvalidateGeneratedData();
@@ -2316,7 +2326,7 @@ bool UCustomizableInstancePrivate::UpdateSkeletalMesh_PostBeginUpdate0(UCustomiz
 			uint32 MeshHash = INDEX_NONE;
 			for (int32 LODIndex = FirstLOD; LODIndex < Component->LODCount && LODIndex < MAX_MESH_LOD_COUNT; ++LODIndex)
 			{
-				const UE::Mutable::Private::FMeshId& MeshId = (*MeshDescriptors)[LODIndex];
+				const UE::Mutable::Private::FMeshId& MeshId = ComponentDescriptor->MeshDescriptors[LODIndex];
 				const uint32 MeshIDHash = UE::Mutable::Private::GetTypeHashPersistent(MeshId);
 
 				if (!Context->bBake)
@@ -2457,11 +2467,7 @@ bool UCustomizableInstancePrivate::UpdateSkeletalMesh_PostBeginUpdate0(UCustomiz
 		
 		if (Context->bUseMeshCache)
 		{
-			const TArray<UE::Mutable::Private::FMeshId>* MeshId = Context->GetMeshDescriptors(ObjectComponentIndex);
-			if (MeshId)
-			{
-				CustomizableObject->GetPrivate()->MeshCache.Add(*MeshId, CastChecked<UCustomizableObjectSkeletalMesh>(SkeletalMesh));
-			}
+			CustomizableObject->GetPrivate()->MeshCache.Add(ComponentDescriptor->MeshDescriptors, CastChecked<UCustomizableObjectSkeletalMesh>(SkeletalMesh));
 		}
 
 		if (UCustomizableObjectSkeletalMesh* StreamableMesh = Cast<UCustomizableObjectSkeletalMesh>(SkeletalMesh))
@@ -5033,12 +5039,8 @@ UE::Tasks::FTask LoadAdditionalAssetsAndData(const TSharedRef<FUpdateContextPriv
 		{
 			const FInstanceUpdateData::FComponent& Component = Components[InstanceComponentIndex];
 			FCustomizableInstanceComponentData* ComponentData = InstancePrivate->GetComponentData(Component.Id);
-
-			TSharedPtr<const UE::Mutable::Private::FMesh> FirstComponentMesh = Context->InstanceUpdateData.LODs.IsValidIndex(Component.FirstLOD) ?
-				Context->InstanceUpdateData.LODs[Component.FirstLOD].Mesh :
-				nullptr;
-
-			if (FirstComponentMesh && FirstComponentMesh->IsReference())
+			
+			if (Component.SkeletalMesh->PassthroughObject)
 			{
 				// Pass-through components don't have a Reference Mesh so don't access it
 				continue;
@@ -5723,6 +5725,7 @@ void UCustomizableInstancePrivate::BuildMaterials(const TSharedRef<FUpdateContex
 	UpdateContext->MeshIdRegistry = CustomizableObject->GetPrivate()->MeshIdRegistry;
 	UpdateContext->ImageIdRegistry = CustomizableObject->GetPrivate()->ImageIdRegistry;
     UpdateContext->MaterialIdRegistry = CustomizableObject->GetPrivate()->MaterialIdRegistry,
+    UpdateContext->SkeletalMeshIdRegistry = CustomizableObject->GetPrivate()->SkeletalMeshIdRegistry,
 	UpdateContext->ExternalResourceProvider = Context->ExternalResourceProvider;
 	UpdateContext->ModelStreamableBulkData = CustomizableObject->GetPrivate()->GetModelStreamableBulkData();
 	UpdateContext->Parameters = Context->Parameters;
diff --git a/Engine/Plugins/Mutable/Source/CustomizableObject/Private/MuCO/CustomizableObjectMeshUpdate.cpp b/Engine/Plugins/Mutable/Source/CustomizableObject/Private/MuCO/CustomizableObjectMeshUpdate.cpp
--- a/Engine/Plugins/Mutable/Source/CustomizableObject/Private/MuCO/CustomizableObjectMeshUpdate.cpp
+++ b/Engine/Plugins/Mutable/Source/CustomizableObject/Private/MuCO/CustomizableObjectMeshUpdate.cpp
@@ -53,6 +53,7 @@ FCustomizableObjectMeshStreamIn::FCustomizableObjectMeshStreamIn(
 	OperationData->MeshIdRegistry = InMesh->MeshIdRegistry;
 	OperationData->ImageIdRegistry = InMesh->ImageIdRegistry;
 	OperationData->MaterialIdRegistry = InMesh->MaterialIdRegistry;
+	OperationData->SkeletalMeshIdRegistry = InMesh->SkeletalMeshIdRegistry;
 	OperationData->ExternalResourceProvider = InMesh->ExternalResourceProvider;
 	OperationData->State = InMesh->State;
 
@@ -254,6 +255,7 @@ namespace impl
 			OperationData->MeshIdRegistry,
 			OperationData->ImageIdRegistry,
 			OperationData->MaterialIdRegistry,
+			OperationData->SkeletalMeshIdRegistry,
 			OperationData->State,
 			LODMask);
 		
diff --git a/Engine/Plugins/Mutable/Source/CustomizableObject/Private/MuCO/CustomizableObjectMeshUpdate.h b/Engine/Plugins/Mutable/Source/CustomizableObject/Private/MuCO/CustomizableObjectMeshUpdate.h
--- a/Engine/Plugins/Mutable/Source/CustomizableObject/Private/MuCO/CustomizableObjectMeshUpdate.h
+++ b/Engine/Plugins/Mutable/Source/CustomizableObject/Private/MuCO/CustomizableObjectMeshUpdate.h
@@ -29,6 +29,7 @@ struct FMutableUpdateMeshContext
 	TSharedPtr<UE::Mutable::Private::FMeshIdRegistry> MeshIdRegistry;
 	TSharedPtr<UE::Mutable::Private::FImageIdRegistry> ImageIdRegistry;
 	TSharedPtr<UE::Mutable::Private::FMaterialIdRegistry> MaterialIdRegistry;
+	TSharedPtr<UE::Mutable::Private::FSkeletalMeshIdRegistry> SkeletalMeshIdRegistry;
 	TSharedPtr<FUnrealMutableResourceProvider> ExternalResourceProvider;
 	TSharedPtr<const UE::Mutable::Private::FParameters> Parameters;
 
diff --git a/Engine/Plugins/Mutable/Source/CustomizableObject/Private/MuCO/CustomizableObjectMipDataProvider.cpp b/Engine/Plugins/Mutable/Source/CustomizableObject/Private/MuCO/CustomizableObjectMipDataProvider.cpp
--- a/Engine/Plugins/Mutable/Source/CustomizableObject/Private/MuCO/CustomizableObjectMipDataProvider.cpp
+++ b/Engine/Plugins/Mutable/Source/CustomizableObject/Private/MuCO/CustomizableObjectMipDataProvider.cpp
@@ -146,6 +146,7 @@ namespace Impl
 				OperationData->UpdateContext->MeshIdRegistry,
 				OperationData->UpdateContext->ImageIdRegistry,
 				OperationData->UpdateContext->MaterialIdRegistry,
+				OperationData->UpdateContext->SkeletalMeshIdRegistry,
 				OperationData->UpdateContext->State,
 				LODMask);
 			
diff --git a/Engine/Plugins/Mutable/Source/CustomizableObject/Private/MuCO/CustomizableObjectSkeletalMesh.cpp b/Engine/Plugins/Mutable/Source/CustomizableObject/Private/MuCO/CustomizableObjectSkeletalMesh.cpp
--- a/Engine/Plugins/Mutable/Source/CustomizableObject/Private/MuCO/CustomizableObjectSkeletalMesh.cpp
+++ b/Engine/Plugins/Mutable/Source/CustomizableObject/Private/MuCO/CustomizableObjectSkeletalMesh.cpp
@@ -30,6 +30,7 @@ void UCustomizableObjectSkeletalMesh::InitMutableStreamingData(const TSharedRef<
 	MeshIdRegistry = CustomizableObject.GetPrivate()->MeshIdRegistry;
 	ImageIdRegistry = CustomizableObject.GetPrivate()->ImageIdRegistry;
 	MaterialIdRegistry = CustomizableObject.GetPrivate()->MaterialIdRegistry;
+	SkeletalMeshIdRegistry = CustomizableObject.GetPrivate()->SkeletalMeshIdRegistry;
 	ExternalResourceProvider = InOperationData->ExternalResourceProvider;
 
 	Parameters = InOperationData->Parameters;
diff --git a/Engine/Plugins/Mutable/Source/CustomizableObject/Private/MuCO/CustomizableObjectSystem.cpp b/Engine/Plugins/Mutable/Source/CustomizableObject/Private/MuCO/CustomizableObjectSystem.cpp
--- a/Engine/Plugins/Mutable/Source/CustomizableObject/Private/MuCO/CustomizableObjectSystem.cpp
+++ b/Engine/Plugins/Mutable/Source/CustomizableObject/Private/MuCO/CustomizableObjectSystem.cpp
@@ -593,23 +593,23 @@ const FInstanceUpdateData::FComponent* FUpdateContextPrivate::GetComponentUpdate
 }
 
 
-void FUpdateContextPrivate::InitMeshDescriptors(int32 Size)
+void FUpdateContextPrivate::InitComponentDescriptors(int32 Size)
 {
-	MeshDescriptors.SetNum(Size);
+	ComponentDescriptors.SetNum(Size);
 }
 
 
-const TArray<TArray<UE::Mutable::Private::FMeshId>>& FUpdateContextPrivate::GetMeshDescriptors() const
+const TArray<FComponentDescriptor>& FUpdateContextPrivate::GetComponentDescriptors() const
 {
-	return MeshDescriptors;
+	return ComponentDescriptors;
 }
 
 
-TArray<UE::Mutable::Private::FMeshId>* FUpdateContextPrivate::GetMeshDescriptors(FCustomizableObjectComponentIndex Index)
+FComponentDescriptor* FUpdateContextPrivate::GetComponentDescriptor(FCustomizableObjectComponentIndex Index)
 {
-	if (MeshDescriptors.IsValidIndex(Index.GetValue()))
+	if (ComponentDescriptors.IsValidIndex(Index.GetValue()))
 	{
-		return &MeshDescriptors[Index.GetValue()];
+		return &ComponentDescriptors[Index.GetValue()];
 	}
 	return nullptr;
 }
@@ -2664,6 +2664,7 @@ namespace impl
 				OperationData->MeshIdRegistry,
 				OperationData->ImageIdRegistry,
 				OperationData->MaterialIdRegistry,
+				OperationData->SkeletalMeshIdRegistry,
 				OperationData->GetCapturedDescriptor().GetState(),
 				UE::Mutable::Private::FSystem::AllLODs);
 
@@ -2693,18 +2694,23 @@ namespace impl
 				FCustomizableObjectComponentIndex ObjectComponentIndex = Component.Id;
 				const FName ComponentName = OperationData->ComponentNames[ObjectComponentIndex.GetValue()];
 
-				Component.LODCount = OperationData->NumLODsAvailable[ComponentName];
-				if (!Component.LODCount)
-				{
-					// It happens in degenerated cases with empty components.
-					continue;
-				}
+				Component.SkeletalMeshId = Instance->GetSkeletalMeshId(InstanceComponentIndex);
+
+				UE::Tasks::TTask<TSharedPtr<const UE::Mutable::Private::FSkeletalMesh>> Task = OperationData->MutableSystem->GetSkeletalMesh(OperationData->InstanceID, Component.SkeletalMeshId, UE::Mutable::Private::EMeshContentFlags::AllFlags);
+				Component.SkeletalMesh = Task.GetResult(); // TODO GMT Remove. We can do this since now FSkeletalMesh will always be passthrough.
 
 				Component.OverlayMaterialId = Instance->GetOverlayMaterialId(InstanceComponentIndex);
 				if (Component.OverlayMaterialId)
 				{
 					GetMaterialOverrideComponent.Add(InstanceComponentIndex);
 				}
+
+				Component.LODCount = OperationData->NumLODsAvailable[ComponentName];
+				if (!Component.LODCount)
+				{
+					// It happens in degenerated cases with empty components.
+					continue;
+				}
 				
 				const uint8 FirstResidentLOD = OperationData->FirstResidentLOD[ComponentName];
 
@@ -3530,6 +3536,7 @@ namespace impl
 		Operation->MeshIdRegistry = ObjectPrivate->MeshIdRegistry;
 		Operation->ImageIdRegistry = ObjectPrivate->ImageIdRegistry;
 		Operation->MaterialIdRegistry = ObjectPrivate->MaterialIdRegistry;
+		Operation->SkeletalMeshIdRegistry = ObjectPrivate->SkeletalMeshIdRegistry;
 
 		UE::Tasks::FTask CacheRuntimeTexturesEvent = UE::Tasks::MakeCompletedTask<void>();
 
@@ -3689,6 +3696,7 @@ namespace impl
 				Operation->MeshIdRegistry,
 				Operation->ImageIdRegistry,
 				Operation->MaterialIdRegistry,
+				Operation->SkeletalMeshIdRegistry,
 				Operation->GetCapturedDescriptor().GetState(),
 				UE::Mutable::Private::FSystem::AllLODs);
 			
@@ -3697,22 +3705,22 @@ namespace impl
 		
 		FixLODs(Operation);
 		
-		Operation->InitMeshDescriptors(Operation->NumObjectComponents);
+		Operation->InitComponentDescriptors(Operation->NumObjectComponents);
 
 		const TMap<FName, uint8>& NumLODsAvailable = Operation->NumLODsAvailable;
 
 		for (int32 InstanceComponentIndex = 0; InstanceComponentIndex < Operation->NumInstanceComponents; ++InstanceComponentIndex)
 		{
 			FCustomizableObjectComponentIndex ObjectComponentIndex(Operation->MutableInstance->GetComponentId(InstanceComponentIndex));
-			TArray<UE::Mutable::Private::FMeshId>* MeshId = Operation->GetMeshDescriptors(ObjectComponentIndex);
-			check(MeshId);
+			FComponentDescriptor* ComponentDescriptor = Operation->GetComponentDescriptor(ObjectComponentIndex);
+			check(ComponentDescriptor);
 
-			if (!MeshId)
+			if (!ComponentDescriptor)
 			{
 				continue;
 			}
 
-			MeshId->Init({}, MAX_MESH_LOD_COUNT);
+			ComponentDescriptor->MeshDescriptors.Init({}, MAX_MESH_LOD_COUNT);
 
 			const FName& ComponentName = Operation->ComponentNames[ObjectComponentIndex.GetValue()];
 
@@ -3722,15 +3730,15 @@ namespace impl
 
 			for (int32 LODIndex = FirstLOD; LODIndex < NumLODsAvailable[ComponentName]; ++LODIndex)
 			{
-				(*MeshId)[LODIndex] = Operation->MutableInstance->GetMeshId(InstanceComponentIndex, LODIndex);
+				ComponentDescriptor->MeshDescriptors[LODIndex] = Operation->MutableInstance->GetMeshId(InstanceComponentIndex, LODIndex);
 			}
 		}
 
 		if (Operation->bUseMeshCache)
 		{
-			for (const TArray<UE::Mutable::Private::FMeshId>& MeshId : Operation->GetMeshDescriptors())
+			for (const FComponentDescriptor& ComponentDescriptor : Operation->GetComponentDescriptors())
 			{
-				if (UCustomizableObjectSkeletalMesh* CachedMesh = ObjectPrivate->MeshCache.Get(MeshId))
+				if (UCustomizableObjectSkeletalMesh* CachedMesh = ObjectPrivate->MeshCache.Get(ComponentDescriptor.MeshDescriptors))
 				{
 					Operation->Objects.Emplace(CachedMesh);
 				}
diff --git a/Engine/Plugins/Mutable/Source/CustomizableObject/Public/MuCO/CustomizableObjectSkeletalMesh.h b/Engine/Plugins/Mutable/Source/CustomizableObject/Public/MuCO/CustomizableObjectSkeletalMesh.h
--- a/Engine/Plugins/Mutable/Source/CustomizableObject/Public/MuCO/CustomizableObjectSkeletalMesh.h
+++ b/Engine/Plugins/Mutable/Source/CustomizableObject/Public/MuCO/CustomizableObjectSkeletalMesh.h
@@ -50,6 +50,7 @@ class UCustomizableObjectSkeletalMesh : public USkeletalMesh
 	TSharedPtr<UE::Mutable::Private::FMeshIdRegistry> MeshIdRegistry;
 	TSharedPtr<UE::Mutable::Private::FImageIdRegistry> ImageIdRegistry;
 	TSharedPtr<UE::Mutable::Private::FMaterialIdRegistry> MaterialIdRegistry;
+	TSharedPtr<UE::Mutable::Private::FSkeletalMeshIdRegistry> SkeletalMeshIdRegistry;
 
 	TSharedPtr<FUnrealMutableResourceProvider> ExternalResourceProvider;
 
diff --git a/Engine/Plugins/Mutable/Source/CustomizableObjectEditor/Private/MuCOE/GenerateMutableSource/GenerateMutableSource.cpp b/Engine/Plugins/Mutable/Source/CustomizableObjectEditor/Private/MuCOE/GenerateMutableSource/GenerateMutableSource.cpp
--- a/Engine/Plugins/Mutable/Source/CustomizableObjectEditor/Private/MuCOE/GenerateMutableSource/GenerateMutableSource.cpp
+++ b/Engine/Plugins/Mutable/Source/CustomizableObjectEditor/Private/MuCOE/GenerateMutableSource/GenerateMutableSource.cpp
@@ -1770,51 +1770,22 @@ TSharedPtr<UE::Mutable::Private::FMesh> GenerateMeshConstant(
 		return nullptr;
 	}
 
-	bool bForceLoad = false;
-	bool bIsCompileTime = false;
-	if (!Source.bIsPassthrough)
-	{
-		bForceLoad = true;
-		// TODO: For now, even in none, we turn them into compile-time references, since runtime meshes are not
-		// supported yet.
-		//if (GenerationContext.CompilationContext->Options.OptimizationLevel == 0)
-		//{
-		//	bIsCompileTime = false;
-		//}
-		//else
-		{
-			bIsCompileTime = true;
-		}
-	}
-
 	UE::Mutable::Private::PASSTHROUGH_ID Id;
 
-	if (Source.bIsPassthrough)
-	{
-		Id = GenerationContext.AddPassthroughObject(*Source.Mesh);
-	}
-	else if (bIsCompileTime)
-	{
-		FMutableGraphGenerationContext::FGeneratedReferencedMesh InvalidEntry;
-		InvalidEntry.ID = UE::Mutable::Private::PASSTHROUGH_ID_INVALID;
+	FMutableGraphGenerationContext::FGeneratedReferencedMesh InvalidEntry;
+	InvalidEntry.ID = UE::Mutable::Private::PASSTHROUGH_ID_INVALID;
 
-		FMutableGraphGenerationContext::FGeneratedReferencedMesh* Entry = &GenerationContext.CompileTimeMeshMap.FindOrAdd(Source, InvalidEntry);
-		int32 Num = GenerationContext.CompileTimeMeshMap.Num();
+	FMutableGraphGenerationContext::FGeneratedReferencedMesh* Entry = &GenerationContext.CompileTimeMeshMap.FindOrAdd(Source, InvalidEntry);
+	int32 Num = GenerationContext.CompileTimeMeshMap.Num();
 
-		if (Entry->ID == UE::Mutable::Private::PASSTHROUGH_ID_INVALID)
-		{
-			Entry->ID = Num - 1;
-		}
-
-		Id = Entry->ID;
-	}
-	else
+	if (Entry->ID == UE::Mutable::Private::PASSTHROUGH_ID_INVALID)
 	{
-		Id = UE::Mutable::Private::PASSTHROUGH_ID_INVALID;
-		unimplemented();
+		Entry->ID = Num - 1;
 	}
 
-	TSharedPtr<UE::Mutable::Private::FMesh> Result = UE::Mutable::Private::FMesh::CreateAsReference(UE::Mutable::Private::TPassthroughObjectPtr<USkeletalMesh>(Id), bForceLoad); // TODO GMT If bIsCompileTime it should not create a TPassthroughObjectPtr since it is not a Passthrough
+	Id = Entry->ID;
+
+	TSharedPtr<UE::Mutable::Private::FMesh> Result = UE::Mutable::Private::FMesh::CreateAsReference(UE::Mutable::Private::TPassthroughObjectPtr<USkeletalMesh>(Id));
 	return Result;
 }
 
diff --git a/Engine/Plugins/Mutable/Source/CustomizableObjectEditor/Private/MuCOE/GenerateMutableSource/GenerateMutableSourceComponent.cpp b/Engine/Plugins/Mutable/Source/CustomizableObjectEditor/Private/MuCOE/GenerateMutableSource/GenerateMutableSourceComponent.cpp
--- a/Engine/Plugins/Mutable/Source/CustomizableObjectEditor/Private/MuCOE/GenerateMutableSource/GenerateMutableSourceComponent.cpp
+++ b/Engine/Plugins/Mutable/Source/CustomizableObjectEditor/Private/MuCOE/GenerateMutableSource/GenerateMutableSourceComponent.cpp
@@ -2,6 +2,7 @@
 
 #include "GenerateMutableSourceComponent.h"
 
+#include "GenerateMutableSourceSkeletalMesh.h"
 #include "Factories/FbxMeshImportData.h"
 #include "MuCOE/GenerateMutableSource/GenerateMutableSource.h"
 
@@ -25,93 +26,17 @@
 #include "MuT/NodeComponentNew.h"
 #include "MuT/NodeComponentSwitch.h"
 #include "MuT/NodeComponentVariation.h"
-#include "MuT/NodeLOD.h"
 #include "MuT/NodeScalar.h"
-#include "MuT/NodeScalarConstant.h"
 #include "MuT/NodeSurfaceNew.h"
 #include "MuT/NodeMeshConstant.h"
 #include "MuT/NodeMaterialConstant.h"
 
 #include "Rendering/SkeletalMeshLODModel.h"
-#include "Rendering/SkeletalMeshModel.h"
 #include "Interfaces/ITargetPlatform.h"
 
 
 #define LOCTEXT_NAMESPACE "CustomizableObjectEditor"
 
-/** Generate LOD pins of the given NodeComponentBase (NodeComponent, NodeComponentExtend...).
- * @param TypedComponentMesh Given component node.
- * @param NodeComponent Core node to connect LOD generated pins. */
-void GenerateMutableSourceComponentMesh(FMutableGraphGenerationContext& GenerationContext, const ICustomizableObjectNodeComponentMeshInterface& TypedComponentMesh, UE::Mutable::Private::Ptr<UE::Mutable::Private::NodeComponent> NodeComponent)
-{
-	int32 FirstLOD = -1;
-
-	const int32 NumLODsInRoot = GenerationContext.NumLODs[GenerationContext.CurrentMeshComponent];
-	for (int32 CurrentLOD = 0; CurrentLOD < NumLODsInRoot; ++CurrentLOD)
-	{
-		GenerationContext.CurrentLOD = CurrentLOD;
-
-		if (!NodeComponent->LODs.IsValidIndex(CurrentLOD))
-		{
-			NodeComponent->LODs.Add(new UE::Mutable::Private::NodeLOD());
-		}
-
-		UE::Mutable::Private::Ptr<UE::Mutable::Private::NodeLOD> LODNode = NodeComponent->LODs[CurrentLOD];
-
-		LODNode->SetMessageContext(&TypedComponentMesh);
-
-		const int32 NumLODs = TypedComponentMesh.GetLODPins().Num();
-
-		const bool bUseAutomaticLods = GenerationContext.CurrentAutoLODStrategy == ECustomizableObjectAutomaticLODStrategy::AutomaticFromMesh;
-		FirstLOD = (CurrentLOD < NumLODs) && (FirstLOD == INDEX_NONE || !bUseAutomaticLods) ? CurrentLOD : FirstLOD;
-
-		if (FirstLOD < 0)
-		{
-			continue;
-		}
-
-		if (GenerationContext.CurrentLOD < GenerationContext.FirstLODAvailable[GenerationContext.CurrentMeshComponent])
-		{
-			continue;
-		} 
-	
-		// Generate all relevant LODs for this object up until the current LODIndex.
-		for (int32 LODIndex = FirstLOD; LODIndex <= CurrentLOD; ++LODIndex)
-		{
-			if (!TypedComponentMesh.GetLODPins().IsValidIndex(LODIndex))
-			{
-				continue;
-			}
-			
-			const UEdGraphPin* LODPin = TypedComponentMesh.GetLODPins()[LODIndex].Get();
-			check(LODPin);
-			
-			GenerationContext.FromLOD = LODIndex;
-
-			TArray<UEdGraphPin*> ConnectedLODPins = FollowInputPinArray(*LODPin);
-
-			// Process non modifier nodes.
-			for (UEdGraphPin* const ChildNodePin : ConnectedLODPins)
-			{
-				// Modifiers are shared for all components and are processed per LOD and not component.
-				if (Cast<UCustomizableObjectNodeModifierBase>(ChildNodePin->GetOwningNode()))
-				{
-					FString Msg = FString::Printf(TEXT("The object has legacy modifier connections that cannot be generated. Their connections should be updated."));
-					GenerationContext.Log(FText::FromString(Msg), TypedComponentMesh.GetOwningNode(), EMessageSeverity::Warning);
-					continue;
-				}
-				
-				UE::Mutable::Private::Ptr<UE::Mutable::Private::NodeSurface> SurfaceNode = GenerateMutableSourceSurface(ChildNodePin, GenerationContext);
-				LODNode->Surfaces.Add(SurfaceNode);
-			}
-		}
-	}
-
-	// Clear the context state for LODs
-	GenerationContext.CurrentLOD = 0;
-	GenerationContext.FromLOD = 0;
-}
-
 
 UE::Mutable::Private::Ptr<UE::Mutable::Private::NodeComponent> GenerateMutableSourceComponent(const UEdGraphPin* Pin, FMutableGraphGenerationContext& GenerationContext)
 {
@@ -169,7 +94,7 @@ UE::Mutable::Private::Ptr<UE::Mutable::Private::NodeComponent> GenerateMutableSo
 		GenerationContext.CurrentMeshComponent = ComponentName;
 		GenerationContext.CurrentAutoLODStrategy = TypedComponentMesh->AutoLODStrategy;
 		
-		GenerateMutableSourceComponentMesh(GenerationContext, *TypedComponentMesh, NodeComponentNew);
+		NodeComponentNew->SkeletalMesh = UE::Mutable::Private::GenerateMutableSourceSkeletalMesh(Pin, GenerationContext);
 
 		GenerationContext.CurrentMeshComponent = {};
 		GenerationContext.CurrentAutoLODStrategy = {};
@@ -212,7 +137,7 @@ UE::Mutable::Private::Ptr<UE::Mutable::Private::NodeComponent> GenerateMutableSo
 				TypedParentComponentMesh->AutoLODStrategy :
 				TypedComponentMeshExtend->AutoLODStrategy;
 			
-			GenerateMutableSourceComponentMesh(GenerationContext, *TypedComponentMeshExtend, NodeComponentEdit);
+			NodeComponentEdit->SkeletalMesh = UE::Mutable::Private::GenerateMutableSourceSkeletalMeshNew(Pin, GenerationContext);
 
 			GenerationContext.CurrentMeshComponent = {};
 			GenerationContext.CurrentAutoLODStrategy = {};
@@ -237,35 +162,7 @@ UE::Mutable::Private::Ptr<UE::Mutable::Private::NodeComponent> GenerateMutableSo
 			GenerationContext.Log(FText::FromString(Msg), TypedComponentPassthroughMesh, EMessageSeverity::Warning);
 			return nullptr;
 		}
-
-		if (TypedComponentPassthroughMesh->SkeletalMesh.IsNull())
-		{
-			FString Msg = FString::Printf(TEXT("No mesh set for component node."));
-			GenerationContext.Log(FText::FromString(Msg), TypedComponentPassthroughMesh, EMessageSeverity::Warning);
-			return nullptr;
-		}
 		
-		USkeletalMesh* SkeletalMesh = Cast<USkeletalMesh>(GenerationContext.LoadObject(TSoftObjectPtr<USkeletalMesh>(TypedComponentPassthroughMesh->SkeletalMesh)));
-		if (!SkeletalMesh)
-		{
-			FString Msg = FString::Printf(TEXT("Only SkeletalMeshes are supported in this node, for now."));
-			GenerationContext.Log(FText::FromString(Msg), TypedComponentPassthroughMesh, EMessageSeverity::Warning);
-			return nullptr;
-		}
-
-		// Create the referenced mesh node.
-		UE::Mutable::Private::Ptr<UE::Mutable::Private::NodeMeshConstant> MeshNode;
-		{
-			MeshNode = new UE::Mutable::Private::NodeMeshConstant();
-
-			FMutableSourceMeshData Source;
-			Source.Mesh = SkeletalMesh;
-			Source.bIsPassthrough = true;
-			TSharedPtr<UE::Mutable::Private::FMesh> MutableMesh = GenerateMutableSkeletalMesh(Source, 0, 0, {}, GenerationContext, Node, false);
-
-			MeshNode->Value = MutableMesh;
-		}
-
 		// Create the component node
 		UE::Mutable::Private::Ptr<UE::Mutable::Private::NodeComponentNew> ComponentNode = new UE::Mutable::Private::NodeComponentNew;
 		if (GenerationContext.ComponentNames.Contains(ComponentName))
@@ -274,53 +171,7 @@ UE::Mutable::Private::Ptr<UE::Mutable::Private::NodeComponent> GenerateMutableSo
 			GenerationContext.Log(FText::FromString(Msg), TypedComponentPassthroughMesh, EMessageSeverity::Warning);
 		}
 		ComponentNode->Id = GenerationContext.ComponentNames.Add(ComponentName);
-
-		// Create a LOD for each pass-through mesh LOD.
-		const FSkeletalMeshModel* Model = SkeletalMesh->GetImportedModel();
-		int32 SkeletalMeshLODCount = Model->LODModels.Num();
-		for (int32 LODIndex=0; LODIndex<SkeletalMeshLODCount; ++LODIndex)
-		{
-			UE::Mutable::Private::Ptr<UE::Mutable::Private::NodeLOD> LODNode = new UE::Mutable::Private::NodeLOD;
-			ComponentNode->LODs.Add(LODNode);
-
-			const FSkeletalMeshLODModel& LODModel = Model->LODModels[LODIndex];
-			int32 SectionCount = LODModel.Sections.Num();
-			for (int32 SectionIndex = 0; SectionIndex < SectionCount; ++SectionIndex)
-			{
-				// Is there a pin in the unreal node for this section?
-				if (UEdGraphPin* InMaterialPin = TypedComponentPassthroughMesh->GetMaterialPin(LODIndex,SectionIndex))
-				{
-					if (UEdGraphPin* ConnectedMaterialPin = FollowInputPin(*InMaterialPin))
-					{
-						GenerationContext.ComponentMeshOverride = MeshNode;
-					
-						UE::Mutable::Private::Ptr<UE::Mutable::Private::NodeSurface> SurfaceNode = GenerateMutableSourceSurface(ConnectedMaterialPin, GenerationContext);
-						LODNode->Surfaces.Add(SurfaceNode);
-
-						GenerationContext.ComponentMeshOverride = nullptr;
-					}
-					else
-					{
-						// Add an empty surface node anyway.
-						UE::Mutable::Private::Ptr<UE::Mutable::Private::NodeSurfaceNew> SurfaceNode = new UE::Mutable::Private::NodeSurfaceNew;
-						SurfaceNode->Mesh = MeshNode;
-						SurfaceNode->SurfaceGuid = GenerationContext.GetNodeIdUnique(TypedComponentPassthroughMesh);
-						LODNode->Surfaces.Add(SurfaceNode);
-					}
-				}
-				else
-				{
-					// The SKM has changed but the node has not been refreshed
-					// This may happen if the node was saved without pins but with a SKM set (old bug)
-					
-					// Add an empty surface node anyway.
-					UE::Mutable::Private::Ptr<UE::Mutable::Private::NodeSurfaceNew> SurfaceNode = new UE::Mutable::Private::NodeSurfaceNew;
-					SurfaceNode->Mesh = MeshNode;
-					SurfaceNode->SurfaceGuid = GenerationContext.GetNodeIdUnique(TypedComponentPassthroughMesh);
-					LODNode->Surfaces.Add(SurfaceNode);
-				}
-			}
-		}
+		ComponentNode->SkeletalMesh = UE::Mutable::Private::GenerateMutableSourceSkeletalMesh(Pin, GenerationContext);
 
 		GenerationContext.CurrentMeshComponent = FName();
 		Result = ComponentNode;
diff --git a/Engine/Plugins/Mutable/Source/CustomizableObjectEditor/Private/MuCOE/GenerateMutableSource/GenerateMutableSourceMesh.cpp b/Engine/Plugins/Mutable/Source/CustomizableObjectEditor/Private/MuCOE/GenerateMutableSource/GenerateMutableSourceMesh.cpp
--- a/Engine/Plugins/Mutable/Source/CustomizableObjectEditor/Private/MuCOE/GenerateMutableSource/GenerateMutableSourceMesh.cpp
+++ b/Engine/Plugins/Mutable/Source/CustomizableObjectEditor/Private/MuCOE/GenerateMutableSource/GenerateMutableSourceMesh.cpp
@@ -2893,16 +2893,8 @@ TSharedPtr<UE::Mutable::Private::FMesh> GenerateMutableSkeletalMesh(FMutableSour
 		return MutableMesh;
 	}
 
-	// At some point we will want all meshes to be references at compile-time. For now, just create the actual pass-through meshes.
-	if (Source.bIsPassthrough)
-	{
-		MutableMesh = GenerateMeshConstant(Source, GenerationContext);
-	}
-	else
-	{
-		Source.bOnlyConnectedLOD = bOnlyConnectedLOD;
-		MutableMesh = ConvertSkeletalMeshToMutable(Source, LODIndexConnected, SectionIndexConnected, GenerationContext, CurrentNode);
-	}
+	Source.bOnlyConnectedLOD = bOnlyConnectedLOD;
+	MutableMesh = ConvertSkeletalMeshToMutable(Source, LODIndexConnected, SectionIndexConnected, GenerationContext, CurrentNode);
 
 	if (MutableMesh)
 	{
diff --git a/Engine/Plugins/Mutable/Source/CustomizableObjectEditor/Private/MuCOE/GenerateMutableSource/GenerateMutableSourceSkeletalMesh.cpp b/Engine/Plugins/Mutable/Source/CustomizableObjectEditor/Private/MuCOE/GenerateMutableSource/GenerateMutableSourceSkeletalMesh.cpp
--- a/Engine/Plugins/Mutable/Source/CustomizableObjectEditor/Private/MuCOE/GenerateMutableSource/GenerateMutableSourceSkeletalMesh.cpp
+++ b/Engine/Plugins/Mutable/Source/CustomizableObjectEditor/Private/MuCOE/GenerateMutableSource/GenerateMutableSourceSkeletalMesh.cpp
@@ -3,6 +3,10 @@
 #include "GenerateMutableSourceSkeletalMesh.h"
 
 #include "GenerateMutableSource.h"
+#include "GenerateMutableSourceSurface.h"
+#include "MuCOE/GraphTraversal.h"
+#include "MuCOE/Nodes/CustomizableObjectNodeComponentPassthroughMesh.h"
+#include "MuCOE/Nodes/CustomizableObjectNodeModifierBase.h"
 #include "MuCOE/Nodes/CustomizableObjectNodeSkeletalMeshParameter.h"
 #include "MuR/MutableTrace.h"
 #include "MuT/NodeSkeletalMeshParameter.h"
@@ -21,17 +25,17 @@ namespace UE::Mutable::Private
 	
 		UCustomizableObjectNodeParameter* ParameterNode = CastChecked<UCustomizableObjectNodeParameter>(Pin->GetOwningNode());
 
-		UE::Mutable::Private::Ptr<UE::Mutable::Private::NodeSkeletalMeshParameter> Result;
+		Ptr<NodeSkeletalMeshParameter> Result;
 
 		const FGuid NodeGUID = ParameterNode->NodeGuid;
-		if (const UE::Mutable::Private::Ptr<UE::Mutable::Private::NodeSkeletalMeshParameter>* GeneratedSkeletalMeshParameter = GenerationContext.GeneratedSkeletalMeshParameters.Find(NodeGUID))
+		if (const Ptr<NodeSkeletalMeshParameter>* GeneratedSkeletalMeshParameter = GenerationContext.GeneratedSkeletalMeshParameters.Find(NodeGUID))
 		{
 			return GeneratedSkeletalMeshParameter->get();
 		}
 	
 		if (const UCustomizableObjectNodeSkeletalMeshParameter* TypedNodeParam = Cast<UCustomizableObjectNodeSkeletalMeshParameter>(ParameterNode))
 		{
-			UE::Mutable::Private::Ptr<UE::Mutable::Private::NodeSkeletalMeshParameter> SkeletalMeshNode = new UE::Mutable::Private::NodeSkeletalMeshParameter();
+			Ptr<NodeSkeletalMeshParameter> SkeletalMeshNode = new NodeSkeletalMeshParameter();
 
 			SkeletalMeshNode->Name = TypedNodeParam->GetParameterName(&GenerationContext.MacroNodesStack);
 			SkeletalMeshNode->UID = GenerationContext.GetNodeIdUnique(ParameterNode).ToString();
@@ -60,7 +64,94 @@ namespace UE::Mutable::Private
 		return Result;
 	}
 
-	
+
+	Ptr<NodeSkeletalMeshNew> GenerateMutableSourceSkeletalMeshNew(const UEdGraphPin* Pin, FMutableGraphGenerationContext& GenerationContext)
+	{
+		UEdGraphNode* Node = Pin->GetOwningNode();
+		
+		if (const ICustomizableObjectNodeComponentMeshInterface* TypedComponentMesh = Cast<ICustomizableObjectNodeComponentMeshInterface>(Node))
+		{
+			Ptr<NodeSkeletalMeshNew> SkeletalMeshNode = new NodeSkeletalMeshNew();
+
+			int32 FirstLOD = -1;
+
+			const int32 NumLODsInRoot = GenerationContext.NumLODs[GenerationContext.CurrentMeshComponent];
+			for (int32 CurrentLOD = 0; CurrentLOD < NumLODsInRoot; ++CurrentLOD)
+			{
+				GenerationContext.CurrentLOD = CurrentLOD;
+
+				if (!SkeletalMeshNode->LODs.IsValidIndex(CurrentLOD))
+				{
+					SkeletalMeshNode->LODs.Add(new NodeLOD());
+				}
+
+				Ptr<NodeLOD> LODNode = SkeletalMeshNode->LODs[CurrentLOD];
+
+				LODNode->SetMessageContext(&TypedComponentMesh);
+
+				const int32 NumLODs = TypedComponentMesh->GetLODPins().Num();
+
+				const bool bUseAutomaticLods = GenerationContext.CurrentAutoLODStrategy == ECustomizableObjectAutomaticLODStrategy::AutomaticFromMesh;
+				FirstLOD = (CurrentLOD < NumLODs) && (FirstLOD == INDEX_NONE || !bUseAutomaticLods) ? CurrentLOD : FirstLOD;
+
+				if (FirstLOD < 0)
+				{
+					continue;
+				}
+
+				if (GenerationContext.CurrentLOD < GenerationContext.FirstLODAvailable[GenerationContext.CurrentMeshComponent])
+				{
+					continue;
+				} 
+			
+				// Generate all relevant LODs for this object up until the current LODIndex.
+				for (int32 LODIndex = FirstLOD; LODIndex <= CurrentLOD; ++LODIndex)
+				{
+					if (!TypedComponentMesh->GetLODPins().IsValidIndex(LODIndex))
+					{
+						continue;
+					}
+					
+					const UEdGraphPin* LODPin = TypedComponentMesh->GetLODPins()[LODIndex].Get();
+					check(LODPin);
+					
+					GenerationContext.FromLOD = LODIndex;
+
+					TArray<UEdGraphPin*> ConnectedLODPins = FollowInputPinArray(*LODPin);
+
+					// Process non modifier nodes.
+					for (UEdGraphPin* const ChildNodePin : ConnectedLODPins)
+					{
+						// Modifiers are shared for all components and are processed per LOD and not component.
+						if (Cast<UCustomizableObjectNodeModifierBase>(ChildNodePin->GetOwningNode()))
+						{
+							FString Msg = FString::Printf(TEXT("The object has legacy modifier connections that cannot be generated. Their connections should be updated."));
+							GenerationContext.Log(FText::FromString(Msg), TypedComponentMesh->GetOwningNode(), EMessageSeverity::Warning);
+							continue;
+						}
+						
+						Ptr<NodeSurface> SurfaceNode = GenerateMutableSourceSurface(ChildNodePin, GenerationContext);
+						LODNode->Surfaces.Add(SurfaceNode);
+					}
+				}
+			}
+
+			// Clear the context state for LODs
+			GenerationContext.CurrentLOD = 0;
+			GenerationContext.FromLOD = 0;
+
+			return SkeletalMeshNode;
+		}
+		
+		else
+		{
+			GenerationContext.Log(LOCTEXT("UnimplementedNode", "Node type not implemented yet."), Node);
+		}
+
+		return {};
+	}
+
+
 	Ptr<NodeSkeletalMesh> GenerateMutableSourceSkeletalMesh(const UEdGraphPin* Pin, FMutableGraphGenerationContext& GenerationContext)
 	{
 		MUTABLE_CPUPROFILER_SCOPE(GenerateMutableSourceSkeletalMesh);
@@ -87,6 +178,19 @@ namespace UE::Mutable::Private
 			Result = GenerateMutableSourceSkeletalMeshParameter(Pin, GenerationContext);
 		}
 
+		else if (const UCustomizableObjectNodeComponentMesh* ComponentMesh = Cast<UCustomizableObjectNodeComponentMesh>(Node))
+		{
+			Result = GenerateMutableSourceSkeletalMeshNew(Pin, GenerationContext);
+		}
+
+		else if (const UCustomizableObjectNodeComponentPassthroughMesh* ComponentPassthroughMesh = Cast<UCustomizableObjectNodeComponentPassthroughMesh>(Node))
+		{
+			if (const UEdGraphPin* ConnectedPin = FollowInputPin(*ComponentPassthroughMesh->SkeletalMeshPin.Get()))
+			{
+				Result = GenerateMutableSourceSkeletalMesh(ConnectedPin, GenerationContext);
+			}
+		}
+		
 		else
 		{
 			GenerationContext.Log(LOCTEXT("UnimplementedNode", "Node type not implemented yet."), Node);
diff --git a/Engine/Plugins/Mutable/Source/CustomizableObjectEditor/Private/MuCOE/GenerateMutableSource/GenerateMutableSourceSkeletalMesh.h b/Engine/Plugins/Mutable/Source/CustomizableObjectEditor/Private/MuCOE/GenerateMutableSource/GenerateMutableSourceSkeletalMesh.h
--- a/Engine/Plugins/Mutable/Source/CustomizableObjectEditor/Private/MuCOE/GenerateMutableSource/GenerateMutableSourceSkeletalMesh.h
+++ b/Engine/Plugins/Mutable/Source/CustomizableObjectEditor/Private/MuCOE/GenerateMutableSource/GenerateMutableSourceSkeletalMesh.h
@@ -4,8 +4,10 @@
 
 #include "MuR/Ptr.h"
 #include "MuT/NodeSkeletalMesh.h"
+#include "MuT/NodeSkeletalMeshNew.h"
 #include "MuT/NodeSkeletalMeshParameter.h"
 
+class ICustomizableObjectNodeComponentMeshInterface;
 struct FMutableGraphGenerationContext;
 class UEdGraphPin;
 
@@ -14,4 +16,7 @@ namespace UE::Mutable::Private
 	Ptr<NodeSkeletalMesh> GenerateMutableSourceSkeletalMesh(const UEdGraphPin* Pin, FMutableGraphGenerationContext& GenerationContext);
 
 	Ptr<NodeSkeletalMeshParameter> GenerateMutableSourceSkeletalMeshParameter(const UEdGraphPin* Pin, FMutableGraphGenerationContext& GenerationContext);
+
+	Ptr<NodeSkeletalMeshNew> GenerateMutableSourceSkeletalMeshNew(const UEdGraphPin* Pin, FMutableGraphGenerationContext& GenerationContext);
+
 }
diff --git a/Engine/Plugins/Mutable/Source/CustomizableObjectEditor/Private/MuCOE/Nodes/CustomizableObjectNodeComponentPassthroughMesh.cpp b/Engine/Plugins/Mutable/Source/CustomizableObjectEditor/Private/MuCOE/Nodes/CustomizableObjectNodeComponentPassthroughMesh.cpp
--- a/Engine/Plugins/Mutable/Source/CustomizableObjectEditor/Private/MuCOE/Nodes/CustomizableObjectNodeComponentPassthroughMesh.cpp
+++ b/Engine/Plugins/Mutable/Source/CustomizableObjectEditor/Private/MuCOE/Nodes/CustomizableObjectNodeComponentPassthroughMesh.cpp
@@ -18,65 +18,23 @@ class UObject;
 
 #define LOCTEXT_NAMESPACE "CustomizableObjectEditor"
 
-const FName UCustomizableObjectNodeComponentPassthroughMesh::OutputPinName(TEXT("Component"));
-
-
-void UCustomizableObjectNodeComponentMeshPinDataSection::Init(int32 InLODIndex, int32 InSectionIndex)
-{
-	LODIndex = InLODIndex;
-	SectionIndex = InSectionIndex;
-}
-
-
-int32 UCustomizableObjectNodeComponentMeshPinDataSection::GetLODIndex() const
-{
-	return LODIndex;
-}
-
-
-int32 UCustomizableObjectNodeComponentMeshPinDataSection::GetSectionIndex() const
-{
-	return SectionIndex;
-}
-
-
-bool UCustomizableObjectNodeComponentMeshPinDataSection::Equals(const UCustomizableObjectNodePinData& Other) const
-{
-	if (GetClass() != Other.GetClass())
-	{
-		return false;
-	}
-
-	const UCustomizableObjectNodeComponentMeshPinDataSection& OtherTyped = static_cast<const UCustomizableObjectNodeComponentMeshPinDataSection&>(Other);
-	if (LODIndex != OtherTyped.LODIndex ||
-		SectionIndex != OtherTyped.SectionIndex)
-	{
-		return false;
-	}
-
-	return Super::Equals(Other);
-}
-
-
-void UCustomizableObjectNodeComponentPassthroughMesh::PostEditChangeProperty(FPropertyChangedEvent& PropertyChangedEvent)
-{
-	Super::PostEditChangeProperty(PropertyChangedEvent);
-
-	FProperty* PropertyThatChanged = PropertyChangedEvent.Property;
-	if (PropertyThatChanged && PropertyThatChanged->GetName() == GET_MEMBER_NAME_CHECKED(UCustomizableObjectNodeComponentPassthroughMesh, SkeletalMesh))
-	{
-		ReconstructNode();
-	}
-}
-
 
 void UCustomizableObjectNodeComponentPassthroughMesh::BackwardsCompatibleFixup(int32 CustomizableObjectCustomVersion)
 {
 	Super::BackwardsCompatibleFixup(CustomizableObjectCustomVersion);
-
-	if (CustomizableObjectCustomVersion == FCustomizableObjectCustomVersion::ChangedComponentsInheritance)
+	
+	if (CustomizableObjectCustomVersion == FCustomizableObjectCustomVersion::PinSkeletalMesh)
 	{
-		OutputPin = FindPin(OutputPinName, EGPD_Output);
+		TArray<UEdGraphPin*> OldPins = Pins;
+		for (UEdGraphPin* Pin : OldPins)
+		{
+			if (Pin->PinType.PinCategory == UEdGraphSchema_CustomizableObject::PC_MeshSection)
+			{
+				CustomRemovePin(*Pin);
+			}
+		}
+		
+		SkeletalMeshPin = CustomCreatePin(EGPD_Input, UEdGraphSchema_CustomizableObject::PC_PassthroughSkeletalMesh, TEXT("Skeletal Mesh"));
 	}
 }
 
@@ -85,51 +43,7 @@ void UCustomizableObjectNodeComponentPassthroughMesh::AllocateDefaultPins(UCusto
 {
 	Super::AllocateDefaultPins(RemapPins);
 
-	if (SkeletalMesh.IsNull())
-	{
-		return;
-	}
-	
-	// Support SkeletalMeshes only, for now.
-	USkeletalMesh* StrongSkeletalMesh = UE::Mutable::Private::LoadObject<USkeletalMesh>(SkeletalMesh);
-	if (!StrongSkeletalMesh)
-	{
-		return;
-	}
-	
-	if (const FSkeletalMeshModel* ImportedModel = StrongSkeletalMesh->GetImportedModel())
-	{
-		const int32 NumLODs = StrongSkeletalMesh->GetLODNum();
-		for (int32 LODIndex = 0; LODIndex < NumLODs; ++LODIndex)
-		{
-			const int32 NumSections = ImportedModel->LODModels[LODIndex].Sections.Num();
-			for (int32 SectionIndex = 0; SectionIndex < NumSections; ++SectionIndex)
-			{
-				// Ignore disabled sections.
-				const bool bIsSectionDisabled = ImportedModel ->LODModels[LODIndex].Sections[SectionIndex].bDisabled;
-				if (bIsSectionDisabled)
-				{
-					continue;
-				}
-
-				UMaterialInterface* MaterialInterface = GetMaterialInterfaceFor(LODIndex, SectionIndex, ImportedModel);
-				
-				// Material for the section
-				{
-					UCustomizableObjectNodeComponentMeshPinDataMaterial* PinData = NewObject<UCustomizableObjectNodeComponentMeshPinDataMaterial>(this);
-					PinData->Init(LODIndex, SectionIndex);
-					
-					FString SectionFriendlyName = MaterialInterface ? MaterialInterface->GetName() : FString::Printf(TEXT("Section %i"), SectionIndex);
-					FString MaterialPinName = FString::Printf(TEXT("LOD %i - Section %i - Material"), LODIndex, SectionIndex);
-
-					UEdGraphPin* Pin = CustomCreatePin(EGPD_Input, UEdGraphSchema_CustomizableObject::PC_MeshSection, FName(*MaterialPinName), PinData);
-					Pin->PinFriendlyName = FText::FromString(FString::Printf(TEXT("LOD %i - Section %i - %s"), LODIndex, SectionIndex, *SectionFriendlyName));
-					Pin->PinToolTip = MaterialPinName;
-				}
-				
-			}
-		}
-	}
+	SkeletalMeshPin = CustomCreatePin(EGPD_Input, UEdGraphSchema_CustomizableObject::PC_PassthroughSkeletalMesh, TEXT("Skeletal Mesh"));
 }
 
 
@@ -145,129 +59,11 @@ FText UCustomizableObjectNodeComponentPassthroughMesh::GetNodeTitle(ENodeTitleTy
 }
 
 
-void UCustomizableObjectNodeComponentPassthroughMesh::GetPinSection(const UEdGraphPin& Pin, int32& OutLODIndex, int32& OutSectionIndex) const
-{
-	if (const UCustomizableObjectNodeComponentMeshPinDataSection* PinData = Cast<UCustomizableObjectNodeComponentMeshPinDataSection>(GetPinData(Pin)))
-	{
-		OutLODIndex = PinData->GetLODIndex();
-		OutSectionIndex = PinData->GetSectionIndex();
-		return;
-	}
-
-	OutLODIndex = -1;
-	OutSectionIndex = -1;
-}
-
-
-UMaterialInterface* UCustomizableObjectNodeComponentPassthroughMesh::GetMaterialFor(const UEdGraphPin* Pin) const
-{
-	USkeletalMesh* StrongSkeletalMesh = UE::Mutable::Private::LoadObject<USkeletalMesh>(SkeletalMesh);
-	if (StrongSkeletalMesh)
-	{
-		if (FSkeletalMaterial* SkeletalMaterial = GetSkeletalMaterialFor(*Pin))
-		{
-			return SkeletalMaterial->MaterialInterface;
-		}
-	}
-
-	return nullptr;
-}
-
-
-FSkeletalMaterial* UCustomizableObjectNodeComponentPassthroughMesh::GetSkeletalMaterialFor(const UEdGraphPin& Pin) const
-{
-	int32 LODIndex;
-	int32 SectionIndex;
-	GetPinSection(Pin, LODIndex, SectionIndex);
-
-	return GetSkeletalMaterialFor(LODIndex, SectionIndex);
-}
-
-
 FText UCustomizableObjectNodeComponentPassthroughMesh::GetTooltipText() const
 {
 	return LOCTEXT("ComponentPassthroughMesh_Tooltip", "Define a new object component based on a Skeletal Mesh.");
 }
 
 
-UMaterialInterface* UCustomizableObjectNodeComponentPassthroughMesh::GetMaterialInterfaceFor(const int32 LODIndex, const int32 SectionIndex, const FSkeletalMeshModel* ImportedModel) const
-{
-	if (FSkeletalMaterial* SkeletalMaterial = GetSkeletalMaterialFor(LODIndex, SectionIndex, ImportedModel))
-	{
-		return SkeletalMaterial->MaterialInterface;
-	}
-
-	return nullptr;
-}
-
-
-FSkeletalMaterial* UCustomizableObjectNodeComponentPassthroughMesh::GetSkeletalMaterialFor(const int32 LODIndex, const int32 SectionIndex, const FSkeletalMeshModel* ImportedModel) const
-{
-	USkeletalMesh* StrongSkeletalMesh = Cast<USkeletalMesh>(UE::Mutable::Private::LoadObject(SkeletalMesh));
-	if (!StrongSkeletalMesh)
-	{
-		return nullptr;
-	}
-
-	// We assume that LODIndex and MaterialIndex are valid for the imported model
-	int32 SkeletalMeshMaterialIndex = INDEX_NONE;
-	
-	// Check if we have lod info map to get the correct material index
-	if (const FSkeletalMeshLODInfo* LodInfo = StrongSkeletalMesh->GetLODInfo(LODIndex))
-	{
-		if (LodInfo->LODMaterialMap.IsValidIndex(SectionIndex))
-		{
-			SkeletalMeshMaterialIndex = LodInfo->LODMaterialMap[SectionIndex];
-		}
-	}
-
-	// Only deduce index when the explicit mapping is not found or there is no remap
-	if (SkeletalMeshMaterialIndex == INDEX_NONE)
-	{
-	if (ImportedModel && ImportedModel->LODModels.IsValidIndex(LODIndex) && ImportedModel->LODModels[LODIndex].Sections.IsValidIndex(SectionIndex))
-	{
-		SkeletalMeshMaterialIndex = ImportedModel->LODModels[LODIndex].Sections[SectionIndex].MaterialIndex;
-	}
-	else
-	{
-		FSkeletalMeshModel* AuxImportedModel = StrongSkeletalMesh->GetImportedModel();
-
-		if (AuxImportedModel)
-		{
-			if (AuxImportedModel->LODModels.IsValidIndex(LODIndex) && AuxImportedModel->LODModels[LODIndex].Sections.IsValidIndex(SectionIndex))
-			{
-				SkeletalMeshMaterialIndex = AuxImportedModel->LODModels[LODIndex].Sections[SectionIndex].MaterialIndex;
-			}
-		}
-	}
-	}
-	
-	if (StrongSkeletalMesh->GetMaterials().IsValidIndex(SkeletalMeshMaterialIndex))
-	{
-		return &StrongSkeletalMesh->GetMaterials()[SkeletalMeshMaterialIndex];
-	}
-	
-	return nullptr;
-}
-
-
-
-UEdGraphPin* UCustomizableObjectNodeComponentPassthroughMesh::GetMaterialPin(const int32 LODIndex, const int32 SectionIndex) const
-{
-	for (UEdGraphPin* Pin : GetAllNonOrphanPins())
-	{
-		if (const UCustomizableObjectNodeComponentMeshPinDataSection* PinData = Cast<UCustomizableObjectNodeComponentMeshPinDataSection>(GetPinData(*Pin)))
-		{
-			if (PinData->GetLODIndex() == LODIndex &&
-				PinData->GetSectionIndex() == SectionIndex)
-			{
-				return Pin;
-			}
-		}
-	}
-
-	return nullptr;
-}
-
-
 #undef LOCTEXT_NAMESPACE
+
diff --git a/Engine/Plugins/Mutable/Source/CustomizableObjectEditor/Private/MuCOE/Nodes/CustomizableObjectNodeComponentPassthroughMesh.h b/Engine/Plugins/Mutable/Source/CustomizableObjectEditor/Private/MuCOE/Nodes/CustomizableObjectNodeComponentPassthroughMesh.h
--- a/Engine/Plugins/Mutable/Source/CustomizableObjectEditor/Private/MuCOE/Nodes/CustomizableObjectNodeComponentPassthroughMesh.h
+++ b/Engine/Plugins/Mutable/Source/CustomizableObjectEditor/Private/MuCOE/Nodes/CustomizableObjectNodeComponentPassthroughMesh.h
@@ -3,9 +3,6 @@
 #pragma once
 
 #include "MuCOE/Nodes/CustomizableObjectNodeComponent.h"
-#include "MuCOE/RemapPins/CustomizableObjectNodeRemapPinsByNameDefaultPin.h"
-#include "SGraphNode.h"
-#include "GameplayTagContainer.h"
 
 #include "CustomizableObjectNodeComponentPassthroughMesh.generated.h"
 
@@ -30,50 +27,13 @@ struct FPropertyChangedEvent;
 struct FSkeletalMaterial;
 
 
-/** PinData of a pin that belongs to a Skeletal Mesh Section. */
-UCLASS()
-class UCustomizableObjectNodeComponentMeshPinDataSection : public UCustomizableObjectNodePinData
-{
-	GENERATED_BODY()
-
-public:
-	void Init(int32 InLODIndex, int32 InSectionIndex);
-
-	int32 GetLODIndex() const;
-
-	int32 GetSectionIndex() const;
-
-protected:
-	virtual bool Equals(const UCustomizableObjectNodePinData& Other) const override;
-
-private:
-	UPROPERTY()
-	int32 LODIndex = -1;
-
-	UPROPERTY()
-	int32 SectionIndex = -1;
-};
-
-
-/** PinData of a Mesh pin. */
-UCLASS()
-class UCustomizableObjectNodeComponentMeshPinDataMaterial : public UCustomizableObjectNodeComponentMeshPinDataSection
-{
-	GENERATED_BODY()
-};
-
-
 UCLASS(MinimalAPI)
 class UCustomizableObjectNodeComponentPassthroughMesh : public UCustomizableObjectNodeComponent
 {
 public:
 	GENERATED_BODY()
-
-	UPROPERTY(EditAnywhere, Category = PassthroughMesh)
-	TSoftObjectPtr<USkeletalMesh> SkeletalMesh;
-
+	
 	// UObject interface.
-	UE_API virtual void PostEditChangeProperty(FPropertyChangedEvent& PropertyChangedEvent) override;
 	UE_API virtual void BackwardsCompatibleFixup(int32 CustomizableObjectCustomVersion) override;
 	
 	// UEdGraphNode interface
@@ -84,23 +44,9 @@ public:
 	UE_API virtual void AllocateDefaultPins(UCustomizableObjectNodeRemapPins* RemapPins) override;
 	UE_API virtual bool IsExperimental() const override;
 
-	// Own interface
-	
-	UE_API void GetPinSection(const UEdGraphPin& Pin, int32& OutLODIndex, int32& OutSectionIndex) const;
-
-	/** Find the pin for a given lod and section. */
-	UE_API UEdGraphPin* GetMaterialPin(const int32 LODIndex, const int32 SectionIndex) const;
-
-	/** Returns the material associated to the given output pin. */
-	UE_API UMaterialInterface* GetMaterialFor(const UEdGraphPin* Pin) const;
-	UE_API FSkeletalMaterial* GetSkeletalMaterialFor(const UEdGraphPin& Pin) const;
-
-private:
-	UE_API UMaterialInterface* GetMaterialInterfaceFor(const int32 LODIndex, const int32 MaterialIndex, const FSkeletalMeshModel* ImportedModel = nullptr) const;
-	UE_API FSkeletalMaterial* GetSkeletalMaterialFor(const int32 LODIndex, const int32 MaterialIndex, const FSkeletalMeshModel* ImportedModel = nullptr) const;
-
-	static UE_API const FName OutputPinName;
-
+	UPROPERTY()
+	FEdGraphPinReference SkeletalMeshPin;
 };
 
+
 #undef UE_API
diff --git a/Engine/Plugins/Mutable/Source/CustomizableObjectEditor/Private/MuCOE/Nodes/CustomizableObjectNodeSkeletalMeshParameter.cpp b/Engine/Plugins/Mutable/Source/CustomizableObjectEditor/Private/MuCOE/Nodes/CustomizableObjectNodeSkeletalMeshParameter.cpp
--- a/Engine/Plugins/Mutable/Source/CustomizableObjectEditor/Private/MuCOE/Nodes/CustomizableObjectNodeSkeletalMeshParameter.cpp
+++ b/Engine/Plugins/Mutable/Source/CustomizableObjectEditor/Private/MuCOE/Nodes/CustomizableObjectNodeSkeletalMeshParameter.cpp
@@ -152,7 +152,8 @@ UEdGraphPin* UCustomizableObjectNodeSkeletalMeshParameter::GetMeshPin(const int3
 	{
 		if (const UCustomizableObjectNodeSkeletalMeshParameterPinDataSection* PinData = Cast<UCustomizableObjectNodeSkeletalMeshParameterPinDataSection>(GetPinData(*Pin)))
 		{
-			if (PinData->GetSectionIndex() == SectionIndex)
+			if (PinData->GetLODIndex() == LODIndex &&
+				PinData->GetSectionIndex() == SectionIndex)
 			{
 				return Pin;
 			}
diff --git a/Engine/Plugins/Mutable/Source/CustomizableObjectEditor/Private/MuCOE/Nodes/CustomizableObjectNodeSkeletalMeshParameter.h b/Engine/Plugins/Mutable/Source/CustomizableObjectEditor/Private/MuCOE/Nodes/CustomizableObjectNodeSkeletalMeshParameter.h
--- a/Engine/Plugins/Mutable/Source/CustomizableObjectEditor/Private/MuCOE/Nodes/CustomizableObjectNodeSkeletalMeshParameter.h
+++ b/Engine/Plugins/Mutable/Source/CustomizableObjectEditor/Private/MuCOE/Nodes/CustomizableObjectNodeSkeletalMeshParameter.h
@@ -89,6 +89,9 @@ public:
 
 	UPROPERTY(EditAnywhere, Category=CustomizableObject)
 	TSoftObjectPtr<USkeletalMesh> ReferenceValue;
+
+	UPROPERTY()
+	FEdGraphPinReference PassthroughPin;
 };
 
 #undef UE_API
diff --git a/Engine/Plugins/Mutable/Source/CustomizableObjectEditor/Private/MuCOE/SMutableCodeViewer.cpp b/Engine/Plugins/Mutable/Source/CustomizableObjectEditor/Private/MuCOE/SMutableCodeViewer.cpp
--- a/Engine/Plugins/Mutable/Source/CustomizableObjectEditor/Private/MuCOE/SMutableCodeViewer.cpp
+++ b/Engine/Plugins/Mutable/Source/CustomizableObjectEditor/Private/MuCOE/SMutableCodeViewer.cpp
@@ -1692,13 +1692,22 @@ void SMutableCodeViewer::GenerateElementRecursive(const int32& InStateIndex, UE:
 	case UE::Mutable::Private::EOpType::IN_ADDSTRING:
 	case UE::Mutable::Private::EOpType::IN_ADDCOMPONENT:
 	case UE::Mutable::Private::EOpType::IN_ADDSURFACE:
+	case UE::Mutable::Private::EOpType::IN_ADDSKELETALMESH:
 	{
 		UE::Mutable::Private::OP::InstanceAddArgs Args = InProgram.GetOpArgs<UE::Mutable::Private::OP::InstanceAddArgs>(InParentAddress);
 		AddOpFunc(Args.instance, TEXT("instance "));
 		AddOpFunc(Args.value, TEXT("value "));
 		break;
 	}
-		
+
+	case UE::Mutable::Private::EOpType::IN_ADDMATERIAL:
+	{
+		UE::Mutable::Private::OP::InstanceAddMaterialArgs Args = InProgram.GetOpArgs<UE::Mutable::Private::OP::InstanceAddMaterialArgs>(InParentAddress);
+		AddOpFunc(Args.Instance, TEXT("instance "));
+		AddOpFunc(Args.Material, TEXT("material "));
+		break;
+	}
+	
 	case UE::Mutable::Private::EOpType::IM_COMPOSE:
 	{
 		UE::Mutable::Private::OP::ImageComposeArgs Args = InProgram.GetOpArgs<UE::Mutable::Private::OP::ImageComposeArgs>(InParentAddress);
@@ -2854,7 +2863,7 @@ void SMutableCodeViewer::Tick(const FGeometry& AllottedGeometry, const double In
 	UE::Mutable::Private::FSettings Settings;
 	const TSharedPtr<UE::Mutable::Private::FSystem> System = MakeShared<UE::Mutable::Private::FSystem>(Settings);
 
-	System->GetPrivate()->BeginBuild(MutableModel, Object->GetPrivate()->MeshIdRegistry, Object->GetPrivate()->ImageIdRegistry, Object->GetPrivate()->MaterialIdRegistry);
+	System->GetPrivate()->BeginBuild(MutableModel, Object->GetPrivate()->MeshIdRegistry, Object->GetPrivate()->ImageIdRegistry, Object->GetPrivate()->MaterialIdRegistry, Object->GetPrivate()->SkeletalMeshIdRegistry);
 	
 	switch (OperationDataType)
 	{
diff --git a/Engine/Plugins/Mutable/Source/CustomizableObjectEditor/Private/MuCOE/SMutableGraphViewer.cpp b/Engine/Plugins/Mutable/Source/CustomizableObjectEditor/Private/MuCOE/SMutableGraphViewer.cpp
--- a/Engine/Plugins/Mutable/Source/CustomizableObjectEditor/Private/MuCOE/SMutableGraphViewer.cpp
+++ b/Engine/Plugins/Mutable/Source/CustomizableObjectEditor/Private/MuCOE/SMutableGraphViewer.cpp
@@ -61,6 +61,7 @@
 #include "MuT/NodeObjectNew.h"
 #include "MuT/NodeModifierMeshTransformInMesh.h"
 #include "MuT/NodeModifierMeshTransformWithBone.h"
+#include "MuT/NodeSkeletalMeshNew.h"
 #include "Widgets/MutableExpanderArrow.h"
 
 class FExtender;
@@ -102,8 +103,8 @@ static const char* MutableNodeNames[] =
 	"MeshVariation",
 	"MeshReshape",
 	"MeshClipDeform",
-	"MeshParameter",
 	"MeshExternal",
+	"MeshSkeletalMeshBreak",
 
 	"Image",
 	"ImageConstant",
@@ -188,8 +189,9 @@ static const char* MutableNodeNames[] =
 	"LOD",
 
 	"SkeletalMesh",
+	"SkeletalMeshNew",
 	"SkeletalMeshParameter",
-	
+
 	"Component",
 	"ComponentNew",
 	"ComponentEdit",
@@ -554,22 +556,15 @@ void SMutableGraphViewer::GetChildrenForInfo(TSharedPtr<FMutableGraphTreeElement
 	else if (ParentNode->GetType() == UE::Mutable::Private::NodeComponentNew::GetStaticType())
 	{
 		UE::Mutable::Private::NodeComponentNew* ComponentVar = StaticCast<UE::Mutable::Private::NodeComponentNew*>(ParentNode);
-		for (int32 LODIndex = 0; LODIndex < ComponentVar->LODs.Num(); LODIndex++)
-		{
-			AddChildFunc(ComponentVar->LODs[LODIndex].get(), FString::Printf(TEXT("LOD [%d]"), LODIndex));
-		}
-
+		AddChildFunc(ComponentVar->SkeletalMesh.get(), TEXT("SKELETAL MESH"));
 		AddChildFunc(ComponentVar->OverlayMaterial.get(), TEXT("OVERLAY MATERIAL"));
 	}
 
 	else if (ParentNode->GetType() == UE::Mutable::Private::NodeComponentEdit::GetStaticType())
 	{
 		UE::Mutable::Private::NodeComponentEdit* ComponentVar = StaticCast<UE::Mutable::Private::NodeComponentEdit*>(ParentNode);
-		for (int32 LODIndex = 0; LODIndex < ComponentVar->LODs.Num(); LODIndex++)
-		{
-			AddChildFunc(ComponentVar->LODs[LODIndex].get(), FString::Printf(TEXT("LOD [%d]"), LODIndex));
-		}
-		}
+		AddChildFunc(ComponentVar->SkeletalMesh.get(), TEXT("SKELETAL MESH"));
+	}
 
 	else if (ParentNode->GetType() == UE::Mutable::Private::NodeComponentSwitch::GetStaticType())
 	{
@@ -581,6 +576,24 @@ void SMutableGraphViewer::GetChildrenForInfo(TSharedPtr<FMutableGraphTreeElement
 		}
 	}
 
+	else if (ParentNode->GetType() == UE::Mutable::Private::NodeSkeletalMeshParameter::GetStaticType())
+	{
+		UE::Mutable::Private::NodeSkeletalMeshParameter* SkeletalMeshVar = StaticCast<UE::Mutable::Private::NodeSkeletalMeshParameter*>(ParentNode);
+		for (int32 RangeIndex = 0; RangeIndex < SkeletalMeshVar->Ranges.Num(); RangeIndex++)
+		{
+			AddChildFunc(SkeletalMeshVar->Ranges[RangeIndex].get(), FString::Printf(TEXT("RANGE [%d]"), RangeIndex));
+		}
+	}
+	
+	else if (ParentNode->GetType() == UE::Mutable::Private::NodeSkeletalMeshNew::GetStaticType())
+	{
+		UE::Mutable::Private::NodeSkeletalMeshNew* SkeletalMeshVar = StaticCast<UE::Mutable::Private::NodeSkeletalMeshNew*>(ParentNode);
+		for (int32 LODIndex = 0; LODIndex < SkeletalMeshVar->LODs.Num(); LODIndex++)
+		{
+			AddChildFunc(SkeletalMeshVar->LODs[LODIndex].get(), FString::Printf(TEXT("LOD [%d]"), LODIndex));
+		}
+	}
+	
 	else if (ParentNode->GetType() == UE::Mutable::Private::NodeMeshConstant::GetStaticType())
 	{
 		UE::Mutable::Private::NodeMeshConstant* MeshConstantVar = StaticCast<UE::Mutable::Private::NodeMeshConstant*>(ParentNode);
diff --git a/Engine/Plugins/Mutable/Source/MutableRuntime/Internal/MuR/Instance.h b/Engine/Plugins/Mutable/Source/MutableRuntime/Internal/MuR/Instance.h
--- a/Engine/Plugins/Mutable/Source/MutableRuntime/Internal/MuR/Instance.h
+++ b/Engine/Plugins/Mutable/Source/MutableRuntime/Internal/MuR/Instance.h
@@ -74,6 +74,8 @@ namespace UE::Mutable::Private
 
 		FMaterialId OverlayMaterialId;
 
+		FSkeletalMeshId SkeletalMeshId;
+
 		TArray<FInstanceLOD, TInlineAllocator<1>> LODs;
 	};
 
@@ -246,10 +248,12 @@ namespace UE::Mutable::Private
         //! \param StringIndex Index of the string, from 0 to GetStringCount()-1
         UE_API FName GetStringName( int32 ComponentIndex, int32 LODIndex, int32 SurfaceIndex, int32 StringIndex ) const;
 
+		UE_API FSkeletalMeshId GetSkeletalMeshId(int32 ComponentIndex) const;
+		
 		//! Get the index of the overlay material.
 		//! \param ComponentIndex Index of the component, from 0 to GetComponentCount()-1
 		UE_API FMaterialId GetOverlayMaterialId(int32 ComponentIndex) const;
-
+		
 		UE_API FMaterialId GetMaterialId(int32 ComponentIndex, int32 LODIndex, int32 SurfaceIndex) const;
 
 		//! Get the number of ExtensionData values in a component
@@ -281,6 +285,8 @@ namespace UE::Mutable::Private
 
 		void SetMaterialId(int32 ComponentIndex, int32 LODIndex, int32 SurfaceIndex, const FMaterialId& MaterialId);
 
+		void SetSkeletalMeshId(int32 ComponentIndex, const FSkeletalMeshId& SkeletalMeshId);
+
 		void SetOverlayMaterialId(int32 ComponentIndex, const FMaterialId& MaterialId);
 
 		void SetExtensionData(const TSharedRef<const FExtensionData>& Data, FName Name);
diff --git a/Engine/Plugins/Mutable/Source/MutableRuntime/Internal/MuR/Mesh.h b/Engine/Plugins/Mutable/Source/MutableRuntime/Internal/MuR/Mesh.h
--- a/Engine/Plugins/Mutable/Source/MutableRuntime/Internal/MuR/Mesh.h
+++ b/Engine/Plugins/Mutable/Source/MutableRuntime/Internal/MuR/Mesh.h
@@ -214,10 +214,6 @@ namespace UE::Mutable::Private
 
 		/** The mesh is formatted to be used for wrapping projection */
 		ProjectWrappingFormat = 1 << 1,
-
-		/** The mesh is a reference to an external resource mesh and must be loaded when first referenced. */
-		IsResourceForceLoad = 1 << 3,
-
 	};
 
 
@@ -239,7 +235,7 @@ namespace UE::Mutable::Private
         //-----------------------------------------------------------------------------------------
 				
 		/** Create a new empty mesh that repreents an external resource mesh. */
-		static UE_API TSharedPtr<FMesh> CreateAsReference(const TPassthroughObjectPtr<USkeletalMesh>& PasstroughObjectPtr, bool bForceLoad);
+		static UE_API TSharedPtr<FMesh> CreateAsReference(const TPassthroughObjectPtr<USkeletalMesh>& PasstroughObjectPtr);
 
         /** Deep clone this mesh. */
         UE_API TSharedPtr<FMesh> Clone() const;
@@ -263,9 +259,6 @@ namespace UE::Mutable::Private
 
 		/** Return true if this is a reference to an engine image. */
 		UE_API bool IsReference() const;
-
-		/** If true, this is a reference that must be resolved at compile time. */
-		UE_API bool IsForceLoad() const;
     	
 		/** Mesh references can actually reference a morphed mesh. In that case you can set and get the morph with these functions. */
 		UE_API void SetReferencedMorph(const FString& MorphName);
@@ -441,7 +434,7 @@ namespace UE::Mutable::Private
 		TArray<TStrongObjectPtr<USkeletalMesh>> SkeletalMeshes;
 
 		/** Passthrough Id or TStrongObjectPtr. */
-		TPassthroughObjectPtr<USkeletalMesh> PassthroughObject;
+		TPassthroughObjectPtr<USkeletalMesh> PassthroughObject; // TODO GMT It should not be a TPassthroughObject but instead directly UObject.
 
 		/** If the mesh is a reference the referenced morph name is stored here. Otherwise it is an empty string. */
 		FString ReferencedMorph;
diff --git a/Engine/Plugins/Mutable/Source/MutableRuntime/Internal/MuR/Operations.h b/Engine/Plugins/Mutable/Source/MutableRuntime/Internal/MuR/Operations.h
--- a/Engine/Plugins/Mutable/Source/MutableRuntime/Internal/MuR/Operations.h
+++ b/Engine/Plugins/Mutable/Source/MutableRuntime/Internal/MuR/Operations.h
@@ -345,6 +345,9 @@ namespace UE::Mutable::Private
         //! Add all LODs to an instance. This operation can only appear once in a model.
         IN_ADDLOD,
 
+    	//! Add a Skeletal Mesh to an instance.
+		IN_ADDSKELETALMESH,
+
 		//! Add extension data to an instance
 		IN_ADDEXTENSIONDATA,
 
diff --git a/Engine/Plugins/Mutable/Source/MutableRuntime/Internal/MuR/ResourceID.h b/Engine/Plugins/Mutable/Source/MutableRuntime/Internal/MuR/ResourceID.h
--- a/Engine/Plugins/Mutable/Source/MutableRuntime/Internal/MuR/ResourceID.h
+++ b/Engine/Plugins/Mutable/Source/MutableRuntime/Internal/MuR/ResourceID.h
@@ -17,6 +17,8 @@ namespace UE::Mutable::Private
 		TMemoryTrackedArray<uint8> ParameterValuesBlob;
 		
 		bool operator==(const FGeneratedMeshKey&) const = default;
+
+		bool operator!=(const FGeneratedMeshKey&) const = default;
 	};
 
 	struct FGeneratedMeshData
@@ -40,7 +42,6 @@ namespace UE::Mutable::Private
 
 		bool operator==(const FGeneratedImageKey&) const = default;
 	};
-
 	
 	struct FGeneratedImageData
 	{
@@ -52,7 +53,7 @@ namespace UE::Mutable::Private
 
 	typedef TRegistry<FGeneratedImageKey, FGeneratedImageData> FImageIdRegistry;
 
-	MUTABLERUNTIME_API uint32 GetTypeHashPersistent(const FImageId& Key);
+	MUTABLERUNTIME_API uint32 GetTypeHashPersistent(const FImageId& Id);
 
 	// Material
 	struct FGeneratedMaterialKey
@@ -63,7 +64,6 @@ namespace UE::Mutable::Private
 
 		bool operator==(const FGeneratedMaterialKey&) const = default;
 	};
-
 	
 	struct FGeneratedMaterialData
 	{
@@ -75,5 +75,28 @@ namespace UE::Mutable::Private
 
 	typedef TRegistry<FGeneratedMaterialKey, FGeneratedMaterialData> FMaterialIdRegistry;
 
-	MUTABLERUNTIME_API uint32 GetTypeHashPersistent(const FMaterialId& Key);
+	MUTABLERUNTIME_API uint32 GetTypeHashPersistent(const FMaterialId& Id);
+
+
+	// SkeletalMesh
+	struct FGeneratedSkeletalMeshKey
+	{
+		OP::ADDRESS Address = 0;
+
+		TMemoryTrackedArray<uint8> ParameterValuesBlob;
+
+		bool operator==(const FGeneratedSkeletalMeshKey&) const = default;
+	};
+
+	struct FGeneratedSkeletalMeshData
+	{
+	};
+	
+	MUTABLERUNTIME_API uint32 GetTypeHash(const FGeneratedSkeletalMeshKey& Key);
+	
+	typedef TRegistry<FGeneratedSkeletalMeshKey, FGeneratedSkeletalMeshData>::FHandle FSkeletalMeshId;
+
+	typedef TRegistry<FGeneratedSkeletalMeshKey, FGeneratedSkeletalMeshData> FSkeletalMeshIdRegistry;
+
+	MUTABLERUNTIME_API uint32 GetTypeHashPersistent(const FSkeletalMeshId& Id);
 }
diff --git a/Engine/Plugins/Mutable/Source/MutableRuntime/Internal/MuR/System.h b/Engine/Plugins/Mutable/Source/MutableRuntime/Internal/MuR/System.h
--- a/Engine/Plugins/Mutable/Source/MutableRuntime/Internal/MuR/System.h
+++ b/Engine/Plugins/Mutable/Source/MutableRuntime/Internal/MuR/System.h
@@ -2,6 +2,7 @@
 
 #pragma once
 
+#include "SkeletalMesh.h"
 #include "MuR/Image.h"
 #include "MuR/Mesh.h"
 #include "MuR/Instance.h"
@@ -158,6 +159,7 @@ namespace UE::Mutable::Private
 			const TSharedPtr<FMeshIdRegistry>& InMeshIdRegistry,
 			const TSharedPtr<FImageIdRegistry>& InImageIdRegistry,
 			const TSharedPtr<FMaterialIdRegistry>& InMaterialIdRegistry,
+			const TSharedPtr<FSkeletalMeshIdRegistry>& InSkeletalMeshIdRegistry,
 			int32 StateIndex,
 			uint32 LodMask);
 
@@ -167,6 +169,7 @@ namespace UE::Mutable::Private
 			const TSharedPtr<FMeshIdRegistry>& InMeshIdRegistry,
 			const TSharedPtr<FImageIdRegistry>& InImageIdRegistry,
 			const TSharedPtr<FMaterialIdRegistry>& InMaterialIdRegistry,
+			const TSharedPtr<FSkeletalMeshIdRegistry>& InSkeletalMeshIdRegistry,
 			int32 InStateIndex,
 			uint32 InLodMask);
 		
@@ -180,8 +183,10 @@ namespace UE::Mutable::Private
 		//! the image will still contain a minimum number of mips specified at model compile time.
 		UE_API UE::Tasks::TTask<TSharedPtr<const FImage>> GetImage(FInstance::FID InstanceID, const FImageId& ImageId, int32 MipsToSkip = 0, int32 LOD = 0);
 
-        //! Only valid between BeginUpdate and EndUpdate
+		//! Only valid between BeginUpdate and EndUpdate
 		UE_API UE::Tasks::TTask<TSharedPtr<const FMesh>> GetMesh(FInstance::FID InstanceID, const FMeshId& MeshId, EMeshContentFlags MeshContentFilter = EMeshContentFlags::AllFlags);
+		
+		UE_API UE::Tasks::TTask<TSharedPtr<const FSkeletalMesh>> GetSkeletalMesh(FInstance::FID InstanceID, const FSkeletalMeshId& SkeletalMeshId, EMeshContentFlags MeshContentFilter);
 
 		//! Only valid between BeginUpdate and EndUpdate
 		//! Calculate the description of an image, without generating it.
@@ -194,7 +199,7 @@ namespace UE::Mutable::Private
 		UE_API TSharedPtr<const FImage> GetImageInline(FInstance::FID InstanceID, const FImageId& ImageId, int32 MipsToSkip = 0, int32 LOD = 0);
 
         //! Only valid between BeginUpdate and EndUpdate
-		UE_API TSharedPtr<const FMesh> GetMeshInline(FInstance::FID InstanceID, const FImageId& MeshId, EMeshContentFlags MeshContentFilter = EMeshContentFlags::AllFlags);
+		UE_API TSharedPtr<const FMesh> GetMeshInline(FInstance::FID InstanceID, const FMeshId& MeshId, EMeshContentFlags MeshContentFilter = EMeshContentFlags::AllFlags);
 		
 		UE_API UE::Tasks::TTask<TSharedPtr<const FMaterial>> GetMaterial(FInstance::FID InstanceID, const FMaterialId& MaterialId);
 		
diff --git a/Engine/Plugins/Mutable/Source/MutableRuntime/Internal/MuR/SystemPrivate.h b/Engine/Plugins/Mutable/Source/MutableRuntime/Internal/MuR/SystemPrivate.h
--- a/Engine/Plugins/Mutable/Source/MutableRuntime/Internal/MuR/SystemPrivate.h
+++ b/Engine/Plugins/Mutable/Source/MutableRuntime/Internal/MuR/SystemPrivate.h
@@ -591,6 +591,10 @@ namespace UE::Mutable::Private
 					MeshResults[DataTypeIndex].Value = nullptr;
 					break;
 
+				case EDataType::SkeletalMesh:
+					SkeletalMeshResults[DataTypeIndex] = nullptr;
+					break;
+
 				case EDataType::Instance:
 					InstanceResults[DataTypeIndex] = nullptr;
 					break;
@@ -1337,7 +1341,9 @@ namespace UE::Mutable::Private
 		FImageId GetImageId(const TSharedPtr<const FModel>& Model, const FParameters* Params, uint32 ParamListIndex, OP::ADDRESS RootAt);
 
     	FMaterialId GetMaterialId(const TSharedPtr<const FModel>& Model, const FParameters* Params, uint32 ParamListIndex, OP::ADDRESS RootAt);
-    	
+
+    	FSkeletalMeshId GetSkeletalMeshId(const TSharedPtr<const FModel>& Model, const FParameters* Params, uint32 ParamListIndex, OP::ADDRESS RootAt);
+
 		/** Maximum working memory that mutable should be using. */
 		int64 BudgetBytes = 0;
 
@@ -1362,6 +1368,8 @@ namespace UE::Mutable::Private
 
     	TSharedPtr<FMaterialIdRegistry> CurrentMaterialIdRegistry;
 
+    	TSharedPtr<FSkeletalMeshIdRegistry> CurrentSkeletalMeshIdRegistry;
+
     	/** Only valid in initial generation. We can not load UObjects while streaming. */
     	TSharedPtr<FPassthroughObjectLoader> CurrentPassthroughObjectLoader;
     	
@@ -1629,6 +1637,14 @@ namespace UE::Mutable::Private
 			return Result;
 		}
 
+
+    	/** */
+    	[[nodiscard]] TSharedPtr<const FSkeletalMesh> LoadSkeletalMesh(const FCacheAddress& From)
+		{
+			return CurrentInstanceCache->GetSkeletalMesh(From);
+		}
+    	
+
 		/** */
 		[[nodiscard]] TSharedPtr<const FImage> LoadImage(const FCacheAddress& From, bool bTakeOwnership = false)
 		{
@@ -1918,16 +1934,20 @@ namespace UE::Mutable::Private
 					break;
 				}
 
+				case EDataType::SkeletalMesh:
+				{
+					CurrentInstanceCache->SkeletalMeshResults[Data.DataTypeIndex] = nullptr;
+					break;
+				}
+
 				case EDataType::Layout:
 				{
-					TSharedPtr<const FLayout> Value = CurrentInstanceCache->LayoutResults[Data.DataTypeIndex];
 					CurrentInstanceCache->LayoutResults[Data.DataTypeIndex] = nullptr;
 					break;
 				}
 
 				case EDataType::Instance:
 				{
-					TSharedPtr<const FInstance> Value = CurrentInstanceCache->InstanceResults[Data.DataTypeIndex];
 					CurrentInstanceCache->InstanceResults[Data.DataTypeIndex] = nullptr;
 					break;
 				}
@@ -1976,16 +1996,20 @@ namespace UE::Mutable::Private
 					break;
 				}
 
+				case EDataType::SkeletalMesh:
+				{
+					CurrentInstanceCache->SkeletalMeshResults[Data.DataTypeIndex] = nullptr;
+					break;
+				}
+
 				case EDataType::Layout:
 				{
-					TSharedPtr<const FLayout> Value = CurrentInstanceCache->LayoutResults[Data.DataTypeIndex];
 					CurrentInstanceCache->LayoutResults[Data.DataTypeIndex] = nullptr;
 					break;
 				}
 
 				case EDataType::Instance:
 				{
-					TSharedPtr<const FInstance> Value = CurrentInstanceCache->InstanceResults[Data.DataTypeIndex];
 					CurrentInstanceCache->InstanceResults[Data.DataTypeIndex] = nullptr;
 					break;
 				}
@@ -2106,7 +2130,8 @@ namespace UE::Mutable::Private
 		MUTABLERUNTIME_API void BeginBuild(const TSharedPtr<const FModel>&,
 			const TSharedPtr<FMeshIdRegistry>&,
 			const TSharedPtr<FImageIdRegistry>&,
-			const TSharedPtr<FMaterialIdRegistry>&);
+			const TSharedPtr<FMaterialIdRegistry>&,
+			const TSharedPtr<FSkeletalMeshIdRegistry>&);
     	
 		MUTABLERUNTIME_API void EndBuild();
 
diff --git a/Engine/Plugins/Mutable/Source/MutableRuntime/Private/MuR/CodeRunner.cpp b/Engine/Plugins/Mutable/Source/MutableRuntime/Private/MuR/CodeRunner.cpp
--- a/Engine/Plugins/Mutable/Source/MutableRuntime/Private/MuR/CodeRunner.cpp
+++ b/Engine/Plugins/Mutable/Source/MutableRuntime/Private/MuR/CodeRunner.cpp
@@ -1109,6 +1109,55 @@ namespace UE::Mutable::Private
         		break;
         	}
 
+		case EOpType::IN_ADDSKELETALMESH:
+        	{
+        		OP::InstanceAddArgs Args = Program.GetOpArgs<OP::InstanceAddArgs>(Item.At);
+        		switch (Item.Stage)
+        		{
+        		case 0:
+        			{
+        				AddOp(FScheduledOp(Item.At, Item, 1),
+							FScheduledOp(Args.instance, Item),
+							FScheduledOp(Args.value, Item));
+						
+        				break;
+        			}
+
+        		case 1:
+        			{
+        				TSharedPtr<const FInstance> InInstance;
+        				if (Args.instance)
+        				{
+        					InInstance = LoadInstance(FCacheAddress(Args.instance, Item));
+        				}
+
+        				TSharedPtr<FInstance> Instance;
+        				if (InInstance)
+        				{
+        					Instance = UE::Mutable::Private::CloneOrTakeOver<FInstance>(InInstance);
+        				}
+        				else
+        				{
+        					Instance = MakeShared<FInstance>();
+        				}
+
+
+        				if (Args.value)
+        				{
+        					FSkeletalMeshId SkeletalMeshId = System->WorkingMemoryManager.GetSkeletalMeshId(Model, InParams, Args.RelevantParametersListIndex, Args.value);
+        					Instance->SetSkeletalMeshId(0, SkeletalMeshId);
+        				}
+						
+        				StoreInstance( Item, Instance);
+        				break;
+        			}
+
+        		default:
+        			unimplemented();
+        		}
+
+        		break;
+        	}
         	
         default:
 			check(false);
@@ -1223,36 +1272,7 @@ namespace UE::Mutable::Private
 
         switch (type)
         {
-
-		case EOpType::ME_REFERENCE:
-		{
-			OP::ResourceReferenceArgs Args = Program.GetOpArgs<OP::ResourceReferenceArgs>(Item.At);
-			switch (Item.Stage)
-			{
-			case 0:
-			{
-				TSharedPtr<FMesh> Result;
-				if (Args.ForceLoad)
-				{
-					// This should never be reached because it should have been caught as a Task in IssueOp
-					check(false);
-				}
-				else
-				{
-					TPassthroughObjectPtr<USkeletalMesh> ExternalObject = FindPassthrough<USkeletalMesh>(Args.ID);
-					Result = FMesh::CreateAsReference(ExternalObject, false);
-				}
-				StoreMesh(Item, Result);
-				break;
-			}
-
-			default:
-				check(false);
-			}
-
-			break;
-		}
-
+        
         case EOpType::ME_APPLYLAYOUT:
         {
 			OP::MeshApplyLayoutArgs Args = Program.GetOpArgs<OP::MeshApplyLayoutArgs>(Item.At);
@@ -6799,6 +6819,7 @@ namespace UE::Mutable::Private
         case EOpType::IN_ADDIMAGE:
         case EOpType::IN_ADDMATERIAL:
         case EOpType::IN_ADDOVERLAYMATERIAL:
+    	case EOpType::IN_ADDSKELETALMESH:
             RunCode_InstanceAddResource(Item, InModel, pParams);
             break;
 
diff --git a/Engine/Plugins/Mutable/Source/MutableRuntime/Private/MuR/CodeRunnerBegin.cpp b/Engine/Plugins/Mutable/Source/MutableRuntime/Private/MuR/CodeRunnerBegin.cpp
--- a/Engine/Plugins/Mutable/Source/MutableRuntime/Private/MuR/CodeRunnerBegin.cpp
+++ b/Engine/Plugins/Mutable/Source/MutableRuntime/Private/MuR/CodeRunnerBegin.cpp
@@ -1204,6 +1204,55 @@ namespace UE::Mutable::Private
 		            break;
 		        }
 				
+			case EOpType::IN_ADDSKELETALMESH:
+				{
+					OP::InstanceAddArgs Args = Program.GetOpArgs<OP::InstanceAddArgs>(Item.At);
+					switch (Item.Stage)
+					{
+					case 0:
+						{
+							PushOp(FScheduledOpInline(Item, 1));
+							PushOp(FScheduledOpInline(Args.instance, Item));
+							PushOp(FScheduledOpInline(Args.value, Item));
+							
+							break;
+						}
+
+					case 1:
+						{
+							TSharedPtr<const FInstance> InInstance;
+							if (Args.instance)
+							{
+								InInstance = LoadInstance(Args.instance);
+							}
+
+							TSharedPtr<FInstance> Instance;
+							if (InInstance)
+							{
+								Instance = UE::Mutable::Private::CloneOrTakeOver<FInstance>(InInstance);
+							}
+							else
+							{
+								Instance = MakeShared<FInstance>();
+							}
+							
+							if (Args.value)
+							{
+								FSkeletalMeshId SkeletalMeshId = System.GetPrivate()->WorkingMemoryManager.GetSkeletalMeshId(Model, Params.Get(), Args.RelevantParametersListIndex, Args.value);
+								Instance->SetSkeletalMeshId(0, SkeletalMeshId);
+							}
+							
+							StoreInstance( Item, Instance);
+							break;
+						}
+
+					default:
+						unimplemented();
+					}
+
+					break;
+				}
+				
 			case EOpType::IN_ADDEXTENSIONDATA:
 				{
 					OP::InstanceAddExtensionDataArgs Args = Program.GetOpArgs<OP::InstanceAddExtensionDataArgs>(Item.At);
@@ -1627,20 +1676,7 @@ namespace UE::Mutable::Private
 					StoreNone(Item);
 					break;
 				}
-				
-			case EOpType::ME_REFERENCE:
-				{
-					OP::ResourceReferenceArgs Args = Program.GetOpArgs<OP::ResourceReferenceArgs>(Item.At);
-
-					if (!Args.ForceLoad) // Not a passthrough. Only happens with None optimization.
-					{
-						GetMemory().CurrentPassthroughObjectLoader->Add<USkeletalMesh>(Args.ID);
-					}
-
-					StoreNone(Item);
-					break;
-				}
-
+			
 			case EOpType::ME_APPLYLAYOUT:
 				{
 					OP::MeshApplyLayoutArgs Args = Program.GetOpArgs<OP::MeshApplyLayoutArgs>(Item.At);
diff --git a/Engine/Plugins/Mutable/Source/MutableRuntime/Private/MuR/CodeRunnerTasked.cpp b/Engine/Plugins/Mutable/Source/MutableRuntime/Private/MuR/CodeRunnerTasked.cpp
--- a/Engine/Plugins/Mutable/Source/MutableRuntime/Private/MuR/CodeRunnerTasked.cpp
+++ b/Engine/Plugins/Mutable/Source/MutableRuntime/Private/MuR/CodeRunnerTasked.cpp
@@ -2687,29 +2687,6 @@ namespace UE::Mutable::Private
 			break;
 		}
 
-		case EOpType::ME_REFERENCE:
-		{
-			OP::ResourceReferenceArgs Args = Model->GetPrivate()->Program.GetOpArgs<OP::ResourceReferenceArgs>(Item.At);
-
-			// We only convert references to meshes if indicated in the operation.
-			if (Args.ForceLoad)
-			{
-				check(Item.Stage == 0);
-
-				FExternalResourceId FullId;
-				FullId.ReferenceResourceId = Args.ID;
-
-				// TODO
-				int32 LODIndex = 0;
-				int32 SectionIndex = 0;
-				int32 MeshID = 0;
-
-				Issued = MakeShared<FMeshExternalLoadTask>(Item, FullId, LODIndex, SectionIndex, MeshID);
-			}
-
-			break;
-		}
-
 		case EOpType::IM_PIXELFORMAT:
 		{
 			if (Item.Stage == 1)
diff --git a/Engine/Plugins/Mutable/Source/MutableRuntime/Private/MuR/Instance.cpp b/Engine/Plugins/Mutable/Source/MutableRuntime/Private/MuR/Instance.cpp
--- a/Engine/Plugins/Mutable/Source/MutableRuntime/Private/MuR/Instance.cpp
+++ b/Engine/Plugins/Mutable/Source/MutableRuntime/Private/MuR/Instance.cpp
@@ -319,12 +319,21 @@ namespace UE::Mutable::Private
         return NAME_None;
     }
 
+	
+	FSkeletalMeshId FInstance::GetSkeletalMeshId(int32 ComponentIndex) const
+	{
+    	check(Components.IsValidIndex(ComponentIndex));
+    	return Components[ComponentIndex].SkeletalMeshId;
+	}
+
+	
 	FMaterialId FInstance::GetOverlayMaterialId(int32 ComponentIndex) const
 	{
 		check(Components.IsValidIndex(ComponentIndex));
 		return Components[ComponentIndex].OverlayMaterialId;
 	}
 
+	
 	FMaterialId FInstance::GetMaterialId(int32 ComponentIndex, int32 LODIndex, int32 SurfaceIndex) const
 	{
 		check(Components.IsValidIndex(ComponentIndex));
@@ -528,6 +537,17 @@ namespace UE::Mutable::Private
 		Components[ComponentIndex].LODs[LODIndex].Surfaces[SurfaceIndex].MaterialId = MaterialId;
 	}
 
+	
+	void FInstance::SetSkeletalMeshId(int32 ComponentIndex, const FSkeletalMeshId& SkeletalMeshId)
+    {
+    	while (ComponentIndex >= Components.Num())
+    	{
+    		AddComponent();
+    	}
+
+    	Components[ComponentIndex].SkeletalMeshId = SkeletalMeshId;
+    }
+
 
 	void FInstance::SetOverlayMaterialId(int32 ComponentIndex, const FMaterialId& MaterialId)
 	{
diff --git a/Engine/Plugins/Mutable/Source/MutableRuntime/Private/MuR/Mesh.cpp b/Engine/Plugins/Mutable/Source/MutableRuntime/Private/MuR/Mesh.cpp
--- a/Engine/Plugins/Mutable/Source/MutableRuntime/Private/MuR/Mesh.cpp
+++ b/Engine/Plugins/Mutable/Source/MutableRuntime/Private/MuR/Mesh.cpp
@@ -43,16 +43,11 @@ TSharedPtr<FMesh> FMesh::StaticUnserialise(FInputArchive& Arch)
 }
 
 
-TSharedPtr<FMesh> FMesh::CreateAsReference(const TPassthroughObjectPtr<USkeletalMesh>& PasstroughObjectPtr,  bool bForceLoad)
+TSharedPtr<FMesh> FMesh::CreateAsReference(const TPassthroughObjectPtr<USkeletalMesh>& PasstroughObjectPtr)
 {
 	TSharedPtr<FMesh> Result = MakeShared<FMesh>();
 	Result->PassthroughObject = PasstroughObjectPtr;
-
-	if (bForceLoad)
-	{
-		EnumAddFlags(Result->Flags, EMeshFlags::IsResourceForceLoad);
-	}
-
+	
 	return Result;
 }
 
@@ -61,12 +56,6 @@ bool FMesh::IsReference() const
 {
 	return PassthroughObject.IsValid();
 }
-
-
-bool FMesh::IsForceLoad() const
-{
-	return EnumHasAnyFlags(Flags, EMeshFlags::IsResourceForceLoad);
-}
 	
 
 void FMesh::SetReferencedMorph(const FString& MorphName)
diff --git a/Engine/Plugins/Mutable/Source/MutableRuntime/Private/MuR/Operations.cpp b/Engine/Plugins/Mutable/Source/MutableRuntime/Private/MuR/Operations.cpp
--- a/Engine/Plugins/Mutable/Source/MutableRuntime/Private/MuR/Operations.cpp
+++ b/Engine/Plugins/Mutable/Source/MutableRuntime/Private/MuR/Operations.cpp
@@ -152,6 +152,7 @@ namespace UE::Mutable::Private
 		{ EDataType::Instance,		},	// IN_ADDSURFACE
 		{ EDataType::Instance,		},	// IN_ADDCOMPONENT
 		{ EDataType::Instance,		},	// IN_ADDLOD
+		{ EDataType::Instance,		},	// IN_ADDSKELETALMESH
 		{ EDataType::Instance,		},	// IN_ADDEXTENSIONDATA
 		{ EDataType::Instance,		},	// IN_ADDOVERLAYMATERIAL
 		{ EDataType::Instance,		},	// IN_ADDMATERIAL
@@ -902,6 +903,16 @@ namespace UE::Mutable::Private
 			break;
 		}
 
+        case EOpType::IN_ADDSKELETALMESH:
+        {
+        	const OP::InstanceAddArgs Args = program.GetOpArgs<OP::InstanceAddArgs>(at);
+
+        	f(Args.instance);
+        	f(Args.value);
+
+        	break;
+        }
+
         //-------------------------------------------------------------------------------------
         case EOpType::LA_PACK:
         {
diff --git a/Engine/Plugins/Mutable/Source/MutableRuntime/Private/MuR/ResourceID.cpp b/Engine/Plugins/Mutable/Source/MutableRuntime/Private/MuR/ResourceID.cpp
--- a/Engine/Plugins/Mutable/Source/MutableRuntime/Private/MuR/ResourceID.cpp
+++ b/Engine/Plugins/Mutable/Source/MutableRuntime/Private/MuR/ResourceID.cpp
@@ -59,4 +59,22 @@ uint32 UE::Mutable::Private::GetTypeHashPersistent(const FMaterialId& Id)
 	{
 		return 0;
 	}
-}
\ No newline at end of file
+}
+
+
+uint32 UE::Mutable::Private::GetTypeHash(const FGeneratedSkeletalMeshKey& Key)
+{
+	return HashCombineFast(Key.Address, static_cast<uint32>(FXxHash64::HashBuffer(Key.ParameterValuesBlob.GetData(), Key.ParameterValuesBlob.Num()).Hash));
+}
+
+uint32 UE::Mutable::Private::GetTypeHashPersistent(const FSkeletalMeshId& Id)
+{
+	if (const FGeneratedSkeletalMeshKey* Key = Id.GetKey())
+	{
+		return HashCombine(Key->Address, static_cast<uint32>(FXxHash64::HashBuffer(Key->ParameterValuesBlob.GetData(), Key->ParameterValuesBlob.Num()).Hash));
+	}
+	else
+	{
+		return 0;
+	}
+}
diff --git a/Engine/Plugins/Mutable/Source/MutableRuntime/Private/MuR/System.cpp b/Engine/Plugins/Mutable/Source/MutableRuntime/Private/MuR/System.cpp
--- a/Engine/Plugins/Mutable/Source/MutableRuntime/Private/MuR/System.cpp
+++ b/Engine/Plugins/Mutable/Source/MutableRuntime/Private/MuR/System.cpp
@@ -239,6 +239,7 @@ namespace UE::Mutable::Private
                                      const TSharedPtr<FMeshIdRegistry>& InMeshIdRegistry,
 									 const TSharedPtr<FImageIdRegistry>& InImageIdRegistry,
 									 const TSharedPtr<FMaterialIdRegistry>& InMaterialIdRegistry,
+									 const TSharedPtr<FSkeletalMeshIdRegistry>& InSkeletalMeshIdRegistry,
 									 int32 InStateIndex,
                                      uint32 InLodMask)
     {
@@ -264,6 +265,7 @@ namespace UE::Mutable::Private
 		m_pD->WorkingMemoryManager.CurrentMeshIdRegistry = InMeshIdRegistry;
     	m_pD->WorkingMemoryManager.CurrentImageIdRegistry = InImageIdRegistry;
     	m_pD->WorkingMemoryManager.CurrentMaterialIdRegistry = InMaterialIdRegistry;
+    	m_pD->WorkingMemoryManager.CurrentSkeletalMeshIdRegistry = InSkeletalMeshIdRegistry;
     	
     	
 		OP::ADDRESS RootAt = LiveInstance->Model->GetPrivate()->Program.States[InStateIndex].Root;
@@ -283,7 +285,8 @@ namespace UE::Mutable::Private
     	m_pD->WorkingMemoryManager.CurrentMeshIdRegistry = nullptr;
     	m_pD->WorkingMemoryManager.CurrentImageIdRegistry = nullptr;
     	m_pD->WorkingMemoryManager.CurrentMaterialIdRegistry = nullptr;
-    	
+    	m_pD->WorkingMemoryManager.CurrentSkeletalMeshIdRegistry = nullptr;
+
 		return Result;
 	}
 
@@ -293,6 +296,7 @@ namespace UE::Mutable::Private
 		const TSharedPtr<FMeshIdRegistry>& InMeshIdRegistry,
 		const TSharedPtr<FImageIdRegistry>& InImageIdRegistry,
 		const TSharedPtr<FMaterialIdRegistry>& InMaterialIdRegistry,
+		const TSharedPtr<FSkeletalMeshIdRegistry>& InSkeletalMeshIdRegistry,
 		int32 InStateIndex,
 		uint32 InLodMask)
     {
@@ -319,6 +323,7 @@ namespace UE::Mutable::Private
     	m_pD->WorkingMemoryManager.CurrentMeshIdRegistry = InMeshIdRegistry;
     	m_pD->WorkingMemoryManager.CurrentImageIdRegistry = InImageIdRegistry;
     	m_pD->WorkingMemoryManager.CurrentMaterialIdRegistry = InMaterialIdRegistry;
+    	m_pD->WorkingMemoryManager.CurrentSkeletalMeshIdRegistry = InSkeletalMeshIdRegistry;
 
 		FProgram& Program = LiveInstance->Model->GetPrivate()->Program;
 
@@ -383,7 +388,8 @@ namespace UE::Mutable::Private
     	m_pD->WorkingMemoryManager.CurrentMeshIdRegistry = nullptr;
     	m_pD->WorkingMemoryManager.CurrentImageIdRegistry = nullptr;
     	m_pD->WorkingMemoryManager.CurrentMaterialIdRegistry = nullptr;
-    	
+    	m_pD->WorkingMemoryManager.CurrentSkeletalMeshIdRegistry = nullptr;
+
 		return Result;
 	}
 
@@ -648,7 +654,8 @@ namespace UE::Mutable::Private
 
 	}	
 
-    TSharedPtr<const FMesh> FSystem::GetMeshInline(FInstance::FID instanceID, const FImageId& MeshId, EMeshContentFlags MeshContentFilter)
+	
+    TSharedPtr<const FMesh> FSystem::GetMeshInline(FInstance::FID instanceID, const FMeshId& MeshId, EMeshContentFlags MeshContentFilter)
     {
 		LLM_SCOPE_BYNAME(TEXT("MutableRuntime"));
 		MUTABLE_CPUPROFILER_SCOPE(SystemGetMesh);
@@ -744,6 +751,75 @@ namespace UE::Mutable::Private
 						Result = SystemPrivate->WorkingMemoryManager.LoadMesh(FCacheAddress(RootAddress, 0, ExecutionOptions), true);
 					}
 					
+						SystemPrivate->WorkingMemoryManager.EndRunnerThread();
+						SystemPrivate->WorkingMemoryManager.CurrentInstanceCache = nullptr;
+						SystemPrivate->WorkingMemoryManager.CurrentPassthroughObjectLoader = nullptr;
+					
+						return Result;
+					},
+					UE::Tasks::Prerequisites(RunnerCompletionEvent),
+					UE::Tasks::ETaskPriority::Inherit,
+					UE::Tasks::EExtendedTaskPriority::Inline);
+    }
+
+
+	UE::Tasks::TTask<TSharedPtr<const FSkeletalMesh>> FSystem::GetSkeletalMesh(FInstance::FID InstanceID, const FSkeletalMeshId& SkeletalMeshId, EMeshContentFlags SkeletalMeshContentFilter)
+    {
+		LLM_SCOPE_BYNAME(TEXT("MutableRuntime"));
+		MUTABLE_CPUPROFILER_SCOPE(SystemGetSkeletalMesh);
+		
+		if (!SkeletalMeshId)
+		{
+			return UE::Tasks::MakeCompletedTask<TSharedPtr<const FSkeletalMesh>>(new FSkeletalMesh());
+		}
+		
+		// Find the live instance
+		FLiveInstance* LiveInstance = m_pD->FindLiveInstance(InstanceID);
+		check(LiveInstance);
+		m_pD->WorkingMemoryManager.CurrentInstanceCache = LiveInstance->Cache;
+		m_pD->WorkingMemoryManager.CurrentPassthroughObjectLoader = LiveInstance->PassthroughObjectLoader;
+		m_pD->WorkingMemoryManager.BeginRunnerThread();
+		
+		OP::ADDRESS RootAddress = SkeletalMeshId.GetKey()->Address;
+			
+		UE::Mutable::Private::EOpType OpType = LiveInstance->Model->GetPrivate()->Program.GetOpType(RootAddress);
+		if (GetOpDataType(OpType) != EDataType::SkeletalMesh)
+		{
+			m_pD->WorkingMemoryManager.EndRunnerThread();
+			m_pD->WorkingMemoryManager.CurrentInstanceCache = nullptr;
+			m_pD->WorkingMemoryManager.CurrentPassthroughObjectLoader = nullptr;
+
+			return UE::Tasks::MakeCompletedTask<TSharedPtr<const FSkeletalMesh>>(new FSkeletalMesh());
+		}
+	
+		uint8 ExecutionOptions = static_cast<uint8>(SkeletalMeshContentFilter);
+		TSharedRef<CodeRunner> Runner = CodeRunner::Create(
+				m_pD->Settings,
+				m_pD, 
+				EExecutionStrategy::MinimizeMemory, 
+				LiveInstance->Model, 
+				LiveInstance->OldParameters.Get(), 
+				RootAddress, 
+				FSystem::AllLODs, 
+				ExecutionOptions, 
+				0, 
+				FScheduledOp::EType::Full,
+				LiveInstance->ExternalResourceProvider);
+		
+		constexpr bool bForceInlineExecution = false;
+		UE::Tasks::FTask RunnerCompletionEvent = Runner->StartRun(bForceInlineExecution);
+
+		return UE::Tasks::Launch(TEXT("FSystem::GetSkeletalMeshResultTask"),
+				[SystemPrivate = m_pD, Runner, RootAddress, ExecutionOptions]() -> TSharedPtr<const FSkeletalMesh>
+				{
+					TSharedPtr<const FSkeletalMesh> Result;
+
+					SystemPrivate->bUnrecoverableError = Runner->bUnrecoverableError;
+					if (!Runner->bUnrecoverableError)
+					{
+						Result = SystemPrivate->WorkingMemoryManager.LoadSkeletalMesh(FCacheAddress(RootAddress, 0, ExecutionOptions));
+					}
+					
 					SystemPrivate->WorkingMemoryManager.EndRunnerThread();
 					SystemPrivate->WorkingMemoryManager.CurrentInstanceCache = nullptr;
 					SystemPrivate->WorkingMemoryManager.CurrentPassthroughObjectLoader = nullptr;
@@ -1051,7 +1127,8 @@ namespace UE::Mutable::Private
 	void FSystem::Private::BeginBuild(const TSharedPtr<const FModel>& InModel,
 		const TSharedPtr<FMeshIdRegistry>& InMeshIdRegistry,
 		const TSharedPtr<FImageIdRegistry>& InImageIdRegistry,
-		const TSharedPtr<FMaterialIdRegistry>& InMaterialIdRegistry)
+		const TSharedPtr<FMaterialIdRegistry>& InMaterialIdRegistry,
+		const TSharedPtr<FSkeletalMeshIdRegistry>& InSkeletalMeshIdRegistry)
 	{
 		// We don't have a FLiveInstance, let's create the memory
 		// \TODO: There is no clear moment to remove this... EndBuild?
@@ -1063,6 +1140,7 @@ namespace UE::Mutable::Private
     	WorkingMemoryManager.CurrentMeshIdRegistry = InMeshIdRegistry;
     	WorkingMemoryManager.CurrentImageIdRegistry = InImageIdRegistry;
     	WorkingMemoryManager.CurrentMaterialIdRegistry = InMaterialIdRegistry;
+    	WorkingMemoryManager.CurrentSkeletalMeshIdRegistry = InSkeletalMeshIdRegistry;
     	
 		// Ensure the model cache has been created
 		WorkingMemoryManager.FindOrAddModelCache(InModel);
@@ -1079,6 +1157,7 @@ namespace UE::Mutable::Private
     	WorkingMemoryManager.CurrentMeshIdRegistry = nullptr;
     	WorkingMemoryManager.CurrentImageIdRegistry = nullptr;
     	WorkingMemoryManager.CurrentMaterialIdRegistry = nullptr;
+    	WorkingMemoryManager.CurrentSkeletalMeshIdRegistry = nullptr;
 	}
 
 
@@ -2285,4 +2364,15 @@ namespace UE::Mutable::Private
 
     	return CurrentMaterialIdRegistry->Add(NewKey, {});
     }
+
+	
+	FSkeletalMeshId FWorkingMemoryManager::GetSkeletalMeshId(const TSharedPtr<const FModel>& Model, const FParameters* Params, uint32 ParamListIndex, OP::ADDRESS RootAt)
+	{
+    	FGeneratedSkeletalMeshKey NewKey;
+    	NewKey.Address = RootAt;
+
+    	GetResourceKey(Model, Params, ParamListIndex, RootAt, NewKey.ParameterValuesBlob);
+
+    	return CurrentSkeletalMeshIdRegistry->Add(NewKey, {});
+	}
 }
diff --git a/Engine/Plugins/Mutable/Source/MutableTools/Internal/MuT/Node.h b/Engine/Plugins/Mutable/Source/MutableTools/Internal/MuT/Node.h
--- a/Engine/Plugins/Mutable/Source/MutableTools/Internal/MuT/Node.h
+++ b/Engine/Plugins/Mutable/Source/MutableTools/Internal/MuT/Node.h
@@ -134,6 +134,7 @@ namespace UE::Mutable::Private
 			LOD,
 
 			SkeletalMesh,
+			SkeletalMeshNew,
 			SkeletalMeshParameter,
 			
 			Component,
diff --git a/Engine/Plugins/Mutable/Source/MutableTools/Internal/MuT/NodeComponent.h b/Engine/Plugins/Mutable/Source/MutableTools/Internal/MuT/NodeComponent.h
--- a/Engine/Plugins/Mutable/Source/MutableTools/Internal/MuT/NodeComponent.h
+++ b/Engine/Plugins/Mutable/Source/MutableTools/Internal/MuT/NodeComponent.h
@@ -13,16 +13,11 @@
 
 namespace UE::Mutable::Private
 {
+	class NodeSkeletalMeshNew;
 
-	/** */
 	class NodeComponent : public Node
 	{
 	public:
-
-		TArray<Ptr<NodeLOD>> LODs;
-
-	public:
-
         // Node interface
 		virtual const FNodeType* GetType() const override { return &StaticType; }
 		static const FNodeType* GetStaticType() { return &StaticType; }
diff --git a/Engine/Plugins/Mutable/Source/MutableTools/Internal/MuT/NodeComponentEdit.h b/Engine/Plugins/Mutable/Source/MutableTools/Internal/MuT/NodeComponentEdit.h
--- a/Engine/Plugins/Mutable/Source/MutableTools/Internal/MuT/NodeComponentEdit.h
+++ b/Engine/Plugins/Mutable/Source/MutableTools/Internal/MuT/NodeComponentEdit.h
@@ -3,7 +3,7 @@
 #pragma once
 
 #include "MuR/Ptr.h"
-#include "MuR/RefCounted.h"
+#include "MuT/NodeSkeletalMeshNew.h"
 #include "MuT/Node.h"
 #include "MuT/NodeComponent.h"
 
@@ -22,7 +22,7 @@ namespace UE::Mutable::Private
 
 		Ptr<NodeComponent> Parent;
 
-	public:
+		Ptr<NodeSkeletalMeshNew> SkeletalMesh;
 
 		// Node interface
 		virtual const FNodeType* GetType() const override { return GetStaticType(); }
diff --git a/Engine/Plugins/Mutable/Source/MutableTools/Internal/MuT/NodeComponentNew.h b/Engine/Plugins/Mutable/Source/MutableTools/Internal/MuT/NodeComponentNew.h
--- a/Engine/Plugins/Mutable/Source/MutableTools/Internal/MuT/NodeComponentNew.h
+++ b/Engine/Plugins/Mutable/Source/MutableTools/Internal/MuT/NodeComponentNew.h
@@ -4,7 +4,7 @@
 
 #include "HAL/Platform.h"
 #include "MuR/Ptr.h"
-#include "MuR/RefCounted.h"
+#include "MuT/NodeSkeletalMeshNew.h"
 #include "MuT/Node.h"
 #include "MuT/NodeComponent.h"
 #include "MuT/NodeMaterial.h"
@@ -23,8 +23,8 @@ namespace UE::Mutable::Private
 
 		Ptr<NodeMaterial> OverlayMaterial;
 
-	public:
-
+		Ptr<NodeSkeletalMesh> SkeletalMesh;
+		
 		//-----------------------------------------------------------------------------------------
         // Node interface
 		//-----------------------------------------------------------------------------------------		
diff --git a/Engine/Plugins/Mutable/Source/MutableTools/Internal/MuT/NodeSkeletalMeshNew.h b/Engine/Plugins/Mutable/Source/MutableTools/Internal/MuT/NodeSkeletalMeshNew.h
new file mode 100644
--- /dev/null
+++ b/Engine/Plugins/Mutable/Source/MutableTools/Internal/MuT/NodeSkeletalMeshNew.h
@@ -0,0 +1,36 @@
+// Copyright Epic Games, Inc. All Rights Reserved.
+
+#pragma once
+
+#include "NodeSkeletalMesh.h"
+#include "HAL/Platform.h"
+#include "MuR/PassthroughObject.h"
+#include "MuR/Ptr.h"
+#include "MuT/Node.h"
+#include "MuT/NodeLOD.h"
+
+#define UE_API MUTABLETOOLS_API
+
+
+namespace UE::Mutable::Private
+{
+	class NodeSkeletalMeshNew : public NodeSkeletalMesh
+	{
+	public:
+		TArray<Ptr<NodeLOD>> LODs;
+
+		// Node interface
+		virtual const FNodeType* GetType() const override { return &StaticType; }
+		static const FNodeType* GetStaticType() { return &StaticType; }
+	
+	protected:
+		virtual ~NodeSkeletalMeshNew() override {}
+
+	private:
+		static UE_API FNodeType StaticType;
+	};
+}
+
+
+#undef UE_API
+
diff --git a/Engine/Plugins/Mutable/Source/MutableTools/Private/MuT/ASTOpAddSkeletalMesh.cpp b/Engine/Plugins/Mutable/Source/MutableTools/Private/MuT/ASTOpAddSkeletalMesh.cpp
new file mode 100644
--- /dev/null
+++ b/Engine/Plugins/Mutable/Source/MutableTools/Private/MuT/ASTOpAddSkeletalMesh.cpp
@@ -0,0 +1,118 @@
+// Copyright Epic Games, Inc. All Rights Reserved.
+
+#include "MuT/ASTOpAddSkeletalMesh.h"
+
+#include "HAL/PlatformMath.h"
+#include "MuR/ModelPrivate.h"
+#include "MuT/CodeOptimiser.h"
+
+
+namespace UE::Mutable::Private
+{
+	ASTOpAddSkeletalMesh::ASTOpAddSkeletalMesh() :
+		Instance(this),
+		SkeletalMesh(this)
+	{
+	}
+
+
+	ASTOpAddSkeletalMesh::~ASTOpAddSkeletalMesh()
+	{
+		RemoveChildren();
+	}
+
+	bool ASTOpAddSkeletalMesh::IsEqual(const ASTOp& OtherUntyped) const
+	{
+		if (OtherUntyped.GetOpType() == GetOpType())
+		{
+			const ASTOpAddSkeletalMesh* Other = static_cast<const ASTOpAddSkeletalMesh*>(&OtherUntyped);
+			return Instance == Other->Instance && SkeletalMesh == Other->SkeletalMesh;
+		}
+
+		return false;
+	}
+
+		
+	uint64 ASTOpAddSkeletalMesh::Hash() const
+	{
+		uint64 Result = std::hash<void*>()(Instance.child().get());
+		
+		hash_combine(Result, SkeletalMesh.child().get());
+		
+		return Result;
+	}
+
+		
+	Ptr<ASTOp> ASTOpAddSkeletalMesh::Clone(MapChildFuncRef MapChild) const
+	{
+		Ptr<ASTOpAddSkeletalMesh> New = new ASTOpAddSkeletalMesh();
+
+		New->Instance = MapChild(Instance.child());
+		New->SkeletalMesh = MapChild(SkeletalMesh.child());
+
+		return New;
+	}
+
+	void ASTOpAddSkeletalMesh::ForEachChild(const TFunctionRef<void(ASTChild&)> Func)
+	{
+		Func(Instance);
+		Func(SkeletalMesh);
+	}
+		
+
+	void ASTOpAddSkeletalMesh::Link(FProgram& Program, FLinkerOptions*)
+	{
+		if (linkedAddress)
+		{
+			return;
+		}
+
+		OP::InstanceAddArgs Args;
+		FMemory::Memzero(Args);
+
+		if (Instance)
+		{
+			Args.instance = Instance->linkedAddress;
+		}
+
+		if (SkeletalMesh)
+		{
+			Args.value = SkeletalMesh->linkedAddress;
+		}
+
+		SubtreeRelevantParametersVisitorAST Visitor;
+		Visitor.Run(SkeletalMesh.child());
+
+		TArray<uint16> Params;
+		for (const FString& ParameterName : Visitor.Parameters)
+		{
+			for (int32 Index = 0; Index < Program.Parameters.Num(); ++Index)
+			{
+				const FParameterDesc& ParameterDesc = Program.Parameters[Index];
+				if (ParameterDesc.Name == ParameterName)
+				{
+					Params.Add(static_cast<uint16>(Index));
+					break;
+				}
+			}
+		}
+
+		Params.Sort();
+
+		int32 Index = Program.ParameterLists.Find(Params);
+		if (Index != INDEX_NONE)
+		{
+			Args.RelevantParametersListIndex = Index;
+		}
+		else
+		{
+			Args.RelevantParametersListIndex = static_cast<uint32>(Program.ParameterLists.Num());
+			Program.ParameterLists.Add(Params);
+		}
+		
+		linkedAddress = static_cast<OP::ADDRESS>(Program.OpAddress.Num());
+		Program.OpAddress.Add(static_cast<uint32_t>(Program.ByteCode.Num()));
+		AppendCode(Program.ByteCode, GetOpType());
+		AppendCode(Program.ByteCode, Args);
+	}
+}
diff --git a/Engine/Plugins/Mutable/Source/MutableTools/Private/MuT/ASTOpAddSkeletalMesh.h b/Engine/Plugins/Mutable/Source/MutableTools/Private/MuT/ASTOpAddSkeletalMesh.h
new file mode 100644
--- /dev/null
+++ b/Engine/Plugins/Mutable/Source/MutableTools/Private/MuT/ASTOpAddSkeletalMesh.h
@@ -0,0 +1,35 @@
+// Copyright Epic Games, Inc. All Rights Reserved.
+
+#pragma once
+
+#include "HAL/Platform.h"
+#include "Math/NumericLimits.h"
+#include "MuR/Operations.h"
+#include "MuR/Ptr.h"
+#include "MuT/AST.h"
+
+
+namespace UE::Mutable::Private
+{
+	struct FProgram;
+
+
+	class ASTOpAddSkeletalMesh final : public ASTOp
+	{
+	public:
+		ASTChild Instance;
+		ASTChild SkeletalMesh;
+
+		ASTOpAddSkeletalMesh();
+		ASTOpAddSkeletalMesh(const ASTOpAddSkeletalMesh&) = delete;
+		~ASTOpAddSkeletalMesh();
+
+		EOpType GetOpType() const override { return EOpType::IN_ADDSKELETALMESH; }
+		uint64 Hash() const override;
+		void ForEachChild(const TFunctionRef<void(ASTChild&)> F) override;
+		bool IsEqual(const ASTOp& OtherUntyped) const override;
+		Ptr<ASTOp> Clone(MapChildFuncRef MapChild) const override;
+		void Link(FProgram& Program, FLinkerOptions* Options) override;
+	};
+}
+
diff --git a/Engine/Plugins/Mutable/Source/MutableTools/Private/MuT/CodeGenerator.cpp b/Engine/Plugins/Mutable/Source/MutableTools/Private/MuT/CodeGenerator.cpp
--- a/Engine/Plugins/Mutable/Source/MutableTools/Private/MuT/CodeGenerator.cpp
+++ b/Engine/Plugins/Mutable/Source/MutableTools/Private/MuT/CodeGenerator.cpp
@@ -2,6 +2,7 @@
 
 #include "MuT/CodeGenerator.h"
 
+#include "ASTOpAddSkeletalMesh.h"
 #include "Containers/Array.h"
 #include "Logging/LogCategory.h"
 #include "Logging/LogMacros.h"
@@ -574,44 +575,27 @@ namespace UE::Mutable::Private
 
 	void CodeGenerator::GenerateComponent_New(const FComponentGenerationOptions& Options, FGenericGenerationResult& Result, const NodeComponentNew* InNode)
 	{
-		TArray<FLODTask> LODTasks;
-		LODTasks.Reserve(InNode->LODs.Num());
+		FSkeletalMeshGenerationResult SkeletalMeshResult;
+		FSkeletalMeshGenerationOptions SkeletalMeshOptions(Options, InNode);
+		GenerateSkeletalMesh(SkeletalMeshResult, SkeletalMeshOptions, InNode->SkeletalMesh.get());
 
-		// Launch tasks for each lod, making sure each LOD depends on the previous one.
-		FLODTask LastLODTask;
-		for (int32 LODIndex = 0; LODIndex < InNode->LODs.Num(); ++LODIndex)
-		{
-			if (const NodeLOD* LODNode = InNode->LODs[LODIndex].get())
-			{
-				FLODGenerationOptions LODOptions(Options, LODIndex, InNode );
-				
-				bool bWasEmpty = false;
-				LastLODTask = GenerateLOD(LODOptions, LODNode, LastLODTask);
-				LODTasks.Add(LastLODTask);
-			}
-		}
+		Ptr<ASTOpAddSkeletalMesh> AddSkeletalMeshOp = new ASTOpAddSkeletalMesh();
+		AddSkeletalMeshOp->Instance = SkeletalMeshResult.InstanceOp;
+		AddSkeletalMeshOp->SkeletalMesh = SkeletalMeshResult.SkeletalMeshOp;
+		
+		Ptr<ASTOp> LastInstOp = AddSkeletalMeshOp;
 
 		// Launch the task that generates the component
 		// There could be more concurrency here, but it doesn't look like it is relevant yet.
 		FComponentTask ComponentTask = LocalPipe.Launch(TEXT("MutableComponentNew"),
 			[
-				LODTasks, InNode, Options,
+				LastInstOp, InNode, Options,
 				this
 			]
 			() mutable
 			{
 				FGenericGenerationResult Result;
-
-				// Create the expression for each component in this object
-				Ptr<ASTOpAddLOD> LODsOp = new ASTOpAddLOD();
-				for (FLODTask& LODTask : LODTasks)
-				{
-					FGenericGenerationResult LODResult = LODTask.GetResult();
-					LODsOp->lods.Emplace(LODsOp, LODResult.Op);
-				}
-
-				Ptr<ASTOp> LastInstOp = LODsOp;
-
+					
 				if (Ptr<NodeMaterial> MaterialNode = InNode->OverlayMaterial)
 				{
 					Ptr<ASTOpAddMaterial> AddOverlayMaterialOp = new ASTOpAddMaterial();
@@ -661,9 +645,7 @@ namespace UE::Mutable::Private
 				}
 
 				return Result;
-			},
-			LODTasks
-			);
+			});
 
 		// Sync point: we currently don't support task-based generation beyond components so we wait here.
 		// Otherwise we could return the component task to be chained with higher level tasks.
diff --git a/Engine/Plugins/Mutable/Source/MutableTools/Private/MuT/CodeGenerator.h b/Engine/Plugins/Mutable/Source/MutableTools/Private/MuT/CodeGenerator.h
--- a/Engine/Plugins/Mutable/Source/MutableTools/Private/MuT/CodeGenerator.h
+++ b/Engine/Plugins/Mutable/Source/MutableTools/Private/MuT/CodeGenerator.h
@@ -194,6 +194,18 @@ namespace UE::Mutable::Private
 			Ptr<ASTOp> BaseInstance;
 		};
 
+    	struct FSkeletalMeshGenerationOptions : public FGenericGenerationOptions
+    	{
+    		FSkeletalMeshGenerationOptions(const FGenericGenerationOptions& BaseOptions, const NodeComponentNew* InComponent)
+    		{
+    			Component = InComponent;
+    			State = BaseOptions.State;
+    			ActiveTags = BaseOptions.ActiveTags;
+    		}
+
+    		const NodeComponentNew* Component = nullptr;
+    	};
+    	
 		struct FLODGenerationOptions : public FGenericGenerationOptions
 		{
 			FLODGenerationOptions(const FGenericGenerationOptions& BaseOptions, int32 InLODIndex, const NodeComponentNew* InComponent)
@@ -847,7 +859,8 @@ namespace UE::Mutable::Private
     	
     	struct FSkeletalMeshGenerationResult
     	{
-    		Ptr<ASTOp> Op;
+    		Ptr<ASTOp> InstanceOp;
+    		Ptr<ASTOp> SkeletalMeshOp;
     	};
 
     	typedef TMap<FGeneratedCacheKey, FSkeletalMeshGenerationResult> FGeneratedSkeletalMeshesMap;
@@ -859,8 +872,9 @@ namespace UE::Mutable::Private
     	};
     	FGeneratedSkeletalMeshes GeneratedSkeletalMeshes;
     	
-    	void GenerateSkeletalMesh(FSkeletalMeshGenerationResult& OutResult, const FGenericGenerationOptions&, const Ptr<const NodeSkeletalMesh>&);
-    	void GenerateSkeletalMesh_Parameter(FSkeletalMeshGenerationResult& OutResult, const FGenericGenerationOptions&, const NodeSkeletalMeshParameter*);
+    	void GenerateSkeletalMesh(FSkeletalMeshGenerationResult& OutResult, const FSkeletalMeshGenerationOptions&, const Ptr<const NodeSkeletalMesh>&);
+    	void GenerateSkeletalMesh_New(FSkeletalMeshGenerationResult& OutResult, const FSkeletalMeshGenerationOptions&, const NodeSkeletalMeshNew*);
+    	void GenerateSkeletalMesh_Parameter(FSkeletalMeshGenerationResult& OutResult, const FSkeletalMeshGenerationOptions&, const NodeSkeletalMeshParameter*);
     	
     	//-----------------------------------------------------------------------------------------
     	
diff --git a/Engine/Plugins/Mutable/Source/MutableTools/Private/MuT/CodeGenerator_FirstPass.cpp b/Engine/Plugins/Mutable/Source/MutableTools/Private/MuT/CodeGenerator_FirstPass.cpp
--- a/Engine/Plugins/Mutable/Source/MutableTools/Private/MuT/CodeGenerator_FirstPass.cpp
+++ b/Engine/Plugins/Mutable/Source/MutableTools/Private/MuT/CodeGenerator_FirstPass.cpp
@@ -30,6 +30,7 @@
 #include "MuT/NodeSurfaceVariation.h"
 #include "MuT/NodeSurfaceSwitch.h"
 #include "MuT/NodeScalarEnumParameter.h"
+#include "MuT/NodeSkeletalMeshNew.h"
 
 namespace UE::Mutable::Private
 {
@@ -185,6 +186,14 @@ namespace UE::Mutable::Private
 		{
 			Generate_Modifier(static_cast<const NodeModifier*>(Root));
 		}
+		else if (Root->GetType() == NodeSkeletalMeshNew::GetStaticType())
+		{
+			Generate_SkeletalMeshNew(static_cast<const NodeSkeletalMeshNew*>(Root));
+		}
+		else if (Root->GetType() == NodeSkeletalMeshParameter::GetStaticType())
+		{
+			// Do nothing.
+		}
 		else
 		{
 			// This node type is not supported in this pass.
@@ -432,16 +441,7 @@ namespace UE::Mutable::Private
 
         CurrentComponent = InNode;
 
-		CurrentLOD = 0;
-		for (const Ptr<NodeLOD>& c : InNode->LODs)
-		{
-			if (c)
-			{
-				Generate_LOD(c.get());
-			}
-			++CurrentLOD;
-		}
-		CurrentLOD = -1;
+		Generate_Generic(InNode->SkeletalMesh.get());
 
 		CurrentComponent = nullptr;
 	}
@@ -452,16 +452,7 @@ namespace UE::Mutable::Private
 	{
 		CurrentComponent = InNode->GetParentComponentNew();
 
-		CurrentLOD = 0;
-		for (const Ptr<NodeLOD>& c : InNode->LODs)
-		{
-			if (c)
-			{
-				Generate_LOD(c.get());
-			}
-			++CurrentLOD;
-		}
-		CurrentLOD = -1;
+		Generate_Generic(InNode->SkeletalMesh.get());
 
 		CurrentComponent = nullptr;
 	}
@@ -495,6 +486,21 @@ namespace UE::Mutable::Private
 		CurrentNegativeTags = OldNegativeTags;
 	}
 
+	
+	void FirstPassGenerator::Generate_SkeletalMeshNew(const NodeSkeletalMeshNew* InNode)
+	{
+		CurrentLOD = 0;
+		for (const Ptr<NodeLOD>& c : InNode->LODs)
+		{
+			if (c)
+			{
+				Generate_LOD(c.get());
+			}
+			++CurrentLOD;
+		}
+		CurrentLOD = -1;
+	}
+
 
 	//---------------------------------------------------------------------------------------------
 	void FirstPassGenerator::Generate_ComponentSwitch(const NodeComponentSwitch* InNode)
diff --git a/Engine/Plugins/Mutable/Source/MutableTools/Private/MuT/CodeGenerator_FirstPass.h b/Engine/Plugins/Mutable/Source/MutableTools/Private/MuT/CodeGenerator_FirstPass.h
--- a/Engine/Plugins/Mutable/Source/MutableTools/Private/MuT/CodeGenerator_FirstPass.h
+++ b/Engine/Plugins/Mutable/Source/MutableTools/Private/MuT/CodeGenerator_FirstPass.h
@@ -91,6 +91,7 @@ namespace UE::Mutable::Private
 	};
 
 	using FLODTask = UE::Tasks::TTask<FGenericGenerationResult>;
+	using FSkeletalMeshTask = UE::Tasks::TTask<FGenericGenerationResult>;
 	using FComponentTask = UE::Tasks::TTask<FGenericGenerationResult>;
 
 	/** First pass of the code generation process.
@@ -115,6 +116,7 @@ namespace UE::Mutable::Private
 		void Generate_ComponentEdit(const NodeComponentEdit*);
 		void Generate_ComponentSwitch(const NodeComponentSwitch*);
 		void Generate_ComponentVariation(const NodeComponentVariation*);
+		void Generate_SkeletalMeshNew(const NodeSkeletalMeshNew*);
 		void Generate_LOD(const NodeLOD*);
 		void Generate_ObjectNew(const NodeObjectNew*);
 		void Generate_ObjectGroup(const NodeObjectGroup*);
diff --git a/Engine/Plugins/Mutable/Source/MutableTools/Private/MuT/CodeGenerator_Mesh.cpp b/Engine/Plugins/Mutable/Source/MutableTools/Private/MuT/CodeGenerator_Mesh.cpp
--- a/Engine/Plugins/Mutable/Source/MutableTools/Private/MuT/CodeGenerator_Mesh.cpp
+++ b/Engine/Plugins/Mutable/Source/MutableTools/Private/MuT/CodeGenerator_Mesh.cpp
@@ -750,24 +750,7 @@ namespace UE::Mutable::Private
 		MUTABLE_CPUPROFILER_SCOPE(GenerateMesh_Constant);
 
 		TSharedPtr<FMesh> Mesh = Node->Value;
-
-		// True passthrough?
-		if (Mesh && Mesh->IsReference() && !Mesh->IsForceLoad())
-		{
-			Ptr<ASTOpReferenceResource> ReferenceOp = new ASTOpReferenceResource();
-			ReferenceOp->Type = EOpType::ME_REFERENCE;
-			ReferenceOp->ID = Mesh->PassthroughObject.GetId();
-			ReferenceOp->bForceLoad = false;
-
-			FMeshGenerationResult Result;
-			Result.BaseMeshOp = ReferenceOp;
-			Result.MeshOp = ReferenceOp;
-
-			// We won't be able to do anything with the passthrough mesh at compile time.
-			//check(InOptions.OverrideLayouts.IsEmpty());
-			return UE::Tasks::MakeCompletedTask<FMeshGenerationResult>(Result);
-		}
-
+		
 		if (!Mesh)
 		{
 			Ptr<ASTOpConstantResource> ConstantOp = new ASTOpConstantResource();
@@ -795,7 +778,7 @@ namespace UE::Mutable::Private
 		TSharedPtr< TSharedPtr<FMesh> > ResolveMesh = MakeShared<TSharedPtr<FMesh>>();
 
 		// Compile-time reference? Add task to resolve as requisite.
-		if (Mesh->IsReference() && Mesh->IsForceLoad())
+		if (Mesh->IsReference())
 		{
 			PASSTHROUGH_ID MeshID = Mesh->PassthroughObject.GetId();
 
@@ -1391,7 +1374,7 @@ namespace UE::Mutable::Private
 		TSharedPtr<FMesh> Mesh = Node->ReferenceMesh;
 		if (Mesh)
 		{
-			if (Mesh->IsReference() && Mesh->IsForceLoad())
+			if (Mesh->IsReference())
 			{
 				PASSTHROUGH_ID MeshID = Mesh->PassthroughObject.GetId();
 
@@ -1407,9 +1390,10 @@ namespace UE::Mutable::Private
 		}
 
 		FSkeletalMeshGenerationResult SkeletalMeshResult;
-		GenerateSkeletalMesh(SkeletalMeshResult, StaticOptions, Node->SkeletalMeshParameter);
+		FSkeletalMeshGenerationOptions SkeletalMeshOptions(StaticOptions, nullptr);
+		GenerateSkeletalMesh(SkeletalMeshResult, SkeletalMeshOptions, Node->SkeletalMeshParameter);
 
-		Ptr<ASTOp> SkeletalMeshParameter = SkeletalMeshResult.Op;
+		Ptr<ASTOp> SkeletalMeshParameter = SkeletalMeshResult.SkeletalMeshOp;
 		
 		// Local pipe because we call GenerateRange in there.
 		FMeshTask MeshParamTask = this->LocalPipe.Launch(TEXT("MutableMeshParameter"),
diff --git a/Engine/Plugins/Mutable/Source/MutableTools/Private/MuT/CodeGenerator_SkeletalMesh.cpp b/Engine/Plugins/Mutable/Source/MutableTools/Private/MuT/CodeGenerator_SkeletalMesh.cpp
--- a/Engine/Plugins/Mutable/Source/MutableTools/Private/MuT/CodeGenerator_SkeletalMesh.cpp
+++ b/Engine/Plugins/Mutable/Source/MutableTools/Private/MuT/CodeGenerator_SkeletalMesh.cpp
@@ -10,7 +10,7 @@
 
 namespace UE::Mutable::Private
 {
-    void CodeGenerator::GenerateSkeletalMesh(FSkeletalMeshGenerationResult& Result, const FGenericGenerationOptions& Options, const Ptr<const NodeSkeletalMesh>& Node)
+    void CodeGenerator::GenerateSkeletalMesh(FSkeletalMeshGenerationResult& Result, const FSkeletalMeshGenerationOptions& Options, const Ptr<const NodeSkeletalMesh>& Node)
 	{
     	if (!Node)
     	{
@@ -32,10 +32,13 @@ namespace UE::Mutable::Private
     		}
 	    }
 
-
 	    // Generate for each different type of node
 		switch (Node->GetType()->Type)
 		{
+	    case Node::EType::SkeletalMeshNew:
+			GenerateSkeletalMesh_New(Result, Options, static_cast<const NodeSkeletalMeshNew*>(Node.get()));
+			break;
+
 		case Node::EType::SkeletalMeshParameter:
 			GenerateSkeletalMesh_Parameter(Result, Options, static_cast<const NodeSkeletalMeshParameter*>(Node.get()));
 			break;
@@ -52,7 +55,46 @@ namespace UE::Mutable::Private
 	}
 
 
-	void CodeGenerator::GenerateSkeletalMesh_Parameter(FSkeletalMeshGenerationResult& Result, const FGenericGenerationOptions& Options, const NodeSkeletalMeshParameter* Node)
+    void CodeGenerator::GenerateSkeletalMesh_New(FSkeletalMeshGenerationResult& Result, const FSkeletalMeshGenerationOptions& Options, const NodeSkeletalMeshNew* Node)
+    {
+    	TArray<FLODTask> LODTasks;
+    	LODTasks.Reserve(Node->LODs.Num());
+    	
+    	// Launch tasks for each lod, making sure each LOD depends on the previous one.
+    	FLODTask LastLODTask;
+    	for (int32 LODIndex = 0; LODIndex < Node->LODs.Num(); ++LODIndex)
+    	{
+    		if (const NodeLOD* LODNode = Node->LODs[LODIndex].get())
+    		{
+    			FLODGenerationOptions LODOptions(Options, LODIndex, Options.Component);
+				
+    			bool bWasEmpty = false;
+    			LastLODTask = GenerateLOD(LODOptions, LODNode, LastLODTask);
+    			LODTasks.Add(LastLODTask);
+    		}
+    	}
+
+    	// Create the expression for each component in this object
+    	Ptr<ASTOpAddLOD> LODsOp = new ASTOpAddLOD();
+    	for (FLODTask& LODTask : LODTasks)
+    	{
+    		if (WaitCallback.IsSet())
+			{
+				while (!LODTask.IsCompleted())
+				{
+					WaitCallback();
+				}
+			}
+			
+    		FGenericGenerationResult LODResult = LODTask.GetResult();
+    		LODsOp->lods.Emplace(LODsOp, LODResult.Op);
+    	}
+
+    	Result.InstanceOp = LODsOp;
+    }
+
+
+    void CodeGenerator::GenerateSkeletalMesh_Parameter(FSkeletalMeshGenerationResult& Result, const FSkeletalMeshGenerationOptions& Options, const NodeSkeletalMeshParameter* Node)
 	{
 		Ptr<ASTOpParameter> op;
 
@@ -90,6 +132,6 @@ namespace UE::Mutable::Private
 			}
 		}
 
-		Result.Op = op;
+		Result.SkeletalMeshOp = op;
 	}
 }
diff --git a/Engine/Plugins/Mutable/Source/MutableTools/Private/MuT/CodeOptimiser.cpp b/Engine/Plugins/Mutable/Source/MutableTools/Private/MuT/CodeOptimiser.cpp
--- a/Engine/Plugins/Mutable/Source/MutableTools/Private/MuT/CodeOptimiser.cpp
+++ b/Engine/Plugins/Mutable/Source/MutableTools/Private/MuT/CodeOptimiser.cpp
@@ -414,7 +414,7 @@ namespace UE::Mutable::Private
 			Model->GetPrivate()->Program.States.Add(state);
 
 			TSharedPtr<FParameters> LocalParams = FModel::NewParameters(Model);
-			System->GetPrivate()->BeginBuild(Model, nullptr, nullptr, nullptr);
+			System->GetPrivate()->BeginBuild(Model, nullptr, nullptr, nullptr, nullptr);
 
 			// Calculate the value and replace this op by a constant
 			switch( DataType )
@@ -606,7 +606,6 @@ namespace UE::Mutable::Private
 						break;
 					}
 
-					case EOpType::ME_REFERENCE:
 					case EOpType::IM_REFERENCE:
 						// If we are in a reference-resolution optimization phase, then the ops are not special.
 						if (Pass < 2)
@@ -789,19 +788,11 @@ namespace UE::Mutable::Private
 
 					}
 
-					if (SubgraphType == EOpType::ME_REFERENCE
-						||
-						SubgraphType == EOpType::ME_CONSTANT)
+					if (SubgraphType == EOpType::ME_CONSTANT)
 					{
 						PASSTHROUGH_ID MeshID = PASSTHROUGH_ID_INVALID;
 						FString MeshMorph;
-						if (SubgraphType == EOpType::ME_REFERENCE)
-						{
-							bIsReference = true;
-							const ASTOpReferenceResource* Typed = static_cast<const ASTOpReferenceResource*>(SubgraphRoot.get());
-							MeshID = Typed->ID;
-						}
-						else if (SubgraphType == EOpType::ME_CONSTANT)
+						if (SubgraphType == EOpType::ME_CONSTANT)
 						{
 							const ASTOpConstantResource* Typed = static_cast<const ASTOpConstantResource*>(SubgraphRoot.get());
 							const FMesh* Value = static_cast<const FMesh*>(Typed->GetValue().Get());
diff --git a/Engine/Plugins/Mutable/Source/MutableTools/Private/MuT/Compiler.cpp b/Engine/Plugins/Mutable/Source/MutableTools/Private/MuT/Compiler.cpp
--- a/Engine/Plugins/Mutable/Source/MutableTools/Private/MuT/Compiler.cpp
+++ b/Engine/Plugins/Mutable/Source/MutableTools/Private/MuT/Compiler.cpp
@@ -189,6 +189,7 @@ namespace UE::Mutable::Private
 		{ false,	{ 0, 0, 0, 0, 0, 0, 0, 0, 0, 0, 0, 0, 0, 0, 0, 0, 0, 0, 0, 0, 0, 0, 0, 0, 0, 0, 0, 0, 0 }		},	// IN_ADDSURFACE
 		{ false,	{ 0, 0, 0, 0, 0, 0, 0, 0, 0, 0, 0, 0, 0, 0, 0, 0, 0, 0, 0, 0, 0, 0, 0, 0, 0, 0, 0, 0, 0 }		},	// IN_ADDCOMPONENT
 		{ false,	{ 0, 0, 0, 0, 0, 0, 0, 0, 0, 0, 0, 0, 0, 0, 0, 0, 0, 0, 0, 0, 0, 0, 0, 0, 0, 0, 0, 0, 0 }		},	// IN_ADDLOD
+		{ false,	{ 0, 0, 0, 0, 0, 0, 0, 0, 0, 0, 0, 0, 0, 0, 0, 0, 0, 0, 0, 0, 0, 0, 0, 0, 0, 0, 0, 0, 0 }		},	// IN_ADDSKELETALMESH
 		{ false,	{ 0, 0, 0, 0, 0, 0, 0, 0, 0, 0, 0, 0, 0, 0, 0, 0, 0, 0, 0, 0, 0, 0, 0, 0, 0, 0, 0, 0, 0 }		},	// IN_ADDEXTENSIONDATA
 		{ false,	{ 0, 0, 0, 0, 0, 0, 0, 0, 0, 0, 0, 0, 0, 0, 0, 0, 0, 0, 0, 0, 0, 0, 0, 0, 0, 0, 0, 0, 0 }		},	// IN_ADDOVERLAYMATERIAL
 		{ false,	{ 0, 0, 0, 0, 0, 0, 0, 0, 0, 0, 0, 0, 0, 0, 0, 0, 0, 0, 0, 0, 0, 0, 0, 0, 0, 0, 0, 0, 0 }		},	// IN_ADDMATERIAL
diff --git a/Engine/Plugins/Mutable/Source/MutableTools/Private/MuT/ErrorLog.cpp b/Engine/Plugins/Mutable/Source/MutableTools/Private/MuT/ErrorLog.cpp
--- a/Engine/Plugins/Mutable/Source/MutableTools/Private/MuT/ErrorLog.cpp
+++ b/Engine/Plugins/Mutable/Source/MutableTools/Private/MuT/ErrorLog.cpp
@@ -334,6 +334,7 @@ namespace UE::Mutable::Private
         TEXT("IN_ADDSURFACE    "),
         TEXT("IN_ADDCOMPONENT  "),
         TEXT("IN_ADDLOD        "),
+        TEXT("IN_ADDSKELETALMESH"),
 		TEXT("IN_ADDEXTENSIDATA"),
 		TEXT("IN_ADDOVERLAYMATERIAL"),
 		TEXT("IN_ADDMATERIAL   "),
diff --git a/Engine/Plugins/Mutable/Source/MutableTools/Private/MuT/Node.cpp b/Engine/Plugins/Mutable/Source/MutableTools/Private/MuT/Node.cpp
--- a/Engine/Plugins/Mutable/Source/MutableTools/Private/MuT/Node.cpp
+++ b/Engine/Plugins/Mutable/Source/MutableTools/Private/MuT/Node.cpp
@@ -114,6 +114,7 @@
 #include "MuT/NodeMeshExternal.h"
 #include "MuT/NodeScalarExternal.h"
 #include "MuT/NodeSkeletalMesh.h"
+#include "MuT/NodeSkeletalMeshNew.h"
 #include "MuT/NodeSkeletalMeshParameter.h"
 
 
@@ -162,6 +163,7 @@ namespace UE::Mutable::Private
 	FNodeType NodeLOD::StaticType = FNodeType(Node::EType::LOD, Node::GetStaticType());
 
 	FNodeType NodeSkeletalMesh::StaticType = FNodeType(Node::EType::SkeletalMesh, Node::GetStaticType());
+	FNodeType NodeSkeletalMeshNew::StaticType = FNodeType(Node::EType::SkeletalMeshNew, Node::GetStaticType());
 	FNodeType NodeSkeletalMeshParameter::StaticType = FNodeType(Node::EType::SkeletalMeshParameter, NodeSkeletalMesh::GetStaticType());
 
 	FNodeType NodeExtensionData::StaticType = FNodeType(Node::EType::ExtensionData, Node::GetStaticType());
