[mutable] Added option to, within the node UI, add and remove pins.
- Added an option within the node to add a new pin as done in other UE nodes. The new button is shown at the bottom of the outputs of the node that does implement a couple of CONode methods.
- Added option to remove a given pin if that was added using the in-node new "Add Pin" option.
- Updated PinViewer so it does not rely that much on our custom callbacks and instead uses the ones provided by the UEGraph hosting our nodes.
- Removed option to add pins using the mutable "PinViewer" as now that is delegated to the "Add Pin" button now present in, for example, the "Break Material" mutable node.
- Standardized the way our node detail panels do request the refreshing of the details.
- Fixed potential crash in relation to how the mutable detail panels did request their contents refresh.
- Disabled the possibility of changing the type of a mutable node pin if said pin is connected to another node.
- Removed some not required includes and fixed some minor typos.

diff --git a/Engine/Plugins/Mutable/Source/CustomizableObjectEditor/Private/MuCOE/COVariable.cpp b/Engine/Plugins/Mutable/Source/CustomizableObjectEditor/Private/MuCOE/COVariable.cpp
--- a/Engine/Plugins/Mutable/Source/CustomizableObjectEditor/Private/MuCOE/COVariable.cpp
+++ b/Engine/Plugins/Mutable/Source/CustomizableObjectEditor/Private/MuCOE/COVariable.cpp
@@ -5,20 +5,19 @@
 #include "DetailLayoutBuilder.h"
 #include "MuCOE/EdGraphSchema_CustomizableObject.h"
 #include "MuCOE/Nodes/CustomizableObjectNode.h"
-#include "IDetailsView.h"
 #include "PropertyCustomizationHelpers.h"
 #include "Widgets/Input/SComboBox.h"
 #include "Widgets/Input/SEditableTextBox.h"
 #include "Widgets/SWidget.h"
 
 
-TSharedRef<IPropertyTypeCustomization> FCOVariableCustomzation::MakeInstance()
+TSharedRef<IPropertyTypeCustomization> FCOVariableCustomization::MakeInstance()
 {
-	return MakeShared<FCOVariableCustomzation>();
+	return MakeShared<FCOVariableCustomization>();
 }
 
 
-void FCOVariableCustomzation::CustomizeHeader(TSharedRef<IPropertyHandle> StructPropertyHandle, FDetailWidgetRow& InHeaderRow, IPropertyTypeCustomizationUtils& InCustomizationUtils)
+void FCOVariableCustomization::CustomizeHeader(TSharedRef<IPropertyHandle> StructPropertyHandle, FDetailWidgetRow& InHeaderRow, IPropertyTypeCustomizationUtils& InCustomizationUtils)
 {
 	// Get the property handles to edit the variable values
 	NamePropertyHandle = StructPropertyHandle->GetChildHandle("Name");
@@ -33,6 +32,9 @@ void FCOVariableCustomzation::CustomizeHeader(TSharedRef<IPropertyHandle> Struct
 	{
 		BaseObjectNode = Cast<UCustomizableObjectNode>(OuterObjects[0]->GetTypedOuter(UCustomizableObjectNode::StaticClass()));
 
+		BaseObjectPinData = Cast<UCustomizableObjectNodePinData>(OuterObjects[0]);
+		ensure(BaseObjectPinData.IsValid());
+
 		//You can add more types here. This class is not meant to work just with nodes.
 	}
 
@@ -49,9 +51,9 @@ void FCOVariableCustomzation::CustomizeHeader(TSharedRef<IPropertyHandle> Struct
 	InHeaderRow.NameContent()
 	[
 		SNew(SEditableTextBox)
-		.Text(this, &FCOVariableCustomzation::GetVariableName)
-		.OnTextChanged(this, &FCOVariableCustomzation::OnVariableNameChanged)
-		.OnTextCommitted(this, &FCOVariableCustomzation::OnVariableNameCommitted)
+		.Text(this, &FCOVariableCustomization::GetVariableName)
+		.OnTextChanged(this, &FCOVariableCustomization::OnVariableNameChanged)
+		.OnTextCommitted(this, &FCOVariableCustomization::OnVariableNameCommitted)
 		.Font(IDetailLayoutBuilder::GetDetailFont())
 	]
 	.ValueContent()
@@ -59,7 +61,8 @@ void FCOVariableCustomzation::CustomizeHeader(TSharedRef<IPropertyHandle> Struct
 		// Combo button that summons the dropdown menu
 		SNew(SComboBox<FName>)
 		.OptionsSource(&AllowedVariableTypes)
-		.OnGenerateWidget(this, &FCOVariableCustomzation::OnGenerateVariableTypeRow)
+		.IsEnabled(this, &FCOVariableCustomization::ShouldVariableTypePropertyBeEnabled)
+		.OnGenerateWidget(this, &FCOVariableCustomization::OnGenerateVariableTypeRow)
 		.OnSelectionChanged_Lambda
 		([this](FName Item, ESelectInfo::Type)
 			{
@@ -73,14 +76,30 @@ void FCOVariableCustomzation::CustomizeHeader(TSharedRef<IPropertyHandle> Struct
 		)
 		.Content()
 		[
-			GenerateCurrentSelectectedTypeWidget()
+			GenerateCurrentSelectedTypeWidget()
 		]
 	]
-	.OverrideResetToDefault(FResetToDefaultOverride::Create(FSimpleDelegate::CreateSP(this, &FCOVariableCustomzation::ResetToDefault)));
+	.OverrideResetToDefault(FResetToDefaultOverride::Create(FSimpleDelegate::CreateSP(this, &FCOVariableCustomization::ResetToDefault)));
 }
 
 
-FText FCOVariableCustomzation::GetVariableName() const
+bool FCOVariableCustomization::ShouldVariableTypePropertyBeEnabled() const
+{
+	if (const TStrongObjectPtr<UCustomizableObjectNodePinData> NodePinData = BaseObjectPinData.Pin())
+	{
+		if (const TStrongObjectPtr<UCustomizableObjectNode> Node = BaseObjectNode.Pin())
+		{
+			if (const UEdGraphPin* RelatedPin = Node->GetPin(*NodePinData.Get()))
+			{
+				return RelatedPin->LinkedTo.Num() == 0;
+			}
+		}
+	}
+
+	return false;
+}
+
+FText FCOVariableCustomization::GetVariableName() const
 {
 	FName VariableName;
 	NamePropertyHandle->GetValue(VariableName);
@@ -89,14 +108,14 @@ FText FCOVariableCustomzation::GetVariableName() const
 }
 
 
-void FCOVariableCustomzation::ResetToDefault()
+void FCOVariableCustomization::ResetToDefault()
 {
 	NamePropertyHandle->SetValue(FName("NewVar"));
 	TypePropertyHandle->SetValue(AllowedVariableTypes.Num() ? AllowedVariableTypes[0] : NAME_None);
 }
 
 
-void FCOVariableCustomzation::OnVariableNameChanged(const FText& InNewText)
+void FCOVariableCustomization::OnVariableNameChanged(const FText& InNewText)
 {
 	FName NewName = FName(InNewText.ToString());
 
@@ -109,7 +128,7 @@ void FCOVariableCustomzation::OnVariableNameChanged(const FText& InNewText)
 }
 
 
-void FCOVariableCustomzation::OnVariableNameCommitted(const FText& InNewText, ETextCommit::Type InTextCommit)
+void FCOVariableCustomization::OnVariableNameCommitted(const FText& InNewText, ETextCommit::Type InTextCommit)
 {
 	FName NewName = FName(InNewText.ToString());
 
@@ -127,7 +146,7 @@ void FCOVariableCustomzation::OnVariableNameCommitted(const FText& InNewText, ET
 }
 
 
-TSharedRef<SWidget> FCOVariableCustomzation::OnGenerateVariableTypeRow(FName Type)
+TSharedRef<SWidget> FCOVariableCustomization::OnGenerateVariableTypeRow(FName Type)
 {
 	const FSlateBrush* IconBrush = FAppStyle::GetBrush(TEXT("Kismet.VariableList.TypeIcon"));
 
@@ -152,7 +171,7 @@ TSharedRef<SWidget> FCOVariableCustomzation::OnGenerateVariableTypeRow(FName Typ
 }
 
 
-TSharedRef<SWidget> FCOVariableCustomzation::GenerateCurrentSelectectedTypeWidget()
+TSharedRef<SWidget> FCOVariableCustomization::GenerateCurrentSelectedTypeWidget()
 {
 	FName CurrentType;
 	TypePropertyHandle->GetValue(CurrentType);
diff --git a/Engine/Plugins/Mutable/Source/CustomizableObjectEditor/Private/MuCOE/COVariable.h b/Engine/Plugins/Mutable/Source/CustomizableObjectEditor/Private/MuCOE/COVariable.h
--- a/Engine/Plugins/Mutable/Source/CustomizableObjectEditor/Private/MuCOE/COVariable.h
+++ b/Engine/Plugins/Mutable/Source/CustomizableObjectEditor/Private/MuCOE/COVariable.h
@@ -7,23 +7,24 @@
 #include "COVariable.generated.h"
 
 
+class UCustomizableObjectNodePinData;
 class UCustomizableObjectNode;
 class SWidget;
 
 namespace ETextCommit { enum Type : int; }
 
-/** This struct can beused to add new variables to nodes. */
+/** This struct can be used to add new variables to nodes. */
 
 USTRUCT()
 struct FCOVariable
 {
 	GENERATED_USTRUCT_BODY()
 
-	/** Name of this variable. Can be eddited from any details. */
+	/** Name of this variable. Can be edited from any details. */
 	UPROPERTY(EditAnywhere, Category = CustomizableObject)
 	FName Name;
 
-	/** The type of this variable. Ususally a pin category */
+	/** The type of this variable. Usually a pin category */
 	UPROPERTY(EditAnywhere, Category = CustomizableObject)
 	FName Type;
 
@@ -34,14 +35,14 @@ struct FCOVariable
 };
 
 
-class FCOVariableCustomzation: public IPropertyTypeCustomization
+class FCOVariableCustomization: public IPropertyTypeCustomization
 {
 public:
 	static TSharedRef<IPropertyTypeCustomization> MakeInstance();
 
 	// IPropertyTypeCustomization interface
-	void CustomizeHeader(TSharedRef<IPropertyHandle> StructPropertyHandle, class FDetailWidgetRow& HeaderRow, IPropertyTypeCustomizationUtils& StructCustomizationUtils) override;
-	void CustomizeChildren(TSharedRef<IPropertyHandle> StructPropertyHandle, class IDetailChildrenBuilder& ChildBuilder, IPropertyTypeCustomizationUtils& StructCustomizationUtils) override {};
+	virtual void CustomizeHeader(TSharedRef<IPropertyHandle> StructPropertyHandle, class FDetailWidgetRow& HeaderRow, IPropertyTypeCustomizationUtils& StructCustomizationUtils) override;
+	virtual void CustomizeChildren(TSharedRef<IPropertyHandle> StructPropertyHandle, class IDetailChildrenBuilder& ChildBuilder, IPropertyTypeCustomizationUtils& StructCustomizationUtils) override {};
 
 private:
 
@@ -56,13 +57,17 @@ private:
 
 	/** ComboBox Callbacks */
 	TSharedRef<SWidget> OnGenerateVariableTypeRow(FName Option);
-	TSharedRef<SWidget> GenerateCurrentSelectectedTypeWidget();
+	bool ShouldVariableTypePropertyBeEnabled() const;
+	TSharedRef<SWidget> GenerateCurrentSelectedTypeWidget();
 
 private:
 
 	/** Weak pointer to the Customizable Object node that contains this property */
 	TWeakObjectPtr<UCustomizableObjectNode> BaseObjectNode;
 
+	/** Weak pointer to the PinData object used to determine which pin the property slate is referring to. */
+	TWeakObjectPtr<UCustomizableObjectNodePinData> BaseObjectPinData;
+	
 	TArray<FName> AllowedVariableTypes;
 
 	/** Property handle of the variable name */
diff --git a/Engine/Plugins/Mutable/Source/CustomizableObjectEditor/Private/MuCOE/CustomizableObjectEditorModule.cpp b/Engine/Plugins/Mutable/Source/CustomizableObjectEditor/Private/MuCOE/CustomizableObjectEditorModule.cpp
--- a/Engine/Plugins/Mutable/Source/CustomizableObjectEditor/Private/MuCOE/CustomizableObjectEditorModule.cpp
+++ b/Engine/Plugins/Mutable/Source/CustomizableObjectEditor/Private/MuCOE/CustomizableObjectEditorModule.cpp
@@ -303,7 +303,7 @@ void FCustomizableObjectEditorModule::StartupModule()
 	PropertyModule.RegisterCustomPropertyTypeLayout(FCustomizableObjectVariation::StaticStruct()->GetFName(), FOnGetPropertyTypeCustomizationInstance::CreateStatic(&FCustomizableObjectVariationCustomization::MakeInstance));
 	PropertyModule.RegisterCustomPropertyTypeLayout(FCustomizableObjectTextureVariation::StaticStruct()->GetFName(), FOnGetPropertyTypeCustomizationInstance::CreateStatic(&FCustomizableObjectVariationCustomization::MakeInstance));
 	PropertyModule.RegisterCustomPropertyTypeLayout(FTableNodeCompilationFilter::StaticStruct()->GetFName(), FOnGetPropertyTypeCustomizationInstance::CreateStatic(&FCustomizableObjectNodeTableCompilationFilterEditor::MakeInstance));
-	PropertyModule.RegisterCustomPropertyTypeLayout(COVariableUStructName, FOnGetPropertyTypeCustomizationInstance::CreateStatic(&FCOVariableCustomzation::MakeInstance));
+	PropertyModule.RegisterCustomPropertyTypeLayout(COVariableUStructName, FOnGetPropertyTypeCustomizationInstance::CreateStatic(&FCOVariableCustomization::MakeInstance));
 	PropertyModule.RegisterCustomPropertyTypeLayout("InstancedStruct", FOnGetPropertyTypeCustomizationInstance::CreateStatic(&FExternalTypeCustomization::MakeInstance), MakeShared<FExternalTypeTypeIdentifier>());
 
 	PropertyModule.NotifyCustomizationModuleChanged();
diff --git a/Engine/Plugins/Mutable/Source/CustomizableObjectEditor/Private/MuCOE/EdGraphSchema_CustomizableObject.cpp b/Engine/Plugins/Mutable/Source/CustomizableObjectEditor/Private/MuCOE/EdGraphSchema_CustomizableObject.cpp
--- a/Engine/Plugins/Mutable/Source/CustomizableObjectEditor/Private/MuCOE/EdGraphSchema_CustomizableObject.cpp
+++ b/Engine/Plugins/Mutable/Source/CustomizableObjectEditor/Private/MuCOE/EdGraphSchema_CustomizableObject.cpp
@@ -1108,10 +1108,32 @@ void UEdGraphSchema_CustomizableObject::GetContextMenuActions(UToolMenu* Menu, U
 						FSlateIcon(),
 						FUIAction(FExecuteAction::CreateLambda([Node = const_cast<UCustomizableObjectNode*>(Node), Pin]()
 						{
+							if (!Node || !Pin)
+							{
+								return;
+							}
+								
 							Node->SetPinHidden(*const_cast<UEdGraphPin*>(Pin), true);
 						})));
 				}
-				
+
+				if (Node->CanPinBeRemoved(*Pin))
+				{
+					FToolMenuSection& Section = Menu->FindOrAddSection("EdGraphSchemaPinActions");
+					Section.AddMenuEntry("RemovePin",
+						LOCTEXT("RemovePin_Label", "Remove Pin"),
+						LOCTEXT("RemovePin_Tooltip", "Removes the selected pin."),
+						FSlateIcon(),
+						FUIAction(FExecuteAction::CreateLambda([Node = const_cast<UCustomizableObjectNode*>(Node), Pin]()
+						{
+							if (!Node || !Pin)
+							{
+								return;
+							}
+								
+							Node->TryRemovePin(*const_cast<UEdGraphPin*>(Pin));
+						})));
+				}
 				
 				if (TSharedPtr<IDetailsView> Widget = Node->CustomizePinDetails(*Pin))
 				{
diff --git a/Engine/Plugins/Mutable/Source/CustomizableObjectEditor/Private/MuCOE/Nodes/CONodeMaterialBreak.cpp b/Engine/Plugins/Mutable/Source/CustomizableObjectEditor/Private/MuCOE/Nodes/CONodeMaterialBreak.cpp
--- a/Engine/Plugins/Mutable/Source/CustomizableObjectEditor/Private/MuCOE/Nodes/CONodeMaterialBreak.cpp
+++ b/Engine/Plugins/Mutable/Source/CustomizableObjectEditor/Private/MuCOE/Nodes/CONodeMaterialBreak.cpp
@@ -105,9 +105,21 @@ bool UCONodeMaterialBreak::HasPinViewer() const
 }
 
 
-bool UCONodeMaterialBreak::CanCreatePinsFromPinViewer() const
+bool UCONodeMaterialBreak::CanPinBeRemoved(const UEdGraphPin& Pin) const
 {
-	return true;
+	const bool bCanPinBeRemoved = !Pin.LinkedTo.Num() && !Pin.bOrphanedPin;
+	if (!bCanPinBeRemoved)
+	{
+		return false;
+	}
+	
+	return Pin.Direction == EGPD_Output;
+}
+
+
+EAddPinNodeButtonLocation UCONodeMaterialBreak::GetAddPinButtonNodeSide() const
+{
+	return EAddPinNodeButtonLocation::OUTPUT;
 }
 
 
@@ -123,18 +135,21 @@ TArray<FName> UCONodeMaterialBreak::GetAllowedPinViewerCreationTypes() const
 }
 
 
-void UCONodeMaterialBreak::CreatePinFromPinViewer()
+void UCONodeMaterialBreak::AddPinFromUI()
 {
-	FName NewPinCategorery = UEdGraphSchema_CustomizableObject::PC_Float;
+	Super::AddPinFromUI();
+	
+	const FName NewPinCategory = UEdGraphSchema_CustomizableObject::PC_Float;
 	
 	UCONodeMaterialBreakPinData* PinData = NewObject<UCONodeMaterialBreakPinData>(this);
 	PinData->VariableData.Name = "NewPin";
-	PinData->VariableData.Type = NewPinCategorery;
+	PinData->VariableData.Type = NewPinCategory;
 	PinData->VariableData.Id = FGuid::NewGuid();
 
-	CustomCreatePin(EEdGraphPinDirection::EGPD_Output, NewPinCategorery, "NewPin", PinData);
-
-	GetGraph()->NotifyGraphChanged();
+	if (CustomCreatePin(EEdGraphPinDirection::EGPD_Output, NewPinCategory, "NewPin", PinData))
+	{
+		GetGraph()->NotifyNodeChanged(this);
+	}
 }
 
 
diff --git a/Engine/Plugins/Mutable/Source/CustomizableObjectEditor/Private/MuCOE/Nodes/CONodeMaterialBreak.h b/Engine/Plugins/Mutable/Source/CustomizableObjectEditor/Private/MuCOE/Nodes/CONodeMaterialBreak.h
--- a/Engine/Plugins/Mutable/Source/CustomizableObjectEditor/Private/MuCOE/Nodes/CONodeMaterialBreak.h
+++ b/Engine/Plugins/Mutable/Source/CustomizableObjectEditor/Private/MuCOE/Nodes/CONodeMaterialBreak.h
@@ -53,16 +53,15 @@ public:
 	UE_API virtual void AllocateDefaultPins(UCustomizableObjectNodeRemapPins* RemapPins) override;
 	UE_API virtual TSharedPtr<IDetailsView> CustomizePinDetails(const UEdGraphPin& Pin) const override;
 	UE_API virtual bool HasPinViewer() const override;
-	UE_API virtual bool CanCreatePinsFromPinViewer() const override;
-	UE_API virtual void CreatePinFromPinViewer() override;
+	UE_API virtual EAddPinNodeButtonLocation GetAddPinButtonNodeSide() const override;
+	UE_API virtual void AddPinFromUI() override;
+	UE_API virtual bool CanPinBeRemoved(const UEdGraphPin& Pin) const override;
 	UE_API virtual UCustomizableObjectNodeRemapPins* CreateRemapPinsDefault() const override;
 	UE_API virtual TArray<FName> GetAllowedPinViewerCreationTypes() const override;
 
 	// Own interface
 	FName GetPinParameterName(const UEdGraphPin& Pin) const;
 
-public: 
-
 	UPROPERTY()
 	FEdGraphPinReference MaterialPinRef;
 };
diff --git a/Engine/Plugins/Mutable/Source/CustomizableObjectEditor/Private/MuCOE/Nodes/CustomizableObjectNode.cpp b/Engine/Plugins/Mutable/Source/CustomizableObjectEditor/Private/MuCOE/Nodes/CustomizableObjectNode.cpp
--- a/Engine/Plugins/Mutable/Source/CustomizableObjectEditor/Private/MuCOE/Nodes/CustomizableObjectNode.cpp
+++ b/Engine/Plugins/Mutable/Source/CustomizableObjectEditor/Private/MuCOE/Nodes/CustomizableObjectNode.cpp
@@ -11,12 +11,11 @@
 #include "MuCOE/CustomizableObjectPin.h"
 #include "MuCOE/EdGraphSchema_CustomizableObject.h"
 #include "MuCOE/GraphTraversal.h"
-#include "MuCOE/Nodes/CustomizableObjectNodeComponentMeshAddTo.h"
 #include "MuCOE/Nodes/CustomizableObjectNodeObject.h"
-#include "MuCOE/Nodes/CustomizableObjectNodeObjectGroup.h"
 #include "MuCOE/RemapPins/CustomizableObjectNodeRemapPinsByName.h"
 #include "MuCOE/RemapPins/CustomizableObjectNodeRemapPinsByPosition.h"
 #include "Toolkits/ToolkitManager.h"
+#include "EDGraph/EdGraph.h"
 
 #include UE_INLINE_GENERATED_CPP_BY_NAME(CustomizableObjectNode)
 
@@ -320,7 +319,7 @@ void UCustomizableObjectNode::FixupReconstructPins(UCustomizableObjectNodeRemapP
 			.Log();
 		}
 
-		Graph->NotifyGraphChanged();
+		Graph->NotifyNodeChanged(this);
 	}
 	
 	PostReconstructNodeDelegate.Broadcast();
@@ -796,8 +795,45 @@ bool UCustomizableObjectNode::HasPinViewer() const
 }
 
 
-bool UCustomizableObjectNode::CanCreatePinsFromPinViewer() const
+EAddPinNodeButtonLocation UCustomizableObjectNode::GetAddPinButtonNodeSide() const
 {
+	return EAddPinNodeButtonLocation::NONE;
+}
+
+
+void UCustomizableObjectNode::AddPinFromUI()
+{
+	// You are not expected to call this method if not part of the "Add pin" in-node button operation. Use "CustomAddPin()" method instead.
+	ensure(GetAddPinButtonNodeSide() != EAddPinNodeButtonLocation::NONE);
+	
+	FScopedTransaction Transaction(LOCTEXT("AddPinFromNodeUITransaction", "Add Pin from Node UI"));
+	Modify();
+}
+
+
+bool UCustomizableObjectNode::CanPinBeRemoved(const UEdGraphPin& Pin) const
+{
+	return false;
+}
+
+
+bool UCustomizableObjectNode::TryRemovePin(UEdGraphPin& Pin)
+{
+	FScopedTransaction Transaction(LOCTEXT("TryRemovePinTransaction", "Try removing a Pin"));
+	Modify();
+
+	// Make sure the pin that is requested for removal was actually added by the UI.
+	if (CanPinBeRemoved(Pin) && CustomRemovePin(Pin))
+	{
+		// Notify that a pin has been removed so the PinViewer reacts to the new pin layout
+		if (UEdGraph* Graph = GetGraph())
+		{
+			Graph->NotifyNodeChanged(this);
+		}
+
+		return true;
+	}
+
 	return false;
 }
 
@@ -843,7 +879,7 @@ bool UCustomizableObjectNode::IsNodeSupportedInMacros() const
 
 bool UCustomizableObjectNode::IsInMacro() const
 {
-	if (UCustomizableObjectGraph* Graph = Cast<UCustomizableObjectGraph>(GetGraph()))
+	if (const UCustomizableObjectGraph* Graph = Cast<UCustomizableObjectGraph>(GetGraph()))
 	{
 		return Graph->IsMacro();
 	}
diff --git a/Engine/Plugins/Mutable/Source/CustomizableObjectEditor/Private/MuCOE/Nodes/CustomizableObjectNodeComponentMeshDetails.h b/Engine/Plugins/Mutable/Source/CustomizableObjectEditor/Private/MuCOE/Nodes/CustomizableObjectNodeComponentMeshDetails.h
--- a/Engine/Plugins/Mutable/Source/CustomizableObjectEditor/Private/MuCOE/Nodes/CustomizableObjectNodeComponentMeshDetails.h
+++ b/Engine/Plugins/Mutable/Source/CustomizableObjectEditor/Private/MuCOE/Nodes/CustomizableObjectNodeComponentMeshDetails.h
@@ -10,7 +10,7 @@ class UCustomizableObjectNodeComponentMesh;
 class SWidget;
 
 
-/** Copy Material node details panel. Hides all properties from the inheret Material node. */
+/** Copy Material node details panel. Hides all properties from the inherited Material node. */
 class FCustomizableObjectNodeComponentMeshDetails : public FCustomizableObjectNodeDetails
 {
 public:
@@ -18,7 +18,6 @@ public:
 
 	// ILayoutDetails interface
 	virtual void CustomizeDetails(IDetailLayoutBuilder& DetailBuilder) override;
-	virtual void CustomizeDetails(const TSharedPtr<IDetailLayoutBuilder>& DetailBuilder) override;
 
 	// Own interface
 	
@@ -33,5 +32,4 @@ public:
 private:
 	TWeakObjectPtr<UCustomizableObjectNodeComponentMesh> NodeComponentMesh;
 	
-	TWeakPtr<IDetailLayoutBuilder> WeakDetailsBuilder;
 };
\ No newline at end of file
diff --git a/Engine/Plugins/Mutable/Source/CustomizableObjectEditor/Private/MuCOE/Nodes/CustomizableObjectNodeDetails.cpp b/Engine/Plugins/Mutable/Source/CustomizableObjectEditor/Private/MuCOE/Nodes/CustomizableObjectNodeDetails.cpp
--- a/Engine/Plugins/Mutable/Source/CustomizableObjectEditor/Private/MuCOE/Nodes/CustomizableObjectNodeDetails.cpp
+++ b/Engine/Plugins/Mutable/Source/CustomizableObjectEditor/Private/MuCOE/Nodes/CustomizableObjectNodeDetails.cpp
@@ -5,11 +5,8 @@
 #include "DetailCategoryBuilder.h"
 #include "DetailLayoutBuilder.h"
 #include "DetailWidgetRow.h"
-#include "Input/Reply.h"
 #include "MuCOE/Nodes/CustomizableObjectNode.h"
 #include "MuCOE/PinViewer/SPinViewer.h"
-#include "Widgets/Input/SButton.h"
-#include "Widgets/Images/SImage.h"
 
 class IDetailLayoutBuilder;
 
@@ -35,33 +32,6 @@ void FCustomizableObjectNodeDetails::CustomizeDetails(IDetailLayoutBuilder& Deta
 			}
 			
 			IDetailCategoryBuilder& PinViewerCategoryBuilder = DetailBuilder.EditCategory("PinViewer", LOCTEXT("PinViewer", "Pins"), ECategoryPriority::Uncommon);
-			
-			if (Node->CanCreatePinsFromPinViewer())
-			{
-				FText TooltipText = LOCTEXT("NodeAddPinTooltip", "Create a new pin.");
-
-				PinViewerCategoryBuilder.HeaderContent
-				(
-					SNew(SHorizontalBox)
-					+ SHorizontalBox::Slot()
-					.HAlign(HAlign_Right)
-					[
-						SNew(SButton)
-						.ButtonStyle(FAppStyle::Get(), "SimpleButton")
-						.ContentPadding(FMargin(1, 0))
-						.OnClicked(this, &FCustomizableObjectNodeDetails::AddNewPin, Node)
-						.HAlign(HAlign_Right)
-						.ToolTipText(TooltipText)
-						.VAlign(VAlign_Center)
-						[
-							SNew(SImage)
-							.Image(FAppStyle::Get().GetBrush("Icons.PlusCircle"))
-							.ColorAndOpacity(FSlateColor::UseForeground())
-						]
-					]
-				);
-			}
-
 			PinViewerCategoryBuilder.AddCustomRow(LOCTEXT("PinViewerDetailsCategory", "PinViewer")).ShouldAutoExpand(true)
 			[
 				SNew(SPinViewer).Node(Node)
@@ -80,17 +50,13 @@ void FCustomizableObjectNodeDetails::CustomizeDetails(const TSharedPtr<IDetailLa
 }
 
 
-FReply FCustomizableObjectNodeDetails::AddNewPin(UCustomizableObjectNode* Node)
+void FCustomizableObjectNodeDetails::RefreshDetails()
 {
-	if (Node && DetailBuilderPtr)
+	TSharedPtr<IDetailLayoutBuilder> DetailBuilder = DetailBuilderPtr.Pin();
+	if (DetailBuilder.IsValid())
 	{
-		Node->CreatePinFromPinViewer();
-		DetailBuilderPtr->ForceRefreshDetails();
-
-		return FReply::Handled();
+		DetailBuilder->ForceRefreshDetails();
 	}
-
-	return FReply::Unhandled();
 }
 
 #undef LOCTEXT_NAMESPACE
diff --git a/Engine/Plugins/Mutable/Source/CustomizableObjectEditor/Private/MuCOE/Nodes/CustomizableObjectNodeDetails.h b/Engine/Plugins/Mutable/Source/CustomizableObjectEditor/Private/MuCOE/Nodes/CustomizableObjectNodeDetails.h
--- a/Engine/Plugins/Mutable/Source/CustomizableObjectEditor/Private/MuCOE/Nodes/CustomizableObjectNodeDetails.h
+++ b/Engine/Plugins/Mutable/Source/CustomizableObjectEditor/Private/MuCOE/Nodes/CustomizableObjectNodeDetails.h
@@ -20,11 +20,13 @@ public:
 	virtual void CustomizeDetails(IDetailLayoutBuilder& DetailBuilder) override;
 	virtual void CustomizeDetails(const TSharedPtr<IDetailLayoutBuilder>& DetailBuilder) override;
 
-	//
-	FReply AddNewPin(UCustomizableObjectNode* Node);
+protected:
 
-private: 
-
-	// Pointer to the DetailBuilder. Needed to refresh the details when we create a new pin from the pin viewer.
-	TSharedPtr<IDetailLayoutBuilder> DetailBuilderPtr;
+	/** Force the refreshing of this details view. Only use if you are sure this will not pose a performance issue. */
+	void RefreshDetails();
+	
+private:
+	
+	// Pointer to the DetailBuilder.
+	TWeakPtr<IDetailLayoutBuilder> DetailBuilderPtr;
 };
diff --git a/Engine/Plugins/Mutable/Source/CustomizableObjectEditor/Private/MuCOE/Nodes/CustomizableObjectNodeExternalPinDetails.cpp b/Engine/Plugins/Mutable/Source/CustomizableObjectEditor/Private/MuCOE/Nodes/CustomizableObjectNodeExternalPinDetails.cpp
--- a/Engine/Plugins/Mutable/Source/CustomizableObjectEditor/Private/MuCOE/Nodes/CustomizableObjectNodeExternalPinDetails.cpp
+++ b/Engine/Plugins/Mutable/Source/CustomizableObjectEditor/Private/MuCOE/Nodes/CustomizableObjectNodeExternalPinDetails.cpp
@@ -28,7 +28,6 @@ void FCustomizableObjectNodeExternalPinDetails::CustomizeDetails( IDetailLayoutB
 	FCustomizableObjectNodeDetails::CustomizeDetails(DetailBuilder);
 
 	Node = nullptr;
-	DetailBuilderPtr = &DetailBuilder;
 	
 	TSharedPtr<const IDetailsView> DetailsView = DetailBuilder.GetDetailsViewSharedPtr();
 	if (DetailsView && DetailsView->GetSelectedObjects().Num())
@@ -110,10 +109,7 @@ void FCustomizableObjectNodeExternalPinDetails::ParentObjectSelectionChanged(con
 	Node->ExternalObject = Cast<UCustomizableObject>(UE::Mutable::Private::LoadObject(AssetData));
 	Node->SetExternalObjectNodeId(FGuid());
 
-	if (DetailBuilderPtr)
-	{
-		DetailBuilderPtr->ForceRefreshDetails();
-	}
+	RefreshDetails();
 }
 
 
diff --git a/Engine/Plugins/Mutable/Source/CustomizableObjectEditor/Private/MuCOE/Nodes/CustomizableObjectNodeExternalPinDetails.h b/Engine/Plugins/Mutable/Source/CustomizableObjectEditor/Private/MuCOE/Nodes/CustomizableObjectNodeExternalPinDetails.h
--- a/Engine/Plugins/Mutable/Source/CustomizableObjectEditor/Private/MuCOE/Nodes/CustomizableObjectNodeExternalPinDetails.h
+++ b/Engine/Plugins/Mutable/Source/CustomizableObjectEditor/Private/MuCOE/Nodes/CustomizableObjectNodeExternalPinDetails.h
@@ -23,7 +23,6 @@ public:
 	virtual void CustomizeDetails( IDetailLayoutBuilder& DetailBuilder ) override;
 
 private:
-	IDetailLayoutBuilder* DetailBuilderPtr = nullptr;
 
 	UCustomizableObjectNodeExternalPin* Node = nullptr;
 	
diff --git a/Engine/Plugins/Mutable/Source/CustomizableObjectEditor/Private/MuCOE/Nodes/CustomizableObjectNodeObjectDetails.cpp b/Engine/Plugins/Mutable/Source/CustomizableObjectEditor/Private/MuCOE/Nodes/CustomizableObjectNodeObjectDetails.cpp
--- a/Engine/Plugins/Mutable/Source/CustomizableObjectEditor/Private/MuCOE/Nodes/CustomizableObjectNodeObjectDetails.cpp
+++ b/Engine/Plugins/Mutable/Source/CustomizableObjectEditor/Private/MuCOE/Nodes/CustomizableObjectNodeObjectDetails.cpp
@@ -161,7 +161,6 @@ void FCustomizableObjectNodeObjectDetails::CustomizeDetails( IDetailLayoutBuilde
 	FCustomizableObjectNodeDetails::CustomizeDetails(DetailBuilder);
 
 	BaseObjectNode = nullptr;
-	DetailBuilderPtr = &DetailBuilder;
 
 	TSharedPtr<const IDetailsView> DetailsView = DetailBuilder.GetDetailsViewSharedPtr();
 	if (DetailsView && DetailsView->GetSelectedObjects().Num())
@@ -321,10 +320,7 @@ void FCustomizableObjectNodeObjectDetails::ParentObjectSelectionChanged(const FA
 		}
 	}
 
-	if (DetailBuilderPtr)
-	{
-		DetailBuilderPtr->ForceRefreshDetails();
-	}
+	RefreshDetails();
 }
 
 
@@ -350,10 +346,7 @@ void FCustomizableObjectNodeObjectDetails::OnGroupNodeComboBoxSelectionChanged(T
 
 void FCustomizableObjectNodeObjectDetails::OnStatesPropertyChanged()
 {
-	if (DetailBuilderPtr)
-	{
-		DetailBuilderPtr->ForceRefreshDetails();
-	}
+	RefreshDetails();
 }
 
 
diff --git a/Engine/Plugins/Mutable/Source/CustomizableObjectEditor/Private/MuCOE/Nodes/CustomizableObjectNodeObjectDetails.h b/Engine/Plugins/Mutable/Source/CustomizableObjectEditor/Private/MuCOE/Nodes/CustomizableObjectNodeObjectDetails.h
--- a/Engine/Plugins/Mutable/Source/CustomizableObjectEditor/Private/MuCOE/Nodes/CustomizableObjectNodeObjectDetails.h
+++ b/Engine/Plugins/Mutable/Source/CustomizableObjectEditor/Private/MuCOE/Nodes/CustomizableObjectNodeObjectDetails.h
@@ -95,7 +95,6 @@ private:
 
 	TWeakObjectPtr<UCustomizableObjectNodeObject> BaseObjectNode;
 	TArray< TSharedPtr<FString> > GroupNodeComboOptions;
-	IDetailLayoutBuilder* DetailBuilderPtr = nullptr;
 
 	TArray< TSharedPtr<FString> > ParentComboOptions;
 	TArray< class UCustomizableObjectNodeMaterial* > ParentOptionNode;
diff --git a/Engine/Plugins/Mutable/Source/CustomizableObjectEditor/Private/MuCOE/Nodes/CustomizableObjectNodeTableDetails.cpp b/Engine/Plugins/Mutable/Source/CustomizableObjectEditor/Private/MuCOE/Nodes/CustomizableObjectNodeTableDetails.cpp
--- a/Engine/Plugins/Mutable/Source/CustomizableObjectEditor/Private/MuCOE/Nodes/CustomizableObjectNodeTableDetails.cpp
+++ b/Engine/Plugins/Mutable/Source/CustomizableObjectEditor/Private/MuCOE/Nodes/CustomizableObjectNodeTableDetails.cpp
@@ -39,7 +39,6 @@ void FCustomizableObjectNodeTableDetails::CustomizeDetails(const TSharedPtr<IDet
 	IDetailCustomization::CustomizeDetails(DetailBuilder);
 
 	Node = 0;
-	DetailBuilderPtr = DetailBuilder;
 
 	TSharedPtr<const IDetailsView> DetailsView = DetailBuilder->GetDetailsViewSharedPtr();
 
@@ -341,10 +340,7 @@ void FCustomizableObjectNodeTableDetails::GenerateMeshSectionOptions(TArray<FLay
 
 void FCustomizableObjectNodeTableDetails::OnNodePinValueChanged()
 {
-	if (IDetailLayoutBuilder* DetailBuilder = DetailBuilderPtr.Pin().Get()) // Raw because we don't want to keep alive the details builder when calling the force refresh details
-	{
-		DetailBuilder->ForceRefreshDetails();
-	}
+	RefreshDetails();
 }
 
 
diff --git a/Engine/Plugins/Mutable/Source/CustomizableObjectEditor/Private/MuCOE/Nodes/CustomizableObjectNodeTableDetails.h b/Engine/Plugins/Mutable/Source/CustomizableObjectEditor/Private/MuCOE/Nodes/CustomizableObjectNodeTableDetails.h
--- a/Engine/Plugins/Mutable/Source/CustomizableObjectEditor/Private/MuCOE/Nodes/CustomizableObjectNodeTableDetails.h
+++ b/Engine/Plugins/Mutable/Source/CustomizableObjectEditor/Private/MuCOE/Nodes/CustomizableObjectNodeTableDetails.h
@@ -130,10 +130,6 @@ private:
 	// Pointer to the node represented in this details
 	TWeakObjectPtr<UCustomizableObjectNodeTable> Node;
 
-	// Pointer to the Detail Builder to force the refresh on recontruct the node
-	TWeakPtr<IDetailLayoutBuilder> DetailBuilderPtr = nullptr;
-
-
 	// Anim -------------
 	// ComboBox widget to select a column from the NodeTable
 	TSharedPtr<STextComboBox> AnimMeshColumnComboBox;
diff --git a/Engine/Plugins/Mutable/Source/CustomizableObjectEditor/Private/MuCOE/Nodes/FCustomizableObjectNodeComponentMeshDetails.cpp b/Engine/Plugins/Mutable/Source/CustomizableObjectEditor/Private/MuCOE/Nodes/FCustomizableObjectNodeComponentMeshDetails.cpp
--- a/Engine/Plugins/Mutable/Source/CustomizableObjectEditor/Private/MuCOE/Nodes/FCustomizableObjectNodeComponentMeshDetails.cpp
+++ b/Engine/Plugins/Mutable/Source/CustomizableObjectEditor/Private/MuCOE/Nodes/FCustomizableObjectNodeComponentMeshDetails.cpp
@@ -8,7 +8,6 @@
 #include "MuCOE/Nodes/CustomizableObjectNodeComponentMesh.h"
 #include "Widgets/Text/STextBlock.h"
 #include "Widgets/Input/SComboButton.h"
-#include "Widgets/Input/SComboButton.h"
 #include "Framework/MultiBox/MultiBoxBuilder.h"
 
 
@@ -62,7 +61,7 @@ void FCustomizableObjectNodeComponentMeshDetails::CustomizeDetails(IDetailLayout
 
 	DetailBuilder.HideProperty("LODReductionSettings");
 
-	FName BonesToRemovePropertyPath = FName("LODReductionSettings[" + FString::FromInt(NodeComponentMesh.Get()->SelectedLOD) + "].BonesToRemove");
+	const FName BonesToRemovePropertyPath = FName("LODReductionSettings[" + FString::FromInt(NodeComponentMesh.Get()->SelectedLOD) + "].BonesToRemove");
 	TSharedRef<IPropertyHandle> BonesToRemoveProperty = DetailBuilder.GetProperty(BonesToRemovePropertyPath);
 	
 	// Bones to remove widget
@@ -70,14 +69,6 @@ void FCustomizableObjectNodeComponentMeshDetails::CustomizeDetails(IDetailLayout
 }
 
 
-void FCustomizableObjectNodeComponentMeshDetails::CustomizeDetails(const TSharedPtr<IDetailLayoutBuilder>& DetailBuilder)
-{
-	FCustomizableObjectNodeDetails::CustomizeDetails(DetailBuilder);
-
-	WeakDetailsBuilder = DetailBuilder;
-}
-
-
 TSharedRef<SWidget> FCustomizableObjectNodeComponentMeshDetails::OnGenerateLODComboBoxForPicker()
 {
 	return SNew(SComboButton)
@@ -117,7 +108,7 @@ void FCustomizableObjectNodeComponentMeshDetails::OnSelectedLODChanged(int32 New
 {
 	NodeComponentMesh.Get()->SelectedLOD = NewLODIndex;
 
-	WeakDetailsBuilder.Pin()->ForceRefreshDetails();
+	RefreshDetails();
 }
 
 
diff --git a/Engine/Plugins/Mutable/Source/CustomizableObjectEditor/Private/MuCOE/Nodes/SCustomizableObjectNode.cpp b/Engine/Plugins/Mutable/Source/CustomizableObjectEditor/Private/MuCOE/Nodes/SCustomizableObjectNode.cpp
--- a/Engine/Plugins/Mutable/Source/CustomizableObjectEditor/Private/MuCOE/Nodes/SCustomizableObjectNode.cpp
+++ b/Engine/Plugins/Mutable/Source/CustomizableObjectEditor/Private/MuCOE/Nodes/SCustomizableObjectNode.cpp
@@ -2,10 +2,14 @@
 
 #include "SCustomizableObjectNode.h"
 
+#include "CustomizableObjectNodeDetails.h"
+#include "DetailLayoutBuilder.h"
+#include "GraphEditorSettings.h"
 #include "SCustomizableObjectNodePin.h"
 #include "MuCOE/EdGraphSchema_CustomizableObject.h"
 #include "MuCOE/Nodes/CustomizableObjectNode.h"
 #include "Widgets/SBoxPanel.h"
+#include "ScopedTransaction.h"
 
 #define LOCTEXT_NAMESPACE "CustomizableObjectEditor"
 
@@ -58,6 +62,73 @@ void SCustomizableObjectNode::CreateBelowPinControls(TSharedPtr<SVerticalBox> Ma
 	}
 }
 
+void SCustomizableObjectNode::CreateInputSideAddButton(TSharedPtr<SVerticalBox> OutputBox)
+{
+	SGraphNode::CreateInputSideAddButton(OutputBox);
+
+	if (const UCustomizableObjectNode* Node = Cast<UCustomizableObjectNode>(GraphNode))
+	{
+		if (Node->GetAddPinButtonNodeSide() == EAddPinNodeButtonLocation::INPUT)
+		{
+			TSharedRef<SWidget> AddPinButton = AddPinButtonContent(
+			NSLOCTEXT("SequencerNode", "SequencerNodeAddPinButton", "Add pin"),
+			NSLOCTEXT("SequencerNode", "SequencerNodeAddPinButton_ToolTip", "Add new pin"));
+
+			FMargin AddPinPadding = Settings->GetOutputPinPadding();
+			AddPinPadding.Top += 6.0f;
+
+			OutputBox->AddSlot()
+			.AutoHeight()
+			.VAlign(VAlign_Center)
+			.HAlign(HAlign_Right)
+			.Padding(AddPinPadding)
+			[
+				AddPinButton
+			];
+		}
+	}
+}
+
+
+void SCustomizableObjectNode::CreateOutputSideAddButton(TSharedPtr<SVerticalBox> OutputBox)
+{
+	SGraphNode::CreateOutputSideAddButton(OutputBox);
+
+	if (const UCustomizableObjectNode* Node = Cast<UCustomizableObjectNode>(GraphNode))
+	{
+		if (Node->GetAddPinButtonNodeSide() == EAddPinNodeButtonLocation::OUTPUT)
+		{
+			TSharedRef<SWidget> AddPinButton = AddPinButtonContent(
+			NSLOCTEXT("SequencerNode", "SequencerNodeAddPinButton", "Add pin"),
+			NSLOCTEXT("SequencerNode", "SequencerNodeAddPinButton_ToolTip", "Add new pin"));
+
+			FMargin AddPinPadding = Settings->GetOutputPinPadding();
+			AddPinPadding.Top += 6.0f;
+
+			OutputBox->AddSlot()
+			.AutoHeight()
+			.VAlign(VAlign_Center)
+			.HAlign(HAlign_Right)
+			.Padding(AddPinPadding)
+			[
+				AddPinButton
+			];
+		}
+	}
+}
+
+
+FReply SCustomizableObjectNode::OnAddPin()
+{
+	if (UCustomizableObjectNode* Node = Cast<UCustomizableObjectNode>(GraphNode))
+	{
+		Node->AddPinFromUI();
+		return FReply::Handled();
+	}
+
+	return FReply::Unhandled();
+}
+
 
 #undef LOCTEXT_NAMESPACE
 
diff --git a/Engine/Plugins/Mutable/Source/CustomizableObjectEditor/Private/MuCOE/Nodes/SCustomizableObjectNode.h b/Engine/Plugins/Mutable/Source/CustomizableObjectEditor/Private/MuCOE/Nodes/SCustomizableObjectNode.h
--- a/Engine/Plugins/Mutable/Source/CustomizableObjectEditor/Private/MuCOE/Nodes/SCustomizableObjectNode.h
+++ b/Engine/Plugins/Mutable/Source/CustomizableObjectEditor/Private/MuCOE/Nodes/SCustomizableObjectNode.h
@@ -20,4 +20,9 @@ public:
 	// SGraphNode interface
     virtual TSharedPtr<SGraphPin> CreatePinWidget(UEdGraphPin* Pin) const override;
 	virtual void CreateBelowPinControls(TSharedPtr<SVerticalBox> MainBox) override;
+
+protected:
+	virtual void CreateInputSideAddButton(TSharedPtr<SVerticalBox> OutputBox) override;
+    virtual void CreateOutputSideAddButton(TSharedPtr<SVerticalBox> OutputBox) override;
+    virtual FReply OnAddPin() override;
 };
diff --git a/Engine/Plugins/Mutable/Source/CustomizableObjectEditor/Private/MuCOE/PinViewer/SPinViewer.cpp b/Engine/Plugins/Mutable/Source/CustomizableObjectEditor/Private/MuCOE/PinViewer/SPinViewer.cpp
--- a/Engine/Plugins/Mutable/Source/CustomizableObjectEditor/Private/MuCOE/PinViewer/SPinViewer.cpp
+++ b/Engine/Plugins/Mutable/Source/CustomizableObjectEditor/Private/MuCOE/PinViewer/SPinViewer.cpp
@@ -2,10 +2,8 @@
 
 #include "MuCOE/PinViewer/SPinViewer.h"
 
-#include "DetailLayoutBuilder.h"
 #include "DetailWidgetRow.h"
-#include "Framework/Views/TableViewMetadata.h"
-#include "IDetailsView.h"
+#include "GraphEditAction.h"
 #include "MuCOE/Nodes/CustomizableObjectNode.h"
 #include "MuCOE/PinViewer/SPinViewerListRow.h"
 #include "ScopedTransaction.h"
@@ -186,11 +184,17 @@ void SPinViewer::PinsRemapped(const TMap<UEdGraphPin*, UEdGraphPin*>& PinsToRema
 
 void SPinViewer::Construct(const FArguments& InArgs)
 {
-	Node = InArgs._Node;
-	Node->PostReconstructNodeDelegate.AddSP(this, &SPinViewer::UpdateWidget);
-	Node->NodeConnectionListChangedDelegate.AddSP(this, &SPinViewer::UpdateWidget);
+	Node = TStrongObjectPtr<UCustomizableObjectNode>{InArgs._Node};
 	Node->RemapPinsDelegate.AddSP(this, &SPinViewer::PinsRemapped);
+	Node->NodeConnectionListChangedDelegate.AddSP(this, &SPinViewer::UpdateWidget);
 	GeneratePinInfoList();
+
+	Graph = TStrongObjectPtr(Node->GetGraph());
+	if (Graph)
+	{
+		// bind the node refresh to a change in the graph
+		Graph->AddOnGraphChangedHandler(FOnGraphChanged::FDelegate::CreateSP(this, &SPinViewer::OnGraphChanged));
+	}
 	
 	ChildSlot
 	[
@@ -224,7 +228,7 @@ void SPinViewer::Construct(const FArguments& InArgs)
 				SNew(SButton)
 				.Text(LOCTEXT("HideAllText", "Hide All"))
 				.HAlign(HAlign_Center)
-				.ToolTipText(LOCTEXT("ShowLinkedTooltip", "Hide all non connnected pins"))
+				.ToolTipText(LOCTEXT("ShowLinkedTooltip", "Hide all non connected pins"))
 				.OnClicked(this, &SPinViewer::OnHideAllPressed)
 			]
 		]
@@ -264,16 +268,19 @@ void SPinViewer::Construct(const FArguments& InArgs)
 }
 
 
-void SPinViewer::AddReferencedObjects(FReferenceCollector& Collector)
+void SPinViewer::UpdateWidget()
 {
-	Collector.AddReferencedObject(Node);
+	GeneratePinInfoList();
+	ListView->RequestListRefresh();
 }
 
 
-void SPinViewer::UpdateWidget()
+void SPinViewer::OnGraphChanged(const FEdGraphEditAction& EdGraphEditAction)
 {
-	GeneratePinInfoList();
-	ListView->RequestListRefresh();
+	if (EdGraphEditAction.Action == GRAPHACTION_EditNode && EdGraphEditAction.Nodes.Contains(Node.Get()))
+	{
+		UpdateWidget();
+	}
 }
 
 
diff --git a/Engine/Plugins/Mutable/Source/CustomizableObjectEditor/Private/MuCOE/PinViewer/SPinViewer.h b/Engine/Plugins/Mutable/Source/CustomizableObjectEditor/Private/MuCOE/PinViewer/SPinViewer.h
--- a/Engine/Plugins/Mutable/Source/CustomizableObjectEditor/Private/MuCOE/PinViewer/SPinViewer.h
+++ b/Engine/Plugins/Mutable/Source/CustomizableObjectEditor/Private/MuCOE/PinViewer/SPinViewer.h
@@ -3,10 +3,10 @@
 #pragma once
 
 #include "Widgets/Views/SHeaderRow.h"
+#include "GraphEditAction.h"
 
 template <typename ItemType> class SListView;
 
-class FReferenceCollector;
 class IDetailLayoutBuilder;
 class ITableRow;
 class STableViewBase;
@@ -22,7 +22,7 @@ struct FGuid;
  * 
  * This widget will work as long as pins are only created and destroyed during the node reconstruction.
  */
-class SPinViewer : public SCompoundWidget, public FGCObject
+class SPinViewer : public SCompoundWidget
 {
 	friend class SPinViewerListRow;
 
@@ -37,17 +37,13 @@ public:
 	
 	void Construct(const FArguments& InArgs);
 
-	// FGCObject interface
-	void AddReferencedObjects(FReferenceCollector& Collector) override;
-	virtual FString GetReferencerName() const override
-	{
-		return TEXT("SCustomizableObjectNodePinViewer");
-	}
-
 	/** Regenerate the list contents and update the widget. */
 	void UpdateWidget();
 
-	/** Callback to Generate a Nodfe Material Pin row */
+	/** Callback executed each time a graph node does change. */
+	void OnGraphChanged(const FEdGraphEditAction& EdGraphEditAction);
+
+	/** Callback to Generate a Node Material Pin row */
 	TSharedRef<ITableRow> GenerateNodePinRow(TSharedPtr<FEdGraphPinReference> PinReference, const TSharedRef<STableViewBase>& OwnerTable);
 
 	/** SearchBox callback. */
@@ -71,9 +67,11 @@ private:
 
 	void PinsRemapped(const TMap<UEdGraphPin*, UEdGraphPin*>& PinsToRemap);
 
-	/** Pointer to the Node Material */
-	TObjectPtr<UCustomizableObjectNode> Node = nullptr;
+	/** Pointer to the Node */
+	TStrongObjectPtr<UCustomizableObjectNode> Node = nullptr;
 
+	TStrongObjectPtr<UEdGraph> Graph = nullptr;
+	
 	/** ListView elements. */
 	TArray<TSharedPtr<FEdGraphPinReference>> PinReferences;
 
diff --git a/Engine/Plugins/Mutable/Source/CustomizableObjectEditor/Public/MuCOE/Nodes/CustomizableObjectNode.h b/Engine/Plugins/Mutable/Source/CustomizableObjectEditor/Public/MuCOE/Nodes/CustomizableObjectNode.h
--- a/Engine/Plugins/Mutable/Source/CustomizableObjectEditor/Public/MuCOE/Nodes/CustomizableObjectNode.h
+++ b/Engine/Plugins/Mutable/Source/CustomizableObjectEditor/Public/MuCOE/Nodes/CustomizableObjectNode.h
@@ -24,6 +24,22 @@ class ICustomizableObjectEditor;
 
 struct FPropertyChangedEvent;
 
+/**
+ * Determines where the "Add Pin" button should be added within the node.
+ */
+enum class EAddPinNodeButtonLocation : uint8
+{
+	/** No "Add Pin" option will be added to the node UI. */
+	NONE = 0,
+
+	/** An option to add pins on the INPUT side of the node will be added. */
+	INPUT,
+
+	/** An option to add pins on the OUTPUT side of the node will be added. */
+	OUTPUT
+};
+
+
 /** Abstract base class for all pin data. */
 UCLASS(Abstract)
 class UCustomizableObjectNodePinData : public UObject
@@ -134,7 +150,7 @@ public:
 	virtual void PostBackwardsCompatibleFixup() {}
 
 	/** Virtual implementation of RemovePin. Allows to do work before removing a pin.
-	 * Use this function instead of RemovePin. RemovePin does not removes possible attached PinData. */
+	 * Use this function instead of RemovePin. RemovePin does not remove possible attached PinData. */
 	UE_API virtual bool CustomRemovePin(UEdGraphPin& Pin);
 
 	/**
@@ -281,11 +297,29 @@ public:
 	/** @return true if the details should spawn a Pin Viewer. */
 	UE_API virtual bool HasPinViewer() const;
 
-	/** @return true if the pin viewer's category should spawn a button to add new pins */
-	UE_API virtual bool CanCreatePinsFromPinViewer() const;
+	/** Determines to which side the "Add Pin" button should be added to*/
+	UE_API virtual EAddPinNodeButtonLocation GetAddPinButtonNodeSide() const;
 
-	UE_API virtual void CreatePinFromPinViewer() {};
+	/**
+	 * Will request the generation of a new pin and, if successful, will notify the graph of the change.
+	 * @note Method called when the Add Pins button is pressed.
+	 */
+	UE_API virtual void AddPinFromUI();
 
+	/**
+	 * Determines if a pin of this node can or can not be removed.
+	 * Override this method to handle the removal of your node pins.
+	 * @return True if the pin can be removed, false otherwise.
+	 */
+	UE_API virtual bool CanPinBeRemoved(const UEdGraphPin& Pin) const;
+	
+	/**
+	 * Removes one pin that is known to have been added by the Node "Add Pin" UI button.
+	 * @param Pin The pin that you want to get removed.
+	 * @return Will return true if the removal could be performed, false otherwise.
+	 */
+	UE_API bool TryRemovePin(UEdGraphPin& Pin);
+	
 	/** Returns pin custom details. Override if required. */
 	UE_API virtual TSharedPtr<IDetailsView> CustomizePinDetails(const UEdGraphPin& Pin) const;
 
@@ -304,7 +338,7 @@ public:
 		return *CastChecked<T>(PinsDataId[Pin.PinId]);
 	}
 
-	/** returs true if this node can be added to a COMacro Graph */
+	/** Returns true if this node can be added to a COMacro Graph */
 	UE_API virtual bool IsNodeSupportedInMacros() const;
 
 	/** Returns true if this node is inside a Macro Graph */
