[Mutable] Fix static analysis warning.

diff --git a/Engine/Plugins/Mutable/Source/MutableTools/Private/MuT/CodeOptimiser.cpp b/Engine/Plugins/Mutable/Source/MutableTools/Private/MuT/CodeOptimiser.cpp
--- a/Engine/Plugins/Mutable/Source/MutableTools/Private/MuT/CodeOptimiser.cpp
+++ b/Engine/Plugins/Mutable/Source/MutableTools/Private/MuT/CodeOptimiser.cpp
@@ -792,16 +792,14 @@ namespace UE::Mutable::Private
 					{
 						PASSTHROUGH_ID MeshID = PASSTHROUGH_ID_INVALID;
 						FString MeshMorph;
-						if (SubgraphType == EOpType::ME_CONSTANT)
+						
+						const ASTOpConstantResource* Typed = static_cast<const ASTOpConstantResource*>(SubgraphRoot.get());
+						const FMesh* Value = static_cast<const FMesh*>(Typed->GetValue().Get());
+						bIsReference = Value->IsReference();
+						if (bIsReference)
 						{
-							const ASTOpConstantResource* Typed = static_cast<const ASTOpConstantResource*>(SubgraphRoot.get());
-							const FMesh* Value = static_cast<const FMesh*>(Typed->GetValue().Get());
-							bIsReference = Value->IsReference();
-							if (bIsReference)
-							{
-								MeshID = Value->PassthroughObject.GetId();
-								MeshMorph = Value->GetReferencedMorph();
-							}
+							MeshID = Value->PassthroughObject.GetId();
+							MeshMorph = Value->GetReferencedMorph();
 						}
 
 						// Instead of generating the constant we resolve the reference, which also replaces the ASTOp.
