[Mutable] Change DDC key.

DDC is corrupted.

diff --git a/Engine/Plugins/Mutable/Source/CustomizableObject/Internal/MuCO/CustomizableObjectPrivate.h b/Engine/Plugins/Mutable/Source/CustomizableObject/Internal/MuCO/CustomizableObjectPrivate.h
--- a/Engine/Plugins/Mutable/Source/CustomizableObject/Internal/MuCO/CustomizableObjectPrivate.h
+++ b/Engine/Plugins/Mutable/Source/CustomizableObject/Internal/MuCO/CustomizableObjectPrivate.h
@@ -2043,7 +2043,7 @@ public:
 
 #endif
 	
-	static constexpr int32 DerivedDataVersion = 0x290ba346;
+	static constexpr int32 DerivedDataVersion = 0x059cf92d;
 };
 
 #if WITH_EDITOR
