[Mutable] Add Passthrough Skeletal Mesh pin to Node Passthrough Skeletal Mesh Component

diff --git a/Engine/Plugins/Mutable/Source/CustomizableObjectEditor/Private/MuCOE/Nodes/CustomizableObjectNodeSkeletalMeshParameter.cpp b/Engine/Plugins/Mutable/Source/CustomizableObjectEditor/Private/MuCOE/Nodes/CustomizableObjectNodeSkeletalMeshParameter.cpp
--- a/Engine/Plugins/Mutable/Source/CustomizableObjectEditor/Private/MuCOE/Nodes/CustomizableObjectNodeSkeletalMeshParameter.cpp
+++ b/Engine/Plugins/Mutable/Source/CustomizableObjectEditor/Private/MuCOE/Nodes/CustomizableObjectNodeSkeletalMeshParameter.cpp
@@ -19,6 +19,10 @@ void UCustomizableObjectNodeSkeletalMeshParameter::AllocateDefaultPins(UCustomiz
 {
 	NamePin = CustomCreatePin(EGPD_Input, UEdGraphSchema_CustomizableObject::PC_String, FName("Name"));
 
+	const UEdGraphSchema_CustomizableObject* Schema = GetDefault<UEdGraphSchema_CustomizableObject>();
+	
+	PassthroughPin = CustomCreatePin(EGPD_Output, Schema->PC_PassthroughSkeletalMesh, TEXT("Passthrough Skeletal Mesh"));
+	
 	USkeletalMesh* SkeletalMesh = ReferenceValue.LoadSynchronous();
 
 	if (!SkeletalMesh)
@@ -26,8 +30,6 @@ void UCustomizableObjectNodeSkeletalMeshParameter::AllocateDefaultPins(UCustomiz
 		return;
 	}
 	
-	const UEdGraphSchema_CustomizableObject* Schema = GetDefault<UEdGraphSchema_CustomizableObject>();
-
 	if (const FSkeletalMeshModel* ImportedModel = SkeletalMesh->GetImportedModel())
 	{
 		for (int32 LODIndex = 0; LODIndex < ImportedModel->LODModels.Num(); LODIndex++)
@@ -91,6 +93,11 @@ void UCustomizableObjectNodeSkeletalMeshParameter::BackwardsCompatibleFixup(int3
 			NamePin = CustomCreatePin(EGPD_Input, UEdGraphSchema_CustomizableObject::PC_String, FName("Name"));
 		}
 	}
+
+	if (CustomizableObjectCustomVersion == FCustomizableObjectCustomVersion::PinSkeletalMesh)
+	{
+		PassthroughPin = CustomCreatePin(EGPD_Output, UEdGraphSchema_CustomizableObject::PC_PassthroughSkeletalMesh, TEXT("Passthrough Skeletal Mesh"));
+	}
 }
 
 
