[Mutable] Added colors to Performance Analyzer

diff --git a/Engine/Plugins/Mutable/Source/CustomizableObject/Internal/MuCO/LogBenchmarkUtil.h b/Engine/Plugins/Mutable/Source/CustomizableObject/Internal/MuCO/LogBenchmarkUtil.h
--- a/Engine/Plugins/Mutable/Source/CustomizableObject/Internal/MuCO/LogBenchmarkUtil.h
+++ b/Engine/Plugins/Mutable/Source/CustomizableObject/Internal/MuCO/LogBenchmarkUtil.h
@@ -102,6 +102,32 @@ struct FInstanceUpdateStats
 };
 
 
+struct FStatistic
+{
+	double SD = 0.0f;
+	double Avg = 0.0f;
+	double Max = -TNumericLimits<double>::Max();
+};
+
+
+struct FInstanceUpdateStatistics
+{
+	FStatistic TriangleCount;
+	FStatistic QueueTime;
+	FStatistic UpdateTime;
+	FStatistic TaskGetMeshTime;
+	FStatistic TaskLockCacheTime;
+	FStatistic TaskGetImagesTime;
+	FStatistic TaskConvertResourcesTime;
+	FStatistic TaskCallbacksTime;
+	FStatistic UpdatePeakMemory;
+	FStatistic UpdateRealPeakMemory;
+	FStatistic TaskUpdateImageTime;
+	FStatistic TaskUpdateImagePeakMemory;
+	FStatistic TaskUpdateImageRealPeakMemory;
+};
+
+
 // Delegate signature
 DECLARE_MULTICAST_DELEGATE_TwoParams(FOnMeshUpdateReportedSignature, const TSharedRef<FUpdateContextPrivate> /* Context */, FInstanceUpdateStats)
 DECLARE_MULTICAST_DELEGATE_OneParam(FOnImageUpdateReportedSignature, FInstanceUpdateStats)
diff --git a/Engine/Plugins/Mutable/Source/CustomizableObjectEditor/Private/MuCOE/CustomizableObjectEditorPerformanceAnalyzer.cpp b/Engine/Plugins/Mutable/Source/CustomizableObjectEditor/Private/MuCOE/CustomizableObjectEditorPerformanceAnalyzer.cpp
--- a/Engine/Plugins/Mutable/Source/CustomizableObjectEditor/Private/MuCOE/CustomizableObjectEditorPerformanceAnalyzer.cpp
+++ b/Engine/Plugins/Mutable/Source/CustomizableObjectEditor/Private/MuCOE/CustomizableObjectEditorPerformanceAnalyzer.cpp
@@ -53,9 +53,10 @@ namespace InstanceUpdatesMainDataColumns
 
 #pragma region SInstanceUpdateDataRow
 
-void SInstanceUpdateDataRow::Construct(const FArguments& InArgs, const TSharedRef<STableViewBase>& InOwnerTableView, TSharedPtr<FInstanceUpdateDataElement> Element)
+void SInstanceUpdateDataRow::Construct(const FArguments& InArgs, const TSharedRef<STableViewBase>& InOwnerTableView, const TSharedPtr<FInstanceUpdateDataElement>& Element, const TSharedPtr<FInstanceUpdateStatistics>& Statistics)
 {
 	InstanceUpdateElement = Element;
+	InstanceUpdateStatistics = Statistics;
 	
 	SMultiColumnTableRow<TSharedPtr<FInstanceUpdateDataElement>>::Construct(
 		STableRow::FArguments()
@@ -74,6 +75,15 @@ void SInstanceUpdateDataRow::OnInstanceNameNavigation() const
 }
 
 
+FSlateColor TextColor(double Value, const FStatistic& Statistic)
+{
+	const float Normalized = FMath::Clamp((Value - Statistic.Avg) / (Statistic.Max - Statistic.Avg), 0.0, 1.0);
+	const FLinearColor Color = FLinearColor::LerpUsingHSV(FLinearColor::Green, FLinearColor::Red, Normalized);
+	
+	return FSlateColor(Color);
+}
+
+
 TSharedRef<SWidget> SInstanceUpdateDataRow::GenerateWidgetForColumn(const FName& InColumnName)
 {
 	if (!InstanceUpdateElement)
@@ -96,7 +106,8 @@ TSharedRef<SWidget> SInstanceUpdateDataRow::GenerateWidgetForColumn(const FName&
 	else if (InColumnName == InstanceUpdatesMainDataColumns::InstanceTriangleCount)
 	{
 		return SNew(STextBlock)
-			.Text( FText::AsNumber(InstanceUpdateElement->UpdateStats.TriangleCount));
+			.Text( FText::AsNumber(InstanceUpdateElement->UpdateStats.TriangleCount))
+			.ColorAndOpacity(TextColor(InstanceUpdateElement->UpdateStats.TriangleCount, InstanceUpdateStatistics->TriangleCount));
 	}
 	else if (InColumnName == InstanceUpdatesMainDataColumns::InstanceDescriptorColumnID)
 	{
@@ -109,17 +120,17 @@ TSharedRef<SWidget> SInstanceUpdateDataRow::GenerateWidgetForColumn(const FName&
 		switch (InstanceUpdateElement->UpdateStats.UpdateResult)
 		{
 		case EUpdateResult::Success:
-			TextColor = FSlateColor(FColor(70,207,120));
+			TextColor = FSlateColor(FLinearColor::Green);
 			break;
 		case EUpdateResult::Warning:
-			TextColor = FSlateColor(FColor(250,226,7));
+			TextColor = FSlateColor(FLinearColor::Yellow);
 			break;
 		case EUpdateResult::Error:
 		case EUpdateResult::ErrorOptimized:
 		case EUpdateResult::ErrorReplaced:
 		case EUpdateResult::ErrorDiscarded:
 		case EUpdateResult::Error16BitBoneIndex:
-			TextColor = FSlateColor(FColor(197,0,7));
+			TextColor = FSlateColor(FLinearColor::Red);
 			break;
 		}
 		
@@ -130,47 +141,58 @@ TSharedRef<SWidget> SInstanceUpdateDataRow::GenerateWidgetForColumn(const FName&
 	else if (InColumnName == InstanceUpdatesMainDataColumns::QueueTimeColumnID)
 	{
 		return SNew(STextBlock)
-			.Text(FText::FromString(FString::Printf(TEXT("%.2f"),InstanceUpdateElement->UpdateStats.QueueTime)));
+			.Text(FText::FromString(FString::Printf(TEXT("%.2f"),InstanceUpdateElement->UpdateStats.QueueTime)))
+			.ColorAndOpacity(TextColor(InstanceUpdateElement->UpdateStats.QueueTime, InstanceUpdateStatistics->QueueTime));
+
 	}
 	else if (InColumnName == InstanceUpdatesMainDataColumns::UpdateTimeColumnID)
 	{
 		return SNew(STextBlock)
-			.Text(FText::FromString(FString::Printf(TEXT("%.2f"),InstanceUpdateElement->UpdateStats.UpdateTime)));
+			.Text(FText::FromString(FString::Printf(TEXT("%.2f"),InstanceUpdateElement->UpdateStats.UpdateTime)))
+			.ColorAndOpacity(TextColor(InstanceUpdateElement->UpdateStats.UpdateTime, InstanceUpdateStatistics->UpdateTime));
 	}
 	else if (InColumnName == InstanceUpdatesMainDataColumns::GetMeshTimeColumnID)
 	{
 		return SNew(STextBlock)
-			.Text(FText::FromString(FString::Printf(TEXT("%.2f"),InstanceUpdateElement->UpdateStats.TaskGetMeshTime)));
+			.Text(FText::FromString(FString::Printf(TEXT("%.2f"),InstanceUpdateElement->UpdateStats.TaskGetMeshTime)))
+			.ColorAndOpacity(TextColor(InstanceUpdateElement->UpdateStats.TaskGetMeshTime, InstanceUpdateStatistics->TaskGetMeshTime));
 	}
 	else if (InColumnName == InstanceUpdatesMainDataColumns::LockCacheTimeColumnID)
 	{
 		return SNew(STextBlock)
-			.Text(FText::FromString(FString::Printf(TEXT("%.4f"),InstanceUpdateElement->UpdateStats.TaskLockCacheTime)));
+			.Text(FText::FromString(FString::Printf(TEXT("%.4f"),InstanceUpdateElement->UpdateStats.TaskLockCacheTime)))
+			.ColorAndOpacity(TextColor(InstanceUpdateElement->UpdateStats.TaskLockCacheTime, InstanceUpdateStatistics->TaskLockCacheTime));
 	}
 	else if (InColumnName == InstanceUpdatesMainDataColumns::GetImagesTimeColumnID)
 	{
 		return SNew(STextBlock)
-			.Text(FText::FromString(FString::Printf(TEXT("%.4f"),InstanceUpdateElement->UpdateStats.TaskGetImagesTime)));
+			.Text(FText::FromString(FString::Printf(TEXT("%.4f"),InstanceUpdateElement->UpdateStats.TaskGetImagesTime)))
+			.ColorAndOpacity(TextColor(InstanceUpdateElement->UpdateStats.TaskGetImagesTime, InstanceUpdateStatistics->TaskGetImagesTime));
 	}
 	else if (InColumnName == InstanceUpdatesMainDataColumns::ConvertResourcesTimeColumnID)
 	{
 		return SNew(STextBlock)
-			.Text(FText::FromString(FString::Printf(TEXT("%.2f"),InstanceUpdateElement->UpdateStats.TaskConvertResourcesTime)));
+			.Text(FText::FromString(FString::Printf(TEXT("%.2f"),InstanceUpdateElement->UpdateStats.TaskConvertResourcesTime)))
+			.ColorAndOpacity(TextColor(InstanceUpdateElement->UpdateStats.TaskConvertResourcesTime, InstanceUpdateStatistics->TaskConvertResourcesTime));
 	}
 	else if (InColumnName == InstanceUpdatesMainDataColumns::CallbacksTimeColumnID)
 	{
 		return SNew(STextBlock)
-			.Text(FText::FromString(FString::Printf(TEXT("%.2f"),InstanceUpdateElement->UpdateStats.TaskCallbacksTime)));
+			.Text(FText::FromString(FString::Printf(TEXT("%.2f"),InstanceUpdateElement->UpdateStats.TaskCallbacksTime)))
+			.ColorAndOpacity(TextColor(InstanceUpdateElement->UpdateStats.TaskCallbacksTime, InstanceUpdateStatistics->TaskCallbacksTime));
 	}
 	else if (InColumnName == InstanceUpdatesMainDataColumns::UpdatePeakMemoryColumnID)
 	{
 		return SNew(STextBlock)
-			.Text(FText::FromString(FString::Printf(TEXT("%.2f"),InstanceUpdateElement->UpdateStats.UpdatePeakMemory)));
+			.Text(FText::FromString(FString::Printf(TEXT("%.2f"),InstanceUpdateElement->UpdateStats.UpdatePeakMemory)))
+			.ColorAndOpacity(TextColor(InstanceUpdateElement->UpdateStats.UpdatePeakMemory, InstanceUpdateStatistics->UpdatePeakMemory));
 	}
 	else if (InColumnName == InstanceUpdatesMainDataColumns::UpdateRealPeakMemoryColumnID)
 	{
 		return SNew(STextBlock)
-			.Text(FText::FromString(FString::Printf(TEXT("%.2f"),InstanceUpdateElement->UpdateStats.UpdateRealPeakMemory)));
+			.Text(FText::FromString(FString::Printf(TEXT("%.2f"),InstanceUpdateElement->UpdateStats.UpdateRealPeakMemory)))
+			.ColorAndOpacity(TextColor(InstanceUpdateElement->UpdateStats.UpdateRealPeakMemory, InstanceUpdateStatistics->UpdateRealPeakMemory));
+
 	}
 
 	// Invalid column name so no widget will be produced 
@@ -649,7 +671,7 @@ EColumnSortMode::Type SCustomizableObjectEditorPerformanceAnalyzer::GetColumnSor
 
 TSharedRef<ITableRow> SCustomizableObjectEditorPerformanceAnalyzer::OnGenerateInstanceUpdateRow(TSharedPtr<FInstanceUpdateDataElement> InstanceUpdateDataElement, const TSharedRef<STableViewBase>& TableViewBase)
 {
-	TSharedRef<SInstanceUpdateDataRow> Row = SNew(SInstanceUpdateDataRow, TableViewBase, InstanceUpdateDataElement);
+	TSharedRef<SInstanceUpdateDataRow> Row = SNew(SInstanceUpdateDataRow, TableViewBase, InstanceUpdateDataElement, InstanceUpdateStatistics);
 	return Row;
 }
 
@@ -876,13 +898,15 @@ void SCustomizableObjectEditorPerformanceAnalyzer::OnBenchmarkMeshUpdated(TShare
 	// 3) Store the element to keep the instance alive
 	check(InstanceElement);
 	InstanceUpdateElements.Push(InstanceElement);
+
+	ComputeStatistics();
 	
 	// Reset control flags for the next instance
 	CurrentInstance = nullptr;
 	
 	// And once the updates have completed update the list view
 	check(InstanceUpdatesListView);
-	InstanceUpdatesListView->RequestListRefresh();
+	InstanceUpdatesListView->RebuildList();
 	
 	// Stop listening for instance update reported data
 	System->GetPrivate()->LogBenchmarkUtil.OnMeshUpdateReported.Remove(OnMeshUpdatePerfReportDelegateHandle);
@@ -923,4 +947,53 @@ bool SCustomizableObjectEditorPerformanceAnalyzer::AreUpdatesPending() const
 	return CurrentInstance || !InstancesToUpdate.IsEmpty();
 }
 
+
+void SCustomizableObjectEditorPerformanceAnalyzer::ComputeStatistics()
+{
+	const int32 Num = InstanceUpdateElements.Num();
+
+	auto ComputeMax = [this, Num](FStatistic& Statistic, const TFunction<double(int32)>& GetValue)
+	{
+		Statistic = {};
+
+		if (!Num)
+		{
+			return;
+		}
+			
+		for (int32 Index = 0; Index < Num; ++Index)
+		{
+			Statistic.Max = FMath::Max(Statistic.Max, GetValue(Index));
+		}
+
+		for (int32 Index = 0; Index < Num; ++Index)
+		{
+			Statistic.Avg += GetValue(Index);
+		}
+
+		Statistic.Avg /= Num;
+			
+		for (int32 Index = 0; Index < Num; ++Index)
+		{
+			Statistic.SD += FMath::Square(GetValue(Index) - Statistic.Avg);
+		}
+
+		Statistic.SD = FMath::Sqrt(Statistic.SD / Num);
+	};
+	
+	ComputeMax(InstanceUpdateStatistics->TriangleCount, [&](int32 Index){	return InstanceUpdateElements[Index]->UpdateStats.TriangleCount; });
+	ComputeMax(InstanceUpdateStatistics->QueueTime, [&](int32 Index){	return InstanceUpdateElements[Index]->UpdateStats.QueueTime; });
+	ComputeMax(InstanceUpdateStatistics->UpdateTime, [&](int32 Index) { return InstanceUpdateElements[Index]->UpdateStats.UpdateTime; });
+	ComputeMax(InstanceUpdateStatistics->TaskGetMeshTime, [&](int32 Index) { return InstanceUpdateElements[Index]->UpdateStats.TaskGetMeshTime; });
+	ComputeMax(InstanceUpdateStatistics->TaskLockCacheTime, [&](int32 Index) { return InstanceUpdateElements[Index]->UpdateStats.TaskLockCacheTime; });
+	ComputeMax(InstanceUpdateStatistics->TaskGetImagesTime, [&](int32 Index) { return InstanceUpdateElements[Index]->UpdateStats.TaskGetImagesTime; });
+	ComputeMax(InstanceUpdateStatistics->TaskConvertResourcesTime, [&](int32 Index) { return InstanceUpdateElements[Index]->UpdateStats.TaskConvertResourcesTime; });
+	ComputeMax(InstanceUpdateStatistics->TaskCallbacksTime, [&](int32 Index) { return InstanceUpdateElements[Index]->UpdateStats.TaskCallbacksTime; });
+	ComputeMax(InstanceUpdateStatistics->UpdatePeakMemory, [&](int32 Index) { return InstanceUpdateElements[Index]->UpdateStats.UpdatePeakMemory; });
+	ComputeMax(InstanceUpdateStatistics->UpdateRealPeakMemory, [&](int32 Index) { return InstanceUpdateElements[Index]->UpdateStats.UpdateRealPeakMemory; });
+	ComputeMax(InstanceUpdateStatistics->TaskUpdateImageTime, [&](int32 Index) { return InstanceUpdateElements[Index]->UpdateStats.TaskUpdateImageTime; });
+	ComputeMax(InstanceUpdateStatistics->TaskUpdateImagePeakMemory, [&](int32 Index) { return InstanceUpdateElements[Index]->UpdateStats.TaskUpdateImagePeakMemory; });
+	ComputeMax(InstanceUpdateStatistics->TaskUpdateImageRealPeakMemory, [&](int32 Index) { return InstanceUpdateElements[Index]->UpdateStats.TaskUpdateImageRealPeakMemory; });
+}
+
 #undef LOCTEXT_NAMESPACE
diff --git a/Engine/Plugins/Mutable/Source/CustomizableObjectEditor/Private/MuCOE/CustomizableObjectEditorPerformanceAnalyzer.h b/Engine/Plugins/Mutable/Source/CustomizableObjectEditor/Private/MuCOE/CustomizableObjectEditorPerformanceAnalyzer.h
--- a/Engine/Plugins/Mutable/Source/CustomizableObjectEditor/Private/MuCOE/CustomizableObjectEditorPerformanceAnalyzer.h
+++ b/Engine/Plugins/Mutable/Source/CustomizableObjectEditor/Private/MuCOE/CustomizableObjectEditorPerformanceAnalyzer.h
@@ -45,13 +45,12 @@ public:
 class SInstanceUpdateDataRow : public SMultiColumnTableRow<TSharedPtr<FInstanceUpdateDataElement>>
 {
 public:
-	
-	void Construct(const FArguments& InArgs, const TSharedRef<STableViewBase>& InOwnerTableView, TSharedPtr<FInstanceUpdateDataElement> Element);
+	void Construct(const FArguments& InArgs, const TSharedRef<STableViewBase>& InOwnerTableView, const TSharedPtr<FInstanceUpdateDataElement>& Element, const TSharedPtr<FInstanceUpdateStatistics>& Statistics);
 	
 private:
-	
 	/** Backend object whose data we are drawing */
 	TSharedPtr<FInstanceUpdateDataElement> InstanceUpdateElement;
+	TSharedPtr<FInstanceUpdateStatistics> InstanceUpdateStatistics;
 
 	/** Method invoked when the hyperlink showing the name of the instance is clicked */
 	void OnInstanceNameNavigation() const;
@@ -130,8 +129,8 @@ private:
 	/** Determines if there are updates to wait for. Checks not only the queue but also the current instnace being udpated*/
 	bool AreUpdatesPending() const;
 	
-private:
-
+	void ComputeStatistics();
+	
 	/** Delegate pointing to the LogBenchmarkUtil on Mesh Updated event */
 	FDelegateHandle OnMeshUpdatePerfReportDelegateHandle;
 
@@ -153,6 +152,8 @@ private:
 	/** Instance update elements (size == RandomInstancesToGenerate) at the time of the update request */
 	TArray<TSharedPtr<FInstanceUpdateDataElement>> InstanceUpdateElements;
 
+	TSharedPtr<FInstanceUpdateStatistics> InstanceUpdateStatistics = MakeShared<FInstanceUpdateStatistics>();
+	
 	/** The name of the last column the user wanted to sort */
 	FName CurrentSortColumn;
 	/** The sort type applied to the last column that got sorted */
