[Mutable] Refactored Task_Mutable_GetImages

Simplified for upcoming changes.

diff --git a/Engine/Plugins/Mutable/Source/CustomizableObject/Internal/MuCO/CustomizableObjectSystemPrivate.h b/Engine/Plugins/Mutable/Source/CustomizableObject/Internal/MuCO/CustomizableObjectSystemPrivate.h
--- a/Engine/Plugins/Mutable/Source/CustomizableObject/Internal/MuCO/CustomizableObjectSystemPrivate.h
+++ b/Engine/Plugins/Mutable/Source/CustomizableObject/Internal/MuCO/CustomizableObjectSystemPrivate.h
@@ -380,7 +380,11 @@ struct FInstanceUpdateData
 		int32 BaseLOD;
 		int32 BaseMip;
 		
-		uint16 FullImageSizeX, FullImageSizeY;
+		uint16 FullImageSizeX;
+		uint16 FullImageSizeY;
+		
+		UE::Mutable::Private::EImageFormat Format = UE::Mutable::Private::EImageFormat::None;
+
 		TSharedPtr<const UE::Mutable::Private::FImage> Image;
 		TWeakObjectPtr<UTexture2D> Cached;
 
@@ -465,7 +469,7 @@ struct FInstanceUpdateData
 	TArray<FComponent> Components;
 	TArray<FLOD> LODs;
 	TArray<FSurface> Surfaces;
-	TArray<FImage> Images;
+	TArray<TSharedRef<FImage>> Images;
 	TArray<FVector> Vectors;
 	TArray<FScalar> Scalars;
 
@@ -627,6 +631,8 @@ public:
 	/** User config. */
 	bool bIsOnlyGenerateRequestedLODsEnabled = false;
 
+	bool bEnableSkipGenerateResidentMips = false;
+
 	/** This option comes from the operation request. It is used to reduce the number of mipmaps that mutable must generate for images.  */
 	int32 MipsToSkip = 0;
 
diff --git a/Engine/Plugins/Mutable/Source/CustomizableObject/Private/MuCO/CustomizableObjectInstance.cpp b/Engine/Plugins/Mutable/Source/CustomizableObject/Private/MuCO/CustomizableObjectInstance.cpp
--- a/Engine/Plugins/Mutable/Source/CustomizableObject/Private/MuCO/CustomizableObjectInstance.cpp
+++ b/Engine/Plugins/Mutable/Source/CustomizableObject/Private/MuCO/CustomizableObjectInstance.cpp
@@ -6018,19 +6018,19 @@ void UCustomizableInstancePrivate::BuildMaterials(const TSharedRef<FUpdateContex
 
 					for (int32 ImageIndex = 0; ImageIndex < Surface.ImageCount; ++ImageIndex)
 					{
-						const FInstanceUpdateData::FImage& Image = Context->InstanceUpdateData.Images[Surface.FirstImage + ImageIndex];
-						FString KeyName = Image.Name.ToString();
-						TSharedPtr<const UE::Mutable::Private::FImage> MutableImage = Image.Image;
+						const TSharedRef<const FInstanceUpdateData::FImage>& Image = Context->InstanceUpdateData.Images[Surface.FirstImage + ImageIndex];
+						FString KeyName = Image->Name.ToString();
+						TSharedPtr<const UE::Mutable::Private::FImage> MutableImage = Image->Image;
 
 						UTexture2D* MutableTexture = nullptr; // Texture generated by mutable
 						UTexture* PassThroughTexture = nullptr; // Texture not generated by mutable
 
-						const FTextureReuseCacheKey TextureReuseCacheRef = bReuseTextures ? FTextureReuseCacheKey{Image.BaseLOD, ObjectComponentIndex.GetValue(), Surface.SurfaceId, ImageIndex} : FTextureReuseCacheKey{};
+						const FTextureReuseCacheKey TextureReuseCacheRef = bReuseTextures ? FTextureReuseCacheKey{Image->BaseLOD, ObjectComponentIndex.GetValue(), Surface.SurfaceId, ImageIndex} : FTextureReuseCacheKey{};
 
 #if WITH_EDITOR
-						const FTextureCache::FId TextureCacheKey = FTextureCache::FId(Image.ImageID, Context->MipsToSkip, Context->bBake);
+						const FTextureCache::FId TextureCacheKey = FTextureCache::FId(Image->ImageID, Context->MipsToSkip, Context->bBake);
 #else
-						const FTextureCache::FId TextureCacheKey = FTextureCache::FId(Image.ImageID, Context->MipsToSkip);
+						const FTextureCache::FId TextureCacheKey = FTextureCache::FId(Image->ImageID, Context->MipsToSkip);
 #endif
 						
 						// If the mutable image is null, it must be in the cache
@@ -6041,7 +6041,7 @@ void UCustomizableInstancePrivate::BuildMaterials(const TSharedRef<FUpdateContex
 						}
 
 						// Check if the image is a reference to an engine texture
-						if (MutableImage && Image.bIsPassThrough)
+						if (MutableImage && Image->bIsPassThrough)
 						{
 							check(MutableImage->IsReference());
 
@@ -6069,7 +6069,7 @@ void UCustomizableInstancePrivate::BuildMaterials(const TSharedRef<FUpdateContex
 
 								FString MutableTextureName;
 								{
-									const uint32 MutableImageHash = UE::Mutable::Private::GetTypeHashPersistent(Image.ImageID);
+									const uint32 MutableImageHash = UE::Mutable::Private::GetTypeHashPersistent(Image->ImageID);
 									MutableTextureName = FString::Printf(TEXT("T_%s_h%llu"), *Props.TextureParameterName, MutableImageHash);
 									MutableTextureName.ReplaceInline(TEXT(" "), TEXT("_"));
 
@@ -6115,11 +6115,11 @@ void UCustomizableInstancePrivate::BuildMaterials(const TSharedRef<FUpdateContex
 
 								if (MutableTexture)
 								{
-									if (Context->ImageToPlatformDataMap.Contains(Image.ImageID))
+									if (Context->ImageToPlatformDataMap.Contains(Image->ImageID))
 									{
 										SetTexturePropertiesFromMutableImageProps(MutableTexture, Props, bNeverStream);
 
-										FTexturePlatformData* PlatformData = Context->ImageToPlatformDataMap[Image.ImageID];
+										FTexturePlatformData* PlatformData = Context->ImageToPlatformDataMap[Image->ImageID];
 
 										if (ReusedTexture)
 										{
@@ -6136,13 +6136,13 @@ void UCustomizableInstancePrivate::BuildMaterials(const TSharedRef<FUpdateContex
 											MutableTexture->SetPlatformData(PlatformData);
 										}
 
-										Context->ImageToPlatformDataMap.Remove(Image.ImageID);
+										Context->ImageToPlatformDataMap.Remove(Image->ImageID);
 									}
 									else
 									{
 										UE_LOG(LogMutable, Error, TEXT("Required image [%s] with ID [%s] was not generated in the mutable thread, and it is not cached. LOD [%d]. Object Component [%d]"),
 											*Props.TextureParameterName,
-											*Image.ImageID.ToString(),
+											*Image->ImageID.ToString(),
 											LODIndex, ObjectComponentIndex.GetValue());
 										continue;
 									}
@@ -6192,13 +6192,13 @@ void UCustomizableInstancePrivate::BuildMaterials(const TSharedRef<FUpdateContex
 												{
 													MutableMipDataProviderFactory->CustomizableObjectInstance = Public;
 													check(LODIndex < 256 && InstanceComponentIndex < 256 && ImageIndex < 256);
-													MutableMipDataProviderFactory->ImageRef.ImageID = Image.ImageID;
+													MutableMipDataProviderFactory->ImageRef.ImageID = Image->ImageID;
 													MutableMipDataProviderFactory->ImageRef.SurfaceId = Surface.SurfaceId;
-													MutableMipDataProviderFactory->ImageRef.LOD = uint8(Image.BaseLOD);
+													MutableMipDataProviderFactory->ImageRef.LOD = uint8(Image->BaseLOD);
 													MutableMipDataProviderFactory->ImageRef.Component = uint8(InstanceComponentIndex);
 													MutableMipDataProviderFactory->ImageRef.Image = uint8(ImageIndex);
-													MutableMipDataProviderFactory->ImageRef.BaseMip = uint8(Image.BaseMip);
-													MutableMipDataProviderFactory->ImageRef.ConstantImagesNeededToGenerate = Image.ConstantImagesNeededToGenerate;
+													MutableMipDataProviderFactory->ImageRef.BaseMip = uint8(Image->BaseMip);
+													MutableMipDataProviderFactory->ImageRef.ConstantImagesNeededToGenerate = Image->ConstantImagesNeededToGenerate;
 													MutableMipDataProviderFactory->UpdateContext = UpdateContext;
 													MutableTexture->AddAssetUserData(MutableMipDataProviderFactory);
 												}
@@ -6227,7 +6227,7 @@ void UCustomizableInstancePrivate::BuildMaterials(const TSharedRef<FUpdateContex
 							}
 
 							// Decoding Material Layer from Mutable parameter name
-							FString ImageName = Image.Name.ToString();
+							FString ImageName = Image->Name.ToString();
 							FString EncodingString = "-MutableLayerParam:";
 
 							int32 EncodingPosition = ImageName.Find(EncodingString);
diff --git a/Engine/Plugins/Mutable/Source/CustomizableObject/Private/MuCO/CustomizableObjectMeshUpdate.cpp b/Engine/Plugins/Mutable/Source/CustomizableObject/Private/MuCO/CustomizableObjectMeshUpdate.cpp
--- a/Engine/Plugins/Mutable/Source/CustomizableObject/Private/MuCO/CustomizableObjectMeshUpdate.cpp
+++ b/Engine/Plugins/Mutable/Source/CustomizableObject/Private/MuCO/CustomizableObjectMeshUpdate.cpp
@@ -206,7 +206,7 @@ namespace impl
 		const UE::Mutable::Private::FMeshId& MeshID = OperationData->MeshIDs[LODIndex];
 
 		constexpr UE::Mutable::Private::EMeshContentFlags MeshContentFilter = UE::Mutable::Private::EMeshContentFlags::AllFlags;
-		UE::Tasks::TTask<TSharedPtr<const UE::Mutable::Private::FMesh>> GetMeshTask = OperationData->System->GetMesh(InstanceID, MeshID, MeshContentFilter);
+		UE::Tasks::TTask<TSharedPtr<const UE::Mutable::Private::FMesh>> GetMeshTask = OperationData->System->GetMesh({}, InstanceID, MeshID, MeshContentFilter);
 
 		UE::Tasks::AddNested(UE::Tasks::Launch(TEXT("Task_MutableGetMeshes_GetMesh_Post"), [=]() mutable
 			{
diff --git a/Engine/Plugins/Mutable/Source/CustomizableObject/Private/MuCO/CustomizableObjectMipDataProvider.cpp b/Engine/Plugins/Mutable/Source/CustomizableObject/Private/MuCO/CustomizableObjectMipDataProvider.cpp
--- a/Engine/Plugins/Mutable/Source/CustomizableObject/Private/MuCO/CustomizableObjectMipDataProvider.cpp
+++ b/Engine/Plugins/Mutable/Source/CustomizableObject/Private/MuCO/CustomizableObjectMipDataProvider.cpp
@@ -162,8 +162,7 @@ namespace Impl
 		// runtime cannot remember all the resources it has built, and only remembers a fixed amount.
 		UE::Mutable::Private::FImageId MipImageID = Instance->GetImageId(ImageRef.Component, ImageRef.LOD, SurfaceIndex, ImageRef.Image);
 
-		UE::Tasks::TTask<TSharedPtr<const UE::Mutable::Private::FImage>> GetImageTask = 
-				System->GetImage(InstanceID, MipImageID, ImageRef.BaseMip + OperationData->MipsToSkip, ImageRef.LOD);
+		UE::Tasks::TTask<TSharedPtr<const UE::Mutable::Private::FImage>> GetImageTask = System->GetImage({}, InstanceID, MipImageID, ImageRef.BaseMip + OperationData->MipsToSkip, ImageRef.LOD);
 
 		UE::Tasks::AddNested(UE::Tasks::Launch(TEXT("MipDataProvider_EndUpdateImagesTask"),
 				[System, OperationData, InstanceID, StartTime, GetImageTask, EndUpdateImage]() mutable
diff --git a/Engine/Plugins/Mutable/Source/CustomizableObject/Private/MuCO/CustomizableObjectSystem.cpp b/Engine/Plugins/Mutable/Source/CustomizableObject/Private/MuCO/CustomizableObjectSystem.cpp
--- a/Engine/Plugins/Mutable/Source/CustomizableObject/Private/MuCO/CustomizableObjectSystem.cpp
+++ b/Engine/Plugins/Mutable/Source/CustomizableObject/Private/MuCO/CustomizableObjectSystem.cpp
@@ -58,7 +58,6 @@ DECLARE_CYCLE_STAT(TEXT("MutableTask"), STAT_MutableTask, STATGROUP_Game);
 
 #define UE_MUTABLE_UPDATE_REGION TEXT("Mutable Update")
 #define UE_TASK_MUTABLE_GETMESHES_REGION TEXT("Task_Mutable_GetMeshes")
-#define UE_TASK_MUTABLE_GETIMAGES_REGION TEXT("Task_Mutable_GetImages")
 
 
 UCustomizableObjectSystem* UCustomizableObjectSystemPrivate::SSystem = nullptr;
@@ -1955,25 +1954,6 @@ namespace impl
 		double StartTime,
 		const TSharedRef<TArray<FGetMeshData>>& GetMeshesData,
 		int32 GetMeshIndex);
-
-
-	/** Call GetImage.
-	  * Once GetImage is called, the task must end. Following code will be in a subsequent TaskGraph task. */
-	void Task_Mutable_GetImages_GetImage(
-		const TSharedRef<FUpdateContextPrivate>& OperationData,
-		double StartTime,
-		const TSharedPtr<TArray<UE::Mutable::Private::FImageId>>& ImagesInThisInstance,
-		int32 ImageIndex,
-		UE::Tasks::TTask<UE::Mutable::Private::FExtendedImageDesc> GetImageDescTask);
-
-	
-	/** Process the next Image. If there are no more Images, go to the end of the task. */
-	void Task_Mutable_GetImages_Loop(
-		const TSharedRef<FUpdateContextPrivate>& OperationData,
-		double StartTime,
-		const TSharedPtr<TArray<UE::Mutable::Private::FImageId>>& ImagesInThisInstance,
-		int32 ImageIndex);
-
 	
 	void Subtask_Mutable_UpdateParameterRelevancy(const TSharedRef<FUpdateContextPrivate>& OperationData)
 	{
@@ -2051,43 +2031,6 @@ namespace impl
 	}
 	
 
-	// This runs in a worker thread
-	void Subtask_Mutable_PrepareTextures(const TSharedRef<FUpdateContextPrivate>& OperationData)
-	{
-		MUTABLE_CPUPROFILER_SCOPE(Subtask_Mutable_PrepareTextures)
-
-		for (const FInstanceUpdateData::FSurface& Surface : OperationData->InstanceUpdateData.Surfaces)
-		{
-			for (int32 ImageIndex = 0; ImageIndex<Surface.ImageCount; ++ImageIndex)
-			{
-				const FInstanceUpdateData::FImage& Image = OperationData->InstanceUpdateData.Images[Surface.FirstImage+ImageIndex];
-
-				const FName KeyName = Image.Name;
-				TSharedPtr<const UE::Mutable::Private::FImage> MutableImage = Image.Image;
-
-				// If the image is null, it must be in the cache (or repeated in this instance), and we don't need to do anything here.
-				if (MutableImage)
-				{
-					// Image references are just references to texture assets and require no work at all
-					if (!MutableImage->IsReference())
-					{
-						if (!OperationData->ImageToPlatformDataMap.Contains(Image.ImageID))
-						{
-							FTexturePlatformData* PlatformData = MutableCreateImagePlatformData(MutableImage, -1, Image.FullImageSizeX, Image.FullImageSizeY);
-							OperationData->ImageToPlatformDataMap.Add(Image.ImageID, PlatformData);
-						}
-						else
-						{
-							// The ImageID already exists in the ImageToPlatformDataMap, that means the equivalent surface in a lower
-							// LOD already created the PlatformData for that ImageID and added it to the ImageToPlatformDataMap.
-						}
-					}
-				}
-			}
-		}
-	}
-	
-
 	// This runs in a worker thread
 	void Subtask_Mutable_PrepareSkeletonData(const TSharedRef<FUpdateContextPrivate>& OperationData)
 	{
@@ -2345,10 +2288,10 @@ namespace impl
 
 		const int32 ImageIndex = (*GetImagesData)[GetImageIndex].ImageIndex;
 
-		FInstanceUpdateData::FImage& Image = OperationData->InstanceUpdateData.Images[ImageIndex];
+		TSharedRef<FInstanceUpdateData::FImage> Image = OperationData->InstanceUpdateData.Images[ImageIndex];
 
-		Image.Image = GetImageTask.GetResult();
-		check(Image.Image->IsReference());
+		Image->Image = GetImageTask.GetResult();
+		check(Image->Image->IsReference());
 		
 		Task_Mutable_GetMeshes_GetImage_Loop(OperationData, StartTime, GetImagesData, ++GetImageIndex);
 	}
@@ -2371,7 +2314,7 @@ namespace impl
 
 		const FGetImageData& ImageData = (*GetImagesData)[GetImageIndex];
 		
-		UE::Tasks::TTask<TSharedPtr<const UE::Mutable::Private::FImage>> GetImageTask = OperationData->MutableSystem->GetImage(OperationData->InstanceID, ImageData.ImageID, 0, 0);
+		UE::Tasks::TTask<TSharedPtr<const UE::Mutable::Private::FImage>> GetImageTask = OperationData->MutableSystem->GetImage({}, OperationData->InstanceID, ImageData.ImageID, 0, 0);
 
 		UE::Tasks::AddNested(UE::Tasks::Launch(TEXT("Task_Mutable_GetMeshes_GetImage_Post"), [=]()
 		{
@@ -2475,7 +2418,7 @@ namespace impl
 						{
 							MUTABLE_CPUPROFILER_SCOPE(GetScalar)
 
-								const FName ScalarName = MutableInstance->GetScalarName(InstanceComponentIndex, BaseLODIndex, BaseSurfaceIndex, ScalarIndex);
+							const FName ScalarName = MutableInstance->GetScalarName(InstanceComponentIndex, BaseLODIndex, BaseSurfaceIndex, ScalarIndex);
 							const float ScalarValue = MutableInstance->GetScalar(InstanceComponentIndex, BaseLODIndex, BaseSurfaceIndex, ScalarIndex);
 
 							OperationData->InstanceUpdateData.Scalars.Push({ ScalarName, ScalarValue });
@@ -2488,30 +2431,30 @@ namespace impl
 						{
 							MUTABLE_CPUPROFILER_SCOPE(GetImageId);
 
-							const int32 UpdateDataImageIndex = OperationData->InstanceUpdateData.Images.AddDefaulted();
-							FInstanceUpdateData::FImage& Image = OperationData->InstanceUpdateData.Images.Last();
-							Image.Name = MutableInstance->GetImageName(InstanceComponentIndex, BaseLODIndex, BaseSurfaceIndex, ImageIndex);
-							Image.ImageID = MutableInstance->GetImageId(InstanceComponentIndex, BaseLODIndex, BaseSurfaceIndex, ImageIndex);
-							Image.FullImageSizeX = 0;
-							Image.FullImageSizeY = 0;
-							Image.BaseLOD = BaseLODIndex;
-							Image.BaseMip = 0;
+							const int32 UpdateDataImageIndex = OperationData->InstanceUpdateData.Images.Add(MakeShared<FInstanceUpdateData::FImage>());
+							TSharedRef<FInstanceUpdateData::FImage> Image = OperationData->InstanceUpdateData.Images.Last();
+							Image->Name = MutableInstance->GetImageName(InstanceComponentIndex, BaseLODIndex, BaseSurfaceIndex, ImageIndex);
+							Image->ImageID = MutableInstance->GetImageId(InstanceComponentIndex, BaseLODIndex, BaseSurfaceIndex, ImageIndex);
+							Image->FullImageSizeX = 0;
+							Image->FullImageSizeY = 0;
+							Image->BaseLOD = BaseLODIndex;
+							Image->BaseMip = 0;
 
-							FString KeyName = Image.Name.ToString();
+							FString KeyName = Image->Name.ToString();
 							int32 ImageKey = FCString::Atoi(*KeyName);
 
 							if (ImageKey >= 0 && ImageKey < ModelResources.ImageProperties.Num())
 							{
 								const FMutableModelImageProperties& Props = ModelResources.ImageProperties[ImageKey];
 
-								Image.bIsNonProgressive = Props.MipGenSettings == TMGS_NoMipmaps;
+								Image->bIsNonProgressive = Props.MipGenSettings == TMGS_NoMipmaps;
 
 								if (Props.IsPassThrough)
 								{
-									Image.bIsPassThrough = true;
+									Image->bIsPassThrough = true;
 
 									// Since it's known it's a pass-through texture there is no need to cache or convert it so we can generate it here already.
-									GetImagesData->Add({ UpdateDataImageIndex, Image.ImageID });
+									GetImagesData->Add({ UpdateDataImageIndex, Image->ImageID });
 								}
 							}
 							else
@@ -2566,7 +2509,7 @@ namespace impl
 			LastTask = UE::Tasks::Launch(TEXT("Task_MutableGetMeshes_GetMesh_GetMaterial"), [SurfaceIndex, OperationData]()
 			{
 				UE::Mutable::Private::FMaterialId MaterialId = OperationData->InstanceUpdateData.Surfaces[SurfaceIndex].MaterialId;
-				UE::Tasks::TTask<TSharedPtr<const UE::Mutable::Private::FMaterial>> GetMaterialTask = OperationData->MutableSystem->GetMaterial(OperationData->InstanceID, MaterialId);
+				UE::Tasks::TTask<TSharedPtr<const UE::Mutable::Private::FMaterial>> GetMaterialTask = OperationData->MutableSystem->GetMaterial({}, OperationData->InstanceID, MaterialId);
 				
 				UE::Tasks::AddNested(UE::Tasks::Launch(TEXT("Task_MutableGetMeshes_GetMesh_GetMaterial_Post"), [=]() mutable
 				{
@@ -2624,7 +2567,7 @@ namespace impl
 		
 		const FGetMeshData& MeshData = (*GetMeshesData)[GetMeshIndex];
 		UE::Tasks::TTask<TSharedPtr<const UE::Mutable::Private::FMesh>> GetMeshTask = 
-				OperationData->MutableSystem->GetMesh(OperationData->InstanceID, MeshData.MeshID, MeshData.ContentFilter);
+				OperationData->MutableSystem->GetMesh({}, OperationData->InstanceID, MeshData.MeshID, MeshData.ContentFilter);
 		
 		UE::Tasks::AddNested(UE::Tasks::Launch(TEXT("Task_MutableGetMeshes_GetMesh_Post"), [=]()
 		{
@@ -2696,7 +2639,7 @@ namespace impl
 
 				Component.SkeletalMeshId = Instance->GetSkeletalMeshId(InstanceComponentIndex);
 
-				UE::Tasks::TTask<TSharedPtr<const UE::Mutable::Private::FSkeletalMesh>> Task = OperationData->MutableSystem->GetSkeletalMesh(OperationData->InstanceID, Component.SkeletalMeshId, UE::Mutable::Private::EMeshContentFlags::AllFlags);
+				UE::Tasks::TTask<TSharedPtr<const UE::Mutable::Private::FSkeletalMesh>> Task = OperationData->MutableSystem->GetSkeletalMesh({}, OperationData->InstanceID, Component.SkeletalMeshId, UE::Mutable::Private::EMeshContentFlags::AllFlags);
 				Component.SkeletalMesh = Task.GetResult(); // TODO GMT Remove. We can do this since now FSkeletalMesh will always be passthrough.
 
 				Component.OverlayMaterialId = Instance->GetOverlayMaterialId(InstanceComponentIndex);
@@ -2750,7 +2693,7 @@ namespace impl
 				LastTask = UE::Tasks::Launch(TEXT("Task_MutableGetMeshes_GetMesh_GetMaterialOverride"), [InstanceComponentIndex, OperationData]()
 				{
 					UE::Mutable::Private::FMaterialId MaterialId = OperationData->InstanceUpdateData.Components[InstanceComponentIndex].OverlayMaterialId;
-					UE::Tasks::TTask<TSharedPtr<const UE::Mutable::Private::FMaterial>> GetMaterialOverrideTask = OperationData->MutableSystem->GetMaterial(OperationData->InstanceID, MaterialId);
+					UE::Tasks::TTask<TSharedPtr<const UE::Mutable::Private::FMaterial>> GetMaterialOverrideTask = OperationData->MutableSystem->GetMaterial({}, OperationData->InstanceID, MaterialId);
 
 					UE::Tasks::AddNested(UE::Tasks::Launch(TEXT("Task_MutableGetMeshes_GetMesh_GetMaterialOverride_Post"), [=]() mutable
 					{
@@ -2779,241 +2722,202 @@ namespace impl
 	}
 
 	
-	/** End of the GetImages tasks. */
-	void Task_Mutable_GetImages_End(const TSharedRef<FUpdateContextPrivate>& OperationData, double StartTime)
+	void Task_Mutable_GetImageDescriptors(const TSharedRef<FUpdateContextPrivate>& OperationData)
 	{
-		MUTABLE_CPUPROFILER_SCOPE(Task_Mutable_GetImages_End)
-		
-		// TODO: Not strictly mutable: move to another worker thread task to free mutable access?
-		Subtask_Mutable_PrepareTextures(OperationData);
+		MUTABLE_CPUPROFILER_SCOPE(Task_Mutable_GetImageDescriptros)
 		
-		OperationData->TaskGetImagesTime = FPlatformTime::Seconds() - StartTime;
-
-		TRACE_END_REGION(UE_TASK_MUTABLE_GETIMAGES_REGION);
-	}
+		const TSharedPtr<TArray<UE::Mutable::Private::FImageId>> ImagesInThisInstance = MakeShared<TArray<UE::Mutable::Private::FImageId>>();
 
+		UE::Tasks::FTask LastTask;
+		for (int32 ImageIndex = 0; ImageIndex < OperationData->InstanceUpdateData.Images.Num(); ++ImageIndex)
+		{
+			TSharedRef<FInstanceUpdateData::FImage> Image = OperationData->InstanceUpdateData.Images[ImageIndex];
 
-	/** Call GetImageDesc.
-	  * Once GetImageDesc is called, the task must end. Following code will be in a subsequent TaskGraph task. */
-	void Task_Mutable_GetImages_GetImageDesc(
-		const TSharedRef<FUpdateContextPrivate>& OperationData,
-		double StartTime,
-		const TSharedPtr<TArray<UE::Mutable::Private::FImageId>>& ImagesInThisInstance,
-		int32 ImageIndex)
-	{
-		MUTABLE_CPUPROFILER_SCOPE(Task_Mutable_GetImages_GetImageDesc)
+			if (Image->bIsPassThrough)
+			{
+				continue;
+			}
 
-		FInstanceUpdateData::FImage& Image = OperationData->InstanceUpdateData.Images[ImageIndex];
-		
-		// This should only be done when using progressive images, since GetImageDesc does some actual processing.
-		UE::Tasks::TTask<UE::Mutable::Private::FExtendedImageDesc> GetImageDescTask = OperationData->MutableSystem->GetImageDesc(OperationData->InstanceID, Image.ImageID);
-		
-		UE::Tasks::AddNested(UE::Tasks::Launch(TEXT("Task_Mutable_GetImages_GetImage"), [=]()
-		{
-			Task_Mutable_GetImages_GetImage(OperationData, StartTime, ImagesInThisInstance, ImageIndex, GetImageDescTask);
-		},
-		GetImageDescTask,
-		LowLevelTasks::ETaskPriority::Inherit));
-	}
+			UE::Tasks::TTask<UE::Mutable::Private::FExtendedImageDesc> GetImageDescTask = OperationData->MutableSystem->GetImageDesc(LastTask, OperationData->InstanceID, Image->ImageID);
+			UE::Tasks::AddNested(GetImageDescTask);
+			LastTask = GetImageDescTask;
+			
+			UE::Tasks::AddNested(UE::Tasks::Launch(TEXT("GetImageDescResult"), [GetImageDescTask, Image]() mutable
+			{
+				UE::Mutable::Private::FExtendedImageDesc& ImageDesc = GetImageDescTask.GetResult();
 
+				Image->Format = ImageDesc.m_format;
+				Image->ConstantImagesNeededToGenerate = MoveTemp(ImageDesc.ConstantImagesNeededToGenerate);
 
-	/** TaskGraph task after GetImage has completed. */
-	void Task_Mutable_GetImages_GetImage_Post(
-		const TSharedRef<FUpdateContextPrivate>& OperationData,
-    	double StartTime,
-    	const TSharedPtr<TArray<UE::Mutable::Private::FImageId>>& ImagesInThisInstance,
-    	int32 ImageIndex,
-    	UE::Tasks::TTask<TSharedPtr<const UE::Mutable::Private::FImage>> GetImageTask,
-    	int32 MipSizeX,
-		int32 MipSizeY,
-		int32 FullLODContent,
-		int32 MipsToSkip)
-	{
-		MUTABLE_CPUPROFILER_SCOPE(Task_Mutable_GetImages_GetImage_Post)
+				const UCustomizableObjectSystemPrivate* CustomizableObjectSystemPrivate = UCustomizableObjectSystem::GetInstanceChecked()->GetPrivate();
+			
+				const uint16 MaxTextureSizeToGenerate = static_cast<uint16>(CustomizableObjectSystemPrivate->MaxTextureSizeToGenerate);
+				const uint16 MaxSize = FMath::Max(ImageDesc.m_size[0], ImageDesc.m_size[1]);
 
-		FInstanceUpdateData::FImage& Image = OperationData->InstanceUpdateData.Images[ImageIndex];
+				Image->BaseMip = 0;
+				if (MaxTextureSizeToGenerate > 0 && MaxSize > MaxTextureSizeToGenerate)
+				{
+					// Find the reduction factor, and the BaseMip of the texture.
+					const uint32 NextPowerOfTwo = FMath::RoundUpToPowerOfTwo(FMath::DivideAndRoundUp(MaxSize, MaxTextureSizeToGenerate));
+					uint16 Reduction = FMath::Max(NextPowerOfTwo, 2U); // At least divide the texture by a factor of two
+					Image->BaseMip = FMath::FloorLog2(Reduction);
+				}
 
-		Image.Image = GetImageTask.GetResult();
-		
-		check(Image.Image);
+				if (!CVarIgnoreFirstAvailableLODCalculation.GetValueOnAnyThread())
+				{
+					Image->BaseMip = FMath::Max<uint8>(Image->BaseMip, ImageDesc.FirstLODAvailable);
+				}
 
-		// We should have generated exactly this size.
-		const bool bSizeMissmatch = Image.Image->GetSizeX() != MipSizeX || Image.Image->GetSizeY() != MipSizeY;
-		if (bSizeMissmatch)
-		{
-			// Generate a correctly-sized but empty image instead, to avoid crashes.
-			UE_LOG(LogMutable, Warning, TEXT("Mutable generated a wrongly-sized image %s."), *Image.ImageID.ToString());
-			Image.Image = MakeShared<UE::Mutable::Private::FImage>(MipSizeX, MipSizeY, FullLODContent - MipsToSkip, Image.Image->GetFormat(), UE::Mutable::Private::EInitializationType::Black);
+				Image->FullImageSizeX = ImageDesc.m_size[0] >> Image->BaseMip;
+				Image->FullImageSizeY = ImageDesc.m_size[1] >> Image->BaseMip;
+			},
+			GetImageDescTask,
+			UE::Tasks::ETaskPriority::Inherit));
 		}
+	}
 
-		// We need one mip or the complete chain. Otherwise there was a bug.
-		const int32 FullMipCount = Image.Image->GetMipmapCount(Image.Image->GetSizeX(), Image.Image->GetSizeY());
-		const int32 RealMipCount = Image.Image->GetLODCount();
-
-		bool bForceMipchain = 
-			// Did we fail to generate the entire mipchain (if we have mips at all)?
-			(RealMipCount != 1) && (RealMipCount != FullMipCount);
+	
+	void Task_Mutable_GetImages(const TSharedRef<FUpdateContextPrivate>& OperationData)
+	{
+		MUTABLE_CPUPROFILER_SCOPE(Task_Mutable_GetImages)
 
-		if (bForceMipchain)
-		{
-			MUTABLE_CPUPROFILER_SCOPE(GetImage_MipFix);
+		const double StartTime = FPlatformTime::Seconds();		
 
-			UE_LOG(LogMutable, Warning, TEXT("Mutable generated an incomplete mip chain for image %s."), *Image.ImageID.ToString());
+		const TSharedPtr<TArray<UE::Mutable::Private::FImageId>> ImagesInThisInstance = MakeShared<TArray<UE::Mutable::Private::FImageId>>();
 
-			// Force the right number of mips. The missing data will be black.
-			const TSharedPtr<UE::Mutable::Private::FImage> NewImage = MakeShared<UE::Mutable::Private::FImage>(Image.Image->GetSizeX(), Image.Image->GetSizeY(), FullMipCount, Image.Image->GetFormat(), UE::Mutable::Private::EInitializationType::Black);
-			check(NewImage);	
-			// Formats with BytesPerBlock == 0 will not allocate memory. This type of images are not expected here.
-			check(!NewImage->DataStorage.IsEmpty());
+		UE::Tasks::FTask LastTask;
+		for (int32 ImageIndex = 0; ImageIndex < OperationData->InstanceUpdateData.Images.Num(); ++ImageIndex)
+		{
+			TSharedRef<FInstanceUpdateData::FImage> Image = OperationData->InstanceUpdateData.Images[ImageIndex];
 
-			for (int32 L = 0; L < RealMipCount; ++L)
+			if (Image->bIsPassThrough)
 			{
-				TArrayView<uint8> DestView = NewImage->DataStorage.GetLOD(L);
-				TArrayView<const uint8> SrcView = Image.Image->DataStorage.GetLOD(L);
-
-				check(DestView.Num() == SrcView.Num());
-				FMemory::Memcpy(DestView.GetData(), SrcView.GetData(), DestView.Num());
+				continue;
+			}
+			
+			const bool bCached = ImagesInThisInstance->Contains(Image->ImageID) || // See if it is cached from this same instance (can happen with LODs)
+				(UCustomizableObjectSystem::ShouldReuseTexturesBetweenInstances() && OperationData->CachedTextures.Contains(Image->ImageID)); // See if it is cached from another instance
+			
+			if (bCached)
+			{
+				UE_LOG(LogMutable, VeryVerbose, TEXT("Texture resource with id [%s] is cached."), *Image->ImageID.ToString());
+				continue;
 			}
-			Image.Image = NewImage;
-		}
-
-		ImagesInThisInstance->Add(Image.ImageID);
-
-		Task_Mutable_GetImages_Loop(OperationData, StartTime, ImagesInThisInstance, ++ImageIndex);
-	}
-
-	/** See declaration. */
-	void Task_Mutable_GetImages_GetImage(
-		const TSharedRef<FUpdateContextPrivate>& OperationData,
-		double StartTime,
-		const TSharedPtr<TArray<UE::Mutable::Private::FImageId>>& ImagesInThisInstance,
-		int32 ImageIndex,
-		UE::Tasks::TTask<UE::Mutable::Private::FExtendedImageDesc> GetImageDescTask)
-	{
-		MUTABLE_CPUPROFILER_SCOPE(Task_Mutable_GetImages_GetImage)
 
-		UE::Mutable::Private::FExtendedImageDesc& ImageDesc = GetImageDescTask.GetResult();
+			ImagesInThisInstance->Add(Image->ImageID);
+			
+			const int32 MaxSize = FMath::Max(Image->FullImageSizeX, Image->FullImageSizeY);
+			const int32 FullLODCount = FMath::CeilLogTwo(MaxSize) + 1;
+			const int32 MinMipsInImage = FMath::Min(FullLODCount, UTexture::GetStaticMinTextureResidentMipCount());
+			const int32 MaxMipsToSkip = FullLODCount - MinMipsInImage;
+			int32 MipsToSkip = FMath::Min(MaxMipsToSkip, OperationData->MipsToSkip);
 
-		FInstanceUpdateData::FImage& Image = OperationData->InstanceUpdateData.Images[ImageIndex];
-		Image.ConstantImagesNeededToGenerate = MoveTemp(ImageDesc.ConstantImagesNeededToGenerate);
+			if (Image->bIsNonProgressive || !FMath::IsPowerOfTwo(Image->FullImageSizeX) || !FMath::IsPowerOfTwo(Image->FullImageSizeY))
+			{
+				// It doesn't make sense to skip mips as non-power-of-two size textures cannot be streamed anyway
+				MipsToSkip = 0;
+			}
 
-		const UCustomizableObjectSystemPrivate* CustomizableObjectSystemPrivate = UCustomizableObjectSystem::GetInstanceChecked()->GetPrivate();
-		
-		{
-			const uint16 MaxTextureSizeToGenerate = static_cast<uint16>(CustomizableObjectSystemPrivate->MaxTextureSizeToGenerate);
-			const uint16 MaxSize = FMath::Max(ImageDesc.m_size[0], ImageDesc.m_size[1]);
+			const int32 MipSizeX = FMath::Max(Image->FullImageSizeX >> MipsToSkip, 1);
+			const int32 MipSizeY = FMath::Max(Image->FullImageSizeY >> MipsToSkip, 1);
 
-			Image.BaseMip = 0;
-			if (MaxTextureSizeToGenerate > 0 && MaxSize > MaxTextureSizeToGenerate)
+			UE::Tasks::TTask<TSharedPtr<const UE::Mutable::Private::FImage>> GetImageTask;
+			if (MipsToSkip > 0 && OperationData->bEnableSkipGenerateResidentMips != 0 && OperationData->LowPriorityTextures.Find(Image->Name.ToString()) != INDEX_NONE)
 			{
-				// Find the reduction factor, and the BaseMip of the texture.
-				const uint32 NextPowerOfTwo = FMath::RoundUpToPowerOfTwo(FMath::DivideAndRoundUp(MaxSize, MaxTextureSizeToGenerate));
-				uint16 Reduction = FMath::Max(NextPowerOfTwo, 2U); // At least divide the texture by a factor of two
-				Image.BaseMip = FMath::FloorLog2(Reduction);
+				TSharedPtr<const UE::Mutable::Private::FImage> DummyImage = MakeShared<UE::Mutable::Private::FImage>(MipSizeX, MipSizeY, FullLODCount - MipsToSkip, Image->Format, UE::Mutable::Private::EInitializationType::Black);
+				GetImageTask = UE::Tasks::MakeCompletedTask<TSharedPtr<const UE::Mutable::Private::FImage>>(DummyImage);
 			}
-
-			if (!CVarIgnoreFirstAvailableLODCalculation.GetValueOnAnyThread())
+			else
 			{
-				Image.BaseMip = FMath::Max<uint8>(Image.BaseMip, ImageDesc.FirstLODAvailable);
+				GetImageTask = OperationData->MutableSystem->GetImage(LastTask, OperationData->InstanceID, Image->ImageID, Image->BaseMip + MipsToSkip, Image->BaseLOD);
+				UE::Tasks::AddNested(GetImageTask);
+				LastTask = GetImageTask;
 			}
 
-			Image.FullImageSizeX = ImageDesc.m_size[0] >> Image.BaseMip;
-			Image.FullImageSizeY = ImageDesc.m_size[1] >> Image.BaseMip;
-		}
+			UE::Tasks::AddNested(UE::Tasks::Launch(TEXT("GetImageResult"), [Image, GetImageTask, MipSizeX, MipSizeY, FullLODCount, MipsToSkip]() mutable
+			{
+				Image->Image = GetImageTask.GetResult();
+				check(Image->Image);
 
-		const bool bCached = ImagesInThisInstance->Contains(Image.ImageID) || // See if it is cached from this same instance (can happen with LODs)
-			(UCustomizableObjectSystem::ShouldReuseTexturesBetweenInstances() && OperationData->CachedTextures.Contains(Image.ImageID)); // See if it is cached from another instance
+				// We should have generated exactly this size.
+				const bool bSizeMissmatch = Image->Image->GetSizeX() != MipSizeX || Image->Image->GetSizeY() != MipSizeY;
+				if (bSizeMissmatch)
+				{
+					// Generate a correctly-sized but empty image instead, to avoid crashes.
+					UE_LOG(LogMutable, Warning, TEXT("Mutable generated a wrongly-sized image %s."), *Image->ImageID.ToString());
+					Image->Image = MakeShared<UE::Mutable::Private::FImage>(MipSizeX, MipSizeY, FullLODCount - MipsToSkip, Image->Image->GetFormat(), UE::Mutable::Private::EInitializationType::Black);
+				}
 
-		if (bCached)
-		{
-			UE_LOG(LogMutable, VeryVerbose, TEXT("Texture resource with id [%s] is cached."), *Image.ImageID.ToString());
+				// We need one mip or the complete chain. Otherwise there was a bug.
+				const int32 FullMipCount = Image->Image->GetMipmapCount(Image->Image->GetSizeX(), Image->Image->GetSizeY());
+				const int32 RealMipCount = Image->Image->GetLODCount();
 
-			Task_Mutable_GetImages_Loop(OperationData, StartTime, ImagesInThisInstance, ++ImageIndex);
-			return;
-		}
-		
-		const int32 MaxSize = FMath::Max(Image.FullImageSizeX, Image.FullImageSizeY);
-		const int32 FullLODCount = FMath::CeilLogTwo(MaxSize) + 1;
-		const int32 MinMipsInImage = FMath::Min(FullLODCount, UTexture::GetStaticMinTextureResidentMipCount());
-		const int32 MaxMipsToSkip = FullLODCount - MinMipsInImage;
-		int32 MipsToSkip = FMath::Min(MaxMipsToSkip, OperationData->MipsToSkip);
+				bool bForceMipchain = 
+					// Did we fail to generate the entire mipchain (if we have mips at all)?
+					(RealMipCount != 1) && (RealMipCount != FullMipCount);
 
-		if (Image.bIsNonProgressive || !FMath::IsPowerOfTwo(Image.FullImageSizeX) || !FMath::IsPowerOfTwo(Image.FullImageSizeY))
-		{
-			// It doesn't make sense to skip mips as non-power-of-two size textures cannot be streamed anyway
-			MipsToSkip = 0;
-		}
+				if (bForceMipchain)
+				{
+					MUTABLE_CPUPROFILER_SCOPE(GetImage_MipFix);
 
-		const int32 MipSizeX = FMath::Max(Image.FullImageSizeX >> MipsToSkip, 1);
-		const int32 MipSizeY = FMath::Max(Image.FullImageSizeY >> MipsToSkip, 1);
-		if (MipsToSkip > 0 && CustomizableObjectSystemPrivate->EnableSkipGenerateResidentMips != 0 && OperationData->LowPriorityTextures.Find(Image.Name.ToString()) != INDEX_NONE)
-		{
-			TSharedPtr<const UE::Mutable::Private::FImage> NewImage = MakeShared<UE::Mutable::Private::FImage>(MipSizeX, MipSizeY, FullLODCount - MipsToSkip, ImageDesc.m_format, UE::Mutable::Private::EInitializationType::Black);
+					UE_LOG(LogMutable, Warning, TEXT("Mutable generated an incomplete mip chain for image %s."), *Image->ImageID.ToString());
 
-			UE::Tasks::TTask<TSharedPtr<const UE::Mutable::Private::FImage>> DummyTask = UE::Tasks::MakeCompletedTask<TSharedPtr<const UE::Mutable::Private::FImage>>(NewImage);
-			Task_Mutable_GetImages_GetImage_Post(OperationData, StartTime, ImagesInThisInstance, ImageIndex, DummyTask, MipSizeX, MipSizeY, FullLODCount, MipsToSkip);
-		}
-		else
-		{
-			const UE::Tasks::TTask<TSharedPtr<const UE::Mutable::Private::FImage>> GetImageTask = OperationData->MutableSystem->GetImage(OperationData->InstanceID, Image.ImageID, Image.BaseMip + MipsToSkip, Image.BaseLOD);
-			
-			UE::Tasks::AddNested(UE::Tasks::Launch(TEXT("Task_Mutable_GetImages_GetImage_Post"), [=]()
-			{
-				Task_Mutable_GetImages_GetImage_Post(OperationData, StartTime, ImagesInThisInstance, ImageIndex, GetImageTask, MipSizeX, MipSizeY, FullLODCount, MipsToSkip);
-			},
-			GetImageTask,
-			LowLevelTasks::ETaskPriority::Inherit));
-		}
-	}
+					// Force the right number of mips. The missing data will be black.
+					const TSharedPtr<UE::Mutable::Private::FImage> NewImage = MakeShared<UE::Mutable::Private::FImage>(Image->Image->GetSizeX(), Image->Image->GetSizeY(), FullMipCount, Image->Image->GetFormat(), UE::Mutable::Private::EInitializationType::Black);
+					check(NewImage);	
+					// Formats with BytesPerBlock == 0 will not allocate memory. This type of images are not expected here.
+					check(!NewImage->DataStorage.IsEmpty());
 
-	/** See declaration. */
-	void Task_Mutable_GetImages_Loop(
-		const TSharedRef<FUpdateContextPrivate>& OperationData,
-		double StartTime,
-		const TSharedPtr<TArray<UE::Mutable::Private::FImageId>>& ImagesInThisInstance,
-		int32 ImageIndex)
-	{
-		MUTABLE_CPUPROFILER_SCOPE(Task_Mutable_GetImages_Loop)
+					for (int32 L = 0; L < RealMipCount; ++L)
+					{
+						TArrayView<uint8> DestView = NewImage->DataStorage.GetLOD(L);
+						TArrayView<const uint8> SrcView = Image->Image->DataStorage.GetLOD(L);
 
-		// Process next image. Some images are skipped
-		for (; ImageIndex < OperationData->InstanceUpdateData.Images.Num(); ++ImageIndex)
-		{
-			const FInstanceUpdateData::FImage& Image = OperationData->InstanceUpdateData.Images[ImageIndex];
-			if (!Image.bIsPassThrough)
-			{
-				Task_Mutable_GetImages_GetImageDesc(OperationData, StartTime, ImagesInThisInstance, ImageIndex);
-				return;
-			}
+						check(DestView.Num() == SrcView.Num());
+						FMemory::Memcpy(DestView.GetData(), SrcView.GetData(), DestView.Num());
+					}
+					Image->Image = NewImage;
+				}
+			},
+			GetImageTask,
+			UE::Tasks::ETaskPriority::Inherit));
 		}
 
-		// If not image needs to be processed, go to end directly
-		Task_Mutable_GetImages_End(OperationData, StartTime);
+		OperationData->TaskGetImagesTime = FPlatformTime::Seconds() - StartTime;
 	}
 
 
-	namespace Impl
+	void Task_Mutable_PrepareTextures(const TSharedRef<FUpdateContextPrivate>& OperationData)
 	{
-		// This runs in a worker thread.
-		void Task_Mutable_GetImages(const TSharedRef<FUpdateContextPrivate>& OperationData)
-		{
-			MUTABLE_CPUPROFILER_SCOPE(Task_Mutable_GetImages)
-			TRACE_BEGIN_REGION(UE_TASK_MUTABLE_GETIMAGES_REGION);
+		MUTABLE_CPUPROFILER_SCOPE(Subtask_Mutable_PrepareTextures)
 
-			const double StartTime = FPlatformTime::Seconds();		
+		for (const FInstanceUpdateData::FSurface& Surface : OperationData->InstanceUpdateData.Surfaces)
+		{
+			for (int32 ImageIndex = 0; ImageIndex<Surface.ImageCount; ++ImageIndex)
+			{
+				const TSharedRef<FInstanceUpdateData::FImage>& Image = OperationData->InstanceUpdateData.Images[Surface.FirstImage+ImageIndex];
 
-			const TSharedPtr<TArray<UE::Mutable::Private::FImageId>> ImagesInThisInstance = MakeShared<TArray<UE::Mutable::Private::FImageId>>();
-			Task_Mutable_GetImages_Loop(OperationData, StartTime, ImagesInThisInstance, 0);
+				if (Image->Image)
+				{
+					// Image references are just references to texture assets and require no work at all
+					if (!Image->Image->IsReference())
+					{
+						if (!OperationData->ImageToPlatformDataMap.Contains(Image->ImageID))
+						{
+							FTexturePlatformData* PlatformData = MutableCreateImagePlatformData(Image->Image, -1, Image->FullImageSizeX, Image->FullImageSizeY);
+							OperationData->ImageToPlatformDataMap.Add(Image->ImageID, PlatformData);
+						}
+						else
+						{
+							// The ImageID already exists in the ImageToPlatformDataMap, that means the equivalent surface in a lower
+							// LOD already created the PlatformData for that ImageID and added it to the ImageToPlatformDataMap.
+						}
+					}
+				}
+			}
 		}
-		
-	}
-
-
-	/** Start of the GetImages tasks. */
-	void Task_Mutable_GetImages(const TSharedRef<FUpdateContextPrivate>& OperationData)
-	{
-		Impl::Task_Mutable_GetImages(OperationData);
 	}
-
+	
 
 	// This runs in a worker thread.
 	void Task_Mutable_ReleaseInstance(UE::Mutable::Private::FInstance::FID InstanceID, TSharedPtr<UE::Mutable::Private::FSystem> MutableSystem, bool bLiveUpdateMode)
@@ -3304,17 +3208,17 @@ namespace impl
 		{
 			FTextureCache& TextureCache = CustomizableObject->GetPrivate()->TextureCache;
 		
-			for (const FInstanceUpdateData::FImage& Image : OperationData->InstanceUpdateData.Images)
+			for (const TSharedRef<FInstanceUpdateData::FImage>& Image : OperationData->InstanceUpdateData.Images)
 			{
 #if	WITH_EDITOR
-				const FTextureCache::FId TextureCacheKey = FTextureCache::FId(Image.ImageID, OperationData->MipsToSkip, OperationData->bBake);
+				const FTextureCache::FId TextureCacheKey = FTextureCache::FId(Image->ImageID, OperationData->MipsToSkip, OperationData->bBake);
 #else
-				const FTextureCache::FId TextureCacheKey = FTextureCache::FId(Image.ImageID, OperationData->MipsToSkip);
+				const FTextureCache::FId TextureCacheKey = FTextureCache::FId(Image->ImageID, OperationData->MipsToSkip);
 #endif
 				if (const UTexture2D* Texture = TextureCache.Get(TextureCacheKey))
 				{
 					OperationData->Objects.Add(TStrongObjectPtr(Texture));
-					OperationData->CachedTextures.Add(Image.ImageID);
+					OperationData->CachedTextures.Add(Image->ImageID);
 				}
 			}
 		}
@@ -3325,20 +3229,29 @@ namespace impl
 		// Copy data generated in the mutable thread over to the instance
 		ObjectInstancePrivateData->PrepareForUpdate();
 
-		// Task: Mutable GetImages
-		//-------------------------------------------------------------
-		UE::Tasks::FTask Mutable_GetImagesTask;
+		UE::Tasks::FTask Mutable_GetImageDescriptors_Task = SystemPrivate->MutableTaskGraph.AddMutableThreadTask(
+		TEXT("Task_Mutable_GetImageDescriptors"),
+		[OperationData]()
 		{
-			// Task inputs
-			Mutable_GetImagesTask = SystemPrivate->MutableTaskGraph.AddMutableThreadTask(
-				TEXT("Task_Mutable_GetImages"),
-				[OperationData]()
-				{
-					Task_Mutable_GetImages(OperationData);
-				});
-		}
-
+			Task_Mutable_GetImageDescriptors(OperationData);
+		});
+		
+		UE::Tasks::FTask Mutable_GetImagesTask = SystemPrivate->MutableTaskGraph.AddMutableThreadTask(
+		TEXT("Task_Mutable_GetImages"),
+		[OperationData]()
+		{
+			Task_Mutable_GetImages(OperationData);
+		},
+		{ Mutable_GetImageDescriptors_Task });
 
+		UE::Tasks::FTask Mutable_PrepareTextures = SystemPrivate->MutableTaskGraph.AddMutableThreadTask(
+		TEXT("Task_Mutable_PrepareTextures"),
+		[OperationData]()
+		{
+			Task_Mutable_PrepareTextures(OperationData);
+		},
+		{ Mutable_GetImagesTask });
+		
 		// Next Task: Load Unreal Assets
 		//-------------------------------------------------------------
 		UE::Tasks::FTask Game_LoadUnrealAssets = LoadAdditionalAssetsAndData(OperationData);
@@ -3352,7 +3265,7 @@ namespace impl
 				Task_Game_ConvertResources(OperationData);
 			},
 			false,
-			{ Game_LoadUnrealAssets, Mutable_GetImagesTask });
+			{ Game_LoadUnrealAssets, Mutable_PrepareTextures });
 	}
 
 
@@ -3501,6 +3414,7 @@ namespace impl
 
 		Operation->bNeverStream = false;
 		Operation->MipsToSkip = 0;
+		Operation->bEnableSkipGenerateResidentMips = static_cast<bool>(SystemPrivate->EnableSkipGenerateResidentMips);
 
 		SystemPrivate->GetMipStreamingConfig(*Instance, Operation->bProgressiveMipStreamingEnabled, Operation->bNeverStream, Operation->MipsToSkip);
 		
diff --git a/Engine/Plugins/Mutable/Source/MutableRuntime/Internal/MuR/System.h b/Engine/Plugins/Mutable/Source/MutableRuntime/Internal/MuR/System.h
--- a/Engine/Plugins/Mutable/Source/MutableRuntime/Internal/MuR/System.h
+++ b/Engine/Plugins/Mutable/Source/MutableRuntime/Internal/MuR/System.h
@@ -175,18 +175,18 @@ namespace UE::Mutable::Private
 		
 		//! Only valid between BeginUpdate and EndUpdate
 		//! Calculate the description of an image, without generating it.
-		UE_API UE::Tasks::TTask<FExtendedImageDesc> GetImageDesc(FInstance::FID InstanceID, const FImageId& ImageId);
+		UE_API UE::Tasks::TTask<FExtendedImageDesc> GetImageDesc(UE::Tasks::FTask Prerequisite, FInstance::FID InstanceID, const FImageId& ImageId);
 
 		//! Only valid between BeginUpdate and EndUpdate
 		//! \param MipsToSkip Number of mips to skip compared from the full image.
 		//! If 0, all mip levels will be generated. If more levels than possible to discard are specified, 
 		//! the image will still contain a minimum number of mips specified at model compile time.
-		UE_API UE::Tasks::TTask<TSharedPtr<const FImage>> GetImage(FInstance::FID InstanceID, const FImageId& ImageId, int32 MipsToSkip = 0, int32 LOD = 0);
+		UE_API UE::Tasks::TTask<TSharedPtr<const FImage>> GetImage(UE::Tasks::FTask Prerequisite, FInstance::FID InstanceID, const FImageId& ImageId, int32 MipsToSkip = 0, int32 LOD = 0);
 
 		//! Only valid between BeginUpdate and EndUpdate
-		UE_API UE::Tasks::TTask<TSharedPtr<const FMesh>> GetMesh(FInstance::FID InstanceID, const FMeshId& MeshId, EMeshContentFlags MeshContentFilter = EMeshContentFlags::AllFlags);
+		UE_API UE::Tasks::TTask<TSharedPtr<const FMesh>> GetMesh(UE::Tasks::FTask Prerequisite, FInstance::FID InstanceID, const FMeshId& MeshId, EMeshContentFlags MeshContentFilter = EMeshContentFlags::AllFlags);
 		
-		UE_API UE::Tasks::TTask<TSharedPtr<const FSkeletalMesh>> GetSkeletalMesh(FInstance::FID InstanceID, const FSkeletalMeshId& SkeletalMeshId, EMeshContentFlags MeshContentFilter);
+		UE_API UE::Tasks::TTask<TSharedPtr<const FSkeletalMesh>> GetSkeletalMesh(UE::Tasks::FTask Prerequisite, FInstance::FID InstanceID, const FSkeletalMeshId& SkeletalMeshId, EMeshContentFlags MeshContentFilter);
 
 		//! Only valid between BeginUpdate and EndUpdate
 		//! Calculate the description of an image, without generating it.
@@ -201,7 +201,7 @@ namespace UE::Mutable::Private
         //! Only valid between BeginUpdate and EndUpdate
 		UE_API TSharedPtr<const FMesh> GetMeshInline(FInstance::FID InstanceID, const FMeshId& MeshId, EMeshContentFlags MeshContentFilter = EMeshContentFlags::AllFlags);
 		
-		UE_API UE::Tasks::TTask<TSharedPtr<const FMaterial>> GetMaterial(FInstance::FID InstanceID, const FMaterialId& MaterialId);
+		UE_API UE::Tasks::TTask<TSharedPtr<const FMaterial>> GetMaterial(UE::Tasks::FTask Prerequisite, FInstance::FID InstanceID, const FMaterialId& MaterialId);
 		
         //! Invalidate and free the last Instance data returned by a call to BeginUpdate with
         //! the same instance index. After a call to this method, that Instance cannot be used any
diff --git a/Engine/Plugins/Mutable/Source/MutableRuntime/Private/MuR/CodeRunner.cpp b/Engine/Plugins/Mutable/Source/MutableRuntime/Private/MuR/CodeRunner.cpp
--- a/Engine/Plugins/Mutable/Source/MutableRuntime/Private/MuR/CodeRunner.cpp
+++ b/Engine/Plugins/Mutable/Source/MutableRuntime/Private/MuR/CodeRunner.cpp
@@ -180,7 +180,6 @@ namespace UE::Mutable::Private
 		FScheduledOp::EType Type,
 		const TSharedPtr<FExternalResourceProvider>& ExternalResourceProvider)
 		: Settings(InSettings)
-		, RunnerCompletionEvent(TEXT("CodeRunnerCompletioneEventInit"))
 		, ExecutionStrategy(InExecutionStrategy)
 		, System(InSystem)
 		, Model(InModel)
@@ -204,11 +203,7 @@ namespace UE::Mutable::Private
    		{
 			GetMemory().IncreaseHitCount(FCacheAddress(At, 0, ExecutionOptions));
 		}
-    
-		// Start with a completed Event. This is checked at StartRun() to make sure StartRun is not called while there is 
-		// a Run in progress.
-		RunnerCompletionEvent.Trigger();
-
+		
 		ImageLOD = InImageLOD;
 	
 		// Push the root operation
diff --git a/Engine/Plugins/Mutable/Source/MutableRuntime/Private/MuR/CodeRunner.h b/Engine/Plugins/Mutable/Source/MutableRuntime/Private/MuR/CodeRunner.h
--- a/Engine/Plugins/Mutable/Source/MutableRuntime/Private/MuR/CodeRunner.h
+++ b/Engine/Plugins/Mutable/Source/MutableRuntime/Private/MuR/CodeRunner.h
@@ -177,14 +177,12 @@ namespace  UE::Mutable::Private
 
 		/** Only used for correct mip skipping with external images. It is the LOD for which the image is build. */
 		int32 ImageLOD;
-
-		UE::Tasks::FTaskEvent RunnerCompletionEvent;
-
+    	
 		void Run(TUniquePtr<FProfileContext>&& ProfileContext, bool bForceInlineExecution);
 		void AbortRun();
 	public:
 
-		UE::Tasks::FTask StartRun(bool bForceInlineExecution);
+		void StartRun(bool bForceInlineExecution);
 
 		/** */
 		const FExtendedImageDesc& GetImageDescResult(OP::ADDRESS ResultAddres);
diff --git a/Engine/Plugins/Mutable/Source/MutableRuntime/Private/MuR/CodeRunnerTasked.cpp b/Engine/Plugins/Mutable/Source/MutableRuntime/Private/MuR/CodeRunnerTasked.cpp
--- a/Engine/Plugins/Mutable/Source/MutableRuntime/Private/MuR/CodeRunnerTasked.cpp
+++ b/Engine/Plugins/Mutable/Source/MutableRuntime/Private/MuR/CodeRunnerTasked.cpp
@@ -150,10 +150,8 @@ namespace UE::Mutable::Private
 		IssuedTasks.Add(TaskToIssue);
 	}
 
-	UE::Tasks::FTask CodeRunner::StartRun(bool bForceInlineExecution)
+	void CodeRunner::StartRun(bool bForceInlineExecution)
 	{
-		check(RunnerCompletionEvent.IsCompleted());
-
 		if (CVarCodeRunnerForceInline.GetValueOnAnyThread())
 		{
 			bForceInlineExecution = true;
@@ -164,31 +162,22 @@ namespace UE::Mutable::Private
 		HeapData.SetNum(0, EAllowShrinking::No);
 		ImageDescResults.Reset();
 		ImageDescConstantImages.Reset();
-
-		RunnerCompletionEvent = UE::Tasks::FTaskEvent(TEXT("CodeRunnerCompletionEvent"));
 		
 		bool bProfile = false;
 		
 		TUniquePtr<FProfileContext> ProfileContext = bProfile ? MakeUnique<FProfileContext>() : nullptr;
 		Run(MoveTemp(ProfileContext), bForceInlineExecution);
-
-		check(!bForceInlineExecution || RunnerCompletionEvent.IsCompleted());
-
-		return RunnerCompletionEvent;
 	}
 
 	void CodeRunner::AbortRun()
 	{
 		bUnrecoverableError = true;
-		RunnerCompletionEvent.Trigger();
 	}
 
     void CodeRunner::Run(TUniquePtr<FProfileContext>&& ProfileContext, bool bForceInlineExecution)
     {
 		MUTABLE_CPUPROFILER_SCOPE(CodeRunner_Run);
 
-		check(!RunnerCompletionEvent.IsCompleted());
-
 		// TODO: Move MaxAllowedTime somewhere else more accessible, maybe a cvar.
 		const FTimespan MaxAllowedTime = FTimespan::FromMilliseconds(2.0); 
 		const FTimespan TimeOut = FTimespan::FromSeconds(FPlatformTime::Seconds()) + MaxAllowedTime;
@@ -416,7 +405,7 @@ namespace UE::Mutable::Private
 
 					System->WorkingMemoryManager.InvalidateRunnerThread();
 
-					UE::Tasks::Launch(TEXT("CodeRunnerFromIssuedTasksTask"),
+					UE::Tasks::AddNested(UE::Tasks::Launch(TEXT("CodeRunnerFromIssuedTasksTask"),
 						[Runner = AsShared(), ProfileContext = MoveTemp(ProfileContext)]() mutable
 						{
 							Runner->System->WorkingMemoryManager.ResetRunnerThread();
@@ -425,7 +414,7 @@ namespace UE::Mutable::Private
 							Runner->Run(MoveTemp(ProfileContext), bForceInlineExecution);
 						},
 						UE::Tasks::Prerequisites(UE::Tasks::Any(IssuedTasksCompletionEvents)),
-						UE::Tasks::ETaskPriority::Inherit);
+						UE::Tasks::ETaskPriority::Inherit));
 					
 					return;
 				}	
@@ -450,7 +439,7 @@ namespace UE::Mutable::Private
 				{
 					System->WorkingMemoryManager.InvalidateRunnerThread();
 
-					UE::Tasks::Launch(TEXT("CodeRunnerFromTimeoutTask"),
+					UE::Tasks::AddNested(UE::Tasks::Launch(TEXT("CodeRunnerFromTimeoutTask"),
 						[Runner = AsShared(), ProfileContext = MoveTemp(ProfileContext)]() mutable
 						{
 							Runner->System->WorkingMemoryManager.ResetRunnerThread();
@@ -458,7 +447,7 @@ namespace UE::Mutable::Private
 							constexpr bool bForceInlineExecution = false;
 							Runner->Run(MoveTemp(ProfileContext), bForceInlineExecution);
 						},
-						UE::Tasks::ETaskPriority::Inherit);
+						UE::Tasks::ETaskPriority::Inherit));
 					
 					return;
 				}
@@ -496,8 +485,6 @@ namespace UE::Mutable::Private
 				UE_LOG(LogMutableCore, Log, TEXT("    op %4d, %4d times."), MostCommonOps[HistIndex], ProfileContext->RunOpsPerType[MostCommonOps[HistIndex]]);
 			}
 		}
-
-		RunnerCompletionEvent.Trigger();
     }
 
 	
diff --git a/Engine/Plugins/Mutable/Source/MutableRuntime/Private/MuR/System.cpp b/Engine/Plugins/Mutable/Source/MutableRuntime/Private/MuR/System.cpp
--- a/Engine/Plugins/Mutable/Source/MutableRuntime/Private/MuR/System.cpp
+++ b/Engine/Plugins/Mutable/Source/MutableRuntime/Private/MuR/System.cpp
@@ -433,39 +433,58 @@ namespace UE::Mutable::Private
 		return pResult;
 	}
 
-	UE::Tasks::TTask<TSharedPtr<const FImage>> FSystem::GetImage(FInstance::FID instanceID, const FImageId& ImageId, int32 MipsToSkip, int32 InImageLOD)
+	UE::Tasks::TTask<TSharedPtr<const FImage>> FSystem::GetImage(UE::Tasks::FTask Prerequisite, FInstance::FID instanceID, const FImageId& ImageId, int32 MipsToSkip, int32 InImageLOD)
 	{
 		LLM_SCOPE_BYNAME(TEXT("MutableRuntime"));
 		MUTABLE_CPUPROFILER_SCOPE(SystemGetImage);
 
     	if (!ImageId)
     	{
-			return UE::Tasks::MakeCompletedTask<TSharedPtr<const FImage>>(BlackImage());
+    		return UE::Tasks::Launch(TEXT("FSystem::GetImageInvalid"),
+			[]() -> TSharedPtr<const FImage>
+			{
+				return MakeShared<FImage>();
+			},
+			UE::Tasks::Prerequisites(Prerequisite),
+			UE::Tasks::ETaskPriority::Inherit,
+			UE::Tasks::EExtendedTaskPriority::Inline);
     	}
-    	
+
+    	OP::ADDRESS RootAddress = ImageId.GetKey()->Address;
+
 		// Find the live instance
 		FLiveInstance* LiveInstance = m_pD->FindLiveInstance(instanceID);
 		check(LiveInstance);
-		m_pD->WorkingMemoryManager.CurrentInstanceCache = LiveInstance->Cache;
-		m_pD->WorkingMemoryManager.CurrentPassthroughObjectLoader = LiveInstance->PassthroughObjectLoader;
 
-		OP::ADDRESS RootAddress = ImageId.GetKey()->Address;
-		
-		m_pD->WorkingMemoryManager.BeginRunnerThread();
-		
-		UE::Mutable::Private::EOpType OpType = LiveInstance->Model->GetPrivate()->Program.GetOpType(RootAddress);
-		if (GetOpDataType(OpType) != EDataType::Image)
+    	EOpType OpType = LiveInstance->Model->GetPrivate()->Program.GetOpType(RootAddress);
+    	if (GetOpDataType(OpType) != EDataType::Image)
+    	{
+    		return UE::Tasks::Launch(TEXT("FSystem::GetImageInvalid"),
+			[]() -> TSharedPtr<const FImage>
+			{
+				return MakeShared<FImage>();
+			},
+			UE::Tasks::Prerequisites(Prerequisite),
+			UE::Tasks::ETaskPriority::Inherit,
+			UE::Tasks::EExtendedTaskPriority::Inline);
+    	}
+
+    	struct FContext
+    	{
+    		TSharedPtr<CodeRunner> Runner;
+    	};
+    	TSharedRef<FContext> Context = MakeShared<FContext>();
+
+    	UE::Tasks::FTask RunnerTask = UE::Tasks::Launch(TEXT("FSystem::GetImage"),
+		[Context, SystemPrivate = m_pD, RootAddress, LiveInstance, MipsToSkip, InImageLOD]()
 		{
-			m_pD->WorkingMemoryManager.EndRunnerThread();
-			m_pD->WorkingMemoryManager.CurrentInstanceCache = nullptr;
-			m_pD->WorkingMemoryManager.CurrentPassthroughObjectLoader = nullptr;
-			
-			return UE::Tasks::MakeCompletedTask<TSharedPtr<const FImage>>(BlackImage());
-		}
-		
-		TSharedRef<CodeRunner> Runner = CodeRunner::Create(
-				m_pD->Settings,
-				m_pD,
+			SystemPrivate->WorkingMemoryManager.CurrentInstanceCache = LiveInstance->Cache;
+			SystemPrivate->WorkingMemoryManager.CurrentPassthroughObjectLoader = LiveInstance->PassthroughObjectLoader;
+			SystemPrivate->WorkingMemoryManager.BeginRunnerThread();
+
+			Context->Runner = CodeRunner::Create(
+				SystemPrivate->Settings,
+				SystemPrivate,
 				EExecutionStrategy::MinimizeMemory,
 				LiveInstance->Model,
 				LiveInstance->OldParameters.Get(),
@@ -473,36 +492,38 @@ namespace UE::Mutable::Private
 				MipsToSkip, InImageLOD,
 				FScheduledOp::EType::Full,
 				LiveInstance->ExternalResourceProvider);
-	
-		
-		constexpr bool bForceInlineExecution = false;
-		UE::Tasks::FTask RunnerCompletionEvent = Runner->StartRun(bForceInlineExecution);
-
+				
+			constexpr bool bForceInlineExecution = false;
+			Context->Runner->StartRun(bForceInlineExecution);
+		},
+		UE::Tasks::Prerequisites(Prerequisite),
+		UE::Tasks::ETaskPriority::Inherit);
+    	
 		return UE::Tasks::Launch(TEXT("FSystem::GetImageResultTask"),
-				[SystemPrivate = m_pD, Runner, RootAddress, MipsToSkip]() -> TSharedPtr<const FImage>
-				{
-					TSharedPtr<const FImage> Result;
+		[Context, SystemPrivate = m_pD, RootAddress, MipsToSkip]() -> TSharedPtr<const FImage>
+		{
+			TSharedPtr<const FImage> Result;
 
-					SystemPrivate->bUnrecoverableError = Runner->bUnrecoverableError;
-					if (!Runner->bUnrecoverableError)
-					{
-						Result = SystemPrivate->WorkingMemoryManager.LoadImage(FCacheAddress(RootAddress, 0, MipsToSkip), true);
-					}
+			SystemPrivate->bUnrecoverableError = Context->Runner->bUnrecoverableError;
+			if (!Context->Runner->bUnrecoverableError)
+			{
+				Result = SystemPrivate->WorkingMemoryManager.LoadImage(FCacheAddress(RootAddress, 0, MipsToSkip), true);
+			}
 
-					if (!Result)
-					{
-						Result = MakeShared<UE::Mutable::Private::FImage>(16, 16, 1, EImageFormat::RGBA_UByte, EInitializationType::Black);
-					}
+			if (!Result)
+			{
+				Result = MakeShared<FImage>(16, 16, 1, EImageFormat::RGBA_UByte, EInitializationType::Black);
+			}
 
-					SystemPrivate->WorkingMemoryManager.EndRunnerThread();
-					SystemPrivate->WorkingMemoryManager.CurrentInstanceCache = nullptr;
-					SystemPrivate->WorkingMemoryManager.CurrentPassthroughObjectLoader = nullptr;
-					
-					return Result;
-				},
-				UE::Tasks::Prerequisites(RunnerCompletionEvent),
-				UE::Tasks::ETaskPriority::Inherit,
-				UE::Tasks::EExtendedTaskPriority::Inline);
+			SystemPrivate->WorkingMemoryManager.EndRunnerThread();
+			SystemPrivate->WorkingMemoryManager.CurrentInstanceCache = nullptr;
+			SystemPrivate->WorkingMemoryManager.CurrentPassthroughObjectLoader = nullptr;
+			
+			return Result;
+		},
+		UE::Tasks::Prerequisites(RunnerTask),
+		UE::Tasks::ETaskPriority::Inherit,
+		UE::Tasks::EExtendedTaskPriority::Inline);
 	}
 
 
@@ -565,8 +586,7 @@ namespace UE::Mutable::Private
 					LiveInstance->ExternalResourceProvider);
 			
 			constexpr bool bForceInlineExecution = true;
-			UE::Tasks::FTask CompletionEvent = Runner->StartRun(bForceInlineExecution);
-			check(CompletionEvent.IsCompleted());
+			Runner->StartRun(bForceInlineExecution);
 
 			Result = Runner->GetImageDescResult(RootAddress);
 
@@ -580,50 +600,67 @@ namespace UE::Mutable::Private
 	}
 
 
-	UE::Tasks::TTask<FExtendedImageDesc> FSystem::GetImageDesc(FInstance::FID instanceID, const FImageId& ImageId)
+	UE::Tasks::TTask<FExtendedImageDesc> FSystem::GetImageDesc(UE::Tasks::FTask Prerequisite, FInstance::FID InstanceID, const FImageId& ImageId)
 	{
 		LLM_SCOPE_BYNAME(TEXT("MutableRuntime"));
 		MUTABLE_CPUPROFILER_SCOPE(SystemGetImageDesc);
 
 		if (!ImageId)
 		{
-			return UE::Tasks::MakeCompletedTask<FExtendedImageDesc>();
+			return UE::Tasks::Launch(TEXT("FSystem::GetImageDescInvalid"),
+			[]()
+			{
+				return FExtendedImageDesc();
+			},
+			UE::Tasks::Prerequisites(Prerequisite),
+			UE::Tasks::ETaskPriority::Inherit,
+			UE::Tasks::EExtendedTaskPriority::Inline);
 		}
 		
-		// Find the live instance
-		FLiveInstance* LiveInstance = m_pD->FindLiveInstance(instanceID);
-		check(LiveInstance);
-		m_pD->WorkingMemoryManager.CurrentInstanceCache = LiveInstance->Cache;
-		m_pD->WorkingMemoryManager.CurrentPassthroughObjectLoader = LiveInstance->PassthroughObjectLoader;
-
 		OP::ADDRESS RootAddress = ImageId.GetKey()->Address;
+		
+		FLiveInstance* LiveInstance = m_pD->FindLiveInstance(InstanceID);
+		check(LiveInstance);
 
-		const UE::Mutable::Private::FModel* Model = LiveInstance->Model.Get();
-		const UE::Mutable::Private::FProgram& Program = Model->GetPrivate()->Program;
-
-		// TODO: It should be possible to reuse this data if cleared in the correct places only, together with HeapImageDesc.
-		int32 VarValue = CVarClearImageDescCache.GetValueOnAnyThread();
-		if (VarValue != 0)
+		EOpType OpType = LiveInstance->Model->GetPrivate()->Program.GetOpType(RootAddress);
+		if (GetOpDataType(OpType) != EDataType::Image)
 		{
-			m_pD->WorkingMemoryManager.CurrentInstanceCache->ClearDescCache();
+			return UE::Tasks::Launch(TEXT("FSystem::GetImageDescInvalid"),
+			[]()
+			{
+				return FExtendedImageDesc();
+			},
+			UE::Tasks::Prerequisites(Prerequisite),
+			UE::Tasks::ETaskPriority::Inherit,
+			UE::Tasks::EExtendedTaskPriority::Inline);
 		}
 
-		m_pD->WorkingMemoryManager.BeginRunnerThread();
-		UE::Mutable::Private::EOpType OpType = Program.GetOpType(RootAddress);
-		if (GetOpDataType(OpType) != EDataType::Image)
+		struct FContext
 		{
-			m_pD->WorkingMemoryManager.EndRunnerThread();
-			m_pD->WorkingMemoryManager.CurrentInstanceCache = nullptr;
-			m_pD->WorkingMemoryManager.CurrentPassthroughObjectLoader = nullptr;
+			TSharedPtr<CodeRunner> Runner;
+		};
+		TSharedRef<FContext> Context = MakeShared<FContext>();
+		
+		UE::Tasks::FTask RunnerTask = UE::Tasks::Launch(TEXT("FSystem::GetImageDesc"),
+		[Context, SystemPrivate = m_pD, LiveInstance, RootAddress]()
+		{
+			SystemPrivate->WorkingMemoryManager.CurrentInstanceCache = LiveInstance->Cache;
+			SystemPrivate->WorkingMemoryManager.CurrentPassthroughObjectLoader = LiveInstance->PassthroughObjectLoader;
+				
+			// TODO: It should be possible to reuse this data if cleared in the correct places only, together with HeapImageDesc.
+			int32 VarValue = CVarClearImageDescCache.GetValueOnAnyThread();
+			if (VarValue != 0)
+			{
+				SystemPrivate->WorkingMemoryManager.CurrentInstanceCache->ClearDescCache();
+			}
 
-			return UE::Tasks::MakeCompletedTask<FExtendedImageDesc>(); 
-		}
+			SystemPrivate->WorkingMemoryManager.BeginRunnerThread();
 
-		// GetImageDesc may call normal execution paths where meshes are computed.
-		int8 ExecutionOptions = 0;
-		TSharedRef<CodeRunner> Runner = CodeRunner::Create(
-				m_pD->Settings,
-				m_pD,
+			// GetImageDesc may call normal execution paths where meshes are computed.
+			int8 ExecutionOptions = 0;
+			Context->Runner = CodeRunner::Create(
+				SystemPrivate->Settings,
+				SystemPrivate,
 				EExecutionStrategy::MinimizeMemory,
 				LiveInstance->Model,
 				LiveInstance->OldParameters.Get(),
@@ -633,24 +670,27 @@ namespace UE::Mutable::Private
 				0,
 				FScheduledOp::EType::ImageDesc,
 				LiveInstance->ExternalResourceProvider);
-		
-		constexpr bool bForceInlineExecution = false;
-		UE::Tasks::FTask RunnerCompletionEvent = Runner->StartRun(bForceInlineExecution);
+
+			constexpr bool bForceInlineExecution = false;
+			Context->Runner->StartRun(bForceInlineExecution);
+		}, 
+		UE::Tasks::Prerequisites(Prerequisite),
+		UE::Tasks::ETaskPriority::Inherit);
 		
 		return UE::Tasks::Launch(TEXT("FSystem::GetImageDescResultTask"),
-				[SystemPrivate = m_pD, Runner, RootAddress]() -> FExtendedImageDesc
-				{
-					FExtendedImageDesc Result = Runner->GetImageDescResult(RootAddress);
+		[Context, SystemPrivate = m_pD, RootAddress]() -> FExtendedImageDesc
+		{
+			FExtendedImageDesc Result = Context->Runner->GetImageDescResult(RootAddress);
 
-					SystemPrivate->WorkingMemoryManager.EndRunnerThread();
-					SystemPrivate->WorkingMemoryManager.CurrentInstanceCache = nullptr;
-					SystemPrivate->WorkingMemoryManager.CurrentPassthroughObjectLoader = nullptr;
+			SystemPrivate->WorkingMemoryManager.EndRunnerThread();
+			SystemPrivate->WorkingMemoryManager.CurrentInstanceCache = nullptr;
+			SystemPrivate->WorkingMemoryManager.CurrentPassthroughObjectLoader = nullptr;
 
-					return Result;
-				},
-				UE::Tasks::Prerequisites(RunnerCompletionEvent),
-				UE::Tasks::ETaskPriority::Inherit,
-				UE::Tasks::EExtendedTaskPriority::Inline);
+			return Result;
+		},
+		UE::Tasks::Prerequisites(RunnerTask),
+		UE::Tasks::ETaskPriority::Inherit,
+		UE::Tasks::EExtendedTaskPriority::Inline);
 
 	}	
 
@@ -664,8 +704,6 @@ namespace UE::Mutable::Private
 		{
 			return MakeShared<FMesh>();
 		}
-		
-		TSharedPtr<const FMesh> Result;
 
 		// Find the live instance
 		FLiveInstance* LiveInstance = m_pD->FindLiveInstance(instanceID);
@@ -674,12 +712,12 @@ namespace UE::Mutable::Private
 		m_pD->WorkingMemoryManager.CurrentPassthroughObjectLoader = LiveInstance->PassthroughObjectLoader;
 
 		OP::ADDRESS RootAddress = MeshId.GetKey()->Address;
-		Result = m_pD->BuildMesh(
-				LiveInstance->ExternalResourceProvider,
-				LiveInstance->Model, 
-				LiveInstance->OldParameters.Get(), 
-				RootAddress, 
-				MeshContentFilter);
+		TSharedPtr<const FMesh> Result = m_pD->BuildMesh(
+			LiveInstance->ExternalResourceProvider,
+			LiveInstance->Model,
+			LiveInstance->OldParameters.Get(),
+			RootAddress,
+			MeshContentFilter);
 
 		// If the mesh is null it means empty, but we still need to return a valid one
 		if (!Result)
@@ -694,209 +732,277 @@ namespace UE::Mutable::Private
 	}
 	
 
-    UE::Tasks::TTask<TSharedPtr<const FMesh>> FSystem::GetMesh(FInstance::FID instanceID, const FMeshId& MeshId, EMeshContentFlags MeshContentFilter)
-    {
+    UE::Tasks::TTask<TSharedPtr<const FMesh>> FSystem::GetMesh(UE::Tasks::FTask Prerequisite, FInstance::FID instanceID, const FMeshId& MeshId, EMeshContentFlags MeshContentFilter)
+	{
 		LLM_SCOPE_BYNAME(TEXT("MutableRuntime"));
 		MUTABLE_CPUPROFILER_SCOPE(SystemGetMesh);
 		
 		if (!MeshId)
 		{
-			return UE::Tasks::MakeCompletedTask<TSharedPtr<const FMesh>>(new FMesh());
+			return UE::Tasks::Launch(TEXT("FSystem::GetMeshInvalid"),
+			[]() -> TSharedPtr<const FMesh>
+			{
+				return MakeShared<FMesh>();
+			},
+			UE::Tasks::Prerequisites(Prerequisite),
+			UE::Tasks::ETaskPriority::Inherit,
+			UE::Tasks::EExtendedTaskPriority::Inline);
 		}
-		
+
+		OP::ADDRESS RootAddress = MeshId.GetKey()->Address;
+
 		// Find the live instance
 		FLiveInstance* LiveInstance = m_pD->FindLiveInstance(instanceID);
 		check(LiveInstance);
-		m_pD->WorkingMemoryManager.CurrentInstanceCache = LiveInstance->Cache;
-		m_pD->WorkingMemoryManager.CurrentPassthroughObjectLoader = LiveInstance->PassthroughObjectLoader;
-		m_pD->WorkingMemoryManager.BeginRunnerThread();
 		
-		OP::ADDRESS RootAddress = MeshId.GetKey()->Address;
-			
-		UE::Mutable::Private::EOpType OpType = LiveInstance->Model->GetPrivate()->Program.GetOpType(RootAddress);
+		EOpType OpType = LiveInstance->Model->GetPrivate()->Program.GetOpType(RootAddress);
 		if (GetOpDataType(OpType) != EDataType::Mesh)
 		{
-			m_pD->WorkingMemoryManager.EndRunnerThread();
-			m_pD->WorkingMemoryManager.CurrentInstanceCache = nullptr;
-			m_pD->WorkingMemoryManager.CurrentPassthroughObjectLoader = nullptr;
-
-			return UE::Tasks::MakeCompletedTask<TSharedPtr<const FMesh>>(new FMesh());
+			return UE::Tasks::Launch(TEXT("FSystem::GetMeshInvalid"),
+			[]() -> TSharedPtr<const FMesh>
+			{
+				return MakeShared<FMesh>();
+			},
+			UE::Tasks::Prerequisites(Prerequisite),
+			UE::Tasks::ETaskPriority::Inherit,
+			UE::Tasks::EExtendedTaskPriority::Inline);
 		}
+
+		struct FContext
+		{
+			TSharedPtr<CodeRunner> Runner;
+		};
+		TSharedRef<FContext> Context = MakeShared<FContext>();
+		
+		UE::Tasks::FTask RunnerTask = UE::Tasks::Launch(TEXT("FSystem::GetMesh"),
+		[Context, SystemPrivate = m_pD, RootAddress, LiveInstance, MeshContentFilter]()
+		{
+			SystemPrivate->WorkingMemoryManager.CurrentInstanceCache = LiveInstance->Cache;
+			SystemPrivate->WorkingMemoryManager.CurrentPassthroughObjectLoader = LiveInstance->PassthroughObjectLoader;
+			SystemPrivate->WorkingMemoryManager.BeginRunnerThread();
+			
+			uint8 ExecutionOptions = static_cast<uint8>(MeshContentFilter);
+			Context->Runner = CodeRunner::Create(
+					SystemPrivate->Settings,
+					SystemPrivate, 
+					EExecutionStrategy::MinimizeMemory, 
+					LiveInstance->Model, 
+					LiveInstance->OldParameters.Get(), 
+					RootAddress, 
+					FSystem::AllLODs, 
+					ExecutionOptions, 
+					0, 
+					FScheduledOp::EType::Full,
+					LiveInstance->ExternalResourceProvider);
 	
-		uint8 ExecutionOptions = static_cast<uint8>(MeshContentFilter);
-		TSharedRef<CodeRunner> Runner = CodeRunner::Create(
-				m_pD->Settings,
-				m_pD, 
-				EExecutionStrategy::MinimizeMemory, 
-				LiveInstance->Model, 
-				LiveInstance->OldParameters.Get(), 
-				RootAddress, 
-				FSystem::AllLODs, 
-				ExecutionOptions, 
-				0, 
-				FScheduledOp::EType::Full,
-				LiveInstance->ExternalResourceProvider);
+			constexpr bool bForceInlineExecution = false;
+			Context->Runner->StartRun(bForceInlineExecution);
+		},
+		UE::Tasks::Prerequisites(Prerequisite),
+		UE::Tasks::ETaskPriority::Inherit);
 		
-		constexpr bool bForceInlineExecution = false;
-		UE::Tasks::FTask RunnerCompletionEvent = Runner->StartRun(bForceInlineExecution);
-
 		return UE::Tasks::Launch(TEXT("FSystem::GetMeshResultTask"),
-				[SystemPrivate = m_pD, Runner, RootAddress, ExecutionOptions]() -> TSharedPtr<const FMesh>
-				{
-					TSharedPtr<const FMesh> Result;
+		[Context, SystemPrivate = m_pD, RootAddress, MeshContentFilter]() -> TSharedPtr<const FMesh>
+		{
+			TSharedPtr<const FMesh> Result;
 
-					SystemPrivate->bUnrecoverableError = Runner->bUnrecoverableError;
-					if (!Runner->bUnrecoverableError)
-					{
-						Result = SystemPrivate->WorkingMemoryManager.LoadMesh(FCacheAddress(RootAddress, 0, ExecutionOptions), true);
-					}
-					
-						SystemPrivate->WorkingMemoryManager.EndRunnerThread();
-						SystemPrivate->WorkingMemoryManager.CurrentInstanceCache = nullptr;
-						SystemPrivate->WorkingMemoryManager.CurrentPassthroughObjectLoader = nullptr;
-					
-						return Result;
-					},
-					UE::Tasks::Prerequisites(RunnerCompletionEvent),
-					UE::Tasks::ETaskPriority::Inherit,
-					UE::Tasks::EExtendedTaskPriority::Inline);
-    }
+			SystemPrivate->bUnrecoverableError = Context->Runner->bUnrecoverableError;
+			if (!Context->Runner->bUnrecoverableError)
+			{
+				Result = SystemPrivate->WorkingMemoryManager.LoadMesh(FCacheAddress(RootAddress, 0, static_cast<uint8>(MeshContentFilter)), true);
+			}
+			
+			SystemPrivate->WorkingMemoryManager.EndRunnerThread();
+			SystemPrivate->WorkingMemoryManager.CurrentInstanceCache = nullptr;
+			SystemPrivate->WorkingMemoryManager.CurrentPassthroughObjectLoader = nullptr;
+			
+			return Result;
+		},
+		UE::Tasks::Prerequisites(RunnerTask),
+		UE::Tasks::ETaskPriority::Inherit,
+		UE::Tasks::EExtendedTaskPriority::Inline);
+	}
 
 
-	UE::Tasks::TTask<TSharedPtr<const FSkeletalMesh>> FSystem::GetSkeletalMesh(FInstance::FID InstanceID, const FSkeletalMeshId& SkeletalMeshId, EMeshContentFlags SkeletalMeshContentFilter)
+	UE::Tasks::TTask<TSharedPtr<const FSkeletalMesh>> FSystem::GetSkeletalMesh(UE::Tasks::FTask Prerequisite, FInstance::FID InstanceID, const FSkeletalMeshId& SkeletalMeshId, EMeshContentFlags SkeletalMeshContentFilter)
     {
 		LLM_SCOPE_BYNAME(TEXT("MutableRuntime"));
 		MUTABLE_CPUPROFILER_SCOPE(SystemGetSkeletalMesh);
 		
 		if (!SkeletalMeshId)
 		{
-			return UE::Tasks::MakeCompletedTask<TSharedPtr<const FSkeletalMesh>>(new FSkeletalMesh());
+			return UE::Tasks::Launch(TEXT("FSystem::GetSkeletalMeshInvalid"),
+			[]() -> TSharedPtr<const FSkeletalMesh>
+			{
+				return MakeShared<FSkeletalMesh>();
+			},
+			UE::Tasks::Prerequisites(Prerequisite),
+			UE::Tasks::ETaskPriority::Inherit,
+			UE::Tasks::EExtendedTaskPriority::Inline);
 		}
+
+		OP::ADDRESS RootAddress = SkeletalMeshId.GetKey()->Address;
 		
 		// Find the live instance
 		FLiveInstance* LiveInstance = m_pD->FindLiveInstance(InstanceID);
 		check(LiveInstance);
-		m_pD->WorkingMemoryManager.CurrentInstanceCache = LiveInstance->Cache;
-		m_pD->WorkingMemoryManager.CurrentPassthroughObjectLoader = LiveInstance->PassthroughObjectLoader;
-		m_pD->WorkingMemoryManager.BeginRunnerThread();
 		
-		OP::ADDRESS RootAddress = SkeletalMeshId.GetKey()->Address;
-			
 		UE::Mutable::Private::EOpType OpType = LiveInstance->Model->GetPrivate()->Program.GetOpType(RootAddress);
 		if (GetOpDataType(OpType) != EDataType::SkeletalMesh)
 		{
-			m_pD->WorkingMemoryManager.EndRunnerThread();
-			m_pD->WorkingMemoryManager.CurrentInstanceCache = nullptr;
-			m_pD->WorkingMemoryManager.CurrentPassthroughObjectLoader = nullptr;
-
-			return UE::Tasks::MakeCompletedTask<TSharedPtr<const FSkeletalMesh>>(new FSkeletalMesh());
+			return UE::Tasks::Launch(TEXT("FSystem::GetSkeletalMeshInvalid"),
+				[]() -> TSharedPtr<const FSkeletalMesh>
+				{
+					return MakeShared<FSkeletalMesh>();
+				},
+				UE::Tasks::Prerequisites(Prerequisite),
+				UE::Tasks::ETaskPriority::Inherit,
+				UE::Tasks::EExtendedTaskPriority::Inline);
 		}
-	
-		uint8 ExecutionOptions = static_cast<uint8>(SkeletalMeshContentFilter);
-		TSharedRef<CodeRunner> Runner = CodeRunner::Create(
-				m_pD->Settings,
-				m_pD, 
-				EExecutionStrategy::MinimizeMemory, 
-				LiveInstance->Model, 
-				LiveInstance->OldParameters.Get(), 
-				RootAddress, 
-				FSystem::AllLODs, 
-				ExecutionOptions, 
-				0, 
-				FScheduledOp::EType::Full,
-				LiveInstance->ExternalResourceProvider);
+
+		struct FContext
+		{
+			TSharedPtr<CodeRunner> Runner;
+		};
+		TSharedRef<FContext> Context = MakeShared<FContext>();
 		
-		constexpr bool bForceInlineExecution = false;
-		UE::Tasks::FTask RunnerCompletionEvent = Runner->StartRun(bForceInlineExecution);
+		UE::Tasks::FTask RunnerTask = UE::Tasks::Launch(TEXT("FSystem::GetSkeletalMesh"),
+		[Context, SystemPrivate = m_pD, RootAddress, LiveInstance, SkeletalMeshContentFilter]()
+		{
+			SystemPrivate->WorkingMemoryManager.CurrentInstanceCache = LiveInstance->Cache;
+			SystemPrivate->WorkingMemoryManager.CurrentPassthroughObjectLoader = LiveInstance->PassthroughObjectLoader;
+			SystemPrivate->WorkingMemoryManager.BeginRunnerThread();
+
+			TSharedRef<CodeRunner> Runner = CodeRunner::Create(
+					SystemPrivate->Settings,
+					SystemPrivate, 
+					EExecutionStrategy::MinimizeMemory, 
+					LiveInstance->Model, 
+					LiveInstance->OldParameters.Get(), 
+					RootAddress, 
+					FSystem::AllLODs, 
+					static_cast<uint8>(SkeletalMeshContentFilter), 
+					0, 
+					FScheduledOp::EType::Full,
+					LiveInstance->ExternalResourceProvider);
 
+			constexpr bool bForceInlineExecution = false;
+			Runner->StartRun(bForceInlineExecution);
+		},
+		UE::Tasks::Prerequisites(Prerequisite),
+		UE::Tasks::ETaskPriority::Inherit);
+		
 		return UE::Tasks::Launch(TEXT("FSystem::GetSkeletalMeshResultTask"),
-				[SystemPrivate = m_pD, Runner, RootAddress, ExecutionOptions]() -> TSharedPtr<const FSkeletalMesh>
-				{
-					TSharedPtr<const FSkeletalMesh> Result;
+		[Context, SystemPrivate = m_pD, RootAddress, SkeletalMeshContentFilter]() -> TSharedPtr<const FSkeletalMesh>
+		{
+			TSharedPtr<const FSkeletalMesh> Result;
 
-					SystemPrivate->bUnrecoverableError = Runner->bUnrecoverableError;
-					if (!Runner->bUnrecoverableError)
-					{
-						Result = SystemPrivate->WorkingMemoryManager.LoadSkeletalMesh(FCacheAddress(RootAddress, 0, ExecutionOptions));
-					}
-					
-					SystemPrivate->WorkingMemoryManager.EndRunnerThread();
-					SystemPrivate->WorkingMemoryManager.CurrentInstanceCache = nullptr;
-					SystemPrivate->WorkingMemoryManager.CurrentPassthroughObjectLoader = nullptr;
-					
-					return Result;
-				},
-				UE::Tasks::Prerequisites(RunnerCompletionEvent),
-				UE::Tasks::ETaskPriority::Inherit,
-				UE::Tasks::EExtendedTaskPriority::Inline);
+			SystemPrivate->bUnrecoverableError = Context->Runner->bUnrecoverableError;
+			if (!Context->Runner->bUnrecoverableError)
+			{
+				Result = SystemPrivate->WorkingMemoryManager.LoadSkeletalMesh(FCacheAddress(RootAddress, 0, static_cast<uint8>(SkeletalMeshContentFilter)));
+			}
+			
+			SystemPrivate->WorkingMemoryManager.EndRunnerThread();
+			SystemPrivate->WorkingMemoryManager.CurrentInstanceCache = nullptr;
+			SystemPrivate->WorkingMemoryManager.CurrentPassthroughObjectLoader = nullptr;
+			
+			return Result;
+		},
+		UE::Tasks::Prerequisites(RunnerTask),
+		UE::Tasks::ETaskPriority::Inherit,
+		UE::Tasks::EExtendedTaskPriority::Inline);
 	}
 
 	
-	UE::Tasks::TTask<TSharedPtr<const FMaterial>> FSystem::GetMaterial(FInstance::FID InstanceID, const FMaterialId& MaterialId)
+	UE::Tasks::TTask<TSharedPtr<const FMaterial>> FSystem::GetMaterial(UE::Tasks::FTask Prerequisite, FInstance::FID InstanceID, const FMaterialId& MaterialId)
 	{
 		LLM_SCOPE_BYNAME(TEXT("MutableRuntime"));
 		MUTABLE_CPUPROFILER_SCOPE(SystemGetMaterial);
 		
 		if (!MaterialId)
 		{
-			return UE::Tasks::MakeCompletedTask<TSharedPtr<const FMaterial>>(new FMaterial());
+			return UE::Tasks::Launch(TEXT("FSystem::GetMaterialInvalid"),
+			[]() -> TSharedPtr<const FMaterial>
+			{
+				return MakeShared<FMaterial>();
+			},
+			UE::Tasks::Prerequisites(Prerequisite),
+			UE::Tasks::ETaskPriority::Inherit,
+			UE::Tasks::EExtendedTaskPriority::Inline);
 		}
 		
+		OP::ADDRESS RootAddress = MaterialId.GetKey()->Address;
+	
 		// Find the live instance
 		FLiveInstance* LiveInstance = m_pD->FindLiveInstance(InstanceID);
 		check(LiveInstance);
-		m_pD->WorkingMemoryManager.CurrentInstanceCache = LiveInstance->Cache;
-		m_pD->WorkingMemoryManager.CurrentPassthroughObjectLoader = LiveInstance->PassthroughObjectLoader;
-		m_pD->WorkingMemoryManager.BeginRunnerThread();
 		
-		OP::ADDRESS RootAddress = MaterialId.GetKey()->Address;
-			
-		UE::Mutable::Private::EOpType OpType = LiveInstance->Model->GetPrivate()->Program.GetOpType(RootAddress);
+		EOpType OpType = LiveInstance->Model->GetPrivate()->Program.GetOpType(RootAddress);
 		if (GetOpDataType(OpType) != EDataType::Material)
 		{
-			m_pD->WorkingMemoryManager.EndRunnerThread();
-			m_pD->WorkingMemoryManager.CurrentInstanceCache = nullptr;
-			m_pD->WorkingMemoryManager.CurrentPassthroughObjectLoader = nullptr;
-
-			return UE::Tasks::MakeCompletedTask<TSharedPtr<const FMaterial>>(new FMaterial());
+			return UE::Tasks::Launch(TEXT("FSystem::GetMaterialInvalid"),
+			[]() -> TSharedPtr<const FMaterial>
+			{
+				return MakeShared<FMaterial>();
+			},
+			UE::Tasks::Prerequisites(Prerequisite),
+			UE::Tasks::ETaskPriority::Inherit,
+			UE::Tasks::EExtendedTaskPriority::Inline);
 		}
-	
-		TSharedRef<CodeRunner> Runner = CodeRunner::Create(
-				m_pD->Settings,
-				m_pD, 
-				EExecutionStrategy::MinimizeMemory, 
-				LiveInstance->Model, 
-				LiveInstance->OldParameters.Get(), 
-				RootAddress, 
-				FSystem::AllLODs, 
-				0, 
-				0, 
-				FScheduledOp::EType::Full,
-				LiveInstance->ExternalResourceProvider);
+
+		struct FContext
+		{
+			TSharedPtr<CodeRunner> Runner;
+		};
+		TSharedRef<FContext> Context = MakeShared<FContext>();
 		
-		constexpr bool bForceInlineExecution = false;
-		UE::Tasks::FTask RunnerCompletionEvent = Runner->StartRun(bForceInlineExecution);
+		UE::Tasks::FTask RunnerTask = UE::Tasks::Launch(TEXT("FSystem::GetMaterial"),
+		[Context, SystemPrivate = m_pD, RootAddress, LiveInstance]()
+		{
+			SystemPrivate->WorkingMemoryManager.CurrentInstanceCache = LiveInstance->Cache;
+			SystemPrivate->WorkingMemoryManager.CurrentPassthroughObjectLoader = LiveInstance->PassthroughObjectLoader;
+			SystemPrivate->WorkingMemoryManager.BeginRunnerThread();
+				
+			Context->Runner = CodeRunner::Create(
+					SystemPrivate->Settings,
+					SystemPrivate, 
+					EExecutionStrategy::MinimizeMemory, 
+					LiveInstance->Model, 
+					LiveInstance->OldParameters.Get(), 
+					RootAddress, 
+					FSystem::AllLODs, 
+					0, 
+					0, 
+					FScheduledOp::EType::Full,
+					LiveInstance->ExternalResourceProvider);
 
+			constexpr bool bForceInlineExecution = false;
+			Context->Runner->StartRun(bForceInlineExecution);
+		},
+		UE::Tasks::Prerequisites(Prerequisite),
+		UE::Tasks::ETaskPriority::Inherit);
+		
 		return UE::Tasks::Launch(TEXT("FSystem::GetMaterialResultTask"),
-				[SystemPrivate = m_pD, Runner, RootAddress]() -> TSharedPtr<const FMaterial>
-				{
-					TSharedPtr<const FMaterial> Result;
+		[Context, SystemPrivate = m_pD, RootAddress]() -> TSharedPtr<const FMaterial>
+		{
+			TSharedPtr<const FMaterial> Result;
 
-					SystemPrivate->bUnrecoverableError = Runner->bUnrecoverableError;
-					if (!Runner->bUnrecoverableError)
-					{
-						Result = SystemPrivate->WorkingMemoryManager.LoadMaterial(FCacheAddress(RootAddress, 0, 0));
-					}
-					
-					SystemPrivate->WorkingMemoryManager.EndRunnerThread();
-					SystemPrivate->WorkingMemoryManager.CurrentInstanceCache = nullptr;
-					SystemPrivate->WorkingMemoryManager.CurrentPassthroughObjectLoader = nullptr;
-					
-					return Result;
-				},
-				UE::Tasks::Prerequisites(RunnerCompletionEvent),
-				UE::Tasks::ETaskPriority::Inherit,
-				UE::Tasks::EExtendedTaskPriority::Inline);
+			SystemPrivate->bUnrecoverableError = Context->Runner->bUnrecoverableError;
+			if (!Context->Runner->bUnrecoverableError)
+			{
+				Result = SystemPrivate->WorkingMemoryManager.LoadMaterial(FCacheAddress(RootAddress, 0, 0));
+			}
+			
+			SystemPrivate->WorkingMemoryManager.EndRunnerThread();
+			SystemPrivate->WorkingMemoryManager.CurrentInstanceCache = nullptr;
+			SystemPrivate->WorkingMemoryManager.CurrentPassthroughObjectLoader = nullptr;
+			
+			return Result;
+		},
+		UE::Tasks::Prerequisites(RunnerTask),
+		UE::Tasks::ETaskPriority::Inherit,
+		UE::Tasks::EExtendedTaskPriority::Inline);
 	}
 	
 
@@ -1185,9 +1291,8 @@ namespace UE::Mutable::Private
 			ExternalResourceProvider);
 		
 		constexpr bool bForceInlineExecutution = true;
-		UE::Tasks::FTask RunnerCompletionEvent = Runner->StartRun(bForceInlineExecutution);
-		check(RunnerCompletionEvent.IsCompleted());
-
+		Runner->StartRun(bForceInlineExecutution);
+		
 		bUnrecoverableError = Runner->bUnrecoverableError;
 	}
 
