[Mutable] Fix GetSkeletalMesh.

diff --git a/Engine/Plugins/Mutable/Source/MutableRuntime/Private/MuR/System.cpp b/Engine/Plugins/Mutable/Source/MutableRuntime/Private/MuR/System.cpp
--- a/Engine/Plugins/Mutable/Source/MutableRuntime/Private/MuR/System.cpp
+++ b/Engine/Plugins/Mutable/Source/MutableRuntime/Private/MuR/System.cpp
@@ -873,7 +873,7 @@ namespace UE::Mutable::Private
 			SystemPrivate->WorkingMemoryManager.CurrentPassthroughObjectLoader = LiveInstance->PassthroughObjectLoader;
 			SystemPrivate->WorkingMemoryManager.BeginRunnerThread();
 
-			TSharedRef<CodeRunner> Runner = CodeRunner::Create(
+			Context->Runner = CodeRunner::Create(
 					SystemPrivate->Settings,
 					SystemPrivate, 
 					EExecutionStrategy::MinimizeMemory, 
@@ -887,7 +887,7 @@ namespace UE::Mutable::Private
 					LiveInstance->ExternalResourceProvider);
 
 			constexpr bool bForceInlineExecution = false;
-			Runner->StartRun(bForceInlineExecution);
+			Context->Runner->StartRun(bForceInlineExecution);
 		},
 		UE::Tasks::Prerequisites(Prerequisite),
 		UE::Tasks::ETaskPriority::Inherit);
