[Mutable] Allow building ResourceId outside Core.

- Relevant Parameters are no longer embedded in the bytecode.
- Added bForceLazy for upcoming MI_NEW (replacing IN_ADDIMAGE).

diff --git a/Engine/Plugins/Mutable/Source/CustomizableObject/Internal/MuCO/CustomizableObjectPrivate.h b/Engine/Plugins/Mutable/Source/CustomizableObject/Internal/MuCO/CustomizableObjectPrivate.h
--- a/Engine/Plugins/Mutable/Source/CustomizableObject/Internal/MuCO/CustomizableObjectPrivate.h
+++ b/Engine/Plugins/Mutable/Source/CustomizableObject/Internal/MuCO/CustomizableObjectPrivate.h
@@ -1750,6 +1750,8 @@ enum class ECustomizableObjectVersions : int32
 
 	SkeletalMeshParameter,
 
+	RelevantParameters,
+
 	// -----<new versions can be added above this line>--------
 	LastCustomizableObjectVersion
 };
diff --git a/Engine/Plugins/Mutable/Source/MutableRuntime/Internal/MuR/ModelPrivate.h b/Engine/Plugins/Mutable/Source/MutableRuntime/Internal/MuR/ModelPrivate.h
--- a/Engine/Plugins/Mutable/Source/MutableRuntime/Internal/MuR/ModelPrivate.h
+++ b/Engine/Plugins/Mutable/Source/MutableRuntime/Internal/MuR/ModelPrivate.h
@@ -375,6 +375,10 @@ namespace UE::Mutable::Private
 		 */
 		TArray<TArray<uint16>> ParameterLists;
 
+    	/** Relevant parameters for the given instruction. */
+	   	TMap<OP::ADDRESS, int32> RelevantParameterList;
+
+
 		/** Constant Material Data */
 		TArray<TSharedPtr<const FMaterial>> ConstantMaterials;
 
@@ -414,6 +418,7 @@ namespace UE::Mutable::Private
             Arch << Parameters;
             Arch << Ranges;
             Arch << ParameterLists;
+            Arch << RelevantParameterList;
 			Arch << ConstantMaterials;
         }
 
@@ -444,6 +449,7 @@ namespace UE::Mutable::Private
             Arch >> Parameters;
             Arch >> Ranges;
             Arch >> ParameterLists;
+            Arch >> RelevantParameterList;
             Arch >> ConstantMaterials;
         }
 
diff --git a/Engine/Plugins/Mutable/Source/MutableRuntime/Internal/MuR/Operations.h b/Engine/Plugins/Mutable/Source/MutableRuntime/Internal/MuR/Operations.h
--- a/Engine/Plugins/Mutable/Source/MutableRuntime/Internal/MuR/Operations.h
+++ b/Engine/Plugins/Mutable/Source/MutableRuntime/Internal/MuR/Operations.h
@@ -1102,11 +1102,6 @@ namespace UE::Mutable::Private
 			int32 SharedSurfaceId;
 
             ADDRESS name;
-
-            // Index in the FProgram::ParameterLists with the parameters that are relevant
-            // for whatever is added by this operation. This is used only for resources like
-            // images, meshes, or materials.
-            ADDRESS RelevantParametersListIndex;
         };
 
 		struct InstanceAddExtensionDataArgs
@@ -1128,11 +1123,6 @@ namespace UE::Mutable::Private
 			// be added to.
 			ADDRESS Instance;
 			ADDRESS Material;
-
-			// Index in the FProgram::ParameterLists with the parameters that are relevant
-			// for whatever is added by this operation. This is used only for resources like
-			// images, meshes, or materials.
-			ADDRESS RelevantParametersListIndex;
 		};
 
         struct LayoutPackArgs
diff --git a/Engine/Plugins/Mutable/Source/MutableRuntime/Internal/MuR/System.h b/Engine/Plugins/Mutable/Source/MutableRuntime/Internal/MuR/System.h
--- a/Engine/Plugins/Mutable/Source/MutableRuntime/Internal/MuR/System.h
+++ b/Engine/Plugins/Mutable/Source/MutableRuntime/Internal/MuR/System.h
@@ -226,6 +226,14 @@ namespace UE::Mutable::Private
 									const TSharedPtr<const FParameters>& FParameters,
 									bool* Flags );
 
+		static FMeshId GetMeshId(const TSharedPtr<FMeshIdRegistry>& Registry, const TSharedPtr<const FModel>& Model, const FParameters* Params, OP::ADDRESS RootAt);
+    	
+		static FImageId GetImageId(const TSharedPtr<FImageIdRegistry>& Registry, const TSharedPtr<const FModel>& Model, const FParameters* Params, OP::ADDRESS RootAt);
+
+		static FMaterialId GetMaterialId(const TSharedPtr<FMaterialIdRegistry>& Registry, const TSharedPtr<const FModel>& Model, const FParameters* Params, OP::ADDRESS RootAt);
+
+		static FSkeletalMeshId GetSkeletalMeshId(const TSharedPtr<FSkeletalMeshIdRegistry>& Registry, const TSharedPtr<const FModel>& Model, const FParameters* Params, OP::ADDRESS RootAt);
+		
 		//-----------------------------------------------------------------------------------------
 		// Interface pattern
 		//-----------------------------------------------------------------------------------------
diff --git a/Engine/Plugins/Mutable/Source/MutableRuntime/Internal/MuR/SystemPrivate.h b/Engine/Plugins/Mutable/Source/MutableRuntime/Internal/MuR/SystemPrivate.h
--- a/Engine/Plugins/Mutable/Source/MutableRuntime/Internal/MuR/SystemPrivate.h
+++ b/Engine/Plugins/Mutable/Source/MutableRuntime/Internal/MuR/SystemPrivate.h
@@ -94,6 +94,7 @@ namespace UE::Mutable::Private
 		inline FScheduledOp()
 		{
 			Stage = 0;
+			bForceLazy = false;
 			Type = EType::Full;
 		}
 
@@ -105,6 +106,7 @@ namespace UE::Mutable::Private
 			ExecutionIndex = InOpTemplate.ExecutionIndex;
 			Stage = InStage;
 			CustomState = InCustomState;
+			bForceLazy = InOpTemplate.bForceLazy;
 			Type = InOpTemplate.Type;
 		}
 
@@ -116,6 +118,7 @@ namespace UE::Mutable::Private
 			r.ExecutionIndex = InOpTemplate.ExecutionIndex;
 			r.Stage = 0;
 			r.CustomState = InOpTemplate.CustomState;
+			r.bForceLazy = InOpTemplate.bForceLazy;
 			r.Type = InOpTemplate.Type;
 			return r;
 		}
@@ -140,7 +143,9 @@ namespace UE::Mutable::Private
 		//! Internal stage of the operation.
 		//! Stage 0 is usually scheduling of children, and 1 is execution. Some instructions
 		//! may have more steges to schedule children that are optional for execution, etc.
-		uint8 Stage : 7;
+		uint8 Stage : 6;
+
+		uint8 bForceLazy : 1;
 
 		//! Type of calculation we are requesting for this operation.
 		enum class EType : uint8
@@ -1336,14 +1341,6 @@ namespace UE::Mutable::Private
 
 		FLiveInstance CurrentLiveInstance;
 
-		FMeshId GetMeshId(const TSharedPtr<const FModel>& Model, const FParameters* Params, uint32 ParamListIndex, OP::ADDRESS RootAt);
-    	
-		FImageId GetImageId(const TSharedPtr<const FModel>& Model, const FParameters* Params, uint32 ParamListIndex, OP::ADDRESS RootAt);
-
-    	FMaterialId GetMaterialId(const TSharedPtr<const FModel>& Model, const FParameters* Params, uint32 ParamListIndex, OP::ADDRESS RootAt);
-
-    	FSkeletalMeshId GetSkeletalMeshId(const TSharedPtr<const FModel>& Model, const FParameters* Params, uint32 ParamListIndex, OP::ADDRESS RootAt);
-
 		/** Maximum working memory that mutable should be using. */
 		int64 BudgetBytes = 0;
 
diff --git a/Engine/Plugins/Mutable/Source/MutableRuntime/Private/MuR/CodeRunner.cpp b/Engine/Plugins/Mutable/Source/MutableRuntime/Private/MuR/CodeRunner.cpp
--- a/Engine/Plugins/Mutable/Source/MutableRuntime/Private/MuR/CodeRunner.cpp
+++ b/Engine/Plugins/Mutable/Source/MutableRuntime/Private/MuR/CodeRunner.cpp
@@ -159,11 +159,12 @@ namespace UE::Mutable::Private
 		uint32 InLODMask,
 		uint8 ExecutionOptions,
 		int32 InImageLOD,
+		bool bForceLazy,
 		FScheduledOp::EType Type,
 		const TSharedPtr<FExternalResourceProvider>& ExternalResourceProvider)
 	{
 		return MakeShared<CodeRunner>(FPrivateToken {},
-				InSettings, InSystem, InExecutionStrategy, InModel, InParams, At, InLODMask, ExecutionOptions, InImageLOD, Type, ExternalResourceProvider);
+				InSettings, InSystem, InExecutionStrategy, InModel, InParams, At, InLODMask, ExecutionOptions, InImageLOD, bForceLazy, Type, ExternalResourceProvider);
 	}
 
 
@@ -177,6 +178,7 @@ namespace UE::Mutable::Private
 		uint32 InLodMask,
 		uint8 ExecutionOptions,
 		int32 InImageLOD,
+		bool bForceLazy,
 		FScheduledOp::EType Type,
 		const TSharedPtr<FExternalResourceProvider>& ExternalResourceProvider)
 		: Settings(InSettings)
@@ -210,6 +212,7 @@ namespace UE::Mutable::Private
 		FScheduledOp RootOp;
 		RootOp.At = At;
 		RootOp.ExecutionOptions = ExecutionOptions;
+		RootOp.bForceLazy = bForceLazy;
 		RootOp.Type = Type;
 		AddOp(RootOp);
 	}
@@ -957,7 +960,7 @@ namespace UE::Mutable::Private
 
                 if ( Args.value )
                 {
-					FMeshId MeshId = System->WorkingMemoryManager.GetMeshId(InModel, InParams, Args.RelevantParametersListIndex, Args.value);
+					FMeshId MeshId = FSystem::GetMeshId(System->WorkingMemoryManager.CurrentMeshIdRegistry, InModel, InParams, Args.value);
 					OP::ADDRESS NameAd = Args.name;
 					check(NameAd < (uint32)Program.ConstantStrings.Num());
 					const FString& Name = Program.ConstantStrings[NameAd];
@@ -998,7 +1001,7 @@ namespace UE::Mutable::Private
 
                 if ( Args.value )
                 {
-					FImageId ImageId = System->WorkingMemoryManager.GetImageId(InModel, InParams, Args.RelevantParametersListIndex, Args.value);
+					FImageId ImageId = FSystem::GetImageId(System->WorkingMemoryManager.CurrentImageIdRegistry, InModel, InParams, Args.value);
 					OP::ADDRESS NameAd = Args.name;
 					check(NameAd < (uint32)Program.ConstantStrings.Num());
 					const FString& Name = Program.ConstantStrings[NameAd];
@@ -1047,7 +1050,7 @@ namespace UE::Mutable::Private
 
         				if (Args.Material)
         				{
-        					FMaterialId MaterialId = System->WorkingMemoryManager.GetMaterialId(InModel, InParams, Args.RelevantParametersListIndex, Args.Material);
+        					FMaterialId MaterialId = FSystem::GetMaterialId(System->WorkingMemoryManager.CurrentMaterialIdRegistry, InModel, InParams, Args.Material);
         					Result->SetMaterialId(0, 0, 0, MaterialId);
         				}
 
@@ -1089,7 +1092,7 @@ namespace UE::Mutable::Private
 
         				if (Args.Material)
         				{
-        					FMaterialId MaterialId = System->WorkingMemoryManager.GetMaterialId(InModel, InParams, Args.RelevantParametersListIndex, Args.Material);
+        					FMaterialId MaterialId = FSystem::GetMaterialId(System->WorkingMemoryManager.CurrentMaterialIdRegistry, InModel, InParams, Args.Material);
         					Result->SetMaterialId(0, 0, 0, MaterialId);
         				}
 
@@ -1139,7 +1142,7 @@ namespace UE::Mutable::Private
 
         				if (Args.value)
         				{
-        					FSkeletalMeshId SkeletalMeshId = System->WorkingMemoryManager.GetSkeletalMeshId(Model, InParams, Args.RelevantParametersListIndex, Args.value);
+        					FSkeletalMeshId SkeletalMeshId = FSystem::GetSkeletalMeshId(System->WorkingMemoryManager.CurrentSkeletalMeshIdRegistry, InModel, InParams, Args.value);
         					Instance->SetSkeletalMeshId(0, SkeletalMeshId);
         				}
 						
diff --git a/Engine/Plugins/Mutable/Source/MutableRuntime/Private/MuR/CodeRunner.h b/Engine/Plugins/Mutable/Source/MutableRuntime/Private/MuR/CodeRunner.h
--- a/Engine/Plugins/Mutable/Source/MutableRuntime/Private/MuR/CodeRunner.h
+++ b/Engine/Plugins/Mutable/Source/MutableRuntime/Private/MuR/CodeRunner.h
@@ -52,6 +52,7 @@ namespace  UE::Mutable::Private
 				uint32 LodMask,
 				uint8 ExecutionOptions,
 				int32 InImageLOD,
+				bool bForceLazy,
 				FScheduledOp::EType Type,
 				const TSharedPtr<FExternalResourceProvider>& ExternalResourceProvider);
 
@@ -67,6 +68,7 @@ namespace  UE::Mutable::Private
 			uint32 LodMask,
 			uint8 ExecutionOptions,
 			int32 InImageLOD,
+			bool bForceLazy,
 			FScheduledOp::EType Type,
 			const TSharedPtr<FExternalResourceProvider>& ExternalResourceProvider);
 
diff --git a/Engine/Plugins/Mutable/Source/MutableRuntime/Private/MuR/CodeRunnerBegin.cpp b/Engine/Plugins/Mutable/Source/MutableRuntime/Private/MuR/CodeRunnerBegin.cpp
--- a/Engine/Plugins/Mutable/Source/MutableRuntime/Private/MuR/CodeRunnerBegin.cpp
+++ b/Engine/Plugins/Mutable/Source/MutableRuntime/Private/MuR/CodeRunnerBegin.cpp
@@ -784,7 +784,7 @@ namespace UE::Mutable::Private
 							
 							if (Args.value)
 							{
-								FMeshId MeshId = System.GetPrivate()->WorkingMemoryManager.GetMeshId(Model, Params.Get(), Args.RelevantParametersListIndex, Args.value);
+								FMeshId MeshId = FSystem::GetMeshId(System.GetPrivate()->WorkingMemoryManager.CurrentMeshIdRegistry, Model, Params.Get(), Args.value);
 								OP::ADDRESS NameAd = Args.name;
 								check(NameAd < (uint32)Program.ConstantStrings.Num());
 								const FString& Name = Program.ConstantStrings[NameAd];
@@ -836,7 +836,7 @@ namespace UE::Mutable::Private
 							
 							if (Args.value)
 							{
-								FImageId ImageId = System.GetPrivate()->WorkingMemoryManager.GetImageId(Model, Params.Get(), Args.RelevantParametersListIndex, Args.value);
+								FImageId ImageId = FSystem::GetImageId(System.GetPrivate()->WorkingMemoryManager.CurrentImageIdRegistry, Model, Params.Get(), Args.value);
 								OP::ADDRESS NameAd = Args.name;
 								check(NameAd < (uint32)Program.ConstantStrings.Num());
 								const FString& Name = Program.ConstantStrings[NameAd];
@@ -888,7 +888,7 @@ namespace UE::Mutable::Private
 
 							if (Args.Material)
 							{
-								FMaterialId MaterialId = System.GetPrivate()->WorkingMemoryManager.GetMaterialId(Model, Params.Get(), Args.RelevantParametersListIndex, Args.Material);
+        						FMaterialId MaterialId = FSystem::GetMaterialId(System.GetPrivate()->WorkingMemoryManager.CurrentMaterialIdRegistry, Model, Params.Get(), Args.Material);
 								Instance->SetMaterialId(0, 0, 0, MaterialId);
 							}
 							
@@ -937,7 +937,7 @@ namespace UE::Mutable::Private
 
 							if (Args.Material)
 							{
-								FMaterialId MaterialId = System.GetPrivate()->WorkingMemoryManager.GetMaterialId(Model, Params.Get(), Args.RelevantParametersListIndex, Args.Material);
+        						FMaterialId MaterialId = FSystem::GetMaterialId(System.GetPrivate()->WorkingMemoryManager.CurrentMaterialIdRegistry, Model, Params.Get(), Args.Material);
 								Instance->SetOverlayMaterialId(0, MaterialId);
 							}
 							
@@ -1238,7 +1238,7 @@ namespace UE::Mutable::Private
 							
 							if (Args.value)
 							{
-								FSkeletalMeshId SkeletalMeshId = System.GetPrivate()->WorkingMemoryManager.GetSkeletalMeshId(Model, Params.Get(), Args.RelevantParametersListIndex, Args.value);
+        						FSkeletalMeshId SkeletalMeshId = FSystem::GetSkeletalMeshId(System.GetPrivate()->WorkingMemoryManager.CurrentSkeletalMeshIdRegistry, Model, Params.Get(), Args.value);
 								Instance->SetSkeletalMeshId(0, SkeletalMeshId);
 							}
 							
diff --git a/Engine/Plugins/Mutable/Source/MutableRuntime/Private/MuR/System.cpp b/Engine/Plugins/Mutable/Source/MutableRuntime/Private/MuR/System.cpp
--- a/Engine/Plugins/Mutable/Source/MutableRuntime/Private/MuR/System.cpp
+++ b/Engine/Plugins/Mutable/Source/MutableRuntime/Private/MuR/System.cpp
@@ -490,6 +490,7 @@ namespace UE::Mutable::Private
 				LiveInstance->OldParameters.Get(),
 				RootAddress, FSystem::AllLODs,
 				MipsToSkip, InImageLOD,
+				false,
 				FScheduledOp::EType::Full,
 				LiveInstance->ExternalResourceProvider);
 				
@@ -582,6 +583,7 @@ namespace UE::Mutable::Private
 					FSystem::AllLODs,
 					executionOptions,
 					0,
+					false,
 					FScheduledOp::EType::ImageDesc,
 					LiveInstance->ExternalResourceProvider);
 			
@@ -668,6 +670,7 @@ namespace UE::Mutable::Private
 				FSystem::AllLODs,
 				ExecutionOptions,
 				0,
+				false,
 				FScheduledOp::EType::ImageDesc,
 				LiveInstance->ExternalResourceProvider);
 
@@ -791,7 +794,8 @@ namespace UE::Mutable::Private
 					RootAddress, 
 					FSystem::AllLODs, 
 					ExecutionOptions, 
-					0, 
+					0,
+					false, 
 					FScheduledOp::EType::Full,
 					LiveInstance->ExternalResourceProvider);
 	
@@ -882,7 +886,8 @@ namespace UE::Mutable::Private
 					RootAddress, 
 					FSystem::AllLODs, 
 					static_cast<uint8>(SkeletalMeshContentFilter), 
-					0, 
+					0,
+					false, 
 					FScheduledOp::EType::Full,
 					LiveInstance->ExternalResourceProvider);
 
@@ -973,7 +978,8 @@ namespace UE::Mutable::Private
 					RootAddress, 
 					FSystem::AllLODs, 
 					0, 
-					0, 
+					0,
+					true, 
 					FScheduledOp::EType::Full,
 					LiveInstance->ExternalResourceProvider);
 
@@ -1287,6 +1293,7 @@ namespace UE::Mutable::Private
 			InLODs,
 			ExecutionOptions,
 			InImageLOD,
+			false,
 			FScheduledOp::EType::Full,
 			ExternalResourceProvider);
 		
@@ -2438,46 +2445,54 @@ namespace UE::Mutable::Private
 	}
 
 
-	FMeshId FWorkingMemoryManager::GetMeshId(const TSharedPtr<const FModel>& Model, const FParameters* Params, uint32 ParamListIndex, OP::ADDRESS RootAt)
+	FMeshId FSystem::GetMeshId(const TSharedPtr<FMeshIdRegistry>& Registry, const TSharedPtr<const FModel>& Model, const FParameters* Params, OP::ADDRESS RootAt)
     {
     	FGeneratedMeshKey NewKey;
     	NewKey.Address = RootAt;
 
+    	int32* Result = Model->GetPrivate()->Program.RelevantParameterList.Find(RootAt);
+    	const int32 ParamListIndex = Result ? *Result : INDEX_NONE; 
 		GetResourceKey(Model, Params, ParamListIndex, RootAt, NewKey.ParameterValuesBlob);
 
-    	return CurrentMeshIdRegistry->Add(NewKey, {});
+    	return Registry->Add(NewKey, {});
     }
 	
     	
-	FImageId FWorkingMemoryManager::GetImageId(const TSharedPtr<const FModel>& Model, const FParameters* Params, uint32 ParamListIndex, OP::ADDRESS RootAt)
+	FImageId FSystem::GetImageId(const TSharedPtr<FImageIdRegistry>& Registry, const TSharedPtr<const FModel>& Model, const FParameters* Params, OP::ADDRESS RootAt)
     {
     	FGeneratedImageKey NewKey;
     	NewKey.Address = RootAt;
 
+    	int32* Result = Model->GetPrivate()->Program.RelevantParameterList.Find(RootAt);
+    	const int32 ParamListIndex = Result ? *Result : INDEX_NONE; 
     	GetResourceKey(Model, Params, ParamListIndex, RootAt, NewKey.ParameterValuesBlob);
 
-    	return CurrentImageIdRegistry->Add(NewKey, {});
+    	return Registry->Add(NewKey, {});
     }
 
 	
-	FMaterialId FWorkingMemoryManager::GetMaterialId(const TSharedPtr<const FModel>& Model, const FParameters* Params, uint32 ParamListIndex, OP::ADDRESS RootAt)
+	FMaterialId FSystem::GetMaterialId(const TSharedPtr<FMaterialIdRegistry>& Registry, const TSharedPtr<const FModel>& Model, const FParameters* Params, OP::ADDRESS RootAt)
     {
     	FGeneratedMaterialKey NewKey;
     	NewKey.Address = RootAt;
 
+    	int32* Result = Model->GetPrivate()->Program.RelevantParameterList.Find(RootAt);
+    	const int32 ParamListIndex = Result ? *Result : INDEX_NONE; 
     	GetResourceKey(Model, Params, ParamListIndex, RootAt, NewKey.ParameterValuesBlob);
 
-    	return CurrentMaterialIdRegistry->Add(NewKey, {});
+    	return Registry->Add(NewKey, {});
     }
 
 	
-	FSkeletalMeshId FWorkingMemoryManager::GetSkeletalMeshId(const TSharedPtr<const FModel>& Model, const FParameters* Params, uint32 ParamListIndex, OP::ADDRESS RootAt)
+	FSkeletalMeshId FSystem::GetSkeletalMeshId(const TSharedPtr<FSkeletalMeshIdRegistry>& Registry, const TSharedPtr<const FModel>& Model, const FParameters* Params, OP::ADDRESS RootAt)
 	{
     	FGeneratedSkeletalMeshKey NewKey;
     	NewKey.Address = RootAt;
 
+    	int32* Result = Model->GetPrivate()->Program.RelevantParameterList.Find(RootAt);
+    	const int32 ParamListIndex = Result ? *Result : INDEX_NONE; 
     	GetResourceKey(Model, Params, ParamListIndex, RootAt, NewKey.ParameterValuesBlob);
 
-    	return CurrentSkeletalMeshIdRegistry->Add(NewKey, {});
+    	return Registry->Add(NewKey, {});
 	}
 }
diff --git a/Engine/Plugins/Mutable/Source/MutableTools/Private/MuT/ASTOpAddMaterial.cpp b/Engine/Plugins/Mutable/Source/MutableTools/Private/MuT/ASTOpAddMaterial.cpp
--- a/Engine/Plugins/Mutable/Source/MutableTools/Private/MuT/ASTOpAddMaterial.cpp
+++ b/Engine/Plugins/Mutable/Source/MutableTools/Private/MuT/ASTOpAddMaterial.cpp
@@ -87,42 +87,41 @@ void ASTOpAddMaterial::Link(FProgram& Program, FLinkerOptions*)
 	if (Material)
 	{
 		Args.Material = Material->linkedAddress;
-	}
-
-	// Find out relevant parameters. \todo: this may be optimised by reusing partial
-	// values in a LINK_CONTEXT or similar
-	SubtreeRelevantParametersVisitorAST visitor;
-	visitor.Run(Material.child());
-
-	TArray<uint16> params;
-	for (const FString& paramName : visitor.Parameters)
-	{
-		for (int32 i = 0; i < Program.Parameters.Num(); ++i)
+		
+		// Find out relevant parameters. \todo: this may be optimised by reusing partial
+		// values in a LINK_CONTEXT or similar
+		SubtreeRelevantParametersVisitorAST Visitor;
+		Visitor.Run(Material.child());
+		
+		TArray<uint16> Params;
+		for (const FString& ParameterName : Visitor.Parameters)
 		{
-			const auto& param = Program.Parameters[i];
-			if (param.Name == paramName)
+			for (int32 Index = 0; Index < Program.Parameters.Num(); ++Index)
 			{
-				params.Add(uint16(i));
-				break;
+				const FParameterDesc& Param = Program.Parameters[Index];
+				if (Param.Name == ParameterName)
+				{
+					Params.Add(uint16(Index));
+					break;
+				}
 			}
 		}
-	}
 
-	params.Sort();
+		Params.Sort();
 
-	auto it = Program.ParameterLists.Find(params);
+		int32 It = Program.ParameterLists.Find(Params);
 
-	if (it != INDEX_NONE)
-	{
-		Args.RelevantParametersListIndex = it;
-	}
-	else
-	{
-		Args.RelevantParametersListIndex = uint32_t(Program.ParameterLists.Num());
-		Program.ParameterLists.Add(params);
+		if (It != INDEX_NONE)
+		{
+			Program.RelevantParameterList.Add(Material->linkedAddress, It);
+		}
+		else
+		{
+			Program.RelevantParameterList.Add(Material->linkedAddress, Program.ParameterLists.Num());
+			Program.ParameterLists.Add(Params);
+		}
 	}
 
-	
 	linkedAddress = (OP::ADDRESS)Program.OpAddress.Num();
 	Program.OpAddress.Add((uint32_t)Program.ByteCode.Num());
 	AppendCode(Program.ByteCode, Type);
diff --git a/Engine/Plugins/Mutable/Source/MutableTools/Private/MuT/ASTOpAddSkeletalMesh.cpp b/Engine/Plugins/Mutable/Source/MutableTools/Private/MuT/ASTOpAddSkeletalMesh.cpp
--- a/Engine/Plugins/Mutable/Source/MutableTools/Private/MuT/ASTOpAddSkeletalMesh.cpp
+++ b/Engine/Plugins/Mutable/Source/MutableTools/Private/MuT/ASTOpAddSkeletalMesh.cpp
@@ -78,36 +78,37 @@ namespace UE::Mutable::Private
 		if (SkeletalMesh)
 		{
 			Args.value = SkeletalMesh->linkedAddress;
-		}
 
-		SubtreeRelevantParametersVisitorAST Visitor;
-		Visitor.Run(SkeletalMesh.child());
+			SubtreeRelevantParametersVisitorAST Visitor;
+			Visitor.Run(SkeletalMesh.child());
 
-		TArray<uint16> Params;
-		for (const FString& ParameterName : Visitor.Parameters)
-		{
-			for (int32 Index = 0; Index < Program.Parameters.Num(); ++Index)
+			TArray<uint16> Params;
+			for (const FString& ParameterName : Visitor.Parameters)
 			{
-				const FParameterDesc& ParameterDesc = Program.Parameters[Index];
-				if (ParameterDesc.Name == ParameterName)
+				for (int32 Index = 0; Index < Program.Parameters.Num(); ++Index)
 				{
-					Params.Add(static_cast<uint16>(Index));
-					break;
+					const FParameterDesc& ParameterDesc = Program.Parameters[Index];
+					if (ParameterDesc.Name == ParameterName)
+					{
+						Params.Add(static_cast<uint16>(Index));
+						break;
+					}
 				}
 			}
-		}
 
-		Params.Sort();
+			Params.Sort();
 
-		int32 Index = Program.ParameterLists.Find(Params);
-		if (Index != INDEX_NONE)
-		{
-			Args.RelevantParametersListIndex = Index;
-		}
-		else
-		{
-			Args.RelevantParametersListIndex = static_cast<uint32>(Program.ParameterLists.Num());
-			Program.ParameterLists.Add(Params);
+			int32 It = Program.ParameterLists.Find(Params);
+
+			if (It != INDEX_NONE)
+			{
+				Program.RelevantParameterList.Add(SkeletalMesh->linkedAddress, It);
+			}
+			else
+			{
+				Program.RelevantParameterList.Add(SkeletalMesh->linkedAddress, Program.ParameterLists.Num());
+				Program.ParameterLists.Add(Params);
+			}
 		}
 		
 		linkedAddress = static_cast<OP::ADDRESS>(Program.OpAddress.Num());
diff --git a/Engine/Plugins/Mutable/Source/MutableTools/Private/MuT/ASTOpInstanceAdd.cpp b/Engine/Plugins/Mutable/Source/MutableTools/Private/MuT/ASTOpInstanceAdd.cpp
--- a/Engine/Plugins/Mutable/Source/MutableTools/Private/MuT/ASTOpInstanceAdd.cpp
+++ b/Engine/Plugins/Mutable/Source/MutableTools/Private/MuT/ASTOpInstanceAdd.cpp
@@ -98,7 +98,7 @@ namespace UE::Mutable::Private
 	}
 
 
-	void ASTOpInstanceAdd::Link(FProgram& program, FLinkerOptions*)
+	void ASTOpInstanceAdd::Link(FProgram& Program, FLinkerOptions*)
 	{
 		check(type != EOpType::NONE);
 		
@@ -110,54 +110,56 @@ namespace UE::Mutable::Private
 			Args.id = id;
 			Args.ExternalId = ExternalId;
 			Args.SharedSurfaceId = SharedSurfaceId;
-			Args.name = program.AddConstant(name);
+			Args.name = Program.AddConstant(name);
 
 			if (instance) Args.instance = instance->linkedAddress;
-			if (value) Args.value = value->linkedAddress;
 
-			if (type == EOpType::IN_ADDIMAGE ||
-				type == EOpType::IN_ADDMESH)
+			if (value)
 			{
-				// Find out relevant parameters. \todo: this may be optimised by reusing partial
-				// values in a LINK_CONTEXT or similar
-				SubtreeRelevantParametersVisitorAST visitor;
-				visitor.Run(value.child());
+				Args.value = value->linkedAddress;
 
-				TArray<uint16> params;
-				for (const FString& paramName : visitor.Parameters)
+				if (type == EOpType::IN_ADDIMAGE || type == EOpType::IN_ADDMESH)
 				{
-					for (int32 i = 0; i < program.Parameters.Num(); ++i)
+					// Find out relevant parameters. \todo: this may be optimised by reusing partial
+					// values in a LINK_CONTEXT or similar
+					SubtreeRelevantParametersVisitorAST Visitor;
+					Visitor.Run(value.child());
+
+					TArray<uint16> Params;
+					for (const FString& ParameterName : Visitor.Parameters)
 					{
-						const auto& param = program.Parameters[i];
-						if (param.Name == paramName)
+						for (int32 i = 0; i < Program.Parameters.Num(); ++i)
 						{
-							params.Add(uint16(i));
-							break;
+							const auto& param = Program.Parameters[i];
+							if (param.Name == ParameterName)
+							{
+								Params.Add(uint16(i));
+								break;
+							}
 						}
 					}
-				}
 
-				params.Sort();
+					Params.Sort();
 
-				auto it = program.ParameterLists.Find(params);
+					int32 It = Program.ParameterLists.Find(Params);
 
-				if (it != INDEX_NONE)
-				{
-					Args.RelevantParametersListIndex = it;
-				}
-				else
-				{
-					Args.RelevantParametersListIndex = uint32_t(program.ParameterLists.Num());
-					program.ParameterLists.Add(params);
+					if (It != INDEX_NONE)
+					{
+						Program.RelevantParameterList.Add(value->linkedAddress, It);
+					}
+					else
+					{
+						Program.RelevantParameterList.Add(value->linkedAddress, Program.ParameterLists.Num());
+						Program.ParameterLists.Add(Params);
+					}
 				}
 			}
 
-			linkedAddress = (OP::ADDRESS)program.OpAddress.Num();
-			program.OpAddress.Add((uint32_t)program.ByteCode.Num());
-			AppendCode(program.ByteCode, type);
-			AppendCode(program.ByteCode, Args);
+			linkedAddress = (OP::ADDRESS)Program.OpAddress.Num();
+			Program.OpAddress.Add((uint32_t)Program.ByteCode.Num());
+			AppendCode(Program.ByteCode, type);
+			AppendCode(Program.ByteCode, Args);
 		}
-
 	}
 
 
diff --git a/Engine/Plugins/Mutable/Source/MutableTools/Private/MuT/ASTOpInstanceAdd.h b/Engine/Plugins/Mutable/Source/MutableTools/Private/MuT/ASTOpInstanceAdd.h
--- a/Engine/Plugins/Mutable/Source/MutableTools/Private/MuT/ASTOpInstanceAdd.h
+++ b/Engine/Plugins/Mutable/Source/MutableTools/Private/MuT/ASTOpInstanceAdd.h
@@ -30,8 +30,6 @@ struct FProgram;
 		int32 SharedSurfaceId = INDEX_NONE;
 		FString name;
 
-	public:
-
 		ASTOpInstanceAdd();
 		ASTOpInstanceAdd(const ASTOpInstanceAdd&) = delete;
 		~ASTOpInstanceAdd();
@@ -43,10 +41,8 @@ struct FProgram;
 		virtual Ptr<ASTOp> Clone(MapChildFuncRef mapChild) const override;
 		virtual uint64 Hash() const override;
 		virtual void Assert() override;
-		virtual void Link(FProgram& program, FLinkerOptions* Options) override;
+		virtual void Link(FProgram& Program, FLinkerOptions* Options) override;
 		virtual Ptr<ASTOp> OptimiseSink(const FModelOptimizationOptions&, FOptimizeSinkContext&) const override;
 	};
-
-
 }
 
