[Mutable] Fix FScheduledOp stage check.

diff --git a/Engine/Plugins/Mutable/Source/MutableRuntime/Internal/MuR/SystemPrivate.h b/Engine/Plugins/Mutable/Source/MutableRuntime/Internal/MuR/SystemPrivate.h
--- a/Engine/Plugins/Mutable/Source/MutableRuntime/Internal/MuR/SystemPrivate.h
+++ b/Engine/Plugins/Mutable/Source/MutableRuntime/Internal/MuR/SystemPrivate.h
@@ -100,7 +100,7 @@ namespace UE::Mutable::Private
 
 		inline FScheduledOp(OP::ADDRESS InAt, const FScheduledOp& InOpTemplate, uint8 InStage = 0, uint32 InCustomState = 0)
 		{
-			check(InStage < 120);
+			check(InStage <= 63);
 			At = InAt;
 			ExecutionOptions = InOpTemplate.ExecutionOptions;
 			ExecutionIndex = InOpTemplate.ExecutionIndex;
