[mutable] Added material pins to be used as overrides for the material slots of the generated mesh.
- Added option to display the pin editable name in the Pin Viewer (limited to only showing the values held by the node for now).
- Added option to decide the visibility policy for the pins that have an editable text box (Hide If Linked (default) and Always Visible (used for the material override pins)).
- Fixed AddPinFromUI transaction not working.
- Added ComponentName change to the transaction system.
- Added ParameterName changes to the transaction system.
- Added other mutable nodes string value changes to the transaction system.
- Removed the "TryRemove" CONode method and moved contents to the UI side of the node.
- Removed some unnecessary includes.

#rnx

diff --git a/Engine/Plugins/Mutable/Source/CustomizableObjectEditor/Private/MuCOE/CustomizableObjectEditorModule.cpp b/Engine/Plugins/Mutable/Source/CustomizableObjectEditor/Private/MuCOE/CustomizableObjectEditorModule.cpp
--- a/Engine/Plugins/Mutable/Source/CustomizableObjectEditor/Private/MuCOE/CustomizableObjectEditorModule.cpp
+++ b/Engine/Plugins/Mutable/Source/CustomizableObjectEditor/Private/MuCOE/CustomizableObjectEditorModule.cpp
@@ -114,6 +114,7 @@
 #include "Nodes/CONodeExternalOperation.h"
 #include "Nodes/CONodeExternalOperationDetails.h"
 #include "Nodes/CONodeExternalTypeParameter.h"
+#include "Nodes/CustomizableObjectNodeComponentPassthroughMesh.h"
 #include "Nodes/CustomizableObjectNodeTextureParameter.h"
 #include "Nodes/ExternalTypeCustomization.h"
 #include "PhysicsEngine/PhysicsAsset.h"
@@ -282,6 +283,7 @@ void FCustomizableObjectEditorModule::StartupModule()
 	RegisterCustomDetails(PropertyModule, UCustomizableObjectNodeMacroInstance::StaticClass(), FOnGetDetailCustomizationInstance::CreateStatic(&FCustomizableObjectNodeMacroInstanceDetails::MakeInstance));
 	RegisterCustomDetails(PropertyModule, UCustomizableObjectNodeParameter::StaticClass(), FOnGetDetailCustomizationInstance::CreateStatic(&FCustomizableObjectNodeParameterDetails::MakeInstance));
 	RegisterCustomDetails(PropertyModule, UCONodeMaterialBreak::StaticClass(), FOnGetDetailCustomizationInstance::CreateStatic(&FCustomizableObjectNodeDetails::MakeInstance));
+	RegisterCustomDetails(PropertyModule, UCustomizableObjectNodeComponentPassthroughMesh::StaticClass(), FOnGetDetailCustomizationInstance::CreateStatic(&FCustomizableObjectNodeDetails::MakeInstance));
 	RegisterCustomDetails(PropertyModule, UCONodeExternalOperation::StaticClass(), FOnGetDetailCustomizationInstance::CreateStatic(&FCONodeExternalOperationDetails::MakeInstance));
 
 	// Other Objects
diff --git a/Engine/Plugins/Mutable/Source/CustomizableObjectEditor/Private/MuCOE/EdGraphSchema_CustomizableObject.cpp b/Engine/Plugins/Mutable/Source/CustomizableObjectEditor/Private/MuCOE/EdGraphSchema_CustomizableObject.cpp
--- a/Engine/Plugins/Mutable/Source/CustomizableObjectEditor/Private/MuCOE/EdGraphSchema_CustomizableObject.cpp
+++ b/Engine/Plugins/Mutable/Source/CustomizableObjectEditor/Private/MuCOE/EdGraphSchema_CustomizableObject.cpp
@@ -1131,7 +1131,19 @@ void UEdGraphSchema_CustomizableObject::GetContextMenuActions(UToolMenu* Menu, U
 								return;
 							}
 								
-							Node->TryRemovePin(*const_cast<UEdGraphPin*>(Pin));
+							FScopedTransaction Transaction(LOCTEXT("TryRemovePinTransaction", "Try removing a Pin"));
+							Node->Modify();
+
+							// Make sure the pin that is requested for removal was actually added by the UI.
+							if (Node->CanPinBeRemoved(*Pin) && Node->CustomRemovePin(*const_cast<UEdGraphPin*>(Pin)))
+							{
+								// Notify that a pin has been removed so the PinViewer reacts to the new pin layout
+								if (UEdGraph* Graph = Node->GetGraph())
+								{
+									Graph->NotifyNodeChanged(Node);
+								}
+							}
+
 						})));
 				}
 				
diff --git a/Engine/Plugins/Mutable/Source/CustomizableObjectEditor/Private/MuCOE/GenerateMutableSource/GenerateMutableSourceComponent.cpp b/Engine/Plugins/Mutable/Source/CustomizableObjectEditor/Private/MuCOE/GenerateMutableSource/GenerateMutableSourceComponent.cpp
--- a/Engine/Plugins/Mutable/Source/CustomizableObjectEditor/Private/MuCOE/GenerateMutableSource/GenerateMutableSourceComponent.cpp
+++ b/Engine/Plugins/Mutable/Source/CustomizableObjectEditor/Private/MuCOE/GenerateMutableSource/GenerateMutableSourceComponent.cpp
@@ -173,6 +173,35 @@ UE::Mutable::Private::Ptr<UE::Mutable::Private::NodeComponent> GenerateMutableSo
 		ComponentNode->Id = GenerationContext.ComponentNames.Add(ComponentName);
 		ComponentNode->SkeletalMesh = UE::Mutable::Private::GenerateMutableSourceSkeletalMesh(Pin, GenerationContext);
 
+		// Override Material slots
+		for (const FEdGraphPinReference& OverrideMaterialPin : TypedComponentPassthroughMesh->MaterialSlotOverridePins)
+		{
+			// Pin Data
+			FName TargetMaterialSlotName;
+			if (const UEdGraphPin* MaterialOverridePinPointer = OverrideMaterialPin.Get())
+			{
+				const UCONodeComponentPassthroughMeshMaterialOverridePinData& OverrideMaterialPinData =
+					TypedComponentPassthroughMesh->GetPinData<UCONodeComponentPassthroughMeshMaterialOverridePinData>(*MaterialOverridePinPointer);
+
+				TargetMaterialSlotName = OverrideMaterialPinData.TargetMaterialSlotName;
+			}
+
+			// Material node
+			UE::Mutable::Private::Ptr<UE::Mutable::Private::NodeMaterial> MaterialNode = nullptr;
+			if (const UEdGraphPin* MaterialOverridePin = OverrideMaterialPin.Get())
+			{
+				if (const UEdGraphPin* ConnectedPin = FollowInputPin(*MaterialOverridePin))
+				{
+					MaterialNode = GenerateMutableSourceMaterial(ConnectedPin, GenerationContext);
+				}
+			}
+
+			if (MaterialNode)
+			{
+				ComponentNode->OverrideMaterials.Add(TargetMaterialSlotName, MaterialNode);
+			}
+		}
+		
 		GenerationContext.CurrentMeshComponent = FName();
 		Result = ComponentNode;
 	}
diff --git a/Engine/Plugins/Mutable/Source/CustomizableObjectEditor/Private/MuCOE/Nodes/CONodeMaterialBreak.cpp b/Engine/Plugins/Mutable/Source/CustomizableObjectEditor/Private/MuCOE/Nodes/CONodeMaterialBreak.cpp
--- a/Engine/Plugins/Mutable/Source/CustomizableObjectEditor/Private/MuCOE/Nodes/CONodeMaterialBreak.cpp
+++ b/Engine/Plugins/Mutable/Source/CustomizableObjectEditor/Private/MuCOE/Nodes/CONodeMaterialBreak.cpp
@@ -5,6 +5,7 @@
 #include "DetailsViewArgs.h"
 #include "MuCOE/EdGraphSchema_CustomizableObject.h"
 #include "PropertyEditorModule.h"
+#include "ScopedTransaction.h"
 
 
 #include UE_INLINE_GENERATED_CPP_BY_NAME(CONodeMaterialBreak)
@@ -44,12 +45,16 @@ void UCONodeMaterialBreak::AllocateDefaultPins(UCustomizableObjectNodeRemapPins*
 	MaterialPinRef = FEdGraphPinReference(MaterialPin);
 
 	// Get old pins
-	TArray<UEdGraphPin*> OutputPins;
-	GetOutputPins(OutputPins);
+	TArray<UEdGraphPin*> OldPins = GetAllNonOrphanPins();
 
-	// Regenerate the new pins from the old pin's pindata which may have changed
-	for (const UEdGraphPin* Pin : OutputPins)
+	// Regenerate the new pins from the old pin's PinData which may have changed
+	for (const UEdGraphPin* Pin : OldPins)
 	{
+		if (Pin->Direction != EGPD_Output)
+		{
+			continue;
+		}
+		
 		if (UCONodeMaterialBreakPinData* PinData = Cast<UCONodeMaterialBreakPinData>(GetPinData(*Pin)))
 		{
 			// Use the same pin data as the old pin since the new one is a copy
@@ -138,6 +143,9 @@ TArray<FName> UCONodeMaterialBreak::GetAllowedPinViewerCreationTypes() const
 void UCONodeMaterialBreak::AddPinFromUI()
 {
 	Super::AddPinFromUI();
+
+	FScopedTransaction Transaction(LOCTEXT("CONodeMaterialBreak_AddPinFromNodeUITransaction", "Add Pin from Node UI"));
+	Modify();
 	
 	const FName NewPinCategory = UEdGraphSchema_CustomizableObject::PC_Float;
 	
diff --git a/Engine/Plugins/Mutable/Source/CustomizableObjectEditor/Private/MuCOE/Nodes/CONodeSchema.cpp b/Engine/Plugins/Mutable/Source/CustomizableObjectEditor/Private/MuCOE/Nodes/CONodeSchema.cpp
--- a/Engine/Plugins/Mutable/Source/CustomizableObjectEditor/Private/MuCOE/Nodes/CONodeSchema.cpp
+++ b/Engine/Plugins/Mutable/Source/CustomizableObjectEditor/Private/MuCOE/Nodes/CONodeSchema.cpp
@@ -2,9 +2,10 @@
 
 #include "MuCOE/Nodes/CONodeSchema.h"
 
-#include "MuCOE/CustomizableObjectGraph.h"
 #include "MuCOE/RemapPins/CustomizableObjectNodeRemapPinsByPosition.h"
+#include "ScopedTransaction.h"
 
+#define LOCTEXT_NAMESPACE "CustomizableObjectNodeStringConstant"
 
 FText UCONodeSchema::GetNodeTitle(ENodeTitleType::Type TitleType) const
 {
@@ -67,7 +68,11 @@ FText UCONodeSchema::GetPinEditableName(const UEdGraphPin& Pin) const
 
 void UCONodeSchema::SetPinEditableName(const UEdGraphPin& Pin, const FText& Value)
 {
+	FScopedTransaction LocalTransaction(LOCTEXT("UCONodeSchema_SetPinEditableNameTransaction", "Change PinName"));
+	Modify();
+		
 	GetPinData<UCONodeSchemaPinData>(Pin).Name = Value;
+	PinEditableNameChangedDelegate.Broadcast(Pin, Value);
 }
 
 
@@ -90,3 +95,5 @@ void UCONodeSchema::PostEditChangeProperty(FPropertyChangedEvent& PropertyChange
 		}
 	}
 }
+
+#undef LOCTEXT_NAMESPACE
diff --git a/Engine/Plugins/Mutable/Source/CustomizableObjectEditor/Private/MuCOE/Nodes/CustomizableObjectNode.cpp b/Engine/Plugins/Mutable/Source/CustomizableObjectEditor/Private/MuCOE/Nodes/CustomizableObjectNode.cpp
--- a/Engine/Plugins/Mutable/Source/CustomizableObjectEditor/Private/MuCOE/Nodes/CustomizableObjectNode.cpp
+++ b/Engine/Plugins/Mutable/Source/CustomizableObjectEditor/Private/MuCOE/Nodes/CustomizableObjectNode.cpp
@@ -764,6 +764,18 @@ bool UCustomizableObjectNode::CanRenamePin(const UEdGraphPin& Pin) const
 }
 
 
+EEditablePinNameBoxVisibilityPolicy UCustomizableObjectNode::GetEditablePinNameVisibilityPolicy (const UEdGraphPin& Pin) const
+{
+	return EEditablePinNameBoxVisibilityPolicy::HIDE_IF_LINKED;
+}
+
+
+bool UCustomizableObjectNode::ShouldPinViewerShowPinEditableName(const UEdGraphPin& Pin) const
+{
+	return CanRenamePin(Pin);
+}
+
+
 FText UCustomizableObjectNode::GetPinEditableName(const UEdGraphPin& Pin) const
 {
 	return {};
@@ -772,6 +784,7 @@ FText UCustomizableObjectNode::GetPinEditableName(const UEdGraphPin& Pin) const
 
 void UCustomizableObjectNode::SetPinEditableName(const UEdGraphPin& Pin, const FText& Value)
 {
+	
 }
 
 
@@ -804,9 +817,6 @@ void UCustomizableObjectNode::AddPinFromUI()
 {
 	// You are not expected to call this method if not part of the "Add pin" in-node button operation. Use "CustomAddPin()" method instead.
 	ensure(GetAddPinButtonNodeSide() != EAddPinNodeButtonLocation::NONE);
-	
-	FScopedTransaction Transaction(LOCTEXT("AddPinFromNodeUITransaction", "Add Pin from Node UI"));
-	Modify();
 }
 
 
@@ -816,27 +826,6 @@ bool UCustomizableObjectNode::CanPinBeRemoved(const UEdGraphPin& Pin) const
 }
 
 
-bool UCustomizableObjectNode::TryRemovePin(UEdGraphPin& Pin)
-{
-	FScopedTransaction Transaction(LOCTEXT("TryRemovePinTransaction", "Try removing a Pin"));
-	Modify();
-
-	// Make sure the pin that is requested for removal was actually added by the UI.
-	if (CanPinBeRemoved(Pin) && CustomRemovePin(Pin))
-	{
-		// Notify that a pin has been removed so the PinViewer reacts to the new pin layout
-		if (UEdGraph* Graph = GetGraph())
-		{
-			Graph->NotifyNodeChanged(this);
-		}
-
-		return true;
-	}
-
-	return false;
-}
-
-
 TSharedPtr<IDetailsView> UCustomizableObjectNode::CustomizePinDetails(const UEdGraphPin& Pin) const
 {
 	return nullptr;
diff --git a/Engine/Plugins/Mutable/Source/CustomizableObjectEditor/Private/MuCOE/Nodes/CustomizableObjectNodeComponent.cpp b/Engine/Plugins/Mutable/Source/CustomizableObjectEditor/Private/MuCOE/Nodes/CustomizableObjectNodeComponent.cpp
--- a/Engine/Plugins/Mutable/Source/CustomizableObjectEditor/Private/MuCOE/Nodes/CustomizableObjectNodeComponent.cpp
+++ b/Engine/Plugins/Mutable/Source/CustomizableObjectEditor/Private/MuCOE/Nodes/CustomizableObjectNodeComponent.cpp
@@ -7,9 +7,11 @@
 #include "MuCOE/GraphTraversal.h"
 #include "MuCOE/Nodes/CustomizableObjectNodeMacroInstance.h"
 #include "MuCOE/Nodes/CustomizableObjectNodeStaticString.h"
+#include "ScopedTransaction.h"
 
 #include UE_INLINE_GENERATED_CPP_BY_NAME(CustomizableObjectNodeComponent)
 
+#define LOCTEXT_NAMESPACE "CustomizableObjectEditor"
 
 void UCustomizableObjectNodeComponent::AllocateDefaultPins(UCustomizableObjectNodeRemapPins* RemapPins)
 {
@@ -44,6 +46,19 @@ void UCustomizableObjectNodeComponent::BackwardsCompatibleFixup(int32 Customizab
 }
 
 
+bool UCustomizableObjectNodeComponent::ShouldPinViewerShowPinEditableName(const UEdGraphPin& Pin) const
+{
+	// Prevent the PinViewer from showing the internal PinEditableName if the pin has been connected (the name is externally driven by another node)
+	const UEdGraphPin* ComponentPin = GetComponentNamePin();
+	if (ComponentPin && ComponentPin->PinId == Pin.PinId)
+	{
+		return ComponentPin->LinkedTo.IsEmpty() && Super::ShouldPinViewerShowPinEditableName(Pin);
+	}
+	
+	return Super::ShouldPinViewerShowPinEditableName(Pin);
+}
+
+
 bool UCustomizableObjectNodeComponent::CanRenamePin(const UEdGraphPin& Pin) const
 {
 	return Pin.PinType.PinCategory == UEdGraphSchema_CustomizableObject::PC_String;
@@ -52,13 +67,31 @@ bool UCustomizableObjectNodeComponent::CanRenamePin(const UEdGraphPin& Pin) cons
 
 FText UCustomizableObjectNodeComponent::GetPinEditableName(const UEdGraphPin& Pin) const
 {
-	return FText::FromName(GetComponentName());
+	const UEdGraphPin* ComponentPin = GetComponentNamePin();
+	if (ComponentPin && ComponentPin->PinId == Pin.PinId)
+	{
+		return FText::FromName(GetComponentName());
+	}
+
+	return FText::GetEmpty();
 }
 
 
 void UCustomizableObjectNodeComponent::SetPinEditableName(const UEdGraphPin& Pin, const FText& InValue)
 {
-	SetComponentName(FName(*InValue.ToString()));
+	const UEdGraphPin* ComponentPin = GetComponentNamePin();
+	if (ComponentPin && ComponentPin->PinId == Pin.PinId)
+	{
+		FScopedTransaction LocalTransaction(LOCTEXT("UCustomizableObjectNodeComponent_SetComponentPinEditableNameTransaction", "Change ComponentName PinName"));
+		Modify();
+		
+		SetComponentName(FName(*InValue.ToString()));
+		PinEditableNameChangedDelegate.Broadcast(Pin, InValue);
+	}
+	else
+	{
+		Super::SetPinEditableName(Pin, InValue);
+	}
 }
 
 
@@ -82,9 +115,7 @@ void UCustomizableObjectNodeComponent::PinConnectionListChanged(UEdGraphPin* Pin
 
 FName UCustomizableObjectNodeComponent::GetComponentName(TArray<const UCustomizableObjectNodeMacroInstance*>* MacroContext) const
 {
-	const UEdGraphPin* NamePin = GetComponentNamePin();
-
-	if (NamePin)
+	if (const UEdGraphPin* NamePin = GetComponentNamePin())
 	{
 		if (const UEdGraphPin* LinkedPin = FollowInputPin(*NamePin))
 		{
@@ -111,3 +142,5 @@ UEdGraphPin* UCustomizableObjectNodeComponent::GetComponentNamePin() const
 {
 	return ComponentNamePin.Get();
 }
+
+#undef LOCTEXT_NAMESPACE
diff --git a/Engine/Plugins/Mutable/Source/CustomizableObjectEditor/Private/MuCOE/Nodes/CustomizableObjectNodeComponent.h b/Engine/Plugins/Mutable/Source/CustomizableObjectEditor/Private/MuCOE/Nodes/CustomizableObjectNodeComponent.h
--- a/Engine/Plugins/Mutable/Source/CustomizableObjectEditor/Private/MuCOE/Nodes/CustomizableObjectNodeComponent.h
+++ b/Engine/Plugins/Mutable/Source/CustomizableObjectEditor/Private/MuCOE/Nodes/CustomizableObjectNodeComponent.h
@@ -27,6 +27,7 @@ public:
 	// UCustomizableObjectNode interface
 	UE_API virtual bool IsAffectedByLOD() const override;
 	UE_API virtual void BackwardsCompatibleFixup(int32 CustomizableObjectCustomVersion) override;
+	UE_API virtual bool ShouldPinViewerShowPinEditableName(const UEdGraphPin& Pin) const override;
 	UE_API virtual bool CanRenamePin(const UEdGraphPin& Pin) const override;
 	UE_API virtual FText GetPinEditableName(const UEdGraphPin& Pin) const override;
 	UE_API virtual void SetPinEditableName(const UEdGraphPin& Pin, const FText& Value) override;
diff --git a/Engine/Plugins/Mutable/Source/CustomizableObjectEditor/Private/MuCOE/Nodes/CustomizableObjectNodeComponentMeshAddTo.cpp b/Engine/Plugins/Mutable/Source/CustomizableObjectEditor/Private/MuCOE/Nodes/CustomizableObjectNodeComponentMeshAddTo.cpp
--- a/Engine/Plugins/Mutable/Source/CustomizableObjectEditor/Private/MuCOE/Nodes/CustomizableObjectNodeComponentMeshAddTo.cpp
+++ b/Engine/Plugins/Mutable/Source/CustomizableObjectEditor/Private/MuCOE/Nodes/CustomizableObjectNodeComponentMeshAddTo.cpp
@@ -8,6 +8,7 @@
 #include "MuCOE/Nodes/CustomizableObjectNodeComponentMesh.h"
 #include "MuCOE/Nodes/CustomizableObjectNodeMacroInstance.h"
 #include "MuCOE/Nodes/CustomizableObjectNodeStaticString.h"
+#include "ScopedTransaction.h"
 
 #include UE_INLINE_GENERATED_CPP_BY_NAME(CustomizableObjectNodeComponentMeshAddTo)
 
@@ -85,13 +86,33 @@ bool UCustomizableObjectNodeComponentMeshAddTo::CanRenamePin(const UEdGraphPin&
 
 FText UCustomizableObjectNodeComponentMeshAddTo::GetPinEditableName(const UEdGraphPin& Pin) const
 {
-	return FText::FromName(GetParentComponentName());
+	const UEdGraphPin* ParentCompNamePin = GetParentComponentNamePin();
+	if (ParentCompNamePin && ParentCompNamePin->PinId == Pin.PinId)
+	{
+		return FText::FromName(GetParentComponentName());
+	}
+	else
+	{
+		return Super::GetPinEditableName(Pin);
+	}
 }
 
 
 void UCustomizableObjectNodeComponentMeshAddTo::SetPinEditableName(const UEdGraphPin& Pin, const FText& InValue)
 {
-	SetParentComponentName(FName(*InValue.ToString()));
+	const UEdGraphPin* ParentCompNamePin = GetParentComponentNamePin();
+	if (ParentCompNamePin && ParentCompNamePin->PinId == Pin.PinId)
+	{
+		FScopedTransaction LocalTransaction(LOCTEXT("UCustomizableObjectNodeComponent_SetParentComponentPinEditableNameTransaction", "Change ParentComponentName PinName"));
+		Modify();
+		
+		SetParentComponentName(FName(*InValue.ToString()));
+		PinEditableNameChangedDelegate.Broadcast(Pin, InValue);
+	}
+	else
+	{
+		Super::SetPinEditableName(Pin, InValue);
+	}
 }
 
 
diff --git a/Engine/Plugins/Mutable/Source/CustomizableObjectEditor/Private/MuCOE/Nodes/CustomizableObjectNodeComponentPassthroughMesh.cpp b/Engine/Plugins/Mutable/Source/CustomizableObjectEditor/Private/MuCOE/Nodes/CustomizableObjectNodeComponentPassthroughMesh.cpp
--- a/Engine/Plugins/Mutable/Source/CustomizableObjectEditor/Private/MuCOE/Nodes/CustomizableObjectNodeComponentPassthroughMesh.cpp
+++ b/Engine/Plugins/Mutable/Source/CustomizableObjectEditor/Private/MuCOE/Nodes/CustomizableObjectNodeComponentPassthroughMesh.cpp
@@ -2,13 +2,14 @@
 
 #include "MuCOE/Nodes/CustomizableObjectNodeComponentPassthroughMesh.h"
 
-#include "Engine/SkeletalMesh.h"
-#include "Modules/ModuleManager.h"
+#include "CONodeMaterialBreak.h"
 #include "MuCO/CustomizableObjectCustomVersion.h"
 #include "MuCO/LoadUtils.h"
 #include "MuCOE/EdGraphSchema_CustomizableObject.h"
-#include "MuCOE/Nodes/CustomizableObjectNodeLayoutBlocks.h"
 #include "MuCOE/UnrealEditorPortabilityHelpers.h"
+#include "ScopedTransaction.h"
+#include "MuCOE/RemapPins/CustomizableObjectNodeRemapPinsByPosition.h"
+#include "MVVM/ViewModels/OutlinerColumns/OutlinerColumnTypes.h"
 
 #include UE_INLINE_GENERATED_CPP_BY_NAME(CustomizableObjectNodeComponentPassthroughMesh)
 
@@ -44,6 +45,26 @@ void UCustomizableObjectNodeComponentPassthroughMesh::AllocateDefaultPins(UCusto
 	Super::AllocateDefaultPins(RemapPins);
 
 	SkeletalMeshPin = CustomCreatePin(EGPD_Input, UEdGraphSchema_CustomizableObject::PC_PassthroughSkeletalMesh, TEXT("Skeletal Mesh"));
+
+	// Get old pins
+	const TArray<UEdGraphPin*> AllPins = GetAllNonOrphanPins();
+	
+	// Regenerate the new pins from the old pins
+	MaterialSlotOverridePins.Reset();
+	for (const UEdGraphPin* OldPin : AllPins)
+	{
+		if (OldPin->Direction != EEdGraphPinDirection::EGPD_Input)
+		{
+			continue;
+		}
+		
+		if (UCONodeComponentPassthroughMeshMaterialOverridePinData* PinData = Cast<UCONodeComponentPassthroughMeshMaterialOverridePinData>(GetPinData(*OldPin)))
+		{
+			// Use the same pin data as the old pin since the new one is a copy
+			const UEdGraphPin* NewPin = CreateMaterialSlotOverridePin(PinData);
+			MaterialSlotOverridePins.Add(NewPin);
+		}
+	}
 }
 
 
@@ -53,6 +74,142 @@ bool UCustomizableObjectNodeComponentPassthroughMesh::IsExperimental() const
 }
 
 
+bool UCustomizableObjectNodeComponentPassthroughMesh::HasPinViewer() const
+{
+	return true;
+}
+
+
+EAddPinNodeButtonLocation UCustomizableObjectNodeComponentPassthroughMesh::GetAddPinButtonNodeSide() const
+{
+	return EAddPinNodeButtonLocation::INPUT;
+}
+
+
+void UCustomizableObjectNodeComponentPassthroughMesh::AddPinFromUI()
+{
+	Super::AddPinFromUI();
+
+	FScopedTransaction Transaction(LOCTEXT("UCustomizableObjectNodeComponentPassthroughMesh_AddPinFromNodeUITransaction", "Add Pin from Node UI"));
+	Modify();
+
+	UCONodeComponentPassthroughMeshMaterialOverridePinData* PinData = NewObject<UCONodeComponentPassthroughMeshMaterialOverridePinData>(this);
+	
+	if (const UEdGraphPin* NewPin = CreateMaterialSlotOverridePin(PinData))
+	{		
+		MaterialSlotOverridePins.Add(NewPin);
+		GetGraph()->NotifyNodeChanged(this);
+	}
+}
+
+
+bool UCustomizableObjectNodeComponentPassthroughMesh::CanPinBeRemoved(const UEdGraphPin& Pin) const
+{
+	const bool bCanPinBeRemoved = !Pin.LinkedTo.Num() && !Pin.bOrphanedPin;
+	if (!bCanPinBeRemoved)
+	{
+		return false;
+	}
+
+	return IsMaterialSlotOverridePin(Pin);
+}
+
+
+
+bool UCustomizableObjectNodeComponentPassthroughMesh::CanRenamePin(const UEdGraphPin& Pin) const
+{
+	return IsMaterialSlotOverridePin(Pin) || Super::CanRenamePin(Pin);
+}
+
+
+EEditablePinNameBoxVisibilityPolicy UCustomizableObjectNodeComponentPassthroughMesh::GetEditablePinNameVisibilityPolicy(const UEdGraphPin& Pin) const
+{
+	if (IsMaterialSlotOverridePin(Pin))
+	{
+		return EEditablePinNameBoxVisibilityPolicy::ALWAYS_VISIBLE;
+	}
+
+	return Super::GetEditablePinNameVisibilityPolicy(Pin);
+}
+
+
+bool UCustomizableObjectNodeComponentPassthroughMesh::ShouldPinViewerShowPinEditableName(const UEdGraphPin& Pin) const
+{
+	if (IsMaterialSlotOverridePin(Pin))
+	{
+		return true;
+	}
+
+	return Super::ShouldPinViewerShowPinEditableName(Pin);
+}
+
+
+void UCustomizableObjectNodeComponentPassthroughMesh::SetPinEditableName(const UEdGraphPin& Pin, const FText& Value)
+{
+	if (IsMaterialSlotOverridePin(Pin))
+	{
+		if (UCONodeComponentPassthroughMeshMaterialOverridePinData* PinData = Cast<UCONodeComponentPassthroughMeshMaterialOverridePinData>(GetPinData(Pin)))
+		{
+			const FName NewName = FName{Value.ToString()};
+			if (NewName != PinData->TargetMaterialSlotName)		
+			{
+				FScopedTransaction LocalTransaction(LOCTEXT("UCustomizableObjectNodeComponentPassthroughMesh_SetOverrideMaterialPinEditableNameTransaction", "Change OverrideMaterial PinName"));
+				PinData->Modify();
+				
+				PinData->TargetMaterialSlotName = NewName;
+
+				// Tell the listeners that this pin editable name has changed
+				PinEditableNameChangedDelegate.Broadcast(Pin, Value);
+			}
+		}
+	}
+	else
+	{
+		Super::SetPinEditableName(Pin, Value);
+	}
+}
+
+
+FText UCustomizableObjectNodeComponentPassthroughMesh::GetPinEditableName(const UEdGraphPin& Pin) const
+{
+	if (IsMaterialSlotOverridePin(Pin))
+	{
+		if (const UCONodeComponentPassthroughMeshMaterialOverridePinData* PinData = Cast<UCONodeComponentPassthroughMeshMaterialOverridePinData>(GetPinData(Pin)))
+		{
+			return FText::FromName(PinData->TargetMaterialSlotName);
+		}
+	}
+	
+	return Super::GetPinEditableName(Pin);	
+}
+
+
+bool UCustomizableObjectNodeComponentPassthroughMesh::IsMaterialSlotOverridePin(const UEdGraphPin& Pin) const
+{
+	return Pin.Direction == EEdGraphPinDirection::EGPD_Input && Pin.PinType.PinCategory == UEdGraphSchema_CustomizableObject::PC_Material;
+}
+
+
+UCustomizableObjectNodeRemapPins* UCustomizableObjectNodeComponentPassthroughMesh::CreateRemapPinsDefault() const
+{
+	return Super::CreateRemapPinsByPosition();
+}
+
+
+bool UCustomizableObjectNodeComponentPassthroughMesh::CustomRemovePin(UEdGraphPin& Pin)
+{
+	MaterialSlotOverridePins.Remove(&Pin);
+	return Super::CustomRemovePin(Pin);
+}
+
+
+UEdGraphPin* UCustomizableObjectNodeComponentPassthroughMesh::CreateMaterialSlotOverridePin(UCONodeComponentPassthroughMeshMaterialOverridePinData* InPinData)
+{
+	return CustomCreatePin(EEdGraphPinDirection::EGPD_Input, UEdGraphSchema_CustomizableObject::PC_Material,
+	"Override Material", InPinData);
+}
+
+
 FText UCustomizableObjectNodeComponentPassthroughMesh::GetNodeTitle(ENodeTitleType::Type TitleType) const
 {
 	return LOCTEXT("ComponentPassthroughMesh", "Passthrough Skeletal Mesh Component");
diff --git a/Engine/Plugins/Mutable/Source/CustomizableObjectEditor/Private/MuCOE/Nodes/CustomizableObjectNodeComponentPassthroughMesh.h b/Engine/Plugins/Mutable/Source/CustomizableObjectEditor/Private/MuCOE/Nodes/CustomizableObjectNodeComponentPassthroughMesh.h
--- a/Engine/Plugins/Mutable/Source/CustomizableObjectEditor/Private/MuCOE/Nodes/CustomizableObjectNodeComponentPassthroughMesh.h
+++ b/Engine/Plugins/Mutable/Source/CustomizableObjectEditor/Private/MuCOE/Nodes/CustomizableObjectNodeComponentPassthroughMesh.h
@@ -6,8 +6,6 @@
 
 #include "CustomizableObjectNodeComponentPassthroughMesh.generated.h"
 
-#define UE_API CUSTOMIZABLEOBJECTEDITOR_API
-
 namespace ENodeTitleType { enum Type : int; }
 
 class FArchive;
@@ -25,6 +23,20 @@ class USkeletalMesh;
 class UTexture2D;
 struct FPropertyChangedEvent;
 struct FSkeletalMaterial;
+class UEdGraphPin;
+
+
+UCLASS(MinimalAPI)
+class UCONodeComponentPassthroughMeshMaterialOverridePinData : public UCustomizableObjectNodePinData
+{
+	GENERATED_BODY()
+
+public:
+
+	UPROPERTY(EditAnywhere, Category = NoCategory)
+	FName TargetMaterialSlotName;
+};
+
 
 
 UCLASS(MinimalAPI)
@@ -34,19 +46,39 @@ public:
 	GENERATED_BODY()
 	
 	// UObject interface.
-	UE_API virtual void BackwardsCompatibleFixup(int32 CustomizableObjectCustomVersion) override;
+	virtual void BackwardsCompatibleFixup(int32 CustomizableObjectCustomVersion) override;
 	
 	// UEdGraphNode interface
-	UE_API virtual FText GetNodeTitle(ENodeTitleType::Type TitleType) const override;
-	UE_API virtual FText GetTooltipText() const override;
+	virtual FText GetNodeTitle(ENodeTitleType::Type TitleType) const override;
+	virtual FText GetTooltipText() const override;
 
 	// UCustomizableObjectNode interface
-	UE_API virtual void AllocateDefaultPins(UCustomizableObjectNodeRemapPins* RemapPins) override;
-	UE_API virtual bool IsExperimental() const override;
+	virtual void AllocateDefaultPins(UCustomizableObjectNodeRemapPins* RemapPins) override;
+	virtual bool IsExperimental() const override;
+	
+	virtual bool HasPinViewer() const override;
+	virtual EAddPinNodeButtonLocation GetAddPinButtonNodeSide() const override;
+	virtual void AddPinFromUI() override;
+	virtual bool CanPinBeRemoved(const UEdGraphPin& Pin) const override;
 
+	virtual bool CanRenamePin(const UEdGraphPin& Pin) const override;
+	virtual EEditablePinNameBoxVisibilityPolicy GetEditablePinNameVisibilityPolicy(const UEdGraphPin& Pin) const override;
+	virtual bool ShouldPinViewerShowPinEditableName(const UEdGraphPin& Pin) const override;
+	virtual void SetPinEditableName(const UEdGraphPin& Pin, const FText& Value) override;
+	virtual FText GetPinEditableName(const UEdGraphPin& Pin) const override;
+	virtual bool CustomRemovePin(UEdGraphPin& Pin) override;
+	
 	UPROPERTY()
 	FEdGraphPinReference SkeletalMeshPin;
-};
 
+	UPROPERTY()
+	TArray<FEdGraphPinReference> MaterialSlotOverridePins;
+
+private:
+	UEdGraphPin* CreateMaterialSlotOverridePin(UCONodeComponentPassthroughMeshMaterialOverridePinData* InPinData);
 
-#undef UE_API
+	bool IsMaterialSlotOverridePin(const UEdGraphPin& Pin) const;
+
+protected:
+	virtual UCustomizableObjectNodeRemapPins* CreateRemapPinsDefault() const override;
+};
diff --git a/Engine/Plugins/Mutable/Source/CustomizableObjectEditor/Private/MuCOE/Nodes/CustomizableObjectNodeObject.cpp b/Engine/Plugins/Mutable/Source/CustomizableObjectEditor/Private/MuCOE/Nodes/CustomizableObjectNodeObject.cpp
--- a/Engine/Plugins/Mutable/Source/CustomizableObjectEditor/Private/MuCOE/Nodes/CustomizableObjectNodeObject.cpp
+++ b/Engine/Plugins/Mutable/Source/CustomizableObjectEditor/Private/MuCOE/Nodes/CustomizableObjectNodeObject.cpp
@@ -16,13 +16,10 @@
 #include "MuCOE/Nodes/CustomizableObjectNodeMaterialVariation.h"
 #include "Misc/UObjectToken.h"
 #include "Logging/MessageLog.h"
-#include "Containers/Queue.h"
 #include "MuCO/CustomizableObjectCustomVersion.h"
 #include "MuCOE/CustomizableObjectEditor_Deprecated.h"
 #include "MuCOE/Nodes/CustomizableObjectNodeComponentMeshAddTo.h"
-#include "MuCOE/Nodes/CustomizableObjectNodeCopyMaterial.h"
 #include "MuCOE/Nodes/CustomizableObjectNodeMaterialSwitch.h"
-#include "MuCOE/CustomizableObjectMacroLibrary/CustomizableObjectMacroLibrary.h"
 #include "Settings/EditorStyleSettings.h"
 
 #include UE_INLINE_GENERATED_CPP_BY_NAME(CustomizableObjectNodeObject)
@@ -627,13 +624,31 @@ bool UCustomizableObjectNodeObject::CanRenamePin(const UEdGraphPin& Pin) const
 
 FText UCustomizableObjectNodeObject::GetPinEditableName(const UEdGraphPin& Pin) const
 {
-	return FText::FromString(GetObjectName());
+	if (NamePin.Get()->PinId == Pin.PinId)
+	{
+		return FText::FromString(GetObjectName());
+	}
+	else
+	{
+		return Super::GetPinEditableName(Pin);
+	}
 }
 
 
 void UCustomizableObjectNodeObject::SetPinEditableName(const UEdGraphPin& Pin, const FText& InValue)
 {
-	SetObjectName(InValue.ToString());
+	if (NamePin.Get()->PinId == Pin.PinId)
+	{
+		FScopedTransaction LocalTransaction(LOCTEXT("UCustomizableObjectNodeObject_SetNamePinEditableNameTransaction", "Change NamePin PinName"));
+		Modify();
+		
+		SetObjectName(InValue.ToString());
+		PinEditableNameChangedDelegate.Broadcast(Pin, InValue);
+	}
+	else
+	{
+		Super::SetPinEditableName(Pin, InValue);
+	}
 }
 
 
diff --git a/Engine/Plugins/Mutable/Source/CustomizableObjectEditor/Private/MuCOE/Nodes/CustomizableObjectNodeObjectGroup.cpp b/Engine/Plugins/Mutable/Source/CustomizableObjectEditor/Private/MuCOE/Nodes/CustomizableObjectNodeObjectGroup.cpp
--- a/Engine/Plugins/Mutable/Source/CustomizableObjectEditor/Private/MuCOE/Nodes/CustomizableObjectNodeObjectGroup.cpp
+++ b/Engine/Plugins/Mutable/Source/CustomizableObjectEditor/Private/MuCOE/Nodes/CustomizableObjectNodeObjectGroup.cpp
@@ -3,11 +3,11 @@
 #include "MuCOE/Nodes/CustomizableObjectNodeObjectGroup.h"
 
 #include "MuCO/CustomizableObjectCustomVersion.h"
-#include "MuCOE/CustomizableObjectMacroLibrary/CustomizableObjectMacroLibrary.h"
 #include "MuCOE/EdGraphSchema_CustomizableObject.h"
 #include "MuCOE/GraphTraversal.h"
 #include "MuCOE/Nodes/CustomizableObjectNodeObject.h"
 #include "MuCOE/Nodes/CustomizableObjectNodeStaticString.h"
+#include "ScopedTransaction.h"
 
 #include UE_INLINE_GENERATED_CPP_BY_NAME(CustomizableObjectNodeObjectGroup)
 
@@ -167,13 +167,31 @@ bool UCustomizableObjectNodeObjectGroup::CanRenamePin(const UEdGraphPin& Pin) co
 
 FText UCustomizableObjectNodeObjectGroup::GetPinEditableName(const UEdGraphPin& Pin) const
 {
-	return FText::FromString(GetGroupName());
+	if (NamePin.Get()->PinId == Pin.PinId)
+	{
+		return FText::FromString(GetGroupName());
+	}
+	else
+	{
+		return Super::GetPinEditableName(Pin);
+	}
 }
 
 
 void UCustomizableObjectNodeObjectGroup::SetPinEditableName(const UEdGraphPin& Pin, const FText& InValue)
 {
-	SetGroupName(InValue.ToString());
+	if (NamePin.Get()->PinId == Pin.PinId)
+	{
+		FScopedTransaction LocalTransaction(LOCTEXT("UCustomizableObjectNodeObjectGroup_SetGroupNamePinEditableNameTransaction", "Change Group Name PinName"));
+		Modify();
+		
+		SetGroupName(InValue.ToString());
+		PinEditableNameChangedDelegate.Broadcast(Pin, InValue);
+	}
+	else
+	{
+		Super::SetPinEditableName(Pin, InValue);
+	}
 }
 
 
diff --git a/Engine/Plugins/Mutable/Source/CustomizableObjectEditor/Private/MuCOE/Nodes/CustomizableObjectNodeParameter.cpp b/Engine/Plugins/Mutable/Source/CustomizableObjectEditor/Private/MuCOE/Nodes/CustomizableObjectNodeParameter.cpp
--- a/Engine/Plugins/Mutable/Source/CustomizableObjectEditor/Private/MuCOE/Nodes/CustomizableObjectNodeParameter.cpp
+++ b/Engine/Plugins/Mutable/Source/CustomizableObjectEditor/Private/MuCOE/Nodes/CustomizableObjectNodeParameter.cpp
@@ -4,10 +4,10 @@
 
 #include "MuCO/CustomizableObjectCustomVersion.h"
 #include "MuCOE/EdGraphSchema_CustomizableObject.h"
-#include "MuCOE/CustomizableObjectMacroLibrary/CustomizableObjectMacroLibrary.h"
 #include "MuCOE/GraphTraversal.h"
 #include "MuCOE/Nodes/CustomizableObjectNodeMacroInstance.h"
 #include "MuCOE/Nodes/CustomizableObjectNodeStaticString.h"
+#include "ScopedTransaction.h"
 
 #include UE_INLINE_GENERATED_CPP_BY_NAME(CustomizableObjectNodeParameter)
 
@@ -58,13 +58,31 @@ bool UCustomizableObjectNodeParameter::CanRenamePin(const UEdGraphPin& Pin) cons
 
 FText UCustomizableObjectNodeParameter::GetPinEditableName(const UEdGraphPin& Pin) const
 {
-	return FText::FromString(GetParameterName());
+	if (NamePin.Get()->PinId == Pin.PinId)
+	{
+		return FText::FromString(GetParameterName());
+	}
+	else
+	{
+		return Super::GetPinEditableName(Pin);
+	}
 }
 
 
 void UCustomizableObjectNodeParameter::SetPinEditableName(const UEdGraphPin& Pin, const FText& InValue)
 {
-	SetParameterName(InValue.ToString());
+	if (NamePin.Get()->PinId == Pin.PinId)
+	{
+		FScopedTransaction LocalTransaction(LOCTEXT("UCustomizableObjectNodeParameter_SetParameterName", "Change ParameterName PinName"));
+		Modify();
+		
+		SetParameterName(InValue.ToString());
+		PinEditableNameChangedDelegate.Broadcast(Pin, InValue);
+	}
+	else
+	{
+		Super::SetPinEditableName(Pin, InValue);
+	}
 }
 
 
diff --git a/Engine/Plugins/Mutable/Source/CustomizableObjectEditor/Private/MuCOE/Nodes/CustomizableObjectNodeStaticString.cpp b/Engine/Plugins/Mutable/Source/CustomizableObjectEditor/Private/MuCOE/Nodes/CustomizableObjectNodeStaticString.cpp
--- a/Engine/Plugins/Mutable/Source/CustomizableObjectEditor/Private/MuCOE/Nodes/CustomizableObjectNodeStaticString.cpp
+++ b/Engine/Plugins/Mutable/Source/CustomizableObjectEditor/Private/MuCOE/Nodes/CustomizableObjectNodeStaticString.cpp
@@ -4,6 +4,7 @@
 
 #include "MuCOE/EdGraphSchema_CustomizableObject.h"
 #include "MuCOE/Nodes/SCustomizableObjectNode.h"
+#include "ScopedTransaction.h"
 
 #include UE_INLINE_GENERATED_CPP_BY_NAME(CustomizableObjectNodeStaticString)
 
@@ -53,7 +54,11 @@ FText UCustomizableObjectNodeStaticString::GetPinEditableName(const UEdGraphPin&
 
 void UCustomizableObjectNodeStaticString::SetPinEditableName(const UEdGraphPin& Pin, const FText& InValue)
 {
+	FScopedTransaction LocalTransaction(LOCTEXT("UCustomizableObjectNodeStaticString_SetStringPinEditableNameTransaction", "Change Literal String PinName"));
+	Modify();
+	
 	Value = InValue.ToString();
+	PinEditableNameChangedDelegate.Broadcast(Pin, InValue);
 }
 
 
diff --git a/Engine/Plugins/Mutable/Source/CustomizableObjectEditor/Private/MuCOE/Nodes/CustomizableObjectNodeVariation.cpp b/Engine/Plugins/Mutable/Source/CustomizableObjectEditor/Private/MuCOE/Nodes/CustomizableObjectNodeVariation.cpp
--- a/Engine/Plugins/Mutable/Source/CustomizableObjectEditor/Private/MuCOE/Nodes/CustomizableObjectNodeVariation.cpp
+++ b/Engine/Plugins/Mutable/Source/CustomizableObjectEditor/Private/MuCOE/Nodes/CustomizableObjectNodeVariation.cpp
@@ -3,10 +3,10 @@
 #include "MuCOE/Nodes/CustomizableObjectNodeVariation.h"
 
 #include "MuCO/CustomizableObjectCustomVersion.h"
-#include "MuCOE/CustomizableObjectMacroLibrary/CustomizableObjectMacroLibrary.h"
 #include "MuCOE/EdGraphSchema_CustomizableObject.h"
 #include "MuCOE/GraphTraversal.h"
 #include "MuCOE/Nodes/CustomizableObjectNodeStaticString.h"
+#include "ScopedTransaction.h"
 
 #include UE_INLINE_GENERATED_CPP_BY_NAME(CustomizableObjectNodeVariation)
 
@@ -88,8 +88,11 @@ void UCustomizableObjectNodeVariation::SetPinEditableName(const UEdGraphPin& Pin
 	{
 		if (VariationTagPin(VariationIndex) == &Pin && VariationsData.IsValidIndex(VariationIndex))
 		{
-			VariationsData[VariationIndex].Tag = InValue.ToString();
+			FScopedTransaction LocalTransaction(LOCTEXT("UCustomizableObjectNodeVariation_SetVariationTagPinEditableNameTransaction", "Change Variation Tag PinName"));
+			Modify();
 
+			VariationsData[VariationIndex].Tag = InValue.ToString();
+			PinEditableNameChangedDelegate.Broadcast(Pin, InValue);
 			break;
 		}
 	}
diff --git a/Engine/Plugins/Mutable/Source/CustomizableObjectEditor/Private/MuCOE/Nodes/SCustomizableObjectNodePin.cpp b/Engine/Plugins/Mutable/Source/CustomizableObjectEditor/Private/MuCOE/Nodes/SCustomizableObjectNodePin.cpp
--- a/Engine/Plugins/Mutable/Source/CustomizableObjectEditor/Private/MuCOE/Nodes/SCustomizableObjectNodePin.cpp
+++ b/Engine/Plugins/Mutable/Source/CustomizableObjectEditor/Private/MuCOE/Nodes/SCustomizableObjectNodePin.cpp
@@ -33,10 +33,23 @@ TSharedRef<SWidget> SCustomizableObjectNodePin::GetDefaultValueWidget()
 		UCustomizableObjectNode* Node = CastChecked<UCustomizableObjectNode>(GetPinObj()->GetOwningNode());
 		if (Node && Node->CanRenamePin(*GraphPin))
 		{
-			return SNew(SEditableTextBox)
-				.Text(this, &SCustomizableObjectNodePin::GetNodeStringValue)
-				.OnTextCommitted(this, &SCustomizableObjectNodePin::OnTextCommited)
-				.Visibility(this, &SCustomizableObjectNodePin::GetWidgetVisibility);
+			switch (Node->GetEditablePinNameVisibilityPolicy(*GraphPin))
+			{
+			case EEditablePinNameBoxVisibilityPolicy::ALWAYS_VISIBLE:
+				{
+					return SNew(SEditableTextBox)
+								.Text(this, &SCustomizableObjectNodePin::GetNodeStringValue)
+								.OnTextCommitted(this, &SCustomizableObjectNodePin::OnTextCommited)
+								.Visibility(EVisibility::Visible);
+				}
+			case EEditablePinNameBoxVisibilityPolicy::HIDE_IF_LINKED:
+				{
+					return SNew(SEditableTextBox)
+						.Text(this, &SCustomizableObjectNodePin::GetNodeStringValue)
+						.OnTextCommitted(this, &SCustomizableObjectNodePin::OnTextCommited)
+						.Visibility(this, &SCustomizableObjectNodePin::GetHideIfLinkedWidgetVisibility);
+				}
+			}
 		}	
 	}
 
@@ -132,7 +145,7 @@ void SCustomizableObjectNodePin::OnTextCommited(const FText& InValue, ETextCommi
 }
 
 
-EVisibility SCustomizableObjectNodePin::GetWidgetVisibility() const
+EVisibility SCustomizableObjectNodePin::GetHideIfLinkedWidgetVisibility() const
 {
 	return GraphPinObj->LinkedTo.Num() ? EVisibility::Collapsed : EVisibility::Visible;
 }
diff --git a/Engine/Plugins/Mutable/Source/CustomizableObjectEditor/Private/MuCOE/Nodes/SCustomizableObjectNodePin.h b/Engine/Plugins/Mutable/Source/CustomizableObjectEditor/Private/MuCOE/Nodes/SCustomizableObjectNodePin.h
--- a/Engine/Plugins/Mutable/Source/CustomizableObjectEditor/Private/MuCOE/Nodes/SCustomizableObjectNodePin.h
+++ b/Engine/Plugins/Mutable/Source/CustomizableObjectEditor/Private/MuCOE/Nodes/SCustomizableObjectNodePin.h
@@ -28,7 +28,7 @@ private:
 	
 	void OnTextCommited(const FText& InValue, ETextCommit::Type InCommitInfo);
 	
-	EVisibility GetWidgetVisibility() const;
+	EVisibility GetHideIfLinkedWidgetVisibility() const;
 
 
 	const FSlateBrush* PassThroughImageConnected = nullptr;
diff --git a/Engine/Plugins/Mutable/Source/CustomizableObjectEditor/Private/MuCOE/PinViewer/SPinViewer.cpp b/Engine/Plugins/Mutable/Source/CustomizableObjectEditor/Private/MuCOE/PinViewer/SPinViewer.cpp
--- a/Engine/Plugins/Mutable/Source/CustomizableObjectEditor/Private/MuCOE/PinViewer/SPinViewer.cpp
+++ b/Engine/Plugins/Mutable/Source/CustomizableObjectEditor/Private/MuCOE/PinViewer/SPinViewer.cpp
@@ -27,14 +27,31 @@ const FName SPinViewer::COLUMN_VISIBILITY(TEXT("Visible"));
 
 FText SPinViewer::GetPinName(const UEdGraphPin& Pin)
 {
+	FString PinComposedName;
+	
 	if (!Pin.PinFriendlyName.IsEmpty())
 	{
-		return Pin.PinFriendlyName;
+		PinComposedName = Pin.PinFriendlyName.ToString();
 	}
 	else
 	{
-		return FText::FromName(Pin.PinName);
+		PinComposedName =  Pin.PinName.ToString();
+	}
+
+	// Append Pin Editable Name
+	if (const UCustomizableObjectNode* OwningNode = Cast<UCustomizableObjectNode>(Pin.GetOwningNode()))
+	{
+		if (OwningNode->ShouldPinViewerShowPinEditableName(Pin))
+		{
+			const FString EditablePinName = OwningNode->GetPinEditableName(Pin).ToString();
+			if (!EditablePinName.IsEmpty())
+			{
+				PinComposedName = FString::Printf(TEXT("%s : %s"), *PinComposedName, *EditablePinName);
+			}
+		}
 	}
+	
+	return FText::FromString(PinComposedName);
 }
 
 
@@ -187,6 +204,7 @@ void SPinViewer::Construct(const FArguments& InArgs)
 	Node = TStrongObjectPtr<UCustomizableObjectNode>{InArgs._Node};
 	Node->RemapPinsDelegate.AddSP(this, &SPinViewer::PinsRemapped);
 	Node->NodeConnectionListChangedDelegate.AddSP(this, &SPinViewer::UpdateWidget);
+	Node->PinEditableNameChangedDelegate.AddSP(this, &SPinViewer::OnPinEditableNameChange);
 	GeneratePinInfoList();
 
 	Graph = TStrongObjectPtr(Node->GetGraph());
@@ -275,6 +293,13 @@ void SPinViewer::UpdateWidget()
 }
 
 
+void SPinViewer::OnPinEditableNameChange(const UEdGraphPin& EdGraphPin, FText Text)
+{
+	// TODO: OPTIMIZATION : it would be nice to take into consideration what pin has changed so we do not update the whole widget but only the affected row.
+	UpdateWidget();
+}
+
+
 void SPinViewer::OnGraphChanged(const FEdGraphEditAction& EdGraphEditAction)
 {
 	if (EdGraphEditAction.Action == GRAPHACTION_EditNode && EdGraphEditAction.Nodes.Contains(Node.Get()))
diff --git a/Engine/Plugins/Mutable/Source/CustomizableObjectEditor/Private/MuCOE/PinViewer/SPinViewer.h b/Engine/Plugins/Mutable/Source/CustomizableObjectEditor/Private/MuCOE/PinViewer/SPinViewer.h
--- a/Engine/Plugins/Mutable/Source/CustomizableObjectEditor/Private/MuCOE/PinViewer/SPinViewer.h
+++ b/Engine/Plugins/Mutable/Source/CustomizableObjectEditor/Private/MuCOE/PinViewer/SPinViewer.h
@@ -61,6 +61,9 @@ public:
 	/** Get the pin name which is actually displayed. */
 	static FText GetPinName(const UEdGraphPin& Pin);
 
+	/** Callback method invoked each time a pin of the node we keep a reference of does get the Editable Name changed. */
+	void OnPinEditableNameChange(const UEdGraphPin& EdGraphPin, FText Text);
+	
 private:
 	/** Regenerate the list contents. */
 	void GeneratePinInfoList();
diff --git a/Engine/Plugins/Mutable/Source/CustomizableObjectEditor/Public/MuCOE/Nodes/CustomizableObjectNode.h b/Engine/Plugins/Mutable/Source/CustomizableObjectEditor/Public/MuCOE/Nodes/CustomizableObjectNode.h
--- a/Engine/Plugins/Mutable/Source/CustomizableObjectEditor/Public/MuCOE/Nodes/CustomizableObjectNode.h
+++ b/Engine/Plugins/Mutable/Source/CustomizableObjectEditor/Public/MuCOE/Nodes/CustomizableObjectNode.h
@@ -40,6 +40,19 @@ enum class EAddPinNodeButtonLocation : uint8
 };
 
 
+/**
+ * Determines the visibility behaviour for the editable pin name text boxes
+ */
+enum class EEditablePinNameBoxVisibilityPolicy : uint8
+{
+	/** Default behaviour : Any linked pin will have it's editable textbox hidden. */
+	HIDE_IF_LINKED = 0,
+
+	/** Always show the text box. */
+	ALWAYS_VISIBLE
+};
+
+
 /** Abstract base class for all pin data. */
 UCLASS(Abstract)
 class UCustomizableObjectNodePinData : public UObject
@@ -88,6 +101,8 @@ DECLARE_MULTICAST_DELEGATE(FNodeConnectionListChangedDelegate);
 
 DECLARE_MULTICAST_DELEGATE_OneParam(FPinConnectionListChangedDelegate, UEdGraphPin*);
 
+DECLARE_MULTICAST_DELEGATE_TwoParams(FPostEditablePinNameChangedDelegate, const UEdGraphPin&, const FText NewContent);
+
 DECLARE_MULTICAST_DELEGATE(FDestroyNodeDelegate);
 
 using FRemapPinsDelegateParameter = TMap<UEdGraphPin*, UEdGraphPin*>; // Required for the delegate macro.
@@ -281,11 +296,18 @@ public:
 	 */
 	UE_API void SetPinHidden(const TArray<UEdGraphPin*>& Pins, bool bHidden);
 
-	/** Return true if a pin can be hidden. Override it to have a custom behaviour. */
+	/** Return true if a pin can be hidden. */
 	UE_API virtual bool CanPinBeHidden(const UEdGraphPin& Pin) const;
 
+	/** Determines if a pin can be renamed or not which determines if the pin will spawn with a TextBox slate at the side of the pin. */
 	UE_API virtual bool CanRenamePin(const UEdGraphPin& Pin) const;
 
+	/** Determines the policy to use when drawing the pin slate. */
+	UE_API virtual EEditablePinNameBoxVisibilityPolicy GetEditablePinNameVisibilityPolicy (const UEdGraphPin& Pin) const;
+
+	/** Determines if the Pin Viewer is allowed to include the Pin Editable Name as part of the PinName string shown. */
+	UE_API virtual bool ShouldPinViewerShowPinEditableName(const UEdGraphPin& Pin) const;
+
 	UE_API virtual FText GetPinEditableName(const UEdGraphPin& Pin) const;
 
 	UE_API virtual void SetPinEditableName(const UEdGraphPin& Pin, const FText& Value);
@@ -313,13 +335,6 @@ public:
 	 */
 	UE_API virtual bool CanPinBeRemoved(const UEdGraphPin& Pin) const;
 	
-	/**
-	 * Removes one pin that is known to have been added by the Node "Add Pin" UI button.
-	 * @param Pin The pin that you want to get removed.
-	 * @return Will return true if the removal could be performed, false otherwise.
-	 */
-	UE_API bool TryRemovePin(UEdGraphPin& Pin);
-	
 	/** Returns pin custom details. Override if required. */
 	UE_API virtual TSharedPtr<IDetailsView> CustomizePinDetails(const UEdGraphPin& Pin) const;
 
@@ -358,6 +373,8 @@ public:
 
 	FPinConnectionListChangedDelegate PinConnectionListChangedDelegate;
 
+	FPostEditablePinNameChangedDelegate PinEditableNameChangedDelegate;
+	
 	FDestroyNodeDelegate DestroyNodeDelegate;
 
 	FRemapPinsDelegate RemapPinsDelegate;
@@ -375,7 +392,7 @@ protected:
 	/**
 	 * Creates and returns the node default node remap pins action.
 	 * 
-	 * By default all nodes remap by name. Override if a node requires a diferent default remap pins action.
+	 * By default all nodes remap by name. Override if a node requires a different default remap pins action.
 	 */
 	UE_API virtual UCustomizableObjectNodeRemapPins* CreateRemapPinsDefault() const;
 
diff --git a/Engine/Plugins/Mutable/Source/MutableTools/Internal/MuT/NodeComponentNew.h b/Engine/Plugins/Mutable/Source/MutableTools/Internal/MuT/NodeComponentNew.h
--- a/Engine/Plugins/Mutable/Source/MutableTools/Internal/MuT/NodeComponentNew.h
+++ b/Engine/Plugins/Mutable/Source/MutableTools/Internal/MuT/NodeComponentNew.h
@@ -24,6 +24,9 @@ namespace UE::Mutable::Private
 		Ptr<NodeMaterial> OverlayMaterial;
 
 		Ptr<NodeSkeletalMesh> SkeletalMesh;
+
+		/** Map where the key is the target Material slot and the value the replacement material to be used in the target material slot*/
+		TMap<const FName, Ptr<NodeMaterial>> OverrideMaterials;
 		
 		//-----------------------------------------------------------------------------------------
         // Node interface
