[mutable] Added support for Output pin Custom Visual Widget.
- Require to be able to have Text boxes in the output pins.
- Removed hidden StaticStringNode input pin as now we will be using the editable textbox spawned by the output pin.
- Corrected some code indentations

#rnx

diff --git a/Engine/Plugins/Mutable/Source/CustomizableObject/Internal/MuCO/CustomizableObjectCustomVersion.h b/Engine/Plugins/Mutable/Source/CustomizableObject/Internal/MuCO/CustomizableObjectCustomVersion.h
--- a/Engine/Plugins/Mutable/Source/CustomizableObject/Internal/MuCO/CustomizableObjectCustomVersion.h
+++ b/Engine/Plugins/Mutable/Source/CustomizableObject/Internal/MuCO/CustomizableObjectCustomVersion.h
@@ -210,6 +210,8 @@ struct FCustomizableObjectCustomVersion
 
 		PinSkeletalMesh,
 
+		NodeStaticStringRemoveHiddenInputPin,
+
 		// -----<new versions can be added above this line>-------------------------------------------------
 		VersionPlusOne,
 		LatestVersion = VersionPlusOne - 1
diff --git a/Engine/Plugins/Mutable/Source/CustomizableObjectEditor/Private/MuCOE/Nodes/CustomizableObjectNodeStaticString.cpp b/Engine/Plugins/Mutable/Source/CustomizableObjectEditor/Private/MuCOE/Nodes/CustomizableObjectNodeStaticString.cpp
--- a/Engine/Plugins/Mutable/Source/CustomizableObjectEditor/Private/MuCOE/Nodes/CustomizableObjectNodeStaticString.cpp
+++ b/Engine/Plugins/Mutable/Source/CustomizableObjectEditor/Private/MuCOE/Nodes/CustomizableObjectNodeStaticString.cpp
@@ -5,6 +5,7 @@
 #include "MuCOE/EdGraphSchema_CustomizableObject.h"
 #include "MuCOE/Nodes/SCustomizableObjectNode.h"
 #include "ScopedTransaction.h"
+#include "MuCO/CustomizableObjectCustomVersion.h"
 
 #include UE_INLINE_GENERATED_CPP_BY_NAME(CustomizableObjectNodeStaticString)
 
@@ -31,18 +32,13 @@ FLinearColor UCustomizableObjectNodeStaticString::GetNodeTitleColor() const
 
 void UCustomizableObjectNodeStaticString::AllocateDefaultPins(UCustomizableObjectNodeRemapPins* RemapPins)
 {
-	const UEdGraphSchema_CustomizableObject* Schema = GetDefault<UEdGraphSchema_CustomizableObject>();
-
-	CustomCreatePin(EGPD_Output, Schema->PC_String, FName("Value"));
-
-	UEdGraphPin* StringPin = CustomCreatePin(EGPD_Input, Schema->PC_String, FName("String"));
-	StringPin->bNotConnectable = true;
+	CustomCreatePin(EGPD_Output, UEdGraphSchema_CustomizableObject::PC_String, FName("Value"));
 }
 
 
 bool UCustomizableObjectNodeStaticString::CanRenamePin(const UEdGraphPin& Pin) const
 {
-	return Pin.PinType.PinCategory == UEdGraphSchema_CustomizableObject::PC_String;
+	return true;
 }
 
 
@@ -62,4 +58,24 @@ void UCustomizableObjectNodeStaticString::SetPinEditableName(const UEdGraphPin&
 }
 
 
+void UCustomizableObjectNodeStaticString::BackwardsCompatibleFixup(int32 CustomizableObjectCustomVersion)
+{
+	Super::BackwardsCompatibleFixup(CustomizableObjectCustomVersion);
+
+	if (CustomizableObjectCustomVersion == FCustomizableObjectCustomVersion::NodeStaticStringRemoveHiddenInputPin)
+	{
+		if (UEdGraphPin* OldInputPin = FindPin(TEXT("String")))
+		{
+			CustomRemovePin(*OldInputPin);
+		}
+	}
+}
+
+
+EEditablePinNameBoxVisibilityPolicy UCustomizableObjectNodeStaticString::GetEditablePinNameVisibilityPolicy(const UEdGraphPin& Pin) const
+{
+	return EEditablePinNameBoxVisibilityPolicy::ALWAYS_VISIBLE;
+}
+
+
 #undef LOCTEXT_NAMESPACE
diff --git a/Engine/Plugins/Mutable/Source/CustomizableObjectEditor/Private/MuCOE/Nodes/CustomizableObjectNodeStaticString.h b/Engine/Plugins/Mutable/Source/CustomizableObjectEditor/Private/MuCOE/Nodes/CustomizableObjectNodeStaticString.h
--- a/Engine/Plugins/Mutable/Source/CustomizableObjectEditor/Private/MuCOE/Nodes/CustomizableObjectNodeStaticString.h
+++ b/Engine/Plugins/Mutable/Source/CustomizableObjectEditor/Private/MuCOE/Nodes/CustomizableObjectNodeStaticString.h
@@ -15,16 +15,18 @@ class UCustomizableObjectNodeStaticString : public UCustomizableObjectNode
 
 public:
 	// EdGraphNode interface
-	UE_API virtual FText GetNodeTitle(ENodeTitleType::Type TitleType) const override;
-	UE_API virtual FText GetTooltipText() const override;
-	UE_API virtual FLinearColor GetNodeTitleColor() const override;
+	virtual FText GetNodeTitle(ENodeTitleType::Type TitleType) const override;
+	virtual FText GetTooltipText() const override;
+	virtual FLinearColor GetNodeTitleColor() const override;
 	
 	// UCustomizableObjectNode interface
-	UE_API virtual void AllocateDefaultPins(UCustomizableObjectNodeRemapPins* RemapPins) override;
-	UE_API virtual bool CanRenamePin(const UEdGraphPin& Pin) const override;
-	UE_API virtual FText GetPinEditableName(const UEdGraphPin& Pin) const override;
-	UE_API virtual void SetPinEditableName(const UEdGraphPin& Pin, const FText& Value) override;
-	
+	virtual void AllocateDefaultPins(UCustomizableObjectNodeRemapPins* RemapPins) override;
+	virtual bool CanRenamePin(const UEdGraphPin& Pin) const override;
+	virtual FText GetPinEditableName(const UEdGraphPin& Pin) const override;
+	virtual void SetPinEditableName(const UEdGraphPin& Pin, const FText& Value) override;
+	virtual void BackwardsCompatibleFixup(int32 CustomizableObjectCustomVersion) override;
+	virtual EEditablePinNameBoxVisibilityPolicy GetEditablePinNameVisibilityPolicy(const UEdGraphPin& Pin) const override;
+
 	UPROPERTY()
 	FString Value;
 };
diff --git a/Engine/Plugins/Mutable/Source/CustomizableObjectEditor/Private/MuCOE/Nodes/SCustomizableObjectNodePin.cpp b/Engine/Plugins/Mutable/Source/CustomizableObjectEditor/Private/MuCOE/Nodes/SCustomizableObjectNodePin.cpp
--- a/Engine/Plugins/Mutable/Source/CustomizableObjectEditor/Private/MuCOE/Nodes/SCustomizableObjectNodePin.cpp
+++ b/Engine/Plugins/Mutable/Source/CustomizableObjectEditor/Private/MuCOE/Nodes/SCustomizableObjectNodePin.cpp
@@ -10,18 +10,39 @@
 
 void SCustomizableObjectNodePin::Construct(const FArguments& InArgs, UEdGraphPin* InGraphPinObj)
 {
-	SGraphPin::Construct(SGraphPin::FArguments(), InGraphPinObj);
+	const SGraphPin::FArguments SGraphPinArguments = SGraphPin::FArguments();
+	SGraphPin::Construct(SGraphPinArguments, InGraphPinObj);
+
+	// Snippet from SGraphPin::Construct(SGraphPin::FArguments(), InGraphPinObj); of the creation and setup of the ValueWidget
+	// Since the base class does no longer support having an editable text box for the outputs of the node we do create our own
+	if (InGraphPinObj->Direction == EEdGraphPinDirection::EGPD_Output)
+	{
+		ValueWidget = GetDefaultValueWidget();
+
+		if (ValueWidget != SNullWidget::NullWidget)
+		{
+			TSharedPtr<SBox> ValueBox;
+			LabelAndValue->AddSlot()
+				.Padding(FMargin(SGraphPinArguments._SideToSideMargin, 0.0f, 0, 0.0f))
+				.VAlign(VAlign_Center)
+				[
+					SAssignNew(ValueBox, SBox)
+					.Padding(0.0f)
+					[
+						ValueWidget.ToSharedRef()
+					]
+				];
+
+			if (!DoesWidgetHandleSettingEditingEnabled())
+			{
+				ValueBox->SetEnabled(TAttribute<bool>(this, &SGraphPin::IsEditingEnabled));
+			}
+		}
+	}
 
 	// Cache pin icons.
 	PassThroughImageConnected = UE_MUTABLE_GET_BRUSH(TEXT("Graph.ExecPin.Connected"));
 	PassThroughImageDisconnected = UE_MUTABLE_GET_BRUSH(TEXT("Graph.ExecPin.Disconnected"));
-
-	// Hide pin Icon and Label
-	if (InGraphPinObj->bNotConnectable)
-	{
-		SetShowLabel(false);
-		PinImage->SetVisibility(EVisibility::Collapsed);
-	}
 }
 
 
