[Mutable] Deduplicate Materials

Equal FMaterials were stored as different constants.

diff --git a/Engine/Plugins/Mutable/Source/CustomizableObject/Internal/MuCO/CustomizableObjectPrivate.h b/Engine/Plugins/Mutable/Source/CustomizableObject/Internal/MuCO/CustomizableObjectPrivate.h
--- a/Engine/Plugins/Mutable/Source/CustomizableObject/Internal/MuCO/CustomizableObjectPrivate.h
+++ b/Engine/Plugins/Mutable/Source/CustomizableObject/Internal/MuCO/CustomizableObjectPrivate.h
@@ -1752,6 +1752,8 @@ enum class ECustomizableObjectVersions : int32
 
 	RelevantParameters,
 
+	DeduplicateFMaterials,
+
 	// -----<new versions can be added above this line>--------
 	LastCustomizableObjectVersion
 };
diff --git a/Engine/Plugins/Mutable/Source/MutableRuntime/Internal/MuR/ModelPrivate.h b/Engine/Plugins/Mutable/Source/MutableRuntime/Internal/MuR/ModelPrivate.h
--- a/Engine/Plugins/Mutable/Source/MutableRuntime/Internal/MuR/ModelPrivate.h
+++ b/Engine/Plugins/Mutable/Source/MutableRuntime/Internal/MuR/ModelPrivate.h
@@ -693,10 +693,19 @@ namespace UE::Mutable::Private
             return index;
         }
 
-		OP::ADDRESS AddConstant(TSharedPtr<const FMaterial> Material)
+		OP::ADDRESS AddConstant(const TSharedPtr<const FMaterial>& Material)
 		{
-			OP::ADDRESS index = OP::ADDRESS(ConstantMaterials.AddUnique(Material));
-			return index;
+			for (int32 Index=0; Index<ConstantMaterials.Num(); ++Index)
+			{
+				if (ConstantMaterials[Index]== Material || *ConstantMaterials[Index] == *Material)
+				{
+					return static_cast<OP::ADDRESS>(Index);
+				}
+			}
+
+			OP::ADDRESS Index = static_cast<OP::ADDRESS>(ConstantMaterials.Num());
+			ConstantMaterials.Add(Material);
+			return Index;
 		}
 
 		inline TSharedPtr<const FImage> GetImageLOD(FConstantResourceIndex Index) const
