[Mutable] Add bone pose priority to resolve conflicts when merging meshes with different poses.

#rnx

diff --git a/Engine/Plugins/Mutable/Source/CustomizableObject/Internal/MuCO/CustomizableObjectPrivate.h b/Engine/Plugins/Mutable/Source/CustomizableObject/Internal/MuCO/CustomizableObjectPrivate.h
--- a/Engine/Plugins/Mutable/Source/CustomizableObject/Internal/MuCO/CustomizableObjectPrivate.h
+++ b/Engine/Plugins/Mutable/Source/CustomizableObject/Internal/MuCO/CustomizableObjectPrivate.h
@@ -1754,6 +1754,8 @@ enum class ECustomizableObjectVersions : int32
 
 	DeduplicateFMaterials,
 
+	AddBonePosePriority,
+
 	// -----<new versions can be added above this line>--------
 	LastCustomizableObjectVersion
 };
diff --git a/Engine/Plugins/Mutable/Source/CustomizableObject/Internal/MuCO/UnrealToMutableTextureConversionUtils.h b/Engine/Plugins/Mutable/Source/CustomizableObject/Internal/MuCO/UnrealToMutableTextureConversionUtils.h
--- a/Engine/Plugins/Mutable/Source/CustomizableObject/Internal/MuCO/UnrealToMutableTextureConversionUtils.h
+++ b/Engine/Plugins/Mutable/Source/CustomizableObject/Internal/MuCO/UnrealToMutableTextureConversionUtils.h
@@ -111,6 +111,7 @@ struct FMutableSourceMeshData
 
 	/** */
 	int32 SocketPriority = 0;
+	int16 BonePosePriority = 0;
 
 	/** */
 	bool bOnlyConnectedLOD = false;
@@ -131,6 +132,7 @@ inline uint32 GetTypeHash(const FMutableSourceMeshData& Key)
 	GuidHash = HashCombineFast(GuidHash, GetTypeHash(Key.BaseSectionIndex));
 	GuidHash = HashCombineFast(GuidHash, GetTypeHash(Key.LODOffset));
 	GuidHash = HashCombineFast(GuidHash, GetTypeHash(Key.SocketPriority));
+	GuidHash = HashCombineFast(GuidHash, GetTypeHash(Key.BonePosePriority));
 	GuidHash = HashCombineFast(GuidHash, GetTypeHash(Key.bOnlyConnectedLOD));
 	GuidHash = HashCombineFast(GuidHash, GetTypeHash(Key.Component));
 	GuidHash = HashCombineFast(GuidHash, GetTypeHash(Key.Flags));
diff --git a/Engine/Plugins/Mutable/Source/CustomizableObjectEditor/Private/MuCOE/GenerateMutableSource/GenerateMutableSource.cpp b/Engine/Plugins/Mutable/Source/CustomizableObjectEditor/Private/MuCOE/GenerateMutableSource/GenerateMutableSource.cpp
--- a/Engine/Plugins/Mutable/Source/CustomizableObjectEditor/Private/MuCOE/GenerateMutableSource/GenerateMutableSource.cpp
+++ b/Engine/Plugins/Mutable/Source/CustomizableObjectEditor/Private/MuCOE/GenerateMutableSource/GenerateMutableSource.cpp
@@ -270,6 +270,9 @@ FMutableGraphGenerationContext::FMutableGraphGenerationContext(const TSharedPtr<
 
 	// Default SocketPriority for mesh generation.
 	SocketPriorityStack.Push(0);
+
+	// Default BonePosePriority for mesh generation.
+	BonePosePriorityStack.Push(0);
 }
 
 
@@ -1045,6 +1048,9 @@ UE::Mutable::Private::Ptr<UE::Mutable::Private::NodeObject> GenerateMutableSourc
 		// All sockets from all mesh parts plugged into this group node will have the following priority when there's a socket name clash
 		GenerationContext.SocketPriorityStack.Push(TypedNodeGroup->SocketPriority);
 
+		// All bone poses from all mesh parts plugged into this group node will have the following priority when there's a EBoneUsageFlags clash
+		GenerationContext.BonePosePriorityStack.Push(TypedNodeGroup->BonePosePriority);
+
 		// Getting node Id. We may need a different ID if the group node is inside a Macro and this Macro is used more than once.
 		// Group nodes of Macros Can Not have external children so we can change their ids here safetely.
 		FGuid GroupNodeId = GenerationContext.GetNodeIdUnchecked(TypedNodeGroup);
@@ -1342,6 +1348,10 @@ UE::Mutable::Private::Ptr<UE::Mutable::Private::NodeObject> GenerateMutableSourc
 		// Go back to the parent group node's socket priority if it exists
 		ensure(GenerationContext.SocketPriorityStack.Num() > 0);
 		GenerationContext.SocketPriorityStack.Pop();
+
+		// Go back to the parent group node's bone pose priority if it exists
+		ensure(GenerationContext.BonePosePriorityStack.Num() > 0);
+		GenerationContext.BonePosePriorityStack.Pop();
 	}
 
 	else if (const UCustomizableObjectNodeMacroInstance* TypedNodeMacro = Cast<UCustomizableObjectNodeMacroInstance>(Node))
diff --git a/Engine/Plugins/Mutable/Source/CustomizableObjectEditor/Private/MuCOE/GenerateMutableSource/GenerateMutableSource.h b/Engine/Plugins/Mutable/Source/CustomizableObjectEditor/Private/MuCOE/GenerateMutableSource/GenerateMutableSource.h
--- a/Engine/Plugins/Mutable/Source/CustomizableObjectEditor/Private/MuCOE/GenerateMutableSource/GenerateMutableSource.h
+++ b/Engine/Plugins/Mutable/Source/CustomizableObjectEditor/Private/MuCOE/GenerateMutableSource/GenerateMutableSource.h
@@ -892,6 +892,11 @@ struct FMutableGraphGenerationContext
 	*/
 	TArray<int32> SocketPriorityStack;
 	
+	/** Used to propagate the BonePose priority defined in group nodes to their child skeletal mesh nodes
+	* It's a stack because group nodes are recursive
+	*/
+	TArray<int16> BonePosePriorityStack;
+
 	// Stores what param names use a certain table as a table can be used from multiple table nodes, useful for partial compilations to restrict params
 	TMap<FString, FMutableParamNameSet> TableToParamNames;
 
diff --git a/Engine/Plugins/Mutable/Source/CustomizableObjectEditor/Private/MuCOE/GenerateMutableSource/GenerateMutableSourceMesh.cpp b/Engine/Plugins/Mutable/Source/CustomizableObjectEditor/Private/MuCOE/GenerateMutableSource/GenerateMutableSourceMesh.cpp
--- a/Engine/Plugins/Mutable/Source/CustomizableObjectEditor/Private/MuCOE/GenerateMutableSource/GenerateMutableSourceMesh.cpp
+++ b/Engine/Plugins/Mutable/Source/CustomizableObjectEditor/Private/MuCOE/GenerateMutableSource/GenerateMutableSourceMesh.cpp
@@ -1467,6 +1467,8 @@ UE::Tasks::FTask ConvertSkeletalMeshToMutable(TSharedRef<FMeshConversionContext>
 		}
 
 		MutableMesh->SetBoneMap(MutableBoneMap);
+
+		MutableMesh->BonePosePriority = Source.BonePosePriority;
 	}
 
 	if (!bIgnoreSkinning)
@@ -2651,6 +2653,7 @@ TSharedPtr<UE::Mutable::Private::FMesh> ConvertSkeletalMeshToMutable(
 	Source.MessageContext = CurrentNode;
 	Source.Component = GenerationContext.CurrentMeshComponent;
 	Source.SocketPriority = GenerationContext.SocketPriorityStack.Last();
+	Source.BonePosePriority = GenerationContext.BonePosePriorityStack.Last();
 
 	// Prepare source data for realtime morph targets
 	const UCustomizableObjectNodeSkeletalMesh* NodeTypedSkMesh = Cast<UCustomizableObjectNodeSkeletalMesh>(CurrentNode);
diff --git a/Engine/Plugins/Mutable/Source/CustomizableObjectEditor/Private/MuCOE/Nodes/CustomizableObjectNodeObjectGroup.h b/Engine/Plugins/Mutable/Source/CustomizableObjectEditor/Private/MuCOE/Nodes/CustomizableObjectNodeObjectGroup.h
--- a/Engine/Plugins/Mutable/Source/CustomizableObjectEditor/Private/MuCOE/Nodes/CustomizableObjectNodeObjectGroup.h
+++ b/Engine/Plugins/Mutable/Source/CustomizableObjectEditor/Private/MuCOE/Nodes/CustomizableObjectNodeObjectGroup.h
@@ -35,6 +35,11 @@ public:
 	UPROPERTY(EditAnywhere, Category = UI, meta = (DisplayName = "Parameter UI Metadata"))
 	FMutableParamUIMetadata ParamUIMetadata;
 
+	// The pose defined in meshes deriving from this node will inherit this bone pose priority. Decides the bone pose to keep when two bones
+	// that have the same bone usage flags (E.g. EBoneUsageFlags::Skinning). The one with higher priority will be kept.
+	UPROPERTY(EditAnywhere, Category = MeshPose)
+	int16 BonePosePriority = 0;
+
 	// The sockets defined in meshes deriving from this node will inherit this socket priority. When in the generated merged mesh there
 	// are clashes with socket names, the one with higher priority will be kept and the other discarded.
 	UPROPERTY(EditAnywhere, Category = MeshSockets)
diff --git a/Engine/Plugins/Mutable/Source/MutableRuntime/Internal/MuR/Mesh.h b/Engine/Plugins/Mutable/Source/MutableRuntime/Internal/MuR/Mesh.h
--- a/Engine/Plugins/Mutable/Source/MutableRuntime/Internal/MuR/Mesh.h
+++ b/Engine/Plugins/Mutable/Source/MutableRuntime/Internal/MuR/Mesh.h
@@ -488,6 +488,10 @@ namespace UE::Mutable::Private
 			FBoneName BoneId;
 
 			EBoneUsageFlags BoneUsageFlags = EBoneUsageFlags::None;
+			
+			// Not serialized. Add Metadata may overwrite it. 
+			int16 BonePosePriority = 0;
+
 			FTransform3f BoneTransform;
 
 			inline void Serialise(FOutputArchive& arch) const;
@@ -497,6 +501,7 @@ namespace UE::Mutable::Private
 			{
 				return BoneUsageFlags == Other.BoneUsageFlags &&
 				BoneId == Other.BoneId &&
+				BonePosePriority == Other.BonePosePriority &&
 				BoneTransform.Equals(Other.BoneTransform);
 			}
 		};
@@ -507,6 +512,9 @@ namespace UE::Mutable::Private
 		 */
 		TMemoryTrackedArray<FBonePose> BonePoses;
 
+		/** Priority used to decide which bone pose to keep when merging two meshes with same BoneUsageFlags. Serialized in the Metadata args */
+		int16 BonePosePriority = 0;
+
 		/** Array containing the bonemaps of all surfaces in the mesh. */
 		TArray<FBoneName> BoneMap;
     	
@@ -564,6 +572,7 @@ namespace UE::Mutable::Private
 				bEqual &= AdditionalBuffers[i] == Other.AdditionalBuffers[i];
 			}
 
+			bEqual &= BonePosePriority == Other.BonePosePriority;
 			bEqual &= BonePoses.Num() == Other.BonePoses.Num();
 			for (int32 i = 0; bEqual && i < BonePoses.Num(); ++i)
 			{
diff --git a/Engine/Plugins/Mutable/Source/MutableRuntime/Internal/MuR/OpMeshMerge.h b/Engine/Plugins/Mutable/Source/MutableRuntime/Internal/MuR/OpMeshMerge.h
--- a/Engine/Plugins/Mutable/Source/MutableRuntime/Internal/MuR/OpMeshMerge.h
+++ b/Engine/Plugins/Mutable/Source/MutableRuntime/Internal/MuR/OpMeshMerge.h
@@ -816,10 +816,16 @@ namespace UE::Mutable::Private
 							(EnumHasAnyFlags(BonePose.BoneUsageFlags, EBoneUsageFlags::Reshaped) ? 1 : 0);
 					};
 
-					if (ComputeBoneMergePriority(ResultBonePose) < ComputeBoneMergePriority(SecondBonePose))
+					const int32 ResultBonePriority = ComputeBoneMergePriority(ResultBonePose);
+					const int32 SecondBonePriority = ComputeBoneMergePriority(SecondBonePose);
+
+					if (ResultBonePriority < SecondBonePriority ||
+						(ResultBonePriority == SecondBonePriority && ResultBonePose.BonePosePriority < SecondBonePose.BonePosePriority))
 					{
 						//ResultBonePose.BoneName = SecondBonePose.BoneName;
 						ResultBonePose.BoneTransform = SecondBonePose.BoneTransform;
+						ResultBonePose.BonePosePriority = SecondBonePose.BonePosePriority;
+
 						// Merge usage flags
 						EnumAddFlags(ResultBonePose.BoneUsageFlags, SecondBonePose.BoneUsageFlags);
 					}
diff --git a/Engine/Plugins/Mutable/Source/MutableRuntime/Internal/MuR/Operations.h b/Engine/Plugins/Mutable/Source/MutableRuntime/Internal/MuR/Operations.h
--- a/Engine/Plugins/Mutable/Source/MutableRuntime/Internal/MuR/Operations.h
+++ b/Engine/Plugins/Mutable/Source/MutableRuntime/Internal/MuR/Operations.h
@@ -1056,7 +1056,7 @@ namespace UE::Mutable::Private
          
         struct MeshAddMetadataArgs
         {
-            enum class EnumFlags : uint32
+            enum class EnumFlags : uint16
             {
                 None           = 0,
                 IsTagList      = 1 << 0,
@@ -1083,6 +1083,7 @@ namespace UE::Mutable::Private
             };
 
             EnumFlags Flags; 
+			int16 BonePosePriority;
             ADDRESS Source;
             TagAddressOrUInt32ListAddress Tags;
             SkeletonIdOrUInt32ListAddress SkeletonIds;
diff --git a/Engine/Plugins/Mutable/Source/MutableRuntime/Private/MuR/CodeRunner.cpp b/Engine/Plugins/Mutable/Source/MutableRuntime/Private/MuR/CodeRunner.cpp
--- a/Engine/Plugins/Mutable/Source/MutableRuntime/Private/MuR/CodeRunner.cpp
+++ b/Engine/Plugins/Mutable/Source/MutableRuntime/Private/MuR/CodeRunner.cpp
@@ -3191,6 +3191,16 @@ namespace UE::Mutable::Private
 					{
 						Result->SkeletonIDs.Add(OpArgs.SkeletonIds.SkeletonId);
 					}
+
+					if (Result->BonePosePriority != OpArgs.BonePosePriority)
+					{
+						Result->BonePosePriority = OpArgs.BonePosePriority;
+
+						for (FMesh::FBonePose& BonePose : Result->BonePoses)
+						{
+							BonePose.BonePosePriority = OpArgs.BonePosePriority;
+						}
+					}
 				}
 
 				StoreMesh(Item, Result);
diff --git a/Engine/Plugins/Mutable/Source/MutableRuntime/Private/MuR/Mesh.cpp b/Engine/Plugins/Mutable/Source/MutableRuntime/Private/MuR/Mesh.cpp
--- a/Engine/Plugins/Mutable/Source/MutableRuntime/Private/MuR/Mesh.cpp
+++ b/Engine/Plugins/Mutable/Source/MutableRuntime/Private/MuR/Mesh.cpp
@@ -157,6 +157,7 @@ void FMesh::CopyFrom(const FMesh& From, EMeshCopyFlags InFlags)
 	if (EnumHasAnyFlags(InFlags, EMeshCopyFlags::WithPoses))
 	{
 		BonePoses = From.BonePoses;
+		BonePosePriority = From.BonePosePriority;
 	}
 
 	// Copy BoneMap
@@ -569,7 +570,7 @@ void UE::Mutable::Private::FMesh::SetBonePose(int32 Index, const FBoneName& Bone
 	check(BonePoses.IsValidIndex(Index));
 	if (BonePoses.IsValidIndex(Index))
 	{
-		BonePoses[Index] = FBonePose{ BoneId, BoneUsageFlags, Transform };
+		BonePoses[Index] = FBonePose{ BoneId, BoneUsageFlags, 0, Transform };
 	}
 }
 
diff --git a/Engine/Plugins/Mutable/Source/MutableRuntime/Private/MuR/OpMeshFormat.cpp b/Engine/Plugins/Mutable/Source/MutableRuntime/Private/MuR/OpMeshFormat.cpp
--- a/Engine/Plugins/Mutable/Source/MutableRuntime/Private/MuR/OpMeshFormat.cpp
+++ b/Engine/Plugins/Mutable/Source/MutableRuntime/Private/MuR/OpMeshFormat.cpp
@@ -710,6 +710,7 @@ namespace UE::Mutable::Private
 		Result->AdditionalBuffers = Source->AdditionalBuffers;
 
 		Result->BonePoses = Source->BonePoses;
+		Result->BonePosePriority = Source->BonePosePriority;
 		Result->BoneMap = Source->BoneMap;
 
 		Result->SkeletonIDs = Source->SkeletonIDs;
diff --git a/Engine/Plugins/Mutable/Source/MutableTools/Private/MuT/ASTOpMeshAddMetadata.cpp b/Engine/Plugins/Mutable/Source/MutableTools/Private/MuT/ASTOpMeshAddMetadata.cpp
--- a/Engine/Plugins/Mutable/Source/MutableTools/Private/MuT/ASTOpMeshAddMetadata.cpp
+++ b/Engine/Plugins/Mutable/Source/MutableTools/Private/MuT/ASTOpMeshAddMetadata.cpp
@@ -34,7 +34,8 @@ namespace UE::Mutable::Private
 				Source      == Other->Source      && 
 				Tags        == Other->Tags        && 
 				ResourceIds == Other->ResourceIds &&
-				SkeletonIds == Other->SkeletonIds;
+				SkeletonIds == Other->SkeletonIds &&
+				BonePosePriority == Other->BonePosePriority;
 		}
 
 		return false;
@@ -48,6 +49,7 @@ namespace UE::Mutable::Private
 		NewOp->Tags = Tags;
 		NewOp->ResourceIds = ResourceIds;
 		NewOp->SkeletonIds = SkeletonIds;
+		NewOp->BonePosePriority = BonePosePriority;
 
 		return NewOp;
 	}
@@ -118,6 +120,8 @@ namespace UE::Mutable::Private
 				Args.SkeletonIds.SkeletonId = SkeletonIds[0];
 			}
 
+			Args.BonePosePriority = BonePosePriority;
+
 			Program.OpAddress.Add((uint32)Program.ByteCode.Num());
 			AppendCode(Program.ByteCode, EOpType::ME_ADDMETADATA);
 			AppendCode(Program.ByteCode, Args);
diff --git a/Engine/Plugins/Mutable/Source/MutableTools/Private/MuT/ASTOpMeshAddMetadata.h b/Engine/Plugins/Mutable/Source/MutableTools/Private/MuT/ASTOpMeshAddMetadata.h
--- a/Engine/Plugins/Mutable/Source/MutableTools/Private/MuT/ASTOpMeshAddMetadata.h
+++ b/Engine/Plugins/Mutable/Source/MutableTools/Private/MuT/ASTOpMeshAddMetadata.h
@@ -29,6 +29,9 @@ namespace UE::Mutable::Private
 		/** SkeletonIds to add. */
 		TArray<uint32> SkeletonIds;
 
+		/** Bone Pose priority when merging two meshes. */
+		int16 BonePosePriority = 0;
+
 	public:
 
 		ASTOpMeshAddMetadata();
diff --git a/Engine/Plugins/Mutable/Source/MutableTools/Private/MuT/CodeGenerator_Mesh.cpp b/Engine/Plugins/Mutable/Source/MutableTools/Private/MuT/CodeGenerator_Mesh.cpp
--- a/Engine/Plugins/Mutable/Source/MutableTools/Private/MuT/CodeGenerator_Mesh.cpp
+++ b/Engine/Plugins/Mutable/Source/MutableTools/Private/MuT/CodeGenerator_Mesh.cpp
@@ -814,8 +814,9 @@ namespace UE::Mutable::Private
 				TArray<FString> Tags = ReferenceMesh->Tags;
 				TArray<uint64> ResourceIds = ReferenceMesh->StreamedResources;
 				TArray<uint32> SkeletonIds = ReferenceMesh->SkeletonIDs;
+				int32 BonePosePriority = 0;
 
-				if (Mesh->Tags.Num() + Mesh->StreamedResources.Num() + Mesh->SkeletonIDs.Num() > 0)
+				if (Mesh->Tags.Num() + Mesh->StreamedResources.Num() + Mesh->SkeletonIDs.Num() > 0 || Mesh->BonePosePriority != 0)
 				{
 					for (FString& Tag : Mesh->Tags)
 					{
@@ -832,11 +833,14 @@ namespace UE::Mutable::Private
 						SkeletonIds.AddUnique(SkeletonId);
 					}
 
+					BonePosePriority = Mesh->BonePosePriority;
+
 					TSharedPtr<FMesh> MetadatalessMesh = CloneOrTakeOver(Mesh);
 
 					MetadatalessMesh->Tags.SetNum(0, EAllowShrinking::No);
 					MetadatalessMesh->StreamedResources.SetNum(0, EAllowShrinking::No);
 					MetadatalessMesh->SkeletonIDs.SetNum(0, EAllowShrinking::No);
+					MetadatalessMesh->BonePosePriority = 0;
 
 					Mesh = MetadatalessMesh;
 				}
@@ -908,7 +912,7 @@ namespace UE::Mutable::Private
 				Result.BaseMeshOp = LastMeshOp;
 
 				// Add the metadata operation
-				if (Tags.Num() + ResourceIds.Num() + SkeletonIds.Num() > 0)
+				if (Tags.Num() + ResourceIds.Num() + SkeletonIds.Num() > 0 || BonePosePriority != 0)
 				{
 					Ptr<ASTOpMeshAddMetadata> AddMetadataOp = new ASTOpMeshAddMetadata;
 					AddMetadataOp->Source = LastMeshOp;
@@ -916,6 +920,7 @@ namespace UE::Mutable::Private
 					AddMetadataOp->Tags        = MoveTemp(Tags);
 					AddMetadataOp->ResourceIds = MoveTemp(ResourceIds);
 					AddMetadataOp->SkeletonIds = MoveTemp(SkeletonIds);
+					AddMetadataOp->BonePosePriority = BonePosePriority;
 
 					LastMeshOp = AddMetadataOp;
 				}
