[Mutable] Remove material layer hack.
+ The material layer index used to be encoded inside the parameter name. Now is sotored inside the ByteCode.
+ This is a previews commit to a major task. In the future this extre bytes will be stored inside the material parameters (UE-351597)

#rnx

diff --git a/Engine/Plugins/Mutable/Source/CustomizableObject/Internal/MuCO/CustomizableObjectPrivate.h b/Engine/Plugins/Mutable/Source/CustomizableObject/Internal/MuCO/CustomizableObjectPrivate.h
--- a/Engine/Plugins/Mutable/Source/CustomizableObject/Internal/MuCO/CustomizableObjectPrivate.h
+++ b/Engine/Plugins/Mutable/Source/CustomizableObject/Internal/MuCO/CustomizableObjectPrivate.h
@@ -1756,6 +1756,8 @@ enum class ECustomizableObjectVersions : int32
 
 	AddBonePosePriority,
 
+	RemovedMaterialLayerEncoding,
+
 	// -----<new versions can be added above this line>--------
 	LastCustomizableObjectVersion
 };
diff --git a/Engine/Plugins/Mutable/Source/CustomizableObject/Internal/MuCO/CustomizableObjectSystemPrivate.h b/Engine/Plugins/Mutable/Source/CustomizableObject/Internal/MuCO/CustomizableObjectSystemPrivate.h
--- a/Engine/Plugins/Mutable/Source/CustomizableObject/Internal/MuCO/CustomizableObjectSystemPrivate.h
+++ b/Engine/Plugins/Mutable/Source/CustomizableObject/Internal/MuCO/CustomizableObjectSystemPrivate.h
@@ -392,18 +392,24 @@ struct FInstanceUpdateData
 
 		bool bIsPassThrough = false;
 		bool bIsNonProgressive = false;
+
+		int32 MaterialLayer;
 	};
 
 	struct FVector
 	{
 		FName Name;
 		FLinearColor Vector;
+
+		int32 MaterialLayer;
 	};
 
 	struct FScalar
 	{
 		FName Name;
 		float Scalar;
+		
+		int32 MaterialLayer;
 	};
 
 	struct FSurface
diff --git a/Engine/Plugins/Mutable/Source/CustomizableObject/Private/MuCO/CustomizableObjectInstance.cpp b/Engine/Plugins/Mutable/Source/CustomizableObject/Private/MuCO/CustomizableObjectInstance.cpp
--- a/Engine/Plugins/Mutable/Source/CustomizableObject/Private/MuCO/CustomizableObjectInstance.cpp
+++ b/Engine/Plugins/Mutable/Source/CustomizableObject/Private/MuCO/CustomizableObjectInstance.cpp
@@ -5942,65 +5942,13 @@ void UCustomizableInstancePrivate::BuildMaterials(const TSharedRef<FUpdateContex
 					for (int32 VectorIndex = 0; VectorIndex < Surface.VectorCount; ++VectorIndex)
 					{
 						const FInstanceUpdateData::FVector& Vector = Context->InstanceUpdateData.Vectors[Surface.FirstVector + VectorIndex];
-
-						// Decoding Material Layer from Mutable parameter name
-						FString EncodingString = "-MutableLayerParam:";
-
-						FString VectorName = Vector.Name.ToString();
-						int32 EncodingPosition = VectorName.Find(EncodingString);
-						int32 LayerIndex = -1;
-
-						if (EncodingPosition == INDEX_NONE)
-						{
-							MutableMaterialPlaceholder.AddParam(
-								FMutableMaterialPlaceholder::FMutableMaterialPlaceHolderParam(FName(Vector.Name), -1, Vector.Vector));
-						}
-						else
-						{
-							//Getting layer index
-							int32 LayerPosition = VectorName.Len() - (EncodingPosition + EncodingString.Len());
-							FString IndexString = VectorName.RightChop(VectorName.Len() - LayerPosition);
-							LayerIndex = FCString::Atof(*IndexString);
-
-							//Getting parameter name
-							FString Sufix = EncodingString + FString::FromInt(LayerIndex);
-							VectorName.RemoveFromEnd(Sufix);
-
-							MutableMaterialPlaceholder.AddParam(
-								FMutableMaterialPlaceholder::FMutableMaterialPlaceHolderParam(FName(VectorName), LayerIndex, Vector.Vector));
-						}
+						MutableMaterialPlaceholder.AddParam(FMutableMaterialPlaceholder::FMutableMaterialPlaceHolderParam(FName(Vector.Name), Vector.MaterialLayer, Vector.Vector));
 					}
 
 					for (int32 ScalarIndex = 0; ScalarIndex < Surface.ScalarCount; ++ScalarIndex)
 					{
 						const FInstanceUpdateData::FScalar& Scalar = Context->InstanceUpdateData.Scalars[Surface.FirstScalar + ScalarIndex];
-
-						// Decoding Material Layer from Mutable parameter name
-						FString EncodingString = "-MutableLayerParam:";
-
-						FString ScalarName = Scalar.Name.ToString();
-						int32 EncodingPosition = ScalarName.Find(EncodingString);
-						int32 LayerIndex = -1;
-
-						if (EncodingPosition == INDEX_NONE)
-						{
-							MutableMaterialPlaceholder.AddParam(
-								FMutableMaterialPlaceholder::FMutableMaterialPlaceHolderParam(FName(Scalar.Name), -1, Scalar.Scalar));
-						}
-						else
-						{
-							//Getting layer index
-							int32 LayerPosition = ScalarName.Len() - (EncodingPosition + EncodingString.Len());
-							FString IndexString = ScalarName.RightChop(ScalarName.Len() - LayerPosition);
-							LayerIndex = FCString::Atof(*IndexString);
-
-							//Getting parameter name
-							FString Sufix = EncodingString + FString::FromInt(LayerIndex);
-							ScalarName.RemoveFromEnd(Sufix);
-
-							MutableMaterialPlaceholder.AddParam(
-								FMutableMaterialPlaceholder::FMutableMaterialPlaceHolderParam(FName(ScalarName), LayerIndex, Scalar.Scalar));
-						}
+						MutableMaterialPlaceholder.AddParam(FMutableMaterialPlaceholder::FMutableMaterialPlaceHolderParam(FName(Scalar.Name), Scalar.MaterialLayer, Scalar.Scalar));
 					}
 				}
 
@@ -6208,7 +6156,7 @@ void UCustomizableInstancePrivate::BuildMaterials(const TSharedRef<FUpdateContex
 										}
 									}
 									
-									TextureCache.Add(TextureCacheKey, MutableTexture);						
+									TextureCache.Add(TextureCacheKey, MutableTexture);
 								}
 								else
 								{
@@ -6226,28 +6174,7 @@ void UCustomizableInstancePrivate::BuildMaterials(const TSharedRef<FUpdateContex
 								NewGeneratedTextures.Add(TextureData);
 							}
 
-							// Decoding Material Layer from Mutable parameter name
-							FString ImageName = Image->Name.ToString();
-							FString EncodingString = "-MutableLayerParam:";
-
-							int32 EncodingPosition = ImageName.Find(EncodingString);
-							int32 LayerIndex = -1;
-
-							if (EncodingPosition == INDEX_NONE)
-							{
-								MutableMaterialPlaceholder.AddParam(
-									FMutableMaterialPlaceholder::FMutableMaterialPlaceHolderParam(*Props.TextureParameterName, -1, TextureData));
-							}
-							else
-							{
-								//Getting layer index
-								int32 LayerPosition = ImageName.Len() - (EncodingPosition + EncodingString.Len());
-								FString IndexString = ImageName.RightChop(ImageName.Len() - LayerPosition);
-								LayerIndex = FCString::Atof(*IndexString);
-
-								MutableMaterialPlaceholder.AddParam(
-									FMutableMaterialPlaceholder::FMutableMaterialPlaceHolderParam(*Props.TextureParameterName, LayerIndex, TextureData));
-							}
+							MutableMaterialPlaceholder.AddParam(FMutableMaterialPlaceholder::FMutableMaterialPlaceHolderParam(*Props.TextureParameterName, Image->MaterialLayer, TextureData));
 						}
 						else
 						{
@@ -6347,6 +6274,7 @@ void UCustomizableInstancePrivate::BuildMaterials(const TSharedRef<FUpdateContex
 					{
 						for (const FMutableMaterialPlaceholder::FMutableMaterialPlaceHolderParam& Param : MutableMaterialPlaceholder.GetParams())
 						{
+							//TODO: Add the possibility to modify blend parameters not just layer parameters.
 							switch (Param.Type)
 							{
 							case FMutableMaterialPlaceholder::EPlaceHolderParamType::Vector:
diff --git a/Engine/Plugins/Mutable/Source/CustomizableObject/Private/MuCO/CustomizableObjectSystem.cpp b/Engine/Plugins/Mutable/Source/CustomizableObject/Private/MuCO/CustomizableObjectSystem.cpp
--- a/Engine/Plugins/Mutable/Source/CustomizableObject/Private/MuCO/CustomizableObjectSystem.cpp
+++ b/Engine/Plugins/Mutable/Source/CustomizableObject/Private/MuCO/CustomizableObjectSystem.cpp
@@ -2409,6 +2409,7 @@ namespace impl
 							FInstanceUpdateData::FVector& Vector = OperationData->InstanceUpdateData.Vectors.Last();
 							Vector.Name = MutableInstance->GetVectorName(InstanceComponentIndex, BaseLODIndex, BaseSurfaceIndex, VectorIndex);
 							Vector.Vector = MutableInstance->GetVector(InstanceComponentIndex, BaseLODIndex, BaseSurfaceIndex, VectorIndex);
+							Vector.MaterialLayer = MutableInstance->GetVectorLayerIndex(InstanceComponentIndex, BaseLODIndex, BaseSurfaceIndex, VectorIndex);
 						}
 
 						// Scalars
@@ -2420,8 +2421,9 @@ namespace impl
 
 							const FName ScalarName = MutableInstance->GetScalarName(InstanceComponentIndex, BaseLODIndex, BaseSurfaceIndex, ScalarIndex);
 							const float ScalarValue = MutableInstance->GetScalar(InstanceComponentIndex, BaseLODIndex, BaseSurfaceIndex, ScalarIndex);
+							const int32 MaterialLayer = MutableInstance->GetScalarLayerIndex(InstanceComponentIndex, BaseLODIndex, BaseSurfaceIndex, ScalarIndex);
 
-							OperationData->InstanceUpdateData.Scalars.Push({ ScalarName, ScalarValue });
+							OperationData->InstanceUpdateData.Scalars.Push({ ScalarName, ScalarValue, MaterialLayer });
 						}
 
 						// Images
@@ -2435,6 +2437,7 @@ namespace impl
 							TSharedRef<FInstanceUpdateData::FImage> Image = OperationData->InstanceUpdateData.Images.Last();
 							Image->Name = MutableInstance->GetImageName(InstanceComponentIndex, BaseLODIndex, BaseSurfaceIndex, ImageIndex);
 							Image->ImageID = MutableInstance->GetImageId(InstanceComponentIndex, BaseLODIndex, BaseSurfaceIndex, ImageIndex);
+							Image->MaterialLayer = MutableInstance->GetImageLayerIndex(InstanceComponentIndex, BaseLODIndex, BaseSurfaceIndex, ImageIndex);
 							Image->FullImageSizeX = 0;
 							Image->FullImageSizeY = 0;
 							Image->BaseLOD = BaseLODIndex;
diff --git a/Engine/Plugins/Mutable/Source/CustomizableObjectEditor/Private/MuCOE/GenerateMutableSource/GenerateMutableSource.h b/Engine/Plugins/Mutable/Source/CustomizableObjectEditor/Private/MuCOE/GenerateMutableSource/GenerateMutableSource.h
--- a/Engine/Plugins/Mutable/Source/CustomizableObjectEditor/Private/MuCOE/GenerateMutableSource/GenerateMutableSource.h
+++ b/Engine/Plugins/Mutable/Source/CustomizableObjectEditor/Private/MuCOE/GenerateMutableSource/GenerateMutableSource.h
@@ -380,10 +380,17 @@ struct FGroupProjectorImageInfo
 	FString TextureName;
 	FString RealTextureName;
 	int32 UVLayout = 0;
-
-	FGroupProjectorImageInfo(UE::Mutable::Private::NodeImagePtr InImageNode, const FString& InTextureName, const FString& InRealTextureName, UCustomizableObjectNodeMaterialBase* InTypedNodeMat, UE::Mutable::Private::Ptr<UE::Mutable::Private::NodeSurfaceNew> InSurfNode, int32 InUVLayout)
-		: TypedNodeMat(InTypedNodeMat), TextureName(InTextureName), RealTextureName(InRealTextureName),
-		UVLayout(InUVLayout)
+	int32 LayerIndex = -1;
+
+	FGroupProjectorImageInfo(
+		UE::Mutable::Private::NodeImagePtr InImageNode,
+		const FString& InTextureName,
+		const FString& InRealTextureName,
+		UCustomizableObjectNodeMaterialBase* InTypedNodeMat,
+		UE::Mutable::Private::Ptr<UE::Mutable::Private::NodeSurfaceNew> InSurfNode,
+		int32 InUVLayout,
+		int32 LayerIndex)
+		: TypedNodeMat(InTypedNodeMat), TextureName(InTextureName), RealTextureName(InRealTextureName), UVLayout(InUVLayout), LayerIndex(LayerIndex)
 	{
 		ImageNode = InImageNode;
 		SurfNode = InSurfNode;
diff --git a/Engine/Plugins/Mutable/Source/CustomizableObjectEditor/Private/MuCOE/GenerateMutableSource/GenerateMutableSourceSurface.cpp b/Engine/Plugins/Mutable/Source/CustomizableObjectEditor/Private/MuCOE/GenerateMutableSource/GenerateMutableSourceSurface.cpp
--- a/Engine/Plugins/Mutable/Source/CustomizableObjectEditor/Private/MuCOE/GenerateMutableSource/GenerateMutableSourceSurface.cpp
+++ b/Engine/Plugins/Mutable/Source/CustomizableObjectEditor/Private/MuCOE/GenerateMutableSource/GenerateMutableSourceSurface.cpp
@@ -276,7 +276,9 @@ UE::Mutable::Private::Ptr<UE::Mutable::Private::NodeSurface> GenerateMutableSour
 
 					check(Props.ImagePropertiesIndex != INDEX_NONE);
 					const FString SurfNodeImageName = FString::Printf(TEXT("%d"), Props.ImagePropertiesIndex);
-					SurfNode->Images[ImageIndex].Name = SurfNodeImageName;
+					const int32 LayerIndex = TypedNodeMat->GetParameterLayerIndex(EMaterialParameterType::Texture, ImageIndex);
+					SurfNode->Images[ImageIndex].ParameterInfo.ParameterName = SurfNodeImageName;
+					SurfNode->Images[ImageIndex].ParameterInfo.MaterialLayerIndex = LayerIndex;
 					SurfNode->Images[ImageIndex].LayoutIndex = -1;
 					SurfNode->Images[ImageIndex].MaterialName = Material->GetName();
 					SurfNode->Images[ImageIndex].MaterialParameterName = Props.TextureParameterName;
@@ -666,11 +668,10 @@ UE::Mutable::Private::Ptr<UE::Mutable::Private::NodeSurface> GenerateMutableSour
 						check(Props.ImagePropertiesIndex != INDEX_NONE);
 						const FString SurfNodeImageName = FString::Printf(TEXT("%d"), Props.ImagePropertiesIndex);
 
-						// Encoding material layer in mutable name
+						// Name and material layer
 						const int32 LayerIndex = TypedNodeMat->GetParameterLayerIndex(EMaterialParameterType::Texture, ImageIndex);
-						const FString LayerEncoding = LayerIndex != INDEX_NONE ? "-MutableLayerParam:" + FString::FromInt(LayerIndex) : "";
-						
-						SurfNode->Images[ImageIndex].Name = SurfNodeImageName + LayerEncoding;
+						SurfNode->Images[ImageIndex].ParameterInfo.ParameterName = SurfNodeImageName;
+						SurfNode->Images[ImageIndex].ParameterInfo.MaterialLayerIndex = LayerIndex;
 
 						// If we are generating an implicit component (with a passthrough mesh) we don't apply any layout.
 						int32 UVLayout = -1;
@@ -687,7 +688,7 @@ UE::Mutable::Private::Ptr<UE::Mutable::Private::NodeSurface> GenerateMutableSour
 							// Add to the GroupProjectorLODCache to potentially reuse this projection texture in higher LODs
 							ensure(LOD == GenerationContext.FirstLODAvailable[GenerationContext.CurrentMeshComponent]);
 							GenerationContext.GroupProjectorLODCache.Add(MaterialImageId,
-								FGroupProjectorImageInfo(ImageNode, ImageName, ImageName, TypedNodeMat, SurfNode, UVLayout));
+								FGroupProjectorImageInfo(ImageNode, ImageName, ImageName, TypedNodeMat, SurfNode, UVLayout, LayerIndex));
 						}
 					}
 				}
@@ -696,7 +697,8 @@ UE::Mutable::Private::Ptr<UE::Mutable::Private::NodeSurface> GenerateMutableSour
 					ensure(LOD > GenerationContext.FirstLODAvailable[GenerationContext.CurrentMeshComponent]);
 					check(ProjectorInfo->SurfNode->Images[ImageIndex].Image == ProjectorInfo->ImageNode);
 					SurfNode->Images[ImageIndex].Image = ProjectorInfo->ImageNode;
-					SurfNode->Images[ImageIndex].Name = ProjectorInfo->TextureName;
+					SurfNode->Images[ImageIndex].ParameterInfo.ParameterName = ProjectorInfo->TextureName;
+					SurfNode->Images[ImageIndex].ParameterInfo.MaterialLayerIndex = ProjectorInfo->LayerIndex;
 					SurfNode->Images[ImageIndex].LayoutIndex = ProjectorInfo->UVLayout;
 				}
 			}
@@ -721,14 +723,10 @@ UE::Mutable::Private::Ptr<UE::Mutable::Private::NodeSurface> GenerateMutableSour
 				{
 					UE::Mutable::Private::Ptr<UE::Mutable::Private::NodeColour> ColorNode = GenerateMutableSourceColor(ConnectedPin, GenerationContext);
 
-					// Encoding material layer in mutable name
-					if (const int32 LayerIndex = TypedNodeMat->GetParameterLayerIndex(EMaterialParameterType::Vector, VectorIndex); LayerIndex != INDEX_NONE)
-					{
-						VectorName += "-MutableLayerParam:" + FString::FromInt(LayerIndex);
-					}
-
+					const int32 LayerIndex = TypedNodeMat->GetParameterLayerIndex(EMaterialParameterType::Vector, VectorIndex);
+					SurfNode->Vectors[VectorIndex].ParameterInfo.ParameterName = VectorName;
+					SurfNode->Vectors[VectorIndex].ParameterInfo.MaterialLayerIndex = LayerIndex;
 					SurfNode->Vectors[VectorIndex].Vector = ColorNode;
-					SurfNode->Vectors[VectorIndex].Name = VectorName;
 				}
 			}
 		}
@@ -752,20 +750,14 @@ UE::Mutable::Private::Ptr<UE::Mutable::Private::NodeSurface> GenerateMutableSour
 				{
 					UE::Mutable::Private::Ptr<UE::Mutable::Private::NodeScalar> ScalarNode = GenerateMutableSourceFloat(ConnectedPin, GenerationContext);
 
-					// Encoding material layer in mutable name
-					if (const int32 LayerIndex = TypedNodeMat->GetParameterLayerIndex(EMaterialParameterType::Scalar, ScalarIndex); LayerIndex != INDEX_NONE)
-					{
-						ScalarName += "-MutableLayerParam:" + FString::FromInt(LayerIndex);
-					}
-
+					const int32 LayerIndex = TypedNodeMat->GetParameterLayerIndex(EMaterialParameterType::Scalar, ScalarIndex);
 					SurfNode->Scalars[ScalarIndex].Scalar = ScalarNode;
-					SurfNode->Scalars[ScalarIndex].Name = ScalarName;
+					SurfNode->Scalars[ScalarIndex].ParameterInfo.ParameterName = ScalarName;
+					SurfNode->Scalars[ScalarIndex].ParameterInfo.MaterialLayerIndex = LayerIndex;
 				}
 			}
 		}
 
-		
-
 		for (const FString& Tag : TypedNodeMat->GetEnableTags(&GenerationContext.MacroNodesStack))
 		{
 			SurfNode->Tags.AddUnique(Tag);
diff --git a/Engine/Plugins/Mutable/Source/CustomizableObjectEditor/Private/MuCOE/SMutableGraphViewer.cpp b/Engine/Plugins/Mutable/Source/CustomizableObjectEditor/Private/MuCOE/SMutableGraphViewer.cpp
--- a/Engine/Plugins/Mutable/Source/CustomizableObjectEditor/Private/MuCOE/SMutableGraphViewer.cpp
+++ b/Engine/Plugins/Mutable/Source/CustomizableObjectEditor/Private/MuCOE/SMutableGraphViewer.cpp
@@ -467,17 +467,17 @@ void SMutableGraphViewer::GetChildrenForInfo(TSharedPtr<FMutableGraphTreeElement
 
 		for (int32 l = 0; l < SurfaceNew->Images.Num(); ++l)
 		{
-			AddChildFunc(SurfaceNew->Images[l].Image.get(), FString::Printf(TEXT("IMAGE [%s]"), *SurfaceNew->Images[l].Name));
+			AddChildFunc(SurfaceNew->Images[l].Image.get(), FString::Printf(TEXT("IMAGE [%s]"), *SurfaceNew->Images[l].ParameterInfo.ParameterName));
 		}
 
 		for (int32 l = 0; l < SurfaceNew->Vectors.Num(); ++l)
 		{
-			AddChildFunc(SurfaceNew->Vectors[l].Vector.get(), FString::Printf(TEXT("VECTOR [%s]"), *SurfaceNew->Vectors[l].Name));
+			AddChildFunc(SurfaceNew->Vectors[l].Vector.get(), FString::Printf(TEXT("VECTOR [%s]"), *SurfaceNew->Vectors[l].ParameterInfo.ParameterName));
 		}
 
 		for (int32 l = 0; l < SurfaceNew->Scalars.Num(); ++l)
 		{
-			AddChildFunc(SurfaceNew->Scalars[l].Scalar.get(), FString::Printf(TEXT("SCALAR [%s]"), *SurfaceNew->Scalars[l].Name));
+			AddChildFunc(SurfaceNew->Scalars[l].Scalar.get(), FString::Printf(TEXT("SCALAR [%s]"), *SurfaceNew->Scalars[l].ParameterInfo.ParameterName));
 		}
 
 		for (int32 l = 0; l < SurfaceNew->Strings.Num(); ++l)
diff --git a/Engine/Plugins/Mutable/Source/MutableRuntime/Internal/MuR/Instance.h b/Engine/Plugins/Mutable/Source/MutableRuntime/Internal/MuR/Instance.h
--- a/Engine/Plugins/Mutable/Source/MutableRuntime/Internal/MuR/Instance.h
+++ b/Engine/Plugins/Mutable/Source/MutableRuntime/Internal/MuR/Instance.h
@@ -28,6 +28,7 @@ namespace UE::Mutable::Private
 		{
 			FImageId Id;
 			FName Name;
+			int32 LayerIndex;
 		};
 
 		TArray<FInstanceImage> Images;
@@ -36,6 +37,7 @@ namespace UE::Mutable::Private
 		{
 			FVector4f Value;
 			FName Name;
+			int32 LayerIndex;
 		};
 
 		TArray<FInstanceVector> Vectors;
@@ -44,6 +46,7 @@ namespace UE::Mutable::Private
 		{
 			float Value;
 			FName Name;
+			int32 LayerIndex;
 		};
 
 		TArray<FInstanceScalar> Scalars;
@@ -188,6 +191,13 @@ namespace UE::Mutable::Private
         //! \param ImageIndex Index of the image, from 0 to GetImageCount()-1
 		UE_API FName GetImageName( int32 ComponentIndex, int32 LODIndex, int32 SurfaceIndex, int32 ImageIndex ) const;
 
+		//! Get the Layer Index of an image in a Material
+		//! \param ComponentIndex Index of the component, from 0 to GetComponentCount()-1
+		//! \param LODIndex Index of the level of detail, from 0 to GetLODCount()-1
+		//! \param SurfaceIndex Index of the surface, from 0 to GetSurfaceCount()-1
+        //! \param ImageIndex Index of the image, from 0 to GetImageCount()-1
+		UE_API int32 GetImageLayerIndex(int32 ComponentIndex, int32 LODIndex, int32 SurfaceIndex, int32 ImageIndex) const;
+
 		//! Get the number of vectors in a component
 		//! \param ComponentIndex Index of the component, from 0 to GetComponentCount()-1
 		//! \param LODIndex Index of the level of detail, from 0 to GetLODCount()-1
@@ -208,6 +218,13 @@ namespace UE::Mutable::Private
         //! \param VectorIndex Index of the vector, from 0 to GetVectorCount()-1
 		UE_API FName GetVectorName( int32 ComponentIndex, int32 LODIndex, int32 SurfaceIndex, int32 VectorIndex ) const;
 
+		//! Get the Layer Index of a vector in a Material
+		//! \param ComponentIndex Index of the component, from 0 to GetComponentCount()-1
+		//! \param LODIndex Index of the level of detail, from 0 to GetLODCount()-1
+		//! \param SurfaceIndex Index of the surface, from 0 to GetSurfaceCount()-1
+		//! \param VectorIndex Index of the vector, from 0 to GetVectorCount()-1
+		UE_API int32 GetVectorLayerIndex(int32 ComponentIndex, int32 LODIndex, int32 SurfaceIndex, int32 VectorIndex) const;
+
         //! Get the number of scalar values in a component
         //! \param ComponentIndex Index of the component, from 0 to GetComponentCount()-1
 		//! \param LODIndex Index of the level of detail, from 0 to GetLODCount()-1
@@ -228,6 +245,13 @@ namespace UE::Mutable::Private
         //! \param ScalarIndex Index of the scalar, from 0 to GetScalarCount()-1
 		UE_API FName GetScalarName( int32 ComponentIndex, int32 LODIndex, int32 SurfaceIndex, int32 ScalarIndex ) const;
 
+		//! Get the Layer Index of a scalar in a Material
+		//! \param ComponentIndex Index of the component, from 0 to GetComponentCount()-1
+		//! \param LODIndex Index of the level of detail, from 0 to GetLODCount()-1
+		//! \param SurfaceIndex Index of the surface, from 0 to GetSurfaceCount()-1
+		//! \param ScalarIndex Index of the scalar, from 0 to GetScalarCount()-1
+		UE_API int32 GetScalarLayerIndex(int32 ComponentIndex, int32 LODIndex, int32 SurfaceIndex, int32 ScalarIndex) const;
+
         //! Get the number of string values in a component
         //! \param ComponentIndex Index of the component, from 0 to GetComponentCount()-1
 		//! \param LODIndex Index of the level of detail, from 0 to GetLODCount()-1
@@ -275,11 +299,11 @@ namespace UE::Mutable::Private
 
 		void SetSurfaceName(int32 ComponentIndex, int32 LODIndex, int32 SurfaceIndex, FName Name);
 
-		int32 AddImage(int32 ComponentIndex, int32 LODIndex, int32 SurfaceIndex, const FImageId& Imag, FName Name);
+		int32 AddImage(int32 ComponentIndex, int32 LODIndex, int32 SurfaceIndex, const FImageId& Image, FName Name, int32 LayerIndex);
 
-		int32 AddVector(int32 ComponentIndex, int32 LODIndex, int32 SurfaceIndex, const FVector4f&, FName Name);
+		int32 AddVector(int32 ComponentIndex, int32 LODIndex, int32 SurfaceIndex, const FVector4f& Vector, FName Name, int32 LayerIndex);
 
-		int32 AddScalar(int32 ComponentIndex, int32 LODIndex, int32 SurfaceIndex, float, FName Name);
+		int32 AddScalar(int32 ComponentIndex, int32 LODIndex, int32 SurfaceIndex, float Scalar, FName Name, int32 LayerIndex);
 
 		int32 AddString(int32 ComponentIndex, int32 LODIndex, int32 SurfaceIndex, const FString& Value, FName Name);
 
diff --git a/Engine/Plugins/Mutable/Source/MutableRuntime/Internal/MuR/Operations.h b/Engine/Plugins/Mutable/Source/MutableRuntime/Internal/MuR/Operations.h
--- a/Engine/Plugins/Mutable/Source/MutableRuntime/Internal/MuR/Operations.h
+++ b/Engine/Plugins/Mutable/Source/MutableRuntime/Internal/MuR/Operations.h
@@ -1103,6 +1103,7 @@ namespace UE::Mutable::Private
 			int32 SharedSurfaceId;
 
             ADDRESS name;
+			int32 LayerIndex; // Will be removed in UE-351597.
         };
 
 		struct InstanceAddExtensionDataArgs
diff --git a/Engine/Plugins/Mutable/Source/MutableRuntime/Private/MuR/CodeRunner.cpp b/Engine/Plugins/Mutable/Source/MutableRuntime/Private/MuR/CodeRunner.cpp
--- a/Engine/Plugins/Mutable/Source/MutableRuntime/Private/MuR/CodeRunner.cpp
+++ b/Engine/Plugins/Mutable/Source/MutableRuntime/Private/MuR/CodeRunner.cpp
@@ -563,7 +563,7 @@ namespace UE::Mutable::Private
                     check(  nameAd < (uint32)Program.ConstantStrings.Num() );
                     const FString& Name = Program.ConstantStrings[ nameAd ];
 
-                    Result->AddVector( 0, 0, 0, value, FName(Name) );
+                    Result->AddVector( 0, 0, 0, value, FName(Name), Args.LayerIndex );
                 }
                 StoreInstance( Item, Result );
                 break;
@@ -608,7 +608,7 @@ namespace UE::Mutable::Private
                     check(  nameAd < (uint32)Program.ConstantStrings.Num() );
                     const FString& Name = Program.ConstantStrings[ nameAd ];
 
-                    Result->AddScalar( 0, 0, 0, value, FName(Name));
+                    Result->AddScalar( 0, 0, 0, value, FName(Name), Args.LayerIndex);
                 }
                 StoreInstance( Item, Result );
                 break;
@@ -1005,7 +1005,7 @@ namespace UE::Mutable::Private
 					OP::ADDRESS NameAd = Args.name;
 					check(NameAd < (uint32)Program.ConstantStrings.Num());
 					const FString& Name = Program.ConstantStrings[NameAd];
-					Result->AddImage(0, 0, 0, ImageId, FName(Name) );
+					Result->AddImage(0, 0, 0, ImageId, FName(Name), Args.LayerIndex);
                 }
                 StoreInstance( Item, Result );
                 break;
diff --git a/Engine/Plugins/Mutable/Source/MutableRuntime/Private/MuR/CodeRunnerBegin.cpp b/Engine/Plugins/Mutable/Source/MutableRuntime/Private/MuR/CodeRunnerBegin.cpp
--- a/Engine/Plugins/Mutable/Source/MutableRuntime/Private/MuR/CodeRunnerBegin.cpp
+++ b/Engine/Plugins/Mutable/Source/MutableRuntime/Private/MuR/CodeRunnerBegin.cpp
@@ -840,7 +840,7 @@ namespace UE::Mutable::Private
 								OP::ADDRESS NameAd = Args.name;
 								check(NameAd < (uint32)Program.ConstantStrings.Num());
 								const FString& Name = Program.ConstantStrings[NameAd];
-								Instance->AddImage(0, 0, 0, ImageId, FName(Name) );
+								Instance->AddImage(0, 0, 0, ImageId, FName(Name), Args.LayerIndex);
 							}
 							
 							StoreInstance(Item, Instance);
diff --git a/Engine/Plugins/Mutable/Source/MutableRuntime/Private/MuR/Instance.cpp b/Engine/Plugins/Mutable/Source/MutableRuntime/Private/MuR/Instance.cpp
--- a/Engine/Plugins/Mutable/Source/MutableRuntime/Private/MuR/Instance.cpp
+++ b/Engine/Plugins/Mutable/Source/MutableRuntime/Private/MuR/Instance.cpp
@@ -241,7 +241,18 @@ namespace UE::Mutable::Private
 
         return Components[ComponentIndex].LODs[LODIndex].Surfaces[SurfaceIndex].Images[ImageIndex].Name;
 	}
+
 	
+	int32 FInstance::GetImageLayerIndex(int32 ComponentIndex, int32 LODIndex, int32 SurfaceIndex, int32 ImageIndex) const
+	{
+		check(Components.IsValidIndex(ComponentIndex));
+		check(Components[ComponentIndex].LODs.IsValidIndex(LODIndex));
+		check(Components[ComponentIndex].LODs[LODIndex].Surfaces.IsValidIndex(SurfaceIndex));
+		check(Components[ComponentIndex].LODs[LODIndex].Surfaces[SurfaceIndex].Images.IsValidIndex(ImageIndex));
+
+		return Components[ComponentIndex].LODs[LODIndex].Surfaces[SurfaceIndex].Images[ImageIndex].LayerIndex;
+	}
+
 	
 	FVector4f FInstance::GetVector(int32 ComponentIndex, int32 LODIndex, int32 SurfaceIndex, int32 VectorIndex) const
 	{
@@ -263,6 +274,17 @@ namespace UE::Mutable::Private
 
         return Components[ComponentIndex].LODs[LODIndex].Surfaces[SurfaceIndex].Vectors[VectorIndex].Name;
 	}
+
+
+	int32 FInstance::GetVectorLayerIndex(int32 ComponentIndex, int32 LODIndex, int32 SurfaceIndex, int32 VectorIndex) const
+	{
+		check(Components.IsValidIndex(ComponentIndex));
+		check(Components[ComponentIndex].LODs.IsValidIndex(LODIndex));
+		check(Components[ComponentIndex].LODs[LODIndex].Surfaces.IsValidIndex(SurfaceIndex));
+		check(Components[ComponentIndex].LODs[LODIndex].Surfaces[SurfaceIndex].Vectors.IsValidIndex(VectorIndex));
+
+		return Components[ComponentIndex].LODs[LODIndex].Surfaces[SurfaceIndex].Vectors[VectorIndex].LayerIndex;
+	}
 	
 	
     float FInstance::GetScalar(int32 ComponentIndex, int32 LODIndex, int32 SurfaceIndex, int32 ScalarIndex ) const
@@ -286,6 +308,17 @@ namespace UE::Mutable::Private
         return Components[ComponentIndex].LODs[LODIndex].Surfaces[SurfaceIndex].Scalars[ScalarIndex].Name;
 	}
 
+
+	int32 FInstance::GetScalarLayerIndex(int32 ComponentIndex, int32 LODIndex, int32 SurfaceIndex, int32 ScalarIndex) const
+	{
+		check(Components.IsValidIndex(ComponentIndex));
+		check(Components[ComponentIndex].LODs.IsValidIndex(LODIndex));
+        check(Components[ComponentIndex].LODs[LODIndex].Surfaces.IsValidIndex(SurfaceIndex));
+		check(Components[ComponentIndex].LODs[LODIndex].Surfaces[SurfaceIndex].Scalars.IsValidIndex(ScalarIndex));
+
+        return Components[ComponentIndex].LODs[LODIndex].Surfaces[SurfaceIndex].Scalars[ScalarIndex].LayerIndex;
+	}
+
 	
     FString FInstance::GetString(int32 ComponentIndex, int32 LODIndex, int32 SurfaceIndex, int32 StringIndex ) const
     {
@@ -434,7 +467,7 @@ namespace UE::Mutable::Private
 	}
 
 
-    int32 FInstance::AddImage( int32 ComponentIndex, int32 LODIndex, int32 SurfaceIndex, const FImageId& ImageId, FName Name)
+    int32 FInstance::AddImage( int32 ComponentIndex, int32 LODIndex, int32 SurfaceIndex, const FImageId& ImageId, FName Name, int32 LayerIndex)
 	{
 		// Automatically create the necessary lods and components
 		while (ComponentIndex >= Components.Num())
@@ -451,11 +484,11 @@ namespace UE::Mutable::Private
 		}
 
 		FInstanceSurface& Surface = Components[ComponentIndex].LODs[LODIndex].Surfaces[SurfaceIndex];
-		return Surface.Images.Add({ ImageId, Name });
+		return Surface.Images.Add({ ImageId, Name, LayerIndex });
 	}
 	
 	
-    int32 FInstance::AddVector( int32 ComponentIndex, int32 LODIndex, int32 SurfaceIndex, const FVector4f& vec, FName Name)
+    int32 FInstance::AddVector( int32 ComponentIndex, int32 LODIndex, int32 SurfaceIndex, const FVector4f& Vector, FName Name, int32 LayerIndex)
 	{
 		// Automatically create the necessary lods and components
 		while (ComponentIndex >= Components.Num())
@@ -472,11 +505,11 @@ namespace UE::Mutable::Private
 		}
 
 		FInstanceSurface& Surface = Components[ComponentIndex].LODs[LODIndex].Surfaces[SurfaceIndex];
-		return Surface.Vectors.Add({ vec, Name } );
+		return Surface.Vectors.Add({ Vector, Name , LayerIndex } );
 	}
 
 	
-    int32 FInstance::AddScalar( int32 ComponentIndex, int32 LODIndex, int32 SurfaceIndex, float sca, FName Name)
+    int32 FInstance::AddScalar( int32 ComponentIndex, int32 LODIndex, int32 SurfaceIndex, float Scalar, FName Name, int32 LayerIndex)
     {
         // Automatically create the necessary lods and components
 		while (ComponentIndex >= Components.Num())
@@ -493,7 +526,7 @@ namespace UE::Mutable::Private
 		}
 
 		FInstanceSurface& Surface = Components[ComponentIndex].LODs[LODIndex].Surfaces[SurfaceIndex];
-		return Surface.Scalars.Add({ sca, Name });
+		return Surface.Scalars.Add({ Scalar, Name, LayerIndex });
     }
 	
     
diff --git a/Engine/Plugins/Mutable/Source/MutableTools/Internal/MuT/NodeSurfaceNew.h b/Engine/Plugins/Mutable/Source/MutableTools/Internal/MuT/NodeSurfaceNew.h
--- a/Engine/Plugins/Mutable/Source/MutableTools/Internal/MuT/NodeSurfaceNew.h
+++ b/Engine/Plugins/Mutable/Source/MutableTools/Internal/MuT/NodeSurfaceNew.h
@@ -39,9 +39,21 @@ namespace UE::Mutable::Private
 
 		Ptr<NodeMaterial> Material;
 
+		struct FParameterIdentifier
+		{
+			/** Name of the parameter inside the material. */
+			FString ParameterName;
+
+			/** Index of the Material Layer. If it isn't a layered material then -1. */
+			int32 MaterialLayerIndex = -1;
+
+			//TODO: Add the possibility to modify blend parameters not just layer parameters.
+			//ParameterType: Global, Blend or Layer
+		};
+
 		struct FImageData
 		{
-			FString Name;
+			FParameterIdentifier ParameterInfo;
 			FString MaterialName;
 			FString MaterialParameterName;
 			Ptr<NodeImage> Image;
@@ -54,7 +66,7 @@ namespace UE::Mutable::Private
 
 		struct FVectorData
 		{
-			FString Name;
+			FParameterIdentifier ParameterInfo;
 			Ptr<NodeColour> Vector;
 		};
 
@@ -62,7 +74,7 @@ namespace UE::Mutable::Private
 
 		struct FScalar
 		{
-			FString Name;
+			FParameterIdentifier ParameterInfo;
 			Ptr<NodeScalar> Scalar;
 		};
 
diff --git a/Engine/Plugins/Mutable/Source/MutableTools/Private/MuT/ASTOpInstanceAdd.cpp b/Engine/Plugins/Mutable/Source/MutableTools/Private/MuT/ASTOpInstanceAdd.cpp
--- a/Engine/Plugins/Mutable/Source/MutableTools/Private/MuT/ASTOpInstanceAdd.cpp
+++ b/Engine/Plugins/Mutable/Source/MutableTools/Private/MuT/ASTOpInstanceAdd.cpp
@@ -39,7 +39,8 @@ namespace UE::Mutable::Private
 				id == other->id &&
 				ExternalId == other->ExternalId &&
 				SharedSurfaceId == other->SharedSurfaceId &&
-				name == other->name;
+				name == other->name &&
+				LayerIndex == other->LayerIndex;
 		}
 		return false;
 	}
@@ -55,6 +56,7 @@ namespace UE::Mutable::Private
 		n->ExternalId = ExternalId;
 		n->SharedSurfaceId = SharedSurfaceId;
 		n->name = name;
+		n->LayerIndex = LayerIndex;
 		return n;
 	}
 
@@ -111,6 +113,7 @@ namespace UE::Mutable::Private
 			Args.ExternalId = ExternalId;
 			Args.SharedSurfaceId = SharedSurfaceId;
 			Args.name = Program.AddConstant(name);
+			Args.LayerIndex = LayerIndex;
 
 			if (instance) Args.instance = instance->linkedAddress;
 
diff --git a/Engine/Plugins/Mutable/Source/MutableTools/Private/MuT/ASTOpInstanceAdd.h b/Engine/Plugins/Mutable/Source/MutableTools/Private/MuT/ASTOpInstanceAdd.h
--- a/Engine/Plugins/Mutable/Source/MutableTools/Private/MuT/ASTOpInstanceAdd.h
+++ b/Engine/Plugins/Mutable/Source/MutableTools/Private/MuT/ASTOpInstanceAdd.h
@@ -29,6 +29,7 @@ struct FProgram;
 		uint32_t ExternalId = 0;
 		int32 SharedSurfaceId = INDEX_NONE;
 		FString name;
+		int32 LayerIndex = -1; // TODO(Max) UE-351597:This will be moved to the materials so meshes won't have this value anymore.
 
 		ASTOpInstanceAdd();
 		ASTOpInstanceAdd(const ASTOpInstanceAdd&) = delete;
diff --git a/Engine/Plugins/Mutable/Source/MutableTools/Private/MuT/CodeGenerator.cpp b/Engine/Plugins/Mutable/Source/MutableTools/Private/MuT/CodeGenerator.cpp
--- a/Engine/Plugins/Mutable/Source/MutableTools/Private/MuT/CodeGenerator.cpp
+++ b/Engine/Plugins/Mutable/Source/MutableTools/Private/MuT/CodeGenerator.cpp
@@ -1350,7 +1350,8 @@ namespace UE::Mutable::Private
 								op->type = EOpType::IN_ADDIMAGE;
 								op->instance = LastSurfOp;
 								op->value = imageAd;
-								op->name = SurfaceNode->Images[ImageIndex].Name;
+								op->name = SurfaceNode->Images[ImageIndex].ParameterInfo.ParameterName;
+								op->LayerIndex = SurfaceNode->Images[ImageIndex].ParameterInfo.MaterialLayerIndex;
 
 								LastSurfOp = op;
 							}
@@ -1550,7 +1551,8 @@ namespace UE::Mutable::Private
 									op->value = imageAd;
 
 									// Name
-									op->name = SurfaceNode->Images[ImageIndex].Name;
+									op->name = SurfaceNode->Images[ImageIndex].ParameterInfo.ParameterName;
+									op->LayerIndex = SurfaceNode->Images[ImageIndex].ParameterInfo.MaterialLayerIndex;
 
 									LastSurfOp = op;
 								}
@@ -1582,7 +1584,8 @@ namespace UE::Mutable::Private
 							op->value = VectorResult.op;
 
 							// Name
-							op->name = SurfaceNode->Vectors[t].Name;
+							op->name = SurfaceNode->Vectors[t].ParameterInfo.ParameterName;
+							op->LayerIndex = SurfaceNode->Vectors[t].ParameterInfo.MaterialLayerIndex;
 
 							LastSurfOp = op;
 						}
@@ -1606,7 +1609,8 @@ namespace UE::Mutable::Private
 							op->value = ScalarResult.op;
 
 							// Name
-							op->name = SurfaceNode->Scalars[t].Name;
+							op->name = SurfaceNode->Scalars[t].ParameterInfo.ParameterName;
+							op->LayerIndex = SurfaceNode->Scalars[t].ParameterInfo.MaterialLayerIndex;
 
 							LastSurfOp = op;
 						}
