[Mutable] Refactored Task_Mutable_GetMeshes

Simplified for upcoming changes.

diff --git a/Engine/Plugins/Mutable/Source/CustomizableObject/Internal/MuCO/CustomizableObjectSystemPrivate.h b/Engine/Plugins/Mutable/Source/CustomizableObject/Internal/MuCO/CustomizableObjectSystemPrivate.h
--- a/Engine/Plugins/Mutable/Source/CustomizableObject/Internal/MuCO/CustomizableObjectSystemPrivate.h
+++ b/Engine/Plugins/Mutable/Source/CustomizableObject/Internal/MuCO/CustomizableObjectSystemPrivate.h
@@ -472,9 +472,9 @@ struct FInstanceUpdateData
 		TSharedPtr<const UE::Mutable::Private::FSkeletalMesh> SkeletalMesh;
 	};
 
-	TArray<FComponent> Components;
-	TArray<FLOD> LODs;
-	TArray<FSurface> Surfaces;
+	TArray<TSharedRef<FComponent>> Components;
+	TArray<TSharedRef<FLOD>> LODs;
+	TArray<TSharedRef<FSurface>> Surfaces;
 	TArray<TSharedRef<FImage>> Images;
 	TArray<FVector> Vectors;
 	TArray<FScalar> Scalars;
@@ -711,7 +711,7 @@ public:
 	double StartUpdateTime = 0.0;
 	double UpdateTime = 0.0;
 
-	double TaskGetMeshTime = 0.0;
+	double TaskGetMeshesTime = 0.0;
 	double TaskLockCacheTime = 0.0;
 	double TaskGetImagesTime = 0.0;
 	double TaskConvertResourcesTime = 0.0f;
@@ -1009,20 +1009,6 @@ private:
 };
 
 
-namespace impl
-{
-	void CreateMutableInstance(const TSharedRef<FUpdateContextPrivate>& Operation);
-
-	void FixLODs(const TSharedRef<FUpdateContextPrivate>& Operation);
-	
-	void Subtask_Mutable_PrepareSkeletonData(const TSharedRef<FUpdateContextPrivate>& OperationData);
-
-	void Subtask_Mutable_UpdateParameterRelevancy(const TSharedRef<FUpdateContextPrivate>& OperationData);
-	
-	void Subtask_Mutable_PrepareTextures(const TSharedRef<FUpdateContextPrivate>& OperationData);
-}
-
-
 /** Set OnlyLOD to -1 to generate all mips */
 CUSTOMIZABLEOBJECT_API FTexturePlatformData* MutableCreateImagePlatformData(TSharedPtr<const UE::Mutable::Private::FImage> MutableImage, int32 OnlyLOD, uint16 FullSizeX, uint16 FullSizeY);
 
diff --git a/Engine/Plugins/Mutable/Source/CustomizableObject/Internal/MuCO/LogBenchmarkUtil.h b/Engine/Plugins/Mutable/Source/CustomizableObject/Internal/MuCO/LogBenchmarkUtil.h
--- a/Engine/Plugins/Mutable/Source/CustomizableObject/Internal/MuCO/LogBenchmarkUtil.h
+++ b/Engine/Plugins/Mutable/Source/CustomizableObject/Internal/MuCO/LogBenchmarkUtil.h
@@ -86,7 +86,7 @@ struct FInstanceUpdateStats
 	
 	double QueueTime = 0.0;
 	double UpdateTime = 0.0;
-	double TaskGetMeshTime = 0.0;
+	double TaskGetMeshesTime = 0.0;
 	double TaskLockCacheTime = 0.0;
 
 	double TaskGetImagesTime = 0.0;
diff --git a/Engine/Plugins/Mutable/Source/CustomizableObject/Private/MuCO/CustomizableObjectInstance.cpp b/Engine/Plugins/Mutable/Source/CustomizableObject/Private/MuCO/CustomizableObjectInstance.cpp
--- a/Engine/Plugins/Mutable/Source/CustomizableObject/Private/MuCO/CustomizableObjectInstance.cpp
+++ b/Engine/Plugins/Mutable/Source/CustomizableObject/Private/MuCO/CustomizableObjectInstance.cpp
@@ -1972,8 +1972,8 @@ bool DoComponentsNeedUpdate(const TSharedRef<FUpdateContextPrivate>& Context, bo
 
 	for (int32 InstanceComponentIndex = 0; InstanceComponentIndex < NumInstanceComponents; ++InstanceComponentIndex)
 	{
-		const FInstanceUpdateData::FComponent& Component = Context->InstanceUpdateData.Components[InstanceComponentIndex];
-		const FCustomizableObjectComponentIndex ObjectComponentIndex = Component.Id;
+		const TSharedRef<FInstanceUpdateData::FComponent>& Component = Context->InstanceUpdateData.Components[InstanceComponentIndex];
+		const FCustomizableObjectComponentIndex ObjectComponentIndex = Component->Id;
 		const FName ComponentName = Context->ComponentNames[ObjectComponentIndex.GetValue()];
 		
 		bool bComponentWithMesh = false;
@@ -1984,32 +1984,32 @@ bool DoComponentsNeedUpdate(const TSharedRef<FUpdateContextPrivate>& Context, bo
 
 		bool bMeshChanged = true;
 		
-		if (Component.SkeletalMesh->PassthroughObject)
+		if (Component->SkeletalMesh->PassthroughObject)
 		{
 			if (const FCustomizableInstanceComponentData* ComponentData = InstancePrivate->GetComponentData(ObjectComponentIndex)) // Could be nullptr if the component has not been generated.
 			{
-				bMeshChanged = Component.SkeletalMeshId != ComponentData->LastComponentDescriptor.SkeletalMeshDescriptor;
+				bMeshChanged = Component->SkeletalMeshId != ComponentData->LastComponentDescriptor.SkeletalMeshDescriptor;
 			}
 		}
 		else
 		{
 			// Look for invalid meshes
-			for (int32 LODIndex = FirstLOD; LODIndex < Component.LODCount; ++LODIndex)
+			for (int32 LODIndex = FirstLOD; LODIndex < Component->LODCount; ++LODIndex)
 			{
-				const FInstanceUpdateData::FLOD& LOD = Context->InstanceUpdateData.LODs[Component.FirstLOD + LODIndex];
+				const TSharedRef<FInstanceUpdateData::FLOD>& LOD = Context->InstanceUpdateData.LODs[Component->FirstLOD + LODIndex];
 
-				if (!LOD.Mesh)
+				if (!LOD->Mesh)
 				{
 					bHasInvalidMesh |= bComponentWithMesh; // Mesh will be invalid if the UE::Mutable::Private::mesh is missing on higher LODs.
 					continue;
 				}
 
-				if (LOD.SurfaceCount == 0 && !LOD.Mesh->IsReference())
+				if (LOD->SurfaceCount == 0 && !LOD->Mesh->IsReference())
 				{
 					continue;
 				}
 
-				if (!LOD.Mesh->IsReference() && LOD.Mesh->GetVertexCount() == 0)
+				if (!LOD->Mesh->IsReference() && LOD->Mesh->GetVertexCount() == 0)
 				{
 					UE_LOG(LogMutable, Error, TEXT("Failed to generate SkeletalMesh for CO Instance [%s]. CO [%s] has invalid geometry for LOD [%d] Component [%d]."),
 						*Instance->GetName(), *CustomizableObject->GetName(),
@@ -2017,7 +2017,7 @@ bool DoComponentsNeedUpdate(const TSharedRef<FUpdateContextPrivate>& Context, bo
 					bHasInvalidMesh = true; // Unreal does not support empty LODs.
 				}
 
-				check(LOD.SurfaceCount > 0);
+				check(LOD->SurfaceCount > 0);
 				bComponentWithMesh = true;
 			}
 			
@@ -2204,7 +2204,7 @@ bool UCustomizableInstancePrivate::UpdateSkeletalMesh_PostBeginUpdate0(UCustomiz
 
 			for (const auto& InstanceUpdateComponent : Context->InstanceUpdateData.Components)
 			{
-				FName InstanceUpdateComponentName = ModelResources.ComponentNamesPerObjectComponent[InstanceUpdateComponent.Id.GetValue()];
+				FName InstanceUpdateComponentName = ModelResources.ComponentNamesPerObjectComponent[InstanceUpdateComponent->Id.GetValue()];
 
 				if (DeletedComponentCandidateName == InstanceUpdateComponentName)
 				{
@@ -2299,7 +2299,7 @@ bool UCustomizableInstancePrivate::UpdateSkeletalMesh_PostBeginUpdate0(UCustomiz
 		TSharedPtr<const UE::Mutable::Private::FMesh> ComponentMesh;
 		for (int32 LODIndex = 0; LODIndex < Component->LODCount && !ComponentMesh; ++LODIndex)
 		{
-			ComponentMesh = Context->InstanceUpdateData.LODs[Component->FirstLOD + LODIndex].Mesh;
+			ComponentMesh = Context->InstanceUpdateData.LODs[Component->FirstLOD + LODIndex]->Mesh;
 		}
 		
 		if (!ComponentMesh)
@@ -4133,14 +4133,14 @@ void BuildOrCopyElementData(const TSharedRef<FUpdateContextPrivate>& Context, US
 
 	for (int32 LODIndex = Context->FirstLODAvailable[ComponentName]; LODIndex < Component->LODCount; ++LODIndex)
 	{
-		const FInstanceUpdateData::FLOD& LOD = Context->InstanceUpdateData.LODs[Component->FirstLOD+LODIndex];
+		const TSharedRef<FInstanceUpdateData::FLOD>& LOD = Context->InstanceUpdateData.LODs[Component->FirstLOD+LODIndex];
 
-		if (!LOD.SurfaceCount)
+		if (!LOD->SurfaceCount)
 		{
 			continue;
 		}
 
-		for (int32 SurfaceIndex = 0; SurfaceIndex < LOD.SurfaceCount; ++SurfaceIndex)
+		for (int32 SurfaceIndex = 0; SurfaceIndex < LOD->SurfaceCount; ++SurfaceIndex)
 		{
 			new(SkeletalMesh->GetResourceForRendering()->LODRenderData[LODIndex].RenderSections) FSkelMeshRenderSection();
 		}
@@ -4311,9 +4311,9 @@ void BuildOrCopyClothingData(
 			
 		for (int32 LODIndex = Context->FirstLODAvailable[ComponentName]; LODIndex < Component->LODCount; ++LODIndex)
 		{
-			const FInstanceUpdateData::FLOD& LOD = Context->InstanceUpdateData.LODs[Component->FirstLOD+LODIndex];
+			const TSharedRef<FInstanceUpdateData::FLOD>& LOD = Context->InstanceUpdateData.LODs[Component->FirstLOD+LODIndex];
 
-			if (TSharedPtr<const UE::Mutable::Private::FMesh> MutableMesh = LOD.Mesh)
+			if (TSharedPtr<const UE::Mutable::Private::FMesh> MutableMesh = LOD->Mesh)
 			{
 				for (int32 SectionIndex = 0; SectionIndex < MutableMesh->GetSurfaceCount(); ++SectionIndex)
 				{
@@ -4420,9 +4420,9 @@ void BuildOrCopyClothingData(
 
 		for (int32 LODIndex = Context->FirstResidentLOD[ComponentName]; LODIndex < Component->LODCount; ++LODIndex)
 		{
-			const FInstanceUpdateData::FLOD& LOD = Context->InstanceUpdateData.LODs[Component->FirstLOD+LODIndex];
+			const TSharedRef<FInstanceUpdateData::FLOD>& LOD = Context->InstanceUpdateData.LODs[Component->FirstLOD+LODIndex];
 
-			if (TSharedPtr<const UE::Mutable::Private::FMesh> MutableMesh = LOD.Mesh)
+			if (TSharedPtr<const UE::Mutable::Private::FMesh> MutableMesh = LOD->Mesh)
 			{
 				GetSectionClothData(*MutableMesh, LODIndex, Context->InstanceUpdateData.ClothingMeshData, SectionsClothData, NumClothingDataNotFound);
 			}
@@ -4600,9 +4600,9 @@ void BuildOrCopyClothingData(
 	// Zero all LODs (even those which do not use cloth).
 	for (int32 LODIndex = Context->FirstLODAvailable[ComponentName]; LODIndex < Component->LODCount; ++LODIndex)
 	{
-		const FInstanceUpdateData::FLOD& LOD = Context->InstanceUpdateData.LODs[Component->FirstLOD + LODIndex];
+		const TSharedRef<FInstanceUpdateData::FLOD>& LOD = Context->InstanceUpdateData.LODs[Component->FirstLOD + LODIndex];
 		
-		TSharedPtr<const UE::Mutable::Private::FMesh> MutableMesh = LOD.Mesh;
+		TSharedPtr<const UE::Mutable::Private::FMesh> MutableMesh = LOD->Mesh;
 		if (!MutableMesh)
 		{
 			continue;
@@ -4721,10 +4721,10 @@ bool BuildOrCopyRenderData(const TSharedRef<FUpdateContextPrivate>& Context, USk
 	{
 		MUTABLE_CPUPROFILER_SCOPE(BuildRenderData);
 		
-		const FInstanceUpdateData::FLOD& LOD = Context->InstanceUpdateData.LODs[Component->FirstLOD+LODIndex];
+		const TSharedRef<FInstanceUpdateData::FLOD>& LOD = Context->InstanceUpdateData.LODs[Component->FirstLOD+LODIndex];
 
 		// There could be components without a mesh in LODs
-		if (!LOD.Mesh || LOD.SurfaceCount == 0)
+		if (!LOD->Mesh || LOD->SurfaceCount == 0)
 		{
 			UE_LOG(LogMutable, Warning, TEXT("Building instance: generated mesh [%s] has LOD [%d] of object component index [%d] with no mesh.")
 				, *SkeletalMesh->GetName()
@@ -4744,11 +4744,11 @@ bool BuildOrCopyRenderData(const TSharedRef<FUpdateContextPrivate>& Context, USk
 		// Set active and required bones
 		{
 			const TArray<UE::Mutable::Private::FBoneName>& ActiveBones = Context->InstanceUpdateData.ActiveBones;
-			LODResource.ActiveBoneIndices.Reserve(LOD.ActiveBoneCount);
+			LODResource.ActiveBoneIndices.Reserve(LOD->ActiveBoneCount);
 
-			for (uint32 Index = 0; Index < LOD.ActiveBoneCount; ++Index)
+			for (uint32 Index = 0; Index < LOD->ActiveBoneCount; ++Index)
 			{
-				const uint16 ActiveBoneIndex = BoneInfoMap[ActiveBones[LOD.FirstActiveBone + Index]].Value;
+				const uint16 ActiveBoneIndex = BoneInfoMap[ActiveBones[LOD->FirstActiveBone + Index]].Value;
 				LODResource.ActiveBoneIndices.Add(ActiveBoneIndex);
 			}
 
@@ -4757,13 +4757,13 @@ bool BuildOrCopyRenderData(const TSharedRef<FUpdateContextPrivate>& Context, USk
 		}
 
 		// Find referenced surface metadata.
-		const int32 MeshNumSurfaces = LOD.Mesh->Surfaces.Num();
+		const int32 MeshNumSurfaces = LOD->Mesh->Surfaces.Num();
 		TArray<const FMutableSurfaceMetadata*> MeshSurfacesMetadata;
 		MeshSurfacesMetadata.Init(nullptr, MeshNumSurfaces);
 
 		for (int32 MeshSectionIndex = 0; MeshSectionIndex < MeshNumSurfaces; ++MeshSectionIndex)
 		{
-			uint32 MeshSurfaceId = LOD.Mesh->GetSurfaceId(MeshSectionIndex);
+			uint32 MeshSurfaceId = LOD->Mesh->GetSurfaceId(MeshSectionIndex);
 			int32 InstanceSurfaceIndex = Context->MutableInstance->FindSurfaceById(InstanceComponentIndex.GetValue(), LODIndex, MeshSurfaceId);
 			
 			if (InstanceSurfaceIndex < 0)
@@ -4781,7 +4781,7 @@ bool BuildOrCopyRenderData(const TSharedRef<FUpdateContextPrivate>& Context, USk
 			else
 			{
 				// In case the surface does not have metadata, check if any submesh has surface metadata.
-				for (const UE::Mutable::Private::FSurfaceSubMesh& SubMesh : LOD.Mesh->Surfaces[MeshSectionIndex].SubMeshes)	
+				for (const UE::Mutable::Private::FSurfaceSubMesh& SubMesh : LOD->Mesh->Surfaces[MeshSectionIndex].SubMeshes)	
 				{
 					const FMutableMeshMetadata* FoundMeshMetadata = ModelResources.MeshMetadata.Find(SubMesh.ExternalId);
 
@@ -4805,10 +4805,10 @@ bool BuildOrCopyRenderData(const TSharedRef<FUpdateContextPrivate>& Context, USk
 		// Set RenderSections
 		UnrealConversionUtils::SetupRenderSections(
 			LODResource,
-			LOD.Mesh.Get(),
+			LOD->Mesh.Get(),
 			Context->InstanceUpdateData.BoneMaps,
 			BoneInfoMap,
-			LOD.FirstBoneMap,
+			LOD->FirstBoneMap,
 			MeshSurfacesMetadata);
 
 		// Set SkinWeightProfiles
@@ -4817,7 +4817,7 @@ bool BuildOrCopyRenderData(const TSharedRef<FUpdateContextPrivate>& Context, USk
 		// Active SkinWeightProfiles ID and Name
 		TArray<TPair<uint32, FName>> ActiveSkinWeightProfiles;
 		
-		const UE::Mutable::Private::FMeshBufferSet& MutableMeshVertexBuffers = LOD.Mesh->GetVertexBuffers();
+		const UE::Mutable::Private::FMeshBufferSet& MutableMeshVertexBuffers = LOD->Mesh->GetVertexBuffers();
 		const int32 NumBuffers = MutableMeshVertexBuffers.GetBufferCount();
 
 		for (int32 BufferIndex = 0; BufferIndex < NumBuffers; ++BufferIndex)
@@ -4861,21 +4861,21 @@ bool BuildOrCopyRenderData(const TSharedRef<FUpdateContextPrivate>& Context, USk
 			// Copy Vertices
 			UnrealConversionUtils::CopyMutableVertexBuffers(
 				LODResource,
-				LOD.Mesh.Get(),
+				LOD->Mesh.Get(),
 				SkeletalMesh->GetLODInfo(LODIndex)->bAllowCPUAccess);
 
 			// SurfaceIDs. Required to copy index buffers with padding
 			TArray<uint32> SurfaceIDs;
-			SurfaceIDs.SetNum(LOD.SurfaceCount);
+			SurfaceIDs.SetNum(LOD->SurfaceCount);
 
-			for (int32 SurfaceIndex = 0; SurfaceIndex < LOD.SurfaceCount; ++SurfaceIndex)
+			for (int32 SurfaceIndex = 0; SurfaceIndex < LOD->SurfaceCount; ++SurfaceIndex)
 			{
-				SurfaceIDs[SurfaceIndex] = LOD.Mesh->GetSurfaceId(SurfaceIndex);
+				SurfaceIDs[SurfaceIndex] = LOD->Mesh->GetSurfaceId(SurfaceIndex);
 			}
 
 			// Copy indices.
 			bool bMarkRenderStateDirty = false;
-			if (!UnrealConversionUtils::CopyMutableIndexBuffers(LODResource, LOD.Mesh.Get(), SurfaceIDs, bMarkRenderStateDirty))
+			if (!UnrealConversionUtils::CopyMutableIndexBuffers(LODResource, LOD->Mesh.Get(), SurfaceIDs, bMarkRenderStateDirty))
 			{
 				// End with failure
 				return false;
@@ -4886,19 +4886,19 @@ bool BuildOrCopyRenderData(const TSharedRef<FUpdateContextPrivate>& Context, USk
 				LODResource,
 				*SkeletalMesh,
 				LODIndex,
-				LOD.Mesh.Get(),
+				LOD->Mesh.Get(),
 				ActiveSkinWeightProfiles);
 		}
-		else // Streamable LOD. 
+		else // Streamable LOD-> 
 		{
 			// Init VertexBuffers for streaming
 			UnrealConversionUtils::InitVertexBuffersWithDummyData(
 				LODResource,
-				LOD.Mesh.Get(),
+				LOD->Mesh.Get(),
 				SkeletalMesh->GetLODInfo(LODIndex)->bAllowCPUAccess);
 
 			// Init IndexBuffers for streaming
-			UnrealConversionUtils::InitIndexBuffersWithDummyData(LODResource, LOD.Mesh.Get());
+			UnrealConversionUtils::InitIndexBuffersWithDummyData(LODResource, LOD->Mesh.Get());
 		}
 
 		if (LODResource.StaticVertexBuffers.ColorVertexBuffer.GetNumVertices())
@@ -4968,7 +4968,7 @@ UE::Tasks::FTask LoadAdditionalAssetsAndData(const TSharedRef<FUpdateContextPriv
 
 	TArray<FSoftObjectPath> AssetsToStream;
 
-	TArray<FInstanceUpdateData::FComponent>& Components = Context->InstanceUpdateData.Components;
+	TArray<TSharedRef<FInstanceUpdateData::FComponent>>& Components = Context->InstanceUpdateData.Components;
 	
 	const int32 NumClothingAssets = ModelResources.ClothingAssetsData.Num();
 	InstancePrivate->ClothingPhysicsAssets.Reset(NumClothingAssets);
@@ -5037,16 +5037,16 @@ UE::Tasks::FTask LoadAdditionalAssetsAndData(const TSharedRef<FUpdateContextPriv
 	{
 		for (int32 InstanceComponentIndex = 0; InstanceComponentIndex < Context->InstanceUpdateData.Components.Num(); ++InstanceComponentIndex)
 		{
-			const FInstanceUpdateData::FComponent& Component = Components[InstanceComponentIndex];
-			FCustomizableInstanceComponentData* ComponentData = InstancePrivate->GetComponentData(Component.Id);
+			const TSharedRef<const FInstanceUpdateData::FComponent>& Component = Components[InstanceComponentIndex];
+			FCustomizableInstanceComponentData* ComponentData = InstancePrivate->GetComponentData(Component->Id);
 			
-			if (Component.SkeletalMesh->PassthroughObject)
+			if (Component->SkeletalMesh->PassthroughObject)
 			{
 				// Pass-through components don't have a Reference Mesh so don't access it
 				continue;
 			}
 			
-			const FCustomizableObjectComponentIndex ObjectComponentIndex = Component.Id;
+			const FCustomizableObjectComponentIndex ObjectComponentIndex = Component->Id;
 			const FMutableRefSkeletalMeshData& RefSkeletalMeshData = ModelResources.ReferenceSkeletalMeshesData[ObjectComponentIndex.GetValue()];
 
 			for (int32 AssetUserDataIndex : RefSkeletalMeshData.AssetUserDataIndices)
@@ -5057,18 +5057,18 @@ UE::Tasks::FTask LoadAdditionalAssetsAndData(const TSharedRef<FUpdateContextPriv
 				ComponentData->StreamedResourceIndex.AddUnique(AssetUserDataIndex);
 			}
 
-			const FName ComponentName = Context->ComponentNames[Component.Id.GetValue()];
+			const FName ComponentName = Context->ComponentNames[Component->Id.GetValue()];
 
 			if (bUpdateMeshes)
 			{
 				// Morphs
 				{
 					// Data
-					for (int32 LODIndex = Context->GetFirstRequestedLOD()[ComponentName]; LODIndex < Component.LODCount; ++LODIndex)
+					for (int32 LODIndex = Context->GetFirstRequestedLOD()[ComponentName]; LODIndex < Component->LODCount; ++LODIndex)
 					{
-						const FInstanceUpdateData::FLOD& LOD = Context->InstanceUpdateData.LODs[Component.FirstLOD + LODIndex];
+						const TSharedRef<FInstanceUpdateData::FLOD>& LOD = Context->InstanceUpdateData.LODs[Component->FirstLOD + LODIndex];
 
-						TSharedPtr<const UE::Mutable::Private::FMesh> MutableMesh = LOD.Mesh;
+						TSharedPtr<const UE::Mutable::Private::FMesh> MutableMesh = LOD->Mesh;
 						if (!MutableMesh)
 						{
 							continue;
@@ -5082,11 +5082,11 @@ UE::Tasks::FTask LoadAdditionalAssetsAndData(const TSharedRef<FUpdateContextPriv
 						Context->FirstResidentLOD[ComponentName] :
 						Context->GetFirstRequestedLOD()[ComponentName];
 				
-					for (int32 LODIndex = FirstLOD; LODIndex < Component.LODCount; ++LODIndex)
+					for (int32 LODIndex = FirstLOD; LODIndex < Component->LODCount; ++LODIndex)
 					{
-						const FInstanceUpdateData::FLOD& LOD = Context->InstanceUpdateData.LODs[Component.FirstLOD + LODIndex];
+						const TSharedRef<FInstanceUpdateData::FLOD>& LOD = Context->InstanceUpdateData.LODs[Component->FirstLOD + LODIndex];
 
-						TSharedPtr<const UE::Mutable::Private::FMesh> MutableMesh = LOD.Mesh;
+						TSharedPtr<const UE::Mutable::Private::FMesh> MutableMesh = LOD->Mesh;
 						if (!MutableMesh)
 						{
 							continue;
@@ -5100,11 +5100,11 @@ UE::Tasks::FTask LoadAdditionalAssetsAndData(const TSharedRef<FUpdateContextPriv
 				{
 					// Data
 					// From FirstResidentLOD instead of FirstRequestedLOD since clothing we generate all LODs, even the hacky streaming copied ones.
-					for (int32 LODIndex = Context->FirstResidentLOD[ComponentName]; LODIndex < Component.LODCount; ++LODIndex)
+					for (int32 LODIndex = Context->FirstResidentLOD[ComponentName]; LODIndex < Component->LODCount; ++LODIndex)
 					{
-						const FInstanceUpdateData::FLOD& LOD = Context->InstanceUpdateData.LODs[Component.FirstLOD + LODIndex];
+						const TSharedRef<FInstanceUpdateData::FLOD>& LOD = Context->InstanceUpdateData.LODs[Component->FirstLOD + LODIndex];
 
-						TSharedPtr<const UE::Mutable::Private::FMesh> MutableMesh = LOD.Mesh;
+						TSharedPtr<const UE::Mutable::Private::FMesh> MutableMesh = LOD->Mesh;
 						if (!MutableMesh)
 						{
 							continue;
@@ -5117,11 +5117,11 @@ UE::Tasks::FTask LoadAdditionalAssetsAndData(const TSharedRef<FUpdateContextPriv
 					}
 
 					// Metadata
-					for (int32 LODIndex = Context->FirstLODAvailable[ComponentName]; LODIndex < Component.LODCount; ++LODIndex)
+					for (int32 LODIndex = Context->FirstLODAvailable[ComponentName]; LODIndex < Component->LODCount; ++LODIndex)
 					{
-						const FInstanceUpdateData::FLOD& LOD = Context->InstanceUpdateData.LODs[Component.FirstLOD + LODIndex];
+						const TSharedRef<FInstanceUpdateData::FLOD>& LOD = Context->InstanceUpdateData.LODs[Component->FirstLOD + LODIndex];
 
-						TSharedPtr<const UE::Mutable::Private::FMesh> MutableMesh = LOD.Mesh;
+						TSharedPtr<const UE::Mutable::Private::FMesh> MutableMesh = LOD->Mesh;
 						if (!MutableMesh)
 						{
 							continue;
@@ -5178,11 +5178,11 @@ UE::Tasks::FTask LoadAdditionalAssetsAndData(const TSharedRef<FUpdateContextPriv
 				}
 			}
 			
-			for (int32 LODIndex = Context->FirstLODAvailable[ComponentName]; LODIndex < Component.LODCount; ++LODIndex)
+			for (int32 LODIndex = Context->FirstLODAvailable[ComponentName]; LODIndex < Component->LODCount; ++LODIndex)
 			{
-				const FInstanceUpdateData::FLOD& LOD = Context->InstanceUpdateData.LODs[Component.FirstLOD + LODIndex];
+				const TSharedRef<FInstanceUpdateData::FLOD>& LOD = Context->InstanceUpdateData.LODs[Component->FirstLOD + LODIndex];
 
-				TSharedPtr<const UE::Mutable::Private::FMesh> MutableMesh = LOD.Mesh;
+				TSharedPtr<const UE::Mutable::Private::FMesh> MutableMesh = LOD->Mesh;
 				if (!MutableMesh)
 				{
 					continue;
@@ -5751,9 +5751,9 @@ void UCustomizableInstancePrivate::BuildMaterials(const TSharedRef<FUpdateContex
 
 	for (int32 InstanceComponentIndex = 0; InstanceComponentIndex < Context->NumInstanceComponents; ++InstanceComponentIndex)
 	{
-		const FInstanceUpdateData::FComponent& Component = Context->InstanceUpdateData.Components[InstanceComponentIndex];
+		const TSharedRef<const FInstanceUpdateData::FComponent>& Component = Context->InstanceUpdateData.Components[InstanceComponentIndex];
 
-		const FCustomizableObjectComponentIndex ObjectComponentIndex = Component.Id;
+		const FCustomizableObjectComponentIndex ObjectComponentIndex = Component->Id;
 
 		if (!ModelResources.ComponentNamesPerObjectComponent.IsValidIndex(ObjectComponentIndex.GetValue()))
 		{
@@ -5783,9 +5783,9 @@ void UCustomizableInstancePrivate::BuildMaterials(const TSharedRef<FUpdateContex
 			ComponentData->OverrideMaterials.Reset();
 			ComponentData->OverlayMaterial = nullptr;
 
-			if (Component.OverlayMaterial)
+			if (Component->OverlayMaterial->PassthroughObject)
 			{
-				ComponentData->OverlayMaterial = Component.OverlayMaterial->PassthroughObject.Get();
+				ComponentData->OverlayMaterial = Component->OverlayMaterial->PassthroughObject.Get();
 				OverlayMaterial = ComponentData->OverlayMaterial;
 			}
 		}
@@ -5810,9 +5810,9 @@ void UCustomizableInstancePrivate::BuildMaterials(const TSharedRef<FUpdateContex
 			Context->FirstLODAvailable[ComponentName] :
 			Context->GetFirstRequestedLOD()[ComponentName];
 
-		for (int32 LODIndex = FirstLOD; LODIndex < Component.LODCount; LODIndex++)
+		for (int32 LODIndex = FirstLOD; LODIndex < Component->LODCount; LODIndex++)
 		{
-			const FInstanceUpdateData::FLOD& LOD = Context->InstanceUpdateData.LODs[Component.FirstLOD+LODIndex];
+			const TSharedRef<const FInstanceUpdateData::FLOD>& LOD = Context->InstanceUpdateData.LODs[Component->FirstLOD + LODIndex];
 			
 			if (!bUseOverrideMaterialsOnly && LODIndex < SkeletalMesh->GetLODNum())
 			{
@@ -5826,12 +5826,12 @@ void UCustomizableInstancePrivate::BuildMaterials(const TSharedRef<FUpdateContex
 				RefSkeletalMeshData = &ModelResources.ReferenceSkeletalMeshesData[ObjectComponentIndex.GetValue()];
 			}
 
-			for (int32 SurfaceIndex = 0; SurfaceIndex < LOD.SurfaceCount; ++SurfaceIndex)
+			for (int32 SurfaceIndex = 0; SurfaceIndex < LOD->SurfaceCount; ++SurfaceIndex)
 			{
-				const FInstanceUpdateData::FSurface& Surface = Context->InstanceUpdateData.Surfaces[LOD.FirstSurface + SurfaceIndex];
+				const TSharedRef<const FInstanceUpdateData::FSurface>& Surface = Context->InstanceUpdateData.Surfaces[LOD->FirstSurface + SurfaceIndex];
 
 				// Reuse MaterialSlot from the previous LOD.
-				if (const int32 MaterialIndex = SurfaceIdToMaterialIndex.Find(Surface.SharedSurfaceId); MaterialIndex != INDEX_NONE)
+				if (const int32 MaterialIndex = SurfaceIdToMaterialIndex.Find(Surface->SharedSurfaceId); MaterialIndex != INDEX_NONE)
 				{
 					if (!bUseOverrideMaterialsOnly)
 					{
@@ -5843,7 +5843,7 @@ void UCustomizableInstancePrivate::BuildMaterials(const TSharedRef<FUpdateContex
 				}
 
 				// Is this a material in a passthrough mesh that we don't modify?
-				if (!Surface.Material)
+				if (!Surface->Material)
 				{
 					Materials.Emplace();
 #if WITH_EDITOR
@@ -5853,7 +5853,7 @@ void UCustomizableInstancePrivate::BuildMaterials(const TSharedRef<FUpdateContex
 					continue;
 				}
 
-				UMaterialInterface* MaterialTemplate = Surface.Material->PassthroughObject.Get();
+				UMaterialInterface* MaterialTemplate = Surface->Material->PassthroughObject.Get();
 				if (!MaterialTemplate)
 				{
 					// Missing MaterialTemplate. Use DefaultMaterial instead. 
@@ -5863,27 +5863,27 @@ void UCustomizableInstancePrivate::BuildMaterials(const TSharedRef<FUpdateContex
 				}
 
 				// This section will require a new slot
-				SurfaceIdToMaterialIndex.Add(Surface.SharedSurfaceId);
+				SurfaceIdToMaterialIndex.Add(Surface->SharedSurfaceId);
 
 				// Add and set up the material data for this slot
 				const int32 MaterialSlotIndex = Materials.Num();
 				FSkeletalMaterial& MaterialSlot = Materials.AddDefaulted_GetRef();
 				MaterialSlot.MaterialInterface = MaterialTemplate;
 
-				uint32 UsedSurfaceMetadataId = Surface.SurfaceMetadataId;
+				uint32 UsedSurfaceMetadataId = Surface->SurfaceMetadataId;
 				
 				// If the surface metadata is invalid, check if any of the mesh fragments has metadata. 
 				// For now use the first found, an aggregate may be needed. 
-				if (Surface.SurfaceMetadataId == 0 && LOD.Mesh)
+				if (Surface->SurfaceMetadataId == 0 && LOD->Mesh)
 				{
-					int32 MeshSurfaceIndex = LOD.Mesh->Surfaces.IndexOfByPredicate([SurfaceId = Surface.SurfaceId](const UE::Mutable::Private::FMeshSurface& Surface)
+					int32 MeshSurfaceIndex = LOD->Mesh->Surfaces.IndexOfByPredicate([SurfaceId = Surface->SurfaceId](const UE::Mutable::Private::FMeshSurface& Surface)
 					{
 						return SurfaceId == Surface.Id;
 					});
 
 					if (MeshSurfaceIndex != INDEX_NONE)
 					{
-						for (const UE::Mutable::Private::FSurfaceSubMesh& SubMesh : LOD.Mesh->Surfaces[SurfaceIndex].SubMeshes)	
+						for (const UE::Mutable::Private::FSurfaceSubMesh& SubMesh : LOD->Mesh->Surfaces[SurfaceIndex].SubMeshes)	
 						{
 							const FMutableMeshMetadata* FoundMeshMetadata = ModelResources.MeshMetadata.Find(SubMesh.ExternalId);
 
@@ -5939,15 +5939,15 @@ void UCustomizableInstancePrivate::BuildMaterials(const TSharedRef<FUpdateContex
 				{
 					MUTABLE_CPUPROFILER_SCOPE(ParamLoop);
 
-					for (int32 VectorIndex = 0; VectorIndex < Surface.VectorCount; ++VectorIndex)
+					for (int32 VectorIndex = 0; VectorIndex < Surface->VectorCount; ++VectorIndex)
 					{
-						const FInstanceUpdateData::FVector& Vector = Context->InstanceUpdateData.Vectors[Surface.FirstVector + VectorIndex];
+						const FInstanceUpdateData::FVector& Vector = Context->InstanceUpdateData.Vectors[Surface->FirstVector + VectorIndex];
 						MutableMaterialPlaceholder.AddParam(FMutableMaterialPlaceholder::FMutableMaterialPlaceHolderParam(FName(Vector.Name), Vector.MaterialLayer, Vector.Vector));
 					}
 
-					for (int32 ScalarIndex = 0; ScalarIndex < Surface.ScalarCount; ++ScalarIndex)
+					for (int32 ScalarIndex = 0; ScalarIndex < Surface->ScalarCount; ++ScalarIndex)
 					{
-						const FInstanceUpdateData::FScalar& Scalar = Context->InstanceUpdateData.Scalars[Surface.FirstScalar + ScalarIndex];
+						const FInstanceUpdateData::FScalar& Scalar = Context->InstanceUpdateData.Scalars[Surface->FirstScalar + ScalarIndex];
 						MutableMaterialPlaceholder.AddParam(FMutableMaterialPlaceholder::FMutableMaterialPlaceHolderParam(FName(Scalar.Name), Scalar.MaterialLayer, Scalar.Scalar));
 					}
 				}
@@ -5964,16 +5964,16 @@ void UCustomizableInstancePrivate::BuildMaterials(const TSharedRef<FUpdateContex
 					check((bNeverStream && Context->MipsToSkip == 0) ||
 						(!bNeverStream && Context->MipsToSkip >= 0));
 
-					for (int32 ImageIndex = 0; ImageIndex < Surface.ImageCount; ++ImageIndex)
+					for (int32 ImageIndex = 0; ImageIndex < Surface->ImageCount; ++ImageIndex)
 					{
-						const TSharedRef<const FInstanceUpdateData::FImage>& Image = Context->InstanceUpdateData.Images[Surface.FirstImage + ImageIndex];
+						const TSharedRef<const FInstanceUpdateData::FImage>& Image = Context->InstanceUpdateData.Images[Surface->FirstImage + ImageIndex];
 						FString KeyName = Image->Name.ToString();
 						TSharedPtr<const UE::Mutable::Private::FImage> MutableImage = Image->Image;
 
 						UTexture2D* MutableTexture = nullptr; // Texture generated by mutable
 						UTexture* PassThroughTexture = nullptr; // Texture not generated by mutable
 
-						const FTextureReuseCacheKey TextureReuseCacheRef = bReuseTextures ? FTextureReuseCacheKey{Image->BaseLOD, ObjectComponentIndex.GetValue(), Surface.SurfaceId, ImageIndex} : FTextureReuseCacheKey{};
+						const FTextureReuseCacheKey TextureReuseCacheRef = bReuseTextures ? FTextureReuseCacheKey{Image->BaseLOD, ObjectComponentIndex.GetValue(), Surface->SurfaceId, ImageIndex} : FTextureReuseCacheKey{};
 
 #if WITH_EDITOR
 						const FTextureCache::FId TextureCacheKey = FTextureCache::FId(Image->ImageID, Context->MipsToSkip, Context->bBake);
@@ -6141,7 +6141,7 @@ void UCustomizableInstancePrivate::BuildMaterials(const TSharedRef<FUpdateContex
 													MutableMipDataProviderFactory->CustomizableObjectInstance = Public;
 													check(LODIndex < 256 && InstanceComponentIndex < 256 && ImageIndex < 256);
 													MutableMipDataProviderFactory->ImageRef.ImageID = Image->ImageID;
-													MutableMipDataProviderFactory->ImageRef.SurfaceId = Surface.SurfaceId;
+													MutableMipDataProviderFactory->ImageRef.SurfaceId = Surface->SurfaceId;
 													MutableMipDataProviderFactory->ImageRef.LOD = uint8(Image->BaseLOD);
 													MutableMipDataProviderFactory->ImageRef.Component = uint8(InstanceComponentIndex);
 													MutableMipDataProviderFactory->ImageRef.Image = uint8(ImageIndex);
@@ -6210,8 +6210,8 @@ void UCustomizableInstancePrivate::BuildMaterials(const TSharedRef<FUpdateContex
 					ReuseMaterialCache.Add(MutableMaterialPlaceholder);
 					
 					FGeneratedMaterial& Material = GeneratedMaterials.AddDefaulted_GetRef();
-					Material.SurfaceId = Surface.SurfaceId;
-					Material.Material = Surface.Material;
+					Material.SurfaceId = Surface->SurfaceId;
+					Material.Material = Surface->Material;
 					Material.MaterialInterface = MaterialInterface;
 
 #if WITH_EDITORONLY_DATA
@@ -6420,7 +6420,7 @@ void UCustomizableInstancePrivate::BuildMaterials(const TSharedRef<FUpdateContex
 			int32 InstanceComponentIndex = -1;
 			for (int32 CurrentInstanceIndex=0; CurrentInstanceIndex<Context->InstanceUpdateData.Components.Num(); ++CurrentInstanceIndex)
 			{
-				if (Context->InstanceUpdateData.Components[CurrentInstanceIndex].Id.GetValue() == ObjectComponentIndex)
+				if (Context->InstanceUpdateData.Components[CurrentInstanceIndex]->Id.GetValue() == ObjectComponentIndex)
 				{
 					InstanceComponentIndex = CurrentInstanceIndex;
 					break;
diff --git a/Engine/Plugins/Mutable/Source/CustomizableObject/Private/MuCO/CustomizableObjectSkeletalMesh.cpp b/Engine/Plugins/Mutable/Source/CustomizableObject/Private/MuCO/CustomizableObjectSkeletalMesh.cpp
--- a/Engine/Plugins/Mutable/Source/CustomizableObject/Private/MuCO/CustomizableObjectSkeletalMesh.cpp
+++ b/Engine/Plugins/Mutable/Source/CustomizableObject/Private/MuCO/CustomizableObjectSkeletalMesh.cpp
@@ -45,15 +45,15 @@ void UCustomizableObjectSkeletalMesh::InitMutableStreamingData(const TSharedRef<
 
 	for (int32 LODIndex = FirstLOD; LODIndex < LODCount; ++LODIndex)
 	{
-		const FInstanceUpdateData::FLOD& LOD = InOperationData->InstanceUpdateData.LODs[InstanceUpdateFirstLOD + LODIndex];
-		MeshIDs[LODIndex] = LOD.MeshId;
+		const TSharedRef<const FInstanceUpdateData::FLOD>& LOD = InOperationData->InstanceUpdateData.LODs[InstanceUpdateFirstLOD + LODIndex];
+		MeshIDs[LODIndex] = LOD->MeshId;
 
-		SurfaceIDs[LODIndex].SetNum(LOD.SurfaceCount);
-		check(LOD.Mesh && LOD.Mesh->GetSurfaceCount() == LOD.SurfaceCount);
+		SurfaceIDs[LODIndex].SetNum(LOD->SurfaceCount);
+		check(LOD->Mesh && LOD->Mesh->GetSurfaceCount() == LOD->SurfaceCount);
 
-		for (int32 SurfaceIndex = 0; SurfaceIndex < LOD.SurfaceCount; ++SurfaceIndex)
+		for (int32 SurfaceIndex = 0; SurfaceIndex < LOD->SurfaceCount; ++SurfaceIndex)
 		{
-			SurfaceIDs[LODIndex][SurfaceIndex] = LOD.Mesh->GetSurfaceId(SurfaceIndex);
+			SurfaceIDs[LODIndex][SurfaceIndex] = LOD->Mesh->GetSurfaceId(SurfaceIndex);
 		}
 	}
 
diff --git a/Engine/Plugins/Mutable/Source/CustomizableObject/Private/MuCO/CustomizableObjectSystem.cpp b/Engine/Plugins/Mutable/Source/CustomizableObject/Private/MuCO/CustomizableObjectSystem.cpp
--- a/Engine/Plugins/Mutable/Source/CustomizableObject/Private/MuCO/CustomizableObjectSystem.cpp
+++ b/Engine/Plugins/Mutable/Source/CustomizableObject/Private/MuCO/CustomizableObjectSystem.cpp
@@ -57,7 +57,6 @@ DECLARE_CYCLE_STAT(TEXT("MutablePendingRelease Time"), STAT_MutablePendingReleas
 DECLARE_CYCLE_STAT(TEXT("MutableTask"), STAT_MutableTask, STATGROUP_Game);
 
 #define UE_MUTABLE_UPDATE_REGION TEXT("Mutable Update")
-#define UE_TASK_MUTABLE_GETMESHES_REGION TEXT("Task_Mutable_GetMeshes")
 
 
 UCustomizableObjectSystem* UCustomizableObjectSystemPrivate::SSystem = nullptr;
@@ -572,7 +571,7 @@ FCustomizableObjectComponentIndex FUpdateContextPrivate::GetObjectComponentIndex
 	FCustomizableObjectComponentIndex ObjectComponentIndex;
 	if (InstanceUpdateData.Components.IsValidIndex(InstanceComponentIndex.GetValue()))
 	{
-		ObjectComponentIndex = InstanceUpdateData.Components[InstanceComponentIndex.GetValue()].Id;
+		ObjectComponentIndex = InstanceUpdateData.Components[InstanceComponentIndex.GetValue()]->Id;
 	}
 	else
 	{
@@ -586,7 +585,7 @@ const FInstanceUpdateData::FComponent* FUpdateContextPrivate::GetComponentUpdate
 {
 	if (InstanceUpdateData.Components.IsValidIndex(InstanceComponentIndex.GetValue()))
 	{
-		return &InstanceUpdateData.Components[InstanceComponentIndex.GetValue()];
+		return &InstanceUpdateData.Components[InstanceComponentIndex.GetValue()].Get();
 	}
 	return nullptr;
 }
@@ -1918,44 +1917,9 @@ void UCustomizableObjectSystemPrivate::UpdateMemoryLimit()
 }
 
 
-// Asynchronous tasks performed during the creation or update of a mutable instance. 
-// Check the documentation before modifying and keep it up to date.
-// https://docs.google.com/drawings/d/109NlsdKVxP59K5TuthJkleVG3AROkLJr6N03U4bNp4s
-// When it says "mutable thread" it means any task pool thread, but with the guarantee that no other thread is using the mutable runtime.
-// Naming: Task_<thread>_<description>
 namespace impl
 {
-	struct FGetImageData
-	{
-		int32 ImageIndex;
-		UE::Mutable::Private::FImageId ImageID;
-	};
-	
-
-	/** Process the next Image. If there are no more Images, go to the end of the task. */
-	void Task_Mutable_GetMeshes_GetImage_Loop(
-		const TSharedRef<FUpdateContextPrivate>& OperationData,
-		double StartTime,
-		const TSharedRef<TArray<FGetImageData>>& GetImagesData,
-		int32 GetImageIndex);
-	
-
-	struct FGetMeshData
-	{
-		int32 InstanceUpdateLODIndex;
-		UE::Mutable::Private::FMeshId MeshID;
-		UE::Mutable::Private::EMeshContentFlags ContentFilter;
-	};
-
-
-	/** Process the next Mesh. If there are no more Meshes, go to the process Images loop. */
-	void Task_Mutable_GetMeshes_GetMesh_Loop(
-		const TSharedRef<FUpdateContextPrivate>& OperationData,
-		double StartTime,
-		const TSharedRef<TArray<FGetMeshData>>& GetMeshesData,
-		int32 GetMeshIndex);
-	
-	void Subtask_Mutable_UpdateParameterRelevancy(const TSharedRef<FUpdateContextPrivate>& OperationData)
+	void Task_Mutable_UpdateParameterRelevancy(const TSharedRef<FUpdateContextPrivate>& OperationData)
 	{
 		MUTABLE_CPUPROFILER_SCOPE(Subtask_Mutable_UpdateParameterRelevancy)
 
@@ -2032,7 +1996,7 @@ namespace impl
 	
 
 	// This runs in a worker thread
-	void Subtask_Mutable_PrepareSkeletonData(const TSharedRef<FUpdateContextPrivate>& OperationData)
+	void Task_Mutable_PrepareSkeletonData(const TSharedRef<FUpdateContextPrivate>& OperationData)
 	{
 		MUTABLE_CPUPROFILER_SCOPE(Subtask_Mutable_PrepareSkeletonData);
 
@@ -2041,13 +2005,13 @@ namespace impl
 
 		for (int32 InstanceComponentIndex = 0; InstanceComponentIndex< NumInstanceComponents; ++InstanceComponentIndex)
 		{
-			const FInstanceUpdateData::FComponent& Component = OperationData->InstanceUpdateData.Components[InstanceComponentIndex];
+			const TSharedRef<FInstanceUpdateData::FComponent>& Component = OperationData->InstanceUpdateData.Components[InstanceComponentIndex];
 
-			for (int32 LODIndex = 0; LODIndex < Component.LODCount; ++LODIndex)
+			for (int32 LODIndex = 0; LODIndex < Component->LODCount; ++LODIndex)
 			{
-				FInstanceUpdateData::FLOD& LOD = OperationData->InstanceUpdateData.LODs[Component.FirstLOD + LODIndex];
+				TSharedRef<FInstanceUpdateData::FLOD>& LOD = OperationData->InstanceUpdateData.LODs[Component->FirstLOD + LODIndex];
 
-				TSharedPtr<const UE::Mutable::Private::FMesh> Mesh = LOD.Mesh;
+				TSharedPtr<const UE::Mutable::Private::FMesh> Mesh = LOD->Mesh;
 				if (!Mesh || Mesh->IsReference())
 				{
 					continue;
@@ -2074,14 +2038,14 @@ namespace impl
 
 				// Append BoneMap to the array of BoneMaps
 				const TArray<UE::Mutable::Private::FBoneName>& BoneMap = Mesh->GetBoneMap();
-				LOD.FirstBoneMap = OperationData->InstanceUpdateData.BoneMaps.Num();
-				LOD.BoneMapCount = BoneMap.Num();
+				LOD->FirstBoneMap = OperationData->InstanceUpdateData.BoneMaps.Num();
+				LOD->BoneMapCount = BoneMap.Num();
 				OperationData->InstanceUpdateData.BoneMaps.Append(BoneMap);
 
 				// Add active bone indices and poses
-				LOD.FirstActiveBone = OperationData->InstanceUpdateData.ActiveBones.Num();
-				LOD.ActiveBoneCount = Mesh->GetBonePoseCount();
-				for (uint32 BoneIndex = 0; BoneIndex < LOD.ActiveBoneCount; ++BoneIndex)
+				LOD->FirstActiveBone = OperationData->InstanceUpdateData.ActiveBones.Num();
+				LOD->ActiveBoneCount = Mesh->GetBonePoseCount();
+				for (uint32 BoneIndex = 0; BoneIndex < LOD->ActiveBoneCount; ++BoneIndex)
 				{
 					const UE::Mutable::Private::FBoneName& BoneId = Mesh->GetBonePoseId(BoneIndex);
 
@@ -2123,10 +2087,10 @@ namespace impl
 		OperationData->InstanceUpdateData.RealTimeMorphTargets.Reserve(NumInstanceComponents);
 		for ( int32 InstanceComponentIndex = 0; InstanceComponentIndex< NumInstanceComponents; ++InstanceComponentIndex)
 		{
-			FInstanceUpdateData::FComponent& Component = OperationData->InstanceUpdateData.Components[InstanceComponentIndex];
-			check(Component.Id.IsValid());
+			TSharedRef<FInstanceUpdateData::FComponent>& Component = OperationData->InstanceUpdateData.Components[InstanceComponentIndex];
+			check(Component->Id.IsValid());
 
-			const FName& ComponentName = OperationData->ComponentNames[Component.Id.GetValue()];
+			const FName& ComponentName = OperationData->ComponentNames[Component->Id.GetValue()];
 
 			FSkeletalMeshMorphTargets& ComponentMorphTargetsData = OperationData->InstanceUpdateData.RealTimeMorphTargets.Add(ComponentName);
 			
@@ -2164,15 +2128,15 @@ namespace impl
 			
 			for (int32 LODIndex = OperationData->GetFirstRequestedLOD()[ComponentName]; LODIndex < NumLODsAvailable; ++LODIndex)
 			{
-				const FInstanceUpdateData::FLOD& LOD = UpdateData.LODs[Component.FirstLOD+LODIndex];
+				const TSharedRef<FInstanceUpdateData::FLOD>& LOD = UpdateData.LODs[Component->FirstLOD + LODIndex];
 
-				if (!LOD.Mesh)
+				if (!LOD->Mesh)
 				{
 					continue;
 				}
 
 				TArray<FMorphTargetLODModel> MorphTargets;
-				ReconstructMorphTargets(*LOD.Mesh.Get(), ComponentMorphTargetsData.RealTimeMorphTargetNames, MorphTargetMeshData, MorphTargets);
+				ReconstructMorphTargets(*LOD->Mesh.Get(), ComponentMorphTargetsData.RealTimeMorphTargetNames, MorphTargetMeshData, MorphTargets);
 
 				for (int32 NameIndex = 0; NameIndex < ComponentMorphTargetsData.RealTimeMorphTargetNames.Num(); NameIndex++)
 				{
@@ -2190,12 +2154,12 @@ namespace impl
 			// Find which Sections are being used in each LOD (Streamed and Residents).
 			for (int32 LODIndex = FirstLOD; LODIndex < NumLODsAvailable; ++LODIndex)
 			{
-				const FInstanceUpdateData::FLOD& LOD = UpdateData.LODs[Component.FirstLOD+LODIndex];
-				check(LOD.Mesh);
+				const TSharedRef<FInstanceUpdateData::FLOD>& LOD = UpdateData.LODs[Component->FirstLOD+LODIndex];
+				check(LOD->Mesh);
 				
-				for (int32 SectionIndex = 0; SectionIndex < LOD.Mesh->Surfaces.Num(); ++SectionIndex)
+				for (int32 SectionIndex = 0; SectionIndex < LOD->Mesh->Surfaces.Num(); ++SectionIndex)
 				{
-					const UE::Mutable::Private::FMeshSurface& Surface = LOD.Mesh->Surfaces[SectionIndex];
+					const UE::Mutable::Private::FMeshSurface& Surface = LOD->Mesh->Surfaces[SectionIndex];
 					for (const UE::Mutable::Private::FSurfaceSubMesh& SubMesh : Surface.SubMeshes)
 					{
 						if (SubMesh.ExternalId == 0)
@@ -2247,90 +2211,11 @@ namespace impl
 		// Free unneeded data memory.
 		UpdateData.RealTimeMorphTargetMeshData.Empty();
 	}
-
-	/** End of the GetMeshes tasks. */
-	void Task_Mutable_GetMeshes_End(
-		const TSharedRef<FUpdateContextPrivate>& OperationData,
-		double StartTime)
-	{
-		MUTABLE_CPUPROFILER_SCOPE(Task_Mutable_GetMeshes_End)
-
-		// TODO: Not strictly mutable: move to another worker thread task to free mutable access?
-		Subtask_Mutable_PrepareSkeletonData(OperationData);
-		if (OperationData->GetCapturedDescriptor().GetBuildParameterRelevancy())
-		{
-			Subtask_Mutable_UpdateParameterRelevancy(OperationData);
-		}
-		else
-		{
-			OperationData->RelevantParametersInProgress.Reset();
-		}
-		
-		OperationData->TaskGetMeshTime = FPlatformTime::Seconds() - StartTime;
-
-		TRACE_END_REGION(UE_TASK_MUTABLE_GETMESHES_REGION);
-	}
-
-
-	/** TaskGraph task after GetImage has completed. */
-	void Task_Mutable_GetMeshes_GetImage_Post(
-		const TSharedRef<FUpdateContextPrivate>& OperationData,
-		double StartTime,
-		const TSharedRef<TArray<FGetImageData>>& GetImagesData,
-		int32 GetImageIndex,
-		UE::Tasks::TTask<TSharedPtr<const UE::Mutable::Private::FImage>> GetImageTask)
-	{
-		MUTABLE_CPUPROFILER_SCOPE(Task_Mutable_GetMeshes_GetImage_Post)
-
-		const UCustomizableObjectInstance* Instance = OperationData->Instance.Get();
-		const UCustomizableObject* CustomizableObject = Instance->GetCustomizableObject();
-		const UModelResources& ModelResources = *CustomizableObject->GetPrivate()->GetModelResources();
-
-		const int32 ImageIndex = (*GetImagesData)[GetImageIndex].ImageIndex;
-
-		TSharedRef<FInstanceUpdateData::FImage> Image = OperationData->InstanceUpdateData.Images[ImageIndex];
-
-		Image->Image = GetImageTask.GetResult();
-		check(Image->Image->IsReference());
-		
-		Task_Mutable_GetMeshes_GetImage_Loop(OperationData, StartTime, GetImagesData, ++GetImageIndex);
-	}
-
 	
-	/** See declaration. */
-	void Task_Mutable_GetMeshes_GetImage_Loop(
-		const TSharedRef<FUpdateContextPrivate>& OperationData,
-		double StartTime,
-		const TSharedRef<TArray<FGetImageData>>& GetImagesData,
-		int32 GetImageIndex)
-	{
-		MUTABLE_CPUPROFILER_SCOPE(Task_Mutable_GetMesh_GetImages_Loop)
 
-		if (GetImageIndex >= GetImagesData->Num()) 
-		{
-			Task_Mutable_GetMeshes_End(OperationData, StartTime);
-			return;
-		}
-
-		const FGetImageData& ImageData = (*GetImagesData)[GetImageIndex];
-		
-		UE::Tasks::TTask<TSharedPtr<const UE::Mutable::Private::FImage>> GetImageTask = OperationData->MutableSystem->GetImage({}, OperationData->InstanceID, ImageData.ImageID, 0, 0);
-
-		UE::Tasks::AddNested(UE::Tasks::Launch(TEXT("Task_Mutable_GetMeshes_GetImage_Post"), [=]()
-		{
-			Task_Mutable_GetMeshes_GetImage_Post(OperationData, StartTime, GetImagesData, GetImageIndex, GetImageTask);
-		},
-		GetImageTask,
-		LowLevelTasks::ETaskPriority::Inherit));
-	}
-
-
-	/** Gather all GetImages that have to be called. */
-	void Task_Mutable_GetMeshes_GetImages(
-		const TSharedRef<FUpdateContextPrivate>& OperationData,
-		double StartTime)
+	void Task_Mutable_GetMaterials(const TSharedRef<FUpdateContextPrivate>& OperationData)
 	{
-		MUTABLE_CPUPROFILER_SCOPE(Task_Mutable_GetMeshes_GetImages)
+		MUTABLE_CPUPROFILER_SCOPE(Task_Mutable_GetMaterials)
 		
 		const UCustomizableObjectInstance* Instance = OperationData->Instance.Get();
 		const UCustomizableObject* CustomizableObject = Instance->GetCustomizableObject();
@@ -2339,389 +2224,290 @@ namespace impl
 		const UE::Mutable::Private::FInstance* MutableInstance = OperationData->MutableInstance.Get();
 
 		TArray<uint32> SurfacesSharedId;
-
-		TArray<int32> GetMaterialSurfaces;
 		
-		const TSharedRef<TArray<FGetImageData>> GetImagesData = MakeShared<TArray<FGetImageData>>();
+		UE::Tasks::FTask LastTask = UE::Tasks::MakeCompletedTask<void>();
 		
 		for (int32 InstanceComponentIndex = 0; InstanceComponentIndex < OperationData->NumInstanceComponents; ++InstanceComponentIndex)
 		{
-			FInstanceUpdateData::FComponent& Component = OperationData->InstanceUpdateData.Components[InstanceComponentIndex];
-			const FName& ComponentName = ModelResources.ComponentNamesPerObjectComponent[Component.Id.GetValue()];
+			TSharedRef<FInstanceUpdateData::FComponent>& Component = OperationData->InstanceUpdateData.Components[InstanceComponentIndex];
+			const FName& ComponentName = ModelResources.ComponentNamesPerObjectComponent[Component->Id.GetValue()];
 
-			for (int32 MutableLODIndex = OperationData->FirstLODAvailable[ComponentName]; MutableLODIndex < Component.LODCount; ++MutableLODIndex)
+			for (int32 MutableLODIndex = OperationData->FirstLODAvailable[ComponentName]; MutableLODIndex < Component->LODCount; ++MutableLODIndex)
 			{
-				FInstanceUpdateData::FLOD& LOD = OperationData->InstanceUpdateData.LODs[Component.FirstLOD + MutableLODIndex];
+				TSharedRef<FInstanceUpdateData::FLOD>& LOD = OperationData->InstanceUpdateData.LODs[Component->FirstLOD + MutableLODIndex];
 
-				LOD.FirstSurface = OperationData->InstanceUpdateData.Surfaces.Num();
-				LOD.SurfaceCount = 0;
+				LOD->FirstSurface = OperationData->InstanceUpdateData.Surfaces.Num();
+				LOD->SurfaceCount = 0;
 
-				if (!LOD.Mesh)
+				if (!LOD->Mesh)
 				{
 					continue;
 				}
+				
+				// Materials and images
+				const int32 SurfaceCount = LOD->Mesh->GetSurfaceCount();
+				for (int32 MeshSurfaceIndex = 0; MeshSurfaceIndex < SurfaceCount; ++MeshSurfaceIndex)
+				{
+					const uint32 SurfaceId = LOD->Mesh->GetSurfaceId(MeshSurfaceIndex);
 
-				// This lambda does all the work to fill up the surface data
-				auto AddSurface = 
-					[&LOD, &SurfacesSharedId, &ModelResources, &GetImagesData, &GetMaterialSurfaces, OperationData, MutableInstance, CustomizableObject, InstanceComponentIndex, MutableLODIndex]
-					(uint32 SurfaceId, uint32 SurfaceMetadataId, int32 InstanceSurfaceIndex)
-					{
-						int32 BaseSurfaceIndex = InstanceSurfaceIndex;
-						int32 BaseLODIndex = MutableLODIndex;
-
-						int32 SurfaceIndex = OperationData->InstanceUpdateData.Surfaces.AddDefaulted();
-						FInstanceUpdateData::FSurface& Surface = OperationData->InstanceUpdateData.Surfaces.Last();
-						++LOD.SurfaceCount;
-						
-						const uint32 SharedSurfaceId = MutableInstance->GetSharedSurfaceId(InstanceComponentIndex, MutableLODIndex, InstanceSurfaceIndex);
-						const int32 SharedSurfaceIndex = SurfacesSharedId.Find(SharedSurfaceId);
-						SurfacesSharedId.Add(SharedSurfaceId);
-
-						if (SharedSurfaceIndex != INDEX_NONE)
-						{
-							Surface = OperationData->InstanceUpdateData.Surfaces[SharedSurfaceIndex];
-							return;
-						}
+					const int32 InstanceSurfaceIndex = MutableInstance->FindSurfaceById(InstanceComponentIndex, MutableLODIndex, SurfaceId);
+					check(LOD->Mesh->GetVertexCount() > 0 || InstanceSurfaceIndex >= 0);
 
-						// New Surface.MaterialIndex is decoded from a parameter at the end of this if()
-						Surface.SurfaceId = SurfaceId;
-						Surface.SurfaceMetadataId = SurfaceMetadataId;
-						Surface.SharedSurfaceId = SharedSurfaceId;
+					if (InstanceSurfaceIndex < 0)
+					{
+						continue;
+					}
 
-						// Find the first LOD where this surface can be found
-						MutableInstance->FindBaseSurfaceBySharedId(InstanceComponentIndex, Surface.SharedSurfaceId, BaseSurfaceIndex, BaseLODIndex);
+					int32 BaseSurfaceIndex = InstanceSurfaceIndex;
+					int32 BaseLODIndex = MutableLODIndex;
 
-						Surface.SurfaceId = MutableInstance->GetSurfaceId(InstanceComponentIndex, BaseLODIndex, BaseSurfaceIndex);
-						Surface.SurfaceMetadataId = MutableInstance->GetSurfaceCustomId(InstanceComponentIndex, BaseLODIndex, BaseSurfaceIndex);
-						Surface.MaterialId = MutableInstance->GetMaterialId(InstanceComponentIndex, BaseLODIndex, BaseSurfaceIndex);
-						if (Surface.MaterialId)
-						{
-							GetMaterialSurfaces.Add(SurfaceIndex);
-						}
+					int32 SurfaceIndex = OperationData->InstanceUpdateData.Surfaces.Add(MakeShared<FInstanceUpdateData::FSurface>());
+					TSharedRef<FInstanceUpdateData::FSurface>& Surface = OperationData->InstanceUpdateData.Surfaces.Last();
+					++LOD->SurfaceCount;
+					
+					const uint32 SharedSurfaceId = MutableInstance->GetSharedSurfaceId(InstanceComponentIndex, MutableLODIndex, InstanceSurfaceIndex);
+					const int32 SharedSurfaceIndex = SurfacesSharedId.Find(SharedSurfaceId);
+					SurfacesSharedId.Add(SharedSurfaceId);
 
-						// Vectors
-						Surface.FirstVector = OperationData->InstanceUpdateData.Vectors.Num();
-						Surface.VectorCount = MutableInstance->GetVectorCount(InstanceComponentIndex, BaseLODIndex, BaseSurfaceIndex);
-						for (int32 VectorIndex = 0; VectorIndex < Surface.VectorCount; ++VectorIndex)
-						{
-							MUTABLE_CPUPROFILER_SCOPE(GetVector);
-							OperationData->InstanceUpdateData.Vectors.Push({});
-							FInstanceUpdateData::FVector& Vector = OperationData->InstanceUpdateData.Vectors.Last();
-							Vector.Name = MutableInstance->GetVectorName(InstanceComponentIndex, BaseLODIndex, BaseSurfaceIndex, VectorIndex);
-							Vector.Vector = MutableInstance->GetVector(InstanceComponentIndex, BaseLODIndex, BaseSurfaceIndex, VectorIndex);
-							Vector.MaterialLayer = MutableInstance->GetVectorLayerIndex(InstanceComponentIndex, BaseLODIndex, BaseSurfaceIndex, VectorIndex);
-						}
+					if (SharedSurfaceIndex != INDEX_NONE)
+					{
+						Surface = OperationData->InstanceUpdateData.Surfaces[SharedSurfaceIndex];
+						continue;
+					}
 
-						// Scalars
-						Surface.FirstScalar = OperationData->InstanceUpdateData.Scalars.Num();
-						Surface.ScalarCount = MutableInstance->GetScalarCount(InstanceComponentIndex, BaseLODIndex, BaseSurfaceIndex);
-						for (int32 ScalarIndex = 0; ScalarIndex < Surface.ScalarCount; ++ScalarIndex)
-						{
-							MUTABLE_CPUPROFILER_SCOPE(GetScalar)
+					// New Surface.MaterialIndex is decoded from a parameter at the end of this if()
+					Surface->SurfaceId = SurfaceId;
+					Surface->SurfaceMetadataId = 0;
+					Surface->SharedSurfaceId = SharedSurfaceId;
 
-							const FName ScalarName = MutableInstance->GetScalarName(InstanceComponentIndex, BaseLODIndex, BaseSurfaceIndex, ScalarIndex);
-							const float ScalarValue = MutableInstance->GetScalar(InstanceComponentIndex, BaseLODIndex, BaseSurfaceIndex, ScalarIndex);
-							const int32 MaterialLayer = MutableInstance->GetScalarLayerIndex(InstanceComponentIndex, BaseLODIndex, BaseSurfaceIndex, ScalarIndex);
+					// Find the first LOD where this surface can be found
+					MutableInstance->FindBaseSurfaceBySharedId(InstanceComponentIndex, Surface->SharedSurfaceId, BaseSurfaceIndex, BaseLODIndex);
 
-							OperationData->InstanceUpdateData.Scalars.Push({ ScalarName, ScalarValue, MaterialLayer });
-						}
+					Surface->SurfaceId = MutableInstance->GetSurfaceId(InstanceComponentIndex, BaseLODIndex, BaseSurfaceIndex);
+					Surface->SurfaceMetadataId = MutableInstance->GetSurfaceCustomId(InstanceComponentIndex, BaseLODIndex, BaseSurfaceIndex);
 
-						// Images
-						Surface.FirstImage = OperationData->InstanceUpdateData.Images.Num();
-						Surface.ImageCount = MutableInstance->GetImageCount(InstanceComponentIndex, BaseLODIndex, BaseSurfaceIndex);
-						for (int32 ImageIndex = 0; ImageIndex < Surface.ImageCount; ++ImageIndex)
-						{
-							MUTABLE_CPUPROFILER_SCOPE(GetImageId);
-
-							const int32 UpdateDataImageIndex = OperationData->InstanceUpdateData.Images.Add(MakeShared<FInstanceUpdateData::FImage>());
-							TSharedRef<FInstanceUpdateData::FImage> Image = OperationData->InstanceUpdateData.Images.Last();
-							Image->Name = MutableInstance->GetImageName(InstanceComponentIndex, BaseLODIndex, BaseSurfaceIndex, ImageIndex);
-							Image->ImageID = MutableInstance->GetImageId(InstanceComponentIndex, BaseLODIndex, BaseSurfaceIndex, ImageIndex);
-							Image->MaterialLayer = MutableInstance->GetImageLayerIndex(InstanceComponentIndex, BaseLODIndex, BaseSurfaceIndex, ImageIndex);
-							Image->FullImageSizeX = 0;
-							Image->FullImageSizeY = 0;
-							Image->BaseLOD = BaseLODIndex;
-							Image->BaseMip = 0;
-
-							FString KeyName = Image->Name.ToString();
-							int32 ImageKey = FCString::Atoi(*KeyName);
-
-							if (ImageKey >= 0 && ImageKey < ModelResources.ImageProperties.Num())
-							{
-								const FMutableModelImageProperties& Props = ModelResources.ImageProperties[ImageKey];
+					Surface->MaterialId = MutableInstance->GetMaterialId(InstanceComponentIndex, BaseLODIndex, BaseSurfaceIndex);
 
-								Image->bIsNonProgressive = Props.MipGenSettings == TMGS_NoMipmaps;
+					UE::Tasks::TTask<TSharedPtr<const UE::Mutable::Private::FMaterial>> GetMaterialTask = OperationData->MutableSystem->GetMaterial(LastTask, OperationData->InstanceID, Surface->MaterialId);
+					UE::Tasks::AddNested(GetMaterialTask);
+					LastTask = GetMaterialTask;
+					
+					UE::Tasks::AddNested(UE::Tasks::Launch(TEXT("GetMaterialResult"), [Surface, GetMaterialTask]() mutable
+					{
+						Surface->Material = GetMaterialTask.GetResult();
+					},
+					GetMaterialTask,
+					LowLevelTasks::ETaskPriority::Inherit));
 
-								if (Props.IsPassThrough)
-								{
-									Image->bIsPassThrough = true;
+					// Vectors
+					Surface->FirstVector = OperationData->InstanceUpdateData.Vectors.Num();
+					Surface->VectorCount = MutableInstance->GetVectorCount(InstanceComponentIndex, BaseLODIndex, BaseSurfaceIndex);
+					for (int32 VectorIndex = 0; VectorIndex < Surface->VectorCount; ++VectorIndex)
+					{
+						MUTABLE_CPUPROFILER_SCOPE(GetVector);
+						OperationData->InstanceUpdateData.Vectors.Push({});
+						FInstanceUpdateData::FVector& Vector = OperationData->InstanceUpdateData.Vectors.Last();
+						Vector.Name = MutableInstance->GetVectorName(InstanceComponentIndex, BaseLODIndex, BaseSurfaceIndex, VectorIndex);
+						Vector.Vector = MutableInstance->GetVector(InstanceComponentIndex, BaseLODIndex, BaseSurfaceIndex, VectorIndex);
+						Vector.MaterialLayer = MutableInstance->GetVectorLayerIndex(InstanceComponentIndex, BaseLODIndex, BaseSurfaceIndex, VectorIndex);
+					}
 
-									// Since it's known it's a pass-through texture there is no need to cache or convert it so we can generate it here already.
-									GetImagesData->Add({ UpdateDataImageIndex, Image->ImageID });
-								}
-							}
-							else
-							{
-								// This means the compiled model (maybe coming from derived data) has images that the asset doesn't know about.
-								UE_LOG(LogMutable, Error, TEXT("CustomizableObject derived data out of sync with asset for [%s]. Try recompiling it."), *CustomizableObject->GetName());
-							}
-						}
+					// Scalars
+					Surface->FirstScalar = OperationData->InstanceUpdateData.Scalars.Num();
+					Surface->ScalarCount = MutableInstance->GetScalarCount(InstanceComponentIndex, BaseLODIndex, BaseSurfaceIndex);
+					for (int32 ScalarIndex = 0; ScalarIndex < Surface->ScalarCount; ++ScalarIndex)
+					{
+						MUTABLE_CPUPROFILER_SCOPE(GetScalar)
 
-					};
+						const FName ScalarName = MutableInstance->GetScalarName(InstanceComponentIndex, BaseLODIndex, BaseSurfaceIndex, ScalarIndex);
+						const float ScalarValue = MutableInstance->GetScalar(InstanceComponentIndex, BaseLODIndex, BaseSurfaceIndex, ScalarIndex);
 
-				// Materials and images
+						const int32 MaterialLayer = MutableInstance->GetScalarLayerIndex(InstanceComponentIndex, BaseLODIndex, BaseSurfaceIndex, ScalarIndex);
 
-				// If the mesh is a reference mesh, it won't have the surface information in the mutable mesh. We need to get it from the instance
-				// and all defined surfaces will be present. 
-				if (LOD.Mesh->IsReference())
-				{
-					const int32 SurfaceCount = MutableInstance->GetSurfaceCount(InstanceComponentIndex, MutableLODIndex);
-					for (int32 SurfaceIndex = 0; SurfaceIndex < SurfaceCount; ++SurfaceIndex)
-					{
-						uint32 SurfaceId = MutableInstance->GetSurfaceId(InstanceComponentIndex, MutableLODIndex, SurfaceIndex);
-						uint32 SurfaceMetadataId = MutableInstance->GetSurfaceCustomId(InstanceComponentIndex, MutableLODIndex, SurfaceIndex);
-						AddSurface(SurfaceId, SurfaceMetadataId, SurfaceIndex);
+						OperationData->InstanceUpdateData.Scalars.Push({ ScalarName, ScalarValue, MaterialLayer });
 					}
-				}
 
-				// If the mesh is a not a reference mesh, we have to add only the materials of the surfaces that appear in the actual final mesh. 
-				else
-				{
-					const int32 SurfaceCount = LOD.Mesh->GetSurfaceCount();
-					for (int32 MeshSurfaceIndex = 0; MeshSurfaceIndex < SurfaceCount; ++MeshSurfaceIndex)
+					// Images
+					Surface->FirstImage = OperationData->InstanceUpdateData.Images.Num();
+					Surface->ImageCount = MutableInstance->GetImageCount(InstanceComponentIndex, BaseLODIndex, BaseSurfaceIndex);
+					for (int32 ImageIndex = 0; ImageIndex < Surface->ImageCount; ++ImageIndex)
 					{
-						const uint32 SurfaceId = LOD.Mesh->GetSurfaceId(MeshSurfaceIndex);
+						MUTABLE_CPUPROFILER_SCOPE(GetImageId);
+						
+						TSharedRef<FInstanceUpdateData::FImage>& Image = OperationData->InstanceUpdateData.Images.Add_GetRef(MakeShared<FInstanceUpdateData::FImage>());
+						Image->Name = MutableInstance->GetImageName(InstanceComponentIndex, BaseLODIndex, BaseSurfaceIndex, ImageIndex);
+						Image->ImageID = MutableInstance->GetImageId(InstanceComponentIndex, BaseLODIndex, BaseSurfaceIndex, ImageIndex);
+						Image->MaterialLayer = MutableInstance->GetImageLayerIndex(InstanceComponentIndex, BaseLODIndex, BaseSurfaceIndex, ImageIndex);
+						Image->FullImageSizeX = 0;
+						Image->FullImageSizeY = 0;
+						Image->BaseLOD = BaseLODIndex;
+						Image->BaseMip = 0;
+
+						FString KeyName = Image->Name.ToString();
+						int32 ImageKey = FCString::Atoi(*KeyName);
+
+						if (ImageKey >= 0 && ImageKey < ModelResources.ImageProperties.Num())
+						{
+							const FMutableModelImageProperties& Props = ModelResources.ImageProperties[ImageKey];
+
+							Image->bIsNonProgressive = Props.MipGenSettings == TMGS_NoMipmaps;
 
-						const int32 InstanceSurfaceIndex = MutableInstance->FindSurfaceById(InstanceComponentIndex, MutableLODIndex, SurfaceId);
-						check(LOD.Mesh->GetVertexCount() > 0 || InstanceSurfaceIndex >= 0);
+							if (Props.IsPassThrough)
+							{
+								Image->bIsPassThrough = true;
 
-						if (InstanceSurfaceIndex >= 0)
+								UE::Tasks::TTask<TSharedPtr<const UE::Mutable::Private::FImage>> GetImageTask = OperationData->MutableSystem->GetImage(LastTask, OperationData->InstanceID, Image->ImageID, 0, 0);
+								UE::Tasks::AddNested(GetImageTask);
+								LastTask = GetImageTask;
+					
+								UE::Tasks::AddNested(UE::Tasks::Launch(TEXT("GetImageResult"), [Image, GetImageTask]() mutable
+								{
+									Image->Image = GetImageTask.GetResult();
+									check(Image->Image->IsReference());
+								},
+								GetImageTask,
+								LowLevelTasks::ETaskPriority::Inherit));
+							}
+						}
+						else
 						{
-							AddSurface(SurfaceId, 0, InstanceSurfaceIndex);
+							// This means the compiled model (maybe coming from derived data) has images that the asset doesn't know about.
+							UE_LOG(LogMutable, Error, TEXT("CustomizableObject derived data out of sync with asset for [%s]. Try recompiling it."), *CustomizableObject->GetName());
 						}
 					}
 				}
-
 			}	
 		}
-		
-		UE::Tasks::FTask LastTask = UE::Tasks::MakeCompletedTask<void>();
-
-		for (uint32 SurfaceIndex : GetMaterialSurfaces)
-		{
-			LastTask = UE::Tasks::Launch(TEXT("Task_MutableGetMeshes_GetMesh_GetMaterial"), [SurfaceIndex, OperationData]()
-			{
-				UE::Mutable::Private::FMaterialId MaterialId = OperationData->InstanceUpdateData.Surfaces[SurfaceIndex].MaterialId;
-				UE::Tasks::TTask<TSharedPtr<const UE::Mutable::Private::FMaterial>> GetMaterialTask = OperationData->MutableSystem->GetMaterial({}, OperationData->InstanceID, MaterialId);
-				
-				UE::Tasks::AddNested(UE::Tasks::Launch(TEXT("Task_MutableGetMeshes_GetMesh_GetMaterial_Post"), [=]() mutable
-				{
-					OperationData->InstanceUpdateData.Surfaces[SurfaceIndex].Material = GetMaterialTask.GetResult();
-				},
-				GetMaterialTask,
-				LowLevelTasks::ETaskPriority::Inherit));
-			},
-			LastTask,
-			LowLevelTasks::ETaskPriority::Inherit);
-		}
-			
-		UE::Tasks::AddNested(UE::Tasks::Launch(TEXT("Task_Mutable_GetMeshes_GetImage_Loop"), [=]()
-		{
-			Task_Mutable_GetMeshes_GetImage_Loop(OperationData, StartTime, GetImagesData, 0);
-		},
-		LastTask,
-		LowLevelTasks::ETaskPriority::Inherit));
 	}
-
-
-	/** TaskGraph task after GetMesh has completed. */
-	void Task_MutableGetMeshes_GetMesh_Post(
-		const TSharedRef<FUpdateContextPrivate>& OperationData,
-		double StartTime,
-		const TSharedRef<TArray<FGetMeshData>>& GetMeshesData,
-		int32 GetMeshIndex,
-		UE::Tasks::TTask<TSharedPtr<const UE::Mutable::Private::FMesh>> GetMeshTask)
+	
+	
+	void Task_Mutable_GetMeshes(const TSharedRef<FUpdateContextPrivate>& OperationData)
 	{
-		MUTABLE_CPUPROFILER_SCOPE(Task_MutableGetMeshes_GetMesh_Post)
+		MUTABLE_CPUPROFILER_SCOPE(Task_Mutable_GetMeshes)
 
-		const int32 LODIndex = (*GetMeshesData)[GetMeshIndex].InstanceUpdateLODIndex;
-		FInstanceUpdateData::FLOD& LOD = OperationData->InstanceUpdateData.LODs[LODIndex];
+		const double StartTime = FPlatformTime::Seconds();
 
-		LOD.Mesh = GetMeshTask.GetResult();
-		
-		Task_Mutable_GetMeshes_GetMesh_Loop(OperationData, StartTime, GetMeshesData, ++GetMeshIndex);
-	}
-	
+		check(OperationData->Parameters);
+		OperationData->InstanceUpdateData.Clear();
 
-	/** See declaration. */
-	void Task_Mutable_GetMeshes_GetMesh_Loop(
-		const TSharedRef<FUpdateContextPrivate>& OperationData,
-		double StartTime,
-		const TSharedRef<TArray<FGetMeshData>>& GetMeshesData,
-		int32 GetMeshIndex)
-	{
-		MUTABLE_CPUPROFILER_SCOPE(Task_Mutable_GetMeshes_GetMesh_Loop)
+		check(UCustomizableObjectSystem::GetInstance() != nullptr);
+		check(UCustomizableObjectSystem::GetInstance()->GetPrivate() != nullptr);
 
-		if (GetMeshIndex >= GetMeshesData->Num())
+		UCustomizableInstancePrivate* CustomizableObjectInstancePrivateData = OperationData->Instance->GetPrivate();
+		
+		if (OperationData->PixelFormatOverride)
 		{
-			Task_Mutable_GetMeshes_GetImages(OperationData, StartTime);
-			return;
+			OperationData->MutableSystem->SetImagePixelConversionOverride( OperationData->PixelFormatOverride );
 		}
 		
-		const FGetMeshData& MeshData = (*GetMeshesData)[GetMeshIndex];
-		UE::Tasks::TTask<TSharedPtr<const UE::Mutable::Private::FMesh>> GetMeshTask = 
-				OperationData->MutableSystem->GetMesh({}, OperationData->InstanceID, MeshData.MeshID, MeshData.ContentFilter);
+		TSharedPtr<const UE::Mutable::Private::FInstance> Instance = OperationData->MutableSystem->BeginUpdate_MutableThread(OperationData->InstanceID,
+			OperationData->Parameters,
+			OperationData->MeshIdRegistry,
+			OperationData->ImageIdRegistry,
+			OperationData->MaterialIdRegistry,
+			OperationData->SkeletalMeshIdRegistry,
+			OperationData->GetCapturedDescriptor().GetState(),
+			UE::Mutable::Private::FSystem::AllLODs);
+
+		OperationData->MutableInstance = Instance;
 		
-		UE::Tasks::AddNested(UE::Tasks::Launch(TEXT("Task_MutableGetMeshes_GetMesh_Post"), [=]()
-		{
-			Task_MutableGetMeshes_GetMesh_Post(OperationData, StartTime, GetMeshesData, GetMeshIndex, GetMeshTask);
-		},
-		GetMeshTask,
-		LowLevelTasks::ETaskPriority::Inherit));
-	}
+		const bool bForceGeometryGeneration = CVarForceGeometryOnFirstGeneration.GetValueOnAnyThread();
 
-
-	namespace Impl
-	{
-		/** Start of the GetMeshes tasks.
-		  * Gathers all GetMeshes that has to be called. */
-		void Task_Mutable_GetMeshes(const TSharedRef<FUpdateContextPrivate>& OperationData)
+		UE::Tasks::FTask LastTask;
+		
+		OperationData->InstanceUpdateData.Components.Reserve(OperationData->NumInstanceComponents);
+		for (int32 InstanceComponentIndex = 0; InstanceComponentIndex < OperationData->NumInstanceComponents; ++InstanceComponentIndex)
 		{
-			MUTABLE_CPUPROFILER_SCOPE(Task_Mutable_GetMeshes)
-			TRACE_BEGIN_REGION(UE_TASK_MUTABLE_GETMESHES_REGION);
-
-			const double StartTime = FPlatformTime::Seconds();
+			TSharedRef<FInstanceUpdateData::FComponent>& Component = OperationData->InstanceUpdateData.Components.Add_GetRef(MakeShared<FInstanceUpdateData::FComponent>());
+			Component->FirstLOD = OperationData->InstanceUpdateData.LODs.Num();
+			uint64 MutableComponentId = Instance->GetComponentId(InstanceComponentIndex);
+			if (MutableComponentId < 65535)
+			{
+				Component->Id = FCustomizableObjectComponentIndex(MutableComponentId);
+			}
+			else
+			{
+				Component->Id.Invalidate();
+			}
 
-			check(OperationData->Parameters);
-			OperationData->InstanceUpdateData.Clear();
+			FCustomizableObjectComponentIndex ObjectComponentIndex = Component->Id;
+			const FName ComponentName = OperationData->ComponentNames[ObjectComponentIndex.GetValue()];
 
-			check(UCustomizableObjectSystem::GetInstance() != nullptr);
-			check(UCustomizableObjectSystem::GetInstance()->GetPrivate() != nullptr);
+			Component->SkeletalMeshId = Instance->GetSkeletalMeshId(InstanceComponentIndex);
 
-			UCustomizableInstancePrivate* CustomizableObjectInstancePrivateData = OperationData->Instance->GetPrivate();
+			UE::Tasks::TTask<TSharedPtr<const UE::Mutable::Private::FSkeletalMesh>> GetSkeletalMeshTask = OperationData->MutableSystem->GetSkeletalMesh(LastTask, OperationData->InstanceID, Component->SkeletalMeshId, UE::Mutable::Private::EMeshContentFlags::AllFlags);
+			UE::Tasks::AddNested(GetSkeletalMeshTask);
+			LastTask = GetSkeletalMeshTask;
 			
-			if (OperationData->PixelFormatOverride)
+			UE::Tasks::AddNested(UE::Tasks::Launch(TEXT("GetSkeletalMeshResult"), [=]() mutable
 			{
-				OperationData->MutableSystem->SetImagePixelConversionOverride( OperationData->PixelFormatOverride );
-			}
+				Component->SkeletalMesh = GetSkeletalMeshTask.GetResult();
+			},
+			GetSkeletalMeshTask,
+			LowLevelTasks::ETaskPriority::Inherit));
 			
-			TSharedPtr<const UE::Mutable::Private::FInstance> Instance = OperationData->MutableSystem->BeginUpdate_MutableThread(OperationData->InstanceID,
-				OperationData->Parameters,
-				OperationData->MeshIdRegistry,
-				OperationData->ImageIdRegistry,
-				OperationData->MaterialIdRegistry,
-				OperationData->SkeletalMeshIdRegistry,
-				OperationData->GetCapturedDescriptor().GetState(),
-				UE::Mutable::Private::FSystem::AllLODs);
-
-			OperationData->MutableInstance = Instance;
+			Component->OverlayMaterialId = Instance->GetOverlayMaterialId(InstanceComponentIndex);
 			
-			const bool bForceGeometryGeneration = CVarForceGeometryOnFirstGeneration.GetValueOnAnyThread();
-
-			TArray<int32> GetMaterialOverrideComponent;
+			UE::Tasks::TTask<TSharedPtr<const UE::Mutable::Private::FMaterial>> GetMaterialOverrideTask = OperationData->MutableSystem->GetMaterial(LastTask, OperationData->InstanceID, Component->OverlayMaterialId);
+			UE::Tasks::AddNested(GetMaterialOverrideTask);
+			LastTask = GetMaterialOverrideTask;
 			
-			const TSharedRef<TArray<FGetMeshData>> GetMeshesData = MakeShared<TArray<FGetMeshData>>();
+			UE::Tasks::AddNested(UE::Tasks::Launch(TEXT("GetMaterialOverrideResult"), [Component, GetMaterialOverrideTask]() mutable
+			{
+				Component->OverlayMaterial = GetMaterialOverrideTask.GetResult();
+			},
+			GetMaterialOverrideTask,
+			LowLevelTasks::ETaskPriority::Inherit));
 			
-			OperationData->InstanceUpdateData.Components.SetNum(OperationData->NumInstanceComponents);
-			for (int32 InstanceComponentIndex = 0; InstanceComponentIndex < OperationData->NumInstanceComponents; ++InstanceComponentIndex)
+			Component->LODCount = OperationData->NumLODsAvailable[ComponentName];
+			if (!Component->LODCount)
 			{
-				FInstanceUpdateData::FComponent& Component = OperationData->InstanceUpdateData.Components[InstanceComponentIndex];
-				Component.FirstLOD = OperationData->InstanceUpdateData.LODs.Num();
-				uint64 MutableComponentId = Instance->GetComponentId(InstanceComponentIndex);
-				if (MutableComponentId < 65535)
-				{
-					Component.Id = FCustomizableObjectComponentIndex(MutableComponentId);
-				}
-				else
-				{
-					Component.Id.Invalidate();
-				}
-
-				FCustomizableObjectComponentIndex ObjectComponentIndex = Component.Id;
-				const FName ComponentName = OperationData->ComponentNames[ObjectComponentIndex.GetValue()];
-
-				Component.SkeletalMeshId = Instance->GetSkeletalMeshId(InstanceComponentIndex);
+				// It happens in degenerated cases with empty components.
+				continue;
+			}
+			
+			const uint8 FirstResidentLOD = OperationData->FirstResidentLOD[ComponentName];
 
-				UE::Tasks::TTask<TSharedPtr<const UE::Mutable::Private::FSkeletalMesh>> Task = OperationData->MutableSystem->GetSkeletalMesh({}, OperationData->InstanceID, Component.SkeletalMeshId, UE::Mutable::Private::EMeshContentFlags::AllFlags);
-				Component.SkeletalMesh = Task.GetResult(); // TODO GMT Remove. We can do this since now FSkeletalMesh will always be passthrough.
+			// If the LOD is not generated we still add an empty one to keep indexes aligned.
+			OperationData->InstanceUpdateData.LODs.Reserve(Component->FirstLOD + Component->LODCount);
 
-				Component.OverlayMaterialId = Instance->GetOverlayMaterialId(InstanceComponentIndex);
-				if (Component.OverlayMaterialId)
-				{
-					GetMaterialOverrideComponent.Add(InstanceComponentIndex);
-				}
+			const int32 FirstLOD = OperationData->bStreamMeshLODs ?
+				OperationData->FirstLODAvailable[ComponentName] :
+				OperationData->GetFirstRequestedLOD()[ComponentName];
+			
+			for (int32 LODIndex = FirstLOD; LODIndex < Component->LODCount; ++LODIndex)
+			{
+				MUTABLE_CPUPROFILER_SCOPE(GetMesh);
 
-				Component.LODCount = OperationData->NumLODsAvailable[ComponentName];
-				if (!Component.LODCount)
-				{
-					// It happens in degenerated cases with empty components.
-					continue;
-				}
+				TSharedRef<FInstanceUpdateData::FLOD>& LOD = OperationData->InstanceUpdateData.LODs.Add_GetRef(MakeShared<FInstanceUpdateData::FLOD>());
 				
-				const uint8 FirstResidentLOD = OperationData->FirstResidentLOD[ComponentName];
-
-				// If the LOD is not generated we still add an empty one to keep indexes aligned.
-				OperationData->InstanceUpdateData.LODs.SetNum(Component.FirstLOD + Component.LODCount);
-
-				const int32 FirstLOD = OperationData->bStreamMeshLODs ?
-					OperationData->FirstLODAvailable[ComponentName] :
-					OperationData->GetFirstRequestedLOD()[ComponentName];
+				LOD->MeshId = Instance->GetMeshId(InstanceComponentIndex, LODIndex);
+				
+				UE::Mutable::Private::EMeshContentFlags MeshContentFilter = UE::Mutable::Private::EMeshContentFlags::AllFlags; 
 				
-				for (int32 LODIndex = FirstLOD; LODIndex < Component.LODCount; ++LODIndex)
+				if (!bForceGeometryGeneration)
 				{
-					MUTABLE_CPUPROFILER_SCOPE(GetMesh);
-					
-					FInstanceUpdateData::FLOD& LOD = OperationData->InstanceUpdateData.LODs[Component.FirstLOD + LODIndex];
-					
-					LOD.MeshId = Instance->GetMeshId(InstanceComponentIndex, LODIndex);
-					
-					UE::Mutable::Private::EMeshContentFlags MeshContentFilter = UE::Mutable::Private::EMeshContentFlags::AllFlags; 
-					
-					if (!bForceGeometryGeneration)
+					if (LODIndex < FirstResidentLOD)
 					{
-						if (LODIndex < FirstResidentLOD)
-						{
-							EnumRemoveFlags(MeshContentFilter, UE::Mutable::Private::EMeshContentFlags::GeometryData);
-						}
+						EnumRemoveFlags(MeshContentFilter, UE::Mutable::Private::EMeshContentFlags::GeometryData);
 					}
-					
-					GetMeshesData->Add({ Component.FirstLOD + LODIndex, LOD.MeshId, MeshContentFilter });
 				}
-			}
-
-			UE::Tasks::FTask LastTask = UE::Tasks::MakeCompletedTask<void>();
-
-			for (const int32 InstanceComponentIndex : GetMaterialOverrideComponent)
-			{
-				LastTask = UE::Tasks::Launch(TEXT("Task_MutableGetMeshes_GetMesh_GetMaterialOverride"), [InstanceComponentIndex, OperationData]()
+				
+				UE::Tasks::TTask<TSharedPtr<const UE::Mutable::Private::FMesh>> GetMeshTask = OperationData->MutableSystem->GetMesh(LastTask, OperationData->InstanceID, LOD->MeshId, MeshContentFilter);
+				UE::Tasks::AddNested(GetMeshTask);
+				LastTask = GetMeshTask;
+				
+				UE::Tasks::AddNested(UE::Tasks::Launch(TEXT("GetMeshResult"),
+				[LOD, GetMeshTask]() mutable
 				{
-					UE::Mutable::Private::FMaterialId MaterialId = OperationData->InstanceUpdateData.Components[InstanceComponentIndex].OverlayMaterialId;
-					UE::Tasks::TTask<TSharedPtr<const UE::Mutable::Private::FMaterial>> GetMaterialOverrideTask = OperationData->MutableSystem->GetMaterial({}, OperationData->InstanceID, MaterialId);
-
-					UE::Tasks::AddNested(UE::Tasks::Launch(TEXT("Task_MutableGetMeshes_GetMesh_GetMaterialOverride_Post"), [=]() mutable
-					{
-						OperationData->InstanceUpdateData.Components[InstanceComponentIndex].OverlayMaterial = GetMaterialOverrideTask.GetResult();
-					},
-					GetMaterialOverrideTask,
-					LowLevelTasks::ETaskPriority::Inherit));
+					LOD->Mesh = GetMeshTask.GetResult();
 				},
-				LastTask,
-				LowLevelTasks::ETaskPriority::Inherit);
+				GetMeshTask,
+				UE::Tasks::ETaskPriority::Inherit));
 			}
-			
-			UE::Tasks::AddNested(UE::Tasks::Launch(TEXT("Task_MutableGetMeshes_GetMesh_Post"), [=]()
-			{
-				Task_Mutable_GetMeshes_GetMesh_Loop(OperationData, StartTime, GetMeshesData, 0);
-			},
-			LastTask,
-			LowLevelTasks::ETaskPriority::Inherit));
 		}
-	}
-
 
-	void Task_Mutable_GetMeshes(const TSharedRef<FUpdateContextPrivate>& OperationData)
-	{
-		Impl::Task_Mutable_GetMeshes(OperationData);
+		OperationData->TaskGetMeshesTime = FPlatformTime::Seconds() - StartTime;
 	}
 
 	
@@ -2894,11 +2680,11 @@ namespace impl
 	{
 		MUTABLE_CPUPROFILER_SCOPE(Subtask_Mutable_PrepareTextures)
 
-		for (const FInstanceUpdateData::FSurface& Surface : OperationData->InstanceUpdateData.Surfaces)
+		for (const TSharedRef<FInstanceUpdateData::FSurface>& Surface : OperationData->InstanceUpdateData.Surfaces)
 		{
-			for (int32 ImageIndex = 0; ImageIndex<Surface.ImageCount; ++ImageIndex)
+			for (int32 ImageIndex = 0; ImageIndex<Surface->ImageCount; ++ImageIndex)
 			{
-				const TSharedRef<FInstanceUpdateData::FImage>& Image = OperationData->InstanceUpdateData.Images[Surface.FirstImage+ImageIndex];
+				const TSharedRef<FInstanceUpdateData::FImage>& Image = OperationData->InstanceUpdateData.Images[Surface->FirstImage + ImageIndex];
 
 				if (Image->Image)
 				{
@@ -3669,15 +3455,39 @@ namespace impl
 			Task_Mutable_GetMeshes(Operation);
 		},
 		{ CacheRuntimeTexturesEvent });
+		
+		UE::Tasks::FTask Mutable_GetMaterialsTask = SystemPrivate->MutableTaskGraph.AddMutableThreadTask(
+		TEXT("Task_Mutable_GetMaterials"),
+		[Operation]()
+		{
+			Task_Mutable_GetMaterials(Operation);
+		},
+		{ Mutable_GetMeshTask });
+
+		UE::Tasks::FTask Mutable_PrepareSkeletonDataTask = SystemPrivate->MutableTaskGraph.AddMutableThreadTask(
+		TEXT("Task_Mutable_PrepareSkeletonData"),
+		[Operation]()
+		{
+			Task_Mutable_PrepareSkeletonData(Operation);
+		},
+		{ Mutable_GetMaterialsTask });
+
+		UE::Tasks::FTask Mutable_UpdateParameterRelevancy = SystemPrivate->MutableTaskGraph.AddMutableThreadTask(
+		TEXT("Task_Mutable_UpdateParameterRelevancy"),
+		[Operation]()
+		{
+			Task_Mutable_UpdateParameterRelevancy(Operation);
+		},
+		{ Mutable_PrepareSkeletonDataTask });
 
 		SystemPrivate->MutableTaskGraph.AddGameThreadTask(
-			TEXT("Task_Game_LockCache"),
-			[Operation]()
-			{
-				impl::Task_Game_LockCache(Operation);
-			},
-			false,
-			{ Mutable_GetMeshTask });
+		TEXT("Task_Game_LockCache"),
+		[Operation]()
+		{
+			impl::Task_Game_LockCache(Operation);
+		},
+		false,
+		{ Mutable_UpdateParameterRelevancy });
 	}
 } // namespace impl
 
diff --git a/Engine/Plugins/Mutable/Source/CustomizableObject/Private/MuCO/LogBenchmarkUtil.cpp b/Engine/Plugins/Mutable/Source/CustomizableObject/Private/MuCO/LogBenchmarkUtil.cpp
--- a/Engine/Plugins/Mutable/Source/CustomizableObject/Private/MuCO/LogBenchmarkUtil.cpp
+++ b/Engine/Plugins/Mutable/Source/CustomizableObject/Private/MuCO/LogBenchmarkUtil.cpp
@@ -340,7 +340,7 @@ void FLogBenchmarkUtil::FinishUpdateMesh(const TSharedRef<FUpdateContextPrivate>
 	UpdateData.TriangleCount = TriangleCount;
 	UpdateData.QueueTime = Context->QueueTime * 1000;
 	UpdateData.UpdateTime = Context->UpdateTime * 1000;
-	UpdateData.TaskGetMeshTime = Context->TaskGetMeshTime * 1000;
+	UpdateData.TaskGetMeshesTime = Context->TaskGetMeshesTime * 1000;
 	UpdateData.TaskLockCacheTime = Context->TaskLockCacheTime * 1000;
 	UpdateData.TaskGetImagesTime = Context->TaskGetImagesTime * 1000;
 	UpdateData.TaskConvertResourcesTime = Context->TaskConvertResourcesTime * 1000;
@@ -359,7 +359,7 @@ void FLogBenchmarkUtil::FinishUpdateMesh(const TSharedRef<FUpdateContextPrivate>
 		UpdateData.TriangleCount,
 		UpdateData.QueueTime,
 		UpdateData.UpdateTime,
-		UpdateData.TaskGetMeshTime,
+		UpdateData.TaskGetMeshesTime,
 		UpdateData.TaskLockCacheTime,
 		UpdateData.TaskGetImagesTime,
 		UpdateData.TaskConvertResourcesTime,
diff --git a/Engine/Plugins/Mutable/Source/CustomizableObjectEditor/Private/MuCOE/CustomizableObjectEditorPerformanceAnalyzer.cpp b/Engine/Plugins/Mutable/Source/CustomizableObjectEditor/Private/MuCOE/CustomizableObjectEditorPerformanceAnalyzer.cpp
--- a/Engine/Plugins/Mutable/Source/CustomizableObjectEditor/Private/MuCOE/CustomizableObjectEditorPerformanceAnalyzer.cpp
+++ b/Engine/Plugins/Mutable/Source/CustomizableObjectEditor/Private/MuCOE/CustomizableObjectEditorPerformanceAnalyzer.cpp
@@ -154,8 +154,8 @@ TSharedRef<SWidget> SInstanceUpdateDataRow::GenerateWidgetForColumn(const FName&
 	else if (InColumnName == InstanceUpdatesMainDataColumns::GetMeshTimeColumnID)
 	{
 		return SNew(STextBlock)
-			.Text(FText::FromString(FString::Printf(TEXT("%.2f"),InstanceUpdateElement->UpdateStats.TaskGetMeshTime)))
-			.ColorAndOpacity(TextColor(InstanceUpdateElement->UpdateStats.TaskGetMeshTime, InstanceUpdateStatistics->TaskGetMeshTime));
+			.Text(FText::FromString(FString::Printf(TEXT("%.2f"),InstanceUpdateElement->UpdateStats.TaskGetMeshesTime)))
+			.ColorAndOpacity(TextColor(InstanceUpdateElement->UpdateStats.TaskGetMeshesTime, InstanceUpdateStatistics->TaskGetMeshTime));
 	}
 	else if (InColumnName == InstanceUpdatesMainDataColumns::LockCacheTimeColumnID)
 	{
@@ -759,8 +759,8 @@ void SCustomizableObjectEditorPerformanceAnalyzer::OnInstanceUpdateListViewSort(
 		{
 			if (ElementA.IsValid() && ElementB.IsValid())
 			{
-				const double ValA = ElementA->UpdateStats.TaskGetMeshTime;
-				const double ValB =  ElementB->UpdateStats.TaskGetMeshTime;
+				const double ValA = ElementA->UpdateStats.TaskGetMeshesTime;
+				const double ValB =  ElementB->UpdateStats.TaskGetMeshesTime;
 
 				return NewSortMode == EColumnSortMode::Ascending ? ValA < ValB : ValA > ValB;
 			}
@@ -984,7 +984,7 @@ void SCustomizableObjectEditorPerformanceAnalyzer::ComputeStatistics()
 	ComputeMax(InstanceUpdateStatistics->TriangleCount, [&](int32 Index){	return InstanceUpdateElements[Index]->UpdateStats.TriangleCount; });
 	ComputeMax(InstanceUpdateStatistics->QueueTime, [&](int32 Index){	return InstanceUpdateElements[Index]->UpdateStats.QueueTime; });
 	ComputeMax(InstanceUpdateStatistics->UpdateTime, [&](int32 Index) { return InstanceUpdateElements[Index]->UpdateStats.UpdateTime; });
-	ComputeMax(InstanceUpdateStatistics->TaskGetMeshTime, [&](int32 Index) { return InstanceUpdateElements[Index]->UpdateStats.TaskGetMeshTime; });
+	ComputeMax(InstanceUpdateStatistics->TaskGetMeshTime, [&](int32 Index) { return InstanceUpdateElements[Index]->UpdateStats.TaskGetMeshesTime; });
 	ComputeMax(InstanceUpdateStatistics->TaskLockCacheTime, [&](int32 Index) { return InstanceUpdateElements[Index]->UpdateStats.TaskLockCacheTime; });
 	ComputeMax(InstanceUpdateStatistics->TaskGetImagesTime, [&](int32 Index) { return InstanceUpdateElements[Index]->UpdateStats.TaskGetImagesTime; });
 	ComputeMax(InstanceUpdateStatistics->TaskConvertResourcesTime, [&](int32 Index) { return InstanceUpdateElements[Index]->UpdateStats.TaskConvertResourcesTime; });
