[Mutable] New node Modify Material (Interface only).
+ Added experimetal tag to Break Material node.

#rnx

diff --git a/Engine/Plugins/Mutable/Source/CustomizableObjectEditor/Private/MuCOE/CustomizableObjectEditorModule.cpp b/Engine/Plugins/Mutable/Source/CustomizableObjectEditor/Private/MuCOE/CustomizableObjectEditorModule.cpp
--- a/Engine/Plugins/Mutable/Source/CustomizableObjectEditor/Private/MuCOE/CustomizableObjectEditorModule.cpp
+++ b/Engine/Plugins/Mutable/Source/CustomizableObjectEditor/Private/MuCOE/CustomizableObjectEditorModule.cpp
@@ -33,6 +33,7 @@
 #include "MuCOE/CustomizableObjectVersionBridge.h"
 #include "MuCOE/CustomizableObjectNodeObjectGroupDetails.h"
 #include "MuCOE/Nodes/CONodeMaterialBreak.h"
+#include "MuCOE/Nodes/CONodeMaterialModify.h"
 #include "MuCOE/Nodes/CustomizableObjectNodeCopyMaterial.h"
 #include "MuCOE/Nodes/CustomizableObjectNodeExternalPin.h"
 #include "MuCOE/Nodes/CustomizableObjectNodeExternalPinDetails.h"
@@ -285,6 +286,7 @@ void FCustomizableObjectEditorModule::StartupModule()
 	RegisterCustomDetails(PropertyModule, UCONodeMaterialBreak::StaticClass(), FOnGetDetailCustomizationInstance::CreateStatic(&FCustomizableObjectNodeDetails::MakeInstance));
 	RegisterCustomDetails(PropertyModule, UCustomizableObjectNodeComponentPassthroughMesh::StaticClass(), FOnGetDetailCustomizationInstance::CreateStatic(&FCustomizableObjectNodeDetails::MakeInstance));
 	RegisterCustomDetails(PropertyModule, UCONodeExternalOperation::StaticClass(), FOnGetDetailCustomizationInstance::CreateStatic(&FCONodeExternalOperationDetails::MakeInstance));
+	RegisterCustomDetails(PropertyModule, UCONodeMaterialModify::StaticClass(), FOnGetDetailCustomizationInstance::CreateStatic(&FCustomizableObjectNodeDetails::MakeInstance));
 
 	// Other Objects
 	RegisterCustomDetails(PropertyModule, UCustomizableObject::StaticClass(), FOnGetDetailCustomizationInstance::CreateStatic(&FCustomizableObjectDetails::MakeInstance));
diff --git a/Engine/Plugins/Mutable/Source/CustomizableObjectEditor/Private/MuCOE/EdGraphSchema_CustomizableObject.cpp b/Engine/Plugins/Mutable/Source/CustomizableObjectEditor/Private/MuCOE/EdGraphSchema_CustomizableObject.cpp
--- a/Engine/Plugins/Mutable/Source/CustomizableObjectEditor/Private/MuCOE/EdGraphSchema_CustomizableObject.cpp
+++ b/Engine/Plugins/Mutable/Source/CustomizableObjectEditor/Private/MuCOE/EdGraphSchema_CustomizableObject.cpp
@@ -25,6 +25,7 @@
 #include "MuCOE/CustomizableObjectEditor.h"
 #include "MuCOE/Nodes/CONodeMaterialBreak.h"
 #include "MuCOE/Nodes/CONodeMaterialConstant.h"
+#include "MuCOE/Nodes/CONodeMaterialModify.h"
 #include "MuCOE/Nodes/CONodeMaterialVariation.h"
 #include "MuCOE/Nodes/CONodeMaterialSwitch.h"
 #include "MuCOE/Nodes/CONodeModifierTransformWithBone.h"
@@ -653,7 +654,8 @@ void UEdGraphSchema_CustomizableObject::GetGraphContextActions(FGraphContextMenu
 			ContextMenuBuilder.CreateTemplateNode<UCONodeMaterialConstant>(),
 			ContextMenuBuilder.CreateTemplateNode<UCONodeMaterialVariation>(),
 			ContextMenuBuilder.CreateTemplateNode<UCONodeMaterialSwitch>(),
-			ContextMenuBuilder.CreateTemplateNode<UCONodeMaterialBreak>()
+			ContextMenuBuilder.CreateTemplateNode<UCONodeMaterialBreak>(),
+			ContextMenuBuilder.CreateTemplateNode<UCONodeMaterialModify>()
 		};
 
 		AddNewNodeCategoryActionsFiltered(MaterialTemplateNodes, ContextMenuBuilder, NC_Material.ToString(), 2, Filter);
diff --git a/Engine/Plugins/Mutable/Source/CustomizableObjectEditor/Private/MuCOE/Nodes/CONodeMaterialBreak.cpp b/Engine/Plugins/Mutable/Source/CustomizableObjectEditor/Private/MuCOE/Nodes/CONodeMaterialBreak.cpp
--- a/Engine/Plugins/Mutable/Source/CustomizableObjectEditor/Private/MuCOE/Nodes/CONodeMaterialBreak.cpp
+++ b/Engine/Plugins/Mutable/Source/CustomizableObjectEditor/Private/MuCOE/Nodes/CONodeMaterialBreak.cpp
@@ -172,4 +172,10 @@ FName UCONodeMaterialBreak::GetPinParameterName(const UEdGraphPin& Pin) const
 }
 
 
+bool UCONodeMaterialBreak::IsExperimental() const
+{
+	return true;
+}
+
+
 #undef LOCTEXT_NAMESPACE
diff --git a/Engine/Plugins/Mutable/Source/CustomizableObjectEditor/Private/MuCOE/Nodes/CONodeMaterialBreak.h b/Engine/Plugins/Mutable/Source/CustomizableObjectEditor/Private/MuCOE/Nodes/CONodeMaterialBreak.h
--- a/Engine/Plugins/Mutable/Source/CustomizableObjectEditor/Private/MuCOE/Nodes/CONodeMaterialBreak.h
+++ b/Engine/Plugins/Mutable/Source/CustomizableObjectEditor/Private/MuCOE/Nodes/CONodeMaterialBreak.h
@@ -58,6 +58,7 @@ public:
 	UE_API virtual bool CanPinBeRemoved(const UEdGraphPin& Pin) const override;
 	UE_API virtual UCustomizableObjectNodeRemapPins* CreateRemapPinsDefault() const override;
 	UE_API virtual TArray<FName> GetAllowedPinViewerCreationTypes() const override;
+	UE_API virtual bool IsExperimental() const override;
 
 	// Own interface
 	FName GetPinParameterName(const UEdGraphPin& Pin) const;
diff --git a/Engine/Plugins/Mutable/Source/CustomizableObjectEditor/Private/MuCOE/Nodes/CONodeMaterialModify.cpp b/Engine/Plugins/Mutable/Source/CustomizableObjectEditor/Private/MuCOE/Nodes/CONodeMaterialModify.cpp
new file mode 100644
--- /dev/null
+++ b/Engine/Plugins/Mutable/Source/CustomizableObjectEditor/Private/MuCOE/Nodes/CONodeMaterialModify.cpp
@@ -0,0 +1,169 @@
+// Copyright Epic Games, Inc. All Rights Reserved.
+
+#include "MuCOE/Nodes/CONodeMaterialModify.h"
+
+#include "DetailsViewArgs.h"
+#include "MuCOE/EdGraphSchema_CustomizableObject.h"
+#include "PropertyEditorModule.h"
+#include "ScopedTransaction.h"
+
+#define LOCTEXT_NAMESPACE "CustomizableObjectEditor"
+
+
+bool UCONodeMaterialModifyRemapPins::Equal(const UCustomizableObjectNode& Node, const UEdGraphPin& OldPin, const UEdGraphPin& NewPin) const
+{
+	UCONodeMaterialModifyPinData* OldPinData = Cast<UCONodeMaterialModifyPinData>(Node.GetPinData(OldPin));
+	UCONodeMaterialModifyPinData* NewPinData = Cast<UCONodeMaterialModifyPinData>(Node.GetPinData(NewPin));
+
+	if (OldPinData && NewPinData)
+	{
+		return OldPin.Direction == NewPin.Direction && OldPinData->VariableData.Type == NewPinData->VariableData.Type && OldPinData->VariableData.Id == NewPinData->VariableData.Id;
+	}
+	else
+	{
+		return Super::Equal(Node, OldPin, NewPin);
+	}
+}
+
+
+FText UCONodeMaterialModify::GetNodeTitle(ENodeTitleType::Type TitleType) const
+{
+	return LOCTEXT("ModifyMaterialNodeTitle","Modify Material");
+}
+
+
+FLinearColor UCONodeMaterialModify::GetNodeTitleColor() const
+{
+	return GetDefault<UEdGraphSchema_CustomizableObject>()->GetPinTypeColor(UEdGraphSchema_CustomizableObject::PC_Material);
+}
+
+
+FText UCONodeMaterialModify::GetTooltipText() const
+{
+	return LOCTEXT("ModifyMaterialNodeTooltip","Modifys a new material instance from the Reference Material.");
+}
+
+
+void UCONodeMaterialModify::AllocateDefaultPins(UCustomizableObjectNodeRemapPins* RemapPins)
+{
+	// Input Material Pin
+	const FName Type = UEdGraphSchema_CustomizableObject::PC_Material;
+	const FName PinName = UEdGraphSchema_CustomizableObject::GetPinCategoryName(Type);
+	const FText PinFriendlyName = UEdGraphSchema_CustomizableObject::GetPinCategoryFriendlyName(Type);
+	
+	UEdGraphPin* InputMaterialPin = CustomCreatePin(EGPD_Input, Type, PinName);
+	InputMaterialPin->PinFriendlyName = PinFriendlyName;
+
+	MaterialPinRef = FEdGraphPinReference(InputMaterialPin);
+
+	//Output Material Pin
+	UEdGraphPin* OutputMaterialPin = CustomCreatePin(EGPD_Output, Type, PinName);
+	OutputMaterialPin->PinFriendlyName = PinFriendlyName;
+
+	// Parameter Pins
+	// Get old pins
+	TArray<UEdGraphPin*> InputPins = GetAllNonOrphanPins();
+
+	// Regenerate the new pins from the old pin's pindata which may have changed
+	for (const UEdGraphPin* Pin : InputPins)
+	{
+		if (Pin->Direction == EEdGraphPinDirection::EGPD_Input)
+		{
+			if (UCONodeMaterialModifyPinData* PinData = Cast<UCONodeMaterialModifyPinData>(GetPinData(*Pin)))
+			{
+				// Use the same pin data as the old pin since the new one is a copy
+				UEdGraphPin* NewPin = CustomCreatePin(EGPD_Input, PinData->VariableData.Type, PinData->VariableData.Name, PinData);
+			}
+		}
+	}
+}
+
+
+TSharedPtr<IDetailsView> UCONodeMaterialModify::CustomizePinDetails(const UEdGraphPin& Pin) const
+{
+	if (UCONodeMaterialModifyPinData* PinData = Cast<UCONodeMaterialModifyPinData>(GetPinData(Pin)))
+	{
+		FPropertyEditorModule& EditModule = FModuleManager::Get().GetModuleChecked<FPropertyEditorModule>("PropertyEditor");
+
+		FDetailsViewArgs DetailsViewArgs;
+		DetailsViewArgs.bAllowSearch = false;
+		DetailsViewArgs.bHideSelectionTip = true;
+
+		const TSharedRef<IDetailsView> SettingsView = EditModule.CreateDetailView(DetailsViewArgs);
+		SettingsView->SetObject(PinData);
+
+		return SettingsView;
+	}
+	else
+	{
+		return nullptr;
+	}
+}
+
+
+bool UCONodeMaterialModify::HasPinViewer() const
+{
+	return true;
+}
+
+
+EAddPinNodeButtonLocation UCONodeMaterialModify::GetAddPinButtonNodeSide() const
+{
+	return EAddPinNodeButtonLocation::INPUT;
+}
+
+
+UCustomizableObjectNodeRemapPins* UCONodeMaterialModify::CreateRemapPinsDefault() const
+{
+	return NewObject<UCONodeMaterialModifyRemapPins>();
+}
+
+
+TArray<FName> UCONodeMaterialModify::GetAllowedPinViewerCreationTypes() const
+{
+	return { GetDefault<UEdGraphSchema_CustomizableObject>()->PC_Texture,
+		GetDefault<UEdGraphSchema_CustomizableObject>()->PC_PassthroughTexture,
+		GetDefault<UEdGraphSchema_CustomizableObject>()->PC_Color,
+		GetDefault<UEdGraphSchema_CustomizableObject>()->PC_Float };
+}
+
+
+void UCONodeMaterialModify::AddPinFromUI()
+{
+	Super::AddPinFromUI();
+
+	FScopedTransaction Transaction(LOCTEXT("CONodeMaterialModify_AddPinFromNodeUITransaction", "Add Pin from Node UI"));
+	Modify();
+
+	const FName NewPinCategory = UEdGraphSchema_CustomizableObject::PC_Float;
+
+	UCONodeMaterialModifyPinData* PinData = NewObject<UCONodeMaterialModifyPinData>(this);
+	PinData->VariableData.Name = "NewPin";
+	PinData->VariableData.Type = NewPinCategory;
+	PinData->VariableData.Id = FGuid::NewGuid();
+
+	if (CustomCreatePin(EEdGraphPinDirection::EGPD_Input, NewPinCategory, "NewPin", PinData))
+	{
+		GetGraph()->NotifyNodeChanged(this);
+	}
+}
+
+
+FName UCONodeMaterialModify::GetPinParameterName(const UEdGraphPin& Pin) const
+{
+	if (const UCONodeMaterialModifyPinData* PinData = Cast<UCONodeMaterialModifyPinData>(GetPinData(Pin)))
+	{
+		return PinData->VariableData.Name;
+	}
+
+	return NAME_None;
+}
+
+
+bool UCONodeMaterialModify::IsExperimental() const
+{
+	return true;
+}
+
+
+#undef LOCTEXT_NAMESPACE
\ No newline at end of file
diff --git a/Engine/Plugins/Mutable/Source/CustomizableObjectEditor/Private/MuCOE/Nodes/CONodeMaterialModify.h b/Engine/Plugins/Mutable/Source/CustomizableObjectEditor/Private/MuCOE/Nodes/CONodeMaterialModify.h
new file mode 100644
--- /dev/null
+++ b/Engine/Plugins/Mutable/Source/CustomizableObjectEditor/Private/MuCOE/Nodes/CONodeMaterialModify.h
@@ -0,0 +1,69 @@
+// Copyright Epic Games, Inc. All Rights Reserved.
+
+#pragma once
+
+#include "MuCOE/Nodes/CustomizableObjectNode.h"
+#include "MuCOE/RemapPins/CustomizableObjectNodeRemapPinsByName.h"
+#include "MuCOE/COVariable.h"
+#include "CONodeMaterialModify.generated.h"
+
+#define UE_API CUSTOMIZABLEOBJECTEDITOR_API
+
+class UCustomizableObjectNodeRemapPinsByName;
+
+UCLASS()
+class UCONodeMaterialModifyRemapPins: public UCustomizableObjectNodeRemapPinsByName
+{
+public:
+	GENERATED_BODY()
+
+	virtual bool Equal(const UCustomizableObjectNode& Node, const UEdGraphPin& OldPin, const UEdGraphPin& NewPin) const override;
+};
+
+
+// TODO(Max - GM)UE-318756 This should have the same properties than UCustomizableObjectNodeMaterialPinDataImage
+// TODO(Max - GM)UE-318756 Reformulate all the pin logic for this node
+/** Base class for all Material Parameters. */
+UCLASS()
+class UCONodeMaterialModifyPinData : public UCustomizableObjectNodePinData
+{
+	GENERATED_BODY()
+
+public:
+
+	UPROPERTY(EditAnywhere, Category = NoCategory)
+	FCOVariable VariableData;
+};
+
+UCLASS()
+class UCONodeMaterialModify : public UCustomizableObjectNode
+{
+public:
+
+	GENERATED_BODY()
+
+	// Begin EdGraphNode interface
+	virtual FText GetNodeTitle(ENodeTitleType::Type TitleType) const override;
+	virtual FLinearColor GetNodeTitleColor() const override;
+	virtual FText GetTooltipText() const override;
+
+	// UCustomizableObjectNode interface
+	virtual void AllocateDefaultPins(UCustomizableObjectNodeRemapPins* RemapPins) override;
+	virtual TSharedPtr<IDetailsView> CustomizePinDetails(const UEdGraphPin& Pin) const override;
+	virtual bool HasPinViewer() const override;
+	virtual EAddPinNodeButtonLocation GetAddPinButtonNodeSide() const override;
+	virtual void AddPinFromUI() override;
+	virtual UCustomizableObjectNodeRemapPins* CreateRemapPinsDefault() const override;
+	virtual TArray<FName> GetAllowedPinViewerCreationTypes() const override;
+	virtual bool IsExperimental() const override;
+
+	// Own interface
+	FName GetPinParameterName(const UEdGraphPin& Pin) const;
+
+public:
+	
+	UPROPERTY()
+	FEdGraphPinReference MaterialPinRef;
+};
+
+#undef UE_API
\ No newline at end of file
diff --git a/Engine/Plugins/Mutable/Source/CustomizableObjectEditor/Private/MuCOE/Nodes/SCustomizableObjectNode.cpp b/Engine/Plugins/Mutable/Source/CustomizableObjectEditor/Private/MuCOE/Nodes/SCustomizableObjectNode.cpp
--- a/Engine/Plugins/Mutable/Source/CustomizableObjectEditor/Private/MuCOE/Nodes/SCustomizableObjectNode.cpp
+++ b/Engine/Plugins/Mutable/Source/CustomizableObjectEditor/Private/MuCOE/Nodes/SCustomizableObjectNode.cpp
@@ -72,7 +72,8 @@ void SCustomizableObjectNode::CreateInputSideAddButton(TSharedPtr<SVerticalBox>
 		{
 			TSharedRef<SWidget> AddPinButton = AddPinButtonContent(
 			NSLOCTEXT("SequencerNode", "SequencerNodeAddPinButton", "Add pin"),
-			NSLOCTEXT("SequencerNode", "SequencerNodeAddPinButton_ToolTip", "Add new pin"));
+			NSLOCTEXT("SequencerNode", "SequencerNodeAddPinButton_ToolTip", "Add new pin"),
+			false);
 
 			FMargin AddPinPadding = Settings->GetOutputPinPadding();
 			AddPinPadding.Top += 6.0f;
