[mutable] Fixed overlay materials not being applied as expected.
- The overlay material was being used as the override material of the first component section.
- The overlay material is shared between all sections of the generated skeletal mesh so they all get it to render on top of each of the section materials.

#rnx

diff --git a/Engine/Plugins/Mutable/Source/MutableRuntime/Private/MuR/CodeRunner.cpp b/Engine/Plugins/Mutable/Source/MutableRuntime/Private/MuR/CodeRunner.cpp
--- a/Engine/Plugins/Mutable/Source/MutableRuntime/Private/MuR/CodeRunner.cpp
+++ b/Engine/Plugins/Mutable/Source/MutableRuntime/Private/MuR/CodeRunner.cpp
@@ -1051,7 +1051,7 @@ namespace UE::Mutable::Private
         				if (Args.Material)
         				{
         					FMaterialId MaterialId = FSystem::GetMaterialId(System->WorkingMemoryManager.CurrentMaterialIdRegistry, InModel, InParams, Args.Material);
-        					Result->SetMaterialId(0, 0, 0, MaterialId);
+        					Result->SetOverlayMaterialId(0, MaterialId);
         				}
 
         				StoreInstance(Item, Result);
