[mutable] Added graceful handling of cases where one or multiple skeletons to merge are null.
- This prevents a crash caused by the previous assumption that all the skeletons to be merged would be valid at the time of merging. That is incorrect, as failures during the downloading or streaming of IAS data may produce null skeletons.

#rnx

diff --git a/Engine/Plugins/Mutable/Source/CustomizableObject/Private/MuCO/CustomizableObjectInstance.cpp b/Engine/Plugins/Mutable/Source/CustomizableObject/Private/MuCO/CustomizableObjectInstance.cpp
--- a/Engine/Plugins/Mutable/Source/CustomizableObject/Private/MuCO/CustomizableObjectInstance.cpp
+++ b/Engine/Plugins/Mutable/Source/CustomizableObject/Private/MuCO/CustomizableObjectInstance.cpp
@@ -1247,8 +1247,17 @@ USkeleton* MergeSkeletons(UCustomizableObject& CustomizableObject, FCustomizable
 		return MergedSkeleton;
 	}
 
+	// Fail in the case a null skeleton is found. This may be caused by a loading or decryption issue
+	if (ReferencedSkeletons.SkeletonsToMerge.Contains(nullptr))
+	{
+		UE_LOG( LogMutable, Warning, TEXT("Unable to merge the provided skeletons as one or multiple of them are null."));
+		ReferencedSkeletons.SkeletonIds.Empty();
+		ReferencedSkeletons.SkeletonsToMerge.Empty();
+		return nullptr;
+	}
+
 	// No need to merge skeletons
-	if(ReferencedSkeletons.SkeletonsToMerge.Num() == 1)
+	if (ReferencedSkeletons.SkeletonsToMerge.Num() == 1)
 	{
 		const TObjectPtr<USkeleton> RefSkeleton = ReferencedSkeletons.SkeletonsToMerge[0];
 		ReferencedSkeletons.SkeletonIds.Empty();
