[mutable] Get Material from Slot editor node implementation
- Added a new node that returns a material based on the material slot name of the provided skeletal mesh. This includes editor node, compilation node, and runtime operation.
- Made the skeletal mesh parameters being used by the COI get pre-compiled so their data can be safely accessed.
- Updated core code version.
- Fixed some typos
- Removed some unused includes.

#rnx

diff --git a/Engine/Plugins/Mutable/Source/CustomizableObject/Internal/MuCO/CustomizableObjectPrivate.h b/Engine/Plugins/Mutable/Source/CustomizableObject/Internal/MuCO/CustomizableObjectPrivate.h
--- a/Engine/Plugins/Mutable/Source/CustomizableObject/Internal/MuCO/CustomizableObjectPrivate.h
+++ b/Engine/Plugins/Mutable/Source/CustomizableObject/Internal/MuCO/CustomizableObjectPrivate.h
@@ -12,9 +12,7 @@
 #include "MuR/ResourceID.h"
 #include "Serialization/BulkData.h"
 #include "Templates/SharedPointer.h"
-#include "UObject/WeakObjectPtr.h"
 #include "UObject/SoftObjectPtr.h"
-#include "MuR/Types.h"
 
 #if WITH_EDITOR
 #include "Misc/Guid.h"
@@ -1757,6 +1755,8 @@ enum class ECustomizableObjectVersions : int32
 	AddBonePosePriority,
 
 	RemovedMaterialLayerEncoding,
+	
+	AddedSkeletalMeshBreakNode,
 
 	// -----<new versions can be added above this line>--------
 	LastCustomizableObjectVersion
diff --git a/Engine/Plugins/Mutable/Source/CustomizableObject/Private/MuCO/CustomizableObjectSystem.cpp b/Engine/Plugins/Mutable/Source/CustomizableObject/Private/MuCO/CustomizableObjectSystem.cpp
--- a/Engine/Plugins/Mutable/Source/CustomizableObject/Private/MuCO/CustomizableObjectSystem.cpp
+++ b/Engine/Plugins/Mutable/Source/CustomizableObject/Private/MuCO/CustomizableObjectSystem.cpp
@@ -2,13 +2,10 @@
 
 #include "MuCO/CustomizableObjectSystem.h"
 
-#include "Animation/Skeleton.h"
 #include "AssetRegistry/AssetRegistryModule.h"
 #include "Engine/SkeletalMesh.h"
 #include "Engine/SkinnedAsset.h"
-#include "Engine/SkinnedAssetCommon.h"
 #include "Engine/SkeletalMeshLODSettings.h"
-#include "GameFramework/PlayerController.h"
 #include "MuCO/CustomizableInstanceLODManagement.h"
 #include "MuCO/CustomizableObjectInstancePrivate.h"
 #include "MuCO/CustomizableObjectPrivate.h"
@@ -26,8 +23,10 @@
 #include "UObject/UObjectIterator.h"
 #include "Widgets/Notifications/SNotificationList.h"
 #include "ContentStreaming.h"
+#include "SkinnedAssetCompiler.h"
 #include "Components/SkeletalMeshComponent.h"
 #include "MuCO/CustomizableObjectSystemPrivate.h"
+#include "MuR/ParametersPrivate.h"
 
 #if WITH_EDITOR
 #include "Editor.h"
@@ -3452,14 +3451,72 @@ namespace impl
 			}
 		}
 		
+		TArray<	UE::Tasks::FTask> Prerequisites;
+		Prerequisites.Add(CacheRuntimeTexturesEvent);
+		
+#if WITH_EDITOR
+		// Add to the prerequisites of the operation the compilation of all the Skeletal Mesh Parameters
+		// Note : We are accessing the internal parameters (no the ones on the COI) as those do change in a later stage and may not yet have the SKM parameters in them
+		// at this point.
+		{
+			UE::Mutable::Private::FParameters::Private* ParametersPrivate = Operation->Parameters->GetPrivate();
+			check(ParametersPrivate);
+		
+			const int32 ParameterCount = ParametersPrivate->Values.Num();
+			for (int32 ParameterIndex = 0; ParameterIndex < ParameterCount; ++ParameterIndex)
+			{
+				auto AddToPrerequisites = [&Prerequisites](USkeletalMesh* SkeletalMesh) -> void
+				{
+					// If the skeletal mesh is compiling wait for it to be done before proceeding
+					if (SkeletalMesh && SkeletalMesh->IsCompiling())
+					{
+						UE::Tasks::FTaskEvent Event { UE_SOURCE_LOCATION };
+						Prerequisites.Add(Event);
+
+						ExecuteOnGameThread(TEXT("FinishMeshCompilation"), [Event, SkeletalMesh = TStrongObjectPtr(SkeletalMesh)]() mutable
+						{
+							FSkinnedAssetCompilingManager::Get().FinishCompilation({ SkeletalMesh.Get() });
+							Event.Trigger();
+						});
+					}
+				};
+
+				if (!ParametersPrivate->HasMultipleValues(ParameterIndex))
+				{
+					// Add the non ranged value
+					UE::Mutable::Private::FParameterValue& ParameterValue = ParametersPrivate->Values[ParameterIndex];
+					if (TStrongObjectPtr<USkeletalMesh>* SkeletalMeshParameterValue = ParameterValue.TryGet<UE::Mutable::Private::FParamSkeletalMeshType>())
+					{
+						AddToPrerequisites(SkeletalMeshParameterValue->Get());
+					}
+				}
+				else
+				{
+					const TMap<TArray<int32>, UE::Mutable::Private::FParameterValue>& MultiValueParameterData = ParametersPrivate->MultiValues[ParameterIndex];
+
+					TArray<UE::Mutable::Private::FParameterValue> MultiDimensionalParameterValues;
+					MultiValueParameterData.GenerateValueArray(MultiDimensionalParameterValues);
+				
+					for (UE::Mutable::Private::FParameterValue& MultiDimensionalParameterValue : MultiDimensionalParameterValues)
+					{
+						if (TStrongObjectPtr<USkeletalMesh>* MultiDimensionalSkeletalMeshParameterValue = MultiDimensionalParameterValue.TryGet<UE::Mutable::Private::FParamSkeletalMeshType>())
+						{
+							AddToPrerequisites(MultiDimensionalSkeletalMeshParameterValue->Get());
+						}
+					}
+				}
+			}
+		}
+#endif
+		
 		UE::Tasks::FTask Mutable_GetMeshTask = SystemPrivate->MutableTaskGraph.AddMutableThreadTask(
 		TEXT("Task_Mutable_GetMeshes"),
 		[Operation]()
 		{
 			Task_Mutable_GetMeshes(Operation);
 		},
-		{ CacheRuntimeTexturesEvent });
-		
+		{ Prerequisites});
+
 		UE::Tasks::FTask Mutable_GetMaterialsTask = SystemPrivate->MutableTaskGraph.AddMutableThreadTask(
 		TEXT("Task_Mutable_GetMaterials"),
 		[Operation]()
diff --git a/Engine/Plugins/Mutable/Source/CustomizableObjectEditor/Private/MuCOE/CustomizableObjectEditorModule.cpp b/Engine/Plugins/Mutable/Source/CustomizableObjectEditor/Private/MuCOE/CustomizableObjectEditorModule.cpp
--- a/Engine/Plugins/Mutable/Source/CustomizableObjectEditor/Private/MuCOE/CustomizableObjectEditorModule.cpp
+++ b/Engine/Plugins/Mutable/Source/CustomizableObjectEditor/Private/MuCOE/CustomizableObjectEditorModule.cpp
@@ -5,7 +5,6 @@
 #include "Animation/AnimBlueprint.h"
 #include "AssetToolsModule.h"
 #include "CustomizableObjectProjectSettings.h"
-#include "GameFramework/Pawn.h"
 #include "ISettingsModule.h"
 #include "ISettingsSection.h"
 #include "MessageLogModule.h"
@@ -19,11 +18,9 @@
 #include "MuCO/ICustomizableObjectModule.h"		// For instance editor command utility function
 #include "MuCO/UnrealPortabilityHelpers.h"
 #include "MuCOE/CustomizableInstanceDetails.h"
-#include "MuCOE/CustomizableObjectCustomSettings.h"
 #include "MuCOE/CustomizableObjectCustomSettingsDetails.h"
 #include "MuCOE/CustomizableObjectDetails.h"
 #include "MuCOE/CustomizableObjectEditor.h"
-#include "MuCOE/CustomizableObjectEditorLogger.h"
 #include "MuCOE/CustomizableObjectEditorSettings.h"
 #include "MuCOE/CustomizableObjectEditorStyle.h"
 #include "MuCOE/CustomizableObjectIdentifierCustomization.h"
@@ -34,10 +31,8 @@
 #include "MuCOE/CustomizableObjectNodeObjectGroupDetails.h"
 #include "MuCOE/Nodes/CONodeMaterialBreak.h"
 #include "MuCOE/Nodes/CONodeMaterialModify.h"
-#include "MuCOE/Nodes/CustomizableObjectNodeCopyMaterial.h"
 #include "MuCOE/Nodes/CustomizableObjectNodeExternalPin.h"
 #include "MuCOE/Nodes/CustomizableObjectNodeExternalPinDetails.h"
-#include "MuCOE/Nodes/CustomizableObjectNodeLayoutBlocks.h"
 #include "MuCOE/Nodes/CustomizableObjectNodeMacroInstance.h"
 #include "MuCOE/Nodes/CustomizableObjectNodeMacroInstanceDetails.h"
 #include "MuCOE/Nodes/CustomizableObjectNodeMaterial.h"
@@ -77,7 +72,6 @@
 #include "MuCOE/Nodes/CustomizableObjectNodeSkeletalMesh.h"
 #include "MuCOE/Nodes/CustomizableObjectNodeSkeletalMeshDetails.h"
 #include "MuCOE/Nodes/CustomizableObjectNodeStaticMesh.h"
-#include "MuCOE/Nodes/CustomizableObjectNodeSkeletalMeshParameter.h"
 #include "MuCOE/Nodes/CustomizableObjectNodeSkeletalMeshParameterDetails.h"
 #include "MuCOE/Nodes/CustomizableObjectNodeTable.h"
 #include "MuCOE/Nodes/CustomizableObjectNodeTableDetails.h"
@@ -98,9 +92,7 @@
 #include "AssetRegistry/AssetRegistryModule.h"
 #include "MuCO/CustomizableObjectInstance.h"
 #include "MuCOE/Nodes/CustomizableObjectNodeGroupProjectorParameter.h"
-#include "UObject/UObjectIterator.h"
 #include "Subsystems/PlacementSubsystem.h"
-#include "Components/SkeletalMeshComponent.h"
 #include "MuCOE/GraphTraversal.h"
 #include "MuCOE/CustomizableObjectInstanceBaker.h"
 #include "Editor.h"
@@ -109,13 +101,12 @@
 #include "MuCO/CustomizableObjectSystemPrivate.h"
 #include "MuCOE/CompileRequest.h"
 #include "MuCOE/CustomizableObjectGraph.h"
-#include "Nodes/CustomizableObjectNodeComponentMeshDetails.h"
 #include "MuCOE/GenerateMutableSource/GenerateMutableSourceTable.h"
 #include "MuR/External/Operation.h"
 #include "Nodes/CONodeExternalOperation.h"
 #include "Nodes/CONodeExternalOperationDetails.h"
-#include "Nodes/CONodeExternalTypeParameter.h"
 #include "Nodes/CustomizableObjectNodeComponentPassthroughMesh.h"
+#include "Nodes/CONodeSkeletalMeshBreak.h"
 #include "Nodes/CustomizableObjectNodeTextureParameter.h"
 #include "Nodes/ExternalTypeCustomization.h"
 #include "PhysicsEngine/PhysicsAsset.h"
@@ -287,6 +278,7 @@ void FCustomizableObjectEditorModule::StartupModule()
 	RegisterCustomDetails(PropertyModule, UCustomizableObjectNodeComponentPassthroughMesh::StaticClass(), FOnGetDetailCustomizationInstance::CreateStatic(&FCustomizableObjectNodeDetails::MakeInstance));
 	RegisterCustomDetails(PropertyModule, UCONodeExternalOperation::StaticClass(), FOnGetDetailCustomizationInstance::CreateStatic(&FCONodeExternalOperationDetails::MakeInstance));
 	RegisterCustomDetails(PropertyModule, UCONodeMaterialModify::StaticClass(), FOnGetDetailCustomizationInstance::CreateStatic(&FCustomizableObjectNodeDetails::MakeInstance));
+	RegisterCustomDetails(PropertyModule, UCONodeSkeletalMeshBreak::StaticClass(), FOnGetDetailCustomizationInstance::CreateStatic(&FCustomizableObjectNodeDetails::MakeInstance));
 
 	// Other Objects
 	RegisterCustomDetails(PropertyModule, UCustomizableObject::StaticClass(), FOnGetDetailCustomizationInstance::CreateStatic(&FCustomizableObjectDetails::MakeInstance));
diff --git a/Engine/Plugins/Mutable/Source/CustomizableObjectEditor/Private/MuCOE/EdGraphSchema_CustomizableObject.cpp b/Engine/Plugins/Mutable/Source/CustomizableObjectEditor/Private/MuCOE/EdGraphSchema_CustomizableObject.cpp
--- a/Engine/Plugins/Mutable/Source/CustomizableObjectEditor/Private/MuCOE/EdGraphSchema_CustomizableObject.cpp
+++ b/Engine/Plugins/Mutable/Source/CustomizableObjectEditor/Private/MuCOE/EdGraphSchema_CustomizableObject.cpp
@@ -120,8 +120,8 @@
 #include "MuR/External/VectorAdapter.h"
 #include "Nodes/CONodeExternalOperation.h"
 #include "Nodes/CONodeExternalTypeParameter.h"
-#include "Nodes/ExternalTypeCustomization.h"
 #include "Nodes/CONodeSchema.h"
+#include "Nodes/CONodeSkeletalMeshBreak.h"
 #include "UObject/Class.h"
 
 #include UE_INLINE_GENERATED_CPP_BY_NAME(EdGraphSchema_CustomizableObject)
@@ -643,6 +643,7 @@ void UEdGraphSchema_CustomizableObject::GetGraphContextActions(FGraphContextMenu
 			ContextMenuBuilder.CreateTemplateNode<UCustomizableObjectNodeTextureParameter>(),
 			ContextMenuBuilder.CreateTemplateNode<UCustomizableObjectNodeSkeletalMeshParameter>(),
 			ContextMenuBuilder.CreateTemplateNode<UCustomizableObjectNodeMaterialParameter>(),
+			ContextMenuBuilder.CreateTemplateNode<UCONodeSkeletalMeshBreak>(),
 		};
 
 		AddNewNodeCategoryActionsFiltered(ExperimentalTemplateNodes, ContextMenuBuilder, NC_Experimental.ToString(), 2, Filter);
diff --git a/Engine/Plugins/Mutable/Source/CustomizableObjectEditor/Private/MuCOE/GenerateMutableSource/GenerateMutableSourceMaterial.cpp b/Engine/Plugins/Mutable/Source/CustomizableObjectEditor/Private/MuCOE/GenerateMutableSource/GenerateMutableSourceMaterial.cpp
--- a/Engine/Plugins/Mutable/Source/CustomizableObjectEditor/Private/MuCOE/GenerateMutableSource/GenerateMutableSourceMaterial.cpp
+++ b/Engine/Plugins/Mutable/Source/CustomizableObjectEditor/Private/MuCOE/GenerateMutableSource/GenerateMutableSourceMaterial.cpp
@@ -3,6 +3,8 @@
 #include "MuCOE/GenerateMutableSource/GenerateMutableSourceMaterial.h"
 
 #include "GenerateMutableSourceExternal.h"
+#include "GenerateMutableSourceMesh.h"
+#include "GenerateMutableSourceSkeletalMesh.h"
 #include "Engine/TextureLODSettings.h"
 #include "Interfaces/ITargetPlatform.h"
 #include "MuCOE/CustomizableObjectCompiler.h"
@@ -35,7 +37,9 @@
 #include "Materials/MaterialInterface.h"
 #include "Materials/MaterialInstance.h"
 #include "MuCOE/Nodes/CONodeExternalOperation.h"
+#include "MuCOE/Nodes/CONodeSkeletalMeshBreak.h"
 #include "MuT/NodeMaterialExternal.h"
+#include "MuT/NodeMaterialSkeletalMeshBreak.h"
 
 
 #define LOCTEXT_NAMESPACE "CustomizableObjectEditor"
@@ -296,7 +300,7 @@ UE::Mutable::Private::Ptr<UE::Mutable::Private::NodeMaterial> GenerateMutableSou
 		{
 			if (!MaterialParameterNode->ReferenceValue)
 			{
-				GenerationContext.Log(LOCTEXT("NoRefernceMaterialToBreak", "Could not find a Reference Material to break."), MaterialParameterNode);
+				GenerationContext.Log(LOCTEXT("NoReferenceMaterialToBreak", "Could not find a Reference Material to break."), MaterialParameterNode);
 			}
 			else
 			{
@@ -360,6 +364,24 @@ UE::Mutable::Private::Ptr<UE::Mutable::Private::NodeMaterial> GenerateMutableSou
 		Result = MaterialNode;
 	}
 
+	else if (const UCONodeSkeletalMeshBreak* TypedSkeletalMeshBreakNode = Cast<UCONodeSkeletalMeshBreak>(Node))
+	{
+		const UCONodeSkeletalMeshBreakPinData& OverrideMaterialPinData = TypedSkeletalMeshBreakNode->GetPinData<UCONodeSkeletalMeshBreakPinData>(*Pin);
+
+		const FName TargetSlotName = OverrideMaterialPinData.TargetMaterialSlotName;
+		if (!TargetSlotName.IsNone())
+		{
+			if (const UEdGraphPin* ConnectedMeshPin = FollowInputPin(*TypedSkeletalMeshBreakNode->InputPassthroughMesh.Get()))
+			{
+				UE::Mutable::Private::Ptr<UE::Mutable::Private::NodeMaterialSkeletalMeshBreak> MaterialFromSlotNode = new UE::Mutable::Private::NodeMaterialSkeletalMeshBreak();;
+				MaterialFromSlotNode->Mesh = UE::Mutable::Private::GenerateMutableSourceSkeletalMesh(ConnectedMeshPin, GenerationContext);
+				MaterialFromSlotNode->SlotName = TargetSlotName;
+
+				Result = MaterialFromSlotNode;
+			}
+		}
+	}
+
 	else if (const UCustomizableObjectNodeTable* TypedNodeTable = Cast<UCustomizableObjectNodeTable>(Node))
 	{
 		//This node will add a default value in case of error
diff --git a/Engine/Plugins/Mutable/Source/CustomizableObjectEditor/Private/MuCOE/Nodes/CONodeSkeletalMeshBreak.cpp b/Engine/Plugins/Mutable/Source/CustomizableObjectEditor/Private/MuCOE/Nodes/CONodeSkeletalMeshBreak.cpp
new file mode 100644
--- /dev/null
+++ b/Engine/Plugins/Mutable/Source/CustomizableObjectEditor/Private/MuCOE/Nodes/CONodeSkeletalMeshBreak.cpp
@@ -0,0 +1,170 @@
+﻿// Copyright Epic Games, Inc. All Rights Reserved.
+
+#include "CONodeSkeletalMeshBreak.h"
+
+#include "MuCOE/EdGraphSchema_CustomizableObject.h"
+#include "ScopedTransaction.h"
+#include "MuCOE/RemapPins/CustomizableObjectNodeRemapPinsByPosition.h"
+
+#define LOCTEXT_NAMESPACE "CustomizableObjectEditor"
+
+
+FText UCONodeSkeletalMeshBreak::GetNodeTitle(ENodeTitleType::Type TitleType) const
+{
+	return LOCTEXT("UCONodeSkeletalMeshBreak_NodeTitle","Break Skeletal Mesh");
+}
+
+
+FText UCONodeSkeletalMeshBreak::GetTooltipText() const
+{
+	return Super::GetTooltipText();
+}
+
+
+FLinearColor UCONodeSkeletalMeshBreak::GetNodeTitleColor() const
+{
+	return UEdGraphSchema_CustomizableObject::GetPinTypeColor(UEdGraphSchema_CustomizableObject::PC_Material);
+}
+
+
+void UCONodeSkeletalMeshBreak::AllocateDefaultPins(UCustomizableObjectNodeRemapPins* RemapPins)
+{
+	// Passthrough skm input
+	InputPassthroughMesh = CustomCreatePinSimple(EGPD_Input, UEdGraphSchema_CustomizableObject::PC_PassthroughSkeletalMesh);
+	
+	// Regenerate the new pins from the old pin's PinData which may have changed
+	OutputMaterialPins.Reset();
+
+	// Get old pins
+	TArray<UEdGraphPin*> AllNonOrphanPins = GetAllNonOrphanPins();
+	for (const UEdGraphPin* Pin : AllNonOrphanPins)
+	{
+		if (!IsMaterialOutputPin(*Pin))
+		{
+			continue;
+		}
+		
+		if (UCONodeSkeletalMeshBreakPinData* PinData = Cast<UCONodeSkeletalMeshBreakPinData>(GetPinData(*Pin)))
+		{
+			CreateMaterialOutputPin(PinData);
+		}
+	}
+}
+
+
+bool UCONodeSkeletalMeshBreak::IsExperimental() const
+{
+	return true;
+}
+
+
+FText UCONodeSkeletalMeshBreak::GetPinEditableName(const UEdGraphPin& Pin) const
+{
+	const UCustomizableObjectNodePinData* PinData = GetPinData(Pin);
+	if (const UCONodeSkeletalMeshBreakPinData* CastedPinData = Cast<UCONodeSkeletalMeshBreakPinData>(PinData))
+	{
+		return FText::FromName(CastedPinData->TargetMaterialSlotName);
+	}
+
+	return Super::GetPinEditableName(Pin);
+}
+
+
+void UCONodeSkeletalMeshBreak::SetPinEditableName(const UEdGraphPin& Pin, const FText& Value)
+{
+	UCustomizableObjectNodePinData* PinData = GetPinData(Pin);
+	if (UCONodeSkeletalMeshBreakPinData* CastedPinData = Cast<UCONodeSkeletalMeshBreakPinData>(PinData))
+	{
+		CastedPinData->TargetMaterialSlotName =  FName{Value.ToString()};
+	}
+	else
+	{
+		Super::SetPinEditableName(Pin, Value);
+	}
+}
+
+
+bool UCONodeSkeletalMeshBreak::HasPinViewer() const
+{
+	return true;
+}
+
+
+EAddPinNodeButtonLocation UCONodeSkeletalMeshBreak::GetAddPinButtonNodeSide() const
+{
+	return EAddPinNodeButtonLocation::OUTPUT;
+}
+
+
+EEditablePinNameBoxVisibilityPolicy UCONodeSkeletalMeshBreak::GetEditablePinNameVisibilityPolicy(const UEdGraphPin& Pin) const
+{
+	return EEditablePinNameBoxVisibilityPolicy::ALWAYS_VISIBLE;
+}
+
+
+void UCONodeSkeletalMeshBreak::AddPinFromUI()
+{
+	FScopedTransaction Transaction(LOCTEXT("UCONodeSkeletalMeshBreak_AddPinFromNodeUITransaction", "Add Material Output Pin from Node UI"));
+	Modify();
+
+	UCONodeSkeletalMeshBreakPinData* PinData = NewObject<UCONodeSkeletalMeshBreakPinData>(this);;
+	if (CreateMaterialOutputPin(PinData))
+	{
+		GetGraph()->NotifyNodeChanged(this);
+	}
+}
+
+
+bool UCONodeSkeletalMeshBreak::CanPinBeRemoved(const UEdGraphPin& Pin) const
+{
+	const bool bCanPinBeRemoved = !Pin.LinkedTo.Num() && !Pin.bOrphanedPin;
+	if (!bCanPinBeRemoved)
+	{
+		return false;
+	}
+	
+	return IsMaterialOutputPin(Pin);
+}
+
+
+bool UCONodeSkeletalMeshBreak::CustomRemovePin(UEdGraphPin& Pin)
+{
+	OutputMaterialPins.Remove(&Pin);
+	return Super::CustomRemovePin(Pin);
+}
+
+
+bool UCONodeSkeletalMeshBreak::CanRenamePin(const UEdGraphPin& Pin) const
+{
+	return IsMaterialOutputPin(Pin);
+}
+
+
+bool UCONodeSkeletalMeshBreak::CreateMaterialOutputPin(const UCONodeSkeletalMeshBreakPinData* InPinData)
+{
+	UEdGraphPin* NewPin = CustomCreatePin(EGPD_Output, UEdGraphSchema_CustomizableObject::PC_Material, FName{"Material Slot"}, (UCONodeSkeletalMeshBreakPinData*)InPinData);;
+	if (NewPin)
+	{
+		OutputMaterialPins.Add(NewPin);
+		return true;
+	}
+	
+	return false;
+}
+
+
+UCustomizableObjectNodeRemapPins* UCONodeSkeletalMeshBreak::CreateRemapPinsDefault() const
+{
+	return Super::CreateRemapPinsByPosition();
+}
+
+
+bool UCONodeSkeletalMeshBreak::IsMaterialOutputPin(const UEdGraphPin& Pin) const
+{
+	return Pin.Direction == EEdGraphPinDirection::EGPD_Output && Pin.PinType.PinCategory == UEdGraphSchema_CustomizableObject::PC_Material;
+}
+
+
+
+
+#undef LOCTEXT_NAMESPACE
diff --git a/Engine/Plugins/Mutable/Source/CustomizableObjectEditor/Private/MuCOE/Nodes/CONodeSkeletalMeshBreak.h b/Engine/Plugins/Mutable/Source/CustomizableObjectEditor/Private/MuCOE/Nodes/CONodeSkeletalMeshBreak.h
new file mode 100644
--- /dev/null
+++ b/Engine/Plugins/Mutable/Source/CustomizableObjectEditor/Private/MuCOE/Nodes/CONodeSkeletalMeshBreak.h
@@ -0,0 +1,64 @@
+﻿// Copyright Epic Games, Inc. All Rights Reserved.
+
+#pragma once
+
+#include "MuCOE/Nodes/CustomizableObjectNode.h"
+
+#include "CONodeSkeletalMeshBreak.generated.h"
+
+UCLASS(MinimalAPI)
+class UCONodeSkeletalMeshBreakPinData : public UCustomizableObjectNodePinData
+{
+	GENERATED_BODY()
+
+public:
+
+	UPROPERTY(EditAnywhere, Category = NoCategory)
+	FName TargetMaterialSlotName;
+};
+
+
+UCLASS(MinimalAPI)
+class UCONodeSkeletalMeshBreak : public UCustomizableObjectNode
+{
+	GENERATED_BODY()
+
+public:
+	
+	// UCustomizableObjectNode interface
+	virtual void AllocateDefaultPins(UCustomizableObjectNodeRemapPins* RemapPins) override;
+	virtual bool IsExperimental() const override;
+
+	// UEdGraphNode interface
+	virtual FText GetNodeTitle(ENodeTitleType::Type TitleType) const override;
+	virtual FText GetTooltipText() const override;
+	virtual FLinearColor GetNodeTitleColor() const override;
+	
+	virtual FText GetPinEditableName(const UEdGraphPin& Pin) const override;
+	virtual void SetPinEditableName(const UEdGraphPin& Pin, const FText& Value) override;
+	virtual bool HasPinViewer() const override;
+	virtual EAddPinNodeButtonLocation GetAddPinButtonNodeSide() const override;
+	virtual EEditablePinNameBoxVisibilityPolicy GetEditablePinNameVisibilityPolicy(const UEdGraphPin& Pin) const override;
+	virtual void AddPinFromUI() override;
+	virtual bool CanPinBeRemoved(const UEdGraphPin& Pin) const override;
+	virtual bool CustomRemovePin(UEdGraphPin& Pin) override;
+	virtual bool CanRenamePin(const UEdGraphPin& Pin) const override;
+
+	UPROPERTY()
+	FEdGraphPinReference InputPassthroughMesh;
+
+	UPROPERTY()
+	TArray<FEdGraphPinReference> OutputMaterialPins;
+	
+private:
+
+	bool IsMaterialOutputPin(const UEdGraphPin& Pin) const;
+
+	bool CreateMaterialOutputPin(const UCONodeSkeletalMeshBreakPinData* InPinData);
+
+protected:
+	
+	virtual UCustomizableObjectNodeRemapPins* CreateRemapPinsDefault() const override;
+
+};
+
diff --git a/Engine/Plugins/Mutable/Source/CustomizableObjectEditor/Private/MuCOE/SMutableGraphViewer.cpp b/Engine/Plugins/Mutable/Source/CustomizableObjectEditor/Private/MuCOE/SMutableGraphViewer.cpp
--- a/Engine/Plugins/Mutable/Source/CustomizableObjectEditor/Private/MuCOE/SMutableGraphViewer.cpp
+++ b/Engine/Plugins/Mutable/Source/CustomizableObjectEditor/Private/MuCOE/SMutableGraphViewer.cpp
@@ -231,6 +231,7 @@ static const char* MutableNodeNames[] =
 	"MaterialVariation",
 	"MaterialParameter",
 	"MaterialExternal",
+	"MaterialSkeletalMeshBreak",
 
 	"ImageMaterialBreak",
 	"ScalarMaterialBreak",
diff --git a/Engine/Plugins/Mutable/Source/MutableRuntime/Internal/MuR/ModelPrivate.h b/Engine/Plugins/Mutable/Source/MutableRuntime/Internal/MuR/ModelPrivate.h
--- a/Engine/Plugins/Mutable/Source/MutableRuntime/Internal/MuR/ModelPrivate.h
+++ b/Engine/Plugins/Mutable/Source/MutableRuntime/Internal/MuR/ModelPrivate.h
@@ -609,9 +609,9 @@ namespace UE::Mutable::Private
             return index;
         }
 
-        OP::ADDRESS AddConstant(const FString& Str)
+        OP::CONSTANT_STRING_INDEX AddConstant(const FString& Str)
         {     
-			return (OP::ADDRESS)ConstantStrings.AddUnique(Str);
+			return (OP::CONSTANT_STRING_INDEX)ConstantStrings.AddUnique(Str);
         }
 		
 		OP::ADDRESS AddConstant(const TArray<uint32>& UInt32List)
diff --git a/Engine/Plugins/Mutable/Source/MutableRuntime/Internal/MuR/Operations.h b/Engine/Plugins/Mutable/Source/MutableRuntime/Internal/MuR/Operations.h
--- a/Engine/Plugins/Mutable/Source/MutableRuntime/Internal/MuR/Operations.h
+++ b/Engine/Plugins/Mutable/Source/MutableRuntime/Internal/MuR/Operations.h
@@ -381,6 +381,9 @@ namespace UE::Mutable::Private
     	//! External operation
 		MI_EXTERNAL,
 
+    	// Get a material from a given skeletal mesh
+		MI_SKELETALMESH_BREAK,
+    	
     	//-----------------------------------------------------------------------------------------
 		// FInstancedStruct operations
 		//-----------------------------------------------------------------------------------------
@@ -432,7 +435,7 @@ namespace UE::Mutable::Private
     struct TOperation
     {
         typedef ADDRESS_TYPE ADDRESS;
-        typedef ADDRESS_TYPE CONSTANT_STRING_ADDRESS;
+        typedef ADDRESS_TYPE CONSTANT_STRING_INDEX;
         typedef ADDRESS_TYPE STRING_ADDRESS;
 		typedef ADDRESS_TYPE LIST_UINT32_ADDRESS;
 
@@ -499,6 +502,12 @@ namespace UE::Mutable::Private
 			LIST_UINT32_ADDRESS ImageParameterAddress;
 		};
 
+    	struct MaterialSkeletalMeshBreakArgs
+    	{
+    		CONSTANT_STRING_INDEX SlotNameStringIndex;
+    		ADDRESS SkeletalMesh;
+    	};
+
     	struct MeshSkeletalMeshBreakArgs
     	{
     		ADDRESS SkeletalMeshParameter;
diff --git a/Engine/Plugins/Mutable/Source/MutableRuntime/Private/MuR/CodeRunner.cpp b/Engine/Plugins/Mutable/Source/MutableRuntime/Private/MuR/CodeRunner.cpp
--- a/Engine/Plugins/Mutable/Source/MutableRuntime/Private/MuR/CodeRunner.cpp
+++ b/Engine/Plugins/Mutable/Source/MutableRuntime/Private/MuR/CodeRunner.cpp
@@ -2,6 +2,7 @@
 
 #include "MuR/CodeRunner.h"
 
+#include "SkinnedAssetCompiler.h"
 #include "GenericPlatform/GenericPlatformMath.h"
 #include "HAL/UnrealMemory.h"
 #include "Logging/LogCategory.h"
@@ -1484,9 +1485,9 @@ namespace UE::Mutable::Private
         {
 			const uint8* Data = Program.GetOpArgsPointer(Item.At);
 
-			OP::CONSTANT_STRING_ADDRESS NameAt = 0;
-			FMemory::Memcpy(&NameAt, Data, sizeof(OP::CONSTANT_STRING_ADDRESS)); 
-			Data += sizeof(OP::CONSTANT_STRING_ADDRESS);
+			OP::CONSTANT_STRING_INDEX NameAt = 0;
+			FMemory::Memcpy(&NameAt, Data, sizeof(OP::CONSTANT_STRING_INDEX)); 
+			Data += sizeof(OP::CONSTANT_STRING_INDEX);
 				
 			OP::ADDRESS FactorAt = 0;
 			FMemory::Memcpy(&FactorAt, Data, sizeof(OP::ADDRESS)); 
@@ -6590,6 +6591,65 @@ namespace UE::Mutable::Private
 
 			break;
 		}
+
+		case EOpType::MI_SKELETALMESH_BREAK:
+		{
+			MUTABLE_CPUPROFILER_SCOPE(RunCodeMaterial_MI_SKELETALMESH_BREAK);
+
+			OP::MaterialSkeletalMeshBreakArgs Args = Program.GetOpArgs<OP::MaterialSkeletalMeshBreakArgs >(Item.At);
+			
+			switch (Item.Stage)
+			{
+			case 0:
+				{
+					check(Program.ConstantStrings.IsValidIndex(Args.SlotNameStringIndex));
+					const FString& SlotName = Program.ConstantStrings[Args.SlotNameStringIndex];
+					check (!SlotName.IsEmpty())
+				
+					// Request the preparation of the skeletal mesh before doing any work
+					AddOp( FScheduledOp( Item.At, Item, 1),
+						   FScheduledOp( Args.SkeletalMesh, Item));
+					break;
+				}
+			case 1:
+				{
+					TSharedPtr<const FSkeletalMesh> LoadedSkeletalMesh = LoadSkeletalMesh(FCacheAddress(Args.SkeletalMesh, Item));
+					
+					TPassthroughObjectPtr<USkeletalMesh> PassthroughMesh = LoadedSkeletalMesh->PassthroughObject;
+					check(PassthroughMesh);
+					
+					TObjectPtr<UMaterialInterface> FoundMaterialInterface = nullptr;
+					if (USkeletalMesh* EngineSkeletalMesh = PassthroughMesh.Get())
+					{
+						const FString& SlotName = Program.ConstantStrings[Args.SlotNameStringIndex];
+						
+						// From the Skeletal mesh get the material in the override slot with the provided name
+						const TArray<FSkeletalMaterial> Materials = EngineSkeletalMesh->GetMaterials();
+						for (const FSkeletalMaterial& Material : Materials)
+						{
+							if (Material.MaterialSlotName == SlotName)
+							{
+								FoundMaterialInterface = Material.MaterialInterface;
+								break;
+							}
+						}
+					}
+
+					TSharedPtr<FMaterial> MutableMaterial = MakeShared<FMaterial>();
+					MutableMaterial->PassthroughObject = TPassthroughObjectPtr(FoundMaterialInterface.Get());
+
+					StoreMaterial(Item, MutableMaterial);
+					
+					break;
+				}
+			
+			default:
+				check(false);
+			}
+			
+			break;
+		}
+			
 		default: unimplemented();
 		}
 	}
diff --git a/Engine/Plugins/Mutable/Source/MutableRuntime/Private/MuR/CodeRunnerBegin.cpp b/Engine/Plugins/Mutable/Source/MutableRuntime/Private/MuR/CodeRunnerBegin.cpp
--- a/Engine/Plugins/Mutable/Source/MutableRuntime/Private/MuR/CodeRunnerBegin.cpp
+++ b/Engine/Plugins/Mutable/Source/MutableRuntime/Private/MuR/CodeRunnerBegin.cpp
@@ -6,12 +6,7 @@
 #include "MuR/MutableTrace.h"
 #include "MuR/Operations.h"
 #include "MuR/SystemPrivate.h"
-#include "MuR/External/FloatAdapter.h"
-#include "MuR/External/MaterialAdapter.h"
-#include "MuR/External/MeshAdapter.h"
 #include "MuR/External/Operation.h"
-#include "MuR/External/TextureAdapter.h"
-#include "MuR/External/VectorAdapter.h"
 
 
 namespace UE::Mutable::Private
@@ -1719,9 +1714,9 @@ namespace UE::Mutable::Private
 				{
 					const uint8* Data = Program.GetOpArgsPointer(Item.At);
 
-					OP::CONSTANT_STRING_ADDRESS NameAt = 0;
-					FMemory::Memcpy(&NameAt, Data, sizeof(OP::CONSTANT_STRING_ADDRESS)); 
-					Data += sizeof(OP::CONSTANT_STRING_ADDRESS);
+					OP::CONSTANT_STRING_INDEX NameAt = 0;
+					FMemory::Memcpy(&NameAt, Data, sizeof(OP::CONSTANT_STRING_INDEX)); 
+					Data += sizeof(OP::CONSTANT_STRING_INDEX);
 					
 					OP::ADDRESS FactorAt = 0;
 					FMemory::Memcpy(&FactorAt, Data, sizeof(OP::ADDRESS)); 
@@ -2498,6 +2493,16 @@ namespace UE::Mutable::Private
 					
 					break;
 				}
+
+			case EOpType::MI_SKELETALMESH_BREAK:
+				{
+					const OP::MaterialSkeletalMeshBreakArgs Args = Program.GetOpArgs<OP::MaterialSkeletalMeshBreakArgs>(Item.At);
+					PushOp(FScheduledOpInline(Args.SlotNameStringIndex, Item));
+					PushOp(FScheduledOpInline(Args.SkeletalMesh, Item));
+					StoreNone(Item);
+					
+					break;
+				}
 			
 			default:
 				unimplemented();
diff --git a/Engine/Plugins/Mutable/Source/MutableRuntime/Private/MuR/Operations.cpp b/Engine/Plugins/Mutable/Source/MutableRuntime/Private/MuR/Operations.cpp
--- a/Engine/Plugins/Mutable/Source/MutableRuntime/Private/MuR/Operations.cpp
+++ b/Engine/Plugins/Mutable/Source/MutableRuntime/Private/MuR/Operations.cpp
@@ -162,7 +162,8 @@ namespace UE::Mutable::Private
 		{ EDataType::Layout,			},	// LA_REMOVEBLOCKS
 		{ EDataType::Layout,			},	// LA_FROMMESH
 
-		{ EDataType::Material,		},	// IM_EXTERNAL
+		{ EDataType::Material,		},	// MI_EXTERNAL
+		{ EDataType::Material,		},	// MI_SKELETALMESH_BREAK
 
     	{ EDataType::InstancedStruct,},	// IS_EXTERNAL
 	};
@@ -967,6 +968,13 @@ namespace UE::Mutable::Private
 
        		break;
         }
+
+        case EOpType::MI_SKELETALMESH_BREAK:
+        {
+        	const OP::MaterialSkeletalMeshBreakArgs Args = program.GetOpArgs<OP::MaterialSkeletalMeshBreakArgs>(at);
+        	f(Args.SkeletalMesh);
+        	break;
+        }
         	
         default:
 			check( false );
diff --git a/Engine/Plugins/Mutable/Source/MutableTools/Internal/MuT/Node.h b/Engine/Plugins/Mutable/Source/MutableTools/Internal/MuT/Node.h
--- a/Engine/Plugins/Mutable/Source/MutableTools/Internal/MuT/Node.h
+++ b/Engine/Plugins/Mutable/Source/MutableTools/Internal/MuT/Node.h
@@ -176,7 +176,8 @@ namespace UE::Mutable::Private
 			MaterialVariation,
 			MaterialParameter,
 			MaterialExternal,
-			
+			MaterialSkeletalMeshBreak,
+
 			ImageMaterialBreak,
 			ScalarMaterialBreak,
 			ColourMaterialBreak,
diff --git a/Engine/Plugins/Mutable/Source/MutableTools/Internal/MuT/NodeMaterialSkeletalMeshBreak.h b/Engine/Plugins/Mutable/Source/MutableTools/Internal/MuT/NodeMaterialSkeletalMeshBreak.h
new file mode 100644
--- /dev/null
+++ b/Engine/Plugins/Mutable/Source/MutableTools/Internal/MuT/NodeMaterialSkeletalMeshBreak.h
@@ -0,0 +1,44 @@
+﻿// Copyright Epic Games, Inc. All Rights Reserved.
+
+#pragma once
+
+#include "MuR/Ptr.h"
+#include "MuT/NodeMaterial.h"
+#include "MuT/NodeSkeletalMesh.h"
+
+
+#define UE_API MUTABLETOOLS_API
+
+namespace UE::Mutable::Private
+{
+	/** Provided a mesh get the material whose material slot name does match the "SlotName" set. */
+	class NodeMaterialSkeletalMeshBreak : public NodeMaterial
+	{
+	public:
+
+		/** The name of the material slot to get from the given mesh. */
+		FName SlotName;
+
+		/** The mesh to use as source for the material set in the slot with "SlotName" name. */
+		Ptr<NodeSkeletalMesh> Mesh;
+		
+	public:
+
+		// Node Interface
+		virtual const FNodeType* GetType() const override { return &StaticType; }
+		static const FNodeType* GetStaticType() { return &StaticType; }
+
+	protected:
+
+		// Forbidden. Manage with the Ptr<> template.
+		inline ~NodeMaterialSkeletalMeshBreak() {}
+
+	private:
+
+		static UE_API FNodeType StaticType;
+
+	};
+	
+}
+
+#undef UE_API
diff --git a/Engine/Plugins/Mutable/Source/MutableTools/Private/MuT/ASTOpMaterialSkeletalMeshBreak.cpp b/Engine/Plugins/Mutable/Source/MutableTools/Private/MuT/ASTOpMaterialSkeletalMeshBreak.cpp
new file mode 100644
--- /dev/null
+++ b/Engine/Plugins/Mutable/Source/MutableTools/Private/MuT/ASTOpMaterialSkeletalMeshBreak.cpp
@@ -0,0 +1,83 @@
+// Copyright Epic Games, Inc. All Rights Reserved.
+
+#include "ASTOpMaterialSkeletalMeshBreak.h"
+
+
+namespace UE::Mutable::Private
+{
+	ASTOpMaterialSkeletalMeshBreak::ASTOpMaterialSkeletalMeshBreak()
+		: SkeletalMesh(this)
+	{
+		
+	}
+
+	ASTOpMaterialSkeletalMeshBreak::~ASTOpMaterialSkeletalMeshBreak()
+	{
+		// Explicit call needed to avoid recursive destruction
+		ASTOp::RemoveChildren();
+	}
+
+	EOpType ASTOpMaterialSkeletalMeshBreak::GetOpType() const
+	{
+		return EOpType::MI_SKELETALMESH_BREAK;
+	}
+
+	uint64 ASTOpMaterialSkeletalMeshBreak::Hash() const
+	{
+		uint64 Hash = std::hash<uint64>()(uint64(GetOpType()));
+		hash_combine(Hash, GetTypeHash(SlotName));
+		hash_combine(Hash, SkeletalMesh.child().get());
+		return Hash;
+	}
+
+	void ASTOpMaterialSkeletalMeshBreak::ForEachChild(const TFunctionRef<void(ASTChild&)> F)
+	{
+		F(SkeletalMesh);
+	}
+
+	bool ASTOpMaterialSkeletalMeshBreak::IsEqual(const ASTOp& InOtherUUntyped) const
+	{
+		if (InOtherUUntyped.GetOpType() == GetOpType())
+		{
+			const ASTOpMaterialSkeletalMeshBreak* Other = static_cast<const ASTOpMaterialSkeletalMeshBreak*>(&InOtherUUntyped);
+			return SlotName == Other->SlotName && SkeletalMesh == Other->SkeletalMesh;
+		}
+
+		return false;
+	}
+
+	Ptr<ASTOp> ASTOpMaterialSkeletalMeshBreak::Clone(MapChildFuncRef MapChild) const
+	{
+		Ptr<ASTOpMaterialSkeletalMeshBreak> New = new ASTOpMaterialSkeletalMeshBreak();
+		New->SkeletalMesh = MapChild(SkeletalMesh.child());
+		New->SlotName = SlotName;
+		
+		return New;
+	}
+
+	void ASTOpMaterialSkeletalMeshBreak::Link(FProgram& Program, FLinkerOptions* LinkerOptions)
+	{
+		// Already linked?
+		if (linkedAddress)
+		{
+			return;
+		}
+		
+		OP::MaterialSkeletalMeshBreakArgs Args;
+		FMemory::Memzero(Args);
+
+		Args.SlotNameStringIndex = Program.AddConstant(SlotName.ToString());
+		
+		if (SkeletalMesh)
+		{
+			Args.SkeletalMesh = SkeletalMesh->linkedAddress;
+		}
+
+		linkedAddress = (OP::ADDRESS)Program.OpAddress.Num();
+		Program.OpAddress.Add((uint32_t)Program.ByteCode.Num());
+		AppendCode(Program.ByteCode, this->GetOpType());
+		AppendCode(Program.ByteCode, Args);
+	}
+}
+
+
diff --git a/Engine/Plugins/Mutable/Source/MutableTools/Private/MuT/ASTOpMaterialSkeletalMeshBreak.h b/Engine/Plugins/Mutable/Source/MutableTools/Private/MuT/ASTOpMaterialSkeletalMeshBreak.h
new file mode 100644
--- /dev/null
+++ b/Engine/Plugins/Mutable/Source/MutableTools/Private/MuT/ASTOpMaterialSkeletalMeshBreak.h
@@ -0,0 +1,38 @@
+// Copyright Epic Games, Inc. All Rights Reserved.
+
+#pragma once
+
+#include "HAL/Platform.h"
+#include "MuR/Operations.h"
+#include "MuR/Ptr.h"
+#include "MuT/AST.h"
+
+
+namespace UE::Mutable::Private
+{
+	class NodeSkeletalMesh;
+	struct FProgram;
+	
+	class ASTOpMaterialSkeletalMeshBreak final : public ASTOp
+	{
+	public:
+
+		/** The name of the material slot to get from the given mesh. */
+		FName SlotName;
+
+		/** The skeletal mesh to use as source when looking for the material with the target "SlotName"*/
+		ASTChild SkeletalMesh;
+
+		ASTOpMaterialSkeletalMeshBreak();
+		virtual ~ASTOpMaterialSkeletalMeshBreak() override;
+
+		// ASTOp interface.
+		virtual EOpType GetOpType() const override;
+		virtual uint64 Hash() const override;
+		virtual void ForEachChild(const TFunctionRef<void(ASTChild&)>) override;
+		virtual bool IsEqual(const ASTOp&) const override;
+		virtual Ptr<ASTOp> Clone(MapChildFuncRef ) const override;
+		virtual void Link(FProgram&, FLinkerOptions*) override;
+	};
+}
+
diff --git a/Engine/Plugins/Mutable/Source/MutableTools/Private/MuT/ASTOpMeshMorph.cpp b/Engine/Plugins/Mutable/Source/MutableTools/Private/MuT/ASTOpMeshMorph.cpp
--- a/Engine/Plugins/Mutable/Source/MutableTools/Private/MuT/ASTOpMeshMorph.cpp
+++ b/Engine/Plugins/Mutable/Source/MutableTools/Private/MuT/ASTOpMeshMorph.cpp
@@ -82,7 +82,7 @@ namespace UE::Mutable::Private
 			program.OpAddress.Add((uint32_t)program.ByteCode.Num());
 			AppendCode(program.ByteCode, EOpType::ME_MORPH);
 
-			OP::ADDRESS NameAt = program.AddConstant(Name.ToString());
+			OP::CONSTANT_STRING_INDEX NameAt = program.AddConstant(Name.ToString());
 			AppendCode(program.ByteCode, NameAt);
 			
 			OP::ADDRESS FactorAt = Factor ? Factor->linkedAddress : 0;
diff --git a/Engine/Plugins/Mutable/Source/MutableTools/Private/MuT/CodeGenerator.h b/Engine/Plugins/Mutable/Source/MutableTools/Private/MuT/CodeGenerator.h
--- a/Engine/Plugins/Mutable/Source/MutableTools/Private/MuT/CodeGenerator.h
+++ b/Engine/Plugins/Mutable/Source/MutableTools/Private/MuT/CodeGenerator.h
@@ -53,6 +53,7 @@ namespace UE::Mutable::Private
 	class NodeSkeletalMesh;
 	class NodeMaterialExternal;
 	class NodeScalarExternal;
+	class NodeMaterialSkeletalMeshBreak;
 }
 
 namespace UE::Mutable::Private
@@ -662,6 +663,7 @@ namespace UE::Mutable::Private
 		void GenerateMaterial_Variation(const FMaterialGenerationOptions&, FMaterialGenerationResult& result, const NodeMaterialVariation* node);
 		void GenerateMaterial_Parameter(const FMaterialGenerationOptions&, FMaterialGenerationResult& result, const NodeMaterialParameter* node);
 		void GenerateMaterial_External(const FMaterialGenerationOptions&, FMaterialGenerationResult& result, const NodeMaterialExternal* node);
+    	void GenerateMaterial_BreakSkeletalMesh(const FMaterialGenerationOptions&, FMaterialGenerationResult& OutResult, const NodeMaterialSkeletalMeshBreak* InNode);
 
 		/** */
 		Ptr<ASTOp> GenerateMissingMaterialCode(const TCHAR* strWhere, const void* ErrorContext);
diff --git a/Engine/Plugins/Mutable/Source/MutableTools/Private/MuT/CodeGenerator_Material.cpp b/Engine/Plugins/Mutable/Source/MutableTools/Private/MuT/CodeGenerator_Material.cpp
--- a/Engine/Plugins/Mutable/Source/MutableTools/Private/MuT/CodeGenerator_Material.cpp
+++ b/Engine/Plugins/Mutable/Source/MutableTools/Private/MuT/CodeGenerator_Material.cpp
@@ -1,5 +1,6 @@
 // Copyright Epic Games, Inc. All Rights Reserved.
 
+#include "ASTOpMaterialSkeletalMeshBreak.h"
 #include "MuR/Operations.h"
 #include "MuR/ParametersPrivate.h"
 #include "MuR/Platform.h"
@@ -26,6 +27,7 @@
 #include "MuT/NodeMaterialVariation.h"
 #include "MuT/NodeMaterialTable.h"
 #include "MuT/NodeMaterialParameter.h"
+#include "MuT/NodeMaterialSkeletalMeshBreak.h"
 #include "MuT/Table.h"
 #include "MuT/TablePrivate.h"
 
@@ -85,6 +87,10 @@ namespace UE::Mutable::Private
 			GenerateMaterial_External(InOptions, Result, static_cast<const NodeMaterialExternal*>(Node));
 			break;
 
+		case Node::EType::MaterialSkeletalMeshBreak:
+			GenerateMaterial_BreakSkeletalMesh(InOptions, Result, static_cast<const NodeMaterialSkeletalMeshBreak*>(Node));
+			break;
+
 		default: 
 			check(false);
 		}
@@ -335,6 +341,28 @@ namespace UE::Mutable::Private
 		Result.op = ExtensionResult.Op;
 	}
 
+	
+	void CodeGenerator::GenerateMaterial_BreakSkeletalMesh(const FMaterialGenerationOptions& InOptions, FMaterialGenerationResult& OutResult,
+		const NodeMaterialSkeletalMeshBreak* InNode)
+	{
+		// Early out
+		if (!InNode)
+		{
+			OutResult = FMaterialGenerationResult();
+			return;
+		}
+		
+		Ptr<ASTOpMaterialSkeletalMeshBreak> MaterialFromMeshSlotOp = new ASTOpMaterialSkeletalMeshBreak();
+		MaterialFromMeshSlotOp->SlotName = InNode->SlotName;
+		
+		FSkeletalMeshGenerationOptions SkeletalMeshOptions{InOptions, nullptr};
+		FSkeletalMeshGenerationResult SkeletalMeshGenerationResult;
+		GenerateSkeletalMesh(SkeletalMeshGenerationResult, SkeletalMeshOptions, InNode->Mesh);
+		MaterialFromMeshSlotOp->SkeletalMesh = SkeletalMeshGenerationResult.SkeletalMeshOp;
+		
+		OutResult.op = MaterialFromMeshSlotOp;
+	}
+
 
 	Ptr<ASTOp> CodeGenerator::GenerateMissingMaterialCode(const TCHAR* strWhere, const void* ErrorContext)
 	{
diff --git a/Engine/Plugins/Mutable/Source/MutableTools/Private/MuT/Compiler.cpp b/Engine/Plugins/Mutable/Source/MutableTools/Private/MuT/Compiler.cpp
--- a/Engine/Plugins/Mutable/Source/MutableTools/Private/MuT/Compiler.cpp
+++ b/Engine/Plugins/Mutable/Source/MutableTools/Private/MuT/Compiler.cpp
@@ -200,6 +200,7 @@ namespace UE::Mutable::Private
 		{ true,		{ 0, 0, 0, 0, 0, 0, 0, 0, 0, 0, 0, 0, 0, 0, 0, 0, 0, 0, 0, 0, 0, 0, 0, 0, 0, 0, 0, 0, 0 }		},	// LA_FROMMESH
 
 		{ true,		{ 0, 0, 0, 0, 0, 0, 0, 0, 0, 0, 0, 0, 0, 0, 0, 0, 0, 0, 0, 0, 0, 0, 0, 0, 0, 0, 0, 0, 0 }		},	// MI_EXTERNAL
+    	{ true,		{ 0, 0, 0, 0, 0, 0, 0, 0, 0, 0, 0, 0, 0, 0, 0, 0, 0, 0, 0, 0, 0, 0, 0, 0, 0, 0, 0, 0, 0 }		},	// MI_SKELETALMESH_BREAK
 
     	{ true,		{ 0, 0, 0, 0, 0, 0, 0, 0, 0, 0, 0, 0, 0, 0, 0, 0, 0, 0, 0, 0, 0, 0, 0, 0, 0, 0, 0, 0, 0 }		},	// IS_EXTERNAL
 	};
diff --git a/Engine/Plugins/Mutable/Source/MutableTools/Private/MuT/ErrorLog.cpp b/Engine/Plugins/Mutable/Source/MutableTools/Private/MuT/ErrorLog.cpp
--- a/Engine/Plugins/Mutable/Source/MutableTools/Private/MuT/ErrorLog.cpp
+++ b/Engine/Plugins/Mutable/Source/MutableTools/Private/MuT/ErrorLog.cpp
@@ -345,6 +345,7 @@ namespace UE::Mutable::Private
 		TEXT("LA_FROMMESH      "),
 
     	TEXT("MI_EXTERNAL      "),
+    	TEXT("MI_FROM_SKELETALMESH_SLOT"),
 
     	TEXT("IS_EXTERNAL      "),
 	};
diff --git a/Engine/Plugins/Mutable/Source/MutableTools/Private/MuT/Node.cpp b/Engine/Plugins/Mutable/Source/MutableTools/Private/MuT/Node.cpp
--- a/Engine/Plugins/Mutable/Source/MutableTools/Private/MuT/Node.cpp
+++ b/Engine/Plugins/Mutable/Source/MutableTools/Private/MuT/Node.cpp
@@ -111,6 +111,7 @@
 #include "MuT/NodeMaterialSwitch.h"
 #include "MuT/NodeMaterialVariation.h"
 #include "MuT/NodeMaterialParameter.h"
+#include "MuT/NodeMaterialSkeletalMeshBreak.h"
 #include "MuT/NodeMeshExternal.h"
 #include "MuT/NodeScalarExternal.h"
 #include "MuT/NodeSkeletalMesh.h"
@@ -258,7 +259,8 @@ namespace UE::Mutable::Private
 	FNodeType NodeMaterialVariation::StaticType = FNodeType(Node::EType::MaterialVariation, NodeMaterial::GetStaticType());
 	FNodeType NodeMaterialParameter::StaticType = FNodeType(Node::EType::MaterialParameter, NodeMaterial::GetStaticType());
 	FNodeType NodeMaterialExternal::StaticType = FNodeType(Node::EType::MaterialExternal, NodeMaterial::GetStaticType());
-	
+	FNodeType NodeMaterialSkeletalMeshBreak::StaticType = FNodeType(Node::EType::MaterialSkeletalMeshBreak, NodeMaterial::GetStaticType());
+
 	FNodeType NodeImageMaterialBreak::StaticType = FNodeType(Node::EType::ImageMaterialBreak, NodeImage::GetStaticType());
 	FNodeType NodeScalarMaterialBreak::StaticType = FNodeType(Node::EType::ScalarMaterialBreak, NodeScalar::GetStaticType());
 	FNodeType NodeColourMaterialBreak::StaticType = FNodeType(Node::EType::ColourMaterialBreak, NodeColour::GetStaticType());
