[Mutable] Replace SurfaceId with SharedSurfaceId and simplify the AddSurface code.

#rnx

diff --git a/Engine/Plugins/Mutable/Source/CustomizableObject/Internal/MuCO/CustomizableObjectPrivate.h b/Engine/Plugins/Mutable/Source/CustomizableObject/Internal/MuCO/CustomizableObjectPrivate.h
--- a/Engine/Plugins/Mutable/Source/CustomizableObject/Internal/MuCO/CustomizableObjectPrivate.h
+++ b/Engine/Plugins/Mutable/Source/CustomizableObject/Internal/MuCO/CustomizableObjectPrivate.h
@@ -1758,6 +1758,8 @@ enum class ECustomizableObjectVersions : int32
 	
 	AddedSkeletalMeshBreakNode,
 
+	ReplaceSurfaceIdWithSharedSurfaceId,
+
 	// -----<new versions can be added above this line>--------
 	LastCustomizableObjectVersion
 };
diff --git a/Engine/Plugins/Mutable/Source/CustomizableObject/Internal/MuCO/CustomizableObjectSystemPrivate.h b/Engine/Plugins/Mutable/Source/CustomizableObject/Internal/MuCO/CustomizableObjectSystemPrivate.h
--- a/Engine/Plugins/Mutable/Source/CustomizableObject/Internal/MuCO/CustomizableObjectSystemPrivate.h
+++ b/Engine/Plugins/Mutable/Source/CustomizableObject/Internal/MuCO/CustomizableObjectSystemPrivate.h
@@ -376,7 +376,7 @@ struct FInstanceUpdateData
 		FName Name;
 		UE::Mutable::Private::FImageId ImageID;
 		
-		// LOD of the ImageId. If the texture is shared between LOD, first LOD where this image can be found. 
+		// LOD of the ImageId.
 		int32 BaseLOD;
 		int32 BaseMip;
 		
@@ -428,7 +428,6 @@ struct FInstanceUpdateData
 
 		/** Id of the surface in the mutable core instance. */
 		uint32 SurfaceId = 0;
-		uint32 SharedSurfaceId = 0;
 
 		/** Id of the metadata associated with this surface. */ 
 		uint32 SurfaceMetadataId = 0;
diff --git a/Engine/Plugins/Mutable/Source/CustomizableObject/Private/MuCO/CustomizableObjectInstance.cpp b/Engine/Plugins/Mutable/Source/CustomizableObject/Private/MuCO/CustomizableObjectInstance.cpp
--- a/Engine/Plugins/Mutable/Source/CustomizableObject/Private/MuCO/CustomizableObjectInstance.cpp
+++ b/Engine/Plugins/Mutable/Source/CustomizableObject/Private/MuCO/CustomizableObjectInstance.cpp
@@ -5840,7 +5840,7 @@ void UCustomizableInstancePrivate::BuildMaterials(const TSharedRef<FUpdateContex
 				const TSharedRef<const FInstanceUpdateData::FSurface>& Surface = Context->InstanceUpdateData.Surfaces[LOD->FirstSurface + SurfaceIndex];
 
 				// Reuse MaterialSlot from the previous LOD.
-				if (const int32 MaterialIndex = SurfaceIdToMaterialIndex.Find(Surface->SharedSurfaceId); MaterialIndex != INDEX_NONE)
+				if (const int32 MaterialIndex = SurfaceIdToMaterialIndex.Find(Surface->SurfaceId); MaterialIndex != INDEX_NONE)
 				{
 					if (!bUseOverrideMaterialsOnly)
 					{
@@ -5872,7 +5872,7 @@ void UCustomizableInstancePrivate::BuildMaterials(const TSharedRef<FUpdateContex
 				}
 
 				// This section will require a new slot
-				SurfaceIdToMaterialIndex.Add(Surface->SharedSurfaceId);
+				SurfaceIdToMaterialIndex.Add(Surface->SurfaceId);
 
 				// Add and set up the material data for this slot
 				const int32 MaterialSlotIndex = Materials.Num();
diff --git a/Engine/Plugins/Mutable/Source/CustomizableObject/Private/MuCO/CustomizableObjectMipDataProvider.cpp b/Engine/Plugins/Mutable/Source/CustomizableObject/Private/MuCO/CustomizableObjectMipDataProvider.cpp
--- a/Engine/Plugins/Mutable/Source/CustomizableObject/Private/MuCO/CustomizableObjectMipDataProvider.cpp
+++ b/Engine/Plugins/Mutable/Source/CustomizableObject/Private/MuCO/CustomizableObjectMipDataProvider.cpp
@@ -162,7 +162,7 @@ namespace Impl
 		// runtime cannot remember all the resources it has built, and only remembers a fixed amount.
 		UE::Mutable::Private::FImageId MipImageID = Instance->GetImageId(ImageRef.Component, ImageRef.LOD, SurfaceIndex, ImageRef.Image);
 
-		UE::Tasks::TTask<TSharedPtr<const UE::Mutable::Private::FImage>> GetImageTask = System->GetImage({}, InstanceID, MipImageID, ImageRef.BaseMip + OperationData->MipsToSkip, ImageRef.LOD);
+		UE::Tasks::TTask<TSharedPtr<const UE::Mutable::Private::FImage>> GetImageTask = System->GetImage({}, InstanceID, MipImageID, ImageRef.BaseMip + OperationData->MipsToSkip);
 
 		UE::Tasks::AddNested(UE::Tasks::Launch(TEXT("MipDataProvider_EndUpdateImagesTask"),
 				[System, OperationData, InstanceID, StartTime, GetImageTask, EndUpdateImage]() mutable
diff --git a/Engine/Plugins/Mutable/Source/CustomizableObject/Private/MuCO/CustomizableObjectSystem.cpp b/Engine/Plugins/Mutable/Source/CustomizableObject/Private/MuCO/CustomizableObjectSystem.cpp
--- a/Engine/Plugins/Mutable/Source/CustomizableObject/Private/MuCO/CustomizableObjectSystem.cpp
+++ b/Engine/Plugins/Mutable/Source/CustomizableObject/Private/MuCO/CustomizableObjectSystem.cpp
@@ -2222,7 +2222,7 @@ namespace impl
 
 		const UE::Mutable::Private::FInstance* MutableInstance = OperationData->MutableInstance.Get();
 
-		TArray<uint32> SurfacesSharedId;
+		TArray<uint32> GeneratedSurfaceIds;
 		
 		UE::Tasks::FTask LastTask = UE::Tasks::MakeCompletedTask<void>();
 		
@@ -2257,35 +2257,23 @@ namespace impl
 						continue;
 					}
 
-					int32 BaseSurfaceIndex = InstanceSurfaceIndex;
-					int32 BaseLODIndex = MutableLODIndex;
-
-					int32 SurfaceIndex = OperationData->InstanceUpdateData.Surfaces.Add(MakeShared<FInstanceUpdateData::FSurface>());
-					TSharedRef<FInstanceUpdateData::FSurface>& Surface = OperationData->InstanceUpdateData.Surfaces.Last();
 					++LOD->SurfaceCount;
 					
-					const uint32 SharedSurfaceId = MutableInstance->GetSharedSurfaceId(InstanceComponentIndex, MutableLODIndex, InstanceSurfaceIndex);
-					const int32 SharedSurfaceIndex = SurfacesSharedId.Find(SharedSurfaceId);
-					SurfacesSharedId.Add(SharedSurfaceId);
+					const int32 GeneratedSurfaceIndex = GeneratedSurfaceIds.Find(SurfaceId);
+					GeneratedSurfaceIds.Add(SurfaceId);
 
-					if (SharedSurfaceIndex != INDEX_NONE)
+					if (GeneratedSurfaceIndex != INDEX_NONE)
 					{
-						Surface = OperationData->InstanceUpdateData.Surfaces[SharedSurfaceIndex];
+						OperationData->InstanceUpdateData.Surfaces.Add(OperationData->InstanceUpdateData.Surfaces[GeneratedSurfaceIndex]);
 						continue;
 					}
 
-					// New Surface.MaterialIndex is decoded from a parameter at the end of this if()
-					Surface->SurfaceId = SurfaceId;
-					Surface->SurfaceMetadataId = 0;
-					Surface->SharedSurfaceId = SharedSurfaceId;
-
-					// Find the first LOD where this surface can be found
-					MutableInstance->FindBaseSurfaceBySharedId(InstanceComponentIndex, Surface->SharedSurfaceId, BaseSurfaceIndex, BaseLODIndex);
-
-					Surface->SurfaceId = MutableInstance->GetSurfaceId(InstanceComponentIndex, BaseLODIndex, BaseSurfaceIndex);
-					Surface->SurfaceMetadataId = MutableInstance->GetSurfaceCustomId(InstanceComponentIndex, BaseLODIndex, BaseSurfaceIndex);
+					OperationData->InstanceUpdateData.Surfaces.Add(MakeShared<FInstanceUpdateData::FSurface>());
+					TSharedRef<FInstanceUpdateData::FSurface>& Surface = OperationData->InstanceUpdateData.Surfaces.Last();
 
-					Surface->MaterialId = MutableInstance->GetMaterialId(InstanceComponentIndex, BaseLODIndex, BaseSurfaceIndex);
+					Surface->SurfaceId = SurfaceId;
+					Surface->SurfaceMetadataId = MutableInstance->GetSurfaceCustomId(InstanceComponentIndex, MutableLODIndex, InstanceSurfaceIndex);
+					Surface->MaterialId = MutableInstance->GetMaterialId(InstanceComponentIndex, MutableLODIndex, InstanceSurfaceIndex);
 
 					UE::Tasks::TTask<TSharedPtr<const UE::Mutable::Private::FMaterial>> GetMaterialTask = OperationData->MutableSystem->GetMaterial(LastTask, OperationData->InstanceID, Surface->MaterialId);
 					UE::Tasks::AddNested(GetMaterialTask);
@@ -2300,46 +2288,46 @@ namespace impl
 
 					// Vectors
 					Surface->FirstVector = OperationData->InstanceUpdateData.Vectors.Num();
-					Surface->VectorCount = MutableInstance->GetVectorCount(InstanceComponentIndex, BaseLODIndex, BaseSurfaceIndex);
+					Surface->VectorCount = MutableInstance->GetVectorCount(InstanceComponentIndex, MutableLODIndex, InstanceSurfaceIndex);
 					for (int32 VectorIndex = 0; VectorIndex < Surface->VectorCount; ++VectorIndex)
 					{
 						MUTABLE_CPUPROFILER_SCOPE(GetVector);
 						OperationData->InstanceUpdateData.Vectors.Push({});
 						FInstanceUpdateData::FVector& Vector = OperationData->InstanceUpdateData.Vectors.Last();
-						Vector.Name = MutableInstance->GetVectorName(InstanceComponentIndex, BaseLODIndex, BaseSurfaceIndex, VectorIndex);
-						Vector.Vector = MutableInstance->GetVector(InstanceComponentIndex, BaseLODIndex, BaseSurfaceIndex, VectorIndex);
-						Vector.MaterialLayer = MutableInstance->GetVectorLayerIndex(InstanceComponentIndex, BaseLODIndex, BaseSurfaceIndex, VectorIndex);
+						Vector.Name = MutableInstance->GetVectorName(InstanceComponentIndex, MutableLODIndex, InstanceSurfaceIndex, VectorIndex);
+						Vector.Vector = MutableInstance->GetVector(InstanceComponentIndex, MutableLODIndex, InstanceSurfaceIndex, VectorIndex);
+						Vector.MaterialLayer = MutableInstance->GetVectorLayerIndex(InstanceComponentIndex, MutableLODIndex, InstanceSurfaceIndex, VectorIndex);
 					}
 
 					// Scalars
 					Surface->FirstScalar = OperationData->InstanceUpdateData.Scalars.Num();
-					Surface->ScalarCount = MutableInstance->GetScalarCount(InstanceComponentIndex, BaseLODIndex, BaseSurfaceIndex);
+					Surface->ScalarCount = MutableInstance->GetScalarCount(InstanceComponentIndex, MutableLODIndex, InstanceSurfaceIndex);
 					for (int32 ScalarIndex = 0; ScalarIndex < Surface->ScalarCount; ++ScalarIndex)
 					{
 						MUTABLE_CPUPROFILER_SCOPE(GetScalar)
 
-						const FName ScalarName = MutableInstance->GetScalarName(InstanceComponentIndex, BaseLODIndex, BaseSurfaceIndex, ScalarIndex);
-						const float ScalarValue = MutableInstance->GetScalar(InstanceComponentIndex, BaseLODIndex, BaseSurfaceIndex, ScalarIndex);
+						const FName ScalarName = MutableInstance->GetScalarName(InstanceComponentIndex, MutableLODIndex, InstanceSurfaceIndex, ScalarIndex);
+						const float ScalarValue = MutableInstance->GetScalar(InstanceComponentIndex, MutableLODIndex, InstanceSurfaceIndex, ScalarIndex);
 
-						const int32 MaterialLayer = MutableInstance->GetScalarLayerIndex(InstanceComponentIndex, BaseLODIndex, BaseSurfaceIndex, ScalarIndex);
+						const int32 MaterialLayer = MutableInstance->GetScalarLayerIndex(InstanceComponentIndex, MutableLODIndex, InstanceSurfaceIndex, ScalarIndex);
 
 						OperationData->InstanceUpdateData.Scalars.Push({ ScalarName, ScalarValue, MaterialLayer });
 					}
 
 					// Images
 					Surface->FirstImage = OperationData->InstanceUpdateData.Images.Num();
-					Surface->ImageCount = MutableInstance->GetImageCount(InstanceComponentIndex, BaseLODIndex, BaseSurfaceIndex);
+					Surface->ImageCount = MutableInstance->GetImageCount(InstanceComponentIndex, MutableLODIndex, InstanceSurfaceIndex);
 					for (int32 ImageIndex = 0; ImageIndex < Surface->ImageCount; ++ImageIndex)
 					{
 						MUTABLE_CPUPROFILER_SCOPE(GetImageId);
 						
 						TSharedRef<FInstanceUpdateData::FImage>& Image = OperationData->InstanceUpdateData.Images.Add_GetRef(MakeShared<FInstanceUpdateData::FImage>());
-						Image->Name = MutableInstance->GetImageName(InstanceComponentIndex, BaseLODIndex, BaseSurfaceIndex, ImageIndex);
-						Image->ImageID = MutableInstance->GetImageId(InstanceComponentIndex, BaseLODIndex, BaseSurfaceIndex, ImageIndex);
-						Image->MaterialLayer = MutableInstance->GetImageLayerIndex(InstanceComponentIndex, BaseLODIndex, BaseSurfaceIndex, ImageIndex);
+						Image->Name = MutableInstance->GetImageName(InstanceComponentIndex, MutableLODIndex, InstanceSurfaceIndex, ImageIndex);
+						Image->ImageID = MutableInstance->GetImageId(InstanceComponentIndex, MutableLODIndex, InstanceSurfaceIndex, ImageIndex);
+						Image->MaterialLayer = MutableInstance->GetImageLayerIndex(InstanceComponentIndex, MutableLODIndex, InstanceSurfaceIndex, ImageIndex);
 						Image->FullImageSizeX = 0;
 						Image->FullImageSizeY = 0;
-						Image->BaseLOD = BaseLODIndex;
+						Image->BaseLOD = MutableLODIndex;
 						Image->BaseMip = 0;
 
 						FString KeyName = Image->Name.ToString();
@@ -2355,7 +2343,7 @@ namespace impl
 							{
 								Image->bIsPassThrough = true;
 
-								UE::Tasks::TTask<TSharedPtr<const UE::Mutable::Private::FImage>> GetImageTask = OperationData->MutableSystem->GetImage(LastTask, OperationData->InstanceID, Image->ImageID, 0, 0);
+								UE::Tasks::TTask<TSharedPtr<const UE::Mutable::Private::FImage>> GetImageTask = OperationData->MutableSystem->GetImage(LastTask, OperationData->InstanceID, Image->ImageID, 0);
 								UE::Tasks::AddNested(GetImageTask);
 								LastTask = GetImageTask;
 					
@@ -2621,7 +2609,7 @@ namespace impl
 			}
 			else
 			{
-				GetImageTask = OperationData->MutableSystem->GetImage(LastTask, OperationData->InstanceID, Image->ImageID, Image->BaseMip + MipsToSkip, Image->BaseLOD);
+				GetImageTask = OperationData->MutableSystem->GetImage(LastTask, OperationData->InstanceID, Image->ImageID, Image->BaseMip + MipsToSkip);
 				UE::Tasks::AddNested(GetImageTask);
 				LastTask = GetImageTask;
 			}
diff --git a/Engine/Plugins/Mutable/Source/CustomizableObjectEditor/Private/MuCOE/SMutableCodeViewer.cpp b/Engine/Plugins/Mutable/Source/CustomizableObjectEditor/Private/MuCOE/SMutableCodeViewer.cpp
--- a/Engine/Plugins/Mutable/Source/CustomizableObjectEditor/Private/MuCOE/SMutableCodeViewer.cpp
+++ b/Engine/Plugins/Mutable/Source/CustomizableObjectEditor/Private/MuCOE/SMutableCodeViewer.cpp
@@ -2878,7 +2878,7 @@ void SMutableCodeViewer::Tick(const FGeometry& AllottedGeometry, const double In
 	case UE::Mutable::Private::EDataType::Image:
 	{
 		check(PreviewImageViewer);
-		TSharedPtr<const UE::Mutable::Private::FImage> MutableImage = System->GetPrivate()->BuildImage(ExternalResourceProvider, MutableModel, PreviewParameters.Get(), SelectedOperationAddress, MipsToSkip, 0);
+		TSharedPtr<const UE::Mutable::Private::FImage> MutableImage = System->GetPrivate()->BuildImage(ExternalResourceProvider, MutableModel, PreviewParameters.Get(), SelectedOperationAddress, MipsToSkip);
 		PreviewImageViewer->SetImage(MutableImage, 0);
 		break;
 	}
diff --git a/Engine/Plugins/Mutable/Source/CustomizableObjectEditor/Private/MuCOE/SMutableInstanceViewer.cpp b/Engine/Plugins/Mutable/Source/CustomizableObjectEditor/Private/MuCOE/SMutableInstanceViewer.cpp
--- a/Engine/Plugins/Mutable/Source/CustomizableObjectEditor/Private/MuCOE/SMutableInstanceViewer.cpp
+++ b/Engine/Plugins/Mutable/Source/CustomizableObjectEditor/Private/MuCOE/SMutableInstanceViewer.cpp
@@ -25,7 +25,6 @@ class USkeleton;
 namespace MutableSurfaceListColumns
 {
 	static const FName IdColumnID("Id");
-	static const FName SharedIdID("SharedId");
 	static const FName CustomColumnID("CustomId");
 	static const FName ImageCountID("Images");
 	static const FName VectorCountID("Vectors");
@@ -61,15 +60,6 @@ public:
 			];
 		}
 
-		if (InColumnName == MutableSurfaceListColumns::SharedIdID)
-		{
-			return SNew(SHorizontalBox) +
-				SHorizontalBox::Slot()
-			[
-				SNew(STextBlock).Text(RowItem->SharedId)
-			];
-		}
-
 		if (InColumnName == MutableSurfaceListColumns::CustomColumnID)
 		{
 			return SNew(SHorizontalBox) + SHorizontalBox::Slot()
@@ -353,7 +343,6 @@ void SMutableInstanceViewer::SetInstance(const TSharedPtr<const UE::Mutable::Pri
 						TSharedPtr<FMutableInstanceViewerSurfaceElement> NewSurface = MakeShareable(new FMutableInstanceViewerSurfaceElement());
 
 						NewSurface->Id = FText::AsNumber(MutableInstance->GetSurfaceId(ComponentIndex, LODIndex, SurfaceIndex));
-						NewSurface->SharedId = FText::AsNumber(MutableInstance->GetSharedSurfaceId(ComponentIndex, LODIndex, SurfaceIndex));
 						NewSurface->CustomId = FText::AsNumber(MutableInstance->GetSurfaceCustomId(ComponentIndex, LODIndex, SurfaceIndex));
 						NewSurface->ImageCount = FText::AsNumber(MutableInstance->GetImageCount(ComponentIndex, LODIndex, SurfaceIndex));
 						NewSurface->VectorCount = FText::AsNumber(MutableInstance->GetVectorCount(ComponentIndex, LODIndex, SurfaceIndex));
@@ -539,10 +528,6 @@ TSharedRef<SWidget> SMutableInstanceViewer::GenerateSurfaceListView(const TShare
 			.DefaultLabel(LOCTEXT("IdTitle", "Id"))
 			.FillWidth(0.14f)
 
-			+ SHeaderRow::Column(MutableSurfaceListColumns::SharedIdID)
-			.DefaultLabel(LOCTEXT("SharedIdTitle", "SharedId"))
-			.FillWidth(0.35f)
-
 			+ SHeaderRow::Column(MutableSurfaceListColumns::CustomColumnID)
 			.DefaultLabel(LOCTEXT("CustomIdTitle", "CustomId"))
 			.FillWidth(0.65f)
diff --git a/Engine/Plugins/Mutable/Source/CustomizableObjectEditor/Private/MuCOE/SMutableInstanceViewer.h b/Engine/Plugins/Mutable/Source/CustomizableObjectEditor/Private/MuCOE/SMutableInstanceViewer.h
--- a/Engine/Plugins/Mutable/Source/CustomizableObjectEditor/Private/MuCOE/SMutableInstanceViewer.h
+++ b/Engine/Plugins/Mutable/Source/CustomizableObjectEditor/Private/MuCOE/SMutableInstanceViewer.h
@@ -23,7 +23,6 @@ class SWidget;
 struct FMutableInstanceViewerSurfaceElement
 {
 	FText Id;
-	FText SharedId;
 	FText CustomId;
 	FText ImageCount;
 	FText VectorCount;
diff --git a/Engine/Plugins/Mutable/Source/MutableRuntime/Internal/MuR/Instance.h b/Engine/Plugins/Mutable/Source/MutableRuntime/Internal/MuR/Instance.h
--- a/Engine/Plugins/Mutable/Source/MutableRuntime/Internal/MuR/Instance.h
+++ b/Engine/Plugins/Mutable/Source/MutableRuntime/Internal/MuR/Instance.h
@@ -143,20 +143,6 @@ namespace UE::Mutable::Private
 		//! \param id ID of the surface to look for.
         UE_API int32 FindSurfaceById( int32 ComponentIndex, int32 LODIndex, uint32 id ) const;
 
-		//! Find the base surface index and Lod index when reusing surfaces between LODs. Return the surface index
-		//! and the LOD it belongs to.
-		//! \param ComponentIndex - Index of the component, from 0 to GetComponentCount()-1
-		//! \param SharedSurfaceId - Id of the surface to look for (as returned by GetSharedSurfaceId).
-		//! \param OutSurfaceIndex - Index of the surface in the OutLODIndex lod. 
-		//! \param OutLODIndex - Index of the first LOD where the surface can be found. 
-		UE_API void FindBaseSurfaceBySharedId(int32 ComponentIndex, uint32 SharedId, int32& OutSurfaceIndex, int32& OutLODIndex) const;
-
-		//! Get an id that can be used to find the same surface on other LODs
-		//! \param ComponentIndex Index of the component, from 0 to GetComponentCount()-1
-		//! \param lod Index of the level of detail, from 0 to GetLODCount()-1
-		//! \param SurfaceIndex Index of the surface, from 0 to GetSurfaceCount()-1
-		UE_API uint32 GetSharedSurfaceId(int32 ComponentIndex, int32 LODIndex, int32 SurfaceIndex) const;
-
         //! Get an optional, opaque application-defined identifier for this surface. The meaning of
         //! this ID depends on each application, and it is specified when creating the source data
         //! that generates this surface.
diff --git a/Engine/Plugins/Mutable/Source/MutableRuntime/Internal/MuR/Operations.h b/Engine/Plugins/Mutable/Source/MutableRuntime/Internal/MuR/Operations.h
--- a/Engine/Plugins/Mutable/Source/MutableRuntime/Internal/MuR/Operations.h
+++ b/Engine/Plugins/Mutable/Source/MutableRuntime/Internal/MuR/Operations.h
@@ -1108,9 +1108,6 @@ namespace UE::Mutable::Private
             uint32 id;
             uint32 ExternalId;
 
-			// Id used to identify shared surfaces between lods.
-			int32 SharedSurfaceId;
-
             ADDRESS name;
 			int32 LayerIndex; // Will be removed in UE-351597.
         };
diff --git a/Engine/Plugins/Mutable/Source/MutableRuntime/Internal/MuR/System.h b/Engine/Plugins/Mutable/Source/MutableRuntime/Internal/MuR/System.h
--- a/Engine/Plugins/Mutable/Source/MutableRuntime/Internal/MuR/System.h
+++ b/Engine/Plugins/Mutable/Source/MutableRuntime/Internal/MuR/System.h
@@ -181,7 +181,7 @@ namespace UE::Mutable::Private
 		//! \param MipsToSkip Number of mips to skip compared from the full image.
 		//! If 0, all mip levels will be generated. If more levels than possible to discard are specified, 
 		//! the image will still contain a minimum number of mips specified at model compile time.
-		UE_API UE::Tasks::TTask<TSharedPtr<const FImage>> GetImage(UE::Tasks::FTask Prerequisite, FInstance::FID InstanceID, const FImageId& ImageId, int32 MipsToSkip = 0, int32 LOD = 0);
+		UE_API UE::Tasks::TTask<TSharedPtr<const FImage>> GetImage(UE::Tasks::FTask Prerequisite, FInstance::FID InstanceID, const FImageId& ImageId, int32 MipsToSkip = 0);
 
 		//! Only valid between BeginUpdate and EndUpdate
 		UE_API UE::Tasks::TTask<TSharedPtr<const FMesh>> GetMesh(UE::Tasks::FTask Prerequisite, FInstance::FID InstanceID, const FMeshId& MeshId, EMeshContentFlags MeshContentFilter = EMeshContentFlags::AllFlags);
@@ -196,7 +196,7 @@ namespace UE::Mutable::Private
 		//! \param MipsToSkip Number of mips to skip compared from the full image.
 		//! If 0, all mip levels will be generated. If more levels than possible to discard are specified, 
 		//! the image will still contain a minimum number of mips specified at model compile time.
-		UE_API TSharedPtr<const FImage> GetImageInline(FInstance::FID InstanceID, const FImageId& ImageId, int32 MipsToSkip = 0, int32 LOD = 0);
+		UE_API TSharedPtr<const FImage> GetImageInline(FInstance::FID InstanceID, const FImageId& ImageId, int32 MipsToSkip = 0);
 
         //! Only valid between BeginUpdate and EndUpdate
 		UE_API TSharedPtr<const FMesh> GetMeshInline(FInstance::FID InstanceID, const FMeshId& MeshId, EMeshContentFlags MeshContentFilter = EMeshContentFlags::AllFlags);
diff --git a/Engine/Plugins/Mutable/Source/MutableRuntime/Internal/MuR/SystemPrivate.h b/Engine/Plugins/Mutable/Source/MutableRuntime/Internal/MuR/SystemPrivate.h
--- a/Engine/Plugins/Mutable/Source/MutableRuntime/Internal/MuR/SystemPrivate.h
+++ b/Engine/Plugins/Mutable/Source/MutableRuntime/Internal/MuR/SystemPrivate.h
@@ -2137,7 +2137,7 @@ namespace UE::Mutable::Private
 		MUTABLERUNTIME_API float BuildScalar(const TSharedPtr<FExternalResourceProvider>& ExternalResourceProvider, const TSharedPtr<const FModel>&, const FParameters*, OP::ADDRESS) ;
 		MUTABLERUNTIME_API FVector4f BuildColour(const TSharedPtr<FExternalResourceProvider>& ExternalResourceProvider, const TSharedPtr<const FModel>&, const FParameters*, OP::ADDRESS) ;
 		MUTABLERUNTIME_API TSharedPtr<const String> BuildString(const TSharedPtr<FExternalResourceProvider>& ExternalResourceProvider, const TSharedPtr<const FModel>&, const FParameters*, OP::ADDRESS) ;
-		MUTABLERUNTIME_API TSharedPtr<const FImage> BuildImage(const TSharedPtr<FExternalResourceProvider>& ExternalResourceProvider, const TSharedPtr<const FModel>&, const FParameters*, OP::ADDRESS, int32 MipsToSkip, int32 LOD);
+		MUTABLERUNTIME_API TSharedPtr<const FImage> BuildImage(const TSharedPtr<FExternalResourceProvider>& ExternalResourceProvider, const TSharedPtr<const FModel>&, const FParameters*, OP::ADDRESS, int32 MipsToSkip);
 		MUTABLERUNTIME_API TSharedPtr<const FMesh> BuildMesh(const TSharedPtr<FExternalResourceProvider>& ExternalResourceProvider, const TSharedPtr<const FModel>&, const FParameters*, OP::ADDRESS, EMeshContentFlags MeshContentFilter);
 		MUTABLERUNTIME_API TSharedPtr<const FInstance> BuildInstance(const TSharedPtr<FExternalResourceProvider>& ExternalResourceProvider, const TSharedPtr<const FModel>&, const FParameters*, OP::ADDRESS);
 		MUTABLERUNTIME_API TSharedPtr<const FLayout> BuildLayout(const TSharedPtr<FExternalResourceProvider>& ExternalResourceProvider, const TSharedPtr<const FModel>&, const FParameters*, OP::ADDRESS) ;
@@ -2176,8 +2176,7 @@ namespace UE::Mutable::Private
 			const FParameters*,
 			OP::ADDRESS,
 			uint32 LODs = FSystem::AllLODs,
-			uint8 ExecutionOptions = 0,
-			int32 LOD = 0);
+			uint8 ExecutionOptions = 0);
 
 		//!
 		void PrepareCache(const FModel*, int32 State);
diff --git a/Engine/Plugins/Mutable/Source/MutableRuntime/Private/MuR/CodeRunner.cpp b/Engine/Plugins/Mutable/Source/MutableRuntime/Private/MuR/CodeRunner.cpp
--- a/Engine/Plugins/Mutable/Source/MutableRuntime/Private/MuR/CodeRunner.cpp
+++ b/Engine/Plugins/Mutable/Source/MutableRuntime/Private/MuR/CodeRunner.cpp
@@ -159,13 +159,12 @@ namespace UE::Mutable::Private
 		OP::ADDRESS At,
 		uint32 InLODMask,
 		uint8 ExecutionOptions,
-		int32 InImageLOD,
 		bool bForceLazy,
 		FScheduledOp::EType Type,
 		const TSharedPtr<FExternalResourceProvider>& ExternalResourceProvider)
 	{
 		return MakeShared<CodeRunner>(FPrivateToken {},
-				InSettings, InSystem, InExecutionStrategy, InModel, InParams, At, InLODMask, ExecutionOptions, InImageLOD, bForceLazy, Type, ExternalResourceProvider);
+				InSettings, InSystem, InExecutionStrategy, InModel, InParams, At, InLODMask, ExecutionOptions, bForceLazy, Type, ExternalResourceProvider);
 	}
 
 
@@ -178,7 +177,6 @@ namespace UE::Mutable::Private
 		OP::ADDRESS At,
 		uint32 InLodMask,
 		uint8 ExecutionOptions,
-		int32 InImageLOD,
 		bool bForceLazy,
 		FScheduledOp::EType Type,
 		const TSharedPtr<FExternalResourceProvider>& ExternalResourceProvider)
@@ -207,8 +205,6 @@ namespace UE::Mutable::Private
 			GetMemory().IncreaseHitCount(FCacheAddress(At, 0, ExecutionOptions));
 		}
 		
-		ImageLOD = InImageLOD;
-	
 		// Push the root operation
 		FScheduledOp RootOp;
 		RootOp.At = At;
@@ -771,7 +767,6 @@ namespace UE::Mutable::Private
                     // IDs
                     Result->Components[0].LODs[0].Surfaces[sindex].InternalId = Args.id;
                     Result->Components[0].LODs[0].Surfaces[sindex].ExternalId = Args.ExternalId;
-                    Result->Components[0].LODs[0].Surfaces[sindex].SharedId = Args.SharedSurfaceId;
                 }
                 StoreInstance( Item, Result );
                 break;
diff --git a/Engine/Plugins/Mutable/Source/MutableRuntime/Private/MuR/CodeRunner.h b/Engine/Plugins/Mutable/Source/MutableRuntime/Private/MuR/CodeRunner.h
--- a/Engine/Plugins/Mutable/Source/MutableRuntime/Private/MuR/CodeRunner.h
+++ b/Engine/Plugins/Mutable/Source/MutableRuntime/Private/MuR/CodeRunner.h
@@ -51,7 +51,6 @@ namespace  UE::Mutable::Private
 				OP::ADDRESS At,
 				uint32 LodMask,
 				uint8 ExecutionOptions,
-				int32 InImageLOD,
 				bool bForceLazy,
 				FScheduledOp::EType Type,
 				const TSharedPtr<FExternalResourceProvider>& ExternalResourceProvider);
@@ -67,7 +66,6 @@ namespace  UE::Mutable::Private
 			OP::ADDRESS At,
 			uint32 LodMask,
 			uint8 ExecutionOptions,
-			int32 InImageLOD,
 			bool bForceLazy,
 			FScheduledOp::EType Type,
 			const TSharedPtr<FExternalResourceProvider>& ExternalResourceProvider);
@@ -177,9 +175,6 @@ namespace  UE::Mutable::Private
 		TMap<OP::ADDRESS, FExtendedImageDesc> ImageDescResults;
 		TArray<int32> ImageDescConstantImages; 
 
-		/** Only used for correct mip skipping with external images. It is the LOD for which the image is build. */
-		int32 ImageLOD;
-    	
 		void Run(TUniquePtr<FProfileContext>&& ProfileContext, bool bForceInlineExecution);
 		void AbortRun();
 	public:
diff --git a/Engine/Plugins/Mutable/Source/MutableRuntime/Private/MuR/CodeRunnerBegin.cpp b/Engine/Plugins/Mutable/Source/MutableRuntime/Private/MuR/CodeRunnerBegin.cpp
--- a/Engine/Plugins/Mutable/Source/MutableRuntime/Private/MuR/CodeRunnerBegin.cpp
+++ b/Engine/Plugins/Mutable/Source/MutableRuntime/Private/MuR/CodeRunnerBegin.cpp
@@ -1108,7 +1108,6 @@ namespace UE::Mutable::Private
 	                    // IDs
 	                    Instance->Components[0].LODs[0].Surfaces[SurfaceIndex].InternalId = Args.id;
 	                    Instance->Components[0].LODs[0].Surfaces[SurfaceIndex].ExternalId = Args.ExternalId;
-	                    Instance->Components[0].LODs[0].Surfaces[SurfaceIndex].SharedId = Args.SharedSurfaceId;
 		            		
 		                StoreInstance(Item, Instance);
 		                break;
diff --git a/Engine/Plugins/Mutable/Source/MutableRuntime/Private/MuR/CodeRunnerTasked.cpp b/Engine/Plugins/Mutable/Source/MutableRuntime/Private/MuR/CodeRunnerTasked.cpp
--- a/Engine/Plugins/Mutable/Source/MutableRuntime/Private/MuR/CodeRunnerTasked.cpp
+++ b/Engine/Plugins/Mutable/Source/MutableRuntime/Private/MuR/CodeRunnerTasked.cpp
@@ -2586,11 +2586,8 @@ namespace UE::Mutable::Private
 				OP::ImageParameterConvertArgs Args = Program.GetOpArgs<OP::ImageParameterConvertArgs>(Item.At);
 			
 				TSharedPtr<const FImage> Image = LoadImage(FCacheAddress(Args.ImageParameter, Item));
-				
-				check(ImageLOD < TNumericLimits<uint8>::Max() && ImageLOD >= 0);
-				check(ImageLOD + static_cast<int32>(Item.ExecutionOptions) < TNumericLimits<uint8>::Max());
 
-				const uint8 MipmapsToSkip = Item.ExecutionOptions + static_cast<uint8>(ImageLOD);
+				const uint8 MipmapsToSkip = Item.ExecutionOptions;
 
 				FExternalResourceId FullId;
 				FullId.ImageParameter = Image->PassthroughObject.Get();
@@ -2622,10 +2619,7 @@ namespace UE::Mutable::Private
 
 				if (bParameterFound && Texture)
 				{
-					check(ImageLOD < TNumericLimits<uint8>::Max() && ImageLOD >= 0);
-					check(ImageLOD + static_cast<int32>(Item.ExecutionOptions) < TNumericLimits<uint8>::Max());
-
-					const uint8 MipmapsToSkip = Item.ExecutionOptions + static_cast<uint8>(ImageLOD);
+					const uint8 MipmapsToSkip = Item.ExecutionOptions;
 
 					CodeRunner::FExternalResourceId FullId;
 					FullId.ImageParameter = Texture;
@@ -2645,7 +2639,7 @@ namespace UE::Mutable::Private
 			{
 				check(Item.Stage==0);
 
-				const uint8 MipmapsToSkip = Item.ExecutionOptions + static_cast<uint8>(ImageLOD);
+				const uint8 MipmapsToSkip = Item.ExecutionOptions;
 
 				FExternalResourceId FullId;
 				FullId.ReferenceResourceId = Args.ID;
diff --git a/Engine/Plugins/Mutable/Source/MutableRuntime/Private/MuR/Instance.cpp b/Engine/Plugins/Mutable/Source/MutableRuntime/Private/MuR/Instance.cpp
--- a/Engine/Plugins/Mutable/Source/MutableRuntime/Private/MuR/Instance.cpp
+++ b/Engine/Plugins/Mutable/Source/MutableRuntime/Private/MuR/Instance.cpp
@@ -113,48 +113,6 @@ namespace UE::Mutable::Private
     }
 	
 	
-	void FInstance::FindBaseSurfaceBySharedId(int32 CompIndex, uint32 SharedId, int32& OutSurfaceIndex, int32& OutLODIndex) const
-	{
-		if (Components.IsValidIndex(CompIndex))
-		{
-			for (int32 LodIndex = 0; LodIndex < Components[CompIndex].LODs.Num(); LodIndex++)
-			{
-				const FInstanceLOD& LOD = Components[CompIndex].LODs[LodIndex];
-				for (int32 SurfaceIndex = 0; SurfaceIndex < LOD.Surfaces.Num(); ++SurfaceIndex)
-				{
-					if (LOD.Surfaces[SurfaceIndex].SharedId == SharedId)
-					{
-						OutSurfaceIndex = SurfaceIndex;
-						OutLODIndex = LodIndex;
-						return;
-					}
-				}
-			}
-
-		}
-
-		OutSurfaceIndex = INDEX_NONE;
-		OutLODIndex = INDEX_NONE;
-	}
-
-	
-	uint32 FInstance::GetSharedSurfaceId(int32 ComponentIndex, int32 LODIndex, int32 SurfaceIndex) const
-	{
-		if (Components.IsValidIndex(ComponentIndex) &&
-			Components[ComponentIndex].LODs.IsValidIndex(LODIndex) &&
-			Components[ComponentIndex].LODs[LODIndex].Surfaces.IsValidIndex(SurfaceIndex))
-		{
-			return Components[ComponentIndex].LODs[LODIndex].Surfaces[SurfaceIndex].SharedId;
-		}
-		else
-		{
-			check(false);
-		}
-
-		return 0;
-	}
-	
-    
     uint32 FInstance::GetSurfaceCustomId(int32 ComponentIndex, int32 LODIndex, int32 SurfaceIndex ) const
     {
         if (Components.IsValidIndex(ComponentIndex) &&
diff --git a/Engine/Plugins/Mutable/Source/MutableRuntime/Private/MuR/System.cpp b/Engine/Plugins/Mutable/Source/MutableRuntime/Private/MuR/System.cpp
--- a/Engine/Plugins/Mutable/Source/MutableRuntime/Private/MuR/System.cpp
+++ b/Engine/Plugins/Mutable/Source/MutableRuntime/Private/MuR/System.cpp
@@ -400,7 +400,7 @@ namespace UE::Mutable::Private
     }
 	
 
-	TSharedPtr<const FImage> FSystem::GetImageInline(FInstance::FID instanceID, const FImageId& ImageId, int32 MipsToSkip, int32 InImageLOD)
+	TSharedPtr<const FImage> FSystem::GetImageInline(FInstance::FID instanceID, const FImageId& ImageId, int32 MipsToSkip)
 	{
 		LLM_SCOPE_BYNAME(TEXT("MutableRuntime"));
 		MUTABLE_CPUPROFILER_SCOPE(SystemGetImage);
@@ -419,7 +419,7 @@ namespace UE::Mutable::Private
 		m_pD->WorkingMemoryManager.CurrentPassthroughObjectLoader = LiveInstance->PassthroughObjectLoader;
 
 		OP::ADDRESS RootAddress = ImageId.GetKey()->Address;
-		pResult = m_pD->BuildImage(LiveInstance->ExternalResourceProvider, LiveInstance->Model, LiveInstance->OldParameters.Get(), RootAddress, MipsToSkip, InImageLOD);
+		pResult = m_pD->BuildImage(LiveInstance->ExternalResourceProvider, LiveInstance->Model, LiveInstance->OldParameters.Get(), RootAddress, MipsToSkip);
 
 		// We always need to return something valid.
 		if (!pResult)
@@ -433,7 +433,7 @@ namespace UE::Mutable::Private
 		return pResult;
 	}
 
-	UE::Tasks::TTask<TSharedPtr<const FImage>> FSystem::GetImage(UE::Tasks::FTask Prerequisite, FInstance::FID instanceID, const FImageId& ImageId, int32 MipsToSkip, int32 InImageLOD)
+	UE::Tasks::TTask<TSharedPtr<const FImage>> FSystem::GetImage(UE::Tasks::FTask Prerequisite, FInstance::FID instanceID, const FImageId& ImageId, int32 MipsToSkip)
 	{
 		LLM_SCOPE_BYNAME(TEXT("MutableRuntime"));
 		MUTABLE_CPUPROFILER_SCOPE(SystemGetImage);
@@ -476,7 +476,7 @@ namespace UE::Mutable::Private
     	TSharedRef<FContext> Context = MakeShared<FContext>();
 
     	UE::Tasks::FTask RunnerTask = UE::Tasks::Launch(TEXT("FSystem::GetImage"),
-		[Context, SystemPrivate = m_pD, RootAddress, LiveInstance, MipsToSkip, InImageLOD]()
+		[Context, SystemPrivate = m_pD, RootAddress, LiveInstance, MipsToSkip]()
 		{
 			SystemPrivate->WorkingMemoryManager.CurrentInstanceCache = LiveInstance->Cache;
 			SystemPrivate->WorkingMemoryManager.CurrentPassthroughObjectLoader = LiveInstance->PassthroughObjectLoader;
@@ -489,7 +489,7 @@ namespace UE::Mutable::Private
 				LiveInstance->Model,
 				LiveInstance->OldParameters.Get(),
 				RootAddress, FSystem::AllLODs,
-				MipsToSkip, InImageLOD,
+				MipsToSkip,
 				false,
 				FScheduledOp::EType::Full,
 				LiveInstance->ExternalResourceProvider);
@@ -582,7 +582,6 @@ namespace UE::Mutable::Private
 					RootAddress,
 					FSystem::AllLODs,
 					executionOptions,
-					0,
 					false,
 					FScheduledOp::EType::ImageDesc,
 					LiveInstance->ExternalResourceProvider);
@@ -669,7 +668,6 @@ namespace UE::Mutable::Private
 				RootAddress,
 				FSystem::AllLODs,
 				ExecutionOptions,
-				0,
 				false,
 				FScheduledOp::EType::ImageDesc,
 				LiveInstance->ExternalResourceProvider);
@@ -794,7 +792,6 @@ namespace UE::Mutable::Private
 					RootAddress, 
 					FSystem::AllLODs, 
 					ExecutionOptions, 
-					0,
 					false, 
 					FScheduledOp::EType::Full,
 					LiveInstance->ExternalResourceProvider);
@@ -886,7 +883,6 @@ namespace UE::Mutable::Private
 					RootAddress, 
 					FSystem::AllLODs, 
 					static_cast<uint8>(SkeletalMeshContentFilter), 
-					0,
 					false, 
 					FScheduledOp::EType::Full,
 					LiveInstance->ExternalResourceProvider);
@@ -977,7 +973,6 @@ namespace UE::Mutable::Private
 					LiveInstance->OldParameters.Get(), 
 					RootAddress, 
 					FSystem::AllLODs, 
-					0, 
 					0,
 					true, 
 					FScheduledOp::EType::Full,
@@ -1280,8 +1275,7 @@ namespace UE::Mutable::Private
 		const FParameters* InParameters,
 		OP::ADDRESS InCodeRoot,
 		uint32 InLODs,
-		uint8 ExecutionOptions,
-		int32 InImageLOD)
+		uint8 ExecutionOptions)
 	{
 		TSharedRef<CodeRunner> Runner = CodeRunner::Create(
 			Settings,
@@ -1292,7 +1286,6 @@ namespace UE::Mutable::Private
 			InCodeRoot,
 			InLODs,
 			ExecutionOptions,
-			InImageLOD,
 			false,
 			FScheduledOp::EType::Full,
 			ExternalResourceProvider);
@@ -1403,7 +1396,7 @@ namespace UE::Mutable::Private
 	}
 
 	
-	TSharedPtr<const FImage> FSystem::Private::BuildImage(const TSharedPtr<FExternalResourceProvider>& ExternalResourceProvider, const TSharedPtr<const FModel>& pModel, const FParameters* Params, OP::ADDRESS at, int32 MipsToSkip, int32 InImageLOD)
+	TSharedPtr<const FImage> FSystem::Private::BuildImage(const TSharedPtr<FExternalResourceProvider>& ExternalResourceProvider, const TSharedPtr<const FModel>& pModel, const FParameters* Params, OP::ADDRESS at, int32 MipsToSkip)
 	{
 		WorkingMemoryManager.BeginRunnerThread();
 
@@ -1412,7 +1405,7 @@ namespace UE::Mutable::Private
 		UE::Mutable::Private::EOpType OpType = pModel->GetPrivate()->Program.GetOpType(at);
 		if (GetOpDataType(OpType) == EDataType::Image)
 		{
-			RunCode(ExternalResourceProvider, pModel, Params, at, FSystem::AllLODs, uint8(MipsToSkip), InImageLOD);
+			RunCode(ExternalResourceProvider, pModel, Params, at, FSystem::AllLODs, uint8(MipsToSkip));
 			if (!bUnrecoverableError)
 			{
 				Result = WorkingMemoryManager.LoadImage(FCacheAddress(at, 0, MipsToSkip), true);
diff --git a/Engine/Plugins/Mutable/Source/MutableTools/Private/MuT/ASTOpInstanceAdd.cpp b/Engine/Plugins/Mutable/Source/MutableTools/Private/MuT/ASTOpInstanceAdd.cpp
--- a/Engine/Plugins/Mutable/Source/MutableTools/Private/MuT/ASTOpInstanceAdd.cpp
+++ b/Engine/Plugins/Mutable/Source/MutableTools/Private/MuT/ASTOpInstanceAdd.cpp
@@ -38,7 +38,6 @@ namespace UE::Mutable::Private
 				value == other->value &&
 				id == other->id &&
 				ExternalId == other->ExternalId &&
-				SharedSurfaceId == other->SharedSurfaceId &&
 				name == other->name &&
 				LayerIndex == other->LayerIndex;
 		}
@@ -54,7 +53,6 @@ namespace UE::Mutable::Private
 		n->value = mapChild(value.child());
 		n->id = id;
 		n->ExternalId = ExternalId;
-		n->SharedSurfaceId = SharedSurfaceId;
 		n->name = name;
 		n->LayerIndex = LayerIndex;
 		return n;
@@ -111,7 +109,6 @@ namespace UE::Mutable::Private
 			FMemory::Memzero(Args);
 			Args.id = id;
 			Args.ExternalId = ExternalId;
-			Args.SharedSurfaceId = SharedSurfaceId;
 			Args.name = Program.AddConstant(name);
 			Args.LayerIndex = LayerIndex;
 
diff --git a/Engine/Plugins/Mutable/Source/MutableTools/Private/MuT/ASTOpInstanceAdd.h b/Engine/Plugins/Mutable/Source/MutableTools/Private/MuT/ASTOpInstanceAdd.h
--- a/Engine/Plugins/Mutable/Source/MutableTools/Private/MuT/ASTOpInstanceAdd.h
+++ b/Engine/Plugins/Mutable/Source/MutableTools/Private/MuT/ASTOpInstanceAdd.h
@@ -27,7 +27,6 @@ struct FProgram;
 		ASTChild value;
 		uint32_t id = 0;
 		uint32_t ExternalId = 0;
-		int32 SharedSurfaceId = INDEX_NONE;
 		FString name;
 		int32 LayerIndex = -1; // TODO(Max) UE-351597:This will be moved to the materials so meshes won't have this value anymore.
 
diff --git a/Engine/Plugins/Mutable/Source/MutableTools/Private/MuT/CodeGenerator.cpp b/Engine/Plugins/Mutable/Source/MutableTools/Private/MuT/CodeGenerator.cpp
--- a/Engine/Plugins/Mutable/Source/MutableTools/Private/MuT/CodeGenerator.cpp
+++ b/Engine/Plugins/Mutable/Source/MutableTools/Private/MuT/CodeGenerator.cpp
@@ -1652,6 +1652,10 @@ namespace UE::Mutable::Private
 						LastSurfOp = Op;
 					}
 				}
+				else
+				{
+					LastSurfOp = SharedSurfaceResults->LastSurfaceOp;
+				}
 
 				FSurfaceGenerationResult SurfaceResult;
 				SurfaceResult.SurfaceOp = LastSurfOp;
@@ -1663,6 +1667,7 @@ namespace UE::Mutable::Private
 					{
 						FSharedSurfaceResultOptions::FSharedSurfaceResult SharedResult;
 						SharedResult.LayoutOps = MeshResults.LayoutOps;
+						SharedResult.LastSurfaceOp = LastSurfOp;
 
 						UE::TUniqueLock Lock(SharedSurfaceOptions.Mutex);
 
@@ -1768,10 +1773,6 @@ namespace UE::Mutable::Private
 				Ptr<ASTOp> LastCompOp;
 				Ptr<ASTOp> LastMeshOp;
 
-				// This generates a different ID for each surface in the LOD and the component. It can be used to match it to the
-				// mesh surface. It cannot be 0 because it is a special case for the merge operation.
-				int32 surfaceID = 1;
-
 				for (int32 SelectedSurfaceIndex=0; SelectedSurfaceIndex <SurfaceTasks.Num(); ++SelectedSurfaceIndex)
 				{
 					FSurfaceTask& SurfaceTask = SurfaceTasks[SelectedSurfaceIndex];
@@ -1784,9 +1785,11 @@ namespace UE::Mutable::Private
 					SurfaceOp->name = SurfaceData.Node->Name;
 					SurfaceOp->instance = LastCompOp;
 					SurfaceOp->value = SurfaceGenerationResult.SurfaceOp;
-					SurfaceOp->id = surfaceID;
+
+					const uint32 SurfaceId = GetTypeHash(SurfaceGenerationResult.SharedSurfaceGuid);
+					SurfaceOp->id = SurfaceId > 0 ? SurfaceId : 1; // It cannot be 0 because it is a special case for the merge operation.
+
 					SurfaceOp->ExternalId = SurfaceData.Node->ExternalId;
-					SurfaceOp->SharedSurfaceId = GetTypeHash(SurfaceGenerationResult.SharedSurfaceGuid);
 
 					Ptr<ASTOp> SurfaceConditionOp = SurfaceData.FinalCondition;
 
@@ -1809,7 +1812,7 @@ namespace UE::Mutable::Private
 						Ptr<ASTOpMeshMerge> MergeOp = new ASTOpMeshMerge();
 						MergeOp->Base = LastMeshOp;
 						MergeOp->Added = Added;
-						MergeOp->NewSurfaceID = surfaceID;
+						MergeOp->NewSurfaceID = SurfaceOp->id;
 						mergeAd = MergeOp;
 					}
 
@@ -1826,8 +1829,6 @@ namespace UE::Mutable::Private
 					{
 						LastMeshOp = mergeAd;
 					}
-
-					++surfaceID;
 				}
 
 				// Add op to optimize the skinning of the resulting mesh
diff --git a/Engine/Plugins/Mutable/Source/MutableTools/Private/MuT/CodeGenerator.h b/Engine/Plugins/Mutable/Source/MutableTools/Private/MuT/CodeGenerator.h
--- a/Engine/Plugins/Mutable/Source/MutableTools/Private/MuT/CodeGenerator.h
+++ b/Engine/Plugins/Mutable/Source/MutableTools/Private/MuT/CodeGenerator.h
@@ -783,6 +783,8 @@ namespace UE::Mutable::Private
 			struct FSharedSurfaceResult
 			{
 				TArray<Ptr<ASTOp>> LayoutOps;
+				Ptr<ASTOp> LastSurfaceOp;
+				uint32 SurfaceId;
 			};
 
 			struct FSharedSurfaceResultKey
@@ -798,7 +800,7 @@ namespace UE::Mutable::Private
 				FORCEINLINE bool operator==(const FSharedSurfaceResultKey& Other) const = default;
 			};
 
-			/** Store the generated layout operations that we intend to share across LODs.
+			/** Store the generated layout and surface operations that we intend to share across LODs.
 			* The key is the combination of the BaseSurfaceGuid and the combination of all ModifierGuids 
 			* affecting this surface.
 			*/
@@ -1160,7 +1162,7 @@ namespace UE::Mutable::Private
 			const FMeshGenerationStaticOptions&,
 			FMeshOptionsTask DynamicOptionsTask,
 			FMeshTask BaseTask,
-			FGuid SharedSurfaceId,
+			FGuid SurfaceGuid,
 			const void* ErrorContext,
 			const NodeMesh* OriginalMeshNode);
 
diff --git a/Engine/Plugins/Mutable/Source/MutableTools/Private/MuT/CodeOptimiser.cpp b/Engine/Plugins/Mutable/Source/MutableTools/Private/MuT/CodeOptimiser.cpp
--- a/Engine/Plugins/Mutable/Source/MutableTools/Private/MuT/CodeOptimiser.cpp
+++ b/Engine/Plugins/Mutable/Source/MutableTools/Private/MuT/CodeOptimiser.cpp
@@ -440,7 +440,7 @@ namespace UE::Mutable::Private
 			{
 				MUTABLE_CPUPROFILER_SCOPE(ConstantImage);
 
-				TSharedPtr<const FImage> pImage = System->GetPrivate()->BuildImage(nullptr, Model, LocalParams.Get(), at, 0, 0);
+				TSharedPtr<const FImage> pImage = System->GetPrivate()->BuildImage(nullptr, Model, LocalParams.Get(), at, 0);
 
 				if (pImage)
 				{
