[Mutable] Fix crash when adding a generated surface to the list of InstanceUpdateData.Surfaces

#rnx

diff --git a/Engine/Plugins/Mutable/Source/CustomizableObject/Private/MuCO/CustomizableObjectSystem.cpp b/Engine/Plugins/Mutable/Source/CustomizableObject/Private/MuCO/CustomizableObjectSystem.cpp
--- a/Engine/Plugins/Mutable/Source/CustomizableObject/Private/MuCO/CustomizableObjectSystem.cpp
+++ b/Engine/Plugins/Mutable/Source/CustomizableObject/Private/MuCO/CustomizableObjectSystem.cpp
@@ -2264,7 +2264,8 @@ namespace impl
 
 					if (GeneratedSurfaceIndex != INDEX_NONE)
 					{
-						OperationData->InstanceUpdateData.Surfaces.Add(OperationData->InstanceUpdateData.Surfaces[GeneratedSurfaceIndex]);
+						TSharedRef<FInstanceUpdateData::FSurface> GeneratedSurface = OperationData->InstanceUpdateData.Surfaces[GeneratedSurfaceIndex];
+						OperationData->InstanceUpdateData.Surfaces.Add(GeneratedSurface);
 						continue;
 					}
 
