[Mutable] Minor UI fixes.

- Removed Schema node.
- Fixed Reroute pin type.
- Aligned Add Pins.
- Renamed Passthrough Skeletal Mesh pin.

diff --git a/Engine/Plugins/Mutable/Source/CustomizableObjectEditor/Private/MuCOE/EdGraphSchema_CustomizableObject.cpp b/Engine/Plugins/Mutable/Source/CustomizableObjectEditor/Private/MuCOE/EdGraphSchema_CustomizableObject.cpp
--- a/Engine/Plugins/Mutable/Source/CustomizableObjectEditor/Private/MuCOE/EdGraphSchema_CustomizableObject.cpp
+++ b/Engine/Plugins/Mutable/Source/CustomizableObjectEditor/Private/MuCOE/EdGraphSchema_CustomizableObject.cpp
@@ -120,7 +120,6 @@
 #include "MuR/External/VectorAdapter.h"
 #include "Nodes/CONodeExternalOperation.h"
 #include "Nodes/CONodeExternalTypeParameter.h"
-#include "Nodes/CONodeSchema.h"
 #include "Nodes/CONodeSkeletalMeshBreak.h"
 #include "UObject/Class.h"
 
@@ -773,14 +772,6 @@ void UEdGraphSchema_CustomizableObject::GetGraphContextActions(FGraphContextMenu
 		TSharedPtr<FCustomizableObjectSchemaAction_NewNode> Action = AddNewNodeAction(ContextMenuBuilder, FString(), Node->GetNodeTitle(ENodeTitleType::ListView), FText(), 1);
 		Action->NodeTemplate = Node;
 	}
-
-	const UCustomizableObjectEditorSettings* EditorSettings = GetDefault<UCustomizableObjectEditorSettings>();
-	if (EditorSettings->bEnableDeveloperOptions)
-	{
-		UCONodeSchema* Node = NewObject<UCONodeSchema>();
-		TSharedPtr<FCustomizableObjectSchemaAction_NewNode> Action = AddNewNodeAction(ContextMenuBuilder, TEXT("Developer"), Node->GetNodeTitle(ENodeTitleType::ListView), FText(), 1);
-		Action->NodeTemplate = Node;
-	}
 	
 	{
 		UCustomizableObjectNodeReroute* Node = NewObject<UCustomizableObjectNodeReroute>();
diff --git a/Engine/Plugins/Mutable/Source/CustomizableObjectEditor/Private/MuCOE/Nodes/CONodeSchema.cpp b/Engine/Plugins/Mutable/Source/CustomizableObjectEditor/Private/MuCOE/Nodes/CONodeSchema.cpp
deleted file mode 100644
diff --git a/Engine/Plugins/Mutable/Source/CustomizableObjectEditor/Private/MuCOE/Nodes/CONodeSchema.h b/Engine/Plugins/Mutable/Source/CustomizableObjectEditor/Private/MuCOE/Nodes/CONodeSchema.h
deleted file mode 100644
diff --git a/Engine/Plugins/Mutable/Source/CustomizableObjectEditor/Private/MuCOE/Nodes/CustomizableObjectNode.cpp b/Engine/Plugins/Mutable/Source/CustomizableObjectEditor/Private/MuCOE/Nodes/CustomizableObjectNode.cpp
--- a/Engine/Plugins/Mutable/Source/CustomizableObjectEditor/Private/MuCOE/Nodes/CustomizableObjectNode.cpp
+++ b/Engine/Plugins/Mutable/Source/CustomizableObjectEditor/Private/MuCOE/Nodes/CustomizableObjectNode.cpp
@@ -794,13 +794,6 @@ FLinearColor UCustomizableObjectNode::GetPinColor(const UEdGraphPin& Pin) const
 }
 
 
-bool UCustomizableObjectNode::IsPassthrough(const UEdGraphPin& Pin) const
-{
-	const UEdGraphSchema_CustomizableObject* Schema = GetDefault<UEdGraphSchema_CustomizableObject>();
-	return Schema->IsPassthrough(Pin.PinType.PinCategory);
-}
-
-
 bool UCustomizableObjectNode::HasPinViewer() const
 {
 	return false;
diff --git a/Engine/Plugins/Mutable/Source/CustomizableObjectEditor/Private/MuCOE/Nodes/CustomizableObjectNodeComponentPassthroughMesh.cpp b/Engine/Plugins/Mutable/Source/CustomizableObjectEditor/Private/MuCOE/Nodes/CustomizableObjectNodeComponentPassthroughMesh.cpp
--- a/Engine/Plugins/Mutable/Source/CustomizableObjectEditor/Private/MuCOE/Nodes/CustomizableObjectNodeComponentPassthroughMesh.cpp
+++ b/Engine/Plugins/Mutable/Source/CustomizableObjectEditor/Private/MuCOE/Nodes/CustomizableObjectNodeComponentPassthroughMesh.cpp
@@ -44,7 +44,7 @@ void UCustomizableObjectNodeComponentPassthroughMesh::AllocateDefaultPins(UCusto
 {
 	Super::AllocateDefaultPins(RemapPins);
 
-	SkeletalMeshPin = CustomCreatePin(EGPD_Input, UEdGraphSchema_CustomizableObject::PC_PassthroughSkeletalMesh, TEXT("Skeletal Mesh"));
+	SkeletalMeshPin = CustomCreatePin(EGPD_Input, UEdGraphSchema_CustomizableObject::PC_PassthroughSkeletalMesh, TEXT("Passthrough Skeletal Mesh"));
 
 	// Get old pins
 	const TArray<UEdGraphPin*> AllPins = GetAllNonOrphanPins();
diff --git a/Engine/Plugins/Mutable/Source/CustomizableObjectEditor/Private/MuCOE/Nodes/SCustomizableObjectNode.cpp b/Engine/Plugins/Mutable/Source/CustomizableObjectEditor/Private/MuCOE/Nodes/SCustomizableObjectNode.cpp
--- a/Engine/Plugins/Mutable/Source/CustomizableObjectEditor/Private/MuCOE/Nodes/SCustomizableObjectNode.cpp
+++ b/Engine/Plugins/Mutable/Source/CustomizableObjectEditor/Private/MuCOE/Nodes/SCustomizableObjectNode.cpp
@@ -81,7 +81,7 @@ void SCustomizableObjectNode::CreateInputSideAddButton(TSharedPtr<SVerticalBox>
 			OutputBox->AddSlot()
 			.AutoHeight()
 			.VAlign(VAlign_Center)
-			.HAlign(HAlign_Right)
+			.HAlign(HAlign_Left)
 			.Padding(AddPinPadding)
 			[
 				AddPinButton
diff --git a/Engine/Plugins/Mutable/Source/CustomizableObjectEditor/Private/MuCOE/Nodes/SCustomizableObjectNodePin.cpp b/Engine/Plugins/Mutable/Source/CustomizableObjectEditor/Private/MuCOE/Nodes/SCustomizableObjectNodePin.cpp
--- a/Engine/Plugins/Mutable/Source/CustomizableObjectEditor/Private/MuCOE/Nodes/SCustomizableObjectNodePin.cpp
+++ b/Engine/Plugins/Mutable/Source/CustomizableObjectEditor/Private/MuCOE/Nodes/SCustomizableObjectNodePin.cpp
@@ -59,9 +59,9 @@ TSharedRef<SWidget> SCustomizableObjectNodePin::GetDefaultValueWidget()
 			case EEditablePinNameBoxVisibilityPolicy::ALWAYS_VISIBLE:
 				{
 					return SNew(SEditableTextBox)
-								.Text(this, &SCustomizableObjectNodePin::GetNodeStringValue)
-								.OnTextCommitted(this, &SCustomizableObjectNodePin::OnTextCommited)
-								.Visibility(EVisibility::Visible);
+						.Text(this, &SCustomizableObjectNodePin::GetNodeStringValue)
+						.OnTextCommitted(this, &SCustomizableObjectNodePin::OnTextCommited)
+						.Visibility(EVisibility::Visible);
 				}
 			case EEditablePinNameBoxVisibilityPolicy::HIDE_IF_LINKED:
 				{
@@ -80,27 +80,16 @@ TSharedRef<SWidget> SCustomizableObjectNodePin::GetDefaultValueWidget()
 
 const FSlateBrush* SCustomizableObjectNodePin::GetPinIcon() const
 {
-	UEdGraphPin* GraphPin = GetPinObj();
-	if (GraphPin && !GraphPin->IsPendingKill())
+	if (UEdGraphSchema_CustomizableObject::IsPassthrough(GraphPinObj->PinType.PinCategory))
 	{
-		UCustomizableObjectNode* Node = CastChecked<UCustomizableObjectNode>(GetPinObj()->GetOwningNode());
-		if (Node)
-		{
-			if (Node->IsPassthrough(*GraphPin))
-			{
-				if (GraphPinObj->LinkedTo.IsEmpty())
-				{
-					return PassThroughImageDisconnected;
-				}
-				else
-				{
-					return PassThroughImageConnected;
-				}
-			}
-		}
+		return GraphPinObj->LinkedTo.Num() ?
+			PassThroughImageConnected :
+			PassThroughImageDisconnected;
+	}
+	else
+	{
+		return SGraphPin::GetPinIcon();
 	}
-	
-	return SGraphPin::GetPinIcon();
 }
 
 
diff --git a/Engine/Plugins/Mutable/Source/CustomizableObjectEditor/Private/MuCOE/Nodes/SCustomizableObjectNodeReroutePin.cpp b/Engine/Plugins/Mutable/Source/CustomizableObjectEditor/Private/MuCOE/Nodes/SCustomizableObjectNodeReroutePin.cpp
--- a/Engine/Plugins/Mutable/Source/CustomizableObjectEditor/Private/MuCOE/Nodes/SCustomizableObjectNodeReroutePin.cpp
+++ b/Engine/Plugins/Mutable/Source/CustomizableObjectEditor/Private/MuCOE/Nodes/SCustomizableObjectNodeReroutePin.cpp
@@ -18,7 +18,7 @@ void SCustomizableObjectNodeReroutePin::Construct(const FArguments& InArgs, UEdG
 
 const FSlateBrush* SCustomizableObjectNodeReroutePin::GetPinIcon() const
 {
-	if (GraphPinObj->PinType.PinCategory == UEdGraphSchema_CustomizableObject::PC_PassthroughTexture)
+	if (UEdGraphSchema_CustomizableObject::IsPassthrough(GraphPinObj->PinType.PinCategory))
 	{
 		return GraphPinObj->LinkedTo.Num() ?
 			PassThroughImageConnected :
diff --git a/Engine/Plugins/Mutable/Source/CustomizableObjectEditor/Public/MuCOE/Nodes/CustomizableObjectNode.h b/Engine/Plugins/Mutable/Source/CustomizableObjectEditor/Public/MuCOE/Nodes/CustomizableObjectNode.h
--- a/Engine/Plugins/Mutable/Source/CustomizableObjectEditor/Public/MuCOE/Nodes/CustomizableObjectNode.h
+++ b/Engine/Plugins/Mutable/Source/CustomizableObjectEditor/Public/MuCOE/Nodes/CustomizableObjectNode.h
@@ -314,8 +314,6 @@ public:
 	
 	UE_API virtual FLinearColor GetPinColor(const UEdGraphPin &Pin) const;
 	
-	UE_API virtual bool IsPassthrough(const UEdGraphPin &Pin) const;
-	
 	/** @return true if the details should spawn a Pin Viewer. */
 	UE_API virtual bool HasPinViewer() const;
 
