[Mutable] Better Performance Analyzer colors.

diff --git a/Engine/Plugins/Mutable/Source/CustomizableObject/Internal/MuCO/LogBenchmarkUtil.h b/Engine/Plugins/Mutable/Source/CustomizableObject/Internal/MuCO/LogBenchmarkUtil.h
--- a/Engine/Plugins/Mutable/Source/CustomizableObject/Internal/MuCO/LogBenchmarkUtil.h
+++ b/Engine/Plugins/Mutable/Source/CustomizableObject/Internal/MuCO/LogBenchmarkUtil.h
@@ -104,9 +104,10 @@ struct FInstanceUpdateStats
 
 struct FStatistic
 {
-	double SD = 0.0f;
-	double Avg = 0.0f;
+	double Min = +TNumericLimits<double>::Max();
 	double Max = -TNumericLimits<double>::Max();
+	double LowerBound = 0.0f;
+	double UpperBound = 0.0f;
 };
 
 
diff --git a/Engine/Plugins/Mutable/Source/CustomizableObjectEditor/Private/MuCOE/CustomizableObjectEditorPerformanceAnalyzer.cpp b/Engine/Plugins/Mutable/Source/CustomizableObjectEditor/Private/MuCOE/CustomizableObjectEditorPerformanceAnalyzer.cpp
--- a/Engine/Plugins/Mutable/Source/CustomizableObjectEditor/Private/MuCOE/CustomizableObjectEditorPerformanceAnalyzer.cpp
+++ b/Engine/Plugins/Mutable/Source/CustomizableObjectEditor/Private/MuCOE/CustomizableObjectEditorPerformanceAnalyzer.cpp
@@ -77,10 +77,34 @@ void SInstanceUpdateDataRow::OnInstanceNameNavigation() const
 
 FSlateColor TextColor(double Value, const FStatistic& Statistic)
 {
-	const float Normalized = FMath::Clamp((Value - Statistic.Avg) / (Statistic.Max - Statistic.Avg), 0.0, 1.0);
-	const FLinearColor Color = FLinearColor::LerpUsingHSV(FLinearColor::Green, FLinearColor::Red, Normalized);
-	
-	return FSlateColor(Color);
+	if (Value < Statistic.LowerBound)
+	{
+		if (Statistic.Min >= Statistic.LowerBound)
+		{
+			return FLinearColor::Green;
+		}
+		else
+		{
+			const float Normalized = (Value - Statistic.Min) / (Statistic.LowerBound - Statistic.Min);
+			return FLinearColor::LerpUsingHSV(FLinearColor::Green, FLinearColor::Red, 1.0 - Normalized);
+		}
+	}
+	else if (Value > Statistic.UpperBound)
+	{
+		if (Statistic.Max <= Statistic.UpperBound)
+		{
+			return FLinearColor::Green;
+		}
+		else
+		{
+			const float Normalized = (Value - Statistic.UpperBound) / (Statistic.Max - Statistic.UpperBound);
+			return FLinearColor::LerpUsingHSV(FLinearColor::Green, FLinearColor::Red, Normalized);
+		}
+	}
+	else
+	{
+		return FLinearColor::Green;
+	}
 }
 
 
@@ -952,7 +976,7 @@ void SCustomizableObjectEditorPerformanceAnalyzer::ComputeStatistics()
 {
 	const int32 Num = InstanceUpdateElements.Num();
 
-	auto ComputeMax = [this, Num](FStatistic& Statistic, const TFunction<double(int32)>& GetValue)
+	auto Compute = [this, Num](FStatistic& Statistic, const TFunction<double(int32)>& GetValue)
 	{
 		Statistic = {};
 
@@ -961,39 +985,39 @@ void SCustomizableObjectEditorPerformanceAnalyzer::ComputeStatistics()
 			return;
 		}
 			
+		TArray<double> Values;
+		Values.Reserve(Num);
 		for (int32 Index = 0; Index < Num; ++Index)
 		{
-			Statistic.Max = FMath::Max(Statistic.Max, GetValue(Index));
+			Values.Add(GetValue(Index));
 		}
 
-		for (int32 Index = 0; Index < Num; ++Index)
-		{
-			Statistic.Avg += GetValue(Index);
-		}
+		Values.Sort();
 
-		Statistic.Avg /= Num;
-			
-		for (int32 Index = 0; Index < Num; ++Index)
-		{
-			Statistic.SD += FMath::Square(GetValue(Index) - Statistic.Avg);
-		}
+		Statistic.Min = Values[0];
+		Statistic.Max = Values.Last();
 
-		Statistic.SD = FMath::Sqrt(Statistic.SD / Num);
+		const double Q1 = Values[static_cast<int32>(Num * 0.25)];
+		const double Q3 = Values[static_cast<int32>(Num * 0.75)];
+		const double IQR = Q3 - Q1;
+			
+		Statistic.LowerBound = Q1 - 1.5 * IQR;
+		Statistic.UpperBound = Q3 + 1.5 * IQR;
 	};
 	
-	ComputeMax(InstanceUpdateStatistics->TriangleCount, [&](int32 Index){	return InstanceUpdateElements[Index]->UpdateStats.TriangleCount; });
-	ComputeMax(InstanceUpdateStatistics->QueueTime, [&](int32 Index){	return InstanceUpdateElements[Index]->UpdateStats.QueueTime; });
-	ComputeMax(InstanceUpdateStatistics->UpdateTime, [&](int32 Index) { return InstanceUpdateElements[Index]->UpdateStats.UpdateTime; });
-	ComputeMax(InstanceUpdateStatistics->TaskGetMeshTime, [&](int32 Index) { return InstanceUpdateElements[Index]->UpdateStats.TaskGetMeshesTime; });
-	ComputeMax(InstanceUpdateStatistics->TaskLockCacheTime, [&](int32 Index) { return InstanceUpdateElements[Index]->UpdateStats.TaskLockCacheTime; });
-	ComputeMax(InstanceUpdateStatistics->TaskGetImagesTime, [&](int32 Index) { return InstanceUpdateElements[Index]->UpdateStats.TaskGetImagesTime; });
-	ComputeMax(InstanceUpdateStatistics->TaskConvertResourcesTime, [&](int32 Index) { return InstanceUpdateElements[Index]->UpdateStats.TaskConvertResourcesTime; });
-	ComputeMax(InstanceUpdateStatistics->TaskCallbacksTime, [&](int32 Index) { return InstanceUpdateElements[Index]->UpdateStats.TaskCallbacksTime; });
-	ComputeMax(InstanceUpdateStatistics->UpdatePeakMemory, [&](int32 Index) { return InstanceUpdateElements[Index]->UpdateStats.UpdatePeakMemory; });
-	ComputeMax(InstanceUpdateStatistics->UpdateRealPeakMemory, [&](int32 Index) { return InstanceUpdateElements[Index]->UpdateStats.UpdateRealPeakMemory; });
-	ComputeMax(InstanceUpdateStatistics->TaskUpdateImageTime, [&](int32 Index) { return InstanceUpdateElements[Index]->UpdateStats.TaskUpdateImageTime; });
-	ComputeMax(InstanceUpdateStatistics->TaskUpdateImagePeakMemory, [&](int32 Index) { return InstanceUpdateElements[Index]->UpdateStats.TaskUpdateImagePeakMemory; });
-	ComputeMax(InstanceUpdateStatistics->TaskUpdateImageRealPeakMemory, [&](int32 Index) { return InstanceUpdateElements[Index]->UpdateStats.TaskUpdateImageRealPeakMemory; });
+	Compute(InstanceUpdateStatistics->TriangleCount, [&](int32 Index){	return InstanceUpdateElements[Index]->UpdateStats.TriangleCount; });
+	Compute(InstanceUpdateStatistics->QueueTime, [&](int32 Index){	return InstanceUpdateElements[Index]->UpdateStats.QueueTime; });
+	Compute(InstanceUpdateStatistics->UpdateTime, [&](int32 Index) { return InstanceUpdateElements[Index]->UpdateStats.UpdateTime; });
+	Compute(InstanceUpdateStatistics->TaskGetMeshTime, [&](int32 Index) { return InstanceUpdateElements[Index]->UpdateStats.TaskGetMeshesTime; });
+	Compute(InstanceUpdateStatistics->TaskLockCacheTime, [&](int32 Index) { return InstanceUpdateElements[Index]->UpdateStats.TaskLockCacheTime; });
+	Compute(InstanceUpdateStatistics->TaskGetImagesTime, [&](int32 Index) { return InstanceUpdateElements[Index]->UpdateStats.TaskGetImagesTime; });
+	Compute(InstanceUpdateStatistics->TaskConvertResourcesTime, [&](int32 Index) { return InstanceUpdateElements[Index]->UpdateStats.TaskConvertResourcesTime; });
+	Compute(InstanceUpdateStatistics->TaskCallbacksTime, [&](int32 Index) { return InstanceUpdateElements[Index]->UpdateStats.TaskCallbacksTime; });
+	Compute(InstanceUpdateStatistics->UpdatePeakMemory, [&](int32 Index) { return InstanceUpdateElements[Index]->UpdateStats.UpdatePeakMemory; });
+	Compute(InstanceUpdateStatistics->UpdateRealPeakMemory, [&](int32 Index) { return InstanceUpdateElements[Index]->UpdateStats.UpdateRealPeakMemory; });
+	Compute(InstanceUpdateStatistics->TaskUpdateImageTime, [&](int32 Index) { return InstanceUpdateElements[Index]->UpdateStats.TaskUpdateImageTime; });
+	Compute(InstanceUpdateStatistics->TaskUpdateImagePeakMemory, [&](int32 Index) { return InstanceUpdateElements[Index]->UpdateStats.TaskUpdateImagePeakMemory; });
+	Compute(InstanceUpdateStatistics->TaskUpdateImageRealPeakMemory, [&](int32 Index) { return InstanceUpdateElements[Index]->UpdateStats.TaskUpdateImageRealPeakMemory; });
 }
 
 #undef LOCTEXT_NAMESPACE
