[Mutable] Fix Task times.

diff --git a/Engine/Plugins/Mutable/Source/CustomizableObject/Internal/MuCO/CustomizableObjectSystemPrivate.h b/Engine/Plugins/Mutable/Source/CustomizableObject/Internal/MuCO/CustomizableObjectSystemPrivate.h
--- a/Engine/Plugins/Mutable/Source/CustomizableObject/Internal/MuCO/CustomizableObjectSystemPrivate.h
+++ b/Engine/Plugins/Mutable/Source/CustomizableObject/Internal/MuCO/CustomizableObjectSystemPrivate.h
@@ -710,10 +710,12 @@ public:
 	double StartUpdateTime = 0.0;
 	double UpdateTime = 0.0;
 
+	double TaskGetMeshesTimeStart = 0.0;
 	double TaskGetMeshesTime = 0.0;
 	double TaskLockCacheTime = 0.0;
+	double TaskGetImagesTimeStart = 0.0;
 	double TaskGetImagesTime = 0.0;
-	double TaskConvertResourcesTime = 0.0f;
+	double TaskConvertResourcesTime = 0.0;
 	double TaskCallbacksTime = 0.0;
 
 	// Update Memory stats
diff --git a/Engine/Plugins/Mutable/Source/CustomizableObject/Private/MuCO/CustomizableObjectSystem.cpp b/Engine/Plugins/Mutable/Source/CustomizableObject/Private/MuCO/CustomizableObjectSystem.cpp
--- a/Engine/Plugins/Mutable/Source/CustomizableObject/Private/MuCO/CustomizableObjectSystem.cpp
+++ b/Engine/Plugins/Mutable/Source/CustomizableObject/Private/MuCO/CustomizableObjectSystem.cpp
@@ -1949,6 +1949,8 @@ namespace impl
 				}
 			}
 		}
+
+		OperationData->TaskGetMeshesTime = FPlatformTime::Seconds() - OperationData->TaskGetMeshesTimeStart;
 	}
 	
 	
@@ -2373,7 +2375,7 @@ namespace impl
 	{
 		MUTABLE_CPUPROFILER_SCOPE(Task_Mutable_GetMeshes)
 
-		const double StartTime = FPlatformTime::Seconds();
+		OperationData->TaskGetMeshesTimeStart = FPlatformTime::Seconds();
 
 		check(OperationData->Parameters);
 		OperationData->InstanceUpdateData.Clear();
@@ -2498,15 +2500,15 @@ namespace impl
 				UE::Tasks::ETaskPriority::Inherit));
 			}
 		}
-
-		OperationData->TaskGetMeshesTime = FPlatformTime::Seconds() - StartTime;
 	}
 
 	
 	void Task_Mutable_GetImageDescriptors(const TSharedRef<FUpdateContextPrivate>& OperationData)
 	{
 		MUTABLE_CPUPROFILER_SCOPE(Task_Mutable_GetImageDescriptros)
-		
+
+		OperationData->TaskGetImagesTimeStart = FPlatformTime::Seconds();		
+
 		const TSharedPtr<TArray<UE::Mutable::Private::FImageId>> ImagesInThisInstance = MakeShared<TArray<UE::Mutable::Private::FImageId>>();
 
 		UE::Tasks::FTask LastTask;
@@ -2561,9 +2563,7 @@ namespace impl
 	void Task_Mutable_GetImages(const TSharedRef<FUpdateContextPrivate>& OperationData)
 	{
 		MUTABLE_CPUPROFILER_SCOPE(Task_Mutable_GetImages)
-
-		const double StartTime = FPlatformTime::Seconds();		
-
+		
 		const TSharedPtr<TArray<UE::Mutable::Private::FImageId>> ImagesInThisInstance = MakeShared<TArray<UE::Mutable::Private::FImageId>>();
 
 		UE::Tasks::FTask LastTask;
@@ -2663,8 +2663,6 @@ namespace impl
 			GetImageTask,
 			UE::Tasks::ETaskPriority::Inherit));
 		}
-
-		OperationData->TaskGetImagesTime = FPlatformTime::Seconds() - StartTime;
 	}
 
 
@@ -2697,6 +2695,8 @@ namespace impl
 				}
 			}
 		}
+
+		OperationData->TaskGetImagesTime = FPlatformTime::Seconds() - OperationData->TaskGetImagesTimeStart;
 	}
 	
 
