[Mutable] Now, if a macro contains a parameter and that macro is instantiated multiple times, all the parameters will be the same.

#rnx

diff --git a/Engine/Plugins/Mutable/Source/CustomizableObjectEditor/Private/MuCOE/GenerateMutableSource/GenerateMutableSource.cpp b/Engine/Plugins/Mutable/Source/CustomizableObjectEditor/Private/MuCOE/GenerateMutableSource/GenerateMutableSource.cpp
--- a/Engine/Plugins/Mutable/Source/CustomizableObjectEditor/Private/MuCOE/GenerateMutableSource/GenerateMutableSource.cpp
+++ b/Engine/Plugins/Mutable/Source/CustomizableObjectEditor/Private/MuCOE/GenerateMutableSource/GenerateMutableSource.cpp
@@ -26,6 +26,7 @@
 #include "MuCOE/Nodes/CustomizableObjectNodeStaticString.h"
 #include "MuCOE/Nodes/CustomizableObjectNodeTable.h"
 #include "MuCOE/Nodes/CustomizableObjectNodeTexture.h"
+#include "MuCOE/Nodes/CustomizableObjectNodeParameter.h"
 #include "MuCOE/Nodes/CustomizableObjectNodePassThroughTexture.h"
 #include "MuCOE/Nodes/CustomizableObjectNodeTextureBinarise.h"
 #include "MuCOE/Nodes/CustomizableObjectNodeTextureColourMap.h"
@@ -392,7 +393,10 @@ FGeneratedKey::FGeneratedKey(void* InFunctionAddress, const UEdGraphPin& InPin,
 	FunctionAddress = InFunctionAddress;
 	Pin = &InPin;
 	LOD = Node.IsAffectedByLOD() ? GenerationContext.CurrentLOD : 0;
-	MacroContext = GenerationContext.MacroNodesStack;
+
+	// Do not use the Stack of macros as Key if the Node is a Parameter Node. 
+	// By doing so, we can define parameters inside mutabel macros and reuse them without having repeated nodes.
+	MacroContext = Node.IsA(UCustomizableObjectNodeParameter::StaticClass()) ? TArray<const UCustomizableObjectNodeMacroInstance*>{} : GenerationContext.MacroNodesStack;
 
 	if (UseMesh)
 	{
diff --git a/Engine/Plugins/Mutable/Source/CustomizableObjectEditor/Private/MuCOE/Nodes/CustomizableObjectNode.cpp b/Engine/Plugins/Mutable/Source/CustomizableObjectEditor/Private/MuCOE/Nodes/CustomizableObjectNode.cpp
--- a/Engine/Plugins/Mutable/Source/CustomizableObjectEditor/Private/MuCOE/Nodes/CustomizableObjectNode.cpp
+++ b/Engine/Plugins/Mutable/Source/CustomizableObjectEditor/Private/MuCOE/Nodes/CustomizableObjectNode.cpp
@@ -875,4 +875,20 @@ TArray<FName> UCustomizableObjectNode::GetAllowedPinViewerCreationTypes() const
 }
 
 
+bool UCustomizableObjectNode::IsCompatibleWithGraph(const UEdGraph* TargetGraph) const
+{
+	if (const UCustomizableObjectGraph* CurrentGraph = Cast<UCustomizableObjectGraph>(TargetGraph))
+	{
+		if (CurrentGraph->IsMacro())
+		{
+			return IsNodeSupportedInMacros();
+		}
+
+		return true;
+	}
+
+	return false;
+}
+
+
 #undef LOCTEXT_NAMESPACE
diff --git a/Engine/Plugins/Mutable/Source/CustomizableObjectEditor/Private/MuCOE/Nodes/CustomizableObjectNodeExternalPin.cpp b/Engine/Plugins/Mutable/Source/CustomizableObjectEditor/Private/MuCOE/Nodes/CustomizableObjectNodeExternalPin.cpp
--- a/Engine/Plugins/Mutable/Source/CustomizableObjectEditor/Private/MuCOE/Nodes/CustomizableObjectNodeExternalPin.cpp
+++ b/Engine/Plugins/Mutable/Source/CustomizableObjectEditor/Private/MuCOE/Nodes/CustomizableObjectNodeExternalPin.cpp
@@ -344,11 +344,6 @@ void UCustomizableObjectNodeExternalPin::PostEditChangeProperty(FPropertyChanged
 	}
 }
 
-bool UCustomizableObjectNodeExternalPin::IsNodeSupportedInMacros() const
-{
-	return false;
-}
-
 
 UCustomizableObjectNodeRemapPins* UCustomizableObjectNodeExternalPin::CreateRemapPinsDefault() const
 {
diff --git a/Engine/Plugins/Mutable/Source/CustomizableObjectEditor/Private/MuCOE/Nodes/CustomizableObjectNodeExternalPin.h b/Engine/Plugins/Mutable/Source/CustomizableObjectEditor/Private/MuCOE/Nodes/CustomizableObjectNodeExternalPin.h
--- a/Engine/Plugins/Mutable/Source/CustomizableObjectEditor/Private/MuCOE/Nodes/CustomizableObjectNodeExternalPin.h
+++ b/Engine/Plugins/Mutable/Source/CustomizableObjectEditor/Private/MuCOE/Nodes/CustomizableObjectNodeExternalPin.h
@@ -45,7 +45,6 @@ public:
 	UE_API virtual void UpdateReferencedNodeId(const FGuid& NewGuid) override;
 	UE_API virtual void PreEditChange(FProperty* PropertyAboutToChange) override;
 	UE_API virtual void PostEditChangeProperty(FPropertyChangedEvent& PropertyChangedEvent) override;
-	UE_API virtual bool IsNodeSupportedInMacros() const override;
 	
 	// Own interface
 	
diff --git a/Engine/Plugins/Mutable/Source/CustomizableObjectEditor/Public/MuCOE/Nodes/CustomizableObjectNode.h b/Engine/Plugins/Mutable/Source/CustomizableObjectEditor/Public/MuCOE/Nodes/CustomizableObjectNode.h
--- a/Engine/Plugins/Mutable/Source/CustomizableObjectEditor/Public/MuCOE/Nodes/CustomizableObjectNode.h
+++ b/Engine/Plugins/Mutable/Source/CustomizableObjectEditor/Public/MuCOE/Nodes/CustomizableObjectNode.h
@@ -359,6 +359,8 @@ public:
 
 	// Returns an array with all the PC_Types that can be created from the pin viewer
 	UE_API virtual TArray<FName> GetAllowedPinViewerCreationTypes() const;
+
+	UE_API virtual bool IsCompatibleWithGraph(UEdGraph const* Graph) const override;
 	
 	UPROPERTY()
 	FPostEditChangePropertyDelegate PostEditChangePropertyDelegate;
