[Mutable] Remove Surface dependency on UMutableTextureMipDataProviderFactory.

diff --git a/Engine/Plugins/Mutable/Source/CustomizableObject/Internal/MuCO/CustomizableObjectSystemPrivate.h b/Engine/Plugins/Mutable/Source/CustomizableObject/Internal/MuCO/CustomizableObjectSystemPrivate.h
--- a/Engine/Plugins/Mutable/Source/CustomizableObject/Internal/MuCO/CustomizableObjectSystemPrivate.h
+++ b/Engine/Plugins/Mutable/Source/CustomizableObject/Internal/MuCO/CustomizableObjectSystemPrivate.h
@@ -56,15 +56,6 @@ struct FMutableImageReference
 {
 	UE::Mutable::Private::FImageId ImageID;
 
-	uint32 SurfaceId = 0;
-
-	uint8 LOD = 0;
-	
-	/** Object component index. */
-	uint8 Component = 0;
-
-	uint8 Image = 0;
-
 	uint8 BaseMip = 0;
 
 	TArray<int32> ConstantImagesNeededToGenerate;
diff --git a/Engine/Plugins/Mutable/Source/CustomizableObject/Private/MuCO/CustomizableObjectInstance.cpp b/Engine/Plugins/Mutable/Source/CustomizableObject/Private/MuCO/CustomizableObjectInstance.cpp
--- a/Engine/Plugins/Mutable/Source/CustomizableObject/Private/MuCO/CustomizableObjectInstance.cpp
+++ b/Engine/Plugins/Mutable/Source/CustomizableObject/Private/MuCO/CustomizableObjectInstance.cpp
@@ -6150,10 +6150,6 @@ void UCustomizableInstancePrivate::BuildMaterials(const TSharedRef<FUpdateContex
 													MutableMipDataProviderFactory->CustomizableObjectInstance = Public;
 													check(LODIndex < 256 && InstanceComponentIndex < 256 && ImageIndex < 256);
 													MutableMipDataProviderFactory->ImageRef.ImageID = Image->ImageID;
-													MutableMipDataProviderFactory->ImageRef.SurfaceId = Surface->SurfaceId;
-													MutableMipDataProviderFactory->ImageRef.LOD = uint8(Image->BaseLOD);
-													MutableMipDataProviderFactory->ImageRef.Component = uint8(InstanceComponentIndex);
-													MutableMipDataProviderFactory->ImageRef.Image = uint8(ImageIndex);
 													MutableMipDataProviderFactory->ImageRef.BaseMip = uint8(Image->BaseMip);
 													MutableMipDataProviderFactory->ImageRef.ConstantImagesNeededToGenerate = Image->ConstantImagesNeededToGenerate;
 													MutableMipDataProviderFactory->UpdateContext = UpdateContext;
diff --git a/Engine/Plugins/Mutable/Source/CustomizableObject/Private/MuCO/CustomizableObjectMipDataProvider.cpp b/Engine/Plugins/Mutable/Source/CustomizableObject/Private/MuCO/CustomizableObjectMipDataProvider.cpp
--- a/Engine/Plugins/Mutable/Source/CustomizableObject/Private/MuCO/CustomizableObjectMipDataProvider.cpp
+++ b/Engine/Plugins/Mutable/Source/CustomizableObject/Private/MuCO/CustomizableObjectMipDataProvider.cpp
@@ -154,15 +154,8 @@ namespace Impl
 		}
 
 		const FMutableImageReference& ImageRef = OperationData->RequestedImage;
-
-		int32 SurfaceIndex = Instance->FindSurfaceById(ImageRef.Component, ImageRef.LOD, ImageRef.SurfaceId);
-		check(SurfaceIndex >= 0);
-
-		// This ID may be different than the ID obtained the first time the image was generated, because the mutable
-		// runtime cannot remember all the resources it has built, and only remembers a fixed amount.
-		UE::Mutable::Private::FImageId MipImageID = Instance->GetImageId(ImageRef.Component, ImageRef.LOD, SurfaceIndex, ImageRef.Image);
-
-		UE::Tasks::TTask<TSharedPtr<const UE::Mutable::Private::FImage>> GetImageTask = System->GetImage({}, InstanceID, MipImageID, ImageRef.BaseMip + OperationData->MipsToSkip);
+		
+		UE::Tasks::TTask<TSharedPtr<const UE::Mutable::Private::FImage>> GetImageTask = System->GetImage({}, InstanceID, ImageRef.ImageID, ImageRef.BaseMip + OperationData->MipsToSkip);
 
 		UE::Tasks::AddNested(UE::Tasks::Launch(TEXT("MipDataProvider_EndUpdateImagesTask"),
 				[System, OperationData, InstanceID, StartTime, GetImageTask, EndUpdateImage]() mutable
