[Mutable] Remove Material references from Generated Textures.

diff --git a/Engine/Plugins/Mutable/Source/CustomizableObjectEditor/Private/MuCOE/SCustomizableObjectEditorTextureAnalyzer.cpp b/Engine/Plugins/Mutable/Source/CustomizableObjectEditor/Private/MuCOE/SCustomizableObjectEditorTextureAnalyzer.cpp
--- a/Engine/Plugins/Mutable/Source/CustomizableObjectEditor/Private/MuCOE/SCustomizableObjectEditorTextureAnalyzer.cpp
+++ b/Engine/Plugins/Mutable/Source/CustomizableObjectEditor/Private/MuCOE/SCustomizableObjectEditorTextureAnalyzer.cpp
@@ -162,7 +162,6 @@ void SCustomizableObjecEditorTextureAnalyzer::FillTextureAnalyzerTable(UCustomiz
 				//Texture Info
 				Entry->Texture = Texture;
 				Entry->TextureName = Texture->GetName();
-				Entry->TextureParameterName = GeneratedTexture.Name;
 				Entry->ResolutionX = Texture->GetSizeX();
 				Entry->ResolutionY = Texture->GetSizeY();
 				Entry->LODBias = Texture->GetCachedLODBias();
@@ -172,13 +171,6 @@ void SCustomizableObjecEditorTextureAnalyzer::FillTextureAnalyzerTable(UCustomiz
 				Entry->IsStreamed = Texture->IsCurrentlyVirtualTextured() ? "Virtual Streamed" : (!Texture->IsStreamable() ? "Not Streamed" : "Streamed");
 				Entry->LODGroup = Texture->LODGroup;
 
-				//Material Info
-				Entry->Material = Material;
-				Entry->MaterialName = Material->GetName();
-				Entry->ParentMaterial = Material->Parent;
-				Entry->MaterialParameterName = Material->Parent ? Material->Parent->GetName() : FString();
-				Entry->ComponentName = GeneratedMaterial.ComponentName;
-
 				//Total size
 				TotalSize += TextureSize;
 			}
@@ -202,21 +194,10 @@ void SCustomizableObjecEditorTextureAnalyzer::OnTextureTableSelectionChanged(TSh
 	if (UCustomizableObjectEditorTextureStats* Current = Cast<UCustomizableObjectEditorTextureStats>(Cell->GetObject().Get()))
 	{
 		//Open Texture Editor Tab
-		if ((SelectedColumn == "Texture" || SelectedColumn == "Parameter Name") && Current->Texture)
+		if ((SelectedColumn == "Texture") && Current->Texture)
 		{
 			GEditor->GetEditorSubsystem<UAssetEditorSubsystem>()->OpenEditorForAsset(Current->Texture);
 		}
-
-		//Open Material Editor Tab
-		if (SelectedColumn == "Material" && Current->Material)
-		{
-			GEditor->GetEditorSubsystem<UAssetEditorSubsystem>()->OpenEditorForAsset(Current->Material);
-		}
-
-		if (SelectedColumn == "Parent" && Current->ParentMaterial)
-		{
-			GEditor->GetEditorSubsystem<UAssetEditorSubsystem>()->OpenEditorForAsset(Current->ParentMaterial);
-		}
 	}
 
 	//Avoid edit mode (not needed)
diff --git a/Engine/Plugins/Mutable/Source/CustomizableObjectEditor/Private/MuCOE/SCustomizableObjectEditorTextureAnalyzer.h b/Engine/Plugins/Mutable/Source/CustomizableObjectEditor/Private/MuCOE/SCustomizableObjectEditorTextureAnalyzer.h
--- a/Engine/Plugins/Mutable/Source/CustomizableObjectEditor/Private/MuCOE/SCustomizableObjectEditorTextureAnalyzer.h
+++ b/Engine/Plugins/Mutable/Source/CustomizableObjectEditor/Private/MuCOE/SCustomizableObjectEditorTextureAnalyzer.h
@@ -30,38 +30,18 @@ class UCustomizableObjectEditorTextureStats : public UObject
 	GENERATED_BODY()
 
 public:
-
-	/** Texture - double click to open */
-	UPROPERTY(VisibleAnywhere, AssetRegistrySearchable, Category = "Stats", meta = (DisplayName = "Parameter Name", ColumnWidth = "50", NoResetToDefault))
-	FString TextureParameterName;
-
 	/** Texture - double click to open */
 	UPROPERTY(VisibleAnywhere, AssetRegistrySearchable, Category = "Stats", meta = (DisplayName = "Texture", ColumnWidth = "40", NoResetToDefault))
 	FString TextureName;
-
-	/** Material - double click to open */
-	UPROPERTY(VisibleAnywhere, AssetRegistrySearchable, Category = "Stats", meta = (DisplayName = "Material", ColumnWidth = "50", NoResetToDefault))
-	FString MaterialName;
-
-	/** Parent Material - double click to open */
-	UPROPERTY(VisibleAnywhere, AssetRegistrySearchable, Category = "Stats", meta = (DisplayName = "Parent", ColumnWidth = "50", NoResetToDefault))
-	FString MaterialParameterName;
-
+	
 	/** Used to open the texture in the editor*/
 	UPROPERTY(meta = (NoResetToDefault))
 	TObjectPtr<UTexture> Texture;
 
-	/** Used to open the material in the editor*/
-	UPROPERTY(meta = (NoResetToDefault))
-	TObjectPtr<UMaterialInterface> Material;
-
-	/** Used to open the parent material in the editor*/
-	UPROPERTY(meta = (NoResetToDefault))
-	TObjectPtr<UMaterialInterface> ParentMaterial;
-
 	/** Resolution of the texture */
 	UPROPERTY(VisibleAnywhere, AssetRegistrySearchable, Category = "Stats", meta = (DisplayName = "Resolution X", ColumnWidth = "40", DisplayRight = "true", NoResetToDefault))
 	int32 ResolutionX;
+
 	UPROPERTY(VisibleAnywhere, AssetRegistrySearchable, Category = "Stats", meta = (DisplayName = "Resolution Y", ColumnWidth = "40", DisplayRight = "true", NoResetToDefault))
 	int32 ResolutionY;
 
@@ -84,11 +64,6 @@ public:
 	/** The Level of detail group of the texture */
 	UPROPERTY(VisibleAnywhere, AssetRegistrySearchable, Category = "Stats", meta = (DisplayName = "Texture Group", ColumnWidth = "70", NoResetToDefault))
 	TEnumAsByte<TextureGroup> LODGroup;
-
-	/** The Component of the texture */
-	UPROPERTY(VisibleAnywhere, AssetRegistrySearchable, Category = "Stats", meta = (DisplayName = "Component", ColumnWidth = "70", NoResetToDefault))
-	FName ComponentName;
-
 };
 
 
