[mutable] Updated the mutable core FMaterial data handling
- The user is now required to provide a non-null FMaterial when calling StoreMaterial.
- LoadMaterial will now always return a non-null ptr to FMaterial.
- We do this as the first step before making the core of mutable work with references instead of pointers as that would simplify the code.

#rnx

diff --git a/Engine/Plugins/Mutable/Source/MutableRuntime/Internal/MuR/SystemPrivate.h b/Engine/Plugins/Mutable/Source/MutableRuntime/Internal/MuR/SystemPrivate.h
--- a/Engine/Plugins/Mutable/Source/MutableRuntime/Internal/MuR/SystemPrivate.h
+++ b/Engine/Plugins/Mutable/Source/MutableRuntime/Internal/MuR/SystemPrivate.h
@@ -754,12 +754,26 @@ namespace UE::Mutable::Private
 
 		TSharedPtr<const FMaterial> GetMaterial(FCacheAddress at)
 		{
-			if (!at.At) return {};
+			if (!at.At)
+			{
+				return MakeShared<FMaterial>();
+			}
 			const FOpExecutionData* Data = OpExecutionData.get_ptr(at);
-			if (!Data) return {};
+			if (!Data)
+			{
+				return MakeShared<FMaterial>();
+			}
 
 			check(Data->DataType == EDataType::Material);
-			return MaterialResults[Data->DataTypeIndex];
+			TSharedPtr<const FMaterial> FoundMaterial = MaterialResults[Data->DataTypeIndex];
+			check(FoundMaterial);
+
+			if (!FoundMaterial)
+			{
+				return MakeShared<FMaterial>();
+			}
+			
+			return FoundMaterial;
 		}
 
 		FProjector GetProjector(FCacheAddress at)
@@ -1022,6 +1036,7 @@ namespace UE::Mutable::Private
 
 		void SetMaterial(FCacheAddress at, TSharedPtr<const FMaterial> Value)
 		{
+			check(Value);		// All material pointers provided must be valid
 			check(at.At < OpExecutionData.size_code());
 			FOpExecutionData* Data = OpExecutionData.get_ptr(at);
 			check(Data->DataType == EDataType::Material || Data->DataType == EDataType::None);
diff --git a/Engine/Plugins/Mutable/Source/MutableRuntime/Private/MuR/CodeRunner.cpp b/Engine/Plugins/Mutable/Source/MutableRuntime/Private/MuR/CodeRunner.cpp
--- a/Engine/Plugins/Mutable/Source/MutableRuntime/Private/MuR/CodeRunner.cpp
+++ b/Engine/Plugins/Mutable/Source/MutableRuntime/Private/MuR/CodeRunner.cpp
@@ -414,7 +414,7 @@ namespace UE::Mutable::Private
 				case EDataType::Layout:			StoreLayout(Item, nullptr); break;
 				case EDataType::Instance:		StoreInstance(Item, nullptr); break;
 				case EDataType::ExtensionData:	StoreExtensionData(Item, MakeShared<FExtensionData>()); break;
-				case EDataType::Material:		StoreMaterial(Item, {}); break;
+				case EDataType::Material:		StoreMaterial(Item, MakeShared<FMaterial>()); break;
 				default:
 					// Not implemented
 					check(false);
@@ -5576,12 +5576,6 @@ namespace UE::Mutable::Private
 			{
 				TSharedPtr<const FMaterial> Material = LoadMaterial(FCacheAddress(Args.Material, Item));
 
-				if (!Material)
-				{
-					StoreImage(Item, nullptr);
-					break;
-				}
-
 				check(Args.ParameterName < (uint32)Program.ConstantStrings.Num());
 				const FString& ParameterName = Program.ConstantStrings[Args.ParameterName];
 				const TVariant<OP::ADDRESS, TSharedPtr<const FImage>>* Image = Material->ImageParameters.Find(FName(ParameterName));
@@ -6083,13 +6077,7 @@ namespace UE::Mutable::Private
 			case 1:
 			{
 				TSharedPtr<const FMaterial> Material = LoadMaterial(FCacheAddress(Args.Material, Item));
-
-				if (!Material)
-				{
-					StoreScalar(Item, 0.0f);
-					break;
-				}
-
+					
 				check(Args.ParameterName < (uint32)Program.ConstantStrings.Num());
 				const FString& ParameterName = Program.ConstantStrings[Args.ParameterName];
 				const float* Value = Material->ScalarParameters.Find(FName(ParameterName));
@@ -6405,12 +6393,6 @@ namespace UE::Mutable::Private
 			{
 				TSharedPtr<const FMaterial> Material = LoadMaterial(FCacheAddress(Args.Material, Item));
 
-				if (!Material)
-				{
-					StoreColor(Item, FVector4f::Zero());
-					break;
-				}
-
 				check(Args.ParameterName < (uint32)Program.ConstantStrings.Num());
 				const FString& ParameterName = Program.ConstantStrings[Args.ParameterName];
 				const FVector4f* Value = Material->ColorParameters.Find(FName(ParameterName));
@@ -7057,9 +7039,12 @@ namespace UE::Mutable::Private
 							break;
 							
 						case EDataType::Material:
-							StoreMaterial(FCacheAddress(Item.At, Item), Context.Output.Get<FMaterialAdapter>().Material);
+						{
+							const TSharedPtr<FMaterial> OutputMaterial = Context.Output.Get<FMaterialAdapter>().Material;
+							StoreMaterial(FCacheAddress(Item.At, Item), OutputMaterial ? OutputMaterial :  MakeShared<FMaterial>());
 							break;
-							
+						}
+						
 						case EDataType::InstancedStruct:
 							StoreInstancedStruct(FCacheAddress(Item.At, Item), Context.Output.Ptr);
 							break;
@@ -8276,38 +8261,30 @@ namespace UE::Mutable::Private
 				CacheMaterialOp.Type = FScheduledOp::EType::Full;
 
 				TSharedPtr<const FMaterial> Material = LoadMaterial(CacheMaterialOp);
+					
+				check(Args.ParameterName < (uint32)Program.ConstantStrings.Num());
+				const FString& ParameterName = Program.ConstantStrings[Args.ParameterName];
+				const TVariant<OP::ADDRESS, TSharedPtr<const FImage>>* Image = Material->ImageParameters.Find(FName(ParameterName));
 
-				if (!Material)
+				if (!Image)
 				{
 					ImageDescResults.FindOrAdd(Item.At);
 					StoreValidDesc(Item);
 				}
 				else
 				{
-					check(Args.ParameterName < (uint32)Program.ConstantStrings.Num());
-					const FString& ParameterName = Program.ConstantStrings[Args.ParameterName];
-					const TVariant<OP::ADDRESS, TSharedPtr<const FImage>>* Image = Material->ImageParameters.Find(FName(ParameterName));
-
-					if (!Image)
+					if (Image->IsType<OP::ADDRESS>())
 					{
-						ImageDescResults.FindOrAdd(Item.At);
-						StoreValidDesc(Item);
+						FScheduledOp ItemNextStage(Item.At, Item, 2);
+						ItemNextStage.CustomState = Image->Get<OP::ADDRESS>();
+
+						AddOp(ItemNextStage,
+							FScheduledOp(Image->Get<OP::ADDRESS>(), Item));
 					}
 					else
 					{
-						if (Image->IsType<OP::ADDRESS>())
-						{
-							FScheduledOp ItemNextStage(Item.At, Item, 2);
-							ItemNextStage.CustomState = Image->Get<OP::ADDRESS>();
-
-							AddOp(ItemNextStage,
-								FScheduledOp(Image->Get<OP::ADDRESS>(), Item));
-						}
-						else
-						{
-							ImageDescResults.FindOrAdd(Item.At);
-							StoreValidDesc(Item);
-						}
+						ImageDescResults.FindOrAdd(Item.At);
+						StoreValidDesc(Item);
 					}
 				}
 
@@ -8320,19 +8297,12 @@ namespace UE::Mutable::Private
 
 				TSharedPtr<const FMaterial> Material = LoadMaterial(CacheMaterialOp);
 
-				if (!Material)
-				{
-					ImageDescResults.FindOrAdd(Item.At);
-				}
-				else
-				{
-					check(Args.ParameterName < (uint32)Program.ConstantStrings.Num());
-					const FString& ParameterName = Program.ConstantStrings[Args.ParameterName];
-					const TVariant<OP::ADDRESS, TSharedPtr<const FImage>>* Image = Material->ImageParameters.Find(FName(ParameterName));
+				check(Args.ParameterName < (uint32)Program.ConstantStrings.Num());
+				const FString& ParameterName = Program.ConstantStrings[Args.ParameterName];
+				const TVariant<OP::ADDRESS, TSharedPtr<const FImage>>* Image = Material->ImageParameters.Find(FName(ParameterName));
 
-					FExtendedImageDesc& Result = ImageDescResults.FindOrAdd(Item.At);
-					Result = ImageDescResults.FindChecked(Image->Get<OP::ADDRESS>());
-				}
+				FExtendedImageDesc& Result = ImageDescResults.FindOrAdd(Item.At);
+				Result = ImageDescResults.FindChecked(Image->Get<OP::ADDRESS>());
 
 				StoreValidDesc(Item);
 				break;
