[Mutable] Fix MI_CONSTANT not being deduplicated.

diff --git a/Engine/Plugins/Mutable/Source/CustomizableObject/Internal/MuCO/CustomizableObjectPrivate.h b/Engine/Plugins/Mutable/Source/CustomizableObject/Internal/MuCO/CustomizableObjectPrivate.h
--- a/Engine/Plugins/Mutable/Source/CustomizableObject/Internal/MuCO/CustomizableObjectPrivate.h
+++ b/Engine/Plugins/Mutable/Source/CustomizableObject/Internal/MuCO/CustomizableObjectPrivate.h
@@ -1760,6 +1760,8 @@ enum class ECustomizableObjectVersions : int32
 
 	ReplaceSurfaceIdWithSharedSurfaceId,
 
+	Deduplicate_MI_CONSTANT,
+
 	// -----<new versions can be added above this line>--------
 	LastCustomizableObjectVersion
 };
diff --git a/Engine/Plugins/Mutable/Source/MutableTools/Private/MuT/CodeOptimiser.cpp b/Engine/Plugins/Mutable/Source/MutableTools/Private/MuT/CodeOptimiser.cpp
--- a/Engine/Plugins/Mutable/Source/MutableTools/Private/MuT/CodeOptimiser.cpp
+++ b/Engine/Plugins/Mutable/Source/MutableTools/Private/MuT/CodeOptimiser.cpp
@@ -1,6 +1,7 @@
 // Copyright Epic Games, Inc. All Rights Reserved.
 
 #include "MuT/CodeOptimiser.h"
+#include "MuT/CodeOptimiserPrivate.h"
 
 #include "MuT/ErrorLog.h"
 #include "MuT/AST.h"
@@ -30,273 +31,53 @@
 
 namespace UE::Mutable::Private
 {
-
-	namespace
-	{
-
-		struct FMeshEntry
-		{
-			TSharedPtr<const FMesh> Mesh;
-			Ptr<ASTOpConstantResource> Op;
-
-			bool operator==(const FMeshEntry& o) const
-			{
-				return Mesh == o.Mesh || *Mesh ==*o.Mesh;
-			}
-		};
-
-		struct FImageEntry
-		{
-			TSharedPtr<const FImage> Image;
-			Ptr<ASTOpConstantResource> Op;
-
-			bool operator==(const FImageEntry& o) const
-			{
-				return Image == o.Image || *Image == *o.Image;
-			}
-		};
-
-		struct FLayoutEntry
-		{
-			TSharedPtr< const FLayout> Layout;
-			Ptr<ASTOpConstantResource> Op;
-
-			bool operator==(const FLayoutEntry& o) const
-			{
-				return Layout == o.Layout || *Layout == *o.Layout;
-			}
-		};
-
-		struct custom_mesh_equal
-		{
-			bool operator()( const TSharedPtr<const FMesh>& a, const TSharedPtr<const FMesh>& b ) const
-			{
-				return a==b || *a==*b;
-			}
-		};
-
-		struct custom_image_equal
-		{
-			bool operator()( const TSharedPtr<const FImage>& a, const TSharedPtr<const FImage>& b ) const
-			{
-				return a==b || *a==*b;
-			}
-		};
-
-		struct custom_layout_equal
-		{
-			bool operator()( const TSharedPtr<const FLayout>& a, const TSharedPtr<const FLayout>& b ) const
-			{
-				return a==b || *a==*b;
-			}
-		};
-	}
-
-
-	bool DuplicatedDataRemoverAST( ASTOpList& roots )
+	bool DuplicatedDataRemoverAST(ASTOpList& Roots)
 	{
 		MUTABLE_CPUPROFILER_SCOPE(DuplicatedDataRemoverAST);
 
-		TArray<Ptr<ASTOpConstantResource>> AllMeshOps;
-		TArray<Ptr<ASTOpConstantResource>> AllImageOps;
-		TArray<Ptr<ASTOpConstantResource>> AllLayoutOps;
-
 		bool bModified = false;
 
-		// Gather constants
-		{
-			MUTABLE_CPUPROFILER_SCOPE(Gather);
+		TMap<FResourceEntry<FMesh>, Ptr<ASTOp>> MeshOps;
+		TMap<FResourceEntry<FImage>, Ptr<ASTOp>> ImageOps;
+		TMap<FResourceEntry<FLayout>, Ptr<ASTOp>> LayoutOps;
+		TMap<FResourceEntry<FMaterial>, Ptr<ASTOp>> MaterialOps;
 
-			ASTOp::Traverse_TopRandom_Unique_NonReentrant( roots, [&](Ptr<ASTOp> n)
+		ASTOp::Traverse_TopRandom_Unique_NonReentrant(Roots, [&](Ptr<ASTOp> Op)
+		{
+			switch (Op->GetOpType())
 			{
-				switch ( n->GetOpType() )
-				{
-
-				case EOpType::ME_CONSTANT:
+			case EOpType::ME_CONSTANT:
 				{
-					ASTOpConstantResource* typedNode = static_cast<ASTOpConstantResource*>(n.get());
-					AllMeshOps.Add(typedNode);
-					break;
+					const Ptr<ASTOpConstantResource> OpResource = static_pointer_cast<ASTOpConstantResource>(Op);
+    				bModified |= Deduplicate<FMesh>(MeshOps, OpResource);
+    				break;
 				}
 
-				case EOpType::IM_CONSTANT:
+			case EOpType::IM_CONSTANT:
 				{
-					ASTOpConstantResource* typedNode = static_cast<ASTOpConstantResource*>(n.get());
-					AllImageOps.Add(typedNode);
+					const Ptr<ASTOpConstantResource> OpResource = static_pointer_cast<ASTOpConstantResource>(Op);
+					bModified |= Deduplicate<FImage>(ImageOps, OpResource);
 					break;
 				}
 
-				case EOpType::LA_CONSTANT:
+			case EOpType::LA_CONSTANT:
 				{
-					ASTOpConstantResource* typedNode = static_cast<ASTOpConstantResource*>(n.get());
-					AllLayoutOps.Add(typedNode);
-					break;
-				}
-
-				//    These should be part of the duplicated code removal, in AST.
-				//            // Names
-				//            case EOpType::IN_ADDMESH:
-				//            case EOpType::IN_ADDIMAGE:
-				//            case EOpType::IN_ADDVECTOR:
-				//            case EOpType::IN_ADDSCALAR:
-				//            case EOpType::IN_ADDCOMPONENT:
-				//            case EOpType::IN_ADDSURFACE:
-
-				default:
+					const Ptr<ASTOpConstantResource> OpResource = static_pointer_cast<ASTOpConstantResource>(Op);
+					bModified |= Deduplicate<FLayout>(LayoutOps, OpResource);
 					break;
 				}
 
-				return true;
-			});
-		}
-
-
-		// Compare meshes
-		{
-			MUTABLE_CPUPROFILER_SCOPE(CompareMeshes);
-
-			TMultiMap< SIZE_T, FMeshEntry > Meshes;
-
-			for (Ptr<ASTOpConstantResource>& typedNode : AllMeshOps)
-			{
-				SIZE_T Key = typedNode->GetValueHash();
-
-				Ptr<ASTOp> Found;
-
-				TArray<FMeshEntry*, TInlineAllocator<4>> Candidates;
-				Meshes.MultiFindPointer(Key, Candidates, false);
-
-				if (!Candidates.IsEmpty())
+			case EOpType::MI_CONSTANT:
 				{
-					TSharedPtr<const FMesh> mesh = StaticCastSharedPtr<const FMesh>(typedNode->GetValue());
-
-					for (FMeshEntry* It : Candidates)
-					{
-						if (!It->Mesh)
-						{
-							It->Mesh = StaticCastSharedPtr<const FMesh>(It->Op->GetValue());
-						}
-
-						if (custom_mesh_equal()(mesh, It->Mesh))
-						{
-							Found = It->Op;
-							break;
-						}
-					}
-				}
-
-				if (Found)
-				{
-					ASTOp::Replace(typedNode, Found);
-					bModified = true;
-				}
-				else
-				{
-					// The mesh will be loaded only if it needs to be compared
-					FMeshEntry e;
-					e.Op = typedNode;
-					Meshes.Add(Key, e);
-				}
-			}
-		}
-
-		// Compare images
-		{
-			MUTABLE_CPUPROFILER_SCOPE(CompareImages);
-
-			TMultiMap< SIZE_T, FImageEntry > Images;
-
-			for (Ptr<ASTOpConstantResource>& typedNode : AllImageOps)
-			{
-				SIZE_T Key = typedNode->GetValueHash();
-
-				Ptr<ASTOp> Found;
-
-				TArray<FImageEntry*,TInlineAllocator<4>> Candidates;
-				Images.MultiFindPointer(Key, Candidates, false);
-				
-				if (!Candidates.IsEmpty())
-				{
-					TSharedPtr<const FImage> image = StaticCastSharedPtr<const FImage>(typedNode->GetValue());
-
-					for (FImageEntry* It: Candidates)
-					{
-						if (!It->Image)
-						{
-							It->Image = StaticCastSharedPtr<const FImage>(It->Op->GetValue());
-						}
-
-						if (custom_image_equal()(image, It->Image))
-						{
-							Found = It->Op;
-							break;
-						}
-					}
-				}
-
-				if (Found)
-				{
-					ASTOp::Replace(typedNode, Found);
-					bModified = true;
-				}
-				else
-				{
-					// The image will be loaded only if it needs to be compared
-					FImageEntry e;
-					e.Op = typedNode;
-					Images.Add(Key, e);
-				}
-			}
-		}
-
-		// Compare layouts
-		{
-			MUTABLE_CPUPROFILER_SCOPE(CompareLayouts);
-
-			TMultiMap< SIZE_T, FLayoutEntry > Layouts;
-
-			for (Ptr<ASTOpConstantResource>& typedNode : AllLayoutOps)
-			{
-				SIZE_T Key = typedNode->GetValueHash();
-
-				Ptr<ASTOp> Found;
-
-				TArray<FLayoutEntry*, TInlineAllocator<4>> Candidates;
-				Layouts.MultiFindPointer(Key, Candidates, false);
-
-				if (!Candidates.IsEmpty())
-				{
-					TSharedPtr<const FLayout> layout = StaticCastSharedPtr<const FLayout>(typedNode->GetValue());
-
-					for (FLayoutEntry* It : Candidates)
-					{
-						if (!It->Layout)
-						{
-							It->Layout = StaticCastSharedPtr<const FLayout>(It->Op->GetValue());
-						}
-
-						if (custom_layout_equal()(layout, It->Layout))
-						{
-							Found = It->Op;
-							break;
-						}
-					}
-				}
-
-				if (Found)
-				{
-					ASTOp::Replace(typedNode, Found);
-					bModified = true;
-				}
-				else
-				{
-					FLayoutEntry e;
-					e.Op = typedNode;
-					Layouts.Add(Key, e);
+					const Ptr<ASTOpConstantResource> OpResource = static_pointer_cast<ASTOpConstantResource>(Op);
+					bModified |= Deduplicate<FMaterial>(MaterialOps, OpResource);
+					break;
 				}
 			}
-		}
 
+			return true;
+		});
+		
 		return bModified;
 	}
 
diff --git a/Engine/Plugins/Mutable/Source/MutableTools/Private/MuT/CodeOptimiserPrivate.h b/Engine/Plugins/Mutable/Source/MutableTools/Private/MuT/CodeOptimiserPrivate.h
new file mode 100644
--- /dev/null
+++ b/Engine/Plugins/Mutable/Source/MutableTools/Private/MuT/CodeOptimiserPrivate.h
@@ -0,0 +1,50 @@
+﻿// Copyright Epic Games, Inc. All Rights Reserved.
+
+#pragma once
+
+#include "ASTOpConstantResource.h"
+#include "MuR/Ptr.h"
+#include "Templates/SharedPointer.h"
+#include "MuT/AST.h"
+
+
+namespace UE::Mutable::Private
+{
+	template<typename T>
+	struct FResourceEntry
+	{
+		uint32 Hash = 0;
+		TSharedPtr<const T> Value;
+
+		bool operator==(const FResourceEntry& Entry) const
+		{
+			return Value == Entry.Value || *Value == *Entry.Value;
+		}
+
+		friend uint32 GetTypeHash(const FResourceEntry& Entry)
+		{
+			return Entry.Hash;
+		}
+	};
+
+
+	template<typename T>
+	bool Deduplicate(TMap<FResourceEntry<T>, Ptr<ASTOp>>& Ops, const Ptr<ASTOpConstantResource>& OpResource)
+	{
+		uint32 Hash = OpResource->GetValueHash();
+		TSharedPtr<const T> Value = StaticCastSharedPtr<const T>(OpResource->GetValue());
+
+		FResourceEntry<T> Entry(Hash, Value);
+		
+		if (Ptr<ASTOp>* Found = Ops.Find(Entry))
+		{
+			ASTOp::Replace(OpResource, *Found);
+			return true;
+		}
+		else
+		{
+			Ops.Add(Entry, OpResource);
+			return false;
+		}
+	}	
+}
