[mutable] Made the MergeSkeletons array only contain valid skeletons.
- It only makes sense for the MergeSkeleton method to merge loaded (not null) skeletons, so with these changes we ensure we only use valid skeletons at the time of merging.
- Added a log to report that a skeleton could not be loaded.

#rnx

diff --git a/Engine/Plugins/Mutable/Source/CustomizableObject/Internal/MuCO/CustomizableObjectSystemPrivate.h b/Engine/Plugins/Mutable/Source/CustomizableObject/Internal/MuCO/CustomizableObjectSystemPrivate.h
--- a/Engine/Plugins/Mutable/Source/CustomizableObject/Internal/MuCO/CustomizableObjectSystemPrivate.h
+++ b/Engine/Plugins/Mutable/Source/CustomizableObject/Internal/MuCO/CustomizableObjectSystemPrivate.h
@@ -644,6 +644,8 @@ public:
 	TSharedPtr<FUnrealMutableResourceProvider> ExternalResourceProvider;
 
 	TSharedPtr<FBoneNames> BoneNames;
+
+	bool bIsUpdateAborted = false;
 	
 	// Number of possible components in the entire CO
 	// TODO: Redundant because it is in the CO
diff --git a/Engine/Plugins/Mutable/Source/CustomizableObject/Private/MuCO/CustomizableObjectInstance.cpp b/Engine/Plugins/Mutable/Source/CustomizableObject/Private/MuCO/CustomizableObjectInstance.cpp
--- a/Engine/Plugins/Mutable/Source/CustomizableObject/Private/MuCO/CustomizableObjectInstance.cpp
+++ b/Engine/Plugins/Mutable/Source/CustomizableObject/Private/MuCO/CustomizableObjectInstance.cpp
@@ -1247,14 +1247,7 @@ USkeleton* MergeSkeletons(UCustomizableObject& CustomizableObject, FCustomizable
 		return MergedSkeleton;
 	}
 
-	// Fail in the case a null skeleton is found. This may be caused by a loading or decryption issue
-	if (ReferencedSkeletons.SkeletonsToMerge.Contains(nullptr))
-	{
-		UE_LOG( LogMutable, Warning, TEXT("Unable to merge the provided skeletons as one or multiple of them are null."));
-		ReferencedSkeletons.SkeletonIds.Empty();
-		ReferencedSkeletons.SkeletonsToMerge.Empty();
-		return nullptr;
-	}
+	check(!ReferencedSkeletons.SkeletonsToMerge.Contains(nullptr));
 
 	// No need to merge skeletons
 	if (ReferencedSkeletons.SkeletonsToMerge.Num() == 1)
@@ -1663,9 +1656,9 @@ UPhysicsAsset* GetOrBuildMainPhysicsAsset(
 	TArray<TObjectPtr<UPhysicsAsset>> ValidAssets;
 
 	const int32 NumPhysicsAssets = ComponentData->PhysicsAssets.PhysicsAssetsToMerge.Num();
-	for (int32 I = 0; I < NumPhysicsAssets; ++I)
+	for (int32 PhysicsAssetIndex = 0; PhysicsAssetIndex < NumPhysicsAssets; ++PhysicsAssetIndex)
 	{
-		const TObjectPtr<UPhysicsAsset>& PhysicsAsset = ComponentData->PhysicsAssets.PhysicsAssetsToMerge[I];
+		const TObjectPtr<UPhysicsAsset>& PhysicsAsset = ComponentData->PhysicsAssets.PhysicsAssetsToMerge[PhysicsAssetIndex];
 
 		if (PhysicsAsset)
 		{
@@ -5223,9 +5216,9 @@ UE::Tasks::FTask LoadAdditionalAssetsAndData(const TSharedRef<FUpdateContextPriv
 
 						if (!PhysicsAsset.IsNull())
 						{
-							if (PhysicsAsset.Get())
+							if (UPhysicsAsset* LoadedPhysicsAssetPtr = PhysicsAsset.Get())
 							{
-								ComponentData->PhysicsAssets.PhysicsAssetsToMerge.Add(PhysicsAsset.Get());
+								ComponentData->PhysicsAssets.PhysicsAssetsToMerge.Add(LoadedPhysicsAssetPtr);
 							}
 							else
 							{
@@ -5533,8 +5526,18 @@ void AdditionalAssetsAsyncLoaded(const TSharedRef<FUpdateContextPrivate> Context
 		FReferencedSkeletons& Skeletons = ComponentData.Skeletons;
 		for (int32 SkeletonIndex : Skeletons.SkeletonIds)
 		{
-			const TSoftObjectPtr<USkeleton>& AssetPtr = ModelResources.Skeletons.IsValidIndex(SkeletonIndex) ? ModelResources.Skeletons[SkeletonIndex] : nullptr;
-			Skeletons.SkeletonsToMerge.AddUnique(AssetPtr.Get());
+			check(ModelResources.Skeletons.IsValidIndex(SkeletonIndex));
+			const TSoftObjectPtr<USkeleton>& AssetPtr = ModelResources.Skeletons[SkeletonIndex];
+			
+			if (USkeleton* Skeleton = AssetPtr.Get())
+			{
+				Skeletons.SkeletonsToMerge.AddUnique(Skeleton);
+			}
+			else
+			{
+				UE_LOG(LogMutable, Error, TEXT("Failed to load Skeleton asset: %s"), *AssetPtr.ToSoftObjectPath().ToString());
+				Context->bIsUpdateAborted = true;
+			}
 		}
 
 		// Loaded PhysicsAssets
diff --git a/Engine/Plugins/Mutable/Source/CustomizableObject/Private/MuCO/CustomizableObjectSystem.cpp b/Engine/Plugins/Mutable/Source/CustomizableObject/Private/MuCO/CustomizableObjectSystem.cpp
--- a/Engine/Plugins/Mutable/Source/CustomizableObject/Private/MuCO/CustomizableObjectSystem.cpp
+++ b/Engine/Plugins/Mutable/Source/CustomizableObject/Private/MuCO/CustomizableObjectSystem.cpp
@@ -2793,7 +2793,6 @@ namespace impl
 			return;
 		}
 		
-		UCustomizableObjectSystemPrivate * CustomizableObjectSystemPrivate = System->GetPrivate();
 
 		// Actual work
 		// TODO MTBL-391: Review This hotfix
@@ -2818,6 +2817,13 @@ namespace impl
 			FinishUpdateGlobal(OperationData);
 			return;
 		}
+
+		if (OperationData->bIsUpdateAborted)
+		{
+			OperationData->UpdateResult = EUpdateResult::Error;
+			FinishUpdateGlobal(OperationData);
+			return;
+		}
 		
 		if (CVarEnableRealTimeMorphTargets.GetValueOnAnyThread())
 		{
