[mutable] Fixed parameter discoverability issue when multiple UVs require layout generation.
- Made the discovery of the Morph Factor child branch as always performed when looking for the relevant parameters no matter who is the parent of the ASTOp as before, if the child of an "LA_FROMMESH" would make the Factor Morph child branch not discoverable.
- This allows the ASTOpMeshMorph to work even when multiple UV channels are found in the base mesh.
- Fixed SMutableCodeViewer Me_MORPH tree display. It was missing one of the arguments of the operation, which was causing the entries of the tree to display the contents of each op branch in the wrong place. This also provoked the usage of the wrong previewer for each of the children types.

#rnx

diff --git a/Engine/Plugins/Mutable/Source/CustomizableObjectEditor/Private/MuCOE/SMutableCodeViewer.cpp b/Engine/Plugins/Mutable/Source/CustomizableObjectEditor/Private/MuCOE/SMutableCodeViewer.cpp
--- a/Engine/Plugins/Mutable/Source/CustomizableObjectEditor/Private/MuCOE/SMutableCodeViewer.cpp
+++ b/Engine/Plugins/Mutable/Source/CustomizableObjectEditor/Private/MuCOE/SMutableCodeViewer.cpp
@@ -1635,6 +1635,11 @@ void SMutableCodeViewer::GenerateElementRecursive(const int32& InStateIndex, UE:
 	{
 		const uint8* Data = InProgram.GetOpArgsPointer(InParentAddress);
 
+		UE::Mutable::Private::OP::CONSTANT_STRING_INDEX NameAt = 0;
+		FMemory::Memcpy(&NameAt, Data, sizeof(UE::Mutable::Private::OP::CONSTANT_STRING_INDEX));
+		AddOpFunc(NameAt, TEXT("name "));
+		Data += sizeof(UE::Mutable::Private::OP::CONSTANT_STRING_INDEX);
+			
 		UE::Mutable::Private::OP::ADDRESS FactorAt = 0;
 		FMemory::Memcpy(&FactorAt, Data, sizeof(UE::Mutable::Private::OP::ADDRESS));
 		AddOpFunc(FactorAt, TEXT("factor "));
diff --git a/Engine/Plugins/Mutable/Source/MutableTools/Internal/MuT/AST.h b/Engine/Plugins/Mutable/Source/MutableTools/Internal/MuT/AST.h
--- a/Engine/Plugins/Mutable/Source/MutableTools/Internal/MuT/AST.h
+++ b/Engine/Plugins/Mutable/Source/MutableTools/Internal/MuT/AST.h
@@ -733,7 +733,7 @@ namespace UE::Mutable::Private
 
 
         //! Kind of top-down, but really not.
-        //! This version is slighlty faster, but doesn't support recursive traversals so
+        //! This version is slightly faster, but doesn't support recursive traversals so
         //! use it only in controlled cases.
         template<typename STATE>
         static inline void Traverse_TopDown_Unique_Imprecise_WithState
diff --git a/Engine/Plugins/Mutable/Source/MutableTools/Private/MuT/ASTOpInstanceAdd.cpp b/Engine/Plugins/Mutable/Source/MutableTools/Private/MuT/ASTOpInstanceAdd.cpp
--- a/Engine/Plugins/Mutable/Source/MutableTools/Private/MuT/ASTOpInstanceAdd.cpp
+++ b/Engine/Plugins/Mutable/Source/MutableTools/Private/MuT/ASTOpInstanceAdd.cpp
@@ -112,8 +112,11 @@ namespace UE::Mutable::Private
 			Args.name = Program.AddConstant(name);
 			Args.LayerIndex = LayerIndex;
 
-			if (instance) Args.instance = instance->linkedAddress;
-
+			if (instance)
+			{
+				Args.instance = instance->linkedAddress;
+			}
+				
 			if (value)
 			{
 				Args.value = value->linkedAddress;
@@ -125,23 +128,23 @@ namespace UE::Mutable::Private
 					SubtreeRelevantParametersVisitorAST Visitor;
 					Visitor.Run(value.child());
 
-					TArray<uint16> Params;
+					TArray<uint16> MatchingParameterIndices;
 					for (const FString& ParameterName : Visitor.Parameters)
 					{
-						for (int32 i = 0; i < Program.Parameters.Num(); ++i)
+						for (int32 ParameterIndex = 0; ParameterIndex < Program.Parameters.Num(); ++ParameterIndex)
 						{
-							const auto& param = Program.Parameters[i];
-							if (param.Name == ParameterName)
+							const FParameterDesc& ParameterDescription = Program.Parameters[ParameterIndex];
+							if (ParameterDescription.Name == ParameterName)
 							{
-								Params.Add(uint16(i));
+								MatchingParameterIndices.Add(uint16(ParameterIndex));
 								break;
 							}
 						}
 					}
 
-					Params.Sort();
+					MatchingParameterIndices.Sort();
 
-					int32 It = Program.ParameterLists.Find(Params);
+					const int32 It = Program.ParameterLists.Find(MatchingParameterIndices);
 
 					if (It != INDEX_NONE)
 					{
@@ -150,7 +153,7 @@ namespace UE::Mutable::Private
 					else
 					{
 						Program.RelevantParameterList.Add(value->linkedAddress, Program.ParameterLists.Num());
-						Program.ParameterLists.Add(Params);
+						Program.ParameterLists.Add(MatchingParameterIndices);
 					}
 				}
 			}
diff --git a/Engine/Plugins/Mutable/Source/MutableTools/Private/MuT/CodeOptimiser_Parameters.cpp b/Engine/Plugins/Mutable/Source/MutableTools/Private/MuT/CodeOptimiser_Parameters.cpp
--- a/Engine/Plugins/Mutable/Source/MutableTools/Private/MuT/CodeOptimiser_Parameters.cpp
+++ b/Engine/Plugins/Mutable/Source/MutableTools/Private/MuT/CodeOptimiser_Parameters.cpp
@@ -1550,10 +1550,10 @@ namespace UE::Mutable::Private
 
 
     //---------------------------------------------------------------------------------------------
-    void SubtreeRelevantParametersVisitorAST::Run( Ptr<ASTOp> root )
+    void SubtreeRelevantParametersVisitorAST::Run( Ptr<ASTOp> Root )
     {
         // Cached?
-		TSet<FString>* it = ResultCache.Find( FState(root,false) );
+		TSet<FString>* it = ResultCache.Find( FState(Root,false) );
         if (it)
         {
             Parameters = *it;
@@ -1565,13 +1565,11 @@ namespace UE::Mutable::Private
             MUTABLE_CPUPROFILER_SCOPE(SubtreeRelevantParametersVisitorAST);
 
             Parameters.Empty();
-
+			
             // The state is the onlyLayoutRelevant flag
-            ASTOp::Traverse_TopDown_Unique_Imprecise_WithState<bool>( root, false,
-                [&]( Ptr<ASTOp>& At, bool& state, TArray<TPair<Ptr<ASTOp>,bool>>& Pending )
+            ASTOp::Traverse_TopDown_Unique_Imprecise({Root},
+                [&]( Ptr<ASTOp>& At )
             {
-                (void)state;
-
                 switch ( At->GetOpType() )
                 {
                 case EOpType::NU_PARAMETER:
@@ -1591,46 +1589,6 @@ namespace UE::Mutable::Private
                     return false;
                 }
 
-				case EOpType::LA_FROMMESH:
-				{
-					// Manually choose how to recurse this op
-					const ASTOpLayoutFromMesh* TypedOp = static_cast<const ASTOpLayoutFromMesh*>(At.get());
-
-					// For that mesh we only want to know about the layouts
-					if (const ASTChild& Mesh = TypedOp->Mesh)
-					{
-						Pending.Add({ Mesh.Child, true });
-					}
-
-					return false;
-				}
-
-                case EOpType::ME_MORPH:
-                {
-                    // Manually choose how to recurse this op
-					const ASTOpMeshMorph* TypedOp = static_cast<const ASTOpMeshMorph*>( At.get() );
-
-                    if ( TypedOp->Base )
-                    {
-						Pending.Add({ TypedOp->Base.Child, state });
-                    }
-
-                    // Mesh morphs don't modify the layouts, so we can ignore the factor and morphs
-                    if (!state)
-                    {
-                        if ( TypedOp->Factor )
-                        {
-							Pending.Add({ TypedOp->Factor.Child, state });
-                        }
-
-                        if ( TypedOp->Target )
-                        {
-							Pending.Add({ TypedOp->Target.Child, state });
-                        }
-                    }
-
-                    return false;
-                }
 
                 default:
                     break;
@@ -1639,7 +1597,7 @@ namespace UE::Mutable::Private
                 return true;
             });
 
-            ResultCache.Add( FState(root,false), Parameters );
+            ResultCache.Add( FState(Root,false), Parameters );
         }
     }
 
