[Mutable] Add subimage block decompression and resize utils.

#rnx

diff --git a/Engine/Plugins/Mutable/Source/MutableRuntime/Internal/MuR/BlockCompression/Miro/Miro.h b/Engine/Plugins/Mutable/Source/MutableRuntime/Internal/MuR/BlockCompression/Miro/Miro.h
--- a/Engine/Plugins/Mutable/Source/MutableRuntime/Internal/MuR/BlockCompression/Miro/Miro.h
+++ b/Engine/Plugins/Mutable/Source/MutableRuntime/Internal/MuR/BlockCompression/Miro/Miro.h
@@ -149,6 +149,59 @@ namespace miro
 
 }
 
+namespace miro::SubImageCompression
+{
+    using FuncRefType = TFunctionRef<void(miro::FImageSize, miro::FImageSize, miro::FImageSize, const uint8*, uint8*, int32)>;
+	
+#if MIRO_INCLUDE_BC
+    void RGB_To_BC3SubImage (FImageSize FromSize, FImageSize ToSize, FImageSize SubSize, const uint8* RESTRICT From, uint8* RESTRICT To, int32 Quality);
+    void RGBA_To_BC3SubImage(FImageSize FromSize, FImageSize ToSize, FImageSize SubSize, const uint8* RESTRICT From, uint8* RESTRICT To, int32 Quality);
+
+    void RGB_To_BC1SubImage (FImageSize FromSize, FImageSize ToSize, FImageSize SubSize, const uint8* RESTRICT From, uint8* RESTRICT To, int32 Quality);
+    void RGBA_To_BC1SubImage(FImageSize FromSize, FImageSize ToSize, FImageSize SubSize, const uint8* RESTRICT From, uint8* RESTRICT To, int32 Quality);
+
+    void RGB_To_BC2SubImage (FImageSize FromSize, FImageSize ToSize, FImageSize SubSize, const uint8* RESTRICT From, uint8* RESTRICT To, int32 Quality);
+    void RGBA_To_BC2SubImage(FImageSize FromSize, FImageSize ToSize, FImageSize SubSize, const uint8* RESTRICT From, uint8* RESTRICT To, int32 Quality);
+
+    void L_To_BC4SubImage(FImageSize FromSize, FImageSize ToSize, FImageSize SubSize, const uint8* RESTRICT From, uint8* RESTRICT To, int32 Quality);
+
+    void RGB_To_BC5SubImage (FImageSize FromSize, FImageSize ToSize, FImageSize SubSize, const uint8* RESTRICT From, uint8* RESTRICT To, int32 Quality);
+    void RGBA_To_BC5SubImage(FImageSize FromSize, FImageSize ToSize, FImageSize SubSize, const uint8* RESTRICT From, uint8* RESTRICT To, int32 Quality);
+
+#endif
+
+#if MIRO_INCLUDE_ASTC
+
+    /** ASTC 4x4 support */
+    void RGBA_To_ASTC4x4RGBALSubImage(FImageSize FromSize, FImageSize ToSize, FImageSize SubSize, const uint8* RESTRICT From, uint8* RESTRICT To, int32 Quality);
+    void RGB_To_ASTC4x4RGBALSubImage (FImageSize FromSize, FImageSize ToSize, FImageSize SubSize, const uint8* RESTRICT From, uint8* RESTRICT To, int32 Quality);
+
+    void RGB_To_ASTC4x4RGBLSubImage (FImageSize FromSize, FImageSize ToSize, FImageSize SubSize, const uint8* RESTRICT From, uint8* RESTRICT To, int32 Quality);
+    void RGBA_To_ASTC4x4RGBLSubImage(FImageSize FromSize, FImageSize ToSize, FImageSize SubSize, const uint8* RESTRICT From, uint8* RESTRICT To, int32 Quality);
+	void L_To_ASTC4x4RGBLSubImage   (FImageSize FromSize, FImageSize ToSize, FImageSize SubSize, const uint8* RESTRICT From, uint8* RESTRICT To, int32 Quality);
+    void RGB_To_ASTC4x4RGLSubImage  (FImageSize FromSize, FImageSize ToSize, FImageSize SubSize, const uint8* RESTRICT From, uint8* RESTRICT To, int32 Quality);
+    void RGBA_To_ASTC4x4RGLSubImage (FImageSize FromSize, FImageSize ToSize, FImageSize SubSize, const uint8* RESTRICT From, uint8* RESTRICT To, int32 Quality);
+
+	/** ASTC 6x6 support */
+	void RGBA_To_ASTC6x6RGBALSubImage(FImageSize FromSize, FImageSize ToSize, FImageSize SubSize, const uint8* RESTRICT From, uint8* RESTRICT To, int32 Quality);
+	void RGB_To_ASTC6x6RGBALSubImage (FImageSize FromSize, FImageSize ToSize, FImageSize SubSize, const uint8* RESTRICT From, uint8* RESTRICT To, int32 Quality);
+	void RGB_To_ASTC6x6RGBLSubImage  (FImageSize FromSize, FImageSize ToSize, FImageSize SubSize, const uint8* RESTRICT From, uint8* RESTRICT To, int32 Quality);
+	void RGBA_To_ASTC6x6RGBLSubImage (FImageSize FromSize, FImageSize ToSize, FImageSize SubSize, const uint8* RESTRICT From, uint8* RESTRICT To, int32 Quality);
+	void RGB_To_ASTC6x6RGLSubImage   (FImageSize FromSize, FImageSize ToSize, FImageSize SubSize, const uint8* RESTRICT From, uint8* RESTRICT To, int32 Quality);
+	void RGBA_To_ASTC6x6RGLSubImage  (FImageSize FromSize, FImageSize ToSize, FImageSize SubSize, const uint8* RESTRICT From, uint8* RESTRICT To, int32 Quality);
+
+	/** ASTC 8x8 support */
+	void RGBA_To_ASTC8x8RGBALSubImage(FImageSize FromSize, FImageSize ToSize, FImageSize SubSize, const uint8* RESTRICT From, uint8* RESTRICT To, int32 Quality);
+	void RGB_To_ASTC8x8RGBALSubImage (FImageSize FromSize, FImageSize ToSize, FImageSize SubSize, const uint8* RESTRICT From, uint8* RESTRICT To, int32 Quality);
+	void RGB_To_ASTC8x8RGBLSubImage  (FImageSize FromSize, FImageSize ToSize, FImageSize SubSize, const uint8* RESTRICT From, uint8* RESTRICT To, int32 Quality);
+	void RGBA_To_ASTC8x8RGBLSubImage (FImageSize FromSize, FImageSize ToSize, FImageSize SubSize, const uint8* RESTRICT From, uint8* RESTRICT To, int32 Quality);
+	void RGB_To_ASTC8x8RGLSubImage   (FImageSize FromSize, FImageSize ToSize, FImageSize SubSize, const uint8* RESTRICT From, uint8* RESTRICT To, int32 Quality);
+	void RGBA_To_ASTC8x8RGLSubImage  (FImageSize FromSize, FImageSize ToSize, FImageSize SubSize, const uint8* RESTRICT From, uint8* RESTRICT To, int32 Quality);
+
+#endif
+
+}
+
 
 namespace miro::SubImageDecompression
 {
diff --git a/Engine/Plugins/Mutable/Source/MutableRuntime/Internal/MuR/ImagePrivate.h b/Engine/Plugins/Mutable/Source/MutableRuntime/Internal/MuR/ImagePrivate.h
--- a/Engine/Plugins/Mutable/Source/MutableRuntime/Internal/MuR/ImagePrivate.h
+++ b/Engine/Plugins/Mutable/Source/MutableRuntime/Internal/MuR/ImagePrivate.h
@@ -5,17 +5,8 @@
 #include "MuR/Image.h"
 #include "MuR/ImageTypes.h"
 
-#include "MuR/SerialisationPrivate.h"
-#include "MuR/MutableMath.h"
-
-#include "MuR/BlockCompression/Miro/Miro.h"
-
 namespace UE::Mutable::Private
 {
-
-    //---------------------------------------------------------------------------------------------
-    //!
-    //---------------------------------------------------------------------------------------------
 	inline EImageFormat GetMostGenericFormat(EImageFormat FormatA, EImageFormat FormatB)
     {
 		if (FormatA == FormatB)
@@ -59,10 +50,6 @@ namespace UE::Mutable::Private
         return FormatA;
     }
 
-
-    //---------------------------------------------------------------------------------------------
-    //!
-    //---------------------------------------------------------------------------------------------
 	inline EImageFormat GetRGBOrRGBAFormat(EImageFormat InFormat)
     {
 		InFormat = GetUncompressedFormat(InFormat);
@@ -106,289 +93,5 @@ namespace UE::Mutable::Private
 	{
 		return GetImageFormatData(Format).PixelsPerBlockX > 1 && Format != GetUncompressedFormat(Format);
 	}
-
-	namespace Private
-	{
-		inline void DecompressionFuncNotFoundFunc(miro::FImageSize, miro::FImageSize, miro::FImageSize, const uint8*, uint8*)
-		{
-			check(false);
-		}
-	}
-
-	inline miro::SubImageDecompression::FuncRefType SelectDecompressionFunction(EImageFormat DestFormat, EImageFormat SrcFormat)
-	{
-		using DecompressionFuncRefType = miro::SubImageDecompression::FuncRefType;
-
-		switch(SrcFormat)
-		{
-			case EImageFormat::BC1:
-			{
-				if (DestFormat == EImageFormat::RGBA_UByte)
-				{
-					return DecompressionFuncRefType(miro::SubImageDecompression::BC1_To_RGBASubImage);
-				}
-				else if (DestFormat == EImageFormat::RGB_UByte)
-				{
-					return DecompressionFuncRefType(miro::SubImageDecompression::BC1_To_RGBSubImage);
-				}
-				break;
-			}
-
-			case EImageFormat::BC2:
-			{
-				if (DestFormat == EImageFormat::RGBA_UByte)
-				{
-					return DecompressionFuncRefType(miro::SubImageDecompression::BC2_To_RGBASubImage);
-				}
-				else if (DestFormat == EImageFormat::RGB_UByte)
-				{
-					return DecompressionFuncRefType(miro::SubImageDecompression::BC2_To_RGBSubImage);
-				}
-				break;
-			}
-
-			case EImageFormat::BC3:
-			{
-				if (DestFormat == EImageFormat::RGBA_UByte)
-				{
-					return DecompressionFuncRefType(miro::SubImageDecompression::BC3_To_RGBASubImage);
-				}
-				else if (DestFormat == EImageFormat::RGB_UByte)
-				{
-					return DecompressionFuncRefType(miro::SubImageDecompression::BC3_To_RGBSubImage);
-				}
-				break;
-			}
-
-			case EImageFormat::BC4:
-			{
-				if (DestFormat == EImageFormat::RGBA_UByte)
-				{
-					return DecompressionFuncRefType(miro::SubImageDecompression::BC4_To_RGBASubImage);
-				}
-				else if (DestFormat == EImageFormat::RGB_UByte)
-				{
-					return DecompressionFuncRefType(miro::SubImageDecompression::BC4_To_RGBSubImage);
-				}
-				if (DestFormat == EImageFormat::L_UByte)
-				{
-					return DecompressionFuncRefType(miro::SubImageDecompression::BC4_To_LSubImage);
-				}
-				break;
-			}
-
-			case EImageFormat::BC5:
-			{
-				if (DestFormat == EImageFormat::RGBA_UByte)
-				{
-					return DecompressionFuncRefType(miro::SubImageDecompression::BC5_To_RGBASubImage);
-				}
-				else if (DestFormat == EImageFormat::RGB_UByte)
-				{
-					return DecompressionFuncRefType(miro::SubImageDecompression::BC5_To_RGBSubImage);
-				}
-				break;
-			}
-	
-			case EImageFormat::ASTC_4x4_RGBA_LDR:
-			{
-				if (DestFormat == EImageFormat::RGBA_UByte)
-				{
-					return DecompressionFuncRefType(miro::SubImageDecompression::ASTC4x4RGBAL_To_RGBASubImage);
-				}
-				else if (DestFormat == EImageFormat::RGB_UByte)
-				{
-					return DecompressionFuncRefType(miro::SubImageDecompression::ASTC4x4RGBAL_To_RGBSubImage);
-				}
-				break;
-			}
-
-			case EImageFormat::ASTC_4x4_RGB_LDR:
-			{
-				if (DestFormat == EImageFormat::RGBA_UByte)
-				{
-					return DecompressionFuncRefType(miro::SubImageDecompression::ASTC4x4RGBL_To_RGBASubImage);
-				}
-				else if (DestFormat == EImageFormat::RGB_UByte)
-				{
-					return DecompressionFuncRefType(miro::SubImageDecompression::ASTC4x4RGBL_To_RGBSubImage);
-				}
-				break;
-			}
-
-			case EImageFormat::ASTC_4x4_RG_LDR:
-			{
-				if (DestFormat == EImageFormat::RGBA_UByte)
-				{
-					return DecompressionFuncRefType(miro::SubImageDecompression::ASTC4x4RGL_To_RGBASubImage);
-				}
-				else if (DestFormat == EImageFormat::RGB_UByte)
-				{
-					return DecompressionFuncRefType(miro::SubImageDecompression::ASTC4x4RGL_To_RGBSubImage);
-				}
-				break;
-			}
-
-			case EImageFormat::ASTC_6x6_RGBA_LDR:
-			{
-				if (DestFormat == EImageFormat::RGBA_UByte)
-				{
-					return DecompressionFuncRefType(miro::SubImageDecompression::ASTC6x6RGBAL_To_RGBASubImage);
-				}
-				else if (DestFormat == EImageFormat::RGB_UByte)
-				{
-					return DecompressionFuncRefType(miro::SubImageDecompression::ASTC6x6RGBAL_To_RGBSubImage);
-				}
-				break;
-			}
-
-			case EImageFormat::ASTC_6x6_RGB_LDR:
-			{
-				if (DestFormat == EImageFormat::RGBA_UByte)
-				{
-					return DecompressionFuncRefType(miro::SubImageDecompression::ASTC6x6RGBL_To_RGBASubImage);
-				}
-				else if (DestFormat == EImageFormat::RGB_UByte)
-				{
-					return DecompressionFuncRefType(miro::SubImageDecompression::ASTC6x6RGBL_To_RGBSubImage);
-				}
-				break;
-			}
-
-			case EImageFormat::ASTC_6x6_RG_LDR:
-			{
-				if (DestFormat == EImageFormat::RGBA_UByte)
-				{
-					return DecompressionFuncRefType(miro::SubImageDecompression::ASTC6x6RGL_To_RGBASubImage);
-				}
-				else if (DestFormat == EImageFormat::RGB_UByte)
-				{
-					return DecompressionFuncRefType(miro::SubImageDecompression::ASTC6x6RGL_To_RGBSubImage);
-				}
-				break;
-			}
-
-			case EImageFormat::ASTC_8x8_RGBA_LDR:
-			{
-				if (DestFormat == EImageFormat::RGBA_UByte)
-				{
-					return DecompressionFuncRefType(miro::SubImageDecompression::ASTC8x8RGBAL_To_RGBASubImage);
-				}
-				else if (DestFormat == EImageFormat::RGB_UByte)
-				{
-					return DecompressionFuncRefType(miro::SubImageDecompression::ASTC8x8RGBAL_To_RGBSubImage);
-				}
-				break;
-			}
-
-			case EImageFormat::ASTC_8x8_RGB_LDR:
-			{
-				if (DestFormat == EImageFormat::RGBA_UByte)
-				{
-					return DecompressionFuncRefType(miro::SubImageDecompression::ASTC8x8RGBL_To_RGBASubImage);
-				}
-				else if (DestFormat == EImageFormat::RGB_UByte)
-				{
-					return DecompressionFuncRefType(miro::SubImageDecompression::ASTC8x8RGBL_To_RGBSubImage);
-				}
-				break;
-			}
-
-			case EImageFormat::ASTC_8x8_RG_LDR:
-			{
-				if (DestFormat == EImageFormat::RGBA_UByte)
-				{
-					return DecompressionFuncRefType(miro::SubImageDecompression::ASTC8x8RGL_To_RGBASubImage);
-				}
-				else if (DestFormat == EImageFormat::RGB_UByte)
-				{
-					return DecompressionFuncRefType(miro::SubImageDecompression::ASTC8x8RGL_To_RGBSubImage);
-				}
-				break;
-			}
-
-			case EImageFormat::ASTC_10x10_RGBA_LDR:
-			{
-				if (DestFormat == EImageFormat::RGBA_UByte)
-				{
-					return DecompressionFuncRefType(miro::SubImageDecompression::ASTC10x10RGBAL_To_RGBASubImage);
-				}
-				else if (DestFormat == EImageFormat::RGB_UByte)
-				{
-					return DecompressionFuncRefType(miro::SubImageDecompression::ASTC10x10RGBAL_To_RGBSubImage);
-				}
-				break;
-			}
-
-			case EImageFormat::ASTC_10x10_RGB_LDR:
-			{
-				if (DestFormat == EImageFormat::RGBA_UByte)
-				{
-					return DecompressionFuncRefType(miro::SubImageDecompression::ASTC10x10RGBL_To_RGBASubImage);
-				}
-				else if (DestFormat == EImageFormat::RGB_UByte)
-				{
-					return DecompressionFuncRefType(miro::SubImageDecompression::ASTC10x10RGBL_To_RGBSubImage);
-				}
-				break;
-			}
-
-			case EImageFormat::ASTC_10x10_RG_LDR:
-			{
-				if (DestFormat == EImageFormat::RGBA_UByte)
-				{
-					return DecompressionFuncRefType(miro::SubImageDecompression::ASTC10x10RGL_To_RGBASubImage);
-				}
-				else if (DestFormat == EImageFormat::RGB_UByte)
-				{
-					return DecompressionFuncRefType(miro::SubImageDecompression::ASTC10x10RGL_To_RGBSubImage);
-				}
-				break;
-			}
-
-			case EImageFormat::ASTC_12x12_RGBA_LDR:
-			{
-				if (DestFormat == EImageFormat::RGBA_UByte)
-				{
-					return DecompressionFuncRefType(miro::SubImageDecompression::ASTC12x12RGBAL_To_RGBASubImage);
-				}
-				else if (DestFormat == EImageFormat::RGB_UByte)
-				{
-					return DecompressionFuncRefType(miro::SubImageDecompression::ASTC12x12RGBAL_To_RGBSubImage);
-				}
-				break;
-			}
-
-			case EImageFormat::ASTC_12x12_RGB_LDR:
-			{
-				if (DestFormat == EImageFormat::RGBA_UByte)
-				{
-					return DecompressionFuncRefType(miro::SubImageDecompression::ASTC12x12RGBL_To_RGBASubImage);
-				}
-				else if (DestFormat == EImageFormat::RGB_UByte)
-				{
-					return DecompressionFuncRefType(miro::SubImageDecompression::ASTC12x12RGBL_To_RGBSubImage);
-				}
-				break;
-			}
-
-			case EImageFormat::ASTC_12x12_RG_LDR:
-			{
-				if (DestFormat == EImageFormat::RGBA_UByte)
-				{
-					return DecompressionFuncRefType(miro::SubImageDecompression::ASTC12x12RGL_To_RGBASubImage);
-				}
-				else if (DestFormat == EImageFormat::RGB_UByte)
-				{
-					return DecompressionFuncRefType(miro::SubImageDecompression::ASTC12x12RGL_To_RGBSubImage);
-				}
-				break;
-			}
-		}
-
-		checkf(false, TEXT("Decompression not supported."));
-		return DecompressionFuncRefType(Private::DecompressionFuncNotFoundFunc);
-	}
-
 }
 
diff --git a/Engine/Plugins/Mutable/Source/MutableRuntime/Private/MuR/BlockCompression/Miro/Miro.cpp b/Engine/Plugins/Mutable/Source/MutableRuntime/Private/MuR/BlockCompression/Miro/Miro.cpp
--- a/Engine/Plugins/Mutable/Source/MutableRuntime/Private/MuR/BlockCompression/Miro/Miro.cpp
+++ b/Engine/Plugins/Mutable/Source/MutableRuntime/Private/MuR/BlockCompression/Miro/Miro.cpp
@@ -31,18 +31,16 @@ inline constexpr void miro_check(bool x)
 
 #endif //UE_MIRO_DEBUG
 
-#ifdef _MSC_VER
-#define MIRO_ALIGN __declspec(align(16))
-#else
-#define MIRO_ALIGN __attribute__((aligned(16)))
-#endif
 
-
-union MIRO_ALIGN FMiroPixelBlock4x4
+struct alignas(16) FMiroPixelBlock4x4_4Bytes
 {
-	uint8 bytes[4 * 4 * 4];
+	uint8 Bytes[4 * 4 * 4];
 };
 
+struct alignas(16) FMiroPixelBlock4x4_1Bytes
+{
+	uint8 Bytes[4 * 4];
+};
 
 namespace impl
 {
@@ -7376,7 +7374,7 @@ namespace miro
 				uint8* rowTo = to + 8 * bx * y;
 				for (uint32 x = 0; x < bx; ++x)
 				{
-					FMiroPixelBlock4x4 block;
+					FMiroPixelBlock4x4_4Bytes block;
 
 					int b = 0;
 					for (uint32 j = 0; j < 4; ++j)
@@ -7391,15 +7389,15 @@ namespace miro
 
 							const uint8* pPixelSource = from + 3 * (sx * rj + ri);
 
-							block.bytes[b++] = pPixelSource[0];
-							block.bytes[b++] = pPixelSource[1];
-							block.bytes[b++] = pPixelSource[2];
-							block.bytes[b++] = 255;
+							block.Bytes[b++] = pPixelSource[0];
+							block.Bytes[b++] = pPixelSource[1];
+							block.Bytes[b++] = pPixelSource[2];
+							block.Bytes[b++] = 255;
 						}
 					}
 
 					int32 STBQuality = Quality < 2 ? STB_DXT_NORMAL : STB_DXT_HIGHQUAL;
-					stb::stb_compress_dxt_block(rowTo, block.bytes, 0, STBQuality);
+					stb::stb_compress_dxt_block(rowTo, block.Bytes, 0, STBQuality);
 
 					rowTo += 8;
 				}
@@ -7839,12 +7837,12 @@ namespace miro
 
 				for (uint32 bx = 0; bx < nx; ++bx)
 				{
-					FMiroPixelBlock4x4 block;
+					FMiroPixelBlock4x4_4Bytes block;
 
 					//uint8 const* colourBlock = from + 8;
 					//uint8 const* alphaBlock = from;
 
-					bcdec::bcdec_bc3(rowFrom, block.bytes, 4 * 4);
+					bcdec::bcdec_bc3(rowFrom, block.Bytes, 4 * 4);
 
 					int b = 0;
 					for (int j = 0; j < 4; ++j)
@@ -7855,10 +7853,10 @@ namespace miro
 							{
 								uint8* pPixelDest = to + 4 * (sx * (by * 4 + j) + (bx * 4 + i));
 
-								pPixelDest[0] = block.bytes[b + 0];
-								pPixelDest[1] = block.bytes[b + 1];
-								pPixelDest[2] = block.bytes[b + 2];
-								pPixelDest[3] = block.bytes[b + 3];
+								pPixelDest[0] = block.Bytes[b + 0];
+								pPixelDest[1] = block.Bytes[b + 1];
+								pPixelDest[2] = block.Bytes[b + 2];
+								pPixelDest[3] = block.Bytes[b + 3];
 							}
 
 							b += 4;
@@ -8370,7 +8368,6 @@ namespace miro
 	}
 
 
-	//---------------------------------------------------------------------------------------------
 	template<uint32 BLOCK_SIZE>
 	void Generic_RGBA_to_ASTCRGBL(uint32 sx, uint32 sy, const uint8* from, uint8* to, int32 Quality)
 	{
@@ -9232,6 +9229,857 @@ namespace miro
 #endif // MIRO_INCLUDE_ASTC
 }
 
+namespace miro::SubImageCompression
+{
+
+#if MIRO_INCLUDE_BC
+    void RGB_To_BC1SubImage(FImageSize FromSize, FImageSize ToSize, FImageSize SubSize, const uint8* RESTRICT From, uint8* RESTRICT To, int32 Quality)
+	{
+		if (SubSize.X <= 0 || SubSize.Y <= 0)
+		{
+			return;
+		}
+		
+		constexpr uint32 BlockSize = 4; 
+
+		const uint32 NumBlocksX = FMath::DivideAndRoundUp(uint32(ToSize.X), BlockSize);
+		const uint32 NumSubBlocksX = FMath::DivideAndRoundUp(uint32(SubSize.X), BlockSize);
+		const uint32 NumSubBlocksY = FMath::DivideAndRoundUp(uint32(SubSize.Y), BlockSize);
+
+		const int32 STBQuality = Quality < 2 ? STB_DXT_NORMAL : STB_DXT_HIGHQUAL;
+		for (uint32 BlockY = 0; BlockY < NumSubBlocksY; ++BlockY)
+		{
+			for (uint32 BlockX = 0; BlockX < NumSubBlocksX; ++BlockX)
+			{
+				FMiroPixelBlock4x4_4Bytes Block;
+				static_assert(sizeof(Block) == BlockSize*BlockSize*4);
+
+				for (uint32 BY = 0; BY < BlockSize; ++BY)
+				{
+					for (uint32 BX = 0; BX < BlockSize; ++BX)
+					{
+						uint32 SrcX = FMath::Min(BlockX * BlockSize + BX, uint32(SubSize.X) - 1);
+						uint32 SrcY = FMath::Min(BlockY * BlockSize + BY, uint32(SubSize.Y) - 1);
+
+						const uint8* FromBlockPtr = From + (SrcY * FromSize.X + SrcX) * 3;
+
+						Block.Bytes[BY * BlockSize * 4 + BX * 4 + 0] = FromBlockPtr[0];
+						Block.Bytes[BY * BlockSize * 4 + BX * 4 + 1] = FromBlockPtr[1];
+						Block.Bytes[BY * BlockSize * 4 + BX * 4 + 2] = FromBlockPtr[2];
+						Block.Bytes[BY * BlockSize * 4 + BX * 4 + 3] = 255;
+					}
+				}
+
+				constexpr int32 CompressedBlockSize = 8;
+				uint8* DestBlockPtr = To + (BlockY * NumBlocksX + BlockX) * CompressedBlockSize;
+
+				stb::stb_compress_dxt_block(DestBlockPtr, Block.Bytes, 0, STBQuality);
+			}
+		}
+	}
+
+    void RGBA_To_BC1SubImage(FImageSize FromSize, FImageSize ToSize, FImageSize SubSize, const uint8* RESTRICT From, uint8* RESTRICT To, int32 Quality)
+	{
+		if (SubSize.X <= 0 || SubSize.Y <= 0)
+		{
+			return;
+		}
+		
+		constexpr uint32 BlockSize = 4; 
+
+		const uint32 NumBlocksX = FMath::DivideAndRoundUp(uint32(ToSize.X), BlockSize);
+		const uint32 NumSubBlocksX = FMath::DivideAndRoundUp(uint32(SubSize.X), BlockSize);
+		const uint32 NumSubBlocksY = FMath::DivideAndRoundUp(uint32(SubSize.Y), BlockSize);
+
+		const int32 STBQuality = Quality < 2 ? STB_DXT_NORMAL : STB_DXT_HIGHQUAL;
+		for (uint32 BlockY = 0; BlockY < NumSubBlocksY; ++BlockY)
+		{
+			for (uint32 BlockX = 0; BlockX < NumSubBlocksX; ++BlockX)
+			{
+				FMiroPixelBlock4x4_4Bytes Block;
+				static_assert(sizeof(Block) == BlockSize*BlockSize*4);
+
+				for (uint32 BY = 0; BY < BlockSize; ++BY)
+				{
+					for (uint32 BX = 0; BX < BlockSize; ++BX)
+					{
+						uint32 SrcX = FMath::Min(BlockX * BlockSize + BX, uint32(SubSize.X) - 1);
+						uint32 SrcY = FMath::Min(BlockY * BlockSize + BY, uint32(SubSize.Y) - 1);
+
+						const uint8* FromBlockPtr = From + (SrcY * FromSize.X + SrcX) * 4;
+
+						Block.Bytes[BY * BlockSize * 4 + BX * 4 + 0] = FromBlockPtr[0];
+						Block.Bytes[BY * BlockSize * 4 + BX * 4 + 1] = FromBlockPtr[1];
+						Block.Bytes[BY * BlockSize * 4 + BX * 4 + 2] = FromBlockPtr[2];
+						Block.Bytes[BY * BlockSize * 4 + BX * 4 + 3] = FromBlockPtr[3];
+					}
+				}
+
+				constexpr int32 CompressedBlockSize = 8;
+				uint8* DestBlockPtr = To + (BlockY * NumBlocksX + BlockX) * CompressedBlockSize;
+
+				stb::stb_compress_dxt_block(DestBlockPtr, Block.Bytes, 0, STBQuality);
+			}
+		}
+	}
+
+	void RGB_To_BC2SubImage(FImageSize FromSize, FImageSize ToSize, FImageSize SubSize, const uint8* From, uint8* To, int32 Quality)
+	{
+		if (SubSize.X <= 0 || SubSize.Y <= 0)
+		{
+			return;
+		}
+		
+		const int32 STBQuality = Quality < 2 ? STB_DXT_NORMAL : STB_DXT_HIGHQUAL;
+		
+		constexpr uint32 BlockSize = 4; 
+
+		const uint32 NumBlocksX = FMath::DivideAndRoundUp(uint32(ToSize.X), BlockSize);
+		const uint32 NumSubBlocksX = FMath::DivideAndRoundUp(uint32(SubSize.X), BlockSize);
+		const uint32 NumSubBlocksY = FMath::DivideAndRoundUp(uint32(SubSize.Y), BlockSize);
+
+		for (uint32 BlockY = 0; BlockY < NumSubBlocksY; ++BlockY)
+		{
+			for (uint32 BlockX = 0; BlockX < NumSubBlocksX; ++BlockX)
+			{
+				FMiroPixelBlock4x4_4Bytes Block;
+				static_assert(sizeof(Block) == BlockSize*BlockSize*4);
+
+				for (uint32 BY = 0; BY < BlockSize; ++BY)
+				{
+					for (uint32 BX = 0; BX < BlockSize; ++BX)
+					{
+						uint32 SrcX = FMath::Min(BlockX * BlockSize + BX, uint32(SubSize.X) - 1);
+						uint32 SrcY = FMath::Min(BlockY * BlockSize + BY, uint32(SubSize.Y) - 1);
+
+						const uint8* FromBlockPtr = From + (SrcY * FromSize.X + SrcX) * 3;
+
+						Block.Bytes[BY * BlockSize * 4 + BX * 4 + 0] = FromBlockPtr[0];
+						Block.Bytes[BY * BlockSize * 4 + BX * 4 + 1] = FromBlockPtr[1];
+						Block.Bytes[BY * BlockSize * 4 + BX * 4 + 2] = FromBlockPtr[2];
+						Block.Bytes[BY * BlockSize * 4 + BX * 4 + 3] = 255;
+					}
+				}
+
+				constexpr int32 CompressedBlockSize = 16;
+				uint8* DestBlockPtr = To + (BlockY * NumBlocksX + BlockX) * CompressedBlockSize;
+
+				stb::stb_compress_dxt_block(DestBlockPtr, Block.Bytes, 2, STBQuality);
+			}
+		}
+	}
+
+	void RGBA_To_BC2SubImage(FImageSize FromSize, FImageSize ToSize, FImageSize SubSize, const uint8* From, uint8* To, int32 Quality)
+	{
+		if (SubSize.X <= 0 || SubSize.Y <= 0)
+		{
+			return;
+		}
+		
+		constexpr uint32 BlockSize = 4; 
+
+		const uint32 NumBlocksX = FMath::DivideAndRoundUp(uint32(ToSize.X), BlockSize);
+		const uint32 NumSubBlocksX = FMath::DivideAndRoundUp(uint32(SubSize.X), BlockSize);
+		const uint32 NumSubBlocksY = FMath::DivideAndRoundUp(uint32(SubSize.Y), BlockSize);
+
+		const int32 STBQuality = Quality < 2 ? STB_DXT_NORMAL : STB_DXT_HIGHQUAL;
+		for (uint32 BlockY = 0; BlockY < NumSubBlocksY; ++BlockY)
+		{
+			for (uint32 BlockX = 0; BlockX < NumSubBlocksX; ++BlockX)
+			{
+				FMiroPixelBlock4x4_4Bytes Block;
+				static_assert(sizeof(Block) == BlockSize*BlockSize*4);
+
+				for (uint32 BY = 0; BY < BlockSize; ++BY)
+				{
+					for (uint32 BX = 0; BX < BlockSize; ++BX)
+					{
+						uint32 SrcX = FMath::Min(BlockX * BlockSize + BX, uint32(SubSize.X) - 1);
+						uint32 SrcY = FMath::Min(BlockY * BlockSize + BY, uint32(SubSize.Y) - 1);
+
+						const uint8* FromBlockPtr = From + (SrcY * FromSize.X + SrcX) * 4;
+
+						Block.Bytes[BY * BlockSize * 4 + BX * 4 + 0] = FromBlockPtr[0];
+						Block.Bytes[BY * BlockSize * 4 + BX * 4 + 1] = FromBlockPtr[1];
+						Block.Bytes[BY * BlockSize * 4 + BX * 4 + 2] = FromBlockPtr[2];
+						Block.Bytes[BY * BlockSize * 4 + BX * 4 + 3] = FromBlockPtr[3];
+					}
+				}
+
+				constexpr int32 CompressedBlockSize = 16;
+				uint8* DestBlockPtr = To + (BlockY * NumBlocksX + BlockX) * CompressedBlockSize;
+
+				stb::stb_compress_dxt_block(DestBlockPtr, Block.Bytes, 2, STBQuality);
+			}
+		}
+	}
+
+	void RGB_To_BC3SubImage(FImageSize FromSize, FImageSize ToSize, FImageSize SubSize, const uint8* From, uint8* To, int32 Quality)
+	{
+		if (SubSize.X <= 0 || SubSize.Y <= 0)
+		{
+			return;
+		}
+		
+		const int32 STBQuality = Quality < 2 ? STB_DXT_NORMAL : STB_DXT_HIGHQUAL;
+		
+		constexpr uint32 BlockSize = 4; 
+
+		const uint32 NumBlocksX = FMath::DivideAndRoundUp(uint32(ToSize.X), BlockSize);
+		const uint32 NumSubBlocksX = FMath::DivideAndRoundUp(uint32(SubSize.X), BlockSize);
+		const uint32 NumSubBlocksY = FMath::DivideAndRoundUp(uint32(SubSize.Y), BlockSize);
+
+		for (uint32 BlockY = 0; BlockY < NumSubBlocksY; ++BlockY)
+		{
+			for (uint32 BlockX = 0; BlockX < NumSubBlocksX; ++BlockX)
+			{
+				FMiroPixelBlock4x4_4Bytes Block;
+				static_assert(sizeof(Block) == BlockSize*BlockSize*4);
+
+				for (uint32 BY = 0; BY < BlockSize; ++BY)
+				{
+					for (uint32 BX = 0; BX < BlockSize; ++BX)
+					{
+						uint32 SrcX = FMath::Min(BlockX * BlockSize + BX, uint32(SubSize.X) - 1);
+						uint32 SrcY = FMath::Min(BlockY * BlockSize + BY, uint32(SubSize.Y) - 1);
+
+						const uint8* FromBlockPtr = From + (SrcY * FromSize.X + SrcX) * 3;
+
+						Block.Bytes[BY * BlockSize * 4 + BX * 4 + 0] = FromBlockPtr[0];
+						Block.Bytes[BY * BlockSize * 4 + BX * 4 + 1] = FromBlockPtr[1];
+						Block.Bytes[BY * BlockSize * 4 + BX * 4 + 2] = FromBlockPtr[2];
+						Block.Bytes[BY * BlockSize * 4 + BX * 4 + 3] = 255;
+					}
+				}
+
+				constexpr int32 CompressedBlockSize = 16;
+				uint8* DestBlockPtr = To + (BlockY * NumBlocksX + BlockX) * CompressedBlockSize;
+
+				stb::stb_compress_dxt_block(DestBlockPtr, Block.Bytes, 1, STBQuality);
+			}
+		}
+	}
+
+	void RGBA_To_BC3SubImage(FImageSize FromSize, FImageSize ToSize, FImageSize SubSize, const uint8* From, uint8* To, int32 Quality)
+	{
+		if (SubSize.X <= 0 || SubSize.Y <= 0)
+		{
+			return;
+		}
+		
+		constexpr uint32 BlockSize = 4; 
+
+		const uint32 NumBlocksX = FMath::DivideAndRoundUp(uint32(ToSize.X), BlockSize);
+		const uint32 NumSubBlocksX = FMath::DivideAndRoundUp(uint32(SubSize.X), BlockSize);
+		const uint32 NumSubBlocksY = FMath::DivideAndRoundUp(uint32(SubSize.Y), BlockSize);
+
+		const int32 STBQuality = Quality < 2 ? STB_DXT_NORMAL : STB_DXT_HIGHQUAL;
+		for (uint32 BlockY = 0; BlockY < NumSubBlocksY; ++BlockY)
+		{
+			for (uint32 BlockX = 0; BlockX < NumSubBlocksX; ++BlockX)
+			{
+				FMiroPixelBlock4x4_4Bytes Block;
+				static_assert(sizeof(Block) == BlockSize*BlockSize*4);
+
+				for (uint32 BY = 0; BY < BlockSize; ++BY)
+				{
+					for (uint32 BX = 0; BX < BlockSize; ++BX)
+					{
+						uint32 SrcX = FMath::Min(BlockX * BlockSize + BX, uint32(SubSize.X) - 1);
+						uint32 SrcY = FMath::Min(BlockY * BlockSize + BY, uint32(SubSize.Y) - 1);
+
+						const uint8* FromBlockPtr = From + (SrcY * FromSize.X + SrcX) * 4;
+
+						Block.Bytes[BY * BlockSize * 4 + BX * 4 + 0] = FromBlockPtr[0];
+						Block.Bytes[BY * BlockSize * 4 + BX * 4 + 1] = FromBlockPtr[1];
+						Block.Bytes[BY * BlockSize * 4 + BX * 4 + 2] = FromBlockPtr[2];
+						Block.Bytes[BY * BlockSize * 4 + BX * 4 + 3] = FromBlockPtr[3];
+					}
+				}
+
+				constexpr int32 CompressedBlockSize = 16;
+				uint8* DestBlockPtr = To + (BlockY * NumBlocksX + BlockX) * CompressedBlockSize;
+
+				stb::stb_compress_dxt_block(DestBlockPtr, Block.Bytes, 1, STBQuality);
+			}
+		}
+	}
+
+	void L_To_BC4SubImage(FImageSize FromSize, FImageSize ToSize, FImageSize SubSize, const uint8* From, uint8* To, int32 Quality)
+	{
+		if (SubSize.X <= 0 || SubSize.Y <= 0)
+		{
+			return;
+		}
+		
+		constexpr uint32 BlockSize = 4; 
+
+		const uint32 NumBlocksX = FMath::DivideAndRoundUp(uint32(ToSize.X), BlockSize);
+		const uint32 NumSubBlocksX = FMath::DivideAndRoundUp(uint32(SubSize.X), BlockSize);
+		const uint32 NumSubBlocksY = FMath::DivideAndRoundUp(uint32(SubSize.Y), BlockSize);
+
+		for (uint32 BlockY = 0; BlockY < NumSubBlocksY; ++BlockY)
+		{
+			for (uint32 BlockX = 0; BlockX < NumSubBlocksX; ++BlockX)
+			{
+				FMiroPixelBlock4x4_1Bytes Block;
+				static_assert(sizeof(Block) == BlockSize*BlockSize);
+
+				for (uint32 BY = 0; BY < BlockSize; ++BY)
+				{
+					for (uint32 BX = 0; BX < BlockSize; ++BX)
+					{
+						uint32 SrcX = FMath::Min(BlockX * BlockSize + BX, uint32(SubSize.X) - 1);
+						uint32 SrcY = FMath::Min(BlockY * BlockSize + BY, uint32(SubSize.Y) - 1);
+
+						const uint8* FromBlockPtr = From + (SrcY * FromSize.X + SrcX) * 1;
+
+						Block.Bytes[BY * BlockSize * 1 + BX * 1 + 0] = FromBlockPtr[0];
+					}
+				}
+
+				constexpr int32 CompressedBlockSize = 8;
+				uint8* DestBlockPtr = To + (BlockY * NumBlocksX + BlockX) * CompressedBlockSize;
+
+				stb::stb_compress_bc4_block(DestBlockPtr, Block.Bytes);
+			}
+		}
+	}
+
+	void RGB_To_BC5SubImage(FImageSize FromSize, FImageSize ToSize, FImageSize SubSize, const uint8* From, uint8* To, int32 Quality)
+	{
+		(void)Quality;
+
+		if (SubSize.X <= 0 || SubSize.Y <= 0)
+		{
+			return;
+		}
+			
+		constexpr uint32 BlockSize = 4; 
+
+		const uint32 NumBlocksX = FMath::DivideAndRoundUp(uint32(ToSize.X), BlockSize);
+		const uint32 NumSubBlocksX = FMath::DivideAndRoundUp(uint32(SubSize.X), BlockSize);
+		const uint32 NumSubBlocksY = FMath::DivideAndRoundUp(uint32(SubSize.Y), BlockSize);
+
+		for (uint32 BlockY = 0; BlockY < NumSubBlocksY; ++BlockY)
+		{
+			for (uint32 BlockX = 0; BlockX < NumSubBlocksX; ++BlockX)
+			{
+				alignas(4) uint8 Block[BlockSize*BlockSize*2];
+
+				for (uint32 BY = 0; BY < BlockSize; ++BY)
+				{
+					for (uint32 BX = 0; BX < BlockSize; ++BX)
+					{
+						uint32 SrcX = FMath::Min(BlockX * BlockSize + BX, uint32(SubSize.X) - 1);
+						uint32 SrcY = FMath::Min(BlockY * BlockSize + BY, uint32(SubSize.Y) - 1);
+
+						const uint8* FromBlockPtr = From + (SrcY * FromSize.X + SrcX) * 3;
+
+						Block[BY * BlockSize * 2 + BX * 2 + 0] = FromBlockPtr[0];
+						Block[BY * BlockSize * 2 + BX * 2 + 1] = FromBlockPtr[1];
+					}
+				}
+
+				constexpr int32 CompressedBlockSize = 16;
+				uint8* DestBlockPtr = To + (BlockY * NumBlocksX + BlockX) * CompressedBlockSize;
+
+				stb::stb_compress_bc5_block(DestBlockPtr, Block);
+			}
+		}
+	}
+
+	void RGBA_To_BC5SubImage(FImageSize FromSize, FImageSize ToSize, FImageSize SubSize, const uint8* From, uint8* To, int32 Quality)
+	{
+		(void)Quality;
+
+		if (SubSize.X <= 0 || SubSize.Y <= 0)
+		{
+			return;
+		}
+			
+		constexpr uint32 BlockSize = 4; 
+
+		const uint32 NumBlocksX = FMath::DivideAndRoundUp(uint32(ToSize.X), BlockSize);
+		const uint32 NumSubBlocksX = FMath::DivideAndRoundUp(uint32(SubSize.X), BlockSize);
+		const uint32 NumSubBlocksY = FMath::DivideAndRoundUp(uint32(SubSize.Y), BlockSize);
+
+		for (uint32 BlockY = 0; BlockY < NumSubBlocksY; ++BlockY)
+		{
+			for (uint32 BlockX = 0; BlockX < NumSubBlocksX; ++BlockX)
+			{
+				alignas(4) uint8 Block[BlockSize*BlockSize*2];
+
+				for (uint32 BY = 0; BY < BlockSize; ++BY)
+				{
+					for (uint32 BX = 0; BX < BlockSize; ++BX)
+					{
+						uint32 SrcX = FMath::Min(BlockX * BlockSize + BX, uint32(SubSize.X) - 1);
+						uint32 SrcY = FMath::Min(BlockY * BlockSize + BY, uint32(SubSize.Y) - 1);
+
+						const uint8* FromBlockPtr = From + (SrcY * FromSize.X + SrcX) * 4;
+
+						Block[BY * BlockSize * 2 + BX * 2 + 0] = FromBlockPtr[0];
+						Block[BY * BlockSize * 2 + BX * 2 + 1] = FromBlockPtr[1];
+					}
+				}
+
+				constexpr int32 CompressedBlockSize = 16;
+				uint8* DestBlockPtr = To + (BlockY * NumBlocksX + BlockX) * CompressedBlockSize;
+
+				stb::stb_compress_bc5_block(DestBlockPtr, Block);
+			}
+		}
+	}
+
+#endif // MIRO_INCLUDE_BC
+
+#if MIRO_INCLUDE_ASTC
+	template<uint32 BlockSize>
+	void GenericSubImage_RGB_To_ASTCRGL(FImageSize FromSize, FImageSize ToSize, FImageSize SubSize, const uint8* RESTRICT From, uint8* RESTRICT To, int32 Quality)
+	{
+		(void)Quality;
+	
+		if (SubSize.X <= 0 || SubSize.Y <= 0)
+		{
+			return;
+		}
+		
+		const uint32 NumBlocksX = FMath::DivideAndRoundUp(uint32(ToSize.X), BlockSize);
+		const uint32 NumSubBlocksX = FMath::DivideAndRoundUp(uint32(SubSize.X), BlockSize);
+		const uint32 NumSubBlocksY = FMath::DivideAndRoundUp(uint32(SubSize.Y), BlockSize);
+
+		constexpr astcrt::PhysicalBlock PhysicalBlockZero = { {0} };
+		
+		for (uint32 BlockY = 0; BlockY < NumSubBlocksY; ++BlockY)
+		{
+			for (uint32 BlockX = 0; BlockX < NumSubBlocksX; ++BlockX)
+			{
+				alignas(4) uint8 Block[BlockSize*BlockSize*4];
+
+				for (uint32 BY = 0; BY < BlockSize; ++BY)
+				{
+					for (uint32 BX = 0; BX < BlockSize; ++BX)
+					{
+						uint32 SrcX = FMath::Min(BlockX * BlockSize + BX, uint32(SubSize.X) - 1);
+						uint32 SrcY = FMath::Min(BlockY * BlockSize + BY, uint32(SubSize.Y) - 1);
+
+						const uint8* FromBlockPtr = From + (SrcY * FromSize.X + SrcX) * 3;
+
+						Block[BY * BlockSize * 4 + BX * 4 + 0] = 0;
+						Block[BY * BlockSize * 4 + BX * 4 + 1] = FromBlockPtr[1];
+						Block[BY * BlockSize * 4 + BX * 4 + 2] = FromBlockPtr[0];
+						Block[BY * BlockSize * 4 + BX * 4 + 3] = 0;
+					}
+				}
+
+				constexpr int32 CompressedBlockSize = 16;
+				uint8* DestBlockPtr = To + (BlockY * NumBlocksX + BlockX) * CompressedBlockSize;
+				
+				astcrt::PhysicalBlock* DestPhysicalBlock = reinterpret_cast<astcrt::PhysicalBlock*>(DestBlockPtr);
+				*DestPhysicalBlock = PhysicalBlockZero;
+				astcrt::compress_block<BlockSize, BlockSize>((astcrt::unorm8_t*)Block, DestPhysicalBlock);
+			}
+		}
+	}
+
+	template<uint32 BlockSize>
+	void GenericSubImage_RGBA_To_ASTCRGL(FImageSize FromSize, FImageSize ToSize, FImageSize SubSize, const uint8* RESTRICT From, uint8* RESTRICT To, int32 Quality)
+	{
+		(void)Quality;
+	
+		if (SubSize.X <= 0 || SubSize.Y <= 0)
+		{
+			return;
+		}
+		
+		const uint32 NumBlocksX = FMath::DivideAndRoundUp(uint32(ToSize.X), BlockSize);
+		const uint32 NumSubBlocksX = FMath::DivideAndRoundUp(uint32(SubSize.X), BlockSize);
+		const uint32 NumSubBlocksY = FMath::DivideAndRoundUp(uint32(SubSize.Y), BlockSize);
+
+		constexpr astcrt::PhysicalBlock PhysicalBlockZero = { {0} };
+		
+		for (uint32 BlockY = 0; BlockY < NumSubBlocksY; ++BlockY)
+		{
+			for (uint32 BlockX = 0; BlockX < NumSubBlocksX; ++BlockX)
+			{
+				alignas(4) uint8 Block[BlockSize*BlockSize*4];
+
+				for (uint32 BY = 0; BY < BlockSize; ++BY)
+				{
+					for (uint32 BX = 0; BX < BlockSize; ++BX)
+					{
+						uint32 SrcX = FMath::Min(BlockX * BlockSize + BX, uint32(SubSize.X) - 1);
+						uint32 SrcY = FMath::Min(BlockY * BlockSize + BY, uint32(SubSize.Y) - 1);
+
+						const uint8* FromBlockPtr = From + (SrcY * FromSize.X + SrcX) * 4;
+
+						Block[BY * BlockSize * 4 + BX * 4 + 0] = 0;
+						Block[BY * BlockSize * 4 + BX * 4 + 1] = FromBlockPtr[1];
+						Block[BY * BlockSize * 4 + BX * 4 + 2] = FromBlockPtr[0];
+						Block[BY * BlockSize * 4 + BX * 4 + 3] = 0;
+					}
+				}
+
+				constexpr int32 CompressedBlockSize = 16;
+				uint8* DestBlockPtr = To + (BlockY * NumBlocksX + BlockX) * CompressedBlockSize;
+				
+				astcrt::PhysicalBlock* DestPhysicalBlock = reinterpret_cast<astcrt::PhysicalBlock*>(DestBlockPtr);
+				*DestPhysicalBlock = PhysicalBlockZero;
+				astcrt::compress_block<BlockSize, BlockSize>((astcrt::unorm8_t*)Block, DestPhysicalBlock);
+			}
+		}
+	}
+
+
+	template<uint32 BlockSize>
+	void GenericSubImage_L_To_ASTCRGBL(FImageSize FromSize, FImageSize ToSize, FImageSize SubSize, const uint8* RESTRICT From, uint8* RESTRICT To, int32 Quality)
+	{
+		(void)Quality;
+	
+		if (SubSize.X <= 0 || SubSize.Y <= 0)
+		{
+			return;
+		}
+		
+		const uint32 NumBlocksX = FMath::DivideAndRoundUp(uint32(ToSize.X), BlockSize);
+		const uint32 NumSubBlocksX = FMath::DivideAndRoundUp(uint32(SubSize.X), BlockSize);
+		const uint32 NumSubBlocksY = FMath::DivideAndRoundUp(uint32(SubSize.Y), BlockSize);
+
+		constexpr astcrt::PhysicalBlock PhysicalBlockZero = { {0} };
+		
+		for (uint32 BlockY = 0; BlockY < NumSubBlocksY; ++BlockY)
+		{
+			for (uint32 BlockX = 0; BlockX < NumSubBlocksX; ++BlockX)
+			{
+				alignas(4) uint8 Block[BlockSize*BlockSize*4];
+
+				for (uint32 BY = 0; BY < BlockSize; ++BY)
+				{
+					for (uint32 BX = 0; BX < BlockSize; ++BX)
+					{
+						uint32 SrcX = FMath::Min(BlockX * BlockSize + BX, uint32(SubSize.X) - 1);
+						uint32 SrcY = FMath::Min(BlockY * BlockSize + BY, uint32(SubSize.Y) - 1);
+
+						const uint8* FromBlockPtr = From + (SrcY * FromSize.X + SrcX) * 1;
+
+						Block[BY * BlockSize * 4 + BX * 4 + 0] = FromBlockPtr[0];
+						Block[BY * BlockSize * 4 + BX * 4 + 1] = FromBlockPtr[0];
+						Block[BY * BlockSize * 4 + BX * 4 + 2] = FromBlockPtr[0];
+						Block[BY * BlockSize * 4 + BX * 4 + 3] = 0;
+					}
+				}
+
+				constexpr int32 CompressedBlockSize = 16;
+				uint8* DestBlockPtr = To + (BlockY * NumBlocksX + BlockX) * CompressedBlockSize;
+				
+				astcrt::PhysicalBlock* DestPhysicalBlock = reinterpret_cast<astcrt::PhysicalBlock*>(DestBlockPtr);
+				*DestPhysicalBlock = PhysicalBlockZero;
+				astcrt::compress_block<BlockSize, BlockSize>((astcrt::unorm8_t*)Block, DestPhysicalBlock);
+			}
+		}
+	}
+
+	template<uint32 BlockSize>
+	void GenericSubImage_RGB_To_ASTCRGBAL(FImageSize FromSize, FImageSize ToSize, FImageSize SubSize, const uint8* RESTRICT From, uint8* RESTRICT To, int32 Quality)
+	{
+		(void)Quality;
+	
+		if (SubSize.X <= 0 || SubSize.Y <= 0)
+		{
+			return;
+		}
+		
+		const uint32 NumBlocksX = FMath::DivideAndRoundUp(uint32(ToSize.X), BlockSize);
+		const uint32 NumSubBlocksX = FMath::DivideAndRoundUp(uint32(SubSize.X), BlockSize);
+		const uint32 NumSubBlocksY = FMath::DivideAndRoundUp(uint32(SubSize.Y), BlockSize);
+
+		constexpr astcrt::PhysicalBlock PhysicalBlockZero = { {0} };
+		
+		for (uint32 BlockY = 0; BlockY < NumSubBlocksY; ++BlockY)
+		{
+			for (uint32 BlockX = 0; BlockX < NumSubBlocksX; ++BlockX)
+			{
+				alignas(4) uint8 Block[BlockSize*BlockSize*4];
+
+				for (uint32 BY = 0; BY < BlockSize; ++BY)
+				{
+					for (uint32 BX = 0; BX < BlockSize; ++BX)
+					{
+						uint32 SrcX = FMath::Min(BlockX * BlockSize + BX, uint32(SubSize.X) - 1);
+						uint32 SrcY = FMath::Min(BlockY * BlockSize + BY, uint32(SubSize.Y) - 1);
+
+						const uint8* FromBlockPtr = From + (SrcY * FromSize.X + SrcX) * 3;
+
+						Block[BY * BlockSize * 4 + BX * 4 + 0] = FromBlockPtr[2];
+						Block[BY * BlockSize * 4 + BX * 4 + 1] = FromBlockPtr[1];
+						Block[BY * BlockSize * 4 + BX * 4 + 2] = FromBlockPtr[0];
+						Block[BY * BlockSize * 4 + BX * 4 + 3] = 0;
+					}
+				}
+
+				constexpr int32 CompressedBlockSize = 16;
+				uint8* DestBlockPtr = To + (BlockY * NumBlocksX + BlockX) * CompressedBlockSize;
+				
+				astcrt::PhysicalBlock* DestPhysicalBlock = reinterpret_cast<astcrt::PhysicalBlock*>(DestBlockPtr);
+				*DestPhysicalBlock = PhysicalBlockZero;
+				astcrt::compress_block<BlockSize, BlockSize>((astcrt::unorm8_t*)Block, DestPhysicalBlock);
+			}
+		}
+	}
+
+	template<uint32 BlockSize>
+	void GenericSubImage_RGBA_To_ASTCRGBAL(FImageSize FromSize, FImageSize ToSize, FImageSize SubSize, const uint8* RESTRICT From, uint8* RESTRICT To, int32 Quality)
+	{
+		(void)Quality;
+	
+		if (SubSize.X <= 0 || SubSize.Y <= 0)
+		{
+			return;
+		}
+		
+		const uint32 NumBlocksX = FMath::DivideAndRoundUp(uint32(ToSize.X), BlockSize);
+		const uint32 NumSubBlocksX = FMath::DivideAndRoundUp(uint32(SubSize.X), BlockSize);
+		const uint32 NumSubBlocksY = FMath::DivideAndRoundUp(uint32(SubSize.Y), BlockSize);
+
+		constexpr astcrt::PhysicalBlock PhysicalBlockZero = { {0} };
+		
+		for (uint32 BlockY = 0; BlockY < NumSubBlocksY; ++BlockY)
+		{
+			for (uint32 BlockX = 0; BlockX < NumSubBlocksX; ++BlockX)
+			{
+				alignas(4) uint8 Block[BlockSize*BlockSize*4];
+
+				for (uint32 BY = 0; BY < BlockSize; ++BY)
+				{
+					for (uint32 BX = 0; BX < BlockSize; ++BX)
+					{
+						uint32 SrcX = FMath::Min(BlockX * BlockSize + BX, uint32(SubSize.X) - 1);
+						uint32 SrcY = FMath::Min(BlockY * BlockSize + BY, uint32(SubSize.Y) - 1);
+
+						const uint8* FromBlockPtr = From + (SrcY * FromSize.X + SrcX) * 4;
+
+						Block[BY * BlockSize * 4 + BX * 4 + 0] = FromBlockPtr[2];
+						Block[BY * BlockSize * 4 + BX * 4 + 1] = FromBlockPtr[1];
+						Block[BY * BlockSize * 4 + BX * 4 + 2] = FromBlockPtr[0];
+						Block[BY * BlockSize * 4 + BX * 4 + 3] = FromBlockPtr[3];
+					}
+				}
+
+				constexpr int32 CompressedBlockSize = 16;
+				uint8* DestBlockPtr = To + (BlockY * NumBlocksX + BlockX) * CompressedBlockSize;
+				
+				astcrt::PhysicalBlock* DestPhysicalBlock = reinterpret_cast<astcrt::PhysicalBlock*>(DestBlockPtr);
+				*DestPhysicalBlock = PhysicalBlockZero;
+				astcrt::compress_block_rgba<BlockSize, BlockSize>((astcrt::unorm8_t*)Block, DestPhysicalBlock);
+			}
+		}
+	}
+
+	template<uint32 BlockSize>
+	void GenericSubImage_RGBA_To_ASTCRGBL(FImageSize FromSize, FImageSize ToSize, FImageSize SubSize, const uint8* RESTRICT From, uint8* RESTRICT To, int32 Quality)
+	{
+		(void)Quality;
+	
+		if (SubSize.X <= 0 || SubSize.Y <= 0)
+		{
+			return;
+		}
+		
+		const uint32 NumBlocksX = FMath::DivideAndRoundUp(uint32(ToSize.X), BlockSize);
+		const uint32 NumSubBlocksX = FMath::DivideAndRoundUp(uint32(SubSize.X), BlockSize);
+		const uint32 NumSubBlocksY = FMath::DivideAndRoundUp(uint32(SubSize.Y), BlockSize);
+
+		constexpr astcrt::PhysicalBlock PhysicalBlockZero = { {0} };
+		
+		for (uint32 BlockY = 0; BlockY < NumSubBlocksY; ++BlockY)
+		{
+			for (uint32 BlockX = 0; BlockX < NumSubBlocksX; ++BlockX)
+			{
+				alignas(4) uint8 Block[BlockSize*BlockSize*4];
+
+				for (uint32 BY = 0; BY < BlockSize; ++BY)
+				{
+					for (uint32 BX = 0; BX < BlockSize; ++BX)
+					{
+						uint32 SrcX = FMath::Min(BlockX * BlockSize + BX, uint32(SubSize.X) - 1);
+						uint32 SrcY = FMath::Min(BlockY * BlockSize + BY, uint32(SubSize.Y) - 1);
+
+						const uint8* FromBlockPtr = From + (SrcY * FromSize.X + SrcX) * 4;
+
+						Block[BY * BlockSize * 4 + BX * 4 + 0] = FromBlockPtr[2];
+						Block[BY * BlockSize * 4 + BX * 4 + 1] = FromBlockPtr[1];
+						Block[BY * BlockSize * 4 + BX * 4 + 2] = FromBlockPtr[0];
+						Block[BY * BlockSize * 4 + BX * 4 + 3] = 0;
+					}
+				}
+
+				constexpr int32 CompressedBlockSize = 16;
+				uint8* DestBlockPtr = To + (BlockY * NumBlocksX + BlockX) * CompressedBlockSize;
+				
+				astcrt::PhysicalBlock* DestPhysicalBlock = reinterpret_cast<astcrt::PhysicalBlock*>(DestBlockPtr);
+				*DestPhysicalBlock = PhysicalBlockZero;
+				astcrt::compress_block<BlockSize, BlockSize>((astcrt::unorm8_t*)Block, DestPhysicalBlock);
+			}
+		}
+	}
+
+	template<uint32 BlockSize>
+	void GenericSubImage_RGB_To_ASTCRGBL(FImageSize FromSize, FImageSize ToSize, FImageSize SubSize, const uint8* RESTRICT From, uint8* RESTRICT To, int32 Quality)
+	{
+		(void)Quality;
+	
+		if (SubSize.X <= 0 || SubSize.Y <= 0)
+		{
+			return;
+		}
+		
+		const uint32 NumBlocksX = FMath::DivideAndRoundUp(uint32(ToSize.X), BlockSize);
+		const uint32 NumSubBlocksX = FMath::DivideAndRoundUp(uint32(SubSize.X), BlockSize);
+		const uint32 NumSubBlocksY = FMath::DivideAndRoundUp(uint32(SubSize.Y), BlockSize);
+
+		constexpr astcrt::PhysicalBlock PhysicalBlockZero = { {0} };
+		
+		for (uint32 BlockY = 0; BlockY < NumSubBlocksY; ++BlockY)
+		{
+			for (uint32 BlockX = 0; BlockX < NumSubBlocksX; ++BlockX)
+			{
+				alignas(4) uint8 Block[BlockSize*BlockSize*4];
+
+				for (uint32 BY = 0; BY < BlockSize; ++BY)
+				{
+					for (uint32 BX = 0; BX < BlockSize; ++BX)
+					{
+						uint32 SrcX = FMath::Min(BlockX * BlockSize + BX, uint32(SubSize.X) - 1);
+						uint32 SrcY = FMath::Min(BlockY * BlockSize + BY, uint32(SubSize.Y) - 1);
+
+						const uint8* FromBlockPtr = From + (SrcY * FromSize.X + SrcX) * 3;
+
+						Block[BY * BlockSize * 4 + BX * 4 + 0] = FromBlockPtr[2];
+						Block[BY * BlockSize * 4 + BX * 4 + 1] = FromBlockPtr[1];
+						Block[BY * BlockSize * 4 + BX * 4 + 2] = FromBlockPtr[0];
+						Block[BY * BlockSize * 4 + BX * 4 + 3] = 0;
+					}
+				}
+
+				constexpr int32 CompressedBlockSize = 16;
+				uint8* DestBlockPtr = To + (BlockY * NumBlocksX + BlockX) * CompressedBlockSize;
+				
+				astcrt::PhysicalBlock* DestPhysicalBlock = reinterpret_cast<astcrt::PhysicalBlock*>(DestBlockPtr);
+				*DestPhysicalBlock = PhysicalBlockZero;
+				astcrt::compress_block<BlockSize, BlockSize>((astcrt::unorm8_t*)Block, DestPhysicalBlock);
+			}
+		}
+	}
+
+
+	void RGBA_To_ASTC4x4RGBALSubImage(FImageSize FromSize, FImageSize ToSize, FImageSize SubSize, const uint8* RESTRICT From, uint8* RESTRICT To, int32 Quality)
+	{
+		GenericSubImage_RGBA_To_ASTCRGBAL<4>(FromSize, ToSize, SubSize, From, To, Quality);
+	}
+
+	void RGB_To_ASTC4x4RGBALSubImage(FImageSize FromSize, FImageSize ToSize, FImageSize SubSize, const uint8* RESTRICT From, uint8* RESTRICT To, int32 Quality)
+	{
+		GenericSubImage_RGB_To_ASTCRGBAL<4>(FromSize, ToSize, SubSize, From, To, Quality);
+	}
+
+	void RGB_To_ASTC4x4RGBLSubImage(FImageSize FromSize, FImageSize ToSize, FImageSize SubSize, const uint8* RESTRICT From, uint8* RESTRICT To, int32 Quality)
+	{
+		GenericSubImage_RGB_To_ASTCRGBL<4>(FromSize, ToSize, SubSize, From, To, Quality);
+	}
+
+	void RGBA_To_ASTC4x4RGBLSubImage(FImageSize FromSize, FImageSize ToSize, FImageSize SubSize, const uint8* RESTRICT From, uint8* RESTRICT To, int32 Quality)
+	{
+		GenericSubImage_RGBA_To_ASTCRGBL<4>(FromSize, ToSize, SubSize, From, To, Quality);
+	}
+
+	void L_To_ASTC4x4RGBLSubImage(FImageSize FromSize, FImageSize ToSize, FImageSize SubSize, const uint8* RESTRICT From, uint8* RESTRICT To, int32 Quality)
+	{
+		GenericSubImage_L_To_ASTCRGBL<4>(FromSize, ToSize, SubSize, From, To, Quality);
+	}
+
+	void RGB_To_ASTC4x4RGLSubImage(FImageSize FromSize, FImageSize ToSize, FImageSize SubSize, const uint8* RESTRICT From, uint8* RESTRICT To, int32 Quality)
+	{
+		GenericSubImage_RGB_To_ASTCRGL<4>(FromSize, ToSize, SubSize, From, To, Quality);
+	}
+
+	void RGBA_To_ASTC4x4RGLSubImage(FImageSize FromSize, FImageSize ToSize, FImageSize SubSize, const uint8* RESTRICT From, uint8* RESTRICT To, int32 Quality)
+	{
+		GenericSubImage_RGBA_To_ASTCRGL<4>(FromSize, ToSize, SubSize, From, To, Quality);
+	}
+
+	//---------------------------------------------------------------------------------------------
+	// 6x6
+	//---------------------------------------------------------------------------------------------
+	void RGBA_To_ASTC6x6RGBALSubImage(FImageSize FromSize, FImageSize ToSize, FImageSize SubSize, const uint8* RESTRICT From, uint8* RESTRICT To, int32 Quality)
+	{
+		GenericSubImage_RGBA_To_ASTCRGBAL<6>(FromSize, ToSize, SubSize, From, To, Quality);
+	}
+
+	void RGB_To_ASTC6x6RGBALSubImage(FImageSize FromSize, FImageSize ToSize, FImageSize SubSize, const uint8* RESTRICT From, uint8* RESTRICT To, int32 Quality)
+	{
+		GenericSubImage_RGB_To_ASTCRGBAL<6>(FromSize, ToSize, SubSize, From, To, Quality);
+	}
+
+	void RGB_To_ASTC6x6RGBLSubImage(FImageSize FromSize, FImageSize ToSize, FImageSize SubSize, const uint8* RESTRICT From, uint8* RESTRICT To, int32 Quality)
+	{
+		GenericSubImage_RGB_To_ASTCRGBL<6>(FromSize, ToSize, SubSize, From, To, Quality);
+	}
+
+	void RGBA_To_ASTC6x6RGBLSubImage(FImageSize FromSize, FImageSize ToSize, FImageSize SubSize, const uint8* RESTRICT From, uint8* RESTRICT To, int32 Quality)
+	{
+		GenericSubImage_RGBA_To_ASTCRGBL<6>(FromSize, ToSize, SubSize, From, To, Quality);
+	}
+
+	void RGB_To_ASTC6x6RGLSubImage(FImageSize FromSize, FImageSize ToSize, FImageSize SubSize, const uint8* RESTRICT From, uint8* RESTRICT To, int32 Quality)
+	{
+		GenericSubImage_RGB_To_ASTCRGL<6>(FromSize, ToSize, SubSize, From, To, Quality);
+	}
+
+	void RGBA_To_ASTC6x6RGLSubImage(FImageSize FromSize, FImageSize ToSize, FImageSize SubSize, const uint8* RESTRICT From, uint8* RESTRICT To, int32 Quality)
+	{
+		GenericSubImage_RGBA_To_ASTCRGL<6>(FromSize, ToSize, SubSize, From, To, Quality);
+	}
+
+
+	//---------------------------------------------------------------------------------------------
+	// 8x8
+	//---------------------------------------------------------------------------------------------
+	void RGBA_To_ASTC8x8RGBALSubImage(FImageSize FromSize, FImageSize ToSize, FImageSize SubSize, const uint8* RESTRICT From, uint8* RESTRICT To, int32 Quality)
+	{
+		GenericSubImage_RGBA_To_ASTCRGBAL<8>(FromSize, ToSize, SubSize, From, To, Quality);
+	}
+
+	void RGB_To_ASTC8x8RGBALSubImage(FImageSize FromSize, FImageSize ToSize, FImageSize SubSize, const uint8* RESTRICT From, uint8* RESTRICT To, int32 Quality)
+	{
+		GenericSubImage_RGB_To_ASTCRGBAL<8>(FromSize, ToSize, SubSize, From, To, Quality);
+	}
+
+	void RGB_To_ASTC8x8RGBLSubImage(FImageSize FromSize, FImageSize ToSize, FImageSize SubSize, const uint8* RESTRICT From, uint8* RESTRICT To, int32 Quality)
+	{
+		GenericSubImage_RGB_To_ASTCRGBL<8>(FromSize, ToSize, SubSize, From, To, Quality);
+	}
+
+	void RGBA_To_ASTC8x8RGBLSubImage(FImageSize FromSize, FImageSize ToSize, FImageSize SubSize, const uint8* RESTRICT From, uint8* RESTRICT To, int32 Quality)
+	{
+		GenericSubImage_RGBA_To_ASTCRGBL<8>(FromSize, ToSize, SubSize, From, To, Quality);
+	}
+
+	void RGB_To_ASTC8x8RGLSubImage(FImageSize FromSize, FImageSize ToSize, FImageSize SubSize, const uint8* RESTRICT From, uint8* RESTRICT To, int32 Quality)
+	{
+		GenericSubImage_RGB_To_ASTCRGL<8>(FromSize, ToSize, SubSize, From, To, Quality);
+	}
+
+	void RGBA_To_ASTC8x8RGLSubImage(FImageSize FromSize, FImageSize ToSize, FImageSize SubSize, const uint8* RESTRICT From, uint8* RESTRICT To, int32 Quality)
+	{
+		GenericSubImage_RGBA_To_ASTCRGL<8>(FromSize, ToSize, SubSize, From, To, Quality);
+	}
+
+
+#endif // MIRO_INCLUDE_ASTC
+
+}
+
 namespace miro::SubImageDecompression
 {
 
@@ -9842,12 +10690,12 @@ namespace miro::SubImageDecompression
 
 			for (uint32 bx = 0; bx < NumSubBlocksX; ++bx)
 			{
-				FMiroPixelBlock4x4 Block;
+				FMiroPixelBlock4x4_4Bytes Block;
 
 				//uint8 const* ColourBlock = From + 8;
 				//uint8 const* AlphaBlock = From;
 
-				bcdec::bcdec_bc3(RowFrom, Block.bytes, 4 * 4);
+				bcdec::bcdec_bc3(RowFrom, Block.Bytes, 4 * 4);
 
 				int32 b = 0;
 				for (int32 j = 0; j < 4; ++j)
@@ -9858,10 +10706,10 @@ namespace miro::SubImageDecompression
 						{
 							uint8* pPixelDest = To + 4 * (ToSize.X * (by * 4 + j) + (bx * 4 + i));
 
-							pPixelDest[0] = Block.bytes[b + 0];
-							pPixelDest[1] = Block.bytes[b + 1];
-							pPixelDest[2] = Block.bytes[b + 2];
-							pPixelDest[3] = Block.bytes[b + 3];
+							pPixelDest[0] = Block.Bytes[b + 0];
+							pPixelDest[1] = Block.Bytes[b + 1];
+							pPixelDest[2] = Block.Bytes[b + 2];
+							pPixelDest[3] = Block.Bytes[b + 3];
 						}
 
 						b += 4;
diff --git a/Engine/Plugins/Mutable/Source/MutableRuntime/Private/MuR/ImageFormatUtils.cpp b/Engine/Plugins/Mutable/Source/MutableRuntime/Private/MuR/ImageFormatUtils.cpp
new file mode 100644
--- /dev/null
+++ b/Engine/Plugins/Mutable/Source/MutableRuntime/Private/MuR/ImageFormatUtils.cpp
@@ -0,0 +1,534 @@
+// Copyright Epic Games, Inc. All Rights Reserved.
+
+#include "MuR/ImageFormatUtils.h"
+#include "MuR/Image.h"
+
+namespace UE::Mutable::Private
+{
+
+namespace SubImageFormat
+{
+namespace
+{
+	void DecompressionFuncNoOp(miro::FImageSize, miro::FImageSize, miro::FImageSize, const uint8*, uint8*)
+	{
+		checkf(false, TEXT("Formats not supported by SelectDecompressionFunction"));
+	}
+
+	void CompressionFuncNoOp(miro::FImageSize, miro::FImageSize, miro::FImageSize, const uint8*, uint8*, int32 Quality)
+	{
+		checkf(false, TEXT("Formats not supported by SelectCompressionFunction"));
+	}
+
+	void UncompressedFormatFuncNoOp(FImageSize, FImageSize, FImageSize, const uint8*, uint8*)
+	{
+		checkf(false, TEXT("Formats not supported by SelectUncompressedFormat"));
+	}
+} // namespace
+
+	miro::SubImageDecompression::FuncRefType SelectDecompressionFunction(EImageFormat DestFormat, EImageFormat SrcFormat)
+	{
+		using DecompressionFuncRefType = miro::SubImageDecompression::FuncRefType;
+		
+		check(DestFormat != EImageFormat::None);
+		check(SrcFormat != EImageFormat::None);
+
+		switch(SrcFormat)
+		{
+			case EImageFormat::BC1:
+			{
+				if (DestFormat == EImageFormat::RGBA_UByte)
+				{
+					return DecompressionFuncRefType(miro::SubImageDecompression::BC1_To_RGBASubImage);
+				}
+				else if (DestFormat == EImageFormat::RGB_UByte)
+				{
+					return DecompressionFuncRefType(miro::SubImageDecompression::BC1_To_RGBSubImage);
+				}
+				break;
+			}
+
+			case EImageFormat::BC2:
+			{
+				if (DestFormat == EImageFormat::RGBA_UByte)
+				{
+					return DecompressionFuncRefType(miro::SubImageDecompression::BC2_To_RGBASubImage);
+				}
+				else if (DestFormat == EImageFormat::RGB_UByte)
+				{
+					return DecompressionFuncRefType(miro::SubImageDecompression::BC2_To_RGBSubImage);
+				}
+				break;
+			}
+
+			case EImageFormat::BC3:
+			{
+				if (DestFormat == EImageFormat::RGBA_UByte)
+				{
+					return DecompressionFuncRefType(miro::SubImageDecompression::BC3_To_RGBASubImage);
+				}
+				else if (DestFormat == EImageFormat::RGB_UByte)
+				{
+					return DecompressionFuncRefType(miro::SubImageDecompression::BC3_To_RGBSubImage);
+				}
+				break;
+			}
+
+			case EImageFormat::BC4:
+			{
+				if (DestFormat == EImageFormat::RGBA_UByte)
+				{
+					return DecompressionFuncRefType(miro::SubImageDecompression::BC4_To_RGBASubImage);
+				}
+				else if (DestFormat == EImageFormat::RGB_UByte)
+				{
+					return DecompressionFuncRefType(miro::SubImageDecompression::BC4_To_RGBSubImage);
+				}
+				if (DestFormat == EImageFormat::L_UByte)
+				{
+					return DecompressionFuncRefType(miro::SubImageDecompression::BC4_To_LSubImage);
+				}
+				break;
+			}
+
+			case EImageFormat::BC5:
+			{
+				if (DestFormat == EImageFormat::RGBA_UByte)
+				{
+					return DecompressionFuncRefType(miro::SubImageDecompression::BC5_To_RGBASubImage);
+				}
+				else if (DestFormat == EImageFormat::RGB_UByte)
+				{
+					return DecompressionFuncRefType(miro::SubImageDecompression::BC5_To_RGBSubImage);
+				}
+				break;
+			}
+	
+			case EImageFormat::ASTC_4x4_RGBA_LDR:
+			{
+				if (DestFormat == EImageFormat::RGBA_UByte)
+				{
+					return DecompressionFuncRefType(miro::SubImageDecompression::ASTC4x4RGBAL_To_RGBASubImage);
+				}
+				else if (DestFormat == EImageFormat::RGB_UByte)
+				{
+					return DecompressionFuncRefType(miro::SubImageDecompression::ASTC4x4RGBAL_To_RGBSubImage);
+				}
+				break;
+			}
+
+			case EImageFormat::ASTC_4x4_RGB_LDR:
+			{
+				if (DestFormat == EImageFormat::RGBA_UByte)
+				{
+					return DecompressionFuncRefType(miro::SubImageDecompression::ASTC4x4RGBL_To_RGBASubImage);
+				}
+				else if (DestFormat == EImageFormat::RGB_UByte)
+				{
+					return DecompressionFuncRefType(miro::SubImageDecompression::ASTC4x4RGBL_To_RGBSubImage);
+				}
+				break;
+			}
+
+			case EImageFormat::ASTC_4x4_RG_LDR:
+			{
+				if (DestFormat == EImageFormat::RGBA_UByte)
+				{
+					return DecompressionFuncRefType(miro::SubImageDecompression::ASTC4x4RGL_To_RGBASubImage);
+				}
+				else if (DestFormat == EImageFormat::RGB_UByte)
+				{
+					return DecompressionFuncRefType(miro::SubImageDecompression::ASTC4x4RGL_To_RGBSubImage);
+				}
+				break;
+			}
+
+			case EImageFormat::ASTC_6x6_RGBA_LDR:
+			{
+				if (DestFormat == EImageFormat::RGBA_UByte)
+				{
+					return DecompressionFuncRefType(miro::SubImageDecompression::ASTC6x6RGBAL_To_RGBASubImage);
+				}
+				else if (DestFormat == EImageFormat::RGB_UByte)
+				{
+					return DecompressionFuncRefType(miro::SubImageDecompression::ASTC6x6RGBAL_To_RGBSubImage);
+				}
+				break;
+			}
+
+			case EImageFormat::ASTC_6x6_RGB_LDR:
+			{
+				if (DestFormat == EImageFormat::RGBA_UByte)
+				{
+					return DecompressionFuncRefType(miro::SubImageDecompression::ASTC6x6RGBL_To_RGBASubImage);
+				}
+				else if (DestFormat == EImageFormat::RGB_UByte)
+				{
+					return DecompressionFuncRefType(miro::SubImageDecompression::ASTC6x6RGBL_To_RGBSubImage);
+				}
+				break;
+			}
+
+			case EImageFormat::ASTC_6x6_RG_LDR:
+			{
+				if (DestFormat == EImageFormat::RGBA_UByte)
+				{
+					return DecompressionFuncRefType(miro::SubImageDecompression::ASTC6x6RGL_To_RGBASubImage);
+				}
+				else if (DestFormat == EImageFormat::RGB_UByte)
+				{
+					return DecompressionFuncRefType(miro::SubImageDecompression::ASTC6x6RGL_To_RGBSubImage);
+				}
+				break;
+			}
+
+			case EImageFormat::ASTC_8x8_RGBA_LDR:
+			{
+				if (DestFormat == EImageFormat::RGBA_UByte)
+				{
+					return DecompressionFuncRefType(miro::SubImageDecompression::ASTC8x8RGBAL_To_RGBASubImage);
+				}
+				else if (DestFormat == EImageFormat::RGB_UByte)
+				{
+					return DecompressionFuncRefType(miro::SubImageDecompression::ASTC8x8RGBAL_To_RGBSubImage);
+				}
+				break;
+			}
+
+			case EImageFormat::ASTC_8x8_RGB_LDR:
+			{
+				if (DestFormat == EImageFormat::RGBA_UByte)
+				{
+					return DecompressionFuncRefType(miro::SubImageDecompression::ASTC8x8RGBL_To_RGBASubImage);
+				}
+				else if (DestFormat == EImageFormat::RGB_UByte)
+				{
+					return DecompressionFuncRefType(miro::SubImageDecompression::ASTC8x8RGBL_To_RGBSubImage);
+				}
+				break;
+			}
+
+			case EImageFormat::ASTC_8x8_RG_LDR:
+			{
+				if (DestFormat == EImageFormat::RGBA_UByte)
+				{
+					return DecompressionFuncRefType(miro::SubImageDecompression::ASTC8x8RGL_To_RGBASubImage);
+				}
+				else if (DestFormat == EImageFormat::RGB_UByte)
+				{
+					return DecompressionFuncRefType(miro::SubImageDecompression::ASTC8x8RGL_To_RGBSubImage);
+				}
+				break;
+			}
+
+			case EImageFormat::ASTC_10x10_RGBA_LDR:
+			{
+				if (DestFormat == EImageFormat::RGBA_UByte)
+				{
+					return DecompressionFuncRefType(miro::SubImageDecompression::ASTC10x10RGBAL_To_RGBASubImage);
+				}
+				else if (DestFormat == EImageFormat::RGB_UByte)
+				{
+					return DecompressionFuncRefType(miro::SubImageDecompression::ASTC10x10RGBAL_To_RGBSubImage);
+				}
+				break;
+			}
+
+			case EImageFormat::ASTC_10x10_RGB_LDR:
+			{
+				if (DestFormat == EImageFormat::RGBA_UByte)
+				{
+					return DecompressionFuncRefType(miro::SubImageDecompression::ASTC10x10RGBL_To_RGBASubImage);
+				}
+				else if (DestFormat == EImageFormat::RGB_UByte)
+				{
+					return DecompressionFuncRefType(miro::SubImageDecompression::ASTC10x10RGBL_To_RGBSubImage);
+				}
+				break;
+			}
+
+			case EImageFormat::ASTC_10x10_RG_LDR:
+			{
+				if (DestFormat == EImageFormat::RGBA_UByte)
+				{
+					return DecompressionFuncRefType(miro::SubImageDecompression::ASTC10x10RGL_To_RGBASubImage);
+				}
+				else if (DestFormat == EImageFormat::RGB_UByte)
+				{
+					return DecompressionFuncRefType(miro::SubImageDecompression::ASTC10x10RGL_To_RGBSubImage);
+				}
+				break;
+			}
+
+			case EImageFormat::ASTC_12x12_RGBA_LDR:
+			{
+				if (DestFormat == EImageFormat::RGBA_UByte)
+				{
+					return DecompressionFuncRefType(miro::SubImageDecompression::ASTC12x12RGBAL_To_RGBASubImage);
+				}
+				else if (DestFormat == EImageFormat::RGB_UByte)
+				{
+					return DecompressionFuncRefType(miro::SubImageDecompression::ASTC12x12RGBAL_To_RGBSubImage);
+				}
+				break;
+			}
+
+			case EImageFormat::ASTC_12x12_RGB_LDR:
+			{
+				if (DestFormat == EImageFormat::RGBA_UByte)
+				{
+					return DecompressionFuncRefType(miro::SubImageDecompression::ASTC12x12RGBL_To_RGBASubImage);
+				}
+				else if (DestFormat == EImageFormat::RGB_UByte)
+				{
+					return DecompressionFuncRefType(miro::SubImageDecompression::ASTC12x12RGBL_To_RGBSubImage);
+				}
+				break;
+			}
+
+			case EImageFormat::ASTC_12x12_RG_LDR:
+			{
+				if (DestFormat == EImageFormat::RGBA_UByte)
+				{
+					return DecompressionFuncRefType(miro::SubImageDecompression::ASTC12x12RGL_To_RGBASubImage);
+				}
+				else if (DestFormat == EImageFormat::RGB_UByte)
+				{
+					return DecompressionFuncRefType(miro::SubImageDecompression::ASTC12x12RGL_To_RGBSubImage);
+				}
+				break;
+			}
+		}
+
+		return DecompressionFuncRefType(DecompressionFuncNoOp);
+	}
+
+	inline miro::SubImageCompression::FuncRefType SelectCompressionFunction(EImageFormat DestFormat, EImageFormat SrcFormat)
+	{
+		using CompressionFuncRefType = miro::SubImageCompression::FuncRefType;
+
+		check(DestFormat != EImageFormat::None);
+		check(SrcFormat != EImageFormat::None);
+		
+		switch (SrcFormat)
+		{
+			case EImageFormat::L_UByte:
+			{
+				switch (DestFormat)
+				{
+					case EImageFormat::BC4:
+					{
+						return CompressionFuncRefType(miro::SubImageCompression::L_To_BC4SubImage);
+					}
+					case EImageFormat::ASTC_4x4_RGB_LDR:
+					{
+						return CompressionFuncRefType(miro::SubImageCompression::L_To_ASTC4x4RGBLSubImage);
+					}	
+				}
+				break;
+			}
+
+			case EImageFormat::RGB_UByte:
+			{
+				switch (DestFormat)
+				{
+					case EImageFormat::BC1:
+					{
+						return CompressionFuncRefType(miro::SubImageCompression::RGB_To_BC1SubImage);
+					}
+					case EImageFormat::BC2:
+					{
+						return CompressionFuncRefType(miro::SubImageCompression::RGB_To_BC2SubImage);
+					}
+					case EImageFormat::BC3:
+					{
+						return CompressionFuncRefType(miro::SubImageCompression::RGB_To_BC3SubImage);
+					}
+					case EImageFormat::BC5:
+					{
+						return CompressionFuncRefType(miro::SubImageCompression::RGB_To_BC5SubImage);
+					}
+					case EImageFormat::ASTC_4x4_RGBA_LDR:
+					{
+						return CompressionFuncRefType(miro::SubImageCompression::RGB_To_ASTC4x4RGBALSubImage);
+					}
+					case EImageFormat::ASTC_4x4_RGB_LDR:
+					{
+						return CompressionFuncRefType(miro::SubImageCompression::RGB_To_ASTC4x4RGBLSubImage);
+					}
+					case EImageFormat::ASTC_4x4_RG_LDR:
+					{
+						return CompressionFuncRefType(miro::SubImageCompression::RGB_To_ASTC4x4RGLSubImage);
+					}
+
+					case EImageFormat::ASTC_6x6_RGBA_LDR:
+					{
+						return CompressionFuncRefType(miro::SubImageCompression::RGB_To_ASTC6x6RGBALSubImage);
+					}
+					case EImageFormat::ASTC_6x6_RGB_LDR:
+					{
+						return CompressionFuncRefType(miro::SubImageCompression::RGB_To_ASTC6x6RGBLSubImage);
+					}
+					case EImageFormat::ASTC_6x6_RG_LDR:
+					{
+						return CompressionFuncRefType(miro::SubImageCompression::RGB_To_ASTC6x6RGLSubImage);
+					}
+
+					case EImageFormat::ASTC_8x8_RGBA_LDR:
+					{
+						return CompressionFuncRefType(miro::SubImageCompression::RGB_To_ASTC8x8RGBALSubImage);
+					}
+					case EImageFormat::ASTC_8x8_RGB_LDR:
+					{
+						return CompressionFuncRefType(miro::SubImageCompression::RGB_To_ASTC8x8RGBLSubImage);
+					}
+					case EImageFormat::ASTC_8x8_RG_LDR:
+					{
+						return CompressionFuncRefType(miro::SubImageCompression::RGB_To_ASTC8x8RGLSubImage);
+					}
+				}
+
+				break;
+			}
+
+			case EImageFormat::RGBA_UByte:
+			{
+				switch (DestFormat)
+				{
+					case EImageFormat::BC1:
+					{
+						return CompressionFuncRefType(miro::SubImageCompression::RGBA_To_BC1SubImage);
+					}
+					case EImageFormat::BC2:
+					{
+						return CompressionFuncRefType(miro::SubImageCompression::RGBA_To_BC2SubImage);
+					}
+					case EImageFormat::BC3:
+					{
+						return CompressionFuncRefType(miro::SubImageCompression::RGBA_To_BC3SubImage);
+					}
+
+					case EImageFormat::BC5:
+					{
+						return CompressionFuncRefType(miro::SubImageCompression::RGBA_To_BC5SubImage);
+					}
+
+					case EImageFormat::ASTC_4x4_RGBA_LDR:
+					{
+						return CompressionFuncRefType(miro::SubImageCompression::RGBA_To_ASTC4x4RGBALSubImage);
+					}
+					case EImageFormat::ASTC_4x4_RGB_LDR:
+					{
+						return CompressionFuncRefType(miro::SubImageCompression::RGBA_To_ASTC4x4RGBLSubImage);
+					}
+					case EImageFormat::ASTC_4x4_RG_LDR:
+					{
+						return CompressionFuncRefType(miro::SubImageCompression::RGBA_To_ASTC4x4RGLSubImage);
+					}
+
+					case EImageFormat::ASTC_6x6_RGBA_LDR:
+					{
+						return CompressionFuncRefType(miro::SubImageCompression::RGBA_To_ASTC6x6RGBALSubImage);
+					}
+					case EImageFormat::ASTC_6x6_RGB_LDR:
+					{
+						return CompressionFuncRefType(miro::SubImageCompression::RGBA_To_ASTC6x6RGBLSubImage);
+					}
+					case EImageFormat::ASTC_6x6_RG_LDR:
+					{
+						return CompressionFuncRefType(miro::SubImageCompression::RGBA_To_ASTC6x6RGLSubImage);
+					}
+
+					case EImageFormat::ASTC_8x8_RGBA_LDR:
+					{
+						return CompressionFuncRefType(miro::SubImageCompression::RGBA_To_ASTC8x8RGBALSubImage);
+					}
+					case EImageFormat::ASTC_8x8_RGB_LDR:
+					{
+						return CompressionFuncRefType(miro::SubImageCompression::RGBA_To_ASTC8x8RGBLSubImage);
+					}
+					case EImageFormat::ASTC_8x8_RG_LDR:
+					{
+						return CompressionFuncRefType(miro::SubImageCompression::RGBA_To_ASTC8x8RGLSubImage);
+					}
+
+				}
+				break;
+			}
+		}
+
+		//check(false);
+		return CompressionFuncRefType(CompressionFuncNoOp);
+	}
+
+	void SubImageVec3ToVec4_U8(FImageSize FromSize, FImageSize ToSize, FImageSize SubSize, const uint8* RESTRICT From, uint8* RESTRICT To)
+	{
+		for (uint16 Y = 0; Y < SubSize.Y; ++Y)
+		{
+			for (uint16 X = 0; X < SubSize.X; ++X)
+			{
+				uint8* ToPtr         = To   + (Y*ToSize.X   + X)*4;
+				const uint8* FromPtr = From + (Y*FromSize.X + X)*3;
+
+				// TODO: Optimize.
+				ToPtr[0] = FromPtr[0];	
+				ToPtr[1] = FromPtr[1];	
+				ToPtr[2] = FromPtr[2];
+				ToPtr[3] = 255;
+			}
+		}
+	}
+
+	void SubImageVec4ToVec3_U8(FImageSize FromSize, FImageSize ToSize, FImageSize SubSize, const uint8* RESTRICT From, uint8* RESTRICT To)
+	{
+		for (uint16 Y = 0; Y < SubSize.Y; ++Y)
+		{
+			for (uint16 X = 0; X < SubSize.X; ++X)
+			{
+				uint8* ToPtr         = To   + (Y*ToSize.X   + X)*3;
+				const uint8* FromPtr = From + (Y*FromSize.X + X)*4;
+
+				// TODO: Optimize.
+				ToPtr[0] = FromPtr[0];	
+				ToPtr[1] = FromPtr[1];	
+				ToPtr[2] = FromPtr[2];
+			}
+		}
+	}
+	
+
+	UncompressedFormatFuncType* SelectUncompressedFormatFunction(EImageFormat DstFormat, EImageFormat SrcFormat)
+	{
+		check(GetUncompressedFormat(DstFormat) == DstFormat);
+		check(GetUncompressedFormat(SrcFormat) == SrcFormat);
+
+		switch (SrcFormat)
+		{
+		case EImageFormat::RGBA_UByte:
+		{
+			switch (DstFormat)
+			{
+				case EImageFormat::RGB_UByte:
+				{
+					return SubImageVec4ToVec3_U8;
+				}
+			}
+			break;
+		}
+		case EImageFormat::RGB_UByte:
+		{
+			switch (DstFormat)
+			{
+				case EImageFormat::RGBA_UByte:
+				{
+					return SubImageVec3ToVec4_U8;
+				}
+			}
+			break;
+		}
+		}
+		return UncompressedFormatFuncNoOp;
+	}
+
+} // namespace SubImageFormat
+} // namespace UE::Mutable::Private
diff --git a/Engine/Plugins/Mutable/Source/MutableRuntime/Private/MuR/ImageFormatUtils.h b/Engine/Plugins/Mutable/Source/MutableRuntime/Private/MuR/ImageFormatUtils.h
new file mode 100644
--- /dev/null
+++ b/Engine/Plugins/Mutable/Source/MutableRuntime/Private/MuR/ImageFormatUtils.h
@@ -0,0 +1,24 @@
+// Copyright Epic Games, Inc. All Rights Reserved.
+
+#pragma once
+
+#include "MuR/BlockCompression/Miro/Miro.h"
+#include "MuR/ImageTypes.h"
+
+namespace UE::Mutable::Private
+{
+
+namespace SubImageFormat
+{	
+
+	extern miro::SubImageDecompression::FuncRefType SelectDecompressionFunction(EImageFormat DestFormat, EImageFormat SrcFormat);
+	extern miro::SubImageCompression::FuncRefType SelectCompressionFunction(EImageFormat DestFormat, EImageFormat SrcFormat);	
+
+	using UncompressedFormatFuncType = void(FImageSize, FImageSize, FImageSize, const uint8*, uint8*); 
+	UncompressedFormatFuncType* SelectUncompressedFormatFunction(EImageFormat DestFormat, EImageFormat SrcFormat);
+	
+	void SubImageVec3ToVec4_U8(FImageSize FromSize, FImageSize ToSize, FImageSize SubSize, const uint8* From, uint8* To);
+	void SubImageVec4ToVec3_U8(FImageSize FromSize, FImageSize ToSize, FImageSize SubSize, const uint8* From, uint8* To);
+
+} // namespace SubImageFormat
+} // namespace UE::Mutable::Private
diff --git a/Engine/Plugins/Mutable/Source/MutableRuntime/Private/MuR/ImageResizeUtils.cpp b/Engine/Plugins/Mutable/Source/MutableRuntime/Private/MuR/ImageResizeUtils.cpp
new file mode 100644
--- /dev/null
+++ b/Engine/Plugins/Mutable/Source/MutableRuntime/Private/MuR/ImageResizeUtils.cpp
@@ -0,0 +1,162 @@
+// Copyright Epic Games, Inc. All Rights Reserved.
+
+#include "MuR/ImageResizeUtils.h"
+#include "HAL/Platform.h"
+
+namespace UE::Mutable::Private
+{
+
+namespace SubImageResize
+{
+	template<uint32 NumChannels>
+	FORCENOINLINE void ImageResizeLinearByLessThanTwoSubImageVectorImpl(
+			FVector2f InOneOverScalingFactor, FVector2f InSrcSubPixelOffset,
+			FImageSize DestSize, FImageSize DestSubSize, FImageSize SrcSize, FImageSize SrcSubSize,
+		  	uint8* RESTRICT SubDest, const uint8* RESTRICT SubSrc)
+	{	
+		static_assert(NumChannels > 0 && NumChannels <= 4);
+
+		check(InSrcSubPixelOffset.X >= 0.0f && InSrcSubPixelOffset.Y >= 0.0f);
+		check(InOneOverScalingFactor.X > 0.0f && InOneOverScalingFactor.Y > 0.0f);
+
+		// For some reason (probably cache related), doing this copy yields marginally but consistently 
+		// better performance than using the function argument. 
+		const FVector2f OneOverScalingFactor = InOneOverScalingFactor;
+		const FVector2f SrcSubPixelOffset = InSrcSubPixelOffset;
+		
+		check(OneOverScalingFactor.X <= 2.0f && OneOverScalingFactor.Y <= 2.0f);
+
+		for (uint16 Y = 0; Y < DestSubSize.Y; ++Y)
+		{
+			for (uint16 X = 0; X < DestSubSize.X; ++X)
+			{
+				auto LoadPixel = [](const uint8* Ptr) -> VectorRegister4Float
+				{	
+					alignas(4) uint8 PixelData[4] = {0};
+					if constexpr (NumChannels == 4)
+					{
+						FMemory::Memcpy(&(PixelData[0]), Ptr, 4);
+					}
+					else
+					{
+						for (uint32 C = 0; C < NumChannels; ++C)
+						{
+							PixelData[C] = Ptr[C];
+						}
+					}
+
+					return VectorIntToFloat(MakeVectorRegisterInt(PixelData[0], PixelData[1], PixelData[2], PixelData[3]));
+				};
+				
+				const FVector2f CoordsF = FVector2f(X, Y)*OneOverScalingFactor + SrcSubPixelOffset;
+	
+				const FIntVector2 Coords = FIntVector2(
+						FMath::Clamp<int32>(FMath::FloorToInt(CoordsF.X), 0, SrcSize.X - 1), 
+						FMath::Clamp<int32>(FMath::FloorToInt(CoordsF.Y), 0, SrcSize.Y - 1));
+
+				const FIntVector2 CoordsPlusOne = FIntVector2(
+						FMath::Clamp<int32>(Coords.X + 1, 0, SrcSize.X - 1),
+						FMath::Clamp<int32>(Coords.Y + 1, 0, SrcSize.Y - 1));
+
+				const VectorRegister4Float Pixel00 = LoadPixel(SubSrc + (Coords.Y*SrcSize.X + Coords.X) * NumChannels);
+				const VectorRegister4Float Pixel10 = LoadPixel(SubSrc + (Coords.Y*SrcSize.X + CoordsPlusOne.X) * NumChannels);
+				const VectorRegister4Float Pixel01 = LoadPixel(SubSrc + (CoordsPlusOne.Y*SrcSize.X + Coords.X) * NumChannels);
+				const VectorRegister4Float Pixel11 = LoadPixel(SubSrc + (CoordsPlusOne.Y*SrcSize.X + CoordsPlusOne.X) * NumChannels);
+
+				const VectorRegister4Float CoordsFVector = MakeVectorRegister(CoordsF.X, CoordsF.Y, 0.0f, 0.0f);
+				const VectorRegister4Float Frac = VectorSubtract(CoordsFVector, VectorFloor(CoordsFVector));
+
+				const VectorRegister4Float FracX = VectorReplicate(Frac, 0);
+				const VectorRegister4Float FracY = VectorReplicate(Frac, 1);
+				
+				VectorRegister4Float Result = VectorMultiplyAdd(FracX, VectorSubtract(Pixel10, Pixel00), Pixel00);
+				Result = VectorMultiplyAdd(
+						FracY,
+						VectorSubtract(VectorMultiplyAdd(FracX, VectorSubtract(Pixel11, Pixel01), Pixel01), Result),
+						Result);
+
+				uint8* Dest = SubDest + (Y*DestSize.X + X)*NumChannels;  
+
+				if constexpr (NumChannels == 4)
+				{
+					VectorStoreByte4(Result, Dest);
+					VectorResetFloatRegisters();
+				}
+				else
+				{
+					alignas(4) uint8 ResultData[4];
+					VectorStoreByte4(Result, ResultData);
+					VectorResetFloatRegisters();
+				
+					for (uint32 C = 0; C < NumChannels; ++C)
+					{
+						Dest[C] = ResultData[C];
+					}
+				}
+			}
+		}
+	}
+
+	template<int32 NumChannels>
+	UE_REWRITE void ImageResizeLinearByLessThanTwoSubImage(
+			FVector2f OneOverScalingFactor, FVector2f SrcSubPixelOffset,
+			FImageSize DestSize, FImageSize DestSubSize, FImageSize SrcSize, FImageSize SrcSubSize,
+			uint8* RESTRICT SubDest, const uint8* RESTRICT SubSrc)
+	{
+			ImageResizeLinearByLessThanTwoSubImageVectorImpl<NumChannels>(
+					OneOverScalingFactor, SrcSubPixelOffset,
+					DestSize, DestSubSize, SrcSize, SrcSubSize, SubDest, SubSrc);
+	}
+
+	void SubImageResizeLinearByLessThanTwoVec1_U8(
+				FVector2f OneOverScalingFactor, 
+				FVector2f SrcSubPixelOffset,
+				FImageSize DestSize, FImageSize DestSubSize, 
+				FImageSize SrcSize, FImageSize SrcSubSize,
+				uint8* RESTRICT SubDest, 
+				const uint8* RESTRICT SubSrc)
+	{
+		ImageResizeLinearByLessThanTwoSubImage<1>(OneOverScalingFactor, SrcSubPixelOffset, DestSize, DestSubSize, SrcSize, SrcSubSize, SubDest, SubSrc); 
+	}
+
+	void SubImageResizeLinearByLessThanTwoVec3_U8(
+				FVector2f OneOverScalingFactor, 
+				FVector2f SrcSubPixelOffset,
+				FImageSize DestSize, FImageSize DestSubSize, 
+				FImageSize SrcSize, FImageSize SrcSubSize,
+				uint8* RESTRICT SubDest, 
+				const uint8* RESTRICT SubSrc)
+	{
+		ImageResizeLinearByLessThanTwoSubImage<3>(OneOverScalingFactor, SrcSubPixelOffset, DestSize, DestSubSize, SrcSize, SrcSubSize, SubDest, SubSrc); 
+	}
+
+	void SubImageResizeLinearByLessThanTwoVec4_U8(
+				FVector2f OneOverScalingFactor, 
+				FVector2f SrcSubPixelOffset,
+				FImageSize DestSize, FImageSize DestSubSize, 
+				FImageSize SrcSize, FImageSize SrcSubSize,
+				uint8* RESTRICT SubDest, 
+				const uint8* RESTRICT SubSrc)
+	{
+		ImageResizeLinearByLessThanTwoSubImage<4>(OneOverScalingFactor, SrcSubPixelOffset, DestSize, DestSubSize, SrcSize, SrcSubSize, SubDest, SubSrc); 
+	}
+
+	void SubImageResizeLinearByLessThanTwoNoOp(FVector2f, FVector2f, FImageSize, FImageSize, FImageSize, FImageSize, uint8*, const uint8*)
+	{
+		checkf(false, TEXT("The format passed to SelectSubImageResizeFunc is not supported."));
+	}
+
+	SubImageResizeFuncType* SelectSubImageResizeFunc(EImageFormat Format)
+	{
+		switch (Format)
+		{
+			case EImageFormat::L_UByte    : return SubImageResizeLinearByLessThanTwoVec1_U8;
+			case EImageFormat::RGB_UByte  : return SubImageResizeLinearByLessThanTwoVec3_U8;
+			case EImageFormat::RGBA_UByte : return SubImageResizeLinearByLessThanTwoVec4_U8;
+			case EImageFormat::BGRA_UByte : return SubImageResizeLinearByLessThanTwoVec4_U8;
+			default: return SubImageResizeLinearByLessThanTwoNoOp;
+		}
+	}
+
+} // namespace SubImageResize
+} // namespace UE::Mutable::Private
diff --git a/Engine/Plugins/Mutable/Source/MutableRuntime/Private/MuR/ImageResizeUtils.h b/Engine/Plugins/Mutable/Source/MutableRuntime/Private/MuR/ImageResizeUtils.h
new file mode 100644
--- /dev/null
+++ b/Engine/Plugins/Mutable/Source/MutableRuntime/Private/MuR/ImageResizeUtils.h
@@ -0,0 +1,42 @@
+// Copyright Epic Games, Inc. All Rights Reserved.
+
+#pragma once
+
+#include "Math/MathFwd.h"
+#include "MuR/ImageTypes.h"
+
+namespace UE::Mutable::Private
+{
+
+namespace SubImageResize
+{
+	using SubImageResizeFuncType = void(FVector2f, FVector2f, FImageSize, FImageSize, FImageSize, FImageSize, uint8*, const uint8*); 
+
+	SubImageResizeFuncType* SelectSubImageResizeFunc(EImageFormat Format);
+
+	void SubImageResizeLinearByLessThanTwoVec1_U8(
+				FVector2f OneOverScalingFactor, 
+				FVector2f SrcSubPixelOffset,
+				FImageSize DestSize, FImageSize DestSubSize, 
+				FImageSize SrcSize, FImageSize SrcSubSize,
+				uint8* RESTRICT SubDest, 
+				const uint8* RESTRICT SubSrc);
+
+	void SubImageResizeLinearByLessThanTwoVec3_U8(
+				FVector2f OneOverScalingFactor, 
+				FVector2f SrcSubPixelOffset,
+				FImageSize DestSize, FImageSize DestSubSize, 
+				FImageSize SrcSize, FImageSize SrcSubSize,
+				uint8* RESTRICT SubDest, 
+				const uint8* RESTRICT SubSrc);
+
+	void SubImageResizeLinearByLessThanTwoVec4_U8(
+				FVector2f OneOverScalingFactor, 
+				FVector2f SrcSubPixelOffset,
+				FImageSize DestSize, FImageSize DestSubSize, 
+				FImageSize SrcSize, FImageSize SrcSubSize,
+				uint8* RESTRICT SubDest, 
+				const uint8* RESTRICT SubSrc);
+
+} // namespace SubImageResize
+} // namespace UE::Mutable::Private
diff --git a/Engine/Plugins/Mutable/Source/MutableRuntime/Private/MuR/OpImageMipmap.cpp b/Engine/Plugins/Mutable/Source/MutableRuntime/Private/MuR/OpImageMipmap.cpp
--- a/Engine/Plugins/Mutable/Source/MutableRuntime/Private/MuR/OpImageMipmap.cpp
+++ b/Engine/Plugins/Mutable/Source/MutableRuntime/Private/MuR/OpImageMipmap.cpp
@@ -1,10 +1,12 @@
 // Copyright Epic Games, Inc. All Rights Reserved.
 
 #include "MuR/ImagePrivate.h"
+#include "MuR/ImageFormatUtils.h"
 #include "MuR/SystemPrivate.h"
-#include "Async/ParallelFor.h"
+#include "MuR/ParallelExecutionUtils.h"
 #include "HAL/UnrealMemory.h"
 
+
 namespace
 {
 
@@ -259,15 +261,15 @@ namespace OpImageMipmapInternal
 		StagingMemory.SetNum(PerJobStagingBytes*NumParallelJobs);
 		uint8 * const StagingMemoryData = StagingMemory.GetData();
 
-		miro::SubImageDecompression::FuncRefType DecompressionFunc = SelectDecompressionFunction(DestFormat, SrcFormat);
+		miro::SubImageDecompression::FuncRefType DecompressionFunc = SubImageFormat::SelectDecompressionFunction(DestFormat, SrcFormat);
 
-		auto ProcessJob = 
-			[
-				NumParallelJobs, NumRowBatchesPerJob, 
-				StagingMemoryData, PerJobStagingBytes, 
-				NumBatches, NumBlocks, PixelsPerBlock, BlockSizeInBytes, DecompressionFunc,
-				Src, SrcSize, Dest, DestSize
-			](int32 JobId)
+		ParallelExecutionUtils::InvokeBatchParallelFor(NumParallelJobs, 
+		[
+			NumParallelJobs, NumRowBatchesPerJob, 
+			StagingMemoryData, PerJobStagingBytes, 
+			NumBatches, NumBlocks, PixelsPerBlock, BlockSizeInBytes, DecompressionFunc,
+			Src, SrcSize, Dest, DestSize
+		](int32 JobId)
 		{
 			const int32 JobRowBegin = JobId*NumRowBatchesPerJob;
 			const int32 JobRowEnd   = FMath::Min(JobRowBegin + NumRowBatchesPerJob, NumBatches.Y);
@@ -352,16 +354,7 @@ namespace OpImageMipmapInternal
 					}
 				}
 			}
-		};
-
-		if (NumParallelJobs == 1)
-		{
-			ProcessJob(0);
-		}
-		else if (NumParallelJobs > 1)
-		{	
-			ParallelFor(NumParallelJobs, ProcessJob);
-		}
+		});
 	}
 } // namespace OpImageMipmapInternal
 
diff --git a/Engine/Plugins/Mutable/Source/MutableRuntime/Private/MuR/OpImageResize.cpp b/Engine/Plugins/Mutable/Source/MutableRuntime/Private/MuR/OpImageResize.cpp
--- a/Engine/Plugins/Mutable/Source/MutableRuntime/Private/MuR/OpImageResize.cpp
+++ b/Engine/Plugins/Mutable/Source/MutableRuntime/Private/MuR/OpImageResize.cpp
@@ -3,6 +3,8 @@
 #include "MuR/ImagePrivate.h"
 #include "MuR/MutableTrace.h"
 #include "MuR/ParallelExecutionUtils.h"
+#include "MuR/ImageResizeUtils.h"
+#include "MuR/ImageFormatUtils.h"
 
 #include "Async/ParallelFor.h"
 #include "HAL/Platform.h"
@@ -706,207 +708,6 @@ namespace UE::Mutable::Private
 
 namespace
 {
-
-	template<uint32 NumChannels>
-	FORCENOINLINE void ImageResizeLinearDownByLessThanTwoSubImageVectorImpl(
-			FVector2f InOneOverScalingFactor, FVector2f InSrcSubPixelOffset,
-			FImageSize DestSize, FImageSize DestSubSize, FImageSize SrcSize, FImageSize SrcSubSize,
-		  	uint8* RESTRICT SubDest, const uint8* RESTRICT SubSrc)
-	{	
-		static_assert(NumChannels > 0 && NumChannels <= 4);
-	
-		// For some reason (probably cache related), doing this copy yields marginally but consistently 
-		// better performance than using the function argument. 
-		const FVector2f OneOverScalingFactor = InOneOverScalingFactor;
-		const FVector2f SrcSubPixelOffset = InSrcSubPixelOffset;
-		
-		check(OneOverScalingFactor.X <= 2.0f && OneOverScalingFactor.Y <= 2.0f && 
-			  OneOverScalingFactor.X >= 1.0f && OneOverScalingFactor.Y >= 1.0f);
-
-		for (uint16 Y = 0; Y < DestSubSize.Y; ++Y)
-		{
-			for (uint16 X = 0; X < DestSubSize.X; ++X)
-			{
-				auto LoadPixel = [](const uint8* Ptr) -> VectorRegister4Float
-				{	
-					alignas(4) uint8 PixelData[4] = {0};
-					if constexpr (NumChannels == 4)
-					{
-						FMemory::Memcpy(&(PixelData[0]), Ptr, 4);
-					}
-					else
-					{
-						for (uint32 C = 0; C < NumChannels; ++C)
-						{
-							PixelData[C] = Ptr[C];
-						}
-					}
-
-					return VectorIntToFloat(MakeVectorRegisterInt(PixelData[0], PixelData[1], PixelData[2], PixelData[3]));
-				};
-				
-				const FVector2f CoordsF = FVector2f(X, Y)*OneOverScalingFactor + SrcSubPixelOffset;
-	
-				const FIntVector2 Coords = FIntVector2(FMath::FloorToInt(CoordsF.X), FMath::FloorToInt(CoordsF.Y));
-				const FIntVector2 CoordsPlusOne = FIntVector2(
-						FMath::Min<int32>(Coords.X + 1, SrcSize.X - 1),
-						FMath::Min<int32>(Coords.Y + 1, SrcSize.Y - 1));
-
-				const VectorRegister4Float Pixel00 = LoadPixel(SubSrc + (Coords.Y*SrcSize.X + Coords.X) * NumChannels);
-				const VectorRegister4Float Pixel10 = LoadPixel(SubSrc + (Coords.Y*SrcSize.X + CoordsPlusOne.X) * NumChannels);
-				const VectorRegister4Float Pixel01 = LoadPixel(SubSrc + (CoordsPlusOne.Y*SrcSize.X + Coords.X) * NumChannels);
-				const VectorRegister4Float Pixel11 = LoadPixel(SubSrc + (CoordsPlusOne.Y*SrcSize.X + CoordsPlusOne.X) * NumChannels);
-
-				const VectorRegister4Float CoordsFVector = MakeVectorRegister(CoordsF.X, CoordsF.Y, 0.0f, 0.0f);
-				const VectorRegister4Float Frac = VectorSubtract(CoordsFVector, VectorFloor(CoordsFVector));
-
-				const VectorRegister4Float FracX = VectorReplicate(Frac, 0);
-				const VectorRegister4Float FracY = VectorReplicate(Frac, 1);
-				
-				VectorRegister4Float Result = VectorMultiplyAdd(FracX, VectorSubtract(Pixel10, Pixel00), Pixel00);
-				Result = VectorMultiplyAdd(
-						FracY,
-						VectorSubtract(VectorMultiplyAdd(FracX, VectorSubtract(Pixel11, Pixel01), Pixel01), Result),
-						Result);
-
-				uint8* Dest = SubDest + (Y*DestSize.X + X)*NumChannels;  
-
-				if constexpr (NumChannels == 4)
-				{
-					VectorStoreByte4(Result, Dest);
-					VectorResetFloatRegisters();
-				}
-				else
-				{
-					alignas(4) uint8 ResultData[4];
-					VectorStoreByte4(Result, ResultData);
-					VectorResetFloatRegisters();
-				
-					for (uint32 C = 0; C < NumChannels; ++C)
-					{
-						Dest[C] = ResultData[C];
-					}
-				}
-			}
-		}
-	}
-
-
-	template<uint32 NumChannels>
-	FORCENOINLINE void ImageResizeLinearDownByLessThanTwoSubImageNonVectorImpl(
-			FVector2f InOneOverScalingFactor, FVector2f InSrcSubPixelOffset,
-			FImageSize DestSize, FImageSize DestSubSize, FImageSize SrcSize, FImageSize SrcSubSize,
-		  	uint8* RESTRICT SubDest, const uint8* RESTRICT SubSrc)
-	{
-		struct alignas(uint64) FPixelData
-		{
-			uint16 Data[4] = {0};
-		};
-
-		static_assert(NumChannels <= 4);
-		
-		// For some reason (probably cache related), doing this copy yields marginally but consistently 
-		// better performance than using the function argument. 
-		const FVector2f OneOverScalingFactor = InOneOverScalingFactor;
-		const FVector2f SrcSubPixelOffset = InSrcSubPixelOffset;
-		
-		check(OneOverScalingFactor.X <= 2.0f && OneOverScalingFactor.Y <= 2.0f && 
-			  OneOverScalingFactor.X >= 1.0f && OneOverScalingFactor.Y >= 1.0f);
-
-		for (uint16 Y = 0; Y < DestSubSize.Y; ++Y)
-		{
-			for (uint16 X = 0; X < DestSubSize.X; ++X)
-			{
-				auto LoadPixel = [](const uint8* Ptr) -> FPixelData
-				{
-					FPixelData Result;
-					if constexpr (NumChannels == 4)
-					{
-						uint32 PackedData;
-						FMemory::Memcpy(&PackedData, Ptr, sizeof(uint32));
-
-						Result.Data[0] = static_cast<uint16>((PackedData >> (8 * 0)) & 0xFF);
-						Result.Data[1] = static_cast<uint16>((PackedData >> (8 * 1)) & 0xFF);
-						Result.Data[2] = static_cast<uint16>((PackedData >> (8 * 2)) & 0xFF);
-						Result.Data[3] = static_cast<uint16>((PackedData >> (8 * 3)) & 0xFF);
-					}
-					else
-					{
-						for (int32 C = 0; C < NumChannels; ++C)
-						{
-							Result.Data[C] = static_cast<uint16>(Ptr[C]);
-						}
-					}
-
-					return Result;
-				};
-				
-				const FVector2f CoordsF = FVector2f(X, Y)*OneOverScalingFactor + SrcSubPixelOffset;
-				
-				using FUint16Vector2 = UE::Math::TIntVector2<uint16>;
-				const FUint16Vector2 FracINorm = FUint16Vector2(FMath::Frac(CoordsF.X) * 255.0f, FMath::Frac(CoordsF.Y) * 255.0f);
-				
-				const FIntVector2 Coords = FIntVector2(FMath::FloorToInt(CoordsF.X), FMath::FloorToInt(CoordsF.Y));
-				const FIntVector2 CoordsPlusOne = FIntVector2(
-						FMath::Min<int32>(Coords.X + 1, SrcSize.X - 1),
-						FMath::Min<int32>(Coords.Y + 1, SrcSize.Y - 1));
-
-				FPixelData PixelData00 = LoadPixel(SubSrc + (Coords.Y*SrcSize.X + Coords.X) * NumChannels);
-				FPixelData PixelData10 = LoadPixel(SubSrc + (Coords.Y*SrcSize.X + CoordsPlusOne.X) * NumChannels);
-				FPixelData PixelData01 = LoadPixel(SubSrc + (CoordsPlusOne.Y*SrcSize.X + Coords.X) * NumChannels);
-				FPixelData PixelData11 = LoadPixel(SubSrc + (CoordsPlusOne.Y*SrcSize.X + CoordsPlusOne.X) * NumChannels);
-				
-				FPixelData Result;
-				for (int32 C = 0; C < NumChannels; ++C)
-				{
-					const uint16 LerpY00 = ((PixelData10.Data[C] * FracINorm.X) + PixelData00.Data[C] * (255 - FracINorm.X)) / 255;
-					const uint16 LerpY01 = ((PixelData11.Data[C] * FracINorm.X) + PixelData01.Data[C] * (255 - FracINorm.X)) / 255;
-					Result.Data[C] = ((LerpY01 * FracINorm.Y) + LerpY00 * (255 - FracINorm.Y)) / 255;
-				}
-
-				uint8* Dest = SubDest + (Y*DestSize.X + X)*NumChannels;  
-
-				if constexpr (NumChannels == 4)
-				{
-					const uint32 PackedPixel = 
-						(static_cast<uint32>((Result.Data[0]) & 0xFF) << (8 * 0)) | 
-						(static_cast<uint32>((Result.Data[1]) & 0xFF) << (8 * 1)) | 
-						(static_cast<uint32>((Result.Data[2]) & 0xFF) << (8 * 2)) | 
-						(static_cast<uint32>((Result.Data[3]) & 0xFF) << (8 * 3));
-
-					FMemory::Memcpy(Dest, &PackedPixel, sizeof(uint32));
-				}
-				else
-				{
-					for (uint32 C = 0; C < NumChannels; ++C)
-					{
-						Dest[C] = static_cast<uint8>(Result.Data[C]);
-					}
-				}
-			}
-		}
-	}
-
-	template<int32 NumChannels, bool bUseVectorImpl>
-	FORCEINLINE void ImageResizeLinearDownByLessThanTwoSubImage(
-			FVector2f OneOverScalingFactor, FVector2f SrcSubPixelOffset,
-			FImageSize DestSize, FImageSize DestSubSize, FImageSize SrcSize, FImageSize SrcSubSize,
-			uint8* RESTRICT SubDest, const uint8* RESTRICT SubSrc)
-	{
-		if constexpr (bUseVectorImpl)
-		{
-			ImageResizeLinearDownByLessThanTwoSubImageVectorImpl<NumChannels>(
-					OneOverScalingFactor, SrcSubPixelOffset,
-					DestSize, DestSubSize, SrcSize, SrcSubSize, SubDest, SubSrc);
-		}
-		else
-		{
-			ImageResizeLinearDownByLessThanTwoSubImageNonVectorImpl<NumChannels>(
-					OneOverScalingFactor, SrcSubPixelOffset,
-					DestSize, DestSubSize, SrcSize, SrcSubSize, SubDest, SubSrc);
-		}
-	}
-
 	template<bool bUseVectorImpl>
 	void ImageResizeLinearDownByLessThanTwo(FImage* Dest, const FImage* InBase)
 	{
@@ -932,38 +733,13 @@ namespace
 		{
 			MUTABLE_CPUPROFILER_SCOPE(ResizeUnCompressed);
 
+			SubImageResize::SubImageResizeFuncType* SubImageResizeFunc = SubImageResize::SelectSubImageResizeFunc(SrcFormat);
+
 			const FVector2f SrcSubPixelOffset = FVector2f(0.0f);
 			// Uncompressed format, resize in one go directly to Dest.
-			switch(SrcFormat)
-			{
-			case EImageFormat::L_UByte:
-			{
-				constexpr int32 NumChannels = 1;
-				ImageResizeLinearDownByLessThanTwoSubImage<NumChannels, bUseVectorImpl>(
-						OneOverScalingFactor, SrcSubPixelOffset,
-						DestSize, DestSize, SrcSize, SrcSize, DestData, SrcData);
-				break;
-			}
-
-			case EImageFormat::RGB_UByte:
-			{
-				constexpr int32 NumChannels = 3;
-				ImageResizeLinearDownByLessThanTwoSubImage<NumChannels, bUseVectorImpl>(
+			SubImageResizeFunc(
 						OneOverScalingFactor, SrcSubPixelOffset,
 						DestSize, DestSize, SrcSize, SrcSize, DestData, SrcData);
-				break;
-			}
-
-			case EImageFormat::RGBA_UByte:
-			case EImageFormat::BGRA_UByte:
-			{
-				constexpr int32 NumChannels = 4;
-				ImageResizeLinearDownByLessThanTwoSubImage<NumChannels, bUseVectorImpl>(
-						OneOverScalingFactor, SrcSubPixelOffset,
-						DestSize, DestSize, SrcSize, SrcSize, DestData, SrcData);
-				break;
-			}
-			}
 		}
 		else if (SrcFormatData.PixelsPerBlockX > 0 && SrcFormatData.PixelsPerBlockY > 0)
 		{
@@ -1012,14 +788,15 @@ namespace
 			StagingMemory.SetNum(PerJobStagingBytes * NumParallelJobs);
 			uint8 * const StagingMemoryData = StagingMemory.GetData();
 
-			miro::SubImageDecompression::FuncRefType DecompressionFunc = SelectDecompressionFunction(DestFormat, SrcFormat);
+			miro::SubImageDecompression::FuncRefType DecompressionFunc = SubImageFormat::SelectDecompressionFunction(DestFormat, SrcFormat);
+			SubImageResize::SubImageResizeFuncType* SubImageResizeFunc = SubImageResize::SelectSubImageResizeFunc(DestFormat);
 
 			ParallelExecutionUtils::InvokeBatchParallelFor(NumParallelJobs, 
 			[
 				NumParallelJobs, NumRowBatchesPerJob, 
 				StagingMemoryData, PerJobStagingBytes, 
 				NumBatches, NumBlocks, PixelsPerBlock, BlockSizeInBytes, DestNumChannels, DecompressionFunc,
-				SrcData, SrcSize, DestData, DestSize, OneOverScalingFactor
+				SubImageResizeFunc, SrcData, SrcSize, DestData, DestSize, OneOverScalingFactor
 			](int32 JobId)
 			{
 				const int32 JobRowBegin = JobId*NumRowBatchesPerJob;
@@ -1081,44 +858,12 @@ namespace
 						const FIntVector2& DestSubOffset = BatchOutBeginInPixels;
 					
 						const uint8* SrcSubData = JobStagingMemoryData;
-						switch(DestNumChannels)
-						{
-						case 1:
-						{
-							constexpr int32 NumChannels = 1;
-						
-							uint8* DestSubData = DestData + (DestSubOffset.Y*DestSize.X + DestSubOffset.X)*NumChannels;
 
-							ImageResizeLinearDownByLessThanTwoSubImage<NumChannels, bUseVectorImpl>( 
-									OneOverScalingFactor, SrcSubPixelOffset,
-									DestSize, DestSubSize, DecSize, SrcSubSize, DestSubData, SrcSubData);
-							break;
-						}
+						uint8* DestSubData = DestData + (DestSubOffset.Y*DestSize.X + DestSubOffset.X)*DestNumChannels;
 
-						case 3:
-						{
-							constexpr int32 NumChannels = 3;
-
-							uint8* DestSubData = DestData + (DestSubOffset.Y*DestSize.X + DestSubOffset.X)*NumChannels;
-							
-							ImageResizeLinearDownByLessThanTwoSubImage<NumChannels, bUseVectorImpl>( 
-									OneOverScalingFactor, SrcSubPixelOffset,
-									DestSize, DestSubSize, DecSize, SrcSubSize, DestSubData, SrcSubData);
-							break;
-						}
-
-						case 4:
-						{
-							constexpr int32 NumChannels = 4;
-
-							uint8* DestSubData = DestData + (DestSubOffset.Y*DestSize.X + DestSubOffset.X)*NumChannels;
-							
-							ImageResizeLinearDownByLessThanTwoSubImage<NumChannels, bUseVectorImpl>(
-									OneOverScalingFactor, SrcSubPixelOffset,
-									DestSize, DestSubSize, DecSize, SrcSubSize, DestSubData, SrcSubData);
-							break;
-						}
-						}
+						SubImageResizeFunc( 
+								OneOverScalingFactor, SrcSubPixelOffset,
+								DestSize, DestSubSize, DecSize, SrcSubSize, DestSubData, SrcSubData);
 					}
 				}
 			});
