[Mutable] Multicompose operation.

- Implementation of the multicompose operation that composes images using multiple blocks from the same image. Currently it allocates as much memory as the largest block decompressed. The implementation is structured so adding batching to reduce memory requirements and parallelize should be relativly simple.

- Now destination block is fully decompressed if block compression alignment or a resize is needed. This is not neccessary as only a small portion at the edeges is actually needed.

#rnx

diff --git a/Engine/Plugins/Mutable/Source/MutableRuntime/Private/MuR/OpImageMultiCompose.cpp b/Engine/Plugins/Mutable/Source/MutableRuntime/Private/MuR/OpImageMultiCompose.cpp
new file mode 100644
--- /dev/null
+++ b/Engine/Plugins/Mutable/Source/MutableRuntime/Private/MuR/OpImageMultiCompose.cpp
@@ -0,0 +1,560 @@
+// Copyright Epic Games, Inc. All Rights Reserved.
+
+#include "MuR/OpImageMultiCompose.h"
+
+#include "MuR/ImagePrivate.h"
+#include "MuR/MutableMath.h"
+#include "MuR/ImageFormatUtils.h"
+#include "MuR/ImageResizeUtils.h"
+
+namespace UE::Mutable::Private
+{
+
+namespace OpImageMultiComposeInternal
+{
+	enum class EActiveStages : uint32
+	{
+		None             = 0 << 0,
+		SrcDecompression = 1 << 0,
+		DstDecompression = 1 << 1,
+		Resize           = 1 << 2,
+		DstCompression   = 1 << 3,
+	};
+	ENUM_CLASS_FLAGS(EActiveStages);
+
+	struct FImageFormatInfo
+	{
+		const FImageFormatData& Dst;
+		const FImageFormatData& Src;
+		const FImageFormatData& DstDecompressed;
+
+		EImageFormat SrcDecompressedFormat;
+		EImageFormat DstDecompressedFormat;
+		EImageFormat MostGenericFormat;
+
+		miro::SubImageDecompression::FuncRefType SrcDecompressionFunc;
+		miro::SubImageDecompression::FuncRefType DstDecompressionFunc;
+		miro::SubImageCompression::FuncRefType DstCompressionFunc;
+		SubImageResize::SubImageResizeFuncType* ResizeFunc;
+		SubImageFormat::UncompressedFormatFuncType* SrcUncompressedFormatFunc;
+		SubImageFormat::UncompressedFormatFuncType* DstUncompressedFormatFunc;
+	};
+
+	template<typename T>
+	struct TSubImage
+	{
+		TArrayView<T> DataView;
+
+		FImageSize Size          = FImageSize{0, 0};
+		FImageSize SubSize       = FImageSize{0, 0};
+		uint16 BytesPerBlock     = 0;
+		uint8  NumElemsPerBlockX = 0;
+		uint8  NumElemsPerBlockY = 0;
+
+		TSubImage<const T> ToConst()
+		{
+			return TSubImage<const T> 
+			{
+				.DataView          = DataView         ,
+
+				.Size              = Size             ,
+				.SubSize           = SubSize          ,
+				.BytesPerBlock     = BytesPerBlock    ,
+				.NumElemsPerBlockX = NumElemsPerBlockX,
+				.NumElemsPerBlockY = NumElemsPerBlockY,
+			};
+		}
+	};
+
+	template<typename T>
+	using TConstSubImage = TSubImage<const T>;
+
+	struct FScratchMemoryAllocation
+	{
+		EActiveStages ActiveStages = EActiveStages::None; 
+		
+		TArray<uint8, FDefaultMemoryTrackingAllocator<MemoryCounters::FImageMemoryCounter>> Buffer;
+
+		TArrayView<uint8> SrcDecompressed;
+		TArrayView<uint8> DstDecompressed;
+	};
+
+	FImageFormatInfo GatherImageFormatInfo(const FImage& Dst, const FImage& Src)
+	{
+		using namespace SubImageFormat;
+
+		EImageFormat SrcDecompressedFormat = GetUncompressedFormat(Src.GetFormat());
+		EImageFormat DstDecompressedFormat = GetUncompressedFormat(Dst.GetFormat());
+		EImageFormat MostGenericFormat     = GetMostGenericFormat(SrcDecompressedFormat, DstDecompressedFormat);
+		return FImageFormatInfo
+		{
+			.Dst                   = GetImageFormatData(Dst.GetFormat()),
+			.Src                   = GetImageFormatData(Src.GetFormat()),
+			.DstDecompressed       = GetImageFormatData(DstDecompressedFormat),
+			.SrcDecompressedFormat = SrcDecompressedFormat,
+			.DstDecompressedFormat = DstDecompressedFormat,
+			.MostGenericFormat     = MostGenericFormat,
+			.SrcDecompressionFunc  = SelectDecompressionFunction(DstDecompressedFormat, Src.GetFormat()),
+			.DstDecompressionFunc  = SelectDecompressionFunction(DstDecompressedFormat, Dst.GetFormat()),
+			.DstCompressionFunc    = SelectCompressionFunction(Dst.GetFormat(), DstDecompressedFormat),
+			.ResizeFunc            = SubImageResize::SelectSubImageResizeFunc(DstDecompressedFormat),
+			.SrcUncompressedFormatFunc = SubImageFormat::SelectUncompressedFormatFunction(SrcDecompressedFormat, DstDecompressedFormat),
+			.DstUncompressedFormatFunc = SubImageFormat::SelectUncompressedFormatFunction(DstDecompressedFormat, MostGenericFormat),
+		};
+	}
+
+	EActiveStages ComputeActiveStages(const FImage& Base, const FImage& Source, const FImageFormatInfo& FormatInfo, const FIntRect& DstRect, const FIntRect& SrcRect)
+	{
+		const bool bIsSrcCompressionRectAligned = 
+				((SrcRect.Min.X % FormatInfo.Src.PixelsPerBlockX) == 0) &
+				((SrcRect.Min.Y % FormatInfo.Src.PixelsPerBlockY) == 0) &
+				((SrcRect.Max.X % FormatInfo.Src.PixelsPerBlockX) == 0) &
+				((SrcRect.Max.Y % FormatInfo.Src.PixelsPerBlockY) == 0);
+
+		const bool bIsDstCompressionRectAligned = 
+				((DstRect.Min.X % FormatInfo.Dst.PixelsPerBlockX) == 0) &
+				((DstRect.Min.Y % FormatInfo.Dst.PixelsPerBlockY) == 0) &
+				((DstRect.Max.X % FormatInfo.Dst.PixelsPerBlockX) == 0) &
+				((DstRect.Max.Y % FormatInfo.Dst.PixelsPerBlockY) == 0);
+
+		const bool bUncompressedFormatMissmatch =
+				FormatInfo.SrcDecompressedFormat != FormatInfo.DstDecompressedFormat && 
+				FormatInfo.SrcDecompressedFormat == Source.GetFormat() && FormatInfo.DstDecompressedFormat == Base.GetFormat();
+
+		const bool bNeedsResize = 
+				(DstRect.Width()  != SrcRect.Width()) |
+				(DstRect.Height() != SrcRect.Height());
+
+		const bool bNeedsSrcDecompression = 
+				(bNeedsResize & (FormatInfo.SrcDecompressedFormat != Source.GetFormat())) |
+				(!(bIsSrcCompressionRectAligned & bIsDstCompressionRectAligned))          |
+				bUncompressedFormatMissmatch;
+
+		const bool bNeedsDstDecompression =
+				(bNeedsResize & (FormatInfo.DstDecompressedFormat != Base.GetFormat())) |
+				(!(bIsSrcCompressionRectAligned & bIsDstCompressionRectAligned));
+
+		const bool bNeedsDstCompression =
+				(FormatInfo.DstDecompressedFormat != Base.GetFormat()) & bNeedsSrcDecompression;
+		
+		EActiveStages Result = EActiveStages::None;
+		EnumAddFlags(Result, bNeedsResize           ? EActiveStages::Resize           : EActiveStages::None);
+		EnumAddFlags(Result, bNeedsSrcDecompression ? EActiveStages::SrcDecompression : EActiveStages::None);
+		EnumAddFlags(Result, bNeedsDstCompression   ? EActiveStages::DstCompression   : EActiveStages::None);
+		EnumAddFlags(Result, bNeedsDstDecompression ? EActiveStages::DstDecompression : EActiveStages::None);
+
+		return Result;
+	}
+
+	int32 AlignDown(int32 Value, int32 Align)
+	{
+		return FMath::DivideAndRoundDown(Value, Align) * Align;
+	}
+
+	int32 AlignUp(int32 Value, int32 Align)
+	{
+		return FMath::DivideAndRoundUp(Value, Align) * Align;
+	}
+
+	FScratchMemoryAllocation AllocateScratchMemory(const FImage& Dst, const FImage& Src, TConstArrayView<FIntRect> DstRects, TConstArrayView<FIntRect> SrcRects, const FImageFormatInfo& FormatInfo)
+	{
+		MUTABLE_CPUPROFILER_SCOPE(ImageMultiCompose_AllocateScratchMemory);
+		
+		// TODO: For now allocate for the largest block. This should allocate only for a specific memory target
+		// splitting the work in smaller batches.
+
+		FScratchMemoryAllocation Result;
+		
+		uint32 SrcDecompressedAllocSize = 0; 
+		uint32 DstDecompressedAllocSize = 0; 
+
+		EActiveStages AllocationActiveStages = EActiveStages::None;
+
+		for (int32 RectIndex = 0, NumRects = DstRects.Num(); RectIndex < NumRects; ++RectIndex)
+		{
+			FIntRect DstRect = DstRects[RectIndex];
+			FIntRect SrcRect = SrcRects[RectIndex];
+
+			FIntRect AlignedSrcRect = FIntRect(
+					AlignDown(SrcRect.Min.X, FormatInfo.Src.PixelsPerBlockX),
+					AlignDown(SrcRect.Min.Y, FormatInfo.Src.PixelsPerBlockY),
+					AlignUp  (SrcRect.Max.X, FormatInfo.Src.PixelsPerBlockX),
+					AlignUp  (SrcRect.Max.Y, FormatInfo.Src.PixelsPerBlockY));
+			
+			FIntRect AlignedDstRect = FIntRect(
+					AlignDown(DstRect.Min.X, FormatInfo.Dst.PixelsPerBlockX),
+					AlignDown(DstRect.Min.Y, FormatInfo.Dst.PixelsPerBlockY),
+					AlignUp  (DstRect.Max.X, FormatInfo.Dst.PixelsPerBlockX),
+					AlignUp  (DstRect.Max.Y, FormatInfo.Dst.PixelsPerBlockY));
+
+			uint32 AlignedDecSrcRectBytes = AlignedSrcRect.Area()*FormatInfo.DstDecompressed.BytesPerBlock;
+			uint32 AlignedDecDstRectBytes = AlignedDstRect.Area()*FormatInfo.DstDecompressed.BytesPerBlock;
+
+			SrcDecompressedAllocSize = FMath::Max(AlignedDecSrcRectBytes, SrcDecompressedAllocSize);
+			DstDecompressedAllocSize = FMath::Max(AlignedDecDstRectBytes, DstDecompressedAllocSize);
+		}
+
+		Result.Buffer.SetNumUninitialized(SrcDecompressedAllocSize + DstDecompressedAllocSize);
+
+		Result.SrcDecompressed = TArrayView(Result.Buffer.GetData(), SrcDecompressedAllocSize);
+		Result.DstDecompressed = TArrayView(Result.Buffer.GetData() + SrcDecompressedAllocSize, DstDecompressedAllocSize);    
+
+		// TODO: We are allocating for all possible stages but should only allocate for the used ones. 
+		Result.ActiveStages = EActiveStages::None;
+		EnumAddFlags(Result.ActiveStages, EActiveStages::Resize | EActiveStages::SrcDecompression | EActiveStages::DstCompression | EActiveStages::DstDecompression);
+
+		return Result;
+	}
+
+	void CopySubImage(const TSubImage<uint8>& Dst, const TConstSubImage<uint8>& Src)
+	{
+		check(Dst.SubSize == Src.SubSize);
+		check(Dst.BytesPerBlock == Src.BytesPerBlock);
+		check(Dst.NumElemsPerBlockX == Src.NumElemsPerBlockX);
+		check(Dst.NumElemsPerBlockY == Src.NumElemsPerBlockY);
+
+		FIntVector2 DstSizeInBlocks = FIntVector2(
+				Dst.Size.X / Dst.NumElemsPerBlockX, Dst.Size.Y / Dst.NumElemsPerBlockY);
+
+		FIntVector2 SrcSizeInBlocks = FIntVector2(
+				Src.Size.X / Src.NumElemsPerBlockX, Src.Size.Y / Src.NumElemsPerBlockY);
+
+		FIntVector2 SubSizeInBlocks = FIntVector2(
+				Dst.SubSize.X / Dst.NumElemsPerBlockX, Dst.SubSize.Y / Dst.NumElemsPerBlockY);
+
+		const int32 DstPitchInBytes = DstSizeInBlocks.X*Dst.BytesPerBlock;
+		const int32 SrcPitchInBytes = SrcSizeInBlocks.X*Src.BytesPerBlock;
+		const int32 RowSizeInBytes = SubSizeInBlocks.X*Dst.BytesPerBlock;
+
+		for (int32 Y = 0; Y < SubSizeInBlocks.Y; ++Y)
+		{
+			FMemory::Memcpy(Dst.DataView.GetData() + Y*DstPitchInBytes, Src.DataView.GetData() + Y*SrcPitchInBytes, RowSizeInBytes);
+		}
+	}
+
+	int32 ComputeNumBatches(const FIntRect& DstRect)
+	{
+		//TODO: Do batching to reduce memory requirements and parallelise.
+		return 1;	
+	}
+
+	struct FBatchWorkingSubImages
+	{
+		TConstSubImage<uint8> Src;
+		TSubImage<uint8> Dst;
+		TSubImage<uint8> DecSrc;
+		TSubImage<uint8> DecDst;
+		TSubImage<uint8> DecSrcAligned;
+		TSubImage<uint8> DecDstAligned;
+
+		FVector2f SrcSubPixelOffset = FVector2f(0, 0);
+	};
+
+	FBatchWorkingSubImages ComputeWorkingSubImages(
+			int32 BatchIndex, 
+			FIntRect DstRect, FIntRect SrcRect, 
+			const TSubImage<uint8>& DstImage, const TConstSubImage<uint8>& SrcImage, 
+			const FImageFormatInfo& FormatInfo, 
+			const FScratchMemoryAllocation& ScratchMemory)
+	{
+		FBatchWorkingSubImages Result;
+
+		// TODO: For now assume there is only one batch for the whole Rect. This should be split in smaller batches
+		// to reduce memory requirements.
+		check(BatchIndex < 1); 
+		FIntRect DstBatchRect = DstRect;
+
+		FIntRect SrcBatchRect = SrcRect;
+
+		check(SrcBatchRect.Area() == SrcRect.Area())
+
+		FIntRect AlignedSrcRect = FIntRect(
+				AlignDown(SrcBatchRect.Min.X, SrcImage.NumElemsPerBlockX),
+				AlignDown(SrcBatchRect.Min.Y, SrcImage.NumElemsPerBlockY),
+				AlignUp  (SrcBatchRect.Max.X, SrcImage.NumElemsPerBlockX),
+				AlignUp  (SrcBatchRect.Max.Y, SrcImage.NumElemsPerBlockY));
+			
+		FIntRect AlignedDstRect = FIntRect(
+				AlignDown(DstBatchRect.Min.X, DstImage.NumElemsPerBlockX),
+				AlignDown(DstBatchRect.Min.Y, DstImage.NumElemsPerBlockY),
+				AlignUp  (DstBatchRect.Max.X, DstImage.NumElemsPerBlockX),
+				AlignUp  (DstBatchRect.Max.Y, DstImage.NumElemsPerBlockY));
+
+		int32 SrcDataOffsetInBlocks = (AlignedSrcRect.Min.Y / SrcImage.NumElemsPerBlockY) * (SrcImage.Size.X / SrcImage.NumElemsPerBlockX) + (AlignedSrcRect.Min.X / SrcImage.NumElemsPerBlockX); 
+		int32 SrcDataOffsetInBytes =  SrcDataOffsetInBlocks * SrcImage.BytesPerBlock;
+
+		Result.Src = TConstSubImage<uint8> 
+		{
+			.DataView      = TConstArrayView<uint8>(SrcImage.DataView.GetData() + SrcDataOffsetInBytes, SrcImage.DataView.Num() - SrcDataOffsetInBytes),
+			.Size          = SrcImage.Size,
+			.SubSize       = FImageSize(AlignedSrcRect.Width(), AlignedSrcRect.Height()),
+			.BytesPerBlock = SrcImage.BytesPerBlock,
+			.NumElemsPerBlockX = SrcImage.NumElemsPerBlockX,
+			.NumElemsPerBlockY = SrcImage.NumElemsPerBlockY,
+		};
+
+		int32 DstDataOffsetInBlocks = (AlignedDstRect.Min.Y / DstImage.NumElemsPerBlockY) * (DstImage.Size.X / DstImage.NumElemsPerBlockX) + (AlignedDstRect.Min.X / DstImage.NumElemsPerBlockX); 
+		int32 DstDataOffsetInBytes =  DstDataOffsetInBlocks * DstImage.BytesPerBlock;
+		Result.Dst = TSubImage<uint8>
+		{
+			.DataView      = TArrayView(DstImage.DataView.GetData() + DstDataOffsetInBytes, DstImage.DataView.Num() - DstDataOffsetInBytes),
+			.Size          = DstImage.Size,
+			.SubSize       = FImageSize(AlignedDstRect.Width(), AlignedDstRect.Height()),
+			.BytesPerBlock = DstImage.BytesPerBlock,
+			.NumElemsPerBlockX = DstImage.NumElemsPerBlockX,
+			.NumElemsPerBlockY = DstImage.NumElemsPerBlockY,
+		};
+
+		if (!ScratchMemory.SrcDecompressed.IsEmpty())
+		{
+			Result.DecSrcAligned = TSubImage<uint8>
+			{
+				.DataView      = TArrayView(ScratchMemory.SrcDecompressed),
+				.Size          = FImageSize(AlignedSrcRect.Width(), AlignedSrcRect.Height()),
+				.SubSize       = FImageSize(AlignedSrcRect.Width(), AlignedSrcRect.Height()),
+				.BytesPerBlock = FormatInfo.DstDecompressed.BytesPerBlock, 
+				.NumElemsPerBlockX = 1,
+				.NumElemsPerBlockY = 1,
+			};
+		}
+
+		if (!ScratchMemory.DstDecompressed.IsEmpty())
+		{
+			Result.DecDstAligned = TSubImage<uint8>
+			{
+				.DataView      = TArrayView(ScratchMemory.DstDecompressed),
+				.Size          = FImageSize(AlignedDstRect.Width(), AlignedDstRect.Height()),
+				.SubSize       = FImageSize(AlignedDstRect.Width(), AlignedDstRect.Height()),
+				.BytesPerBlock = FormatInfo.DstDecompressed.BytesPerBlock, 
+				.NumElemsPerBlockX = 1,
+				.NumElemsPerBlockY = 1,
+			};
+		}
+
+		if (!Result.DecDstAligned.DataView.IsEmpty())
+		{
+			FIntVector2 DstOffset = DstBatchRect.Min - AlignedDstRect.Min;
+			int32 DecDstOffsetInBytes = (DstOffset.Y*Result.DecDstAligned.Size.X + DstOffset.X) * Result.DecDstAligned.BytesPerBlock;
+			
+			Result.DecDst = TSubImage<uint8>
+			{
+				
+				.DataView      = TArrayView(Result.DecDstAligned.DataView.GetData() + DecDstOffsetInBytes, Result.DecDstAligned.DataView.Num() - DecDstOffsetInBytes),
+				.Size          = Result.DecDstAligned.Size,
+				.SubSize       = FImageSize(DstBatchRect.Width(), DstBatchRect.Height()),
+				.BytesPerBlock = FormatInfo.DstDecompressed.BytesPerBlock,
+				.NumElemsPerBlockX = 1,
+				.NumElemsPerBlockY = 1,
+			};
+		}
+
+		if (!Result.DecSrcAligned.DataView.IsEmpty())
+		{
+			FIntVector2 SrcOffset = SrcBatchRect.Min - AlignedSrcRect.Min;
+			int32 DecSrcOffsetInBytes = (SrcOffset.Y*Result.DecSrcAligned.Size.X + SrcOffset.X) * Result.DecSrcAligned.BytesPerBlock;
+
+			Result.DecSrc = TSubImage<uint8>
+			{
+				.DataView      = TArrayView(Result.DecSrcAligned.DataView.GetData() + DecSrcOffsetInBytes, Result.DecSrcAligned.DataView.Num() - DecSrcOffsetInBytes),
+				.Size          = Result.DecSrcAligned.Size,
+				.SubSize       = FImageSize(SrcBatchRect.Width(), SrcBatchRect.Height()),
+				.BytesPerBlock = FormatInfo.DstDecompressed.BytesPerBlock,
+				.NumElemsPerBlockX = 1,
+				.NumElemsPerBlockY = 1,
+			};
+		}
+
+		return Result;
+	}
+} // namespace OpImageMultiComposeInternal
+
+	void ImageMultiCompose(
+			FImage& Base, 
+			const FImage& Source, 
+			TConstArrayView<FIntRect> DstRects, 
+			TConstArrayView<FIntRect> SrcRects,
+			TConstArrayView<int32> SourceMipPerRect,
+			int32 CompressionQuality)
+	{
+		using namespace OpImageMultiComposeInternal;
+
+		check(DstRects.Num() == SrcRects.Num());
+		check(SourceMipPerRect.Num() == SrcRects.Num());
+
+		if (!DstRects.Num())
+		{
+			return;
+		}
+
+		FImageFormatInfo FormatInfo = GatherImageFormatInfo(Base, Source);
+
+		FScratchMemoryAllocation ScratchMemory = AllocateScratchMemory(Base, Source, DstRects, SrcRects, FormatInfo);
+
+		for (int32 RectIndex = 0; RectIndex < SrcRects.Num(); ++RectIndex)
+		{
+			FIntRect DstRect = DstRects[RectIndex];
+			FIntRect SrcRect = SrcRects[RectIndex];
+
+			if (DstRect.Area() == 0)
+			{
+				continue;
+			}
+
+			EActiveStages ActiveStages = ComputeActiveStages(Base, Source, FormatInfo, DstRect, SrcRect);
+			check(EnumHasAllFlags(ScratchMemory.ActiveStages, ActiveStages)); 
+
+			int32 NumBatches = ComputeNumBatches(DstRect);
+
+			uint32 SourceLODIndex = SourceMipPerRect[RectIndex];
+
+			FIntVector2 SourceLODSize = Source.CalculateMipSize(SourceLODIndex); 
+
+			TConstSubImage<uint8> SrcSubImage = TConstSubImage<uint8>
+			{
+				.DataView = Source.DataStorage.GetLOD(SourceLODIndex),
+				.Size     = FImageSize(SourceLODSize.X, SourceLODSize.Y),
+				.SubSize  = FImageSize(SourceLODSize.X, SourceLODSize.Y), 
+				.BytesPerBlock     = FormatInfo.Src.BytesPerBlock,
+				.NumElemsPerBlockX = FormatInfo.Src.PixelsPerBlockX,
+				.NumElemsPerBlockY = FormatInfo.Src.PixelsPerBlockY,
+			};
+
+			TSubImage<uint8> BaseSubImage = TSubImage<uint8>
+			{
+				.DataView = Base.DataStorage.GetLOD(0),
+				.Size     = Base.GetSize(),
+				.SubSize  = Base.GetSize(),
+				.BytesPerBlock     = FormatInfo.Dst.BytesPerBlock,
+				.NumElemsPerBlockX = FormatInfo.Dst.PixelsPerBlockX,
+				.NumElemsPerBlockY = FormatInfo.Dst.PixelsPerBlockY,
+			};
+
+			for (int32 BatchIndex = 0; BatchIndex < NumBatches; ++BatchIndex)
+			{
+				FBatchWorkingSubImages WorkingSubImages = ComputeWorkingSubImages(BatchIndex, DstRect, SrcRect, BaseSubImage, SrcSubImage, FormatInfo, ScratchMemory);
+
+				if (EnumHasAnyFlags(ActiveStages, EActiveStages::DstDecompression))
+				{
+					MUTABLE_CPUPROFILER_SCOPE(ImageMultiCompose_DstDecompression);
+
+					const TSubImage<uint8>& StageSrc = WorkingSubImages.Dst;
+					const TSubImage<uint8>& StageDst = WorkingSubImages.DecDstAligned;
+
+					check(StageSrc.SubSize == StageDst.SubSize);
+
+					const uint8* FromPtr = StageSrc.DataView.GetData();
+					uint8* ToPtr         = StageDst.DataView.GetData();
+
+					miro::FImageSize FromSize = miro::FImageSize(StageSrc.Size.X, StageSrc.Size.Y);
+					miro::FImageSize ToSize   = miro::FImageSize(StageDst.Size.X, StageDst.Size.Y); 
+					miro::FImageSize SubSize  = miro::FImageSize(StageSrc.SubSize.X, StageSrc.SubSize.Y);
+
+					FormatInfo.DstDecompressionFunc(FromSize, ToSize, SubSize, FromPtr, ToPtr);	
+				}
+
+				if (EnumHasAnyFlags(ActiveStages, EActiveStages::SrcDecompression))
+				{
+					MUTABLE_CPUPROFILER_SCOPE(ImageMultiCompose_SrcDecompression);
+
+					const TConstSubImage<uint8>& StageSrc = WorkingSubImages.Src;
+					const TSubImage<uint8>& StageDst = EnumHasAnyFlags(ActiveStages, EActiveStages::DstCompression | EActiveStages::Resize)
+							? WorkingSubImages.DecSrcAligned
+							: WorkingSubImages.Dst;
+		
+					const uint8* FromPtr = StageSrc.DataView.GetData();
+					uint8* ToPtr         = StageDst.DataView.GetData();
+
+					miro::FImageSize FromSize = miro::FImageSize(StageSrc.Size.X, StageSrc.Size.Y);
+					miro::FImageSize ToSize   = miro::FImageSize(StageDst.Size.X, StageDst.Size.Y); 
+					miro::FImageSize SubSize  = miro::FImageSize(StageDst.SubSize.X, StageDst.SubSize.Y);
+
+					const bool bUncompressedFormat = 
+							FormatInfo.SrcDecompressedFormat != FormatInfo.DstDecompressedFormat && 
+							FormatInfo.SrcDecompressedFormat == Source.GetFormat()               &&  
+							FormatInfo.DstDecompressedFormat == Base.GetFormat();
+
+					if (bUncompressedFormat)
+					{
+						FormatInfo.SrcUncompressedFormatFunc(StageSrc.Size, StageDst.Size, StageDst.SubSize, FromPtr, ToPtr); 
+					}
+					else
+					{
+						FormatInfo.SrcDecompressionFunc(FromSize, ToSize, SubSize, FromPtr, ToPtr);
+					}
+				}
+
+				{
+					TConstSubImage<uint8> ConstDecSrc = WorkingSubImages.DecSrc.ToConst();
+
+					const TConstSubImage<uint8>& StageSrc = EnumHasAnyFlags(ActiveStages, EActiveStages::SrcDecompression)
+							? ConstDecSrc
+							: WorkingSubImages.Src;
+
+					const TSubImage<uint8>& StageDst = EnumHasAnyFlags(ActiveStages, EActiveStages::DstCompression)
+							? WorkingSubImages.DecDst
+							: WorkingSubImages.Dst;
+
+					if (EnumHasAnyFlags(ActiveStages, EActiveStages::Resize))
+					{
+						MUTABLE_CPUPROFILER_SCOPE(ImageMultiCompose_Resize);
+
+						const uint8* FromPtr = StageSrc.DataView.GetData();
+						uint8* ToPtr         = StageDst.DataView.GetData();
+
+						FVector2f SrcRectDimsF = FVector2f(SrcRect.Width(), SrcRect.Height());
+						FVector2f DstRectDimsF = FVector2f(DstRect.Width(), DstRect.Height());
+
+						FVector2f OneOverScalingFactor = SrcRectDimsF / DstRectDimsF;
+
+						FormatInfo.ResizeFunc(OneOverScalingFactor, WorkingSubImages.SrcSubPixelOffset, 
+								StageDst.Size, StageDst.SubSize, StageSrc.Size, StageSrc.SubSize, ToPtr, FromPtr);	
+					}
+					else 
+					{
+						MUTABLE_CPUPROFILER_SCOPE(ImageMultiCompose_Copy);
+						CopySubImage(StageDst, StageSrc);
+					}
+				}
+
+				if (EnumHasAnyFlags(ActiveStages, EActiveStages::DstCompression))
+				{
+					MUTABLE_CPUPROFILER_SCOPE(ImageMultiCompose_DstCompression);
+
+					const TSubImage<uint8>& StageSrc = WorkingSubImages.DecDstAligned;
+					const TSubImage<uint8>& StageDst = WorkingSubImages.Dst;
+
+					check(StageSrc.SubSize == StageDst.SubSize);
+
+					const uint8* FromPtr = StageSrc.DataView.GetData();
+					uint8* ToPtr         = StageDst.DataView.GetData();
+
+					miro::FImageSize FromSize = miro::FImageSize(StageSrc.Size.X, StageSrc.Size.Y);
+					miro::FImageSize ToSize   = miro::FImageSize(StageDst.Size.X, StageDst.Size.Y); 
+					miro::FImageSize SubSize  = miro::FImageSize(StageDst.SubSize.X, StageDst.SubSize.Y);
+
+					FormatInfo.DstCompressionFunc(FromSize, ToSize, SubSize, FromPtr, ToPtr, CompressionQuality);
+				}
+			}
+		}
+	}
+
+	float ImageMultiComposeComputeBestLODForSrcRect(const FIntRect& DstRectInPixels, const FIntRect& SrcRectInPixels)
+	{
+		FVector2f DstRectSize = FVector2f(DstRectInPixels.Width(), DstRectInPixels.Height());
+		FVector2f SrcRectSize = FVector2f(SrcRectInPixels.Width(), SrcRectInPixels.Height());
+
+		// Image multicompose only supports scaling down by less than 2 or scale up, for now pick the coord that 
+		// changes the most down so we never get values outside the limitations.
+		//
+		// NOTE: Doing this we may loose resolution if both dimentions change very differently, using layout packing
+		// startegies that mantain the original aspect ratio is advised to get the best results possible. 
+		int32 MostChangingCoord = static_cast<int32>((DstRectSize.X - SrcRectSize.X) > (DstRectSize.Y - SrcRectSize.Y));
+
+		const float BestLOD = 
+				FMath::Log2(FMath::Max(1.0f, FMath::Square(SrcRectSize[MostChangingCoord]))) * 0.5f - 
+				FMath::Log2(FMath::Max(1.0f, FMath::Square(DstRectSize[MostChangingCoord]))) * 0.5f;
+
+		return BestLOD;
+	}
+} // namespace UE::Mutable::Private
+
diff --git a/Engine/Plugins/Mutable/Source/MutableRuntime/Private/MuR/OpImageMultiCompose.h b/Engine/Plugins/Mutable/Source/MutableRuntime/Private/MuR/OpImageMultiCompose.h
new file mode 100644
--- /dev/null
+++ b/Engine/Plugins/Mutable/Source/MutableRuntime/Private/MuR/OpImageMultiCompose.h
@@ -0,0 +1,42 @@
+// Copyright Epic Games, Inc. All Rights Reserved.
+
+#include "HAL/Platform.h"
+#include "Containers/ArrayView.h"
+
+namespace UE::Mutable::Private
+{
+	class FImage;
+	
+	/**
+	 * Composes an image using multiple rects from the same source image. 
+	 * 
+	 * @param Base                Base image where the composing Rects will be copied.
+	 *
+	 * @param Source              Source image from where the composing Rects data will be extracted. 
+	 *
+	 * @param DstRectsInPixels    Rects in pixels of regions of the Base image to be composed. 
+	 *
+	 * @param SrcRectsInPixels    Rects in pixels of regions of the Source image to compose, i-th SrcRectsInPixels 
+	 * 						      region content will be copied to i-th DstRectsInPixels region. 
+	 * 						      SrcRectsInPixels must have the same number of elements than DstRectsInPixels.
+	 *
+	 * @param SourceLODPerRect    Source LOD index where the i-th SrcRect is extracted from. Number of pixels in 
+	 *                            SrcRect are expected to be based on the Source image's LOD index dimensions.
+	 * 					          SourceLODPerRect must have the same number of elements than SrcRectsInPixels.
+	 *
+	 * @param CompressionQuality Quality level used by the block compression functions.
+	 */
+	void ImageMultiCompose(FImage& Base, const FImage& Source, TConstArrayView<FIntRect> DstRectsInPixels, TConstArrayView<FIntRect> SrcRectsInPixels, TConstArrayView<int32> SourceLODPerRect, int32 CompressionQuality);
+	
+	/**
+	 * Computes the best LOD bias for an image going from SrcRectInPixels dimensions to DstRectInPixels dimensions. 
+	 * 
+	 * @param DstRectInPixels The destination rect used for getting the dimensions used for the computation.
+	 *
+	 * @param SrcRectInPixels The source rect used for getting the dimensions used for the computation.
+	 *
+	 * @return                LOD bias needed to apply to SrcRect to get the dimensions closest to DstRect. 
+	 */
+	float ImageMultiComposeComputeBestLODForSrcRect(const FIntRect& DstRectInPixels, const FIntRect& SrcRectInPixels);
+}
+
