[Mutable] Fix crash when openning the Code / Graph Viwer without model.

#Synced-CL 48348317

diff --git a/Engine/Plugins/Mutable/Source/CustomizableObjectEditor/Private/MuCOE/CustomizableObjectCompiler.cpp b/Engine/Plugins/Mutable/Source/CustomizableObjectEditor/Private/MuCOE/CustomizableObjectCompiler.cpp
--- a/Engine/Plugins/Mutable/Source/CustomizableObjectEditor/Private/MuCOE/CustomizableObjectCompiler.cpp
+++ b/Engine/Plugins/Mutable/Source/CustomizableObjectEditor/Private/MuCOE/CustomizableObjectCompiler.cpp
@@ -1594,13 +1594,6 @@ UE::Mutable::Private::Ptr<UE::Mutable::Private::Node> FCustomizableObjectCompile
 {
 	UE_LOG(LogMutable, Log, TEXT("Started Customizable Object Export %s."), *Object->GetName());
 
-	FNotificationInfo Info(LOCTEXT("CustomizableObjectExportInProgress", "Exported Customizable Object"));
-	Info.bFireAndForget = true;
-	Info.bUseThrobber = true;
-	Info.FadeOutDuration = 1.0f;
-	Info.ExpireDuration = 1.0f;
-	FSlateNotificationManager::Get().AddNotification(Info);
-
 	FCompilationOptions CompilerOptions = InCompilerOptions;
 	CompilerOptions.bRealTimeMorphTargetsEnabled = Object->GetPrivate()->IsRealTimeMorphTargetsEnabled();
 	CompilerOptions.bClothingEnabled = Object->GetPrivate()->IsClothingEnabled();
diff --git a/Engine/Plugins/Mutable/Source/CustomizableObjectEditor/Private/MuCOE/CustomizableObjectEditor.cpp b/Engine/Plugins/Mutable/Source/CustomizableObjectEditor/Private/MuCOE/CustomizableObjectEditor.cpp
--- a/Engine/Plugins/Mutable/Source/CustomizableObjectEditor/Private/MuCOE/CustomizableObjectEditor.cpp
+++ b/Engine/Plugins/Mutable/Source/CustomizableObjectEditor/Private/MuCOE/CustomizableObjectEditor.cpp
@@ -1478,6 +1478,14 @@ void FCustomizableObjectEditor::GraphViewer() const
 
 	UE::Mutable::Private::Ptr<UE::Mutable::Private::Node> RootNode = DebuggerCompile(Model, RuntimeTextures, CompilerTextures, RuntimeMeshes, CompilerMeshes);
 	
+	if (!RootNode)
+	{
+		FNotificationInfo Info(LOCTEXT("GraphViewerError", "Unable to open Graph Viewer: No root node"));
+		FSlateNotificationManager::Get().AddNotification(Info);
+		
+		return;
+	}
+	
 	const TSharedPtr<SDockTab> NewMutableObjectTab = SNew(SDockTab)
 	.Label(LOCTEXT("GraphViewer", "Graph Viewer"))
 	[
@@ -1501,7 +1509,15 @@ void FCustomizableObjectEditor::CodeViewer() const
 	TArray<TSoftObjectPtr<const UStreamableRenderAsset>> RuntimeMeshes;
 	TArray<FMutableSourceMeshData> CompilerMeshes;
 
-	UE::Mutable::Private::Ptr<UE::Mutable::Private::Node> RootNode = DebuggerCompile(Model, RuntimeTextures, CompilerTextures, RuntimeMeshes, CompilerMeshes);
+	DebuggerCompile(Model, RuntimeTextures, CompilerTextures, RuntimeMeshes, CompilerMeshes);
+	
+	if (!Model)
+	{
+		FNotificationInfo Info(LOCTEXT("CodeViewerError", "Unable to open Code Viewer: No model"));
+		FSlateNotificationManager::Get().AddNotification(Info);
+
+		return;
+	}
 	
 	const TSharedRef<FBoneNames> BoneNames = MakeShared<FBoneNames>(); // TODO DebuggerCompile does not create a UModelResource
 	TSharedPtr<FUnrealMutableResourceProvider> ExternalResourceProvider = MakeShared<FUnrealMutableResourceProvider>(BoneNames);
