[Mutable] Decouple Materials from Surfaces.

- Renamed ReuseInstanceTextures to ReuseGeneratedResources.
- Split Task_Mutable_GetSurfaces and Task_Mutable_GetMaterials.
- Reworked BuildMaterials. Now split into BuildMaterial and BuildTexture.

#Synced-CL 48348317

diff --git a/Engine/Plugins/Mutable/Source/CustomizableObject/Internal/MuCO/CustomizableObjectInstancePrivate.h b/Engine/Plugins/Mutable/Source/CustomizableObject/Internal/MuCO/CustomizableObjectInstancePrivate.h
--- a/Engine/Plugins/Mutable/Source/CustomizableObject/Internal/MuCO/CustomizableObjectInstancePrivate.h
+++ b/Engine/Plugins/Mutable/Source/CustomizableObject/Internal/MuCO/CustomizableObjectInstancePrivate.h
@@ -58,7 +58,6 @@ enum ECOInstanceFlags
 	ECONone							= 0,  // Should not use the name None here.. it collides with other enum in global namespace
 
 	// Update process
-	ReuseTextures					= 1 << 3, 	// 
 	ReplacePhysicsAssets			= 1 << 4,	// Merge active PhysicsAssets and replace the base physics asset
 
 	// Update priorities
@@ -188,6 +187,29 @@ struct FExtensionInstanceData
 };
 
 
+struct FMaterialReuseCacheKey
+{
+	TStrongObjectPtr<UMaterialInterface> MaterialTemplate;
+	UE::Mutable::Private::OP::ADDRESS Root;
+
+	bool operator==(const FMaterialReuseCacheKey&) const = default;
+};
+
+
+uint32 GetTypeHash(const FMaterialReuseCacheKey& Key);
+
+
+struct FTextureReuseCacheKey
+{
+	UE::Mutable::Private::OP::ADDRESS Root;
+
+	bool operator==(const FTextureReuseCacheKey&) const = default;
+};
+
+
+uint32 GetTypeHash(const FTextureReuseCacheKey& Key);
+
+
 #if WITH_EDITOR
 DECLARE_MULTICAST_DELEGATE_OneParam(FObjectInstanceTransactedDelegate, const FTransactionObjectEvent&);
 #endif //WITH_EDITOR
@@ -202,38 +224,6 @@ enum class ESkeletalMeshStatus : uint8
 };
 
 
-USTRUCT()
-struct FTextureReuseCacheKey
-{
-	GENERATED_BODY()
-	
-	FTextureReuseCacheKey() = default;
-	FTextureReuseCacheKey(int32 InLODIndex, int32 InComponentIndex, uint32 InSurfaceIndex, int32 InImageIndex)
-	{
-		LODIndex = InLODIndex;
-		ComponentIndex = InComponentIndex;
-		SurfaceIndex = InSurfaceIndex;
-		ImageIndex = InImageIndex;
-	}
-
-	friend uint32 GetTypeHash(FTextureReuseCacheKey In)
-	{
-		return HashCombine(GetTypeHash(In.LODIndex), GetTypeHash(In.ComponentIndex), GetTypeHash(In.SurfaceIndex), GetTypeHash(In.ImageIndex));
-	}
-	
-	bool operator==(const FTextureReuseCacheKey& Other) const
-	{
-		return	LODIndex == Other.LODIndex && ComponentIndex == Other.ComponentIndex &&
-				SurfaceIndex == Other.SurfaceIndex && ImageIndex == Other.ImageIndex;
-	}
-
-	uint32 SurfaceIndex = 0;
-	int32 LODIndex = -1;
-	int32 ComponentIndex = -1;
-	int32 ImageIndex = -1;
-};
-
-
 UCLASS(MinimalAPI)
 class UCustomizableInstancePrivate : public UObject
 {
@@ -241,23 +231,19 @@ public:
 	GENERATED_BODY()
 
 	/** The generated skeletal meshes for this Instance. They may be null if the component is empty. */
-	UPROPERTY(Transient, VisibleAnywhere, Category = NoCategory, meta=(DisplayName = "Meshes"))
+	UPROPERTY(Transient, VisibleAnywhere, Category = NoCategory)
 	TMap<FName, TObjectPtr<USkeletalMesh>> SkeletalMeshes;
 	
 	UPROPERTY(Transient, VisibleAnywhere, Category = NoCategory, meta=(DisplayName = "Materials"))
-	TArray<FGeneratedMaterial> GeneratedMaterials;
+	TArray<TObjectPtr<UMaterialInterface>> GeneratedMaterials;
 
-	UPROPERTY(Transient)
-	TArray<FGeneratedTexture> GeneratedTextures;
+	UPROPERTY(Transient, VisibleAnywhere, Category = NoCategory, meta=(DisplayName = "Textures"))
+	TArray<TObjectPtr<UTexture>> GeneratedTextures;
 
 	// Indices of the parameters that are relevant for the given parameter values.
 	// This only gets updated if parameter decorations are generated.
 	TArray<int> RelevantParameters;
 
-	// If Texture reuse is enabled, stores which texture is being used in a particular <LODIndex, ComponentIndex, MeshSurfaceIndex, image>
-	UPROPERTY(Transient)
-	TMap<FTextureReuseCacheKey, TWeakObjectPtr<UTexture2D>> TextureReuseCache;
-	
 	// Only used in LiveUpdateMode to reuse core instances between updates and their temp data to speed up updates, but spend way more memory
 	UE::Mutable::Private::FInstance::FID LiveUpdateModeInstanceID = 0;
 
@@ -283,9 +269,7 @@ public:
 	bool HasCOInstanceFlags(ECOInstanceFlags FlagsToCheck) const { return (InstanceFlagsPrivate & FlagsToCheck) != 0; }
 
 	UE_API void BuildMaterials(const TSharedRef<FUpdateContextPrivate>& Context, UCustomizableObjectInstance* Public);
-
-	UE_API void ReuseTexture(UTexture2D* Texture, TSharedRef<FTexturePlatformData, ESPMode::ThreadSafe>& PlatformData);
-
+	
 	UE_API void TickUpdateCloseCustomizableObjects(UCustomizableObjectInstance& Public, FMutableInstanceUpdateMap& InOutRequestedUpdates);
 	UE_API void UpdateInstanceIfNotGenerated(UCustomizableObjectInstance& Public, FMutableInstanceUpdateMap& InOutRequestedUpdates);
 
@@ -316,8 +300,6 @@ public:
 	UE_API FCustomizableInstanceComponentData* GetComponentData(FCustomizableObjectComponentIndex ObjectComponentIndex);
 	UE_API const FCustomizableInstanceComponentData* GetComponentData(FCustomizableObjectComponentIndex ObjectComponentIndex) const;
 	
-	UE_API UTexture2D* CreateTexture(const FName& TextureName);
-
 	UE_API void InvalidateGeneratedData();
 
 	UE_API bool LoadParametersFromProfile(int32 ProfileIndex);
@@ -406,6 +388,8 @@ public:
 	/** Automatic update required.
 	 * Set to true when a Customizable Object Instance Usage requires an automatic update (e.g., component reattached). */
 	bool bAutomaticUpdateRequired = false;
+
+	bool bReuseInstanceTextures = false;
 	
 	// TEMP VARIABLE to check the Min desired LODs for this instance
 	TWeakObjectPtr<UCustomizableObjectInstanceUsage> NearestToActor;
@@ -426,6 +410,12 @@ public:
 	/** Delegate called when the Instance has been transacted */
 	FObjectInstanceTransactedDelegate OnInstanceTransactedDelegate;
 #endif // WITH_EDITOR
+	
+	TMap<FMaterialReuseCacheKey, TStrongObjectPtr<UMaterialInterface>> LastMaterialsReuseCache;
+	TMap<FMaterialReuseCacheKey, TStrongObjectPtr<UMaterialInterface>> CurrentMaterialsReuseCache;
+
+	TMap<FTextureReuseCacheKey, TStrongObjectPtr<UTexture2D>> LastTextureReuseCache;
+	TMap<FTextureReuseCacheKey, TStrongObjectPtr<UTexture2D>> CurrentTextureReuseCache;
 };
 
 
diff --git a/Engine/Plugins/Mutable/Source/CustomizableObject/Internal/MuCO/CustomizableObjectPrivate.h b/Engine/Plugins/Mutable/Source/CustomizableObject/Internal/MuCO/CustomizableObjectPrivate.h
--- a/Engine/Plugins/Mutable/Source/CustomizableObject/Internal/MuCO/CustomizableObjectPrivate.h
+++ b/Engine/Plugins/Mutable/Source/CustomizableObject/Internal/MuCO/CustomizableObjectPrivate.h
@@ -287,7 +287,6 @@ public:
 		int32 SkippedMips = 0;
 
 #if WITH_EDITORONLY_DATA
-		// TODO: Remove this bool and make the actual generation of the textures a step prior to the build of the materials
 		bool bIsBake = false;
 #endif
 		
diff --git a/Engine/Plugins/Mutable/Source/CustomizableObject/Internal/MuCO/CustomizableObjectSystemPrivate.h b/Engine/Plugins/Mutable/Source/CustomizableObject/Internal/MuCO/CustomizableObjectSystemPrivate.h
--- a/Engine/Plugins/Mutable/Source/CustomizableObject/Internal/MuCO/CustomizableObjectSystemPrivate.h
+++ b/Engine/Plugins/Mutable/Source/CustomizableObject/Internal/MuCO/CustomizableObjectSystemPrivate.h
@@ -36,6 +36,7 @@
 
 #define UE_API CUSTOMIZABLEOBJECT_API
 
+class FMutableUpdateImageContext;
 class FBoneNames;
 class FUnrealMutableResourceProvider;
 class FMutableStreamRequest;
@@ -304,49 +305,6 @@ public:
 };
 
 
-USTRUCT()
-struct FGeneratedTexture
-{
-	GENERATED_USTRUCT_BODY();
-
-	UPROPERTY(Category = NoCategory, VisibleAnywhere)
-	FString Name;
-
-	UPROPERTY(Category = NoCategory, VisibleAnywhere)
-	TObjectPtr<UTexture> Texture = nullptr;
-
-	bool operator==(const FGeneratedTexture& Other) const = default;
-};
-
-
-USTRUCT()
-struct FGeneratedMaterial
-{
-	GENERATED_USTRUCT_BODY();
-
-	/* Output material */
-	UPROPERTY(Category = NoCategory, VisibleAnywhere)
-	TObjectPtr<UMaterialInterface> MaterialInterface;
-
-	UPROPERTY(Category = NoCategory, VisibleAnywhere)
-	TArray< FGeneratedTexture > Textures;
-
-	// Surface or SharedSurface Id
-	uint32 SurfaceId = 0;
-
-	// Core Material. Strong object pointer to a material or a material index
-	TSharedPtr<const UE::Mutable::Private::FMaterial> Material;
-
-#if WITH_EDITORONLY_DATA
-	// Name of the component that contains this material
-	UPROPERTY(Category = NoCategory, VisibleAnywhere)
-	FName ComponentName;
-#endif
-
-	bool operator==(const FGeneratedMaterial& Other) const { return SurfaceId == Other.SurfaceId && Material == Other.Material; };
-};
-
-
 /** Final data per component. */
 struct FSkeletalMeshMorphTargets
 {
@@ -364,11 +322,11 @@ struct FInstanceUpdateData
 {
 	struct FImage
 	{
-		FName Name;
+		FName Name; // TODO MATERIAL_PIN Remove this hack
+		FName ParameterName; // TODO MATERIAL_PIN This does not belong to the Image but the Material
 		UE::Mutable::Private::FImageId ImageID;
 		
 		// LOD of the ImageId.
-		int32 BaseLOD;
 		int32 BaseMip;
 		
 		uint16 FullImageSizeX;
@@ -384,7 +342,7 @@ struct FInstanceUpdateData
 		bool bIsPassThrough = false;
 		bool bIsNonProgressive = false;
 
-		int32 MaterialLayer;
+		int32 MaterialLayer = INDEX_NONE; // TODO MATERIAL_PIN This does not belong to the Image but the Material
 	};
 
 	struct FVector
@@ -403,7 +361,7 @@ struct FInstanceUpdateData
 		int32 MaterialLayer;
 	};
 
-	struct FSurface
+	struct FMaterial
 	{
 		/** Range in the Images array */
 		uint16 FirstImage = 0;
@@ -417,14 +375,21 @@ struct FInstanceUpdateData
 		uint16 FirstScalar = 0;
 		uint16 ScalarCount = 0;
 
+		bool bGetParameters = true; // TODO MATERIAL_PIN Remove once the Material Pin is completed. Only used to avid GetImages/Vector/Scalar twice (Material and Surface).
+		
+		UE::Mutable::Private::FMaterialId MaterialId;
+		TSharedPtr<const UE::Mutable::Private::FMaterial> Material;
+	};
+	
+	struct FSurface
+	{
 		/** Id of the surface in the mutable core instance. */
 		uint32 SurfaceId = 0;
 
 		/** Id of the metadata associated with this surface. */ 
 		uint32 SurfaceMetadataId = 0;
 
-		UE::Mutable::Private::FMaterialId MaterialId;
-		TSharedPtr<const UE::Mutable::Private::FMaterial> Material;
+		uint32 MaterialIndex = INDEX_NONE;
 	};
 
 	struct FLOD
@@ -453,9 +418,7 @@ struct FInstanceUpdateData
 		uint16 FirstLOD = 0;
 		uint16 LODCount = 0;
 
-		/** Overlay material. Only one per mesh, if any. */
-		UE::Mutable::Private::FMaterialId OverlayMaterialId;
-		TSharedPtr<const UE::Mutable::Private::FMaterial> OverlayMaterial;
+		int32 OverlayMaterialIndex = INDEX_NONE;
 
 		UE::Mutable::Private::FSkeletalMeshId SkeletalMeshId;
 		
@@ -464,6 +427,7 @@ struct FInstanceUpdateData
 
 	TArray<TSharedRef<FComponent>> Components;
 	TArray<TSharedRef<FLOD>> LODs;
+	TArray<TSharedRef<FMaterial>> Materials;
 	TArray<TSharedRef<FSurface>> Surfaces;
 	TArray<TSharedRef<FImage>> Images;
 	TArray<FVector> Vectors;
@@ -516,6 +480,7 @@ struct FInstanceUpdateData
 		LODs.Empty();
 		Components.Empty();
 		Surfaces.Empty();
+		Materials.Empty();
 		Images.Empty();
 		Scalars.Empty();
 		Vectors.Empty();
@@ -526,6 +491,9 @@ struct FInstanceUpdateData
 };
 
 
+uint32 GetMaterialHash(const FInstanceUpdateData& UpdateData, const FInstanceUpdateData::FMaterial& Material, bool bPersistent);
+
+
 struct FComponentDescriptor
 {
 	TArray<UE::Mutable::Private::FMeshId> MeshDescriptors;
@@ -603,7 +571,9 @@ public:
 	
 	/** When this option is enabled it will reuse the Mutable core instance and its temp data between updates.  */
 	bool bLiveUpdateMode = false;
+
 	bool bReuseInstanceTextures = false;
+	
 	bool bUseMeshCache = false;
 	
 	/** Whether the mesh to generate should support Mesh LOD streaming or not. */
@@ -629,6 +599,8 @@ public:
 
 	bool bEnableSkipGenerateResidentMips = false;
 
+	bool bIsUpdateAborted = false;
+
 	/** This option comes from the operation request. It is used to reduce the number of mipmaps that mutable must generate for images.  */
 	int32 MipsToSkip = 0;
 
@@ -645,7 +617,7 @@ public:
 
 	TSharedPtr<FBoneNames> BoneNames;
 
-	bool bIsUpdateAborted = false;
+	TSharedPtr<FMutableUpdateImageContext> UpdateImageContext;
 	
 	// Number of possible components in the entire CO
 	// TODO: Redundant because it is in the CO
@@ -730,6 +702,12 @@ public:
 	TArray<TStrongObjectPtr<const UTexture>> LoadedPassThroughTexturesPendingSetMaterial;
 	
 	TArray<UE::Mutable::Private::FImageId> CachedTextures;
+
+	/** Materials already build in this update. */
+	TMap<uint32, TStrongObjectPtr<UMaterialInterface>> MaterialUpdateCache; // TODO MATERIAL_PIN replace key with FMaterialId once MaterialPin completed.
+
+	/** Textures already build in this update. */
+	TMap<FTextureCache::FId, TStrongObjectPtr<UTexture>> TextureUpdateCache;
 };
 
 
@@ -849,7 +827,6 @@ public:
 	
 	static int32 EnableMutableProgressiveMipStreaming;
 	static int32 EnableMutableLiveUpdate;
-	static int32 EnableReuseInstanceTextures;
 	static int32 EnableMutableAnimInfoDebugging;
 	static int32 EnableSkipGenerateResidentMips;
 	static int32 EnableOnlyGenerateRequestedLODs;
diff --git a/Engine/Plugins/Mutable/Source/CustomizableObject/Private/MuCO/CustomizableObjectInstance.cpp b/Engine/Plugins/Mutable/Source/CustomizableObject/Private/MuCO/CustomizableObjectInstance.cpp
--- a/Engine/Plugins/Mutable/Source/CustomizableObject/Private/MuCO/CustomizableObjectInstance.cpp
+++ b/Engine/Plugins/Mutable/Source/CustomizableObject/Private/MuCO/CustomizableObjectInstance.cpp
@@ -95,124 +95,7 @@ const FString IMAGE_PARAMETER_POSTFIX = FString("_SelectedImages");
 const FString POSE_PARAMETER_POSTFIX = FString("_SelectedPoses");
 
 
-// Struct used by BuildMaterials() to identify common materials between LODs
-struct FMutableMaterialPlaceholder
-{
-	enum class EPlaceHolderParamType { Vector, Scalar, Texture };
-
-	struct FMutableMaterialPlaceHolderParam
-	{
-		FName ParamName;
-		EPlaceHolderParamType Type;
-		int32 LayerIndex; // Set to -1 for non-multilayer params
-		float Scalar = 0.0f;
-		FLinearColor Vector = FLinearColor::Black;
-		FGeneratedTexture Texture;
-
-		FMutableMaterialPlaceHolderParam(const FName& InParamName, const int32 InLayerIndex, const FLinearColor& InVector)
-			: ParamName(InParamName), Type(EPlaceHolderParamType::Vector), LayerIndex(InLayerIndex), Vector(InVector) {}
-
-		FMutableMaterialPlaceHolderParam(const FName& InParamName, const int32 InLayerIndex, const float InScalar)
-			: ParamName(InParamName), Type(EPlaceHolderParamType::Scalar), LayerIndex(InLayerIndex), Scalar(InScalar) {}
-
-		FMutableMaterialPlaceHolderParam(const FName& InParamName, const int32 InLayerIndex, const FGeneratedTexture& InTexture)
-			: ParamName(InParamName), Type(EPlaceHolderParamType::Texture), LayerIndex(InLayerIndex), Texture(InTexture) {}
-
-		bool operator<(const FMutableMaterialPlaceHolderParam& Other) const
-		{
-			if (Type != Other.Type)
-			{
-				return Type < Other.Type;
-			}
-			else
-			{
-				return ParamName.LexicalLess(Other.ParamName);
-			}
-		}
-
-		bool operator==(const FMutableMaterialPlaceHolderParam& Other) const = default;
-	};
-
-	uint32 ParentMaterialID = 0;
-	int32 MatIndex = -1;
-
-private:
-	mutable TArray<FMutableMaterialPlaceHolderParam> Params;
-
-public:
-	void AddParam(const FMutableMaterialPlaceHolderParam& NewParam) { Params.Add(NewParam); }
-	
-	const TArray<FMutableMaterialPlaceHolderParam>& GetParams() const { return Params; }
-
-	bool operator==(const FMutableMaterialPlaceholder& Other) const;
-
-	friend uint32 GetTypeHash(const FMutableMaterialPlaceholder& PlaceHolder);
-	
-	friend uint32 GetHash(const FMutableMaterialPlaceholder& PlaceHolder, const bool bPersistent);
-};
-
-
-bool FMutableMaterialPlaceholder::operator==(const FMutableMaterialPlaceholder& Other) const
-{
-	return ParentMaterialID == Other.ParentMaterialID &&
-		   Params == Other.Params;
-}
-
-
-uint32 GetTypeHash(const FMutableMaterialPlaceholder& PlaceHolder)
-{
-	return GetHash(PlaceHolder, false);
-}
-
-
-uint32 GetHash(const FMutableMaterialPlaceholder& PlaceHolder, bool bPersistent)
-{
-	auto FastApproach = [](const uint32 HashA, const uint32 HashB) -> uint32
-	{
-		return HashCombineFast(HashA, HashB);
-	};
-
-	auto BackwardsCompatibleApproach = [](const uint32 HashA, const uint32 HashB) -> uint32
-	{
-		return HashCombine(HashA,HashB);
-	};
-
-	auto HashingMethod = bPersistent ? BackwardsCompatibleApproach : FastApproach;
-	
-	uint32 Hash = GetTypeHash(PlaceHolder.ParentMaterialID);
-
-	// Sort parameters before building the hash.
-	PlaceHolder.Params.Sort();
-
-	for (const FMutableMaterialPlaceholder::FMutableMaterialPlaceHolderParam& Param : PlaceHolder.Params)
-	{
-		uint32 ParamHash = GetTypeHash(Param.ParamName.ToString().ToLower());
-		ParamHash = HashingMethod(ParamHash, (uint32)Param.LayerIndex);
-		ParamHash = HashingMethod(ParamHash, (uint32)Param.Type);
-
-		switch (Param.Type)
-		{
-		case FMutableMaterialPlaceholder::EPlaceHolderParamType::Vector:
-			ParamHash = HashingMethod(ParamHash, GetTypeHash(Param.Vector));
-			break;
-
-		case FMutableMaterialPlaceholder::EPlaceHolderParamType::Scalar:
-			ParamHash = HashingMethod(ParamHash, GetTypeHash(Param.Scalar));
-			break;
-
-		case FMutableMaterialPlaceholder::EPlaceHolderParamType::Texture:
-			ParamHash = HashingMethod(ParamHash, GetTypeHash(Param.Texture.Texture.GetName().ToLower()));
-			break;
-		}
-
-		Hash = HashingMethod(Hash, ParamHash);
-	}
-
-	return Hash;
-}
-
-
-UTexture2D* UCustomizableInstancePrivate::CreateTexture(const FName& TextureName)
+UTexture2D* CreateTexture(const FName& TextureName)
 {
 	UTexture2D* NewTexture = NewObject<UTexture2D>(
 		GetTransientPackage(),
@@ -374,6 +257,21 @@ void UCustomizableInstancePrivate::PrepareForUpdate()
 }
 
 
+uint32 GetTypeHash(const FMaterialReuseCacheKey& Key)
+{
+	uint32 Hash = GetTypeHash(Key.MaterialTemplate);
+
+	Hash = HashCombine(Hash, Key.Root);
+	
+	return Hash;
+}
+
+uint32 GetTypeHash(const FTextureReuseCacheKey& Key)
+{
+	return Key.Root;
+}
+
+
 #if WITH_EDITOR
 void UCustomizableInstancePrivate::PostDuplicate(bool bDuplicateForPIE)
 {
@@ -676,15 +574,7 @@ int32 UCustomizableInstancePrivate::GetState() const
 
 void UCustomizableInstancePrivate::SetState(const int32 InState)
 {
-	const int32 OldState = GetState();
-	
 	GetPublic()->Descriptor.SetState(InState);
-
-	if (OldState != InState)
-	{
-		// State may change texture properties, so invalidate the texture reuse cache
-		TextureReuseCache.Empty();
-	}
 }
 
 
@@ -1627,6 +1517,7 @@ namespace
 	}
 }
 
+
 UPhysicsAsset* GetOrBuildMainPhysicsAsset(
 	const TSharedRef<FUpdateContextPrivate>& Context,
 	TObjectPtr<UPhysicsAsset> TemplateAsset,
@@ -2074,8 +1965,6 @@ bool UCustomizableInstancePrivate::UpdateSkeletalMesh_PostBeginUpdate0(UCustomiz
 		return false;
 	}
 	
-	TextureReuseCache.Empty(); // Sections may have changed, so invalidate the texture reuse cache because it's indexed by section
-
 	TMap<FName, TObjectPtr<USkeletalMesh>> OldSkeletalMeshes = SkeletalMeshes;
 
 	const UModelResources& ModelResources = CustomizableObject->GetPrivate()->GetModelResourcesChecked();
@@ -3498,7 +3387,7 @@ FTexturePlatformData* MutableCreateImagePlatformData(TSharedPtr<const UE::Mutabl
 		}		
 	}
 
-	FTexturePlatformData* PlatformData = new FTexturePlatformData();
+	FTexturePlatformData* PlatformData = new FTexturePlatformData(); // TODO GMT Change to SharedRef
 	PlatformData->SizeX = SizeX;
 	PlatformData->SizeY = SizeY;
 	PlatformData->PixelFormat = PlatformFormat;
@@ -5685,7 +5574,7 @@ void UpdateTextureRegionsMutable(UTexture2D* Texture, int32 MipIndex, uint32 Num
 }
 
 
-void UCustomizableInstancePrivate::ReuseTexture(UTexture2D* Texture, TSharedRef<FTexturePlatformData, ESPMode::ThreadSafe>& PlatformData)
+void ReuseTexture(UTexture2D* Texture, TSharedRef<FTexturePlatformData, ESPMode::ThreadSafe>& PlatformData)
 {
 	uint32 NumMips = PlatformData->Mips.Num();
 
@@ -5714,6 +5603,402 @@ void UCustomizableInstancePrivate::ReuseTexture(UTexture2D* Texture, TSharedRef<
 }
 
 
+
+FName TextureName(const FInstanceUpdateData::FImage& Image, bool bBake)
+{
+	FString MutableTextureName;
+	{
+		const uint32 MutableImageHash = UE::Mutable::Private::GetTypeHashPersistent(Image.ImageID);
+		MutableTextureName = FString::Printf(TEXT("T_%s_h%llu"), *Image.ParameterName.ToString(), MutableImageHash);
+		MutableTextureName.ReplaceInline(TEXT(" "), TEXT("_"));
+
+#if WITH_EDITOR
+		if (bBake)
+		{
+			// If this is a bake operation just add the baked prefix to make the name unique for all baked instances of this resource
+			MutableTextureName = FBakingConfiguration::BakedResourcePrefix + MutableTextureName;
+		}
+#endif
+	}
+		
+	return MakeUniqueObjectName(GetTransientPackage(), UTexture2D::StaticClass(), *MutableTextureName);
+}
+
+UTexture* BuildTexture(const TSharedRef<FUpdateContextPrivate>& Context, const int32 ImageInex, const bool bInReuseTexture)
+{
+	const TSharedRef<const FInstanceUpdateData::FImage>& Image = Context->InstanceUpdateData.Images[ImageInex];
+	
+#if WITH_EDITOR // TODO MATERIAL_PIN Remove WITH_EDITOR
+	const FTextureCache::FId TextureCacheKey(Image->ImageID, Context->MipsToSkip, Context->bBake);
+#else
+	const FTextureCache::FId TextureCacheKey(Image->ImageID, Context->MipsToSkip);
+#endif
+	
+	TStrongObjectPtr<UTexture>* CachedTexture = Context->TextureUpdateCache.Find(TextureCacheKey);
+	if (CachedTexture)
+	{
+		return CachedTexture->Get();
+	}
+	
+	if (Image->Image && Image->Image->PassthroughObject)
+	{
+		UTexture* Texture = Image->Image->PassthroughObject.Get(); // TODO GMT Why not also cache Passthrough?
+
+		Context->Instance->GetPrivate()->GeneratedTextures.AddUnique(Texture);
+
+		return Texture;
+	}
+
+	const UModelResources& ModelResources = *Context->Object->GetPrivate()->GetModelResources();
+
+	// Get the cache of resources of all live instances of this object
+	FTextureCache& TextureCache = Context->Object->GetPrivate()->TextureCache;
+
+	FString CurrentState = Context->Instance->GetCurrentState();
+	bool bNeverStream = Context->bNeverStream;
+
+	check((bNeverStream && Context->MipsToSkip == 0) ||
+		(!bNeverStream && Context->MipsToSkip >= 0));
+	
+	FString KeyName = Image->Name.ToString();
+	TSharedPtr<const UE::Mutable::Private::FImage> MutableImage = Image->Image;
+
+	UTexture2D* Texture; // Texture generated by mutable
+
+	const bool bReuseTexture = bInReuseTexture && Context->bReuseInstanceTextures;
+	TMap<FTextureReuseCacheKey, TStrongObjectPtr<UTexture2D>>& LastTextureReuseCache = Context->Instance->GetPrivate()->LastTextureReuseCache; 
+	TMap<FTextureReuseCacheKey, TStrongObjectPtr<UTexture2D>>& CurrentTextureReuseCache = Context->Instance->GetPrivate()->CurrentTextureReuseCache;
+	
+	const FTextureReuseCacheKey TextureReuseCacheKey(Image->ImageID.GetKey()->Address);
+	
+	// If the mutable image is null, it must be in the cache
+	if (!MutableImage)
+	{
+		Texture = TextureCache.Get(TextureCacheKey);
+		check(Texture);
+	}
+	else
+	{
+		// Find the additional information for this image
+    	int32 ImageKey = FCString::Atoi(*KeyName);
+    	check (ImageKey >= 0 && ImageKey < ModelResources.ImageProperties.Num());
+    	
+    	const FMutableModelImageProperties& Props = ModelResources.ImageProperties[ImageKey];
+    	
+		UTexture2D* ReusedTexture = nullptr;
+		if (bReuseTexture)
+		{
+			if (TStrongObjectPtr<UTexture2D>* ReuseTexture = LastTextureReuseCache.Find(TextureReuseCacheKey))
+			{
+				ReusedTexture = ReuseTexture->Get();
+				LastTextureReuseCache.Remove(TextureReuseCacheKey);
+			}
+		}
+
+		if (ReusedTexture)
+		{
+			// Only uncompressed textures can be reused. This also fixes an issue in the editor where textures supposedly 
+			// uncompressed by their state, are still compressed because the CO has not been compiled at maximum settings
+			// and the uncompressed setting cannot be applied to them.
+			EPixelFormat PixelFormat = ReusedTexture->GetPixelFormat();
+			if (PixelFormat == EPixelFormat::PF_R8G8B8A8)
+			{
+				Texture = ReusedTexture;
+			}
+			else
+			{
+				ReusedTexture = nullptr;
+				
+				const FName Name = TextureName(Image.Get(), Context->bBake);
+				Texture = CreateTexture(Name);
+#if WITH_EDITOR
+				UE_LOG(LogMutable, Warning,
+					TEXT("Tried to reuse an uncompressed texture with name %s. Make sure the selected Mutable state disables texture compression/streaming, that one of the state's runtime parameters affects the texture and that the CO is compiled with max. optimization settings."),
+					*Texture->GetName());
+#endif
+			}
+		}
+		else
+		{
+			const FName Name = TextureName(Image.Get(), Context->bBake);
+			Texture = CreateTexture(Name);
+		}
+
+		if (!Context->ImageToPlatformDataMap.Contains(Image->ImageID))
+		{
+			UE_LOG(LogMutable, Error, TEXT("Required image [%s] with ID [%s] was not generated in the mutable thread, and it is not cached."),
+			*Props.TextureParameterName,
+			*Image->ImageID.ToString());
+			
+			return nullptr;
+		}
+		
+		SetTexturePropertiesFromMutableImageProps(Texture, Props, bNeverStream);
+
+		FTexturePlatformData* PlatformData = Context->ImageToPlatformDataMap[Image->ImageID];
+		Context->ImageToPlatformDataMap.Remove(Image->ImageID);
+
+		if (ReusedTexture)
+		{
+			check(PlatformData->Mips.Num() == ReusedTexture->GetPlatformData()->Mips.Num());
+			check(PlatformData->Mips[0].SizeX == ReusedTexture->GetPlatformData()->Mips[0].SizeX);
+			check(PlatformData->Mips[0].SizeY == ReusedTexture->GetPlatformData()->Mips[0].SizeY);
+
+			// Now the ReusedTexturePlatformData shared ptr owns the platform data
+			// This shared ptr will hold the reused texture platform data (mips) until the reused texture is updated 
+			// and delete it automatically
+			TSharedRef<FTexturePlatformData> PlatformDataRef = TSharedPtr<FTexturePlatformData>(PlatformData).ToSharedRef();
+			ReuseTexture(ReusedTexture, PlatformDataRef);
+		}
+		else
+		{
+			// Now the MutableTexture owns the platform data
+			Texture->SetPlatformData(PlatformData);
+
+			MUTABLE_CPUPROFILER_SCOPE(UpdateResource);
+#if REQUIRES_SINGLEUSE_FLAG_FOR_RUNTIME_TEXTURES
+			for (int32 i = 0; i < Texture->GetPlatformData()->Mips.Num(); ++i)
+			{
+				uint32 DataFlags = Texture->GetPlatformData()->Mips[i].BulkData.GetBulkDataFlags();
+				Texture->GetPlatformData()->Mips[i].BulkData.SetBulkDataFlags(DataFlags | BULKDATA_SingleUse);
+			}
+#endif
+
+			if (bNeverStream)
+			{
+				// To prevent LogTexture Error "Loading non-streamed mips from an external bulk file."
+				for (int32 i = 0; i < Texture->GetPlatformData()->Mips.Num(); ++i)
+				{
+					Texture->GetPlatformData()->Mips[i].BulkData.ClearBulkDataFlags(BULKDATA_PayloadInSeparateFile);
+				}
+			}
+				
+			//if (!bNeverStream) // No need to check bNeverStream. In that case, the texture won't use 
+			// the MutableMipDataProviderFactory anyway and it's needed for detecting Mutable textures elsewhere
+			UMutableTextureMipDataProviderFactory* MutableMipDataProviderFactory = Cast<UMutableTextureMipDataProviderFactory>(Texture->GetAssetUserDataOfClass(UMutableTextureMipDataProviderFactory::StaticClass()));
+			if (!MutableMipDataProviderFactory)
+			{
+				MutableMipDataProviderFactory = NewObject<UMutableTextureMipDataProviderFactory>();
+
+				if (MutableMipDataProviderFactory)
+				{
+					MutableMipDataProviderFactory->CustomizableObjectInstance = Context->Instance.Get();
+					MutableMipDataProviderFactory->ImageRef.ImageID = Image->ImageID;
+					MutableMipDataProviderFactory->ImageRef.BaseMip = uint8(Image->BaseMip);
+					MutableMipDataProviderFactory->ImageRef.ConstantImagesNeededToGenerate = Image->ConstantImagesNeededToGenerate;
+					MutableMipDataProviderFactory->UpdateContext = Context->UpdateImageContext;
+					Texture->AddAssetUserData(MutableMipDataProviderFactory);
+				}
+			}
+
+			Texture->UpdateResource();
+		}
+	}
+
+	if (Texture)
+	{
+		Context->Instance->GetPrivate()->GeneratedTextures.Add(Texture);
+
+		if (bReuseTexture)
+		{
+			CurrentTextureReuseCache.Add(TextureReuseCacheKey, TStrongObjectPtr(Texture));
+		}
+		else
+		{
+			TextureCache.Add(TextureCacheKey, Texture);
+		}
+	}
+
+	Context->TextureUpdateCache.Add(TextureCacheKey, TStrongObjectPtr(Texture));
+	
+	return Texture;
+}
+
+
+struct FBuildMaterialOut
+{
+	UMaterialInterface* Material = nullptr;
+	bool bNotifyPrimitiveUpdated = false;
+};
+
+
+FBuildMaterialOut BuildMaterial(const TSharedRef<FUpdateContextPrivate>& Context, const int32 MaterialIndex, const bool bReuseMaterial)
+{
+	FBuildMaterialOut Out;
+
+	if (MaterialIndex == INDEX_NONE)
+	{
+		Out.Material = UMaterial::GetDefaultMaterial(MD_Surface); // TODO MATERIAL_PIN. It should be possible to return nullptr (Overlay and Override can be nullptr). The caller then shoud be responsible of setting a Default or not.
+		
+		return Out;
+	}
+
+	const TSharedRef<const FInstanceUpdateData::FMaterial>& Material = Context->InstanceUpdateData.Materials[MaterialIndex];
+
+	const uint32 MaterialHash = GetMaterialHash(Context->InstanceUpdateData, Material.Get(), Context->bBake); // TODO MATERIAL_PIN Replace with FMaterialId once MaterialPin completed
+	TStrongObjectPtr<UMaterialInterface>* CachedMaterial = Context->MaterialUpdateCache.Find(MaterialHash);
+	if (CachedMaterial)
+	{
+		Out.Material = CachedMaterial->Get();
+		
+		return Out;
+	}
+	
+	if (!Material->Material->PassthroughObject.Get())
+	{
+		Out.Material = UMaterial::GetDefaultMaterial(MD_Surface);
+		
+		return Out;
+	}
+
+	UMaterialInterface* MaterialTemplate = Material->Material->PassthroughObject.Get();
+
+	TMap<FMaterialReuseCacheKey, TStrongObjectPtr<UMaterialInterface>>& LastMaterialReuseCache = Context->Instance->GetPrivate()->LastMaterialsReuseCache; 
+	TMap<FMaterialReuseCacheKey, TStrongObjectPtr<UMaterialInterface>>& CurrentMaterialReuseCache = Context->Instance->GetPrivate()->CurrentMaterialsReuseCache;
+	
+	const FMaterialReuseCacheKey MaterialReuseCacheKey(TStrongObjectPtr(MaterialTemplate), Material->MaterialId.GetKey()->Address);
+
+	UMaterialInterface* MaterialInterface = nullptr;
+	
+	if (bReuseMaterial)
+	{
+		if (const TStrongObjectPtr<UMaterialInterface>* ReuseMaterial = LastMaterialReuseCache.Find(MaterialReuseCacheKey))
+		{
+			MaterialInterface = ReuseMaterial->Get();	
+			LastMaterialReuseCache.Remove(MaterialReuseCacheKey);
+		}
+	}
+
+	const bool bHasParameters = Material->ScalarCount || Material->VectorCount || Material->ImageCount;
+	
+	if (!MaterialInterface &&
+		bHasParameters)
+	{
+#if WITH_EDITOR
+		FString MIDName = MaterialTemplate->GetName();
+		{
+			// Remove the MI_ or M_ prefixes from the material string so we use it as the name of the MID
+			const FString MaterialPrefix = TEXT("M_");						// Material
+			const FString MaterialInstancePrefix = TEXT("MI_");				// Material Instance
+			const FString MaterialInstanceConstantPrefix = TEXT("MIC_");	// Material Instance Constant
+			
+			if (MIDName.Find(MaterialInstancePrefix, ESearchCase::CaseSensitive) == 0)
+			{
+				MIDName = MIDName.RightChop(MaterialInstancePrefix.Len());
+			}
+			else if (MIDName.Find(MaterialPrefix, ESearchCase::CaseSensitive) == 0)
+			{
+				MIDName = MIDName.RightChop(MaterialPrefix.Len());
+			}
+			else if (MIDName.Find(MaterialInstanceConstantPrefix, ESearchCase::CaseSensitive) == 0)
+			{
+				MIDName = MIDName.RightChop(MaterialInstanceConstantPrefix.Len());
+			}
+
+			MIDName = FString::Printf(TEXT("MID_%s_h%lu"), *MIDName, MaterialHash);
+
+			// Add a "tag" to better identify this asset during the bake
+			if (Context->bBake)
+			{
+				MIDName = FBakingConfiguration::BakedResourcePrefix + MIDName;
+			}
+		}
+		
+		const FName UniqueMIDName = MakeUniqueObjectName(GetTransientPackage(), UMaterialInstanceDynamic::StaticClass(), *MIDName);
+		MaterialInterface = UMaterialInstanceDynamic::Create(MaterialTemplate, GetTransientPackage(), UniqueMIDName);
+#else
+		MaterialInterface = UMaterialInstanceDynamic::Create(MaterialTemplate, GetTransientPackage());
+#endif
+	}
+	
+	if (!MaterialInterface)
+	{
+		MaterialInterface = MaterialTemplate;
+	}
+	
+	if (bHasParameters)
+	{
+		UMaterialInstanceDynamic* MaterialInstance = CastChecked<UMaterialInstanceDynamic>(MaterialInterface);
+		
+		for (int32 VectorIndex = 0; VectorIndex < Material->VectorCount; ++VectorIndex)
+		{
+			const FInstanceUpdateData::FVector& Vector = Context->InstanceUpdateData.Vectors[Material->FirstVector + VectorIndex];
+			if (Vector.MaterialLayer < 0)
+			{
+				FLinearColor Color = Vector.Vector;
+
+				// HACK: We encode an invalid value (Nan) for table option "None.
+				// Decoding "None" color parameters that use the material color
+				if (FMath::IsNaN(Color.R))
+				{
+					FMaterialParameterInfo ParameterInfo(Vector.Name);
+					MaterialTemplate->GetVectorParameterValue(ParameterInfo, Color);
+				}
+
+				MaterialInstance->SetVectorParameterValue(Vector.Name, Color);
+			}
+			else
+			{
+				FMaterialParameterInfo ParameterInfo = FMaterialParameterInfo(Vector.Name, EMaterialParameterAssociation::LayerParameter, Vector.MaterialLayer);
+				MaterialInstance->SetVectorParameterValueByInfo(ParameterInfo, Vector.Vector);
+			}
+		}
+
+		for (int32 ScalarIndex = 0; ScalarIndex < Material->ScalarCount; ++ScalarIndex)
+		{
+			const FInstanceUpdateData::FScalar& Scalar = Context->InstanceUpdateData.Scalars[Material->FirstScalar + ScalarIndex];
+			if (Scalar.MaterialLayer < 0)
+			{
+				MaterialInstance->SetScalarParameterValue(FName(Scalar.Name), Scalar.Scalar);
+			}
+			else
+			{
+				FMaterialParameterInfo ParameterInfo = FMaterialParameterInfo(Scalar.Name, EMaterialParameterAssociation::LayerParameter, Scalar.MaterialLayer);
+				MaterialInstance->SetScalarParameterValueByInfo(ParameterInfo, Scalar.Scalar);
+			}
+		}
+
+		for (int32 ImageIndex = 0; ImageIndex < Material->ImageCount; ++ImageIndex)
+		{
+			const TSharedRef<const FInstanceUpdateData::FImage>& Image = Context->InstanceUpdateData.Images[Material->FirstImage + ImageIndex];
+				
+			UTexture* Texture = BuildTexture(Context, Material->FirstImage + ImageIndex, bReuseMaterial);
+				
+			if (Image->MaterialLayer < 0)
+			{
+				MaterialInstance->SetTextureParameterValue(Image->ParameterName, Texture);
+			}
+			else
+			{
+				FMaterialParameterInfo ParameterInfo = FMaterialParameterInfo(Image->ParameterName, EMaterialParameterAssociation::LayerParameter, Image->MaterialLayer);
+				MaterialInstance->SetTextureParameterValueByInfo(ParameterInfo, Texture);
+			}
+
+			if (!bDisableNotifyComponentsOfTextureUpdates)
+			{
+				Out.bNotifyPrimitiveUpdated = true;
+			}
+		}
+	}
+
+	if (MaterialInterface)
+	{
+		Context->Instance->GetPrivate()->GeneratedMaterials.AddUnique(MaterialInterface);
+
+		if (bReuseMaterial)
+		{
+			CurrentMaterialReuseCache.Add(MaterialReuseCacheKey, TStrongObjectPtr(MaterialInterface)); 			
+		}
+	}
+	
+	Context->MaterialUpdateCache.Add(MaterialHash, TStrongObjectPtr(MaterialInterface));
+
+	Out.Material = MaterialInterface;
+	
+	return Out;
+}
+
+
 void UCustomizableInstancePrivate::BuildMaterials(const TSharedRef<FUpdateContextPrivate>& Context, UCustomizableObjectInstance* Public)
 {
 	MUTABLE_CPUPROFILER_SCOPE(UCustomizableInstancePrivate::BuildMaterials)
@@ -5721,39 +6006,33 @@ void UCustomizableInstancePrivate::BuildMaterials(const TSharedRef<FUpdateContex
 	UCustomizableObject* CustomizableObject = Public->GetCustomizableObject();
 
 	const UModelResources& ModelResources = *CustomizableObject->GetPrivate()->GetModelResources();
-
-	TArray<FGeneratedTexture> NewGeneratedTextures;
-
-	// Temp copy to allow reuse of MaterialInstances
-	TArray<FGeneratedMaterial> OldGeneratedMaterials;
-	Exchange(OldGeneratedMaterials, GeneratedMaterials);
 	
 	// Prepare the data to store in order to regenerate resources for this instance (usually texture mips).
-	TSharedPtr<FMutableUpdateImageContext> UpdateContext = MakeShared<FMutableUpdateImageContext>();
-	UpdateContext->CustomizableObjectPathName = CustomizableObject->GetPathName();
-	UpdateContext->InstancePathName = Public->GetPathName();
-	UpdateContext->System = UCustomizableObjectSystem::GetInstance()->GetPrivate()->MutableSystem;
-	UpdateContext->Model = Context->Model;
-	UpdateContext->MeshIdRegistry = CustomizableObject->GetPrivate()->MeshIdRegistry;
-	UpdateContext->ImageIdRegistry = CustomizableObject->GetPrivate()->ImageIdRegistry;
-    UpdateContext->MaterialIdRegistry = CustomizableObject->GetPrivate()->MaterialIdRegistry,
-    UpdateContext->SkeletalMeshIdRegistry = CustomizableObject->GetPrivate()->SkeletalMeshIdRegistry,
-	UpdateContext->ExternalResourceProvider = Context->ExternalResourceProvider;
-	UpdateContext->ModelStreamableBulkData = CustomizableObject->GetPrivate()->GetModelStreamableBulkData();
-	UpdateContext->Parameters = Context->Parameters;
-	UpdateContext->State = Context->GetCapturedDescriptor().GetState();
+	TSharedRef<FMutableUpdateImageContext> UpdateImageContext = MakeShared<FMutableUpdateImageContext>();
+	UpdateImageContext->CustomizableObjectPathName = CustomizableObject->GetPathName();
+	UpdateImageContext->InstancePathName = Public->GetPathName();
+	UpdateImageContext->System = UCustomizableObjectSystem::GetInstance()->GetPrivate()->MutableSystem;
+	UpdateImageContext->Model = Context->Model;
+	UpdateImageContext->MeshIdRegistry = CustomizableObject->GetPrivate()->MeshIdRegistry;
+	UpdateImageContext->ImageIdRegistry = CustomizableObject->GetPrivate()->ImageIdRegistry;
+    UpdateImageContext->MaterialIdRegistry = CustomizableObject->GetPrivate()->MaterialIdRegistry,
+    UpdateImageContext->SkeletalMeshIdRegistry = CustomizableObject->GetPrivate()->SkeletalMeshIdRegistry,
+	UpdateImageContext->ExternalResourceProvider = Context->ExternalResourceProvider;
+	UpdateImageContext->ModelStreamableBulkData = CustomizableObject->GetPrivate()->GetModelStreamableBulkData();
+	UpdateImageContext->Parameters = Context->Parameters;
+	UpdateImageContext->State = Context->GetCapturedDescriptor().GetState();
 
 	// Cache the descriptor as a string if we want to later report it using our benchmark utility. 
 	if (FLogBenchmarkUtil::IsBenchmarkingReportingEnabled())
 	{
-		UpdateContext->CapturedDescriptor = Context->GetCapturedDescriptor().ToString();
+		UpdateImageContext->CapturedDescriptor = Context->GetCapturedDescriptor().ToString();
 		if (GWorld)
 		{
-			UpdateContext->bLevelBegunPlay = GWorld->GetBegunPlay();
+			UpdateImageContext->bLevelBegunPlay = GWorld->GetBegunPlay();
 		}
 	}
 
-	const bool bReuseTextures = Context->bBake ? false : Context->bReuseInstanceTextures;
+	Context->UpdateImageContext = UpdateImageContext;
 	
 	TArray<bool> RecreateRenderStateOnInstanceComponent;
 	RecreateRenderStateOnInstanceComponent.Init(false, Context->NumInstanceComponents);
@@ -5779,8 +6058,8 @@ void UCustomizableInstancePrivate::BuildMaterials(const TSharedRef<FUpdateContex
 			continue;
 		}
 
-		const bool bReuseMaterials = !Context->bBake && !Context->MeshChangedPerInstanceComponent[InstanceComponentIndex];
-
+		const bool bUVChanged = Context->MeshChangedPerInstanceComponent[InstanceComponentIndex];
+		
 		// If the mesh is not transient, it means it's pass-through so it should use material overrides and not be modified in any way
 		const bool bIsTransientMesh = SkeletalMesh->HasAllFlags(EObjectFlags::RF_Transient);
 
@@ -5794,11 +6073,14 @@ void UCustomizableInstancePrivate::BuildMaterials(const TSharedRef<FUpdateContex
 		{
 			ComponentData->OverrideMaterials.Reset();
 			ComponentData->OverlayMaterial = nullptr;
-
-			if (Component->OverlayMaterial->PassthroughObject)
+			
+			if (Component->OverlayMaterialIndex != INDEX_NONE)
 			{
-				ComponentData->OverlayMaterial = Component->OverlayMaterial->PassthroughObject.Get();
-				OverlayMaterial = ComponentData->OverlayMaterial;
+				FBuildMaterialOut Out = BuildMaterial(Context, Component->OverlayMaterialIndex, !bUVChanged);
+				OverlayMaterial = Out.Material;
+				ComponentData->OverlayMaterial = Out.Material;
+				
+				NotifyUpdateOnInstanceComponent[InstanceComponentIndex] = Out.bNotifyPrimitiveUpdated;
 			}
 		}
 
@@ -5809,10 +6091,7 @@ void UCustomizableInstancePrivate::BuildMaterials(const TSharedRef<FUpdateContex
 		}
 
 		TArray<FSkeletalMaterial> Materials;
-
-		// Maps serializations of FMutableMaterialPlaceholder to Created Dynamic Material instances, used to reuse materials across LODs
-		TSet<FMutableMaterialPlaceholder> ReuseMaterialCache;
-
+		
 		// SurfaceId per MaterialSlotIndex
 		TArray<int32> SurfaceIdToMaterialIndex;
 
@@ -5853,34 +6132,13 @@ void UCustomizableInstancePrivate::BuildMaterials(const TSharedRef<FUpdateContex
 
 					continue;
 				}
-
-				// Is this a material in a passthrough mesh that we don't modify?
-				if (!Surface->Material)
-				{
-					Materials.Emplace();
-#if WITH_EDITOR
-					// Without this, a change of a referenced material and recompilation doesn't show up in the preview.
-					RecreateRenderStateOnInstanceComponent[InstanceComponentIndex] = true;
-#endif
-					continue;
-				}
-
-				UMaterialInterface* MaterialTemplate = Surface->Material->PassthroughObject.Get();
-				if (!MaterialTemplate)
-				{
-					// Missing MaterialTemplate. Use DefaultMaterial instead. 
-					MaterialTemplate = UMaterial::GetDefaultMaterial(MD_Surface);
-					check(MaterialTemplate);
-					UE_LOG(LogMutable, Error, TEXT("Build Materials: Missing referenced template to use as parent material on CustomizableObject [%s]."), *CustomizableObject->GetName());
-				}
-
+				
 				// This section will require a new slot
 				SurfaceIdToMaterialIndex.Add(Surface->SurfaceId);
 
 				// Add and set up the material data for this slot
 				const int32 MaterialSlotIndex = Materials.Num();
 				FSkeletalMaterial& MaterialSlot = Materials.AddDefaulted_GetRef();
-				MaterialSlot.MaterialInterface = MaterialTemplate;
 
 				uint32 UsedSurfaceMetadataId = Surface->SurfaceMetadataId;
 				
@@ -5944,414 +6202,12 @@ void UCustomizableInstancePrivate::BuildMaterials(const TSharedRef<FUpdateContex
 					}
 				}
 
-				FMutableMaterialPlaceholder MutableMaterialPlaceholder;
-				MutableMaterialPlaceholder.ParentMaterialID = MaterialTemplate->GetUniqueID();
-				MutableMaterialPlaceholder.MatIndex = MaterialSlotIndex;
-
-				{
-					MUTABLE_CPUPROFILER_SCOPE(ParamLoop);
-
-					for (int32 VectorIndex = 0; VectorIndex < Surface->VectorCount; ++VectorIndex)
-					{
-						const FInstanceUpdateData::FVector& Vector = Context->InstanceUpdateData.Vectors[Surface->FirstVector + VectorIndex];
-						MutableMaterialPlaceholder.AddParam(FMutableMaterialPlaceholder::FMutableMaterialPlaceHolderParam(FName(Vector.Name), Vector.MaterialLayer, Vector.Vector));
-					}
-
-					for (int32 ScalarIndex = 0; ScalarIndex < Surface->ScalarCount; ++ScalarIndex)
-					{
-						const FInstanceUpdateData::FScalar& Scalar = Context->InstanceUpdateData.Scalars[Surface->FirstScalar + ScalarIndex];
-						MutableMaterialPlaceholder.AddParam(FMutableMaterialPlaceholder::FMutableMaterialPlaceHolderParam(FName(Scalar.Name), Scalar.MaterialLayer, Scalar.Scalar));
-					}
-				}
-
-				{
-					MUTABLE_CPUPROFILER_SCOPE(BuildMaterials_ImageLoop);
-
-					// Get the cache of resources of all live instances of this object
-					FTextureCache& TextureCache = CustomizableObject->GetPrivate()->TextureCache;
-
-					FString CurrentState = Public->GetCurrentState();
-					bool bNeverStream = Context->bNeverStream;
-
-					check((bNeverStream && Context->MipsToSkip == 0) ||
-						(!bNeverStream && Context->MipsToSkip >= 0));
-
-					for (int32 ImageIndex = 0; ImageIndex < Surface->ImageCount; ++ImageIndex)
-					{
-						const TSharedRef<const FInstanceUpdateData::FImage>& Image = Context->InstanceUpdateData.Images[Surface->FirstImage + ImageIndex];
-						FString KeyName = Image->Name.ToString();
-						TSharedPtr<const UE::Mutable::Private::FImage> MutableImage = Image->Image;
-
-						UTexture2D* MutableTexture = nullptr; // Texture generated by mutable
-						UTexture* PassThroughTexture = nullptr; // Texture not generated by mutable
-
-						const FTextureReuseCacheKey TextureReuseCacheRef = bReuseTextures ? FTextureReuseCacheKey{Image->BaseLOD, ObjectComponentIndex.GetValue(), Surface->SurfaceId, ImageIndex} : FTextureReuseCacheKey{};
-
-#if WITH_EDITOR
-						const FTextureCache::FId TextureCacheKey = FTextureCache::FId(Image->ImageID, Context->MipsToSkip, Context->bBake);
-#else
-						const FTextureCache::FId TextureCacheKey = FTextureCache::FId(Image->ImageID, Context->MipsToSkip);
-#endif
-						
-						// If the mutable image is null, it must be in the cache
-						if (!MutableImage)
-						{
-							MutableTexture = TextureCache.Get(TextureCacheKey);
-							check(MutableTexture);
-						}
-
-						// Check if the image is a reference to an engine texture
-						if (MutableImage && Image->bIsPassThrough)
-						{
-							check(MutableImage->IsReference());
-
-							PassThroughTexture = MutableImage->PassthroughObject.Get();
-							
-							if (!PassThroughTexture)
-							{
-								continue;
-							}
-						}
-
-						// Find the additional information for this image
-						int32 ImageKey = FCString::Atoi(*KeyName);
-						if (ImageKey >= 0 && ImageKey < ModelResources.ImageProperties.Num())
-						{
-							const FMutableModelImageProperties& Props = ModelResources.ImageProperties[ImageKey];
-
-							if (!MutableTexture && !PassThroughTexture && MutableImage)
-							{
-								TWeakObjectPtr<UTexture2D>* ReusedTexture = bReuseTextures ? TextureReuseCache.Find(TextureReuseCacheRef) : nullptr;
-								
-								// This shared ptr will hold the reused texture platform data (mips) until the reused texture is updated 
-								// and delete it automatically
-								TSharedPtr<FTexturePlatformData, ESPMode::ThreadSafe> ReusedTexturePlatformData;
-
-								FString MutableTextureName;
-								{
-									const uint32 MutableImageHash = UE::Mutable::Private::GetTypeHashPersistent(Image->ImageID);
-									MutableTextureName = FString::Printf(TEXT("T_%s_h%llu"), *Props.TextureParameterName, MutableImageHash);
-									MutableTextureName.ReplaceInline(TEXT(" "), TEXT("_"));
-
-#if WITH_EDITOR
-									if (Context->bBake)
-									{
-										// If this is a bake operation just add the baked prefix to make the name unique for all baked instances of this resource
-										MutableTextureName = FBakingConfiguration::BakedResourcePrefix + MutableTextureName;
-									}
-#endif
-								}
-								
-								const FName UniqueMutableTextureName = MakeUniqueObjectName(GetTransientPackage(), UTexture2D::StaticClass(), *MutableTextureName);
-
-								if (ReusedTexture && (*ReusedTexture).IsValid() && !(*ReusedTexture)->HasAnyFlags(RF_BeginDestroyed))
-								{
-									// Only uncompressed textures can be reused. This also fixes an issue in the editor where textures supposedly 
-									// uncompressed by their state, are still compressed because the CO has not been compiled at maximum settings
-									// and the uncompressed setting cannot be applied to them.
-									EPixelFormat PixelFormat = (*ReusedTexture)->GetPixelFormat();
-
-									if (PixelFormat == EPixelFormat::PF_R8G8B8A8)
-									{
-										MutableTexture = (*ReusedTexture).Get();
-										check(MutableTexture != nullptr);
-									}
-									else
-									{
-										ReusedTexture = nullptr;
-										MutableTexture = CreateTexture(UniqueMutableTextureName);
-#if WITH_EDITOR
-										UE_LOG(LogMutable, Warning,
-											TEXT("Tried to reuse an uncompressed texture with name %s. Make sure the selected Mutable state disables texture compression/streaming, that one of the state's runtime parameters affects the texture and that the CO is compiled with max. optimization settings."),
-											*MutableTexture->GetName());
-#endif
-									}
-								}
-								else
-								{
-									ReusedTexture = nullptr;
-									MutableTexture = CreateTexture(UniqueMutableTextureName);
-								}
-
-								if (MutableTexture)
-								{
-									if (Context->ImageToPlatformDataMap.Contains(Image->ImageID))
-									{
-										SetTexturePropertiesFromMutableImageProps(MutableTexture, Props, bNeverStream);
-
-										FTexturePlatformData* PlatformData = Context->ImageToPlatformDataMap[Image->ImageID];
-
-										if (ReusedTexture)
-										{
-											check(PlatformData->Mips.Num() == MutableTexture->GetPlatformData()->Mips.Num());
-											check(PlatformData->Mips[0].SizeX == MutableTexture->GetPlatformData()->Mips[0].SizeX);
-											check(PlatformData->Mips[0].SizeY == MutableTexture->GetPlatformData()->Mips[0].SizeY);
-
-											// Now the ReusedTexturePlatformData shared ptr owns the platform data
-											ReusedTexturePlatformData = TSharedPtr<FTexturePlatformData, ESPMode::ThreadSafe>(PlatformData);
-										}
-										else
-										{
-											// Now the MutableTexture owns the platform data
-											MutableTexture->SetPlatformData(PlatformData);
-										}
-
-										Context->ImageToPlatformDataMap.Remove(Image->ImageID);
-									}
-									else
-									{
-										UE_LOG(LogMutable, Error, TEXT("Required image [%s] with ID [%s] was not generated in the mutable thread, and it is not cached. LOD [%d]. Object Component [%d]"),
-											*Props.TextureParameterName,
-											*Image->ImageID.ToString(),
-											LODIndex, ObjectComponentIndex.GetValue());
-										continue;
-									}
-
-									if (bNeverStream)
-									{
-										// To prevent LogTexture Error "Loading non-streamed mips from an external bulk file."
-										for (int32 i = 0; i < MutableTexture->GetPlatformData()->Mips.Num(); ++i)
-										{
-											MutableTexture->GetPlatformData()->Mips[i].BulkData.ClearBulkDataFlags(BULKDATA_PayloadInSeparateFile);
-										}
-									}
-
-									{
-										MUTABLE_CPUPROFILER_SCOPE(UpdateResource);
-#if REQUIRES_SINGLEUSE_FLAG_FOR_RUNTIME_TEXTURES
-										for (int32 i = 0; i < MutableTexture->GetPlatformData()->Mips.Num(); ++i)
-										{
-											uint32 DataFlags = MutableTexture->GetPlatformData()->Mips[i].BulkData.GetBulkDataFlags();
-											MutableTexture->GetPlatformData()->Mips[i].BulkData.SetBulkDataFlags(DataFlags | BULKDATA_SingleUse);
-										}
-#endif
-
-										if (ReusedTexture)
-										{
-											// Must remove texture from cache since it will be reused with a different ImageID
-											TextureCache.Remove(*MutableTexture);
-											
-											check(ReusedTexturePlatformData.IsValid());
-
-											if (ReusedTexturePlatformData.IsValid())
-											{
-												TSharedRef<FTexturePlatformData, ESPMode::ThreadSafe> PlatformDataRef = ReusedTexturePlatformData.ToSharedRef();
-												ReuseTexture(MutableTexture, PlatformDataRef);
-											}
-										}
-										else if (MutableTexture)
-										{	
-											//if (!bNeverStream) // No need to check bNeverStream. In that case, the texture won't use 
-											// the MutableMipDataProviderFactory anyway and it's needed for detecting Mutable textures elsewhere
-											UMutableTextureMipDataProviderFactory* MutableMipDataProviderFactory = Cast<UMutableTextureMipDataProviderFactory>(MutableTexture->GetAssetUserDataOfClass(UMutableTextureMipDataProviderFactory::StaticClass()));
-											if (!MutableMipDataProviderFactory)
-											{
-												MutableMipDataProviderFactory = NewObject<UMutableTextureMipDataProviderFactory>();
-
-												if (MutableMipDataProviderFactory)
-												{
-													MutableMipDataProviderFactory->CustomizableObjectInstance = Public;
-													check(LODIndex < 256 && InstanceComponentIndex < 256 && ImageIndex < 256);
-													MutableMipDataProviderFactory->ImageRef.ImageID = Image->ImageID;
-													MutableMipDataProviderFactory->ImageRef.BaseMip = uint8(Image->BaseMip);
-													MutableMipDataProviderFactory->ImageRef.ConstantImagesNeededToGenerate = Image->ConstantImagesNeededToGenerate;
-													MutableMipDataProviderFactory->UpdateContext = UpdateContext;
-													MutableTexture->AddAssetUserData(MutableMipDataProviderFactory);
-												}
-											}
-
-											MutableTexture->UpdateResource();
-										}
-									}
-									
-									TextureCache.Add(TextureCacheKey, MutableTexture);
-								}
-								else
-								{
-									UE_LOG(LogMutable, Error, TEXT("Texture creation failed."));
-								}
-							}
-
-							FGeneratedTexture TextureData;
-							TextureData.Name = Props.TextureParameterName;
-							TextureData.Texture = MutableTexture ? MutableTexture : PassThroughTexture;
-							
-							// Only add textures generated by mutable to the cache
-							if (MutableTexture)
-							{
-								NewGeneratedTextures.Add(TextureData);
-							}
-
-							MutableMaterialPlaceholder.AddParam(FMutableMaterialPlaceholder::FMutableMaterialPlaceHolderParam(*Props.TextureParameterName, Image->MaterialLayer, TextureData));
-						}
-						else
-						{
-							// This means the compiled model (maybe coming from derived data) has images that the asset doesn't know about.
-							UE_LOG(LogMutable, Error, TEXT("CustomizableObject derived data out of sync with asset for [%s]. Try recompiling it."), *CustomizableObject->GetName());
-						}
-
-						if (bReuseTextures)
-						{
-							if (MutableTexture)
-							{
-								TextureReuseCache.Add(TextureReuseCacheRef, MutableTexture);
-							}
-							else
-							{
-								TextureReuseCache.Remove(TextureReuseCacheRef);
-							}
-						}
-					}
-				}
-
-				// Find or create the material for this slot
-				UMaterialInterface* MaterialInterface = MaterialSlot.MaterialInterface;
-
-				if (FMutableMaterialPlaceholder* FoundMaterialPlaceholder = ReuseMaterialCache.Find(MutableMaterialPlaceholder))
-				{
-					MaterialInterface = Materials[FoundMaterialPlaceholder->MatIndex].MaterialInterface;
-				}
-				else // Material not cached, create a new one
-				{
-					MUTABLE_CPUPROFILER_SCOPE(BuildMaterials_CreateMaterial);
-
-					ReuseMaterialCache.Add(MutableMaterialPlaceholder);
-					
-					FGeneratedMaterial& Material = GeneratedMaterials.AddDefaulted_GetRef();
-					Material.SurfaceId = Surface->SurfaceId;
-					Material.Material = Surface->Material;
-					Material.MaterialInterface = MaterialInterface;
-
-#if WITH_EDITORONLY_DATA
-					Material.ComponentName = ComponentName;
-#endif
-
-					UMaterialInstanceDynamic* MaterialInstance = nullptr;
-
-					if (const int32 OldMaterialIndex = OldGeneratedMaterials.Find(Material); bReuseMaterials && OldMaterialIndex != INDEX_NONE)
-					{
-						const FGeneratedMaterial& OldMaterial = OldGeneratedMaterials[OldMaterialIndex];
-						MaterialInstance = Cast<UMaterialInstanceDynamic>(OldMaterial.MaterialInterface);
-						Material.MaterialInterface = OldMaterial.MaterialInterface;
-					}
-					
-					if (!MaterialInstance && MutableMaterialPlaceholder.GetParams().Num() != 0)
-					{
-#if WITH_EDITOR
-						
-						FString MIDName = MaterialTemplate->GetName();
-						{
-							// Remove the MI_ or M_ prefixes from the material string so we use it as the name of the MID
-							const FString MaterialPrefix = TEXT("M_");						// Material
-							const FString MaterialInstancePrefix = TEXT("MI_");				// Material Instance
-							const FString MaterialInstanceConstantPrefix = TEXT("MIC_");	// Material Instance Constant
-							
-							if (MIDName.Find(MaterialInstancePrefix, ESearchCase::CaseSensitive) == 0)
-							{
-								MIDName = MIDName.RightChop(MaterialInstancePrefix.Len());
-							}
-							else if (MIDName.Find(MaterialPrefix, ESearchCase::CaseSensitive) == 0)
-							{
-								MIDName = MIDName.RightChop(MaterialPrefix.Len());
-							}
-							else if (MIDName.Find(MaterialInstanceConstantPrefix, ESearchCase::CaseSensitive) == 0)
-							{
-								MIDName = MIDName.RightChop(MaterialInstanceConstantPrefix.Len());
-							}
-
-							const uint32 MaterialPlaceHolderHash = GetHash(MutableMaterialPlaceholder, Context->bBake);
-							MIDName = FString::Printf(TEXT("MID_%s_h%lu"), *MIDName, MaterialPlaceHolderHash);
-
-							// Add a "tag" to better identify this asset during the bake
-							if (Context->bBake)
-							{
-								MIDName = FBakingConfiguration::BakedResourcePrefix + MIDName;
-							}
-						}
-						
-						const FName UniqueMIDName = MakeUniqueObjectName(GetTransientPackage(), UMaterialInstanceDynamic::StaticClass(), *MIDName);
-						MaterialInstance = UMaterialInstanceDynamic::Create(MaterialTemplate, GetTransientPackage(), UniqueMIDName);
-#else
-						MaterialInstance = UMaterialInstanceDynamic::Create(MaterialTemplate, GetTransientPackage());
-#endif
-						
-						Material.MaterialInterface = MaterialInstance;
-					}
-
-					if (MaterialInstance)
-					{
-						for (const FMutableMaterialPlaceholder::FMutableMaterialPlaceHolderParam& Param : MutableMaterialPlaceholder.GetParams())
-						{
-							//TODO: Add the possibility to modify blend parameters not just layer parameters.
-							switch (Param.Type)
-							{
-							case FMutableMaterialPlaceholder::EPlaceHolderParamType::Vector:
-								if (Param.LayerIndex < 0)
-								{
-									FLinearColor Color = Param.Vector;
-
-									// HACK: We encode an invalid value (Nan) for table option "None.
-									// Decoding "None" color parameters that use the material color
-									if (FMath::IsNaN(Color.R))
-									{
-										FMaterialParameterInfo ParameterInfo(Param.ParamName);
-										MaterialTemplate->GetVectorParameterValue(ParameterInfo, Color);
-									}
-
-									MaterialInstance->SetVectorParameterValue(Param.ParamName, Color);
-								}
-								else
-								{
-									FMaterialParameterInfo ParameterInfo = FMaterialParameterInfo(Param.ParamName, EMaterialParameterAssociation::LayerParameter, Param.LayerIndex);
-									MaterialInstance->SetVectorParameterValueByInfo(ParameterInfo, Param.Vector);
-								}
-
-								break;
-
-							case FMutableMaterialPlaceholder::EPlaceHolderParamType::Scalar:
-								if (Param.LayerIndex < 0)
-								{
-									MaterialInstance->SetScalarParameterValue(FName(Param.ParamName), Param.Scalar);
-								}
-								else
-								{
-									FMaterialParameterInfo ParameterInfo = FMaterialParameterInfo(Param.ParamName, EMaterialParameterAssociation::LayerParameter, Param.LayerIndex);
-									MaterialInstance->SetScalarParameterValueByInfo(ParameterInfo, Param.Scalar);
-								}
-
-								break;
-
-							case FMutableMaterialPlaceholder::EPlaceHolderParamType::Texture:
-								if (Param.LayerIndex < 0)
-								{
-									MaterialInstance->SetTextureParameterValue(Param.ParamName, Param.Texture.Texture);
-								}
-								else
-								{
-									FMaterialParameterInfo ParameterInfo = FMaterialParameterInfo(Param.ParamName, EMaterialParameterAssociation::LayerParameter, Param.LayerIndex);
-									MaterialInstance->SetTextureParameterValueByInfo(ParameterInfo, Param.Texture.Texture);
-								}
-
-								if (!bDisableNotifyComponentsOfTextureUpdates)
-								{
-									NotifyUpdateOnInstanceComponent[InstanceComponentIndex] = true;
-								}
-
-								Material.Textures.Add(Param.Texture);
-
-								break;
-							}
-						}
-					}
-
-					MaterialInterface = Material.MaterialInterface;
-				}
-
-				// Assign the material to the slot, and add it to the  OverrideMaterials
-				MaterialSlot.MaterialInterface = MaterialInterface;
 				if (ComponentData)
 				{
-					ComponentData->OverrideMaterials.Add(MaterialInterface);
+					FBuildMaterialOut Out = BuildMaterial(Context, Surface->MaterialIndex, !bUVChanged);
+					MaterialSlot.MaterialInterface = Out.Material;
+					ComponentData->OverrideMaterials.Add(Out.Material);
+					NotifyUpdateOnInstanceComponent[InstanceComponentIndex] = Out.bNotifyPrimitiveUpdated;
 				}
 			}
 		}
@@ -6464,7 +6320,8 @@ void UCustomizableInstancePrivate::BuildMaterials(const TSharedRef<FUpdateContex
 	{
 		MUTABLE_CPUPROFILER_SCOPE(BuildMaterials_Exchange);
 
-		Exchange(GeneratedTextures, NewGeneratedTextures);
+		Context->Instance->GetPrivate()->LastMaterialsReuseCache = MoveTemp(Context->Instance->GetPrivate()->CurrentMaterialsReuseCache);
+		Context->Instance->GetPrivate()->LastTextureReuseCache = MoveTemp(Context->Instance->GetPrivate()->CurrentTextureReuseCache);
 	}
 }
 
@@ -6477,7 +6334,7 @@ void UCustomizableObjectInstance::SetReplacePhysicsAssets(bool bReplaceEnabled)
 
 void UCustomizableObjectInstance::SetReuseInstanceTextures(bool bTextureReuseEnabled)
 {
-	bTextureReuseEnabled ? GetPrivate()->SetCOInstanceFlags(ReuseTextures) : GetPrivate()->ClearCOInstanceFlags(ReuseTextures);
+	GetPrivate()->bReuseInstanceTextures = bTextureReuseEnabled;
 }
 
 
diff --git a/Engine/Plugins/Mutable/Source/CustomizableObject/Private/MuCO/CustomizableObjectSystem.cpp b/Engine/Plugins/Mutable/Source/CustomizableObject/Private/MuCO/CustomizableObjectSystem.cpp
--- a/Engine/Plugins/Mutable/Source/CustomizableObject/Private/MuCO/CustomizableObjectSystem.cpp
+++ b/Engine/Plugins/Mutable/Source/CustomizableObject/Private/MuCO/CustomizableObjectSystem.cpp
@@ -273,6 +273,8 @@ FUpdateContextPrivate::FUpdateContextPrivate(UCustomizableObjectInstance& InInst
 #if WITH_EDITOR
 	PixelFormatOverride = SystemPrivate->ImageFormatOverrideFunc;
 #endif
+
+	bReuseInstanceTextures = InInstance.GetPrivate()->bReuseInstanceTextures;
 	
 	bValid = true;
 }
@@ -750,6 +752,41 @@ void FMutablePendingInstanceWork::AddIDRelease(UE::Mutable::Private::FInstance::
 }
 
 
+uint32 GetMaterialHash(const FInstanceUpdateData& UpdateData, const FInstanceUpdateData::FMaterial& Material, bool bPersistent)
+{
+	uint32 Hash = bPersistent ? UE::Mutable::Private::GetTypeHashPersistent(Material.MaterialId) : GetTypeHash(Material.MaterialId);
+	
+	for (int32 Index = Material.FirstScalar; Index < Material.FirstScalar + Material.ScalarCount; ++Index)
+	{
+		const FInstanceUpdateData::FScalar& Scalar = UpdateData.Scalars[Index];
+		
+		Hash = HashCombine(Hash, GetTypeHash(Scalar.Name));
+		Hash = HashCombine(Hash, GetTypeHash(Scalar.Scalar));
+		Hash = HashCombine(Hash, Scalar.MaterialLayer);
+	}
+
+	for (int32 Index = Material.FirstVector; Index < Material.FirstVector + Material.VectorCount; ++Index)
+	{
+		const FInstanceUpdateData::FVector& Vector = UpdateData.Vectors[Index];
+		
+		Hash = HashCombine(Hash, GetTypeHash(Vector.Name));
+		Hash = HashCombine(Hash, GetTypeHash(Vector.Vector));
+		Hash = HashCombine(Hash, Vector.MaterialLayer);
+	}
+
+	for (int32 Index = Material.FirstImage; Index < Material.FirstImage + Material.ImageCount; ++Index)
+	{
+		const TSharedRef<FInstanceUpdateData::FImage>& Image = UpdateData.Images[Index];
+		
+		Hash = HashCombine(Hash, GetTypeHash(Image->ParameterName));
+		Hash = HashCombine(Hash, GetTypeHash(Image->ImageID));
+		Hash = HashCombine(Hash, Image->MaterialLayer);
+	}
+
+	return Hash;
+}
+
+
 UCustomizableObjectSystem* UCustomizableObjectSystem::GetInstance()
 {
 	if (!UCustomizableObjectSystemPrivate::SSystem)
@@ -965,11 +1002,10 @@ static FAutoConsoleVariableRef CVarEnableMutableLiveUpdate(
 	ECVF_Default);
 
 
-int32 UCustomizableObjectSystemPrivate::EnableReuseInstanceTextures = 1;
-
-static FAutoConsoleVariableRef CVarEnableMutableReuseInstanceTextures(
-	TEXT("mutable.EnableReuseInstanceTextures"), UCustomizableObjectSystemPrivate::EnableReuseInstanceTextures,
-	TEXT("If set to 1 or greater and set in the corresponding setting in the current Mutable state, Mutable can reuse instance UTextures (only uncompressed and not streaming, so set the options in the state) and their resources between updates when they are modified. If geometry or state is changed they cannot be reused."),
+static TAutoConsoleVariable<bool> CVarEnableMutableReuseInstanceTextures(
+	TEXT("mutable.EnableReuseInstanceTextures"), 
+	true,
+	TEXT("If set to true set in the corresponding setting in the current Mutable state, Mutable can reuse instance UTextures (only uncompressed and not streaming, so set the options in the state) and their resources between updates when they are modified. If geometry or state is changed they cannot be reused."),
 	ECVF_Default);
 
 
@@ -1845,9 +1881,9 @@ bool UCustomizableObjectSystem::CheckIfDiskOrMipUpdateOperationsPending(const UC
 	{
 		if (IsValid(*CustomizableObjectInstance) && CustomizableObjectInstance->GetCustomizableObject() == &Object)
 		{
-			for (const FGeneratedTexture& GeneratedTexture : CustomizableObjectInstance->GetPrivate()->GeneratedTextures)
+			for (const TObjectPtr<UTexture>& GeneratedTexture : CustomizableObjectInstance->GetPrivate()->GeneratedTextures)
 			{
-				if (GeneratedTexture.Texture->HasPendingInitOrStreaming())
+				if (GeneratedTexture->HasPendingInitOrStreaming())
 				{
 					return true;
 				}
@@ -2214,9 +2250,9 @@ namespace impl
 	}
 	
 
-	void Task_Mutable_GetMaterials(const TSharedRef<FUpdateContextPrivate>& OperationData)
+	void Task_Mutable_GetSurfaces(const TSharedRef<FUpdateContextPrivate>& OperationData)
 	{
-		MUTABLE_CPUPROFILER_SCOPE(Task_Mutable_GetMaterials)
+		MUTABLE_CPUPROFILER_SCOPE(Task_Mutable_GetSurfaces)
 		
 		const UCustomizableObjectInstance* Instance = OperationData->Instance.Get();
 		const UCustomizableObject* CustomizableObject = Instance->GetCustomizableObject();
@@ -2271,105 +2307,168 @@ namespace impl
 						continue;
 					}
 
-					OperationData->InstanceUpdateData.Surfaces.Add(MakeShared<FInstanceUpdateData::FSurface>());
-					TSharedRef<FInstanceUpdateData::FSurface>& Surface = OperationData->InstanceUpdateData.Surfaces.Last();
+					TSharedRef<FInstanceUpdateData::FSurface>& Surface = OperationData->InstanceUpdateData.Surfaces.Add_GetRef(MakeShared<FInstanceUpdateData::FSurface>());
 
 					Surface->SurfaceId = SurfaceId;
 					Surface->SurfaceMetadataId = MutableInstance->GetSurfaceCustomId(InstanceComponentIndex, MutableLODIndex, InstanceSurfaceIndex);
-					Surface->MaterialId = MutableInstance->GetMaterialId(InstanceComponentIndex, MutableLODIndex, InstanceSurfaceIndex);
 
-					UE::Tasks::TTask<TSharedPtr<const UE::Mutable::Private::FMaterial>> GetMaterialTask = OperationData->MutableSystem->GetMaterial(LastTask, OperationData->InstanceID, Surface->MaterialId);
-					UE::Tasks::AddNested(GetMaterialTask);
-					LastTask = GetMaterialTask;
-					
-					UE::Tasks::AddNested(UE::Tasks::Launch(TEXT("GetMaterialResult"), [Surface, GetMaterialTask]() mutable
+					if (UE::Mutable::Private::FMaterialId MaterialId = MutableInstance->GetMaterialId(InstanceComponentIndex, MutableLODIndex, InstanceSurfaceIndex))
 					{
-						Surface->Material = GetMaterialTask.GetResult();
-					},
-					GetMaterialTask,
-					LowLevelTasks::ETaskPriority::Inherit));
+						Surface->MaterialIndex = OperationData->InstanceUpdateData.Materials.Num();
+						TSharedRef<FInstanceUpdateData::FMaterial>& Material = OperationData->InstanceUpdateData.Materials.Add_GetRef(MakeShared<FInstanceUpdateData::FMaterial>());
 
-					// Vectors
-					Surface->FirstVector = OperationData->InstanceUpdateData.Vectors.Num();
-					Surface->VectorCount = MutableInstance->GetVectorCount(InstanceComponentIndex, MutableLODIndex, InstanceSurfaceIndex);
-					for (int32 VectorIndex = 0; VectorIndex < Surface->VectorCount; ++VectorIndex)
-					{
-						MUTABLE_CPUPROFILER_SCOPE(GetVector);
-						OperationData->InstanceUpdateData.Vectors.Push({});
-						FInstanceUpdateData::FVector& Vector = OperationData->InstanceUpdateData.Vectors.Last();
-						Vector.Name = MutableInstance->GetVectorName(InstanceComponentIndex, MutableLODIndex, InstanceSurfaceIndex, VectorIndex);
-						Vector.Vector = MutableInstance->GetVector(InstanceComponentIndex, MutableLODIndex, InstanceSurfaceIndex, VectorIndex);
-						Vector.MaterialLayer = MutableInstance->GetVectorLayerIndex(InstanceComponentIndex, MutableLODIndex, InstanceSurfaceIndex, VectorIndex);
+						Material->MaterialId = MaterialId;
+
+						// TODO MATERIAL_PIN Remove
+
+						Material->bGetParameters = false;
+
+						// Vectors
+						Material->FirstVector = OperationData->InstanceUpdateData.Vectors.Num();
+						Material->VectorCount = MutableInstance->GetVectorCount(InstanceComponentIndex, MutableLODIndex, InstanceSurfaceIndex);
+						for (int32 VectorIndex = 0; VectorIndex < Material->VectorCount; ++VectorIndex)
+						{
+							MUTABLE_CPUPROFILER_SCOPE(GetVector);
+							OperationData->InstanceUpdateData.Vectors.Push({});
+							FInstanceUpdateData::FVector& Vector = OperationData->InstanceUpdateData.Vectors.Last();
+							Vector.Name = MutableInstance->GetVectorName(InstanceComponentIndex, MutableLODIndex, InstanceSurfaceIndex, VectorIndex);
+							Vector.Vector = MutableInstance->GetVector(InstanceComponentIndex, MutableLODIndex, InstanceSurfaceIndex, VectorIndex);
+							Vector.MaterialLayer = MutableInstance->GetVectorLayerIndex(InstanceComponentIndex, MutableLODIndex, InstanceSurfaceIndex, VectorIndex);
+						}
+
+						// Scalars
+						Material->FirstScalar = OperationData->InstanceUpdateData.Scalars.Num();
+						Material->ScalarCount = MutableInstance->GetScalarCount(InstanceComponentIndex, MutableLODIndex, InstanceSurfaceIndex);
+						for (int32 ScalarIndex = 0; ScalarIndex < Material->ScalarCount; ++ScalarIndex)
+						{
+							MUTABLE_CPUPROFILER_SCOPE(GetScalar)
+
+							const FName ScalarName = MutableInstance->GetScalarName(InstanceComponentIndex, MutableLODIndex, InstanceSurfaceIndex, ScalarIndex);
+							const float ScalarValue = MutableInstance->GetScalar(InstanceComponentIndex, MutableLODIndex, InstanceSurfaceIndex, ScalarIndex);
+
+							const int32 MaterialLayer = MutableInstance->GetScalarLayerIndex(InstanceComponentIndex, MutableLODIndex, InstanceSurfaceIndex, ScalarIndex);
+
+							OperationData->InstanceUpdateData.Scalars.Push({ ScalarName, ScalarValue, MaterialLayer });
+						}
+
+						// Images
+						Material->FirstImage = OperationData->InstanceUpdateData.Images.Num();
+						Material->ImageCount = MutableInstance->GetImageCount(InstanceComponentIndex, MutableLODIndex, InstanceSurfaceIndex);
+						for (int32 ImageIndex = 0; ImageIndex < Material->ImageCount; ++ImageIndex)
+						{
+							MUTABLE_CPUPROFILER_SCOPE(GetImageId);
+							
+							TSharedRef<FInstanceUpdateData::FImage>& Image = OperationData->InstanceUpdateData.Images.Add_GetRef(MakeShared<FInstanceUpdateData::FImage>());
+
+							Image->Name = MutableInstance->GetImageName(InstanceComponentIndex, MutableLODIndex, InstanceSurfaceIndex, ImageIndex);
+							Image->ImageID = MutableInstance->GetImageId(InstanceComponentIndex, MutableLODIndex, InstanceSurfaceIndex, ImageIndex);
+							Image->MaterialLayer = MutableInstance->GetImageLayerIndex(InstanceComponentIndex, MutableLODIndex, InstanceSurfaceIndex, ImageIndex);
+							Image->FullImageSizeX = 0;
+							Image->FullImageSizeY = 0;
+							Image->BaseMip = 0;
+
+							FString KeyName = Image->Name.ToString();
+							int32 ImageKey = FCString::Atoi(*KeyName);
+
+							if (ImageKey >= 0 && ImageKey < ModelResources.ImageProperties.Num())
+							{
+								const FMutableModelImageProperties& Props = ModelResources.ImageProperties[ImageKey];
+
+								Image->ParameterName = FName(Props.TextureParameterName);
+								Image->bIsNonProgressive = Props.MipGenSettings == TMGS_NoMipmaps;
+								Image->bIsPassThrough = Props.IsPassThrough;
+							}
+							else
+							{
+								// This means the compiled model (maybe coming from derived data) has images that the asset doesn't know about.
+								UE_LOG(LogMutable, Error, TEXT("CustomizableObject derived data out of sync with asset for [%s]. Try recompiling it."), *CustomizableObject->GetName());
+							}
+						}
 					}
+				}
+			}
+		}
+	}
 
-					// Scalars
-					Surface->FirstScalar = OperationData->InstanceUpdateData.Scalars.Num();
-					Surface->ScalarCount = MutableInstance->GetScalarCount(InstanceComponentIndex, MutableLODIndex, InstanceSurfaceIndex);
-					for (int32 ScalarIndex = 0; ScalarIndex < Surface->ScalarCount; ++ScalarIndex)
-					{
-						MUTABLE_CPUPROFILER_SCOPE(GetScalar)
 
-						const FName ScalarName = MutableInstance->GetScalarName(InstanceComponentIndex, MutableLODIndex, InstanceSurfaceIndex, ScalarIndex);
-						const float ScalarValue = MutableInstance->GetScalar(InstanceComponentIndex, MutableLODIndex, InstanceSurfaceIndex, ScalarIndex);
+	void Task_Mutable_GetMaterials(const TSharedRef<FUpdateContextPrivate>& Context)
+	{
+		MUTABLE_CPUPROFILER_SCOPE(Task_Mutable_GetMaterials)
 
-						const int32 MaterialLayer = MutableInstance->GetScalarLayerIndex(InstanceComponentIndex, MutableLODIndex, InstanceSurfaceIndex, ScalarIndex);
+		const UCustomizableObjectInstance* Instance = Context->Instance.Get();
+		const UCustomizableObject* CustomizableObject = Instance->GetCustomizableObject();
+		const UModelResources* ModelResources = CustomizableObject->GetPrivate()->GetModelResources();
 
-						OperationData->InstanceUpdateData.Scalars.Push({ ScalarName, ScalarValue, MaterialLayer });
+		UE::Tasks::FTask LastTask = UE::Tasks::MakeCompletedTask<void>();
+		
+		for (int32 MaterialIndex = 0; MaterialIndex < Context->InstanceUpdateData.Materials.Num(); ++MaterialIndex)
+		{
+			TSharedRef<FInstanceUpdateData::FMaterial> Material = Context->InstanceUpdateData.Materials[MaterialIndex];
+			
+			UE::Tasks::TTask<TSharedPtr<const UE::Mutable::Private::FMaterial>> GetMaterialTask = Context->MutableSystem->GetMaterial(LastTask, Context->InstanceID, Material->MaterialId);
+			UE::Tasks::AddNested(GetMaterialTask);
+			LastTask = GetMaterialTask;
+			
+			LastTask = UE::Tasks::Launch(TEXT("GetMaterialResult"), [Context, ModelResources = TStrongObjectPtr(ModelResources), Material, GetMaterialTask]() mutable
+			{
+				Material->Material = GetMaterialTask.GetResult();
+
+				if (Material->bGetParameters)
+				{
+					// Scalars
+					Material->FirstScalar = Context->InstanceUpdateData.Scalars.Num();
+					Material->ScalarCount = Material->Material->ScalarParameters.Num();
+					for (TTuple<FName, float> ScalarParameter : Material->Material->ScalarParameters)
+					{
+						FInstanceUpdateData::FScalar& Scalar = Context->InstanceUpdateData.Scalars.AddDefaulted_GetRef();
+						Scalar.Name = ScalarParameter.Key;
+						Scalar.Scalar = ScalarParameter.Value;
+						//Scalar.MaterialLayer = ScalarParameter.Key.Layer; // TODO MATERIAL_PIN
+					}
+
+					// Vectors
+					Material->FirstVector = Context->InstanceUpdateData.Vectors.Num();
+					Material->VectorCount = Material->Material->ColorParameters.Num();
+					for (TTuple<FName, UE::Math::TVector4<float>> VectorParameter : Material->Material->ColorParameters)
+					{
+						FInstanceUpdateData::FVector& Vector = Context->InstanceUpdateData.Vectors.AddDefaulted_GetRef();
+						Vector.Name = VectorParameter.Key;
+						Vector.Vector = VectorParameter.Value;
+						//Vector.MaterialLayer =  VectorParameter.Key.Layer; // TODO MATERIAL_PIN
 					}
 
 					// Images
-					Surface->FirstImage = OperationData->InstanceUpdateData.Images.Num();
-					Surface->ImageCount = MutableInstance->GetImageCount(InstanceComponentIndex, MutableLODIndex, InstanceSurfaceIndex);
-					for (int32 ImageIndex = 0; ImageIndex < Surface->ImageCount; ++ImageIndex)
+					Material->FirstImage = Context->InstanceUpdateData.Images.Num();
+					Material->ImageCount = Material->Material->ImageParameters.Num();
+					for (TTuple<FName, TVariant<uint32, TSharedPtr<const UE::Mutable::Private::FImage>>> ImageParameter : Material->Material->ImageParameters)
 					{
-						MUTABLE_CPUPROFILER_SCOPE(GetImageId);
+						const UE::Mutable::Private::OP::ADDRESS Address = ImageParameter.Value.Get<UE::Mutable::Private::OP::ADDRESS>();
 						
-						TSharedRef<FInstanceUpdateData::FImage>& Image = OperationData->InstanceUpdateData.Images.Add_GetRef(MakeShared<FInstanceUpdateData::FImage>());
-						Image->Name = MutableInstance->GetImageName(InstanceComponentIndex, MutableLODIndex, InstanceSurfaceIndex, ImageIndex);
-						Image->ImageID = MutableInstance->GetImageId(InstanceComponentIndex, MutableLODIndex, InstanceSurfaceIndex, ImageIndex);
-						Image->MaterialLayer = MutableInstance->GetImageLayerIndex(InstanceComponentIndex, MutableLODIndex, InstanceSurfaceIndex, ImageIndex);
+						TSharedRef<FInstanceUpdateData::FImage>& Image = Context->InstanceUpdateData.Images.Add_GetRef(MakeShared<FInstanceUpdateData::FImage>());
+						Image->ParameterName = ImageParameter.Key;
+						// Image->Name = ImageParameter.Key; // TODO MATERIAL_PIN It should not have a name but is it currently used?
+						Image->ImageID = UE::Mutable::Private::FSystem::GetImageId(Context->ImageIdRegistry, Context->Model, Context->Parameters.Get(), Address);
+						//Image->MaterialLayer = ImageParameter.Key.Layer; // TODO MATERIAL_PIN
 						Image->FullImageSizeX = 0;
 						Image->FullImageSizeY = 0;
-						Image->BaseLOD = MutableLODIndex;
 						Image->BaseMip = 0;
 
 						FString KeyName = Image->Name.ToString();
 						int32 ImageKey = FCString::Atoi(*KeyName);
 
-						if (ImageKey >= 0 && ImageKey < ModelResources.ImageProperties.Num())
-						{
-							const FMutableModelImageProperties& Props = ModelResources.ImageProperties[ImageKey];
-
-							Image->bIsNonProgressive = Props.MipGenSettings == TMGS_NoMipmaps;
-
-							if (Props.IsPassThrough)
-							{
-								Image->bIsPassThrough = true;
-
-								UE::Tasks::TTask<TSharedPtr<const UE::Mutable::Private::FImage>> GetImageTask = OperationData->MutableSystem->GetImage(LastTask, OperationData->InstanceID, Image->ImageID, 0);
-								UE::Tasks::AddNested(GetImageTask);
-								LastTask = GetImageTask;
-					
-								UE::Tasks::AddNested(UE::Tasks::Launch(TEXT("GetImageResult"), [Image, GetImageTask]() mutable
-								{
-									Image->Image = GetImageTask.GetResult();
-									check(Image->Image->IsReference());
-								},
-								GetImageTask,
-								LowLevelTasks::ETaskPriority::Inherit));
-							}
-						}
-						else
-						{
-							// This means the compiled model (maybe coming from derived data) has images that the asset doesn't know about.
-							UE_LOG(LogMutable, Error, TEXT("CustomizableObject derived data out of sync with asset for [%s]. Try recompiling it."), *CustomizableObject->GetName());
-						}
+						const FMutableModelImageProperties& Props = ModelResources->ImageProperties[ImageKey];
+						Image->bIsNonProgressive = Props.MipGenSettings == TMGS_NoMipmaps;
+						Image->bIsPassThrough = Props.IsPassThrough;
 					}
 				}
-			}	
+			},
+			GetMaterialTask,
+			LowLevelTasks::ETaskPriority::Inherit);
+
+			UE::Tasks::AddNested(LastTask);
 		}
 	}
-	
+
 	
 	void Task_Mutable_GetMeshes(const TSharedRef<FUpdateContextPrivate>& OperationData)
 	{
@@ -2382,8 +2481,6 @@ namespace impl
 
 		check(UCustomizableObjectSystem::GetInstance() != nullptr);
 		check(UCustomizableObjectSystem::GetInstance()->GetPrivate() != nullptr);
-
-		UCustomizableInstancePrivate* CustomizableObjectInstancePrivateData = OperationData->Instance->GetPrivate();
 		
 		if (OperationData->PixelFormatOverride)
 		{
@@ -2435,19 +2532,13 @@ namespace impl
 			},
 			GetSkeletalMeshTask,
 			LowLevelTasks::ETaskPriority::Inherit));
-			
-			Component->OverlayMaterialId = Instance->GetOverlayMaterialId(InstanceComponentIndex);
-			
-			UE::Tasks::TTask<TSharedPtr<const UE::Mutable::Private::FMaterial>> GetMaterialOverrideTask = OperationData->MutableSystem->GetMaterial(LastTask, OperationData->InstanceID, Component->OverlayMaterialId);
-			UE::Tasks::AddNested(GetMaterialOverrideTask);
-			LastTask = GetMaterialOverrideTask;
-			
-			UE::Tasks::AddNested(UE::Tasks::Launch(TEXT("GetMaterialOverrideResult"), [Component, GetMaterialOverrideTask]() mutable
+
+			if (UE::Mutable::Private::FMaterialId MaterialId = Instance->GetOverlayMaterialId(InstanceComponentIndex))
 			{
-				Component->OverlayMaterial = GetMaterialOverrideTask.GetResult();
-			},
-			GetMaterialOverrideTask,
-			LowLevelTasks::ETaskPriority::Inherit));
+				Component->OverlayMaterialIndex = OperationData->InstanceUpdateData.Materials.Num();
+				TSharedRef<FInstanceUpdateData::FMaterial> Material = OperationData->InstanceUpdateData.Materials.Add_GetRef(MakeShared<FInstanceUpdateData::FMaterial>());
+				Material->MaterialId = MaterialId;
+			}
 			
 			Component->LODCount = OperationData->NumLODsAvailable[ComponentName];
 			if (!Component->LODCount)
@@ -2573,95 +2664,107 @@ namespace impl
 
 			if (Image->bIsPassThrough)
 			{
-				continue;
-			}
-			
-			const bool bCached = ImagesInThisInstance->Contains(Image->ImageID) || // See if it is cached from this same instance (can happen with LODs)
-				(UCustomizableObjectSystem::ShouldReuseTexturesBetweenInstances() && OperationData->CachedTextures.Contains(Image->ImageID)); // See if it is cached from another instance
-			
-			if (bCached)
-			{
-				UE_LOG(LogMutable, VeryVerbose, TEXT("Texture resource with id [%s] is cached."), *Image->ImageID.ToString());
-				continue;
-			}
-
-			ImagesInThisInstance->Add(Image->ImageID);
-			
-			const int32 MaxSize = FMath::Max(Image->FullImageSizeX, Image->FullImageSizeY);
-			const int32 FullLODCount = FMath::CeilLogTwo(MaxSize) + 1;
-			const int32 MinMipsInImage = FMath::Min(FullLODCount, UTexture::GetStaticMinTextureResidentMipCount());
-			const int32 MaxMipsToSkip = FullLODCount - MinMipsInImage;
-			int32 MipsToSkip = FMath::Min(MaxMipsToSkip, OperationData->MipsToSkip);
-
-			if (Image->bIsNonProgressive || !FMath::IsPowerOfTwo(Image->FullImageSizeX) || !FMath::IsPowerOfTwo(Image->FullImageSizeY))
-			{
-				// It doesn't make sense to skip mips as non-power-of-two size textures cannot be streamed anyway
-				MipsToSkip = 0;
-			}
-
-			const int32 MipSizeX = FMath::Max(Image->FullImageSizeX >> MipsToSkip, 1);
-			const int32 MipSizeY = FMath::Max(Image->FullImageSizeY >> MipsToSkip, 1);
-
-			UE::Tasks::TTask<TSharedPtr<const UE::Mutable::Private::FImage>> GetImageTask;
-			if (MipsToSkip > 0 && OperationData->bEnableSkipGenerateResidentMips != 0 && OperationData->LowPriorityTextures.Find(Image->Name.ToString()) != INDEX_NONE)
-			{
-				TSharedPtr<const UE::Mutable::Private::FImage> DummyImage = MakeShared<UE::Mutable::Private::FImage>(MipSizeX, MipSizeY, FullLODCount - MipsToSkip, Image->Format, UE::Mutable::Private::EInitializationType::Black);
-				GetImageTask = UE::Tasks::MakeCompletedTask<TSharedPtr<const UE::Mutable::Private::FImage>>(DummyImage);
-			}
-			else
-			{
-				GetImageTask = OperationData->MutableSystem->GetImage(LastTask, OperationData->InstanceID, Image->ImageID, Image->BaseMip + MipsToSkip);
+				UE::Tasks::TTask<TSharedPtr<const UE::Mutable::Private::FImage>> GetImageTask = OperationData->MutableSystem->GetImage(LastTask, OperationData->InstanceID, Image->ImageID, 0);
 				UE::Tasks::AddNested(GetImageTask);
 				LastTask = GetImageTask;
+					
+				UE::Tasks::AddNested(UE::Tasks::Launch(TEXT("GetImageResult"), [Image, GetImageTask]() mutable
+				{
+					Image->Image = GetImageTask.GetResult();
+					check(Image->Image->IsReference());
+				},
+				GetImageTask,
+				LowLevelTasks::ETaskPriority::Inherit));
 			}
-
-			UE::Tasks::AddNested(UE::Tasks::Launch(TEXT("GetImageResult"), [Image, GetImageTask, MipSizeX, MipSizeY, FullLODCount, MipsToSkip]() mutable
+			else
 			{
-				Image->Image = GetImageTask.GetResult();
-				check(Image->Image);
+				const bool bCached = ImagesInThisInstance->Contains(Image->ImageID) || // See if it is cached from this same instance (can happen with LODs)
+					(UCustomizableObjectSystem::ShouldReuseTexturesBetweenInstances() && OperationData->CachedTextures.Contains(Image->ImageID)); // See if it is cached from another instance
+			
+				if (bCached)
+				{
+					UE_LOG(LogMutable, VeryVerbose, TEXT("Texture resource with id [%s] is cached."), *Image->ImageID.ToString());
+					continue;
+				}
+
+				ImagesInThisInstance->Add(Image->ImageID);
+			
+				const int32 MaxSize = FMath::Max(Image->FullImageSizeX, Image->FullImageSizeY);
+				const int32 FullLODCount = FMath::CeilLogTwo(MaxSize) + 1;
+				const int32 MinMipsInImage = FMath::Min(FullLODCount, UTexture::GetStaticMinTextureResidentMipCount());
+				const int32 MaxMipsToSkip = FullLODCount - MinMipsInImage;
+				int32 MipsToSkip = FMath::Min(MaxMipsToSkip, OperationData->MipsToSkip);
 
-				// We should have generated exactly this size.
-				const bool bSizeMissmatch = Image->Image->GetSizeX() != MipSizeX || Image->Image->GetSizeY() != MipSizeY;
-				if (bSizeMissmatch)
+				if (Image->bIsNonProgressive || !FMath::IsPowerOfTwo(Image->FullImageSizeX) || !FMath::IsPowerOfTwo(Image->FullImageSizeY))
 				{
-					// Generate a correctly-sized but empty image instead, to avoid crashes.
-					UE_LOG(LogMutable, Warning, TEXT("Mutable generated a wrongly-sized image %s."), *Image->ImageID.ToString());
-					Image->Image = MakeShared<UE::Mutable::Private::FImage>(MipSizeX, MipSizeY, FullLODCount - MipsToSkip, Image->Image->GetFormat(), UE::Mutable::Private::EInitializationType::Black);
+					// It doesn't make sense to skip mips as non-power-of-two size textures cannot be streamed anyway
+					MipsToSkip = 0;
 				}
 
-				// We need one mip or the complete chain. Otherwise there was a bug.
-				const int32 FullMipCount = Image->Image->GetMipmapCount(Image->Image->GetSizeX(), Image->Image->GetSizeY());
-				const int32 RealMipCount = Image->Image->GetLODCount();
+				const int32 MipSizeX = FMath::Max(Image->FullImageSizeX >> MipsToSkip, 1);
+				const int32 MipSizeY = FMath::Max(Image->FullImageSizeY >> MipsToSkip, 1);
 
-				bool bForceMipchain = 
-					// Did we fail to generate the entire mipchain (if we have mips at all)?
-					(RealMipCount != 1) && (RealMipCount != FullMipCount);
+				UE::Tasks::TTask<TSharedPtr<const UE::Mutable::Private::FImage>> GetImageTask;
+				if (MipsToSkip > 0 && OperationData->bEnableSkipGenerateResidentMips != 0 && OperationData->LowPriorityTextures.Find(Image->Name.ToString()) != INDEX_NONE)
+				{
+					TSharedPtr<const UE::Mutable::Private::FImage> DummyImage = MakeShared<UE::Mutable::Private::FImage>(MipSizeX, MipSizeY, FullLODCount - MipsToSkip, Image->Format, UE::Mutable::Private::EInitializationType::Black);
+					GetImageTask = UE::Tasks::MakeCompletedTask<TSharedPtr<const UE::Mutable::Private::FImage>>(DummyImage);
+				}
+				else
+				{
+					GetImageTask = OperationData->MutableSystem->GetImage(LastTask, OperationData->InstanceID, Image->ImageID, Image->BaseMip + MipsToSkip);
+					UE::Tasks::AddNested(GetImageTask);
+					LastTask = GetImageTask;
+				}
 
-				if (bForceMipchain)
+				UE::Tasks::AddNested(UE::Tasks::Launch(TEXT("GetImageResult"), [Image, GetImageTask, MipSizeX, MipSizeY, FullLODCount, MipsToSkip]() mutable
 				{
-					MUTABLE_CPUPROFILER_SCOPE(GetImage_MipFix);
+					Image->Image = GetImageTask.GetResult();
+					check(Image->Image);
 
-					UE_LOG(LogMutable, Warning, TEXT("Mutable generated an incomplete mip chain for image %s."), *Image->ImageID.ToString());
+					// We should have generated exactly this size.
+					const bool bSizeMissmatch = Image->Image->GetSizeX() != MipSizeX || Image->Image->GetSizeY() != MipSizeY;
+					if (bSizeMissmatch)
+					{
+						// Generate a correctly-sized but empty image instead, to avoid crashes.
+						UE_LOG(LogMutable, Warning, TEXT("Mutable generated a wrongly-sized image %s."), *Image->ImageID.ToString());
+						Image->Image = MakeShared<UE::Mutable::Private::FImage>(MipSizeX, MipSizeY, FullLODCount - MipsToSkip, Image->Image->GetFormat(), UE::Mutable::Private::EInitializationType::Black);
+					}
+
+					// We need one mip or the complete chain. Otherwise there was a bug.
+					const int32 FullMipCount = Image->Image->GetMipmapCount(Image->Image->GetSizeX(), Image->Image->GetSizeY());
+					const int32 RealMipCount = Image->Image->GetLODCount();
 
-					// Force the right number of mips. The missing data will be black.
-					const TSharedPtr<UE::Mutable::Private::FImage> NewImage = MakeShared<UE::Mutable::Private::FImage>(Image->Image->GetSizeX(), Image->Image->GetSizeY(), FullMipCount, Image->Image->GetFormat(), UE::Mutable::Private::EInitializationType::Black);
-					check(NewImage);	
-					// Formats with BytesPerBlock == 0 will not allocate memory. This type of images are not expected here.
-					check(!NewImage->DataStorage.IsEmpty());
+					bool bForceMipchain = 
+						// Did we fail to generate the entire mipchain (if we have mips at all)?
+						(RealMipCount != 1) && (RealMipCount != FullMipCount);
 
-					for (int32 L = 0; L < RealMipCount; ++L)
+					if (bForceMipchain)
 					{
-						TArrayView<uint8> DestView = NewImage->DataStorage.GetLOD(L);
-						TArrayView<const uint8> SrcView = Image->Image->DataStorage.GetLOD(L);
+						MUTABLE_CPUPROFILER_SCOPE(GetImage_MipFix);
+
+						UE_LOG(LogMutable, Warning, TEXT("Mutable generated an incomplete mip chain for image %s."), *Image->ImageID.ToString());
+
+						// Force the right number of mips. The missing data will be black.
+						const TSharedPtr<UE::Mutable::Private::FImage> NewImage = MakeShared<UE::Mutable::Private::FImage>(Image->Image->GetSizeX(), Image->Image->GetSizeY(), FullMipCount, Image->Image->GetFormat(), UE::Mutable::Private::EInitializationType::Black);
+						check(NewImage);	
+						// Formats with BytesPerBlock == 0 will not allocate memory. This type of images are not expected here.
+						check(!NewImage->DataStorage.IsEmpty());
 
-						check(DestView.Num() == SrcView.Num());
-						FMemory::Memcpy(DestView.GetData(), SrcView.GetData(), DestView.Num());
+						for (int32 L = 0; L < RealMipCount; ++L)
+						{
+							TArrayView<uint8> DestView = NewImage->DataStorage.GetLOD(L);
+							TArrayView<const uint8> SrcView = Image->Image->DataStorage.GetLOD(L);
+
+							check(DestView.Num() == SrcView.Num());
+							FMemory::Memcpy(DestView.GetData(), SrcView.GetData(), DestView.Num());
+						}
+						Image->Image = NewImage;
 					}
-					Image->Image = NewImage;
-				}
-			},
-			GetImageTask,
-			UE::Tasks::ETaskPriority::Inherit));
+				},
+				GetImageTask,
+				UE::Tasks::ETaskPriority::Inherit));
+			}
 		}
 	}
 
@@ -2670,11 +2773,11 @@ namespace impl
 	{
 		MUTABLE_CPUPROFILER_SCOPE(Subtask_Mutable_PrepareTextures)
 
-		for (const TSharedRef<FInstanceUpdateData::FSurface>& Surface : OperationData->InstanceUpdateData.Surfaces)
+		for (const TSharedRef<FInstanceUpdateData::FMaterial>& Material : OperationData->InstanceUpdateData.Materials)
 		{
-			for (int32 ImageIndex = 0; ImageIndex<Surface->ImageCount; ++ImageIndex)
+			for (int32 ImageIndex = 0; ImageIndex < Material->ImageCount; ++ImageIndex)
 			{
-				const TSharedRef<FInstanceUpdateData::FImage>& Image = OperationData->InstanceUpdateData.Images[Surface->FirstImage + ImageIndex];
+				const TSharedRef<FInstanceUpdateData::FImage>& Image = OperationData->InstanceUpdateData.Images[Material->FirstImage + ImageIndex];
 
 				if (Image->Image)
 				{
@@ -3210,21 +3313,15 @@ namespace impl
 			UE_LOG(LogMutable, Warning, TEXT("Instance LiveUpdateMode does not yet support progressive streaming of Mutable textures. Disabling LiveUpdateMode for this update."));
 			Operation->bLiveUpdateMode = false;
 		}
-
-		Operation->bReuseInstanceTextures = false;
 		
-		if (!Operation->bBake && SystemPrivate->EnableReuseInstanceTextures)
+		Operation->bReuseInstanceTextures = !Operation->bBake && CVarEnableMutableReuseInstanceTextures.GetValueOnGameThread() && (Operation->bReuseInstanceTextures || (StateData ? StateData->bReuseInstanceTextures : false));
+
+		if (Operation->bReuseInstanceTextures && !Operation->bNeverStream)
 		{
-			Operation->bReuseInstanceTextures = StateData ? StateData->bReuseInstanceTextures : false;
-			Operation->bReuseInstanceTextures |= InstancePrivate->HasCOInstanceFlags(ReuseTextures);
-			
-			if (Operation->bReuseInstanceTextures && !Operation->bNeverStream)
-			{
-				UE_LOG(LogMutable, Warning, TEXT("Instance texture reuse requires that the current Mutable state is in non-streaming mode. Change it in the Mutable graph base node in the state definition."));
-				Operation->bReuseInstanceTextures = false;
-			}
+			UE_LOG(LogMutable, Warning, TEXT("Instance texture reuse requires that the current Mutable state is in non-streaming mode. Change it in the Mutable graph base node in the state definition."));
+			Operation->bReuseInstanceTextures = false;
 		}
-
+		
 		if (!Operation->bLiveUpdateMode && InstancePrivate->LiveUpdateModeInstanceID != 0)
 		{
 			// The instance was in live update mode last update, but now it's not. So the Id and resources have to be released.
@@ -3290,6 +3387,9 @@ namespace impl
 		Operation->PixelFormatOverride = SystemPrivate->ImageFormatOverrideFunc;
 #endif
 
+		InstancePrivate->GeneratedMaterials.Empty();
+		InstancePrivate->GeneratedTextures.Empty();
+		
 		if (!InstancePrivate->HasCOInstanceFlags(ForceGenerateMipTail))
 		{
 			ObjectPrivate->GetLowPriorityTextureNames(Operation->LowPriorityTextures);
@@ -3511,14 +3611,21 @@ namespace impl
 			Task_Mutable_GetMeshes(Operation);
 		},
 		{ Prerequisites});
-
+		UE::Tasks::FTask Mutable_GetSurfacesTask = SystemPrivate->MutableTaskGraph.AddMutableThreadTask(
+		TEXT("Task_Mutable_GetSurfaces"),
+		[Operation]()
+		{
+			Task_Mutable_GetSurfaces(Operation);
+		},
+		{ Mutable_GetMeshTask });
+		
 		UE::Tasks::FTask Mutable_GetMaterialsTask = SystemPrivate->MutableTaskGraph.AddMutableThreadTask(
 		TEXT("Task_Mutable_GetMaterials"),
 		[Operation]()
 		{
 			Task_Mutable_GetMaterials(Operation);
 		},
-		{ Mutable_GetMeshTask });
+		{ Mutable_GetSurfacesTask });
 
 		UE::Tasks::FTask Mutable_PrepareSkeletonDataTask = SystemPrivate->MutableTaskGraph.AddMutableThreadTask(
 		TEXT("Task_Mutable_PrepareSkeletonData"),
diff --git a/Engine/Plugins/Mutable/Source/CustomizableObjectEditor/Private/MuCOE/SCustomizableObjectEditorTextureAnalyzer.cpp b/Engine/Plugins/Mutable/Source/CustomizableObjectEditor/Private/MuCOE/SCustomizableObjectEditorTextureAnalyzer.cpp
--- a/Engine/Plugins/Mutable/Source/CustomizableObjectEditor/Private/MuCOE/SCustomizableObjectEditorTextureAnalyzer.cpp
+++ b/Engine/Plugins/Mutable/Source/CustomizableObjectEditor/Private/MuCOE/SCustomizableObjectEditorTextureAnalyzer.cpp
@@ -139,41 +139,32 @@ void SCustomizableObjecEditorTextureAnalyzer::FillTextureAnalyzerTable(UCustomiz
 		const UCustomizableObject* CustomizableObject = PreviewInstance->GetCustomizableObject();
 		check(CustomizableObject)
 
-		for (const FGeneratedMaterial& GeneratedMaterial : PreviewInstance->GetPrivate()->GeneratedMaterials)
+		for (const TObjectPtr<UTexture>& GeneratedTexture : PreviewInstance->GetPrivate()->GeneratedTextures)
 		{
-			UMaterialInstanceDynamic* Material = Cast<UMaterialInstanceDynamic>(GeneratedMaterial.MaterialInterface);
-			if (!Material)
+			UTexture2D* Texture = Cast<UTexture2D>(GeneratedTexture);
+			if (!Texture || !Texture->HasAllFlags(EObjectFlags::RF_Transient))
 			{
 				continue;
 			}
 
-			for (const FGeneratedTexture& GeneratedTexture : GeneratedMaterial.Textures)
-			{
-				UTexture2D* Texture = Cast<UTexture2D>(GeneratedTexture.Texture);
-				if (!Texture || !Texture->HasAllFlags(EObjectFlags::RF_Transient))
-				{
-					continue;
-				}
-
-				//New Object for the table
-				UCustomizableObjectEditorTextureStats* Entry = NewObject<UCustomizableObjectEditorTextureStats>();
-				TabTextures.Add(Entry);
-
-				//Texture Info
-				Entry->Texture = Texture;
-				Entry->TextureName = Texture->GetName();
-				Entry->ResolutionX = Texture->GetSizeX();
-				Entry->ResolutionY = Texture->GetSizeY();
-				Entry->LODBias = Texture->GetCachedLODBias();
-				uint32 TextureSize = Texture->CalcTextureMemorySizeEnum(TMC_AllMips);
-				Entry->Size = TextureSize / 1024.0f;
-				Entry->Format = Texture->GetPixelFormat();
-				Entry->IsStreamed = Texture->IsCurrentlyVirtualTextured() ? "Virtual Streamed" : (!Texture->IsStreamable() ? "Not Streamed" : "Streamed");
-				Entry->LODGroup = Texture->LODGroup;
-
-				//Total size
-				TotalSize += TextureSize;
-			}
+			//New Object for the table
+			UCustomizableObjectEditorTextureStats* Entry = NewObject<UCustomizableObjectEditorTextureStats>();
+			TabTextures.Add(Entry);
+
+			//Texture Info
+			Entry->Texture = Texture;
+			Entry->TextureName = Texture->GetName();
+			Entry->ResolutionX = Texture->GetSizeX();
+			Entry->ResolutionY = Texture->GetSizeY();
+			Entry->LODBias = Texture->GetCachedLODBias();
+			uint32 TextureSize = Texture->CalcTextureMemorySizeEnum(TMC_AllMips);
+			Entry->Size = TextureSize / 1024.0f;
+			Entry->Format = Texture->GetPixelFormat();
+			Entry->IsStreamed = Texture->IsCurrentlyVirtualTextured() ? "Virtual Streamed" : (!Texture->IsStreamable() ? "Not Streamed" : "Streamed");
+			Entry->LODGroup = Texture->LODGroup;
+			
+			//Total size
+			TotalSize += TextureSize;
 		}
 
 		TextureAnalyzerTable->SetObjects(ObjectPtrDecay(TabTextures));
diff --git a/Engine/Plugins/Mutable/Source/MutableRuntime/Internal/MuR/ResourceID.h b/Engine/Plugins/Mutable/Source/MutableRuntime/Internal/MuR/ResourceID.h
--- a/Engine/Plugins/Mutable/Source/MutableRuntime/Internal/MuR/ResourceID.h
+++ b/Engine/Plugins/Mutable/Source/MutableRuntime/Internal/MuR/ResourceID.h
@@ -17,8 +17,6 @@ namespace UE::Mutable::Private
 		TMemoryTrackedArray<uint8> ParameterValuesBlob;
 		
 		bool operator==(const FGeneratedMeshKey&) const = default;
-
-		bool operator!=(const FGeneratedMeshKey&) const = default;
 	};
 
 	struct FGeneratedMeshData
diff --git a/Engine/Plugins/Mutable/Source/MutableRuntime/Internal/MuR/System.h b/Engine/Plugins/Mutable/Source/MutableRuntime/Internal/MuR/System.h
--- a/Engine/Plugins/Mutable/Source/MutableRuntime/Internal/MuR/System.h
+++ b/Engine/Plugins/Mutable/Source/MutableRuntime/Internal/MuR/System.h
@@ -226,13 +226,13 @@ namespace UE::Mutable::Private
 									const TSharedPtr<const FParameters>& FParameters,
 									bool* Flags );
 
-		static FMeshId GetMeshId(const TSharedPtr<FMeshIdRegistry>& Registry, const TSharedPtr<const FModel>& Model, const FParameters* Params, OP::ADDRESS RootAt);
+		UE_API static FMeshId GetMeshId(const TSharedPtr<FMeshIdRegistry>& Registry, const TSharedPtr<const FModel>& Model, const FParameters* Params, OP::ADDRESS RootAt);
     	
-		static FImageId GetImageId(const TSharedPtr<FImageIdRegistry>& Registry, const TSharedPtr<const FModel>& Model, const FParameters* Params, OP::ADDRESS RootAt);
+		UE_API static FImageId GetImageId(const TSharedPtr<FImageIdRegistry>& Registry, const TSharedPtr<const FModel>& Model, const FParameters* Params, OP::ADDRESS RootAt);
 
-		static FMaterialId GetMaterialId(const TSharedPtr<FMaterialIdRegistry>& Registry, const TSharedPtr<const FModel>& Model, const FParameters* Params, OP::ADDRESS RootAt);
+		UE_API static FMaterialId GetMaterialId(const TSharedPtr<FMaterialIdRegistry>& Registry, const TSharedPtr<const FModel>& Model, const FParameters* Params, OP::ADDRESS RootAt);
 
-		static FSkeletalMeshId GetSkeletalMeshId(const TSharedPtr<FSkeletalMeshIdRegistry>& Registry, const TSharedPtr<const FModel>& Model, const FParameters* Params, OP::ADDRESS RootAt);
+		UE_API static FSkeletalMeshId GetSkeletalMeshId(const TSharedPtr<FSkeletalMeshIdRegistry>& Registry, const TSharedPtr<const FModel>& Model, const FParameters* Params, OP::ADDRESS RootAt);
 		
 		//-----------------------------------------------------------------------------------------
 		// Interface pattern
