[Mutable] Added IN_ADDOVERRIDEMATERIAL instruction

Added to Instance Components.

#Synced-CL 48348317

diff --git a/Engine/Plugins/Mutable/Source/CustomizableObject/Internal/MuCO/CustomizableObjectPrivate.h b/Engine/Plugins/Mutable/Source/CustomizableObject/Internal/MuCO/CustomizableObjectPrivate.h
--- a/Engine/Plugins/Mutable/Source/CustomizableObject/Internal/MuCO/CustomizableObjectPrivate.h
+++ b/Engine/Plugins/Mutable/Source/CustomizableObject/Internal/MuCO/CustomizableObjectPrivate.h
@@ -1760,6 +1760,8 @@ enum class ECustomizableObjectVersions : int32
 	ReplaceSurfaceIdWithSharedSurfaceId,
 
 	Deduplicate_MI_CONSTANT,
+	
+	AddOverrideMaterials,
 
 	// -----<new versions can be added above this line>--------
 	LastCustomizableObjectVersion
diff --git a/Engine/Plugins/Mutable/Source/CustomizableObjectEditor/Private/MuCOE/SMutableCodeViewer.cpp b/Engine/Plugins/Mutable/Source/CustomizableObjectEditor/Private/MuCOE/SMutableCodeViewer.cpp
--- a/Engine/Plugins/Mutable/Source/CustomizableObjectEditor/Private/MuCOE/SMutableCodeViewer.cpp
+++ b/Engine/Plugins/Mutable/Source/CustomizableObjectEditor/Private/MuCOE/SMutableCodeViewer.cpp
@@ -1704,7 +1704,23 @@ void SMutableCodeViewer::GenerateElementRecursive(const int32& InStateIndex, UE:
 		AddOpFunc(Args.value, TEXT("value "));
 		break;
 	}
+		
+	case UE::Mutable::Private::EOpType::IN_ADDOVERLAYMATERIAL:
+	{
+		UE::Mutable::Private::OP::InstanceAddMaterialArgs Args = InProgram.GetOpArgs<UE::Mutable::Private::OP::InstanceAddMaterialArgs>(InParentAddress);
+		AddOpFunc(Args.Instance, TEXT("instance "));
+		AddOpFunc(Args.Material, TEXT("value "));
+		break;
+	}
 
+	case UE::Mutable::Private::EOpType::IN_ADDOVERRIDEMATERIAL:
+	{
+		UE::Mutable::Private::OP::InstanceAddOverrideMaterialArgs Args = InProgram.GetOpArgs<UE::Mutable::Private::OP::InstanceAddOverrideMaterialArgs>(InParentAddress);
+		AddOpFunc(Args.Instance, TEXT("instance "));
+		AddOpFunc(Args.Material, TEXT("value "));
+		break;
+	}
+	
 	case UE::Mutable::Private::EOpType::IN_ADDMATERIAL:
 	{
 		UE::Mutable::Private::OP::InstanceAddMaterialArgs Args = InProgram.GetOpArgs<UE::Mutable::Private::OP::InstanceAddMaterialArgs>(InParentAddress);
diff --git a/Engine/Plugins/Mutable/Source/MutableRuntime/Internal/MuR/Instance.h b/Engine/Plugins/Mutable/Source/MutableRuntime/Internal/MuR/Instance.h
--- a/Engine/Plugins/Mutable/Source/MutableRuntime/Internal/MuR/Instance.h
+++ b/Engine/Plugins/Mutable/Source/MutableRuntime/Internal/MuR/Instance.h
@@ -77,6 +77,8 @@ namespace UE::Mutable::Private
 
 		FMaterialId OverlayMaterialId;
 
+		TMap<FName, FMaterialId> OverrideMaterialIds;
+
 		FSkeletalMeshId SkeletalMeshId;
 
 		TArray<FInstanceLOD, TInlineAllocator<1>> LODs;
@@ -298,6 +300,8 @@ namespace UE::Mutable::Private
 		void SetSkeletalMeshId(int32 ComponentIndex, const FSkeletalMeshId& SkeletalMeshId);
 
 		void SetOverlayMaterialId(int32 ComponentIndex, const FMaterialId& MaterialId);
+		
+		void AddOverrideMaterial(int32 ComponentIndex, const FName& SlotName, const FMaterialId& MaterialId);
 
 		void SetExtensionData(const TSharedRef<const FExtensionData>& Data, FName Name);
 		
diff --git a/Engine/Plugins/Mutable/Source/MutableRuntime/Internal/MuR/Operations.h b/Engine/Plugins/Mutable/Source/MutableRuntime/Internal/MuR/Operations.h
--- a/Engine/Plugins/Mutable/Source/MutableRuntime/Internal/MuR/Operations.h
+++ b/Engine/Plugins/Mutable/Source/MutableRuntime/Internal/MuR/Operations.h
@@ -354,6 +354,9 @@ namespace UE::Mutable::Private
 		//! Add overlay material to an instance
 		IN_ADDOVERLAYMATERIAL,
 
+		//! Add override material to an instance
+		IN_ADDOVERRIDEMATERIAL,
+
 		//! Add a material to an instance
 		IN_ADDMATERIAL,
 
@@ -1133,6 +1136,13 @@ namespace UE::Mutable::Private
 			ADDRESS Material;
 		};
 
+		struct InstanceAddOverrideMaterialArgs
+		{
+			ADDRESS Instance;
+			ADDRESS Material;
+			CONSTANT_STRING_INDEX SlotName;
+		};
+
         struct LayoutPackArgs
         {
             ADDRESS Source;
diff --git a/Engine/Plugins/Mutable/Source/MutableRuntime/Private/MuR/CodeRunner.cpp b/Engine/Plugins/Mutable/Source/MutableRuntime/Private/MuR/CodeRunner.cpp
--- a/Engine/Plugins/Mutable/Source/MutableRuntime/Private/MuR/CodeRunner.cpp
+++ b/Engine/Plugins/Mutable/Source/MutableRuntime/Private/MuR/CodeRunner.cpp
@@ -1021,9 +1021,6 @@ namespace UE::Mutable::Private
         		{
         		case 0:
         			{
-        				// Must pass in a Material op
-        				check(Args.Material);
-
         				AddOp(FScheduledOp(Item.At, Item, 1),
 							FScheduledOp(Args.Instance, Item),
 							FScheduledOp(Args.Material, Item));
@@ -1061,6 +1058,51 @@ namespace UE::Mutable::Private
         		break;
         	}
 
+        case EOpType::IN_ADDOVERRIDEMATERIAL:
+        	{
+        		OP::InstanceAddOverrideMaterialArgs Args = Program.GetOpArgs<OP::InstanceAddOverrideMaterialArgs>(Item.At);
+        		switch (Item.Stage)
+        		{
+        		case 0:
+        			{
+        				AddOp(FScheduledOp(Item.At, Item, 1),
+							FScheduledOp(Args.Instance, Item),
+							FScheduledOp(Args.Material, Item));
+
+        				break;
+        			}
+
+        		case 1:
+        			{
+        				TSharedPtr<const FInstance> Base = LoadInstance(FCacheAddress(Args.Instance, Item));
+        				TSharedPtr<FInstance> Result;
+        				if (!Base)
+        				{
+        					Result = MakeShared<FInstance>();
+        				}
+        				else
+        				{
+        					Result = UE::Mutable::Private::CloneOrTakeOver<FInstance>(Base);
+        				}
+
+        				if (Args.Material)
+        				{
+        					FMaterialId MaterialId = FSystem::GetMaterialId(System->WorkingMemoryManager.CurrentMaterialIdRegistry, InModel, InParams, Args.Material);
+        					const FString& SlotName = Program.ConstantStrings[Args.SlotName];
+        					Result->AddOverrideMaterial(0, FName(SlotName), MaterialId);
+        				}
+
+        				StoreInstance(Item, Result);
+        				break;
+        			}
+
+        		default:
+        			check(false);
+        		}
+
+        		break;
+        	}
+        	
         case EOpType::IN_ADDMATERIAL:
         	{
         		OP::InstanceAddMaterialArgs Args = Program.GetOpArgs<OP::InstanceAddMaterialArgs>(Item.At);
@@ -6864,6 +6906,7 @@ namespace UE::Mutable::Private
         case EOpType::IN_ADDIMAGE:
         case EOpType::IN_ADDMATERIAL:
         case EOpType::IN_ADDOVERLAYMATERIAL:
+        case EOpType::IN_ADDOVERRIDEMATERIAL:
     	case EOpType::IN_ADDSKELETALMESH:
             RunCode_InstanceAddResource(Item, InModel, pParams);
             break;
diff --git a/Engine/Plugins/Mutable/Source/MutableRuntime/Private/MuR/CodeRunnerBegin.cpp b/Engine/Plugins/Mutable/Source/MutableRuntime/Private/MuR/CodeRunnerBegin.cpp
--- a/Engine/Plugins/Mutable/Source/MutableRuntime/Private/MuR/CodeRunnerBegin.cpp
+++ b/Engine/Plugins/Mutable/Source/MutableRuntime/Private/MuR/CodeRunnerBegin.cpp
@@ -946,6 +946,56 @@ namespace UE::Mutable::Private
 
 					break;
 				}
+
+			case EOpType::IN_ADDOVERRIDEMATERIAL:
+				{
+					OP::InstanceAddOverrideMaterialArgs Args = Program.GetOpArgs<OP::InstanceAddOverrideMaterialArgs>(Item.At);
+					switch (Item.Stage)
+					{
+					case 0:
+						{
+							PushOp(FScheduledOpInline(Item, 1));
+							PushOp(FScheduledOpInline(Args.Instance, Item));
+							PushOp(FScheduledOpInline(Args.Material, Item));
+
+							break;
+						}
+
+					case 1:
+						{
+							TSharedPtr<const FInstance> InInstance;
+							if (Args.Instance)
+							{
+								InInstance = LoadInstance(Args.Instance);
+							}
+
+							TSharedPtr<FInstance> Instance;
+							if (InInstance)
+							{
+								Instance = UE::Mutable::Private::CloneOrTakeOver<FInstance>(InInstance);
+							}
+							else
+							{
+								Instance = MakeShared<FInstance>();
+							}
+
+							if (Args.Material)
+							{
+								FMaterialId MaterialId = FSystem::GetMaterialId(System.GetPrivate()->WorkingMemoryManager.CurrentMaterialIdRegistry, Model, Params.Get(), Args.Material);
+								const FString& SlotName = Program.ConstantStrings[Args.SlotName];
+								Instance->AddOverrideMaterial(0, FName(SlotName), MaterialId);
+							}
+							
+							StoreInstance( Item, Instance);
+							break;
+						}
+
+					default:
+						unimplemented();
+					}
+
+					break;
+				}
 				
 			case EOpType::IN_ADDCOMPONENT:
 				{
diff --git a/Engine/Plugins/Mutable/Source/MutableRuntime/Private/MuR/Instance.cpp b/Engine/Plugins/Mutable/Source/MutableRuntime/Private/MuR/Instance.cpp
--- a/Engine/Plugins/Mutable/Source/MutableRuntime/Private/MuR/Instance.cpp
+++ b/Engine/Plugins/Mutable/Source/MutableRuntime/Private/MuR/Instance.cpp
@@ -550,7 +550,18 @@ namespace UE::Mutable::Private
 		Components[ComponentIndex].OverlayMaterialId = MaterialId;
 	}
 
-	
+
+	void FInstance::AddOverrideMaterial(int32 ComponentIndex, const FName& SlotName, const FMaterialId& MaterialId)
+	{
+    	while (ComponentIndex >= Components.Num())
+    	{
+    		AddComponent();
+    	}
+
+    	Components[ComponentIndex].OverrideMaterialIds.Add(SlotName, MaterialId);
+	}
+
+
 	void FInstance::SetExtensionData(const TSharedRef<const FExtensionData>& Data, FName Name)
 	{
 		NamedExtensionData& Entry = ExtensionData.AddDefaulted_GetRef();
diff --git a/Engine/Plugins/Mutable/Source/MutableRuntime/Private/MuR/Operations.cpp b/Engine/Plugins/Mutable/Source/MutableRuntime/Private/MuR/Operations.cpp
--- a/Engine/Plugins/Mutable/Source/MutableRuntime/Private/MuR/Operations.cpp
+++ b/Engine/Plugins/Mutable/Source/MutableRuntime/Private/MuR/Operations.cpp
@@ -155,6 +155,7 @@ namespace UE::Mutable::Private
 		{ EDataType::Instance,		},	// IN_ADDSKELETALMESH
 		{ EDataType::Instance,		},	// IN_ADDEXTENSIONDATA
 		{ EDataType::Instance,		},	// IN_ADDOVERLAYMATERIAL
+		{ EDataType::Instance,		},	// IN_ADDOVERRIDEMATERIAL
 		{ EDataType::Instance,		},	// IN_ADDMATERIAL
 
 		{ EDataType::Layout,			},	// LA_PACK
@@ -893,6 +894,16 @@ namespace UE::Mutable::Private
 			break;
 		}
 
+		case EOpType::IN_ADDOVERRIDEMATERIAL:
+		{
+        		const OP::InstanceAddOverrideMaterialArgs Args = program.GetOpArgs<OP::InstanceAddOverrideMaterialArgs>(at);
+
+        		f(Args.Instance);
+        		f(Args.Material);
+
+        		break;
+		}
+
 		case EOpType::IN_ADDOVERLAYMATERIAL:
 		case EOpType::IN_ADDMATERIAL:
 		{
diff --git a/Engine/Plugins/Mutable/Source/MutableTools/Private/MuT/ASTOpAddOverrideMaterial.cpp b/Engine/Plugins/Mutable/Source/MutableTools/Private/MuT/ASTOpAddOverrideMaterial.cpp
new file mode 100644
--- /dev/null
+++ b/Engine/Plugins/Mutable/Source/MutableTools/Private/MuT/ASTOpAddOverrideMaterial.cpp
@@ -0,0 +1,137 @@
+// Copyright Epic Games, Inc. All Rights Reserved.
+
+#include "MuT/ASTOpAddOverrideMaterial.h"
+
+#include "MuR/ModelPrivate.h"
+#include "MuT/CodeOptimiser.h"
+
+
+namespace UE::Mutable::Private
+{
+	ASTOpAddOverrideMaterial::ASTOpAddOverrideMaterial() :
+		Instance(this),
+		Material(this)
+	{
+	}
+
+
+	ASTOpAddOverrideMaterial::~ASTOpAddOverrideMaterial()
+	{
+		// Explicit call needed to avoid recursive destruction
+		RemoveChildren();
+	}
+
+
+	EOpType ASTOpAddOverrideMaterial::GetOpType() const
+	{
+		return EOpType::IN_ADDOVERRIDEMATERIAL;
+	}
+
+
+	bool ASTOpAddOverrideMaterial::IsEqual(const ASTOp& OtherUntyped) const
+	{
+		if (OtherUntyped.GetOpType() == GetOpType())
+		{
+			const ASTOpAddOverrideMaterial* Other = static_cast<const ASTOpAddOverrideMaterial*>(&OtherUntyped);
+			return Instance == Other->Instance &&
+				Material == Other->Material &&
+				SlotName == Other->SlotName;
+		}
+
+		return false;
+	}
+
+	
+	uint64 ASTOpAddOverrideMaterial::Hash() const
+	{
+		uint32 Result = GetTypeHash(static_cast<uint32>(GetOpType()));
+		
+		Result = HashCombine(Result, PointerHash(Instance.child().get()));
+		Result = HashCombine(Result, PointerHash(Material.child().get()));
+		Result = HashCombine(Result, GetTypeHash(SlotName));
+		
+		return Result;
+	}
+	
+
+	Ptr<ASTOp> ASTOpAddOverrideMaterial::Clone(MapChildFuncRef MapChild) const
+	{
+		Ptr<ASTOpAddOverrideMaterial> NewInstance = new ASTOpAddOverrideMaterial();
+
+		NewInstance->Instance = MapChild(Instance.child());
+		NewInstance->Material = MapChild(Material.child());
+		NewInstance->SlotName = SlotName;
+
+		return NewInstance;
+	}
+	
+
+	void ASTOpAddOverrideMaterial::ForEachChild(const TFunctionRef<void(ASTChild&)> Func)
+	{
+		Func(Instance);
+		Func(Material);
+	}
+	
+
+	void ASTOpAddOverrideMaterial::Link(FProgram& Program, FLinkerOptions*)
+	{
+		// Already linked?
+		if (linkedAddress)
+		{
+			return;
+		}
+
+		OP::InstanceAddOverrideMaterialArgs Args;
+		FMemory::Memzero(Args);
+
+		if (Instance)
+		{
+			Args.Instance = Instance->linkedAddress;
+		}
+
+		if (Material)
+		{
+			Args.Material = Material->linkedAddress;
+			
+			// Find out relevant parameters. \todo: this may be optimised by reusing partial
+			// values in a LINK_CONTEXT or similar
+			SubtreeRelevantParametersVisitorAST Visitor;
+			Visitor.Run(Material.child());
+			
+			TArray<uint16> Params;
+			for (const FString& ParameterName : Visitor.Parameters)
+			{
+				for (int32 Index = 0; Index < Program.Parameters.Num(); ++Index)
+				{
+					const FParameterDesc& Param = Program.Parameters[Index];
+					if (Param.Name == ParameterName)
+					{
+						Params.Add(uint16(Index));
+						break;
+					}
+				}
+			}
+
+			Params.Sort();
+
+			int32 It = Program.ParameterLists.Find(Params);
+
+			if (It != INDEX_NONE)
+			{
+				Program.RelevantParameterList.Add(Material->linkedAddress, It);
+			}
+			else
+			{
+				Program.RelevantParameterList.Add(Material->linkedAddress, Program.ParameterLists.Num());
+				Program.ParameterLists.Add(Params);
+			}
+		}
+		
+		Args.SlotName = Program.AddConstant(SlotName.ToString());
+
+		linkedAddress = (OP::ADDRESS)Program.OpAddress.Num();
+		Program.OpAddress.Add((uint32_t)Program.ByteCode.Num());
+		AppendCode(Program.ByteCode, GetOpType());
+		AppendCode(Program.ByteCode, Args);
+	}
+}
diff --git a/Engine/Plugins/Mutable/Source/MutableTools/Private/MuT/ASTOpAddOverrideMaterial.h b/Engine/Plugins/Mutable/Source/MutableTools/Private/MuT/ASTOpAddOverrideMaterial.h
new file mode 100644
--- /dev/null
+++ b/Engine/Plugins/Mutable/Source/MutableTools/Private/MuT/ASTOpAddOverrideMaterial.h
@@ -0,0 +1,36 @@
+// Copyright Epic Games, Inc. All Rights Reserved.
+
+#pragma once
+
+#include "HAL/Platform.h"
+#include "MuR/Operations.h"
+#include "MuR/Ptr.h"
+#include "MuT/AST.h"
+
+
+namespace UE::Mutable::Private
+{
+	struct FProgram;
+	
+	class ASTOpAddOverrideMaterial final : public ASTOp
+	{
+	public:
+		ASTChild Instance;
+		
+		ASTChild Material;
+		
+		FName SlotName;
+
+		ASTOpAddOverrideMaterial();
+		ASTOpAddOverrideMaterial(const ASTOpAddOverrideMaterial&) = delete;
+		virtual ~ASTOpAddOverrideMaterial() override;
+
+		virtual EOpType GetOpType() const override;
+		virtual uint64 Hash() const override;
+		virtual void ForEachChild(const TFunctionRef<void(ASTChild&)> Func) override;
+		virtual bool IsEqual(const ASTOp& OtherUntyped) const override;
+		virtual Ptr<ASTOp> Clone(MapChildFuncRef MapChild) const override;
+		virtual void Link(FProgram& Program, FLinkerOptions* Options) override;
+	};
+}
+
diff --git a/Engine/Plugins/Mutable/Source/MutableTools/Private/MuT/CodeGenerator.cpp b/Engine/Plugins/Mutable/Source/MutableTools/Private/MuT/CodeGenerator.cpp
--- a/Engine/Plugins/Mutable/Source/MutableTools/Private/MuT/CodeGenerator.cpp
+++ b/Engine/Plugins/Mutable/Source/MutableTools/Private/MuT/CodeGenerator.cpp
@@ -2,6 +2,7 @@
 
 #include "MuT/CodeGenerator.h"
 
+#include "ASTOpAddOverrideMaterial.h"
 #include "ASTOpAddSkeletalMesh.h"
 #include "Containers/Array.h"
 #include "Logging/LogCategory.h"
@@ -610,6 +611,20 @@ namespace UE::Mutable::Private
 
 					LastInstOp = AddOverlayMaterialOp;
 				}
+				
+				for (TPair<const FName, Ptr<NodeMaterial>> Pair : InNode->OverrideMaterials)
+				{
+					Ptr<ASTOpAddOverrideMaterial> AddOverrideMaterialOp = new ASTOpAddOverrideMaterial();
+					AddOverrideMaterialOp->SlotName = Pair.Key;
+					AddOverrideMaterialOp->Instance = LastInstOp;
+
+					FMaterialGenerationOptions MaterialOptions;
+					FMaterialGenerationResult MaterialResult;
+					GenerateMaterial(MaterialOptions, MaterialResult, Pair.Value);
+					AddOverrideMaterialOp->Material = MaterialResult.op;
+
+					LastInstOp = AddOverrideMaterialOp;
+				}
 
 				Ptr<ASTOpInstanceAdd> InstanceOp = new ASTOpInstanceAdd();
 				InstanceOp->type = EOpType::IN_ADDCOMPONENT;
diff --git a/Engine/Plugins/Mutable/Source/MutableTools/Private/MuT/Compiler.cpp b/Engine/Plugins/Mutable/Source/MutableTools/Private/MuT/Compiler.cpp
--- a/Engine/Plugins/Mutable/Source/MutableTools/Private/MuT/Compiler.cpp
+++ b/Engine/Plugins/Mutable/Source/MutableTools/Private/MuT/Compiler.cpp
@@ -192,6 +192,7 @@ namespace UE::Mutable::Private
 		{ false,	{ 0, 0, 0, 0, 0, 0, 0, 0, 0, 0, 0, 0, 0, 0, 0, 0, 0, 0, 0, 0, 0, 0, 0, 0, 0, 0, 0, 0, 0 }		},	// IN_ADDSKELETALMESH
 		{ false,	{ 0, 0, 0, 0, 0, 0, 0, 0, 0, 0, 0, 0, 0, 0, 0, 0, 0, 0, 0, 0, 0, 0, 0, 0, 0, 0, 0, 0, 0 }		},	// IN_ADDEXTENSIONDATA
 		{ false,	{ 0, 0, 0, 0, 0, 0, 0, 0, 0, 0, 0, 0, 0, 0, 0, 0, 0, 0, 0, 0, 0, 0, 0, 0, 0, 0, 0, 0, 0 }		},	// IN_ADDOVERLAYMATERIAL
+		{ false,	{ 0, 0, 0, 0, 0, 0, 0, 0, 0, 0, 0, 0, 0, 0, 0, 0, 0, 0, 0, 0, 0, 0, 0, 0, 0, 0, 0, 0, 0 }		},	// IN_ADDOVERRIDEMATERIAL
 		{ false,	{ 0, 0, 0, 0, 0, 0, 0, 0, 0, 0, 0, 0, 0, 0, 0, 0, 0, 0, 0, 0, 0, 0, 0, 0, 0, 0, 0, 0, 0 }		},	// IN_ADDMATERIAL
 
 		{ true,		{ 0, 0, 0, 0, 0, 0, 0, 0, 0, 0, 0, 0, 0, 0, 0, 0, 0, 0, 0, 0, 0, 0, 0, 0, 0, 0, 0, 0, 0 }		},	// LA_PACK
diff --git a/Engine/Plugins/Mutable/Source/MutableTools/Private/MuT/ErrorLog.cpp b/Engine/Plugins/Mutable/Source/MutableTools/Private/MuT/ErrorLog.cpp
--- a/Engine/Plugins/Mutable/Source/MutableTools/Private/MuT/ErrorLog.cpp
+++ b/Engine/Plugins/Mutable/Source/MutableTools/Private/MuT/ErrorLog.cpp
@@ -337,6 +337,7 @@ namespace UE::Mutable::Private
         TEXT("IN_ADDSKELETALMESH"),
 		TEXT("IN_ADDEXTENSIDATA"),
 		TEXT("IN_ADDOVERLAYMATERIAL"),
+		TEXT("IN_ADDOVERRIDEMATERIAL"),
 		TEXT("IN_ADDMATERIAL   "),
 		
 		TEXT("LA_PACK          "),
