Fix Mutable MorphTargets with streamed meshes when CPU Morphs are needed. Added functionalty to init MorphTargetVeretexInfoBuffers from compressed data.

- Now CPU compressed morph targets can be generated by the streming update process.
- Added Cvar mutable.EnableCPUMorphTargetStreaming To disable CPU Morph targets generation when streaming. The Cvar is enabled by default.
- Exposed with an accessor method MorphTarget CompressedLODModels.
- MorphTargetVertexInfoBuffers::InitMorphResourcesStreaming now checks ShaderPlatorm and bails out if not supported.

#rnx

#Synced-CL 48140675

diff --git a/Engine/Plugins/Mutable/Source/CustomizableObject/Internal/MuCO/UnrealConversionUtils.h b/Engine/Plugins/Mutable/Source/CustomizableObject/Internal/MuCO/UnrealConversionUtils.h
--- a/Engine/Plugins/Mutable/Source/CustomizableObject/Internal/MuCO/UnrealConversionUtils.h
+++ b/Engine/Plugins/Mutable/Source/CustomizableObject/Internal/MuCO/UnrealConversionUtils.h
@@ -154,7 +154,7 @@ namespace UnrealConversionUtils
 		FSkeletalMeshLODRenderData& LODResource
 	);
 
-	void MorphTargetVertexInfoBuffers(FSkeletalMeshLODRenderData& LODResource, const USkeletalMesh& Owner, const UE::Mutable::Private::FMesh& MutableMesh, const TMap<uint32, FMorphTargetMeshData>& MorphTargetMeshData, int32 LODIndex);
+	void MorphTargetVertexInfoBuffers(FSkeletalMeshLODRenderData& LODResource, const USkeletalMesh& Owner, const UE::Mutable::Private::FMesh& MutableMesh, const TMap<uint32, FMorphTargetMeshData>& MorphTargetMeshData, int32 LODIndex, bool bGenerateCPUMorphTargets);
 
 	CUSTOMIZABLEOBJECT_API UE::Tasks::FTask ConvertSkeletalMeshFromRuntimeData(USkeletalMesh* SkeletalMesh, int32 LODIndex, int32 SectionIndex, TSharedPtr<UE::Mutable::Private::FMesh> Result, const TSharedRef<FBoneNames>& BoneNames);
 
diff --git a/Engine/Plugins/Mutable/Source/CustomizableObject/Private/MuCO/CustomizableObjectInstance.cpp b/Engine/Plugins/Mutable/Source/CustomizableObject/Private/MuCO/CustomizableObjectInstance.cpp
--- a/Engine/Plugins/Mutable/Source/CustomizableObject/Private/MuCO/CustomizableObjectInstance.cpp
+++ b/Engine/Plugins/Mutable/Source/CustomizableObject/Private/MuCO/CustomizableObjectInstance.cpp
@@ -4085,11 +4085,16 @@ void BuildOrCopyMorphTargetsData(const TSharedRef<FUpdateContextPrivate>& Contex
 
 		if (Context->bStreamMeshLODs)
 		{
+			const int32 FirstResidentLOD = Context->FirstResidentLOD[ComponentName];
 			MorphLODModels.SetNum(ComponentMorphTargets->RealTimeMorphsLODData[MorphTargetIndex].Num());
-			
+#if !WITH_EDITOR
+			// Allocate compressed LODModels storage. This will be filled by the streaming LOD update if CPU Morphs are needed. 
+			// Only allocate for streamed LODs, the resident will use the non-compressed generated here.
+			NewMorphTarget->GetCompressedLODModels().SetNum(FirstResidentLOD);
+#endif
 			// Streamed LODs
 			const int32 FirstLODAvailable = Context->FirstLODAvailable[ComponentName];
-			for (int32 LODIndex = FirstLODAvailable; LODIndex < Context->FirstResidentLOD[ComponentName]; ++LODIndex)
+			for (int32 LODIndex = FirstLODAvailable; LODIndex < FirstResidentLOD; ++LODIndex)
 			{
 				// Copy data required for streaming
 				MorphLODModels[LODIndex].NumVertices = 1; // Trick the engine
diff --git a/Engine/Plugins/Mutable/Source/CustomizableObject/Private/MuCO/CustomizableObjectMeshUpdate.cpp b/Engine/Plugins/Mutable/Source/CustomizableObject/Private/MuCO/CustomizableObjectMeshUpdate.cpp
--- a/Engine/Plugins/Mutable/Source/CustomizableObject/Private/MuCO/CustomizableObjectMeshUpdate.cpp
+++ b/Engine/Plugins/Mutable/Source/CustomizableObject/Private/MuCO/CustomizableObjectMeshUpdate.cpp
@@ -23,7 +23,6 @@ CustomizableObjectMeshUpdate.cpp: Helpers to stream in CustomizableObject skelet
 #include "Streaming/RenderAssetUpdate.inl"
 #include "Rendering/SkeletalMeshRenderData.h"
 
-
 template class TRenderAssetUpdate<FSkelMeshUpdateContext>;
 
 #define UE_MUTABLE_UPDATE_MESH_REGION		TEXT("Task_Mutable_UpdateMesh")
@@ -36,6 +35,14 @@ FAutoConsoleVariableRef CVarMutableEnableGCHangFix(
 	TEXT("If true, TaskSynchronization decrement will happen in the Abort method instead of DoCancelMeshUpdate.")
 );
 
+static bool bEnableCPUMorphTargetStreaming = true;
+FAutoConsoleVariableRef CVarMutableEnableCPUMorphTargetStreaming(
+	TEXT("mutable.EnableCPUMorphTargetStreaming"),
+	bEnableCPUMorphTargetStreaming,
+	TEXT("Enable CPU morph targets when mesh streaming is enabled.")
+	TEXT("If true, when streaming meshes the CPU morphs targets will be generated when needed.")
+);
+
 FCustomizableObjectMeshStreamIn::FCustomizableObjectMeshStreamIn(
 	const UCustomizableObjectSkeletalMesh* InMesh,
 	EThreadType CreateResourcesThread,
@@ -416,12 +423,13 @@ void FCustomizableObjectMeshStreamIn::ConvertMesh(const FContext& Context, bool&
 		}
 		
 		const bool bNeedsCPUAccess = Mesh->GetResourceForRendering()->RequiresCPUSkinning(GMaxRHIFeatureLevel) || Mesh->NeedCPUData(LODIndexAbsolute);
+		const bool bGenerateCPUMorphTargetsIfNeeded = bEnableCPUMorphTargetStreaming;
 
 		FSkeletalMeshLODRenderData& LODResource = *Context.LODResourcesView[LODIndex];
 		UnrealConversionUtils::CopyMutableVertexBuffers(LODResource, MutableMesh.Get(), bNeedsCPUAccess);
 		UnrealConversionUtils::CopyMutableIndexBuffers(LODResource, MutableMesh.Get(), Mesh->SurfaceIDs[LODIndexAbsolute], bOutMarkRenderStateDirty);
 		UnrealConversionUtils::CopyMutableSkinWeightProfilesBuffers(LODResource, *const_cast<UCustomizableObjectSkeletalMesh*>(Mesh), LODIndexAbsolute, MutableMesh.Get(), Mesh->SkinWeightProfileIDs);
-		UnrealConversionUtils::MorphTargetVertexInfoBuffers(LODResource, *Mesh, *MutableMesh.Get(), OperationData->MorphTargetMeshData, LODIndexAbsolute);
+		UnrealConversionUtils::MorphTargetVertexInfoBuffers(LODResource, *Mesh, *MutableMesh.Get(), OperationData->MorphTargetMeshData, LODIndexAbsolute, bGenerateCPUMorphTargetsIfNeeded);
 		UnrealConversionUtils::ClothVertexBuffers(LODResource, *MutableMesh.Get(), OperationData->ClothingMeshData, LODIndexAbsolute);
 		
 		UnrealConversionUtils::UpdateSkeletalMeshLODRenderDataBuffersSize(LODResource);
@@ -458,4 +466,4 @@ void FCustomizableObjectMeshStreamIn::MarkRenderStateDirty(const FContext& Conte
 	}
 	
 	PushTask(Context, CreateResourcesThread, SRA_UPDATE_CALLBACK(DoCreateBuffers), (EThreadType)Context.CurrentThread, SRA_UPDATE_CALLBACK(DoCancel));
-}
\ No newline at end of file
+}
diff --git a/Engine/Plugins/Mutable/Source/CustomizableObject/Private/MuCO/UnrealConversionUtils.cpp b/Engine/Plugins/Mutable/Source/CustomizableObject/Private/MuCO/UnrealConversionUtils.cpp
--- a/Engine/Plugins/Mutable/Source/CustomizableObject/Private/MuCO/UnrealConversionUtils.cpp
+++ b/Engine/Plugins/Mutable/Source/CustomizableObject/Private/MuCO/UnrealConversionUtils.cpp
@@ -22,6 +22,7 @@
 #include "Rendering/SkeletalMeshRenderData.h"
 #include "Rendering/MorphTargetVertexCodec.h"
 
+
 class USkeleton;
 
 namespace UnrealConversionUtils
@@ -167,8 +168,54 @@ namespace UnrealConversionUtils
 			}
 	
 		}
-	}
 
+		/** 
+		 * Compute the minimum and maximum values for the morphs models.
+		 *
+		 * @param MorphTargets     Morph target form which to compute the minimum an maximum values.
+		 * @param OutMinimumValues Per morph minimum values result. Must have the same size as MorphTargets.
+		 * @param OutMaximumValues Per morph maximum values result. Must have the same size as MorphTargets. 
+		 */
+		void ComputeMorphTargetsMinAndMaxValues(TConstArrayView<const FMorphTargetLODModel*> MorphTargets, TArrayView<FVector4f> OutMinimumValues, TArrayView<FVector4f> OutMaximumValues)
+		{
+			check(OutMinimumValues.Num() == MorphTargets.Num());
+			check(OutMaximumValues.Num() == MorphTargets.Num());
+
+			for (int32 MorphIndex = 0; MorphIndex < MorphTargets.Num(); ++MorphIndex)
+			{
+				const FMorphTargetLODModel* MorphLODModel = MorphTargets[MorphIndex];
+
+				FVector4f MaximumValues = FVector4f(0.0f, 0.0f, 0.0f, 0.0f);
+				FVector4f MinimumValues = FVector4f(0.0f, 0.0f, 0.0f, 0.0f);
+
+				if (MorphLODModel && MorphLODModel->Vertices.Num())
+				{
+					MaximumValues = FVector4f(-FLT_MAX, -FLT_MAX, -FLT_MAX, -FLT_MAX);
+					MinimumValues = FVector4f( FLT_MAX,  FLT_MAX,  FLT_MAX,  FLT_MAX);
+
+					for (const FMorphTargetDelta& MorphDelta : MorphLODModel->Vertices)
+					{
+						const FVector3f& PositionDelta = MorphDelta.PositionDelta;
+						const FVector3f& TangentZDelta = MorphDelta.TangentZDelta;
+
+						MaximumValues[0] = FMath::Max(MaximumValues[0], PositionDelta.X);
+						MaximumValues[1] = FMath::Max(MaximumValues[1], PositionDelta.Y);
+						MaximumValues[2] = FMath::Max(MaximumValues[2], PositionDelta.Z);
+						MaximumValues[3] = FMath::Max(MaximumValues[3], FMath::Max(TangentZDelta.X, FMath::Max(TangentZDelta.Y, TangentZDelta.Z)));
+
+						MinimumValues[0] = FMath::Min(MinimumValues[0], PositionDelta.X);
+						MinimumValues[1] = FMath::Min(MinimumValues[1], PositionDelta.Y);
+						MinimumValues[2] = FMath::Min(MinimumValues[2], PositionDelta.Z);
+						MinimumValues[3] = FMath::Min(MinimumValues[3], FMath::Min(TangentZDelta.X, FMath::Min(TangentZDelta.Y, TangentZDelta.Z)));
+
+					}
+				}
+
+				OutMinimumValues[MorphIndex] = MinimumValues;
+				OutMaximumValues[MorphIndex] = MaximumValues;
+			}
+		}
+	}
 
 	void SetupRenderSections(
 		FSkeletalMeshLODRenderData& LODResource,
@@ -832,8 +879,7 @@ namespace UnrealConversionUtils
 		}
 	}
 
-
-	void MorphTargetVertexInfoBuffers(FSkeletalMeshLODRenderData& LODResource, const USkeletalMesh& Owner, const UE::Mutable::Private::FMesh& MutableMesh, const TMap<uint32, FMorphTargetMeshData>& MorphTargetMeshData, int32 LODIndex)
+	void MorphTargetVertexInfoBuffers(FSkeletalMeshLODRenderData& LODResource, const USkeletalMesh& Owner, const UE::Mutable::Private::FMesh& MutableMesh, const TMap<uint32, FMorphTargetMeshData>& MorphTargetMeshData, int32 LODIndex, bool bGenerateCPUMorphTargetsIfNeeded)
 	{
 		if (Owner.GetMorphTargets().IsEmpty())
 		{
@@ -880,14 +926,104 @@ namespace UnrealConversionUtils
 
 		TArray<const FMorphTargetLODModel*> MorphsToInit;
 		MorphsToInit.Reserve(MorphTargetLODs.Num());
-		for (FMorphTargetLODModel& MorphTargetLOD : MorphTargetLODs)
+		Algo::Transform(MorphTargetLODs, MorphsToInit, [](const FMorphTargetLODModel& Morph) { return &Morph; });
+
+		const float ErrorTolerance = Owner.GetLODInfo(LODIndex)->MorphTargetPositionErrorTolerance;
+		const int32 NumVertices = LODResource.StaticVertexBuffers.StaticMeshVertexBuffer.GetNumVertices();
+
+		const bool bGenerateCPUMorphTargets = 
+				bGenerateCPUMorphTargetsIfNeeded && !FMorphTargetVertexInfoBuffers::IsPlatformShaderSupported(GMaxRHIShaderPlatform);
+
+#if !WITH_EDITOR
+		TArray<FMorphTargetCompressedLODModel> CompressedMorphs;
+		CompressedMorphs.SetNum(MorphTargetLODs.Num());
+		
+		TBitArray<> SectionNeedsTangents;
+		SectionNeedsTangents.Init(false, LODResource.RenderSections.Num());
+		for (int32 SectionIndex = 0; SectionIndex < SectionNeedsTangents.Num(); ++SectionIndex)
 		{
-			MorphsToInit.Add(&MorphTargetLOD);
+			SectionNeedsTangents[SectionIndex] = FMorphTargetVertexInfoBuffers::NeedsTangents(
+					LODResource.RenderSections[SectionIndex].bRecomputeTangent);
 		}
 
+		const bool bAllRenderSectionsNeedTangents = SectionNeedsTangents.CountSetBits() == SectionNeedsTangents.Num();
+
+		TBitArray<> VertexNeedsTangentsStorage;
+		TBitArray<>* VertexNeedsTangents = nullptr;
+
+		// When generating CPU Morph targets we need to generate tangents.
+		if (!(bGenerateCPUMorphTargets || bAllRenderSectionsNeedTangents))
+		{
+			VertexNeedsTangentsStorage.Init(false, NumVertices);
+			for (int32 SectionIndex = 0; SectionIndex < SectionNeedsTangents.Num(); ++SectionIndex)
+			{
+				const FSkelMeshRenderSection& RenderSection = LODResource.RenderSections[SectionIndex];
+				if (SectionNeedsTangents[SectionIndex])
+				{
+					VertexNeedsTangentsStorage.SetRange(RenderSection.BaseVertexIndex, RenderSection.NumVertices, true);
+				}
+			}
+
+			VertexNeedsTangents = &VertexNeedsTangentsStorage;
+		}
+
+		for (int32 Index = 0, NumMorphs = MorphTargetLODs.Num(); Index < NumMorphs; ++Index)
+		{
+			using namespace UE::MorphTargetVertexCodec;
+
+			CompressedMorphs[Index].PositionPrecision = ComputePositionPrecision(ErrorTolerance);
+			CompressedMorphs[Index].TangentPrecision = ComputeTangentPrecision();
+
+			Encode(MorphTargetLODs[Index].Vertices, 
+					VertexNeedsTangents,
+					CompressedMorphs[Index].PositionPrecision,
+					CompressedMorphs[Index].TangentPrecision,
+					CompressedMorphs[Index].PackedDeltaHeaders,
+					CompressedMorphs[Index].PackedDeltaData);
+
+			// The uncompressed data for this Morph is no longer needed, free it.
+			MorphTargetLODs[Index].Vertices.Empty();
+		}
+#endif
+		
+#if WITH_EDITOR
 		const TBitArray UsesBuiltinMorphTargetCompression(true, MorphTargetLODs.Num());
-		const float ErrorTolerance = Owner.GetLODInfo(LODIndex)->MorphTargetPositionErrorTolerance;
-		LODResource.MorphTargetVertexInfoBuffers.InitMorphResourcesStreaming(LODResource.RenderSections, MorphsToInit, LODResource.StaticVertexBuffers.StaticMeshVertexBuffer.GetNumVertices(), ErrorTolerance);
+		LODResource.MorphTargetVertexInfoBuffers.InitMorphResourcesStreaming(GMaxRHIShaderPlatform, LODResource.RenderSections, MorphsToInit, NumVertices, ErrorTolerance);
+#else
+		TArray<FVector4f> MorphTargetMinimumValues;
+		TArray<FVector4f> MorphTargetMaximumValues;
+
+		MorphTargetMinimumValues.SetNum(MorphTargetLODs.Num());
+		MorphTargetMaximumValues.SetNum(MorphTargetLODs.Num());
+	
+		ComputeMorphTargetsMinAndMaxValues(MorphsToInit, MorphTargetMinimumValues, MorphTargetMaximumValues);
+
+		TArray<const FMorphTargetCompressedLODModel*> CompressedMorphsToInit;
+		CompressedMorphsToInit.Reserve(CompressedMorphs.Num());
+		Algo::Transform(CompressedMorphs, CompressedMorphsToInit, [](const FMorphTargetCompressedLODModel& Morph) { return &Morph; });
+
+		LODResource.MorphTargetVertexInfoBuffers.InitMorphResourcesStreaming(GMaxRHIShaderPlatform, CompressedMorphsToInit, MorphTargetMinimumValues, MorphTargetMaximumValues);
+#endif
+		if (bGenerateCPUMorphTargets)
+		{
+#if !WITH_EDITOR
+			for (int32 Index = 0, NumMorphs = Owner.GetMorphTargets().Num(); Index < NumMorphs; ++Index)
+			{
+				UMorphTarget* MorphLOD = Owner.GetMorphTargets()[Index];
+				check(MorphLOD->GetCompressedLODModels().IsValidIndex(LODIndex));
+				MorphLOD->GetCompressedLODModels()[LODIndex] = MoveTemp(CompressedMorphs[Index]);
+			}
+#else
+			for (int32 Index = 0, NumMorphs = Owner.GetMorphTargets().Num(); Index < NumMorphs; ++Index)
+			{
+				UMorphTarget* MorphLOD = Owner.GetMorphTargets()[Index];
+
+				check(MorphLOD->GetMorphLODModels().IsValidIndex(LODIndex));
+				MorphLOD->GetMorphLODModels()[LODIndex].NumVertices = MorphTargetLODs[Index].Vertices.Num();
+				MorphLOD->GetMorphLODModels()[LODIndex].Vertices    = MoveTemp(MorphTargetLODs[Index].Vertices);
+			}
+#endif
+		}
 	}
 
 	
@@ -1366,28 +1502,14 @@ namespace UnrealConversionUtils
 						}
 						
 						uint32 BatchStartOffset = BatchHeaders.Num();
-
-						float MaximumValues[4] = { -FLT_MAX, -FLT_MAX, -FLT_MAX, -FLT_MAX };
-						float MinimumValues[4] = { +FLT_MAX, +FLT_MAX, +FLT_MAX, +FLT_MAX };
-
-						int32 NumSrcDeltas = MorphModel ? MorphModel->Vertices.Num() : 0;
 						
-						TConstArrayView<FMorphTargetDelta> MorphDeltas(MorphModel->Vertices);
-						for (int32 DeltaIndex = 0; DeltaIndex < NumSrcDeltas; ++DeltaIndex)
-						{
-							const FMorphTargetDelta& MorphDelta = MorphDeltas[DeltaIndex];
-							const FVector3f TangentZDelta = MorphDelta.TangentZDelta;
-
-							MaximumValues[0] = FMath::Max(MaximumValues[0], MorphDelta.PositionDelta.X);
-							MaximumValues[1] = FMath::Max(MaximumValues[1], MorphDelta.PositionDelta.Y);
-							MaximumValues[2] = FMath::Max(MaximumValues[2], MorphDelta.PositionDelta.Z);
-							MaximumValues[3] = FMath::Max(MaximumValues[3], FMath::Max(TangentZDelta.X, FMath::Max(TangentZDelta.Y, TangentZDelta.Z)));
-
-							MinimumValues[0] = FMath::Min(MinimumValues[0], MorphDelta.PositionDelta.X);
-							MinimumValues[1] = FMath::Min(MinimumValues[1], MorphDelta.PositionDelta.Y);
-							MinimumValues[2] = FMath::Min(MinimumValues[2], MorphDelta.PositionDelta.Z);
-							MinimumValues[3] = FMath::Min(MinimumValues[3], FMath::Min(TangentZDelta.X, FMath::Min(TangentZDelta.Y, TangentZDelta.Z)));
-						}
+						FVector4f MinimumValues;
+						FVector4f MaximumValues;
+
+						ComputeMorphTargetsMinAndMaxValues(
+								TConstArrayView<const FMorphTargetLODModel*>(&MorphModel, 1),
+								TArrayView<FVector4f>(&MinimumValues, 1),
+								TArrayView<FVector4f>(&MaximumValues, 1));
 
 						// Encode the actual morph vertex info into the quantized bitstream.
 						Encode(MorphModel->Vertices, &VertexNeedsTangents, PositionPrecision, TangentZPrecision, BatchHeaders, BitstreamData);
@@ -1395,8 +1517,8 @@ namespace UnrealConversionUtils
 						const uint32 MorphNumBatches = BatchHeaders.Num() - BatchStartOffset;
 						Result->Morph.BatchStartOffsetPerMorph.Add(BatchStartOffset);
 						Result->Morph.BatchesPerMorph.Add(MorphNumBatches);
-						Result->Morph.MaximumValuePerMorph.Add(FVector4f(MaximumValues[0], MaximumValues[1], MaximumValues[2], MaximumValues[3]));
-						Result->Morph.MinimumValuePerMorph.Add(FVector4f(MinimumValues[0], MinimumValues[1], MinimumValues[2], MinimumValues[3]));
+						Result->Morph.MaximumValuePerMorph.Add(MaximumValues);
+						Result->Morph.MinimumValuePerMorph.Add(MinimumValues);
 					}
 				
 					Result->Morph.PositionPrecision = PositionPrecision;
diff --git a/Engine/Source/Runtime/Engine/Classes/Animation/MorphTarget.h b/Engine/Source/Runtime/Engine/Classes/Animation/MorphTarget.h
--- a/Engine/Source/Runtime/Engine/Classes/Animation/MorphTarget.h
+++ b/Engine/Source/Runtime/Engine/Classes/Animation/MorphTarget.h
@@ -241,6 +241,15 @@ public:
 	/** morph mesh vertex data for each LOD */
 	virtual TArray<FMorphTargetLODModel>& GetMorphLODModels() { return MorphLODModels; }
 
+
+#if !WITH_EDITOR
+	/** compressed morph mesh vertex data for each LOD */
+	virtual const TArray<FMorphTargetCompressedLODModel>& GetCompressedLODModels() const { return CompressedLODModels; }
+
+	/** compressed morph mesh vertex data for each LOD */
+	virtual TArray<FMorphTargetCompressedLODModel>& GetCompressedLODModels() { return CompressedLODModels; }
+#endif
+
 protected:
 	/** morph mesh vertex data for each LOD */
 	TArray<FMorphTargetLODModel>	MorphLODModels;
diff --git a/Engine/Source/Runtime/Engine/Private/Rendering/MorphTargetVertexInfoBuffers.cpp b/Engine/Source/Runtime/Engine/Private/Rendering/MorphTargetVertexInfoBuffers.cpp
--- a/Engine/Source/Runtime/Engine/Private/Rendering/MorphTargetVertexInfoBuffers.cpp
+++ b/Engine/Source/Runtime/Engine/Private/Rendering/MorphTargetVertexInfoBuffers.cpp
@@ -165,10 +165,16 @@ void FMorphTargetVertexInfoBuffers::InitMorphResources(EShaderPlatform ShaderPla
 	ValidateVertexBuffers(true);
 }
 
-void FMorphTargetVertexInfoBuffers::InitMorphResourcesStreaming(const TArray<FSkelMeshRenderSection>& RenderSections, const TArray<const FMorphTargetLODModel*>& MorphTargets, int32 NumVertices, float TargetPositionErrorTolerance)
+void FMorphTargetVertexInfoBuffers::InitMorphResourcesStreaming(EShaderPlatform ShaderPlatform, const TArray<FSkelMeshRenderSection>& RenderSections, const TArray<const FMorphTargetLODModel*>& MorphTargets, int32 NumVertices, float TargetPositionErrorTolerance)
 {
 	ResetCPUData();
 
+	// UseGPUMorphTargets() can be toggled only on SM5 atm
+	if (!IsPlatformShaderSupported(ShaderPlatform) || MorphTargets.Num() == 0)
+	{
+		return;
+	}
+
 	const TBitArray UsesBuiltinMorphTargetCompression(true, MorphTargets.Num());
 	Compress(RenderSections, MorphTargets, UsesBuiltinMorphTargetCompression, NumVertices, TargetPositionErrorTolerance);
 	
@@ -178,6 +184,157 @@ void FMorphTargetVertexInfoBuffers::InitMorphResourcesStreaming(const TArray<FSk
 	ValidateVertexBuffers(true);
 }
 
+#if !WITH_EDITOR
+void FMorphTargetVertexInfoBuffers::InitMorphResourcesStreaming(
+		EShaderPlatform ShaderPlatform,
+		TConstArrayView<const FMorphTargetCompressedLODModel*> InCompressedMorphTargets, 
+		TConstArrayView<FVector4f> InMinimumValuePerMorph,  
+		TConstArrayView<FVector4f> InMaximumValuePerMorph)
+{
+	ResetCPUData();
+
+	// UseGPUMorphTargets() can be toggled only on SM5 atm
+	if (!IsPlatformShaderSupported(ShaderPlatform) || InCompressedMorphTargets.Num() == 0)
+	{
+		return;
+	}
+
+	InitDataFromCompressedMorphs(InCompressedMorphTargets, InMinimumValuePerMorph, InMaximumValuePerMorph);
+
+	bRHIInitialized = false;
+	bIsMorphCPUDataValid = true;
+	bResourcesInitialized = true;
+	ValidateVertexBuffers(true);
+}
+
+void FMorphTargetVertexInfoBuffers::InitDataFromCompressedMorphs(
+		TConstArrayView<const FMorphTargetCompressedLODModel*> InCompressedMorphTargets, 
+		TConstArrayView<FVector4f> InMinimumValuePerMorph,  
+		TConstArrayView<FVector4f> InMaximumValuePerMorph)
+{
+	// Assum all morphs targets have the same PositionPrecision and TangentZPrecision.
+	{
+		PositionPrecision = 0.0f;
+		TangentZPrecision = 0.0f;
+
+		int32 MorphIndex = 0;
+		for (; MorphIndex < InCompressedMorphTargets.Num(); ++MorphIndex)
+		{
+			const FMorphTargetCompressedLODModel* CompressedMorph = InCompressedMorphTargets[MorphIndex];
+			if (CompressedMorph)
+			{
+				PositionPrecision = CompressedMorph->PositionPrecision;
+				TangentZPrecision = CompressedMorph->TangentPrecision;
+				break;
+			}
+		}
+
+#if DO_CHECK
+		for (; MorphIndex < InCompressedMorphTargets.Num(); ++MorphIndex)
+		{
+			const FMorphTargetCompressedLODModel* CompressedMorph = InCompressedMorphTargets[MorphIndex];
+			if (CompressedMorph)
+			{
+				check(PositionPrecision == CompressedMorph->PositionPrecision);
+				check(TangentZPrecision == CompressedMorph->TangentPrecision);
+			}
+		}
+#endif
+	}	
+
+	BatchStartOffsetPerMorph.Empty(InCompressedMorphTargets.Num());
+	BatchesPerMorph.Empty(InCompressedMorphTargets.Num());
+	MaximumValuePerMorph.Empty(InCompressedMorphTargets.Num());
+	MinimumValuePerMorph.Empty(InCompressedMorphTargets.Num());
+
+	check(InMaximumValuePerMorph.Num() == InCompressedMorphTargets.Num());
+	check(InMinimumValuePerMorph.Num() == InCompressedMorphTargets.Num());
+
+	// Compute allocation size and add metadata.
+	int32 TotalDataSizeInDwords = 0;
+	NumTotalBatches = 0;
+
+	int32 CurrentMorphBatchStartOffsetInBatches = 0;
+	for (int32 MorphIndex = 0; MorphIndex < InCompressedMorphTargets.Num(); ++MorphIndex)
+	{
+		using namespace UE::MorphTargetVertexCodec;
+		
+		const FMorphTargetCompressedLODModel* CompressedMorph = InCompressedMorphTargets[MorphIndex];
+
+		if (CompressedMorph)
+		{
+			const int32 MorphTargetNumBatches = CompressedMorph->PackedDeltaHeaders.Num();
+		
+			BatchesPerMorph.Add(MorphTargetNumBatches);
+			BatchStartOffsetPerMorph.Add(CurrentMorphBatchStartOffsetInBatches);
+			MinimumValuePerMorph.Add(InMinimumValuePerMorph[MorphIndex]);
+			MaximumValuePerMorph.Add(InMaximumValuePerMorph[MorphIndex]);
+
+			CurrentMorphBatchStartOffsetInBatches += MorphTargetNumBatches;
+
+			NumTotalBatches += MorphTargetNumBatches;
+			TotalDataSizeInDwords += 
+					MorphTargetNumBatches * UE::MorphTargetVertexCodec::NumBatchHeaderDwords + 
+					CompressedMorph->PackedDeltaData.Num();
+
+			static_assert(sizeof(decltype(CompressedMorph->PackedDeltaData)::ElementType) == sizeof(uint32));
+		}
+		else
+		{
+			BatchesPerMorph.Add(0);
+			BatchStartOffsetPerMorph.Add(CurrentMorphBatchStartOffsetInBatches);
+			MinimumValuePerMorph.Add(FVector4f(0.0f, 0.0f, 0.0f, 0.0f));
+			MaximumValuePerMorph.Add(FVector4f(0.0f, 0.0f, 0.0f, 0.0f));
+		}
+	}
+
+	// Add extra padding to make sure it is always safe to access the data with load4s.
+	MorphData.Empty(TotalDataSizeInDwords > 0 ? TotalDataSizeInDwords + 3 : 0);
+
+	// Fixup and append headers to the bitstream.
+	int32 AccumulatedDataOffsetFixupInBytes = NumTotalBatches * UE::MorphTargetVertexCodec::NumBatchHeaderDwords * sizeof(uint32);
+	for (const FMorphTargetCompressedLODModel* CompressedMorph : InCompressedMorphTargets)
+	{
+		using namespace UE::MorphTargetVertexCodec;
+
+		if (CompressedMorph)
+		{
+			for (const FDeltaBatchHeader& BatchHeader : CompressedMorph->PackedDeltaHeaders)
+			{
+				TStaticArray<uint32, NumBatchHeaderDwords> HeaderData;
+				WriteHeader(BatchHeader, HeaderData);
+
+				// Fixup data offsets. 
+				// First header dword corresponds to the DataOffset. See UE::MorphTargetVertexCodec::WriteHeader.
+				HeaderData[0] += AccumulatedDataOffsetFixupInBytes;
+				MorphData.Append(HeaderData);
+			}	
+
+			AccumulatedDataOffsetFixupInBytes += CompressedMorph->PackedDeltaData.NumBytes();
+		}
+	}
+
+	// Append deltas data to the bitstream.
+	for (const FMorphTargetCompressedLODModel* CompressedMorph : InCompressedMorphTargets)
+	{
+		if (CompressedMorph)
+		{
+			MorphData.Append(CompressedMorph->PackedDeltaData);
+		}
+	}
+
+	check(TotalDataSizeInDwords == MorphData.Num());
+
+	if (MorphData.Num() > 0)
+	{
+		// Pad to make sure it is always safe to access the data with load4s.
+		MorphData.Add(0u);
+		MorphData.Add(0u);
+		MorphData.Add(0u);
+	}
+}
+#endif
+
 void FMorphTargetVertexInfoBuffers::Compress(const TArray<FSkelMeshRenderSection>& RenderSections, const TArray<const FMorphTargetLODModel*>& MorphTargets, const TBitArray<>& UsesBuiltinMorphTargetCompression, int32 NumVertices, float TargetPositionErrorTolerance)
 {
 	using namespace UE::MorphTargetVertexCodec;
@@ -193,18 +350,14 @@ void FMorphTargetVertexInfoBuffers::Compress(const TArray<FSkelMeshRenderSection
 	MinimumValuePerMorph.Empty(MorphTargets.Num());
 
 	// Mark vertices that are in a section that doesn't recompute tangents as needing tangents
-	const int32 RecomputeTangentsMode = GSkinCacheRecomputeTangents;
 	TBitArray<> VertexNeedsTangents;
 	VertexNeedsTangents.Init(false, NumVertices);
 	for (const FSkelMeshRenderSection& RenderSection : RenderSections)
 	{
-		const bool bRecomputeTangents = RecomputeTangentsMode > 0 && (RenderSection.bRecomputeTangent || RecomputeTangentsMode == 1);
-		if (!bRecomputeTangents)
+		const bool bNeedsTangents = FMorphTargetVertexInfoBuffers::NeedsTangents(RenderSection.bRecomputeTangent);
+		if (bNeedsTangents)
 		{
-			for (uint32 i = 0; i < RenderSection.NumVertices; i++)
-			{
-				VertexNeedsTangents[RenderSection.BaseVertexIndex + i] = true;
-			}
+			VertexNeedsTangents.SetRange(RenderSection.BaseVertexIndex, RenderSection.NumVertices, true);
 		}
 	}
 	
@@ -297,3 +450,8 @@ bool FMorphTargetVertexInfoBuffers::IsPlatformShaderSupported(EShaderPlatform Sh
 {
 	return IsFeatureLevelSupported(ShaderPlatform, ERHIFeatureLevel::SM5);
 }
+
+bool FMorphTargetVertexInfoBuffers::NeedsTangents(bool bRecomputeTangents)
+{
+	return !(GSkinCacheRecomputeTangents > 0 && (bRecomputeTangents || GSkinCacheRecomputeTangents == 1));
+}
diff --git a/Engine/Source/Runtime/Engine/Public/Rendering/MorphTargetVertexInfoBuffers.h b/Engine/Source/Runtime/Engine/Public/Rendering/MorphTargetVertexInfoBuffers.h
--- a/Engine/Source/Runtime/Engine/Public/Rendering/MorphTargetVertexInfoBuffers.h
+++ b/Engine/Source/Runtime/Engine/Public/Rendering/MorphTargetVertexInfoBuffers.h
@@ -10,6 +10,10 @@ struct FSkelMeshRenderSection;
 struct FMorphTargetLODModel;
 class UMorphTarget;
 
+#if !WITH_EDITOR
+struct FMorphTargetCompressedLODModel;
+#endif
+
 class FMorphTargetVertexInfoBuffers : public FRenderResource
 {
 public:
@@ -18,7 +22,11 @@ public:
 
 	ENGINE_API void InitMorphResources(EShaderPlatform ShaderPlatform, const TArray<FSkelMeshRenderSection>& RenderSections, const TArray<UMorphTarget*>& MorphTargets, int32 NumVertices, int32 LODIndex, float TargetPositionErrorTolerance);
 
-	ENGINE_API void InitMorphResourcesStreaming(const TArray<FSkelMeshRenderSection>& RenderSections, const TArray<const FMorphTargetLODModel*>& MorphTargets, int32 NumVertices, float TargetPositionErrorTolerance);
+	ENGINE_API void InitMorphResourcesStreaming(EShaderPlatform ShaderPlatform, const TArray<FSkelMeshRenderSection>& RenderSections, const TArray<const FMorphTargetLODModel*>& MorphTargets, int32 NumVertices, float TargetPositionErrorTolerance);
+
+#if !WITH_EDITOR
+	ENGINE_API void InitMorphResourcesStreaming(EShaderPlatform ShaderPlatform, TConstArrayView<const FMorphTargetCompressedLODModel*> MorphTargets, TConstArrayView<FVector4f> InMinimumValuePerMorph, TConstArrayView<FVector4f> InMaximumValuesPerMorph);
+#endif
 
 	inline bool IsMorphResourcesInitialized() const { return bResourcesInitialized; }
 	inline bool IsRHIInitialized() const { return bRHIInitialized; }
@@ -83,6 +91,7 @@ public:
 	}
 	
 	static ENGINE_API bool IsPlatformShaderSupported(EShaderPlatform ShaderPlatform);
+	static ENGINE_API bool NeedsTangents(bool bRecomputeTangents);
 
 	FBufferRHIRef MorphDataBuffer;
 	FShaderResourceViewRHIRef MorphDataSRV;
@@ -95,6 +104,9 @@ protected:
 
 	void Compress(const TArray<FSkelMeshRenderSection>& RenderSections, const TArray<const FMorphTargetLODModel*>& MorphTargets, const TBitArray<>& UsesBuiltinMorphTargetCompression, int32 NumVertices, float TargetPositionErrorTolerance);
 
+#if !WITH_EDITOR
+	void InitDataFromCompressedMorphs(TConstArrayView<const FMorphTargetCompressedLODModel*> MorphTargets, TConstArrayView<FVector4f> InMinimumValuePerMorph, TConstArrayView<FVector4f> MaximumValuesPerMorph);
+#endif
 	// Transient data. Gets deleted as soon as the GPU resource has been initialized (unless excplicitly disabled by bEmptyMorphCPUDataOnInitRHI).
 	TResourceArray<uint32> MorphData;
 
