[Mutable] More secure management of FTexturePlatformData pointers.

#Synced-CL 48348317

diff --git a/Engine/Plugins/Mutable/Source/CustomizableObject/Internal/MuCO/CustomizableObjectSystemPrivate.h b/Engine/Plugins/Mutable/Source/CustomizableObject/Internal/MuCO/CustomizableObjectSystemPrivate.h
--- a/Engine/Plugins/Mutable/Source/CustomizableObject/Internal/MuCO/CustomizableObjectSystemPrivate.h
+++ b/Engine/Plugins/Mutable/Source/CustomizableObject/Internal/MuCO/CustomizableObjectSystemPrivate.h
@@ -643,7 +643,7 @@ public:
 	
 	TMap<FName, uint8> FirstResidentLOD;
 
-	TMap<UE::Mutable::Private::FImageId, FTexturePlatformData*> ImageToPlatformDataMap;
+	TMap<UE::Mutable::Private::FImageId, TUniquePtr<FTexturePlatformData>> ImageToPlatformDataMap;
 
 	EUpdateResult UpdateResult = EUpdateResult::Success;
 
@@ -711,13 +711,6 @@ public:
 };
 
 
-/** Runtime data used during a mutable instance update */
-struct FMutableReleasePlatformOperationData
-{
-	TMap<UE::Mutable::Private::FImageId, FTexturePlatformData*> ImageToPlatformDataMap;
-};
-
-
 USTRUCT()
 struct FPendingReleaseSkeletalMeshInfo
 {
@@ -981,7 +974,7 @@ private:
 
 
 /** Set OnlyLOD to -1 to generate all mips */
-CUSTOMIZABLEOBJECT_API FTexturePlatformData* MutableCreateImagePlatformData(TSharedPtr<const UE::Mutable::Private::FImage> MutableImage, int32 OnlyLOD, uint16 FullSizeX, uint16 FullSizeY);
+CUSTOMIZABLEOBJECT_API TUniquePtr<FTexturePlatformData> MutableCreateImagePlatformData(TSharedPtr<const UE::Mutable::Private::FImage> MutableImage, int32 OnlyLOD, uint16 FullSizeX, uint16 FullSizeY);
 
 
 /** Return true if Streaming is enabled for the given Object. */
diff --git a/Engine/Plugins/Mutable/Source/CustomizableObject/Private/MuCO/CustomizableObjectInstance.cpp b/Engine/Plugins/Mutable/Source/CustomizableObject/Private/MuCO/CustomizableObjectInstance.cpp
--- a/Engine/Plugins/Mutable/Source/CustomizableObject/Private/MuCO/CustomizableObjectInstance.cpp
+++ b/Engine/Plugins/Mutable/Source/CustomizableObject/Private/MuCO/CustomizableObjectInstance.cpp
@@ -3246,7 +3246,7 @@ void FMutableUpdateCandidate::ApplyLODUpdateParamsToInstance(FUpdateContextPriva
 
 // The memory allocated in the function and pointed by the returned pointer is owned by the caller and must be freed. 
 // If assigned to a UTexture2D, it will be freed by that UTexture2D
-FTexturePlatformData* MutableCreateImagePlatformData(TSharedPtr<const UE::Mutable::Private::FImage> MutableImage, int32 OnlyLOD, uint16 FullSizeX, uint16 FullSizeY)
+TUniquePtr<FTexturePlatformData> MutableCreateImagePlatformData(TSharedPtr<const UE::Mutable::Private::FImage> MutableImage, int32 OnlyLOD, uint16 FullSizeX, uint16 FullSizeY)
 {
 	int32 SizeX = FMath::Max(MutableImage->GetSize()[0], FullSizeX);
 	int32 SizeY = FMath::Max(MutableImage->GetSize()[1], FullSizeY);
@@ -3387,7 +3387,7 @@ FTexturePlatformData* MutableCreateImagePlatformData(TSharedPtr<const UE::Mutabl
 		}		
 	}
 
-	FTexturePlatformData* PlatformData = new FTexturePlatformData(); // TODO GMT Change to SharedRef
+	TUniquePtr<FTexturePlatformData> PlatformData = MakeUnique<FTexturePlatformData>();
 	PlatformData->SizeX = SizeX;
 	PlatformData->SizeY = SizeY;
 	PlatformData->PixelFormat = PlatformFormat;
@@ -3591,7 +3591,7 @@ void ConvertImage(UTexture2D* Texture, TSharedPtr<const UE::Mutable::Private::FI
 		OnlyLOD = FMath::Min( OnlyLOD, MutableImage->GetLODCount()-1 );
 	}
 
-	Texture->SetPlatformData(MutableCreateImagePlatformData(MutableImage, OnlyLOD,0,0) );
+	Texture->SetPlatformData(MutableCreateImagePlatformData(MutableImage, OnlyLOD,0,0).Release());
 }
 
 static int32 EnableRayTracingFix = 0;
@@ -5517,75 +5517,15 @@ void AdditionalAssetsAsyncLoaded(const TSharedRef<FUpdateContextPrivate> Context
 }
 
 
-void UpdateTextureRegionsMutable(UTexture2D* Texture, int32 MipIndex, uint32 NumMips, const FUpdateTextureRegion2D& Region, uint32 SrcPitch, 
-								 const FByteBulkData* BulkData, TSharedRef<FTexturePlatformData, ESPMode::ThreadSafe>& PlatformData)
+void ReuseTexture(UTexture2D* Texture, TUniquePtr<FTexturePlatformData> PlatformData)
 {
-	if (Texture->GetResource())
-	{
-		struct FUpdateTextureRegionsData
-		{
-			FTexture2DResource* Texture2DResource;
-			int32 MipIndex;
-			FUpdateTextureRegion2D Region;
-			uint32 SrcPitch;
-			uint32 NumMips;
-			
-			// The Platform Data mips will be automatically deleted when all FUpdateTextureRegionsData that reference it are deleted
-			// in the render thread after being used to update the texture
-			TSharedRef<FTexturePlatformData, ESPMode::ThreadSafe> PlatformData;
-
-			FUpdateTextureRegionsData(TSharedRef<FTexturePlatformData, ESPMode::ThreadSafe>& InPlatformData) : PlatformData(InPlatformData) {}
-		};
-
-		FUpdateTextureRegionsData* RegionData = new FUpdateTextureRegionsData(PlatformData);
-
-		RegionData->Texture2DResource = (FTexture2DResource*)Texture->GetResource();
-		RegionData->MipIndex = MipIndex;
-		RegionData->Region = Region;
-		RegionData->SrcPitch = SrcPitch;
-		RegionData->NumMips = NumMips;
-
-		ENQUEUE_RENDER_COMMAND(UpdateTextureRegionsMutable)(
-			[RegionData, BulkData](FRHICommandList& CmdList)
-			{
-				check(int32(RegionData->NumMips) >= RegionData->Texture2DResource->GetCurrentMipCount());
-				int32 MipDifference = RegionData->NumMips - RegionData->Texture2DResource->GetCurrentMipCount();
-				check(MipDifference >= 0);
-				int32 CurrentFirstMip = RegionData->Texture2DResource->GetCurrentFirstMip();
-				uint8* SrcData = (uint8*)BulkData->LockReadOnly();
-
-				//uint32 Size = RegionData->SrcPitch / (sizeof(uint8) * 4);
-				//UE_LOG(LogMutable, Warning, TEXT("UpdateTextureRegionsMutable MipIndex = %d, FirstMip = %d, size = %d"),
-				//	RegionData->MipIndex, CurrentFirstMip, Size);
-
-				//checkf(Size <= RegionData->Texture2DResource->GetSizeX(),
-				//	TEXT("UpdateTextureRegionsMutable incorrect size. %d, %d. NumMips=%d"), 
-				//	Size, RegionData->Texture2DResource->GetSizeX(), RegionData->Texture2DResource->GetCurrentMipCount());
-
-				if (RegionData->MipIndex >= CurrentFirstMip + MipDifference)
-				{
-					CmdList.UpdateTexture2D(
-						RegionData->Texture2DResource->GetTexture2DRHI(),
-						RegionData->MipIndex - CurrentFirstMip - MipDifference,
-						RegionData->Region,
-						RegionData->SrcPitch,
-						SrcData);
-				}
-
-				BulkData->Unlock();
-				delete RegionData; // This will implicitly delete the Platform Data if this is the last RegionData referencing it
-			});
-	}
-}
-
-
-void ReuseTexture(UTexture2D* Texture, TSharedRef<FTexturePlatformData, ESPMode::ThreadSafe>& PlatformData)
-{
-	uint32 NumMips = PlatformData->Mips.Num();
+	int32 NumMips = PlatformData->Mips.Num();
 
-	for (uint32 i = 0; i < NumMips; i++)
+	TSharedPtr<FTexturePlatformData> SharedPlatformData = MakeShareable(PlatformData.Release()); 
+	
+	for (int32 MipIndex = 0; MipIndex < NumMips; ++MipIndex)
 	{
-		FTexture2DMipMap& Mip = PlatformData->Mips[i];
+		FTexture2DMipMap& Mip = SharedPlatformData->Mips[MipIndex];
 
 		if (Mip.BulkData.GetElementCount() > 0)
 		{
@@ -5601,8 +5541,36 @@ void ReuseTexture(UTexture2D* Texture, TSharedRef<FTexturePlatformData, ESPMode:
 			check(int32(Region.Width) <= Texture->GetSizeX());
 			check(int32(Region.Height) <= Texture->GetSizeY());
 
-			UpdateTextureRegionsMutable(Texture, i, NumMips, Region, 
-					                    Mip.SizeX * sizeof(uint8) * 4, &Mip.BulkData, PlatformData);
+			uint32 SrcPitch = Mip.SizeX * sizeof(uint8) * 4; 
+			
+			FTexture2DResource* Texture2DResource = static_cast<FTexture2DResource*>(Texture->GetResource());
+			if (Texture2DResource)
+			{
+				ENQUEUE_RENDER_COMMAND(UpdateTextureRegionsMutable)(
+				[MipIndex, NumMips, Region, SrcPitch, Texture2DResource, SharedPlatformData](FRHICommandList& CmdList)
+				{
+					FTexture2DMipMap& Mip = SharedPlatformData->Mips[MipIndex];
+					const FByteBulkData* BulkData = &Mip.BulkData;
+
+					check(NumMips >= Texture2DResource->GetCurrentMipCount());
+					int32 MipDifference = NumMips - Texture2DResource->GetCurrentMipCount();
+					check(MipDifference >= 0);
+					int32 CurrentFirstMip = Texture2DResource->GetCurrentFirstMip();
+					uint8* SrcData = (uint8*)BulkData->LockReadOnly();
+
+					if (MipIndex >= CurrentFirstMip + MipDifference)
+					{
+						CmdList.UpdateTexture2D(
+							Texture2DResource->GetTexture2DRHI(),
+							MipIndex - CurrentFirstMip - MipDifference,
+							Region,
+							SrcPitch,
+							SrcData);
+					}
+
+					BulkData->Unlock();
+				});
+			}
 		}
 	}
 }
@@ -5740,8 +5708,8 @@ UTexture* BuildTexture(const TSharedRef<FUpdateContextPrivate>& Context, const i
 		
 		SetTexturePropertiesFromMutableImageProps(Texture, Props, bNeverStream);
 
-		FTexturePlatformData* PlatformData = Context->ImageToPlatformDataMap[Image->ImageID];
-		Context->ImageToPlatformDataMap.Remove(Image->ImageID);
+		TUniquePtr<FTexturePlatformData> PlatformData;
+		Context->ImageToPlatformDataMap.RemoveAndCopyValue(Image->ImageID, PlatformData);
 
 		if (ReusedTexture)
 		{
@@ -5752,13 +5720,12 @@ UTexture* BuildTexture(const TSharedRef<FUpdateContextPrivate>& Context, const i
 			// Now the ReusedTexturePlatformData shared ptr owns the platform data
 			// This shared ptr will hold the reused texture platform data (mips) until the reused texture is updated 
 			// and delete it automatically
-			TSharedRef<FTexturePlatformData> PlatformDataRef = TSharedPtr<FTexturePlatformData>(PlatformData).ToSharedRef();
-			ReuseTexture(ReusedTexture, PlatformDataRef);
+			ReuseTexture(ReusedTexture, MoveTemp(PlatformData));
 		}
 		else
 		{
 			// Now the MutableTexture owns the platform data
-			Texture->SetPlatformData(PlatformData);
+			Texture->SetPlatformData(PlatformData.Release());
 
 			MUTABLE_CPUPROFILER_SCOPE(UpdateResource);
 #if REQUIRES_SINGLEUSE_FLAG_FOR_RUNTIME_TEXTURES
diff --git a/Engine/Plugins/Mutable/Source/CustomizableObject/Private/MuCO/CustomizableObjectSystem.cpp b/Engine/Plugins/Mutable/Source/CustomizableObject/Private/MuCO/CustomizableObjectSystem.cpp
--- a/Engine/Plugins/Mutable/Source/CustomizableObject/Private/MuCO/CustomizableObjectSystem.cpp
+++ b/Engine/Plugins/Mutable/Source/CustomizableObject/Private/MuCO/CustomizableObjectSystem.cpp
@@ -2786,8 +2786,7 @@ namespace impl
 					{
 						if (!OperationData->ImageToPlatformDataMap.Contains(Image->ImageID))
 						{
-							FTexturePlatformData* PlatformData = MutableCreateImagePlatformData(Image->Image, -1, Image->FullImageSizeX, Image->FullImageSizeY);
-							OperationData->ImageToPlatformDataMap.Add(Image->ImageID, PlatformData);
+							OperationData->ImageToPlatformDataMap.Add(Image->ImageID, MutableCreateImagePlatformData(Image->Image, -1, Image->FullImageSizeX, Image->FullImageSizeY));
 						}
 						else
 						{
@@ -2847,21 +2846,6 @@ namespace impl
 		}
 	}
 
-
-	void Task_Game_ReleasePlatformData(const TSharedPtr<FMutableReleasePlatformOperationData>& OperationData)
-	{
-		MUTABLE_CPUPROFILER_SCOPE(Task_Game_ReleasePlatformData)
-
-		check(OperationData);
-		
-		TMap<UE::Mutable::Private::FImageId, FTexturePlatformData*>& ImageToPlatformDataMap = OperationData->ImageToPlatformDataMap;
-		for (const TPair<UE::Mutable::Private::FImageId, FTexturePlatformData*>& Pair : ImageToPlatformDataMap)
-		{
-			delete Pair.Value; // If this is not null then it must mean it hasn't been used, otherwise they would have taken ownership and nulled it
-		}
-		ImageToPlatformDataMap.Reset();
-	}
-
 	
 	void Task_Game_Callbacks(const TSharedRef<FUpdateContextPrivate>& OperationData)
 	{
@@ -2991,44 +2975,25 @@ namespace impl
 			});
 
 
-		// Next Task: Release Platform Data
-		//-------------------------------------------------------------
-		if (!bInstanceInvalid)
-		{
-			TSharedPtr<FMutableReleasePlatformOperationData> ReleaseOperationData = MakeShared<FMutableReleasePlatformOperationData>();
-			check(ReleaseOperationData);
-			
-			static_assert(
-					std::is_same_v<
-						decltype(std::declval<FMutableReleasePlatformOperationData>().ImageToPlatformDataMap), 
-						decltype(std::declval<FUpdateContextPrivate>().ImageToPlatformDataMap)
-					>, 
-					"Cannot move FMutableReleasePlatformOperationData::ImageToPlatformDataMap to FUpdateContextPrivate::ImageToPlatformDataMap, types do not match.");
-
-			ReleaseOperationData->ImageToPlatformDataMap = MoveTemp(OperationData->ImageToPlatformDataMap);
-			CustomizableObjectSystemPrivate->MutableTaskGraph.AddAnyThreadTask(
-				TEXT("Mutable_ReleasePlatformData"),
-				[ReleaseOperationData]()
-				{
-					Task_Game_ReleasePlatformData(ReleaseOperationData);
-				}
-			);
-
-			// Next Task: Callbacks
-			//-------------------------------------------------------------
-			CustomizableObjectSystemPrivate->MutableTaskGraph.AddGameThreadTask(
-				TEXT("Task_Game_Callbacks"),
-				[OperationData]()
-				{
-					Task_Game_Callbacks(OperationData);
-				},
-				false);
-		}
-		else
+		// If they are still in this container it means they have not been used. No longer required. Free some memory now.
+		OperationData->ImageToPlatformDataMap.Empty();
+		
+		if (bInstanceInvalid)
 		{
 			OperationData->UpdateResult = EUpdateResult::Error;
 			FinishUpdateGlobal(OperationData);
+			return;
 		}
+		
+		// Next Task: Callbacks
+		//-------------------------------------------------------------
+		CustomizableObjectSystemPrivate->MutableTaskGraph.AddGameThreadTask(
+			TEXT("Task_Game_Callbacks"),
+			[OperationData]()
+			{
+				Task_Game_Callbacks(OperationData);
+			},
+			false);
 	}
 
 
