[Mutable] Connected Override Materials to BuildMaterial.

- Fixed some MATERIAL_PIN TODOs.
- Removed deprecated GetOverrideMaterials function (by index).
- Simplified FTextureCache.
- Removed redundant FInstance checks.

#Synced-CL 48348317

diff --git a/Engine/Plugins/Mutable/Source/CustomizableObject/Internal/MuCO/CustomizableObjectPrivate.h b/Engine/Plugins/Mutable/Source/CustomizableObject/Internal/MuCO/CustomizableObjectPrivate.h
--- a/Engine/Plugins/Mutable/Source/CustomizableObject/Internal/MuCO/CustomizableObjectPrivate.h
+++ b/Engine/Plugins/Mutable/Source/CustomizableObject/Internal/MuCO/CustomizableObjectPrivate.h
@@ -286,9 +286,7 @@ public:
 
 		int32 SkippedMips = 0;
 
-#if WITH_EDITORONLY_DATA
 		bool bIsBake = false;
-#endif
 		
 		bool operator==(const FId&) const = default;
 
diff --git a/Engine/Plugins/Mutable/Source/CustomizableObject/Internal/MuCO/CustomizableObjectSystemPrivate.h b/Engine/Plugins/Mutable/Source/CustomizableObject/Internal/MuCO/CustomizableObjectSystemPrivate.h
--- a/Engine/Plugins/Mutable/Source/CustomizableObject/Internal/MuCO/CustomizableObjectSystemPrivate.h
+++ b/Engine/Plugins/Mutable/Source/CustomizableObject/Internal/MuCO/CustomizableObjectSystemPrivate.h
@@ -335,7 +335,6 @@ struct FInstanceUpdateData
 		UE::Mutable::Private::EImageFormat Format = UE::Mutable::Private::EImageFormat::None;
 
 		TSharedPtr<const UE::Mutable::Private::FImage> Image;
-		TWeakObjectPtr<UTexture2D> Cached;
 
 		TArray<int32> ConstantImagesNeededToGenerate;
 
@@ -374,7 +373,7 @@ struct FInstanceUpdateData
 		/** Range in the Scalar array */
 		uint16 FirstScalar = 0;
 		uint16 ScalarCount = 0;
-
+		
 		bool bGetParameters = true; // TODO MATERIAL_PIN Remove once the Material Pin is completed. Only used to avid GetImages/Vector/Scalar twice (Material and Surface).
 		
 		UE::Mutable::Private::FMaterialId MaterialId;
@@ -410,6 +409,13 @@ struct FInstanceUpdateData
 		uint32 BoneMapCount = 0;
 	};
 
+	struct FOverrideMaterial
+	{
+		FName SlotName;
+		
+		uint32 MaterialIndex = INDEX_NONE;
+	};
+	
 	struct FComponent
 	{
 		FCustomizableObjectComponentIndex Id = FCustomizableObjectComponentIndex(0);
@@ -418,8 +424,11 @@ struct FInstanceUpdateData
 		uint16 FirstLOD = 0;
 		uint16 LODCount = 0;
 
-		int32 OverlayMaterialIndex = INDEX_NONE;
-
+		uint16 FirstOverrideMaterial = 0;
+		uint16 OverrideMaterialCount = 0;
+		
+		uint32 OverlayMaterialIndex = INDEX_NONE;
+		
 		UE::Mutable::Private::FSkeletalMeshId SkeletalMeshId;
 		
 		TSharedPtr<const UE::Mutable::Private::FSkeletalMesh> SkeletalMesh;
@@ -427,6 +436,7 @@ struct FInstanceUpdateData
 
 	TArray<TSharedRef<FComponent>> Components;
 	TArray<TSharedRef<FLOD>> LODs;
+	TArray<FOverrideMaterial> OverrideMaterials;
 	TArray<TSharedRef<FMaterial>> Materials;
 	TArray<TSharedRef<FSurface>> Surfaces;
 	TArray<TSharedRef<FImage>> Images;
@@ -480,6 +490,7 @@ struct FInstanceUpdateData
 		LODs.Empty();
 		Components.Empty();
 		Surfaces.Empty();
+		OverrideMaterials.Empty();
 		Materials.Empty();
 		Images.Empty();
 		Scalars.Empty();
diff --git a/Engine/Plugins/Mutable/Source/CustomizableObject/Private/MuCO/CustomizableObject.cpp b/Engine/Plugins/Mutable/Source/CustomizableObject/Private/MuCO/CustomizableObject.cpp
--- a/Engine/Plugins/Mutable/Source/CustomizableObject/Private/MuCO/CustomizableObject.cpp
+++ b/Engine/Plugins/Mutable/Source/CustomizableObject/Private/MuCO/CustomizableObject.cpp
@@ -3674,11 +3674,7 @@ FMutableParameterData::FMutableParameterData(const FMutableParamUIMetadata& InPa
 
 uint32 GetTypeHash(const FTextureCache::FId& Key)
 {
-#if WITH_EDITOR
 	return HashCombine(GetTypeHash(Key.Resource), Key.SkippedMips, Key.bIsBake);
-#else
-	return HashCombine(GetTypeHash(Key.Resource), Key.SkippedMips);
-#endif
 }
 
 
diff --git a/Engine/Plugins/Mutable/Source/CustomizableObject/Private/MuCO/CustomizableObjectInstance.cpp b/Engine/Plugins/Mutable/Source/CustomizableObject/Private/MuCO/CustomizableObjectInstance.cpp
--- a/Engine/Plugins/Mutable/Source/CustomizableObject/Private/MuCO/CustomizableObjectInstance.cpp
+++ b/Engine/Plugins/Mutable/Source/CustomizableObject/Private/MuCO/CustomizableObjectInstance.cpp
@@ -5314,37 +5314,6 @@ FCustomizableObjectInstanceDescriptor& UCustomizableInstancePrivate::GetDescript
 }
 
 
-const TArray<TObjectPtr<UMaterialInterface>>* UCustomizableObjectInstance::GetOverrideMaterials(int32 ComponentIndex) const
-{
-	UCustomizableObject* Object = GetCustomizableObject();
-	if (!Object)
-	{
-		return nullptr;
-	}
-
-	const UModelResources* ModelResources = Object->GetPrivate()->GetModelResources();
-	if (!ModelResources)
-	{
-		return nullptr;
-	}
-
-	if (!ModelResources->ComponentNamesPerObjectComponent.IsValidIndex(ComponentIndex))
-	{
-		return nullptr;	
-	}
-	
-	const FName ComponentName = ModelResources->ComponentNamesPerObjectComponent[ComponentIndex];
-
-	FCustomizableInstanceComponentData* ComponentData = PrivateData->GetComponentData(ComponentName);
-	if (!ComponentData)
-	{
-		return nullptr;
-	}
-
-	return &ComponentData->OverrideMaterials;
-}
-
-
 TArray<UMaterialInterface*> UCustomizableObjectInstance::GetSkeletalMeshComponentOverrideMaterials(const FName& ComponentName) const
 {
 	FCustomizableInstanceComponentData* ComponentData = PrivateData->GetComponentData(ComponentName);
@@ -5601,11 +5570,7 @@ UTexture* BuildTexture(const TSharedRef<FUpdateContextPrivate>& Context, const i
 {
 	const TSharedRef<const FInstanceUpdateData::FImage>& Image = Context->InstanceUpdateData.Images[ImageInex];
 	
-#if WITH_EDITOR // TODO MATERIAL_PIN Remove WITH_EDITOR
 	const FTextureCache::FId TextureCacheKey(Image->ImageID, Context->MipsToSkip, Context->bBake);
-#else
-	const FTextureCache::FId TextureCacheKey(Image->ImageID, Context->MipsToSkip);
-#endif
 	
 	TStrongObjectPtr<UTexture>* CachedTexture = Context->TextureUpdateCache.Find(TextureCacheKey);
 	if (CachedTexture)
@@ -5800,8 +5765,6 @@ FBuildMaterialOut BuildMaterial(const TSharedRef<FUpdateContextPrivate>& Context
 
 	if (MaterialIndex == INDEX_NONE)
 	{
-		Out.Material = UMaterial::GetDefaultMaterial(MD_Surface); // TODO MATERIAL_PIN. It should be possible to return nullptr (Overlay and Override can be nullptr). The caller then shoud be responsible of setting a Default or not.
-		
 		return Out;
 	}
 
@@ -6054,6 +6017,26 @@ void UCustomizableInstancePrivate::BuildMaterials(const TSharedRef<FUpdateContex
 				
 				NotifyUpdateOnInstanceComponent[InstanceComponentIndex] = Out.bNotifyPrimitiveUpdated;
 			}
+			
+			for (int32 MaterialIndex = 0; MaterialIndex < Component->OverrideMaterialCount; ++MaterialIndex)
+			{
+				const FInstanceUpdateData::FOverrideMaterial& OverrideMaterial = Context->InstanceUpdateData.OverrideMaterials[Component->FirstOverrideMaterial + MaterialIndex];
+				
+				const int32 SlotIndex = SkeletalMesh->GetMaterials().IndexOfByPredicate([&](const FSkeletalMaterial& Material)
+				{
+					return Material.MaterialSlotName == OverrideMaterial.SlotName;
+				});
+				
+				if (SlotIndex != INDEX_NONE)
+				{
+					FBuildMaterialOut Out = BuildMaterial(Context, OverrideMaterial.MaterialIndex, !bUVChanged);
+					
+					ComponentData->OverrideMaterials.SetNum(SlotIndex + 1);
+					ComponentData->OverrideMaterials[SlotIndex] = Out.Material;
+					
+					NotifyUpdateOnInstanceComponent[InstanceComponentIndex] = Out.bNotifyPrimitiveUpdated;
+				}
+			}
 		}
 
 		if (!bUseOverrideMaterialsOnly)
@@ -6177,6 +6160,12 @@ void UCustomizableInstancePrivate::BuildMaterials(const TSharedRef<FUpdateContex
 				if (ComponentData)
 				{
 					FBuildMaterialOut Out = BuildMaterial(Context, Surface->MaterialIndex, !bUVChanged);
+					
+					if (!Out.Material)
+					{
+						Out.Material = UMaterial::GetDefaultMaterial(MD_Surface);
+					}
+					
 					MaterialSlot.MaterialInterface = Out.Material;
 					ComponentData->OverrideMaterials.Add(Out.Material);
 					NotifyUpdateOnInstanceComponent[InstanceComponentIndex] = Out.bNotifyPrimitiveUpdated;
diff --git a/Engine/Plugins/Mutable/Source/CustomizableObject/Private/MuCO/CustomizableObjectSystem.cpp b/Engine/Plugins/Mutable/Source/CustomizableObject/Private/MuCO/CustomizableObjectSystem.cpp
--- a/Engine/Plugins/Mutable/Source/CustomizableObject/Private/MuCO/CustomizableObjectSystem.cpp
+++ b/Engine/Plugins/Mutable/Source/CustomizableObject/Private/MuCO/CustomizableObjectSystem.cpp
@@ -2540,6 +2540,22 @@ namespace impl
 				Material->MaterialId = MaterialId;
 			}
 			
+			const int32 OverrideMaterialCount = Instance->GetOverrideMaterialCount(InstanceComponentIndex);
+			Component->FirstOverrideMaterial = OperationData->InstanceUpdateData.OverrideMaterials.Num();
+			for (int32 MaterialIndex = 0; MaterialIndex < OverrideMaterialCount; ++MaterialIndex)
+			{
+				if (UE::Mutable::Private::FMaterialId MaterialId = Instance->GetOverrideMaterialId(InstanceComponentIndex, MaterialIndex))
+				{
+					FInstanceUpdateData::FOverrideMaterial& OverrideMaterial = OperationData->InstanceUpdateData.OverrideMaterials.AddDefaulted_GetRef();
+					++Component->OverrideMaterialCount;
+					OverrideMaterial.SlotName = Instance->GetOverrideMaterialSlotSlotName(InstanceComponentIndex, MaterialIndex);
+					
+					OverrideMaterial.MaterialIndex = OperationData->InstanceUpdateData.Materials.Num();
+					TSharedRef<FInstanceUpdateData::FMaterial> Material = OperationData->InstanceUpdateData.Materials.Add_GetRef(MakeShared<FInstanceUpdateData::FMaterial>());
+					Material->MaterialId = MaterialId;
+				}	
+			}
+			
 			Component->LODCount = OperationData->NumLODsAvailable[ComponentName];
 			if (!Component->LODCount)
 			{
diff --git a/Engine/Plugins/Mutable/Source/CustomizableObject/Public/MuCO/CustomizableObjectInstance.h b/Engine/Plugins/Mutable/Source/CustomizableObject/Public/MuCO/CustomizableObjectInstance.h
--- a/Engine/Plugins/Mutable/Source/CustomizableObject/Public/MuCO/CustomizableObjectInstance.h
+++ b/Engine/Plugins/Mutable/Source/CustomizableObject/Public/MuCO/CustomizableObjectInstance.h
@@ -377,9 +377,6 @@ public:
 	UFUNCTION(BlueprintCallable, Category = CustomizableObjectInstance)
 	CUSTOMIZABLEOBJECT_API USkeletalMesh* GetSkeletalMeshComponentSkeletalMesh(const FName& ComponentName) const;
 	
-	// DEPRECATED: Use GetSkeletalMeshComponentOverrideMaterials 
-	CUSTOMIZABLEOBJECT_API const TArray<TObjectPtr<UMaterialInterface>>* GetOverrideMaterials(int32 ComponentIndex) const;
-
 	UFUNCTION(BlueprintCallable, Category = CustomizableObjectInstance)
 	CUSTOMIZABLEOBJECT_API TArray<UMaterialInterface*> GetSkeletalMeshComponentOverrideMaterials(const FName& ComponentName) const;
 	
diff --git a/Engine/Plugins/Mutable/Source/MutableRuntime/Internal/MuR/Instance.h b/Engine/Plugins/Mutable/Source/MutableRuntime/Internal/MuR/Instance.h
--- a/Engine/Plugins/Mutable/Source/MutableRuntime/Internal/MuR/Instance.h
+++ b/Engine/Plugins/Mutable/Source/MutableRuntime/Internal/MuR/Instance.h
@@ -70,6 +70,13 @@ namespace UE::Mutable::Private
 		TArray<FInstanceSurface, TInlineAllocator<1>> Surfaces;
 	};
 
+	
+	struct FOverrideMaterial
+	{
+		FName SlotName;
+		FMaterialId MaterialId;
+	};
+	
 
 	struct FInstanceComponent
 	{
@@ -77,7 +84,7 @@ namespace UE::Mutable::Private
 
 		FMaterialId OverlayMaterialId;
 
-		TMap<FName, FMaterialId> OverrideMaterialIds;
+		TArray<FOverrideMaterial> OverrideMaterials;
 
 		FSkeletalMeshId SkeletalMeshId;
 
@@ -265,6 +272,12 @@ namespace UE::Mutable::Private
 		//! Get the index of the overlay material.
 		//! \param ComponentIndex Index of the component, from 0 to GetComponentCount()-1
 		UE_API FMaterialId GetOverlayMaterialId(int32 ComponentIndex) const;
+
+		UE_API int32 GetOverrideMaterialCount(int32 ComponentIndex) const;
+
+		UE_API FName GetOverrideMaterialSlotSlotName(int32 ComponentIndex, int32 MaterialIndex) const;
+
+		UE_API FMaterialId GetOverrideMaterialId(int32 ComponentIndex, int32 MaterialIndex) const;
 		
 		UE_API FMaterialId GetMaterialId(int32 ComponentIndex, int32 LODIndex, int32 SurfaceIndex) const;
 
diff --git a/Engine/Plugins/Mutable/Source/MutableRuntime/Private/MuR/Instance.cpp b/Engine/Plugins/Mutable/Source/MutableRuntime/Private/MuR/Instance.cpp
--- a/Engine/Plugins/Mutable/Source/MutableRuntime/Private/MuR/Instance.cpp
+++ b/Engine/Plugins/Mutable/Source/MutableRuntime/Private/MuR/Instance.cpp
@@ -132,159 +132,90 @@ namespace UE::Mutable::Private
 	
     int32 FInstance::GetImageCount(int32 ComponentIndex, int32 LODIndex, int32 SurfaceIndex ) const
 	{
-		check(Components.IsValidIndex(ComponentIndex));
-		check(Components[ComponentIndex].LODs.IsValidIndex(LODIndex));
-		check(Components[ComponentIndex].LODs[LODIndex].Surfaces.IsValidIndex(SurfaceIndex));
-
 		return Components[ComponentIndex].LODs[LODIndex].Surfaces[SurfaceIndex].Images.Num();
 	}
 
 	
     int32 FInstance::GetVectorCount(int32 ComponentIndex, int32 LODIndex, int32 SurfaceIndex ) const
 	{
-		check(Components.IsValidIndex(ComponentIndex));
-		check(Components[ComponentIndex].LODs.IsValidIndex(LODIndex));
-		check(Components[ComponentIndex].LODs[LODIndex].Surfaces.IsValidIndex(SurfaceIndex));
-
 		return Components[ComponentIndex].LODs[LODIndex].Surfaces[SurfaceIndex].Vectors.Num();
 	}
 
 	
     int32 FInstance::GetScalarCount(int32 ComponentIndex, int32 LODIndex, int32 SurfaceIndex ) const
 	{
-		check(Components.IsValidIndex(ComponentIndex));
-		check(Components[ComponentIndex].LODs.IsValidIndex(LODIndex));
-		check(Components[ComponentIndex].LODs[LODIndex].Surfaces.IsValidIndex(SurfaceIndex));
-
 		return Components[ComponentIndex].LODs[LODIndex].Surfaces[SurfaceIndex].Scalars.Num();
 	}
 
     
     int32 FInstance::GetStringCount(int32 ComponentIndex, int32 LODIndex, int32 SurfaceIndex ) const
     {
-		check(Components.IsValidIndex(ComponentIndex));
-		check(Components[ComponentIndex].LODs.IsValidIndex(LODIndex));
-		check(Components[ComponentIndex].LODs[LODIndex].Surfaces.IsValidIndex(SurfaceIndex));
-
 		return Components[ComponentIndex].LODs[LODIndex].Surfaces[SurfaceIndex].Strings.Num();
 	}
 	
     
 	FMeshId FInstance::GetMeshId(int32 ComponentIndex, int32 LODIndex) const
     {
-        check(Components.IsValidIndex(ComponentIndex));
-        check(Components[ComponentIndex].LODs.IsValidIndex(LODIndex));
-
 		return Components[ComponentIndex].LODs[LODIndex].MeshId;
     }
 
 	
 	FImageId FInstance::GetImageId(int32 ComponentIndex, int32 LODIndex, int32 SurfaceIndex, int32 ImageIndex) const
 	{
-		check(Components.IsValidIndex(ComponentIndex));
-		check(Components[ComponentIndex].LODs.IsValidIndex(LODIndex));
-        check(Components[ComponentIndex].LODs[LODIndex].Surfaces.IsValidIndex(SurfaceIndex));
-        check(Components[ComponentIndex].LODs[LODIndex].Surfaces[SurfaceIndex].Images.IsValidIndex(ImageIndex));
-
 		return Components[ComponentIndex].LODs[LODIndex].Surfaces[SurfaceIndex].Images[ImageIndex].Id;
 	}
 
 	
     FName FInstance::GetImageName(int32 ComponentIndex, int32 LODIndex, int32 SurfaceIndex, int32 ImageIndex ) const
 	{
-		check(Components.IsValidIndex(ComponentIndex));
-		check(Components[ComponentIndex].LODs.IsValidIndex(LODIndex));
-        check(Components[ComponentIndex].LODs[LODIndex].Surfaces.IsValidIndex(SurfaceIndex));
-		check(Components[ComponentIndex].LODs[LODIndex].Surfaces[SurfaceIndex].Images.IsValidIndex(ImageIndex));
-
         return Components[ComponentIndex].LODs[LODIndex].Surfaces[SurfaceIndex].Images[ImageIndex].Name;
 	}
 
 	
 	int32 FInstance::GetImageLayerIndex(int32 ComponentIndex, int32 LODIndex, int32 SurfaceIndex, int32 ImageIndex) const
 	{
-		check(Components.IsValidIndex(ComponentIndex));
-		check(Components[ComponentIndex].LODs.IsValidIndex(LODIndex));
-		check(Components[ComponentIndex].LODs[LODIndex].Surfaces.IsValidIndex(SurfaceIndex));
-		check(Components[ComponentIndex].LODs[LODIndex].Surfaces[SurfaceIndex].Images.IsValidIndex(ImageIndex));
-
 		return Components[ComponentIndex].LODs[LODIndex].Surfaces[SurfaceIndex].Images[ImageIndex].LayerIndex;
 	}
 
 	
 	FVector4f FInstance::GetVector(int32 ComponentIndex, int32 LODIndex, int32 SurfaceIndex, int32 VectorIndex) const
 	{
-		check(Components.IsValidIndex(ComponentIndex));
-		check(Components[ComponentIndex].LODs.IsValidIndex(LODIndex));
-        check(Components[ComponentIndex].LODs[LODIndex].Surfaces.IsValidIndex(SurfaceIndex));
-		check(Components[ComponentIndex].LODs[LODIndex].Surfaces[SurfaceIndex].Vectors.IsValidIndex(VectorIndex));
-
         return Components[ComponentIndex].LODs[LODIndex].Surfaces[SurfaceIndex].Vectors[VectorIndex].Value;
 	}
 	
 	
 	FName FInstance::GetVectorName(int32 ComponentIndex, int32 LODIndex, int32 SurfaceIndex, int32 VectorIndex) const
 	{
-		check(Components.IsValidIndex(ComponentIndex));
-		check(Components[ComponentIndex].LODs.IsValidIndex(LODIndex));
-        check(Components[ComponentIndex].LODs[LODIndex].Surfaces.IsValidIndex(SurfaceIndex));
-		check(Components[ComponentIndex].LODs[LODIndex].Surfaces[SurfaceIndex].Vectors.IsValidIndex(VectorIndex));
-
         return Components[ComponentIndex].LODs[LODIndex].Surfaces[SurfaceIndex].Vectors[VectorIndex].Name;
 	}
 
 
 	int32 FInstance::GetVectorLayerIndex(int32 ComponentIndex, int32 LODIndex, int32 SurfaceIndex, int32 VectorIndex) const
 	{
-		check(Components.IsValidIndex(ComponentIndex));
-		check(Components[ComponentIndex].LODs.IsValidIndex(LODIndex));
-		check(Components[ComponentIndex].LODs[LODIndex].Surfaces.IsValidIndex(SurfaceIndex));
-		check(Components[ComponentIndex].LODs[LODIndex].Surfaces[SurfaceIndex].Vectors.IsValidIndex(VectorIndex));
-
 		return Components[ComponentIndex].LODs[LODIndex].Surfaces[SurfaceIndex].Vectors[VectorIndex].LayerIndex;
 	}
 	
 	
     float FInstance::GetScalar(int32 ComponentIndex, int32 LODIndex, int32 SurfaceIndex, int32 ScalarIndex ) const
 	{
-		check(Components.IsValidIndex(ComponentIndex));
-		check(Components[ComponentIndex].LODs.IsValidIndex(LODIndex));
-        check(Components[ComponentIndex].LODs[LODIndex].Surfaces.IsValidIndex(SurfaceIndex));
-		check(Components[ComponentIndex].LODs[LODIndex].Surfaces[SurfaceIndex].Scalars.IsValidIndex(ScalarIndex));
-
         return Components[ComponentIndex].LODs[LODIndex].Surfaces[SurfaceIndex].Scalars[ScalarIndex].Value;
 	}
 
 	
 	FName FInstance::GetScalarName(int32 ComponentIndex, int32 LODIndex, int32 SurfaceIndex, int32 ScalarIndex) const
 	{
-		check(Components.IsValidIndex(ComponentIndex));
-		check(Components[ComponentIndex].LODs.IsValidIndex(LODIndex));
-        check(Components[ComponentIndex].LODs[LODIndex].Surfaces.IsValidIndex(SurfaceIndex));
-		check(Components[ComponentIndex].LODs[LODIndex].Surfaces[SurfaceIndex].Scalars.IsValidIndex(ScalarIndex));
-
         return Components[ComponentIndex].LODs[LODIndex].Surfaces[SurfaceIndex].Scalars[ScalarIndex].Name;
 	}
 
 
 	int32 FInstance::GetScalarLayerIndex(int32 ComponentIndex, int32 LODIndex, int32 SurfaceIndex, int32 ScalarIndex) const
 	{
-		check(Components.IsValidIndex(ComponentIndex));
-		check(Components[ComponentIndex].LODs.IsValidIndex(LODIndex));
-        check(Components[ComponentIndex].LODs[LODIndex].Surfaces.IsValidIndex(SurfaceIndex));
-		check(Components[ComponentIndex].LODs[LODIndex].Surfaces[SurfaceIndex].Scalars.IsValidIndex(ScalarIndex));
-
         return Components[ComponentIndex].LODs[LODIndex].Surfaces[SurfaceIndex].Scalars[ScalarIndex].LayerIndex;
 	}
 
 	
     FString FInstance::GetString(int32 ComponentIndex, int32 LODIndex, int32 SurfaceIndex, int32 StringIndex ) const
     {
-        check(Components.IsValidIndex(ComponentIndex));
-        check(Components[ComponentIndex].LODs.IsValidIndex(LODIndex));
-        check(Components[ComponentIndex].LODs[LODIndex].Surfaces.IsValidIndex(SurfaceIndex));
-		check(Components[ComponentIndex].LODs[LODIndex].Surfaces[SurfaceIndex].Strings.IsValidIndex(StringIndex));
-
 		bool bValid = Components[ComponentIndex].LODs[LODIndex].Surfaces[SurfaceIndex].Strings.IsValidIndex(StringIndex);
 		if (bValid)
         {
@@ -297,10 +228,6 @@ namespace UE::Mutable::Private
     
 	FName FInstance::GetStringName(int32 ComponentIndex, int32 LODIndex, int32 SurfaceIndex, int32 StringIndex) const
     {
-        check(Components.IsValidIndex(ComponentIndex));
-        check(Components[ComponentIndex].LODs.IsValidIndex(LODIndex));
-        check(Components[ComponentIndex].LODs[LODIndex].Surfaces.IsValidIndex(SurfaceIndex));
-
 		bool bValid = Components[ComponentIndex].LODs[LODIndex].Surfaces[SurfaceIndex].Strings.IsValidIndex(StringIndex);
 		if (bValid)
 		{
@@ -313,24 +240,36 @@ namespace UE::Mutable::Private
 	
 	FSkeletalMeshId FInstance::GetSkeletalMeshId(int32 ComponentIndex) const
 	{
-    	check(Components.IsValidIndex(ComponentIndex));
     	return Components[ComponentIndex].SkeletalMeshId;
 	}
 
 	
 	FMaterialId FInstance::GetOverlayMaterialId(int32 ComponentIndex) const
 	{
-		check(Components.IsValidIndex(ComponentIndex));
 		return Components[ComponentIndex].OverlayMaterialId;
 	}
 
+
+	int32 FInstance::GetOverrideMaterialCount(int32 ComponentIndex) const
+	{
+    	return Components[ComponentIndex].OverrideMaterials.Num();
+	}
+
+
+	FName FInstance::GetOverrideMaterialSlotSlotName(int32 ComponentIndex, int32 MaterialIndex) const
+	{
+    	return Components[ComponentIndex].OverrideMaterials[MaterialIndex].SlotName;
+	}
 	
-	FMaterialId FInstance::GetMaterialId(int32 ComponentIndex, int32 LODIndex, int32 SurfaceIndex) const
+	
+	FMaterialId FInstance::GetOverrideMaterialId(int32 ComponentIndex, int32 MaterialIndex) const
 	{
-		check(Components.IsValidIndex(ComponentIndex));
-		check(Components[ComponentIndex].LODs.IsValidIndex(LODIndex));
-		check(Components[ComponentIndex].LODs[LODIndex].Surfaces.IsValidIndex(SurfaceIndex));
+    	return Components[ComponentIndex].OverrideMaterials[MaterialIndex].MaterialId;
+	}
 
+
+	FMaterialId FInstance::GetMaterialId(int32 ComponentIndex, int32 LODIndex, int32 SurfaceIndex) const
+	{
 		return Components[ComponentIndex].LODs[LODIndex].Surfaces[SurfaceIndex].MaterialId;
 	}
 
@@ -558,7 +497,7 @@ namespace UE::Mutable::Private
     		AddComponent();
     	}
 
-    	Components[ComponentIndex].OverrideMaterialIds.Add(SlotName, MaterialId);
+    	Components[ComponentIndex].OverrideMaterials.Add({ SlotName, MaterialId });
 	}
 
 
