[Mutable] Removed Low Priority Textures.

#Synced-CL 48348317

diff --git a/Engine/Plugins/Mutable/Source/CustomizableObject/Internal/MuCO/CustomizableObjectInstancePrivate.h b/Engine/Plugins/Mutable/Source/CustomizableObject/Internal/MuCO/CustomizableObjectInstancePrivate.h
--- a/Engine/Plugins/Mutable/Source/CustomizableObject/Internal/MuCO/CustomizableObjectInstancePrivate.h
+++ b/Engine/Plugins/Mutable/Source/CustomizableObject/Internal/MuCO/CustomizableObjectInstancePrivate.h
@@ -69,9 +69,6 @@ enum ECOInstanceFlags
 	// Types of updates
 	PendingLODsUpdate				= 1 << 9,	// Used to queue an update due to a change in LODs required by the instance
 	PendingLODsDowngrade			= 1 << 10,	// Used to queue a downgrade update to reduce the number of LODs. LOD update goes from a high res level to a low res one, ex: 0 to 1 or 1 to 2
-	
-	// Generation
-	ForceGenerateMipTail			= 1 << 13,	// If set, SkipGenerateResidentMips will be ignored and the mip tail will be generated
 };
 
 ENUM_CLASS_FLAGS(ECOInstanceFlags);
diff --git a/Engine/Plugins/Mutable/Source/CustomizableObject/Internal/MuCO/CustomizableObjectPrivate.h b/Engine/Plugins/Mutable/Source/CustomizableObject/Internal/MuCO/CustomizableObjectPrivate.h
--- a/Engine/Plugins/Mutable/Source/CustomizableObject/Internal/MuCO/CustomizableObjectPrivate.h
+++ b/Engine/Plugins/Mutable/Source/CustomizableObject/Internal/MuCO/CustomizableObjectPrivate.h
@@ -1900,10 +1900,7 @@ public:
 	CUSTOMIZABLEOBJECT_API bool TryUpdateIsChildObject();
 
 	CUSTOMIZABLEOBJECT_API void SetIsChildObject(bool bIsChildObject);
-
-	/** Return the names used by mutable to identify which UE::Mutable::Private::FImage should be considered of LowPriority. */
-	CUSTOMIZABLEOBJECT_API void GetLowPriorityTextureNames(TArray<FString>& OutTextureNames);
-
+	
 	/** Return the MinLOD index to generate based on the active LODSettings (PerPlatformMinLOD or PerQualityLevelMinLOD) */
 	CUSTOMIZABLEOBJECT_API uint8 GetMinLODIndex(const FName& ComponentName) const;
 
diff --git a/Engine/Plugins/Mutable/Source/CustomizableObject/Internal/MuCO/CustomizableObjectSystemPrivate.h b/Engine/Plugins/Mutable/Source/CustomizableObject/Internal/MuCO/CustomizableObjectSystemPrivate.h
--- a/Engine/Plugins/Mutable/Source/CustomizableObject/Internal/MuCO/CustomizableObjectSystemPrivate.h
+++ b/Engine/Plugins/Mutable/Source/CustomizableObject/Internal/MuCO/CustomizableObjectSystemPrivate.h
@@ -575,8 +575,6 @@ public:
 
 	const FInstanceUpdateData::FComponent* GetComponentUpdateData(FCustomizableObjectInstanceComponentIndex InstanceComponentIndex);
 
-	TArray<FString> LowPriorityTextures;
-
 	/** This option comes from the operation request */
 	bool bNeverStream = false;
 	
@@ -607,9 +605,7 @@ public:
 	
 	/** User config. */
 	bool bIsOnlyGenerateRequestedLODsEnabled = false;
-
-	bool bEnableSkipGenerateResidentMips = false;
-
+	
 	bool bIsUpdateAborted = false;
 
 	/** This option comes from the operation request. It is used to reduce the number of mipmaps that mutable must generate for images.  */
diff --git a/Engine/Plugins/Mutable/Source/CustomizableObject/Private/MuCO/CustomizableObject.cpp b/Engine/Plugins/Mutable/Source/CustomizableObject/Private/MuCO/CustomizableObject.cpp
--- a/Engine/Plugins/Mutable/Source/CustomizableObject/Private/MuCO/CustomizableObject.cpp
+++ b/Engine/Plugins/Mutable/Source/CustomizableObject/Private/MuCO/CustomizableObject.cpp
@@ -1964,25 +1964,6 @@ void UCustomizableObjectPrivate::ApplyStateForcedValuesToParameters(FCustomizabl
 }
 
 
-void UCustomizableObjectPrivate::GetLowPriorityTextureNames(TArray<FString>& OutTextureNames)
-{
-	OutTextureNames.Reset(GetPublic()->LowPriorityTextures.Num());
-
-	const UModelResources* LocalModelResources = GetModelResources();
-	if (LocalModelResources && !GetPublic()->LowPriorityTextures.IsEmpty())
-	{
-		const int32 ImageCount = LocalModelResources->ImageProperties.Num();
-		for (int32 ImageIndex = 0; ImageIndex < ImageCount; ++ImageIndex)
-		{
-			if (GetPublic()->LowPriorityTextures.Find(FName(LocalModelResources->ImageProperties[ImageIndex].TextureParameterName)) != INDEX_NONE)
-			{
-				OutTextureNames.Add(FString::FromInt(ImageIndex));
-			}
-		}
-	}
-}
-
-
 uint8 UCustomizableObjectPrivate::GetMinLODIndex(const FName& ComponentName) const
 {
 	uint8 MinLODIdx = 0;
diff --git a/Engine/Plugins/Mutable/Source/CustomizableObject/Private/MuCO/CustomizableObjectInstance.cpp b/Engine/Plugins/Mutable/Source/CustomizableObject/Private/MuCO/CustomizableObjectInstance.cpp
--- a/Engine/Plugins/Mutable/Source/CustomizableObject/Private/MuCO/CustomizableObjectInstance.cpp
+++ b/Engine/Plugins/Mutable/Source/CustomizableObject/Private/MuCO/CustomizableObjectInstance.cpp
@@ -6299,12 +6299,6 @@ void UCustomizableObjectInstance::SetReuseInstanceTextures(bool bTextureReuseEna
 }
 
 
-void UCustomizableObjectInstance::SetForceGenerateResidentMips(bool bForceGenerateResidentMips)
-{
-	bForceGenerateResidentMips ? GetPrivate()->SetCOInstanceFlags(ForceGenerateMipTail) : GetPrivate()->ClearCOInstanceFlags(ForceGenerateMipTail);
-}
-
-
 void UCustomizableObjectInstance::SetIsBeingUsedByComponentInPlay(bool bIsUsedByComponentInPlay)
 {
 	bIsUsedByComponentInPlay ? GetPrivate()->SetCOInstanceFlags(UsedByComponentInPlay) : GetPrivate()->ClearCOInstanceFlags(UsedByComponentInPlay);
diff --git a/Engine/Plugins/Mutable/Source/CustomizableObject/Private/MuCO/CustomizableObjectSystem.cpp b/Engine/Plugins/Mutable/Source/CustomizableObject/Private/MuCO/CustomizableObjectSystem.cpp
--- a/Engine/Plugins/Mutable/Source/CustomizableObject/Private/MuCO/CustomizableObjectSystem.cpp
+++ b/Engine/Plugins/Mutable/Source/CustomizableObject/Private/MuCO/CustomizableObjectSystem.cpp
@@ -1016,13 +1016,6 @@ static FAutoConsoleVariableRef CVarEnableOnlyGenerateRequestedLODs(
 	TEXT("If 1 or greater, Only the RequestedLOD will be generated. If 0, all LODs will be build."),
 	ECVF_Default);
 
-int32 UCustomizableObjectSystemPrivate::EnableSkipGenerateResidentMips = 1;
-
-static FAutoConsoleVariableRef CVarSkipGenerateResidentMips(
-	TEXT("mutable.EnableSkipGenerateResidentMips"), UCustomizableObjectSystemPrivate::EnableSkipGenerateResidentMips,
-	TEXT("If 1 or greater, resident mip generation will be optional. If 0, resident mips will be always generated"),
-	ECVF_Default);
-
 int32 UCustomizableObjectSystemPrivate::MaxTextureSizeToGenerate = 0;
 
 FAutoConsoleVariableRef CVarMaxTextureSizeToGenerate(
@@ -2720,18 +2713,9 @@ namespace impl
 				const int32 MipSizeX = FMath::Max(Image->FullImageSizeX >> MipsToSkip, 1);
 				const int32 MipSizeY = FMath::Max(Image->FullImageSizeY >> MipsToSkip, 1);
 
-				UE::Tasks::TTask<TSharedPtr<const UE::Mutable::Private::FImage>> GetImageTask;
-				if (MipsToSkip > 0 && OperationData->bEnableSkipGenerateResidentMips != 0 && OperationData->LowPriorityTextures.Find(Image->Name.ToString()) != INDEX_NONE)
-				{
-					TSharedPtr<const UE::Mutable::Private::FImage> DummyImage = MakeShared<UE::Mutable::Private::FImage>(MipSizeX, MipSizeY, FullLODCount - MipsToSkip, Image->Format, UE::Mutable::Private::EInitializationType::Black);
-					GetImageTask = UE::Tasks::MakeCompletedTask<TSharedPtr<const UE::Mutable::Private::FImage>>(DummyImage);
-				}
-				else
-				{
-					GetImageTask = OperationData->MutableSystem->GetImage(LastTask, OperationData->InstanceID, Image->ImageID, Image->BaseMip + MipsToSkip);
-					UE::Tasks::AddNested(GetImageTask);
-					LastTask = GetImageTask;
-				}
+				UE::Tasks::TTask<TSharedPtr<const UE::Mutable::Private::FImage>> GetImageTask = OperationData->MutableSystem->GetImage(LastTask, OperationData->InstanceID, Image->ImageID, Image->BaseMip + MipsToSkip);
+				UE::Tasks::AddNested(GetImageTask);
+				LastTask = GetImageTask;
 
 				UE::Tasks::AddNested(UE::Tasks::Launch(TEXT("GetImageResult"), [Image, GetImageTask, MipSizeX, MipSizeY, FullLODCount, MipsToSkip]() mutable
 				{
@@ -3285,7 +3269,6 @@ namespace impl
 
 		Operation->bNeverStream = false;
 		Operation->MipsToSkip = 0;
-		Operation->bEnableSkipGenerateResidentMips = static_cast<bool>(SystemPrivate->EnableSkipGenerateResidentMips);
 
 		SystemPrivate->GetMipStreamingConfig(*Instance, Operation->bProgressiveMipStreamingEnabled, Operation->bNeverStream, Operation->MipsToSkip);
 		
@@ -3371,11 +3354,6 @@ namespace impl
 		InstancePrivate->GeneratedMaterials.Empty();
 		InstancePrivate->GeneratedTextures.Empty();
 		
-		if (!InstancePrivate->HasCOInstanceFlags(ForceGenerateMipTail))
-		{
-			ObjectPrivate->GetLowPriorityTextureNames(Operation->LowPriorityTextures);
-		}
-
 		bool bRequestAllLODs = !Operation->bIsOnlyGenerateRequestedLODsEnabled ||
 			!SystemPrivate->CurrentInstanceLODManagement->IsOnlyGenerateRequestedLODLevelsEnabled();
 
diff --git a/Engine/Plugins/Mutable/Source/CustomizableObject/Public/MuCO/CustomizableObject.h b/Engine/Plugins/Mutable/Source/CustomizableObject/Public/MuCO/CustomizableObject.h
--- a/Engine/Plugins/Mutable/Source/CustomizableObject/Public/MuCO/CustomizableObject.h
+++ b/Engine/Plugins/Mutable/Source/CustomizableObject/Public/MuCO/CustomizableObject.h
@@ -508,11 +508,6 @@ public:
 #endif
 	
 private:
-	/** Textures marked as low priority will generate defaulted resident mips (if texture streaming is enabled).
-	  * Generating defaulted resident mips greatly reduce initial generation times. */
-	UPROPERTY(EditAnywhere, Category = CustomizableObject)
-	TArray<FName> LowPriorityTextures;
-
 	// Customizable Object Population data start ------------------------------------------------------
 	/** Array to store the selected Population Class tags for this Customizable Object */
 	UPROPERTY()
diff --git a/Engine/Plugins/Mutable/Source/CustomizableObject/Public/MuCO/CustomizableObjectInstance.h b/Engine/Plugins/Mutable/Source/CustomizableObject/Public/MuCO/CustomizableObjectInstance.h
--- a/Engine/Plugins/Mutable/Source/CustomizableObject/Public/MuCO/CustomizableObjectInstance.h
+++ b/Engine/Plugins/Mutable/Source/CustomizableObject/Public/MuCO/CustomizableObjectInstance.h
@@ -795,9 +795,6 @@ public:
 	  * WARNING! If texture reuse is enabled, do NOT keep external references to the textures of the instance. The instance owns the textures. */
 	CUSTOMIZABLEOBJECT_API void SetReuseInstanceTextures(bool bTextureReuseEnabled);
 	
-	/** If enabled, low-priority textures will generate resident mipmaps too. */
-	CUSTOMIZABLEOBJECT_API void SetForceGenerateResidentMips(bool bForceGenerateResidentMips);
-	
 	// The following methods should only be used in an LOD management class
 	CUSTOMIZABLEOBJECT_API void SetIsBeingUsedByComponentInPlay(bool bIsUsedByComponent);
 	CUSTOMIZABLEOBJECT_API bool GetIsBeingUsedByComponentInPlay() const;
