[mutable] Added missing spacer (Mutable Debugger)

#rnx
#Synced-CL 48385835

diff --git a/Engine/Plugins/Mutable/Source/CustomizableObjectEditor/Private/MuCOE/SMutableConstantsWidget.cpp b/Engine/Plugins/Mutable/Source/CustomizableObjectEditor/Private/MuCOE/SMutableConstantsWidget.cpp
--- a/Engine/Plugins/Mutable/Source/CustomizableObjectEditor/Private/MuCOE/SMutableConstantsWidget.cpp
+++ b/Engine/Plugins/Mutable/Source/CustomizableObjectEditor/Private/MuCOE/SMutableConstantsWidget.cpp
@@ -779,7 +779,11 @@ void SMutableConstantsWidget::Construct(const FArguments& InArgs,const UE::Mutab
 			]
 		]
 
-
+		+ SScrollBox::Slot()
+		[
+			SNew(SSeparator)
+		]
+		
 		// Physics Constants
 		+ SScrollBox::Slot()
 		.Padding(0,InBetweenListsVerticalPadding)
