[Mutable] Fixed the searching method which finds the base mesh of a node when it goes through one or more mutable macros.

#rnx

#Synced-CL 48355645

diff --git a/Engine/Plugins/Mutable/Source/CustomizableObjectEditor/Private/MuCOE/GenerateMutableSource/GenerateMutableSource.cpp b/Engine/Plugins/Mutable/Source/CustomizableObjectEditor/Private/MuCOE/GenerateMutableSource/GenerateMutableSource.cpp
--- a/Engine/Plugins/Mutable/Source/CustomizableObjectEditor/Private/MuCOE/GenerateMutableSource/GenerateMutableSource.cpp
+++ b/Engine/Plugins/Mutable/Source/CustomizableObjectEditor/Private/MuCOE/GenerateMutableSource/GenerateMutableSource.cpp
@@ -681,7 +681,7 @@ UTexture2D* FindReferenceImage(const UEdGraphPin* Pin, FMutableGraphGenerationCo
 
 	else if (const UCustomizableObjectNodeMacroInstance* TypedNodeMacro = Cast<UCustomizableObjectNodeMacroInstance>(Node))
 	{
-		if (const UEdGraphPin* OutputPin = TypedNodeMacro->GetMacroIOPin(ECOMacroIOType::COMVT_Output, Pin->PinName))
+		if (const UEdGraphPin* OutputPin = TypedNodeMacro->GetMacroTunnelPin(ECOMacroIOType::COMVT_Output, Pin->PinName))
 		{
 			if (const UEdGraphPin* ConnectedPin = FollowInputPin(*OutputPin))
 			{
diff --git a/Engine/Plugins/Mutable/Source/CustomizableObjectEditor/Private/MuCOE/GenerateMutableSource/GenerateMutableSourceImage.cpp b/Engine/Plugins/Mutable/Source/CustomizableObjectEditor/Private/MuCOE/GenerateMutableSource/GenerateMutableSourceImage.cpp
--- a/Engine/Plugins/Mutable/Source/CustomizableObjectEditor/Private/MuCOE/GenerateMutableSource/GenerateMutableSourceImage.cpp
+++ b/Engine/Plugins/Mutable/Source/CustomizableObjectEditor/Private/MuCOE/GenerateMutableSource/GenerateMutableSourceImage.cpp
@@ -1179,7 +1179,7 @@ UE::Mutable::Private::Ptr<UE::Mutable::Private::NodeImage> GenerateMutableSource
 		// Can't use the GenerateMutableSourceMacro function here because GenerateMutableSourceImage needs the extra parameter ReferenceTextureSize 
 		bCacheNode = false;
 
-		if (const UEdGraphPin* OutputPin = TypedNodeMacro->GetMacroIOPin(ECOMacroIOType::COMVT_Output, Pin->PinName))
+		if (const UEdGraphPin* OutputPin = TypedNodeMacro->GetMacroTunnelPin(ECOMacroIOType::COMVT_Output, Pin->PinName))
 		{
 			if (const UEdGraphPin* FollowPin = FollowInputPin(*OutputPin))
 			{
diff --git a/Engine/Plugins/Mutable/Source/CustomizableObjectEditor/Private/MuCOE/GenerateMutableSource/GenerateMutableSourceMacro.h b/Engine/Plugins/Mutable/Source/CustomizableObjectEditor/Private/MuCOE/GenerateMutableSource/GenerateMutableSourceMacro.h
--- a/Engine/Plugins/Mutable/Source/CustomizableObjectEditor/Private/MuCOE/GenerateMutableSource/GenerateMutableSourceMacro.h
+++ b/Engine/Plugins/Mutable/Source/CustomizableObjectEditor/Private/MuCOE/GenerateMutableSource/GenerateMutableSourceMacro.h
@@ -18,7 +18,7 @@ UE::Mutable::Private::Ptr<T> GenerateMutableSourceMacro(const UEdGraphPin& Pin,
 
 	if (const UCustomizableObjectNodeMacroInstance* TypedNodeMacro = Cast<UCustomizableObjectNodeMacroInstance>(Node))
 	{
-		if (const UEdGraphPin* OutputPin = TypedNodeMacro->GetMacroIOPin(ECOMacroIOType::COMVT_Output, Pin.PinName))
+		if (const UEdGraphPin* OutputPin = TypedNodeMacro->GetMacroTunnelPin(ECOMacroIOType::COMVT_Output, Pin.PinName))
 		{
 			if (const UEdGraphPin* FollowPin = FollowInputPin(*OutputPin))
 			{
diff --git a/Engine/Plugins/Mutable/Source/CustomizableObjectEditor/Private/MuCOE/GenerateMutableSource/GenerateMutableSourceMesh.cpp b/Engine/Plugins/Mutable/Source/CustomizableObjectEditor/Private/MuCOE/GenerateMutableSource/GenerateMutableSourceMesh.cpp
--- a/Engine/Plugins/Mutable/Source/CustomizableObjectEditor/Private/MuCOE/GenerateMutableSource/GenerateMutableSourceMesh.cpp
+++ b/Engine/Plugins/Mutable/Source/CustomizableObjectEditor/Private/MuCOE/GenerateMutableSource/GenerateMutableSourceMesh.cpp
@@ -3639,7 +3639,7 @@ UE::Mutable::Private::NodeMeshPtr GenerateMorphMesh(const UEdGraphPin* Pin,
 				Result->bReshapePhysicsVolumes = TypedMorphNode->bReshapePhysicsVolumes;
 				{
 					//const UEdGraphPin* ConnectedPin = FollowInputPin(*TypedMorphNode->MeshPin());
-					const UEdGraphPin* SourceMeshPin = Pin ? FindMeshBaseSource(*Pin, false) : nullptr;
+					const UEdGraphPin* SourceMeshPin = Pin ? FindMeshBaseSource(*Pin, false, GenerationContext.MacroNodesStack) : nullptr;
 					const UEdGraphNode* SkeletalMeshNode = SourceMeshPin ? SourceMeshPin->GetOwningNode() : nullptr;
 
 					TArray<TTuple<USkeletalMesh*, TSoftClassPtr<UAnimInstance>>> SkeletalMeshesToDeform = 
@@ -3724,7 +3724,7 @@ TArray<FMorphNodeData> GenerateMeshMorphStackDefinition(const UEdGraphPin* Pin,
 
 	else if (const UCustomizableObjectNodeMacroInstance* TypedNodeMacro = Cast<UCustomizableObjectNodeMacroInstance>(Node))
 	{
-		if (const UEdGraphPin* OutputPin = TypedNodeMacro->GetMacroIOPin(ECOMacroIOType::COMVT_Output, Pin->PinName))
+		if (const UEdGraphPin* OutputPin = TypedNodeMacro->GetMacroTunnelPin(ECOMacroIOType::COMVT_Output, Pin->PinName))
 		{
 			if (const UEdGraphPin* FollowPin = FollowInputPin(*OutputPin))
 			{
@@ -4272,7 +4272,7 @@ UE::Mutable::Private::Ptr<UE::Mutable::Private::NodeMesh> GenerateMutableSourceM
 			MeshNode->ColorAChannelUsage = ConvertColorUsage(ChannelUsages[3]);
 				
 			const UEdGraphPin* ConnectedPin = FollowInputPin(*TypedNodeReshape->BaseMeshPin());
-			const UEdGraphPin* SourceMeshPin = ConnectedPin ? FindMeshBaseSource(*ConnectedPin, false) : nullptr;
+			const UEdGraphPin* SourceMeshPin = ConnectedPin ? FindMeshBaseSource(*ConnectedPin, false, GenerationContext.MacroNodesStack) : nullptr;
 			const UEdGraphNode* SkeletalMeshNode = SourceMeshPin ? SourceMeshPin->GetOwningNode() : nullptr;
 
 			TArray<TTuple<USkeletalMesh*, TSoftClassPtr<UAnimInstance>>> SkeletalMeshesToDeform = 
@@ -4689,7 +4689,7 @@ UE::Mutable::Private::Ptr<UE::Mutable::Private::NodeMesh> GenerateMutableSourceM
 		// Can't use the GenerateMutableSourceMacro function here because GenerateMutableSourceMesh needs some extra parameters
 		bCacheNode = false;
 
-		if (const UEdGraphPin* OutputPin = TypedNodeMacro->GetMacroIOPin(ECOMacroIOType::COMVT_Output, Pin->PinName))
+		if (const UEdGraphPin* OutputPin = TypedNodeMacro->GetMacroTunnelPin(ECOMacroIOType::COMVT_Output, Pin->PinName))
 		{
 			if (const UEdGraphPin* FollowPin = FollowInputPin(*OutputPin))
 			{
diff --git a/Engine/Plugins/Mutable/Source/CustomizableObjectEditor/Private/MuCOE/GenerateMutableSource/GenerateMutableSourceModifier.cpp b/Engine/Plugins/Mutable/Source/CustomizableObjectEditor/Private/MuCOE/GenerateMutableSource/GenerateMutableSourceModifier.cpp
--- a/Engine/Plugins/Mutable/Source/CustomizableObjectEditor/Private/MuCOE/GenerateMutableSource/GenerateMutableSourceModifier.cpp
+++ b/Engine/Plugins/Mutable/Source/CustomizableObjectEditor/Private/MuCOE/GenerateMutableSource/GenerateMutableSourceModifier.cpp
@@ -288,7 +288,7 @@ UE::Mutable::Private::Ptr<UE::Mutable::Private::NodeModifier> GenerateMutableSou
 				if(ConnectedPin)
 				{
 					//NOTE: This is the same is done in GenerateMutableSourceSurface. 
-					if (const UEdGraphPin* SkeletalMeshPin = FindMeshBaseSource(*ConnectedPin, false, &GenerationContext.MacroNodesStack))
+					if (const UEdGraphPin* SkeletalMeshPin = FindMeshBaseSource(*ConnectedPin, false, GenerationContext.MacroNodesStack))
 					{
 						int32 MetadataLODIndex, MetadataSectionIndex;
 						MetadataLODIndex = MetadataSectionIndex = INDEX_NONE;
diff --git a/Engine/Plugins/Mutable/Source/CustomizableObjectEditor/Private/MuCOE/GenerateMutableSource/GenerateMutableSourceSurface.cpp b/Engine/Plugins/Mutable/Source/CustomizableObjectEditor/Private/MuCOE/GenerateMutableSource/GenerateMutableSourceSurface.cpp
--- a/Engine/Plugins/Mutable/Source/CustomizableObjectEditor/Private/MuCOE/GenerateMutableSource/GenerateMutableSourceSurface.cpp
+++ b/Engine/Plugins/Mutable/Source/CustomizableObjectEditor/Private/MuCOE/GenerateMutableSource/GenerateMutableSourceSurface.cpp
@@ -111,7 +111,7 @@ UE::Mutable::Private::Ptr<UE::Mutable::Private::NodeSurface> GenerateMutableSour
 		// TODO: delete the if clause and the warning when static meshes are operational again
 		if (ConnectedMaterialPin)
 		{
-			if (const UEdGraphPin* StaticMeshPin = FindMeshBaseSource(*ConnectedMaterialPin, true, &GenerationContext.MacroNodesStack))
+			if (const UEdGraphPin* StaticMeshPin = FindMeshBaseSource(*ConnectedMaterialPin, true, GenerationContext.MacroNodesStack))
 			{
 				const UCustomizableObjectNode* StaticMeshNode = CastChecked<UCustomizableObjectNode>(StaticMeshPin->GetOwningNode());
 				GenerationContext.Log(LOCTEXT("UnsupportedStaticMeshes", "Static meshes are currently not supported as material meshes"), StaticMeshNode);
@@ -146,7 +146,7 @@ UE::Mutable::Private::Ptr<UE::Mutable::Private::NodeSurface> GenerateMutableSour
 		if (ConnectedMaterialPin)
 		{
 			//NOTE: This is the same is done in GenerateMutableSourceSurface. 
-			if (const UEdGraphPin* SkeletalMeshPin = FindMeshBaseSource(*ConnectedMaterialPin, false, &GenerationContext.MacroNodesStack))
+			if (const UEdGraphPin* SkeletalMeshPin = FindMeshBaseSource(*ConnectedMaterialPin, false, GenerationContext.MacroNodesStack))
 			{
 				int32 MetadataLODIndex, MetadataSectionIndex;
 				MetadataLODIndex = MetadataSectionIndex = INDEX_NONE;
diff --git a/Engine/Plugins/Mutable/Source/CustomizableObjectEditor/Private/MuCOE/GraphTraversal.cpp b/Engine/Plugins/Mutable/Source/CustomizableObjectEditor/Private/MuCOE/GraphTraversal.cpp
--- a/Engine/Plugins/Mutable/Source/CustomizableObjectEditor/Private/MuCOE/GraphTraversal.cpp
+++ b/Engine/Plugins/Mutable/Source/CustomizableObjectEditor/Private/MuCOE/GraphTraversal.cpp
@@ -450,7 +450,7 @@ void RecursiveVisitNodes(UCustomizableObjectNode& CurrentNode, const TFunction<v
 			else if (UCustomizableObjectNodeMacroInstance* MacroInstanceNode = Cast<UCustomizableObjectNodeMacroInstance>(ConnectedNode))
 			{
 				// First visit the graph of the Macro
-				if (const UEdGraphPin* OutputPin = MacroInstanceNode->GetMacroIOPin(ECOMacroIOType::COMVT_Output, ConnectedPin->PinName))
+				if (const UEdGraphPin* OutputPin = MacroInstanceNode->GetMacroTunnelPin(ECOMacroIOType::COMVT_Output, ConnectedPin->PinName))
 				{
 					if (UCustomizableObjectNode* OutputNode = Cast<UCustomizableObjectNode>(OutputPin->GetOwningNode()))
 					{
@@ -501,7 +501,7 @@ const UEdGraphPin* GraphTraversal::FindIOPinSourceThroughMacroContext(const UEdG
 
 	if (const UCustomizableObjectNodeMacroInstance* NodeMacro = Cast<UCustomizableObjectNodeMacroInstance>(Node))
 	{
-		if (const UEdGraphPin* OutputPin = NodeMacro->GetMacroIOPin(Pin.Direction == EGPD_Output ? ECOMacroIOType::COMVT_Output : ECOMacroIOType::COMVT_Input, Pin.PinName))
+		if (const UEdGraphPin* OutputPin = NodeMacro->GetMacroTunnelPin(Pin.Direction == EGPD_Output ? ECOMacroIOType::COMVT_Output : ECOMacroIOType::COMVT_Input, Pin.PinName))
 		{
 			if (const UEdGraphPin* ConnectedPin = FollowInputPin(*OutputPin))
 			{
@@ -583,7 +583,7 @@ UCustomizableObjectNodeObject* GraphTraversal::GetFullGraphRootNode(const UCusto
 }
 
 
-const UEdGraphPin* FindMeshBaseSource(const UEdGraphPin& Pin, const bool bOnlyLookForStaticMesh, TArray<const UCustomizableObjectNodeMacroInstance*>* MacroContext)
+const UEdGraphPin* FindMeshBaseSource(const UEdGraphPin& Pin, const bool bOnlyLookForStaticMesh, TArray<const UCustomizableObjectNodeMacroInstance*> MacroContext)
 {
 	check(Pin.Direction == EGPD_Output);
 	check(Pin.PinType.PinCategory == UEdGraphSchema_CustomizableObject::PC_Mesh
@@ -732,21 +732,26 @@ const UEdGraphPin* FindMeshBaseSource(const UEdGraphPin& Pin, const bool bOnlyLo
 
 	else if (const UCustomizableObjectNodeMacroInstance* NodeMacro = Cast<UCustomizableObjectNodeMacroInstance>(Node))
 	{
-		const UEdGraphPin* ConnectedPin = GraphTraversal::FindIOPinSourceThroughMacroContext(Pin, MacroContext);
-		
-		if (ConnectedPin)
+		const UEdGraphPin* OutputPin = NodeMacro->GetMacroTunnelPin(ECOMacroIOType::COMVT_Output, Pin.PinName);
+
+		if (const UEdGraphPin* ConnectedPin = FollowInputPin(*OutputPin))
 		{
+			MacroContext.Push(NodeMacro);
 			return FindMeshBaseSource(*ConnectedPin, bOnlyLookForStaticMesh, MacroContext);
 		}
 	}
 
 	else if (const UCustomizableObjectNodeTunnel* NodeTunnel = Cast<UCustomizableObjectNodeTunnel>(Node))
 	{
-		const UEdGraphPin* ConnectedPin = GraphTraversal::FindIOPinSourceThroughMacroContext(Pin, MacroContext);
+		const UCustomizableObjectNodeMacroInstance* MacroInstanceNode = MacroContext.Pop();
+		check(MacroInstanceNode);
 
-		if (ConnectedPin)
+		if (const UEdGraphPin* InputPin = MacroInstanceNode->FindPin(Pin.PinName, EEdGraphPinDirection::EGPD_Input))
 		{
-			return FindMeshBaseSource(*ConnectedPin, bOnlyLookForStaticMesh, MacroContext);
+			if (const UEdGraphPin* ConnectedPin = FollowInputPin(*InputPin))
+			{
+				return FindMeshBaseSource(*ConnectedPin, bOnlyLookForStaticMesh, MacroContext);
+			}
 		}
 	}
 
diff --git a/Engine/Plugins/Mutable/Source/CustomizableObjectEditor/Private/MuCOE/GraphTraversal.h b/Engine/Plugins/Mutable/Source/CustomizableObjectEditor/Private/MuCOE/GraphTraversal.h
--- a/Engine/Plugins/Mutable/Source/CustomizableObjectEditor/Private/MuCOE/GraphTraversal.h
+++ b/Engine/Plugins/Mutable/Source/CustomizableObjectEditor/Private/MuCOE/GraphTraversal.h
@@ -105,7 +105,7 @@ namespace GraphTraversal
 }
 
 /** Given an output pin, return the output pin where the mesh is located. */
-const UEdGraphPin* FindMeshBaseSource(const UEdGraphPin& Pin, const bool bOnlyLookForStaticMesh, TArray<const UCustomizableObjectNodeMacroInstance*>* MacroContext = nullptr);
+const UEdGraphPin* FindMeshBaseSource(const UEdGraphPin& Pin, const bool bOnlyLookForStaticMesh, TArray<const UCustomizableObjectNodeMacroInstance*> MacroContext = {});
 
 /** Return the mapping of Group Object Nodes to Child Object Nodes of the given hierarchy.
  * @param Object Child or root Object to start the search from. */
diff --git a/Engine/Plugins/Mutable/Source/CustomizableObjectEditor/Private/MuCOE/Nodes/CustomizableObjectNodeMacroInstance.cpp b/Engine/Plugins/Mutable/Source/CustomizableObjectEditor/Private/MuCOE/Nodes/CustomizableObjectNodeMacroInstance.cpp
--- a/Engine/Plugins/Mutable/Source/CustomizableObjectEditor/Private/MuCOE/Nodes/CustomizableObjectNodeMacroInstance.cpp
+++ b/Engine/Plugins/Mutable/Source/CustomizableObjectEditor/Private/MuCOE/Nodes/CustomizableObjectNodeMacroInstance.cpp
@@ -163,7 +163,7 @@ bool UCustomizableObjectNodeMacroInstance::IsExperimental() const
 }
 
 
-UEdGraphPin* UCustomizableObjectNodeMacroInstance::GetMacroIOPin(ECOMacroIOType IONodeType, const FName& PinName) const
+UEdGraphPin* UCustomizableObjectNodeMacroInstance::GetMacroTunnelPin(ECOMacroIOType IONodeType, const FName& PinName) const
 {
 	UEdGraphPin* MacroIOPin = nullptr;
 
diff --git a/Engine/Plugins/Mutable/Source/CustomizableObjectEditor/Private/MuCOE/Nodes/CustomizableObjectNodeMacroInstance.h b/Engine/Plugins/Mutable/Source/CustomizableObjectEditor/Private/MuCOE/Nodes/CustomizableObjectNodeMacroInstance.h
--- a/Engine/Plugins/Mutable/Source/CustomizableObjectEditor/Private/MuCOE/Nodes/CustomizableObjectNodeMacroInstance.h
+++ b/Engine/Plugins/Mutable/Source/CustomizableObjectEditor/Private/MuCOE/Nodes/CustomizableObjectNodeMacroInstance.h
@@ -73,8 +73,8 @@ public:
 	UE_API virtual bool IsExperimental() const override;
 
 	// Node MacroInstance Interface
-	/** Returns the pin of the Macro's input/output node with the same name */
-	UE_API UEdGraphPin* GetMacroIOPin(ECOMacroIOType IONodeType, const FName& PinName) const;
+	/** Given a Macro pin, this function returns its I/O tunnel pin. */
+	UE_API UEdGraphPin* GetMacroTunnelPin(ECOMacroIOType IONodeType, const FName& PinName) const;
 };
 
 #undef UE_API
