[Mutable] Use the reference texture to define the expected EImageFormat on Texture Parameters.

#rnx
#Synced-CL 48501169

diff --git a/Engine/Plugins/Mutable/Source/CustomizableObject/Internal/MuCO/CustomizableObjectPrivate.h b/Engine/Plugins/Mutable/Source/CustomizableObject/Internal/MuCO/CustomizableObjectPrivate.h
--- a/Engine/Plugins/Mutable/Source/CustomizableObject/Internal/MuCO/CustomizableObjectPrivate.h
+++ b/Engine/Plugins/Mutable/Source/CustomizableObject/Internal/MuCO/CustomizableObjectPrivate.h
@@ -1761,6 +1761,8 @@ enum class ECustomizableObjectVersions : int32
 	
 	AddOverrideMaterials,
 
+	TexutreParametersFormatFromReferenceTexture,
+
 	// -----<new versions can be added above this line>--------
 	LastCustomizableObjectVersion
 };
@@ -2051,7 +2053,7 @@ public:
 
 #endif
 	
-	static constexpr int32 DerivedDataVersion = 0x059cf92d;
+	static constexpr int32 DerivedDataVersion = 0x5254e536;
 };
 
 #if WITH_EDITOR
diff --git a/Engine/Plugins/Mutable/Source/CustomizableObject/Internal/MuCO/UnrealMutableImageProvider.h b/Engine/Plugins/Mutable/Source/CustomizableObject/Internal/MuCO/UnrealMutableImageProvider.h
--- a/Engine/Plugins/Mutable/Source/CustomizableObject/Internal/MuCO/UnrealMutableImageProvider.h
+++ b/Engine/Plugins/Mutable/Source/CustomizableObject/Internal/MuCO/UnrealMutableImageProvider.h
@@ -30,7 +30,7 @@ public:
 	
 private:
 	static inline const UE::Mutable::Private::FImageDesc DUMMY_IMAGE_DESC = 
-			UE::Mutable::Private::FImageDesc {UE::Mutable::Private::FImageSize(32, 32), UE::Mutable::Private::EImageFormat::RGBA_UByte, 1};
+			UE::Mutable::Private::FImageDesc {UE::Mutable::Private::FImageSize(32, 32), UE::Mutable::Private::EImageFormat::RGBA_UByte, UE::Mutable::Private::EImageFormat::RGBA_UByte, 1};
 
 	/** This will be called if an image Id has been requested by Mutable core but it has not been provided by any provider. */
 	static TSharedPtr<UE::Mutable::Private::FImage> CreateDummy();
diff --git a/Engine/Plugins/Mutable/Source/CustomizableObject/Private/MuCO/UnrealMutableImageProvider.cpp b/Engine/Plugins/Mutable/Source/CustomizableObject/Private/MuCO/UnrealMutableImageProvider.cpp
--- a/Engine/Plugins/Mutable/Source/CustomizableObject/Private/MuCO/UnrealMutableImageProvider.cpp
+++ b/Engine/Plugins/Mutable/Source/CustomizableObject/Private/MuCO/UnrealMutableImageProvider.cpp
@@ -417,7 +417,7 @@ UE::Mutable::Private::FExtendedImageDesc FUnrealMutableResourceProvider::GetImag
 	const UE::Mutable::Private::FImageSize ImageSize = UE::Mutable::Private::FImageSize(Source.GetSizeX(), Source.GetSizeY());
 	const uint8 LODs = 1;
 
-	return UE::Mutable::Private::FExtendedImageDesc { UE::Mutable::Private::FImageDesc { ImageSize, UE::Mutable::Private::EImageFormat::None, LODs }, 0 };
+	return UE::Mutable::Private::FExtendedImageDesc { UE::Mutable::Private::FImageDesc { ImageSize, UE::Mutable::Private::EImageFormat::None, UE::Mutable::Private::EImageFormat::None, LODs }, 0 };
 #else
 	UE::Mutable::Private::FExtendedImageDesc Result;	
 
@@ -448,7 +448,7 @@ UE::Mutable::Private::FExtendedImageDesc FUnrealMutableResourceProvider::GetImag
 	}
 
 	const UE::Mutable::Private::FImageDesc ImageDesc = UE::Mutable::Private::FImageDesc 
-		{ UE::Mutable::Private::FImageSize(Texture2D->GetSizeX(), Texture2D->GetSizeY()), MutableFormat, 1 }; 
+		{ UE::Mutable::Private::FImageSize(Texture2D->GetSizeX(), Texture2D->GetSizeY()), MutableFormat, UE::Mutable::Private::EImageFormat::None, 1 };
 
 	Result = UE::Mutable::Private::FExtendedImageDesc { ImageDesc, (uint8)FirstLODAvailable }; 
 
diff --git a/Engine/Plugins/Mutable/Source/CustomizableObjectEditor/Private/MuCOE/GenerateMutableSource/GenerateMutableSource.cpp b/Engine/Plugins/Mutable/Source/CustomizableObjectEditor/Private/MuCOE/GenerateMutableSource/GenerateMutableSource.cpp
--- a/Engine/Plugins/Mutable/Source/CustomizableObjectEditor/Private/MuCOE/GenerateMutableSource/GenerateMutableSource.cpp
+++ b/Engine/Plugins/Mutable/Source/CustomizableObjectEditor/Private/MuCOE/GenerateMutableSource/GenerateMutableSource.cpp
@@ -277,7 +277,7 @@ FMutableGraphGenerationContext::FMutableGraphGenerationContext(const TSharedPtr<
 }
 
 
-FString FMutableGraphGenerationContext::GetObjectName()
+FString FMutableGraphGenerationContext::GetObjectName() const
 {
 	return CompilationContext->GetObjectName();
 }
diff --git a/Engine/Plugins/Mutable/Source/CustomizableObjectEditor/Private/MuCOE/GenerateMutableSource/GenerateMutableSource.h b/Engine/Plugins/Mutable/Source/CustomizableObjectEditor/Private/MuCOE/GenerateMutableSource/GenerateMutableSource.h
--- a/Engine/Plugins/Mutable/Source/CustomizableObjectEditor/Private/MuCOE/GenerateMutableSource/GenerateMutableSource.h
+++ b/Engine/Plugins/Mutable/Source/CustomizableObjectEditor/Private/MuCOE/GenerateMutableSource/GenerateMutableSource.h
@@ -803,7 +803,7 @@ struct FMutableGraphGenerationContext
 	TArray<TArray<UCustomizableObjectNodeGroupProjectorParameter*>> CurrentGroupProjectors;
 
 	/** Message logging */
-	FString GetObjectName();
+	FString GetObjectName() const;
 	void Log(const FText& Message, const TArray<const UObject*>& Context, const EMessageSeverity::Type MessageSeverity = EMessageSeverity::Warning, const bool bAddBaseObjectInfo = true, const ELoggerSpamBin SpamBin = ELoggerSpamBin::ShowAll) const;
 	void Log(const FText& Message, const UObject* Context = nullptr, const EMessageSeverity::Type MessageSeverity = EMessageSeverity::Warning, const bool bAddBaseObjectInfo = true, const ELoggerSpamBin SpamBin = ELoggerSpamBin::ShowAll) const;
 
diff --git a/Engine/Plugins/Mutable/Source/CustomizableObjectEditor/Private/MuCOE/GenerateMutableSource/GenerateMutableSourceImage.cpp b/Engine/Plugins/Mutable/Source/CustomizableObjectEditor/Private/MuCOE/GenerateMutableSource/GenerateMutableSourceImage.cpp
--- a/Engine/Plugins/Mutable/Source/CustomizableObjectEditor/Private/MuCOE/GenerateMutableSource/GenerateMutableSourceImage.cpp
+++ b/Engine/Plugins/Mutable/Source/CustomizableObjectEditor/Private/MuCOE/GenerateMutableSource/GenerateMutableSourceImage.cpp
@@ -66,6 +66,11 @@
 #include "MuT/NodeImageParameterConvert.h"
 #include "MuT/NodeScalarConstant.h"
 
+#include "Interfaces/ITextureFormat.h"
+#include "Interfaces/ITextureFormatModule.h"
+#include "Interfaces/ITextureFormatManagerModule.h"
+#include "MuT/UnrealPixelFormatOverride.h"
+
 #define LOCTEXT_NAMESPACE "CustomizableObjectEditor"
 
 UE::Mutable::Private::Ptr<UE::Mutable::Private::NodeImage> ResizeTextureByNumMips(const UE::Mutable::Private::Ptr<UE::Mutable::Private::NodeImage>& ImageConstant, int32 MipsToSkip)
@@ -175,23 +180,9 @@ UE::Mutable::Private::Ptr<UE::Mutable::Private::NodeImage> GenerateMutableSource
 		{
 			UE::Mutable::Private::Ptr<UE::Mutable::Private::NodeImageParameterConvert> ImagePassthroughConvertNode = new UE::Mutable::Private::NodeImageParameterConvert();
 			ImagePassthroughConvertNode->ImageParameter = ImageParameterNode;
-			
-			// Force the same format that the default texture if any.
-			UE::Mutable::Private::Ptr<UE::Mutable::Private::NodeImageFormat> FormatNode = new UE::Mutable::Private::NodeImageFormat();
-			// \TODO: Take it from default?
-			//if (TypedNodeParam->DefaultValue)
-			//{
-			//	FormatNode->SetFormat(UE::Mutable::Private::EImageFormat::RGBA_UByte);
-			//}
-			//else
-			{
-				// Force an "easy format" on the texture.
-				FormatNode->Format = UE::Mutable::Private::EImageFormat::RGBA_UByte;
-			}
-			FormatNode->Source = ImagePassthroughConvertNode;
 
 			UE::Mutable::Private::Ptr<UE::Mutable::Private::NodeImageResize> ResizeNode = new UE::Mutable::Private::NodeImageResize();
-			ResizeNode->Base = FormatNode;
+			ResizeNode->Base = ImagePassthroughConvertNode;
 			ResizeNode->bRelative = false;
 
 			FUintVector2 TextureSize(TypedNodeParam->TextureSizeX, TypedNodeParam->TextureSizeY);
@@ -221,7 +212,13 @@ UE::Mutable::Private::Ptr<UE::Mutable::Private::NodeImage> GenerateMutableSource
 			ResizeNode->SizeX = TextureSize.X;
 			ResizeNode->SizeY = TextureSize.Y;
 
-			Result = ResizeNode;
+			UE::Mutable::Private::Ptr<UE::Mutable::Private::NodeImageFormat> ImageFormat = GenerateMutableImageFormat(ResizeNode, TypedNodeParam->ReferenceValue, GenerationContext, TypedNodeParam);
+			
+			// Set expected image formats in game
+			ImagePassthroughConvertNode->ImageDescriptor.m_format = ImageFormat->Format;
+			ImagePassthroughConvertNode->ImageDescriptor.FormatIfAlpha = ImageFormat->FormatIfAlpha;
+
+			Result = ImageFormat;
 		}
 	}
 
@@ -1333,5 +1330,140 @@ UE::Mutable::Private::Ptr<UE::Mutable::Private::NodeImageParameter> GenerateMuta
 }
 
 
+UE::Mutable::Private::Ptr<UE::Mutable::Private::NodeImageFormat> GenerateMutableImageFormat(
+	const UE::Mutable::Private::Ptr<UE::Mutable::Private::NodeImage>& InSource,
+	const UTexture2D* ReferenceTexture,
+	const FMutableGraphGenerationContext& GenerationContext,
+	const UCustomizableObjectNode* Node)
+{
+	UE::Mutable::Private::Ptr<UE::Mutable::Private::NodeImageFormat> ImageFormat = new UE::Mutable::Private::NodeImageFormat();
+	ImageFormat->Source = InSource.get();
+	ImageFormat->SetMessageContext(Node);
+
+	// Use an "easy format" as the default value.
+	ImageFormat->Format = UE::Mutable::Private::EImageFormat::RGBA_UByte;
+
+	if (ReferenceTexture)
+	{
+		TArray<TArray<FTextureBuildSettings>> BuildSettingsPerFormatPerLayer;
+		if (GenerationContext.CompilationContext->Options.TargetPlatform)
+		{
+			ReferenceTexture->GetTargetPlatformBuildSettings(GenerationContext.CompilationContext->Options.TargetPlatform, BuildSettingsPerFormatPerLayer);
+
+			bool bIsServerOnly = GenerationContext.CompilationContext->Options.TargetPlatform->IsServerOnly();
+			// Suppress the message for server only platforms. Images are discarded is ServerOnly. 
+			if (!bIsServerOnly && BuildSettingsPerFormatPerLayer.IsEmpty())
+			{
+				const FString ReplacedImageFormatMsg = FString::Printf(TEXT("In object [%s] for platform [%s] the unsupported image format of texture [%s] is used, RGBA_UByte will be used instead."),
+					*GenerationContext.GetObjectName(),
+					*GenerationContext.CompilationContext->Options.TargetPlatform->PlatformName(),
+					*ReferenceTexture->GetName());
+				const FText ReplacedImageFormatText = FText::FromString(ReplacedImageFormatMsg);
+				GenerationContext.Log(ReplacedImageFormatText, Node, EMessageSeverity::Info);
+				UE_LOG(LogMutable, Log, TEXT("%s"), *ReplacedImageFormatMsg);
+			}
+			else if (BuildSettingsPerFormatPerLayer.Num() > 1)
+			{
+				const FString ReplacedImageFormatMsg = FString::Printf(TEXT("In object [%s] for platform [%s] the image format of texture [%s] has multiple target formats. Only one will be used.."),
+					*GenerationContext.GetObjectName(),
+					*GenerationContext.CompilationContext->Options.TargetPlatform->PlatformName(),
+					*ReferenceTexture->GetName());
+				const FText ReplacedImageFormatText = FText::FromString(ReplacedImageFormatMsg);
+				GenerationContext.Log(ReplacedImageFormatText, Node, EMessageSeverity::Info);
+				UE_LOG(LogMutable, Log, TEXT("%s"), *ReplacedImageFormatMsg);
+			}
+		}
+
+		if (BuildSettingsPerFormatPerLayer.IsValidIndex(0) && BuildSettingsPerFormatPerLayer[0].IsValidIndex(0))
+		{
+			const TArray<FTextureBuildSettings>& BuildSettingsPerLayer = BuildSettingsPerFormatPerLayer[0];
+
+			if (GenerationContext.CompilationContext->Options.TextureCompression != ECustomizableObjectTextureCompression::None)
+			{
+				static ITextureFormatManagerModule* TextureFormatManager = nullptr;
+				if (!TextureFormatManager)
+				{
+					TextureFormatManager = &FModuleManager::LoadModuleChecked<ITextureFormatManagerModule>("TextureFormat");
+					check(TextureFormatManager);
+				}
+
+				EPixelFormat UnrealTargetPlatformFormat = EPixelFormat::PF_Unknown;
+				EPixelFormat UnrealTargetPlatformFormatAlpha = EPixelFormat::PF_Unknown;
+
+				if (ensure(TextureFormatManager))
+				{
+					const ITextureFormat* TextureFormat = TextureFormatManager->FindTextureFormat(BuildSettingsPerLayer[0].TextureFormatName);
+					if (ensure(TextureFormat))
+					{
+						UnrealTargetPlatformFormat = TextureFormat->GetEncodedPixelFormat(BuildSettingsPerLayer[0], false);
+						UnrealTargetPlatformFormatAlpha = TextureFormat->GetEncodedPixelFormat(BuildSettingsPerLayer[0], true);
+					}
+				}
+
+				// \TODO: The QualityFix filter is used while the internal mutable runtime compression doesn't provide enough quality for some large block formats.
+				UE::Mutable::Private::EImageFormat MutableFormat = QualityAndPerformanceFix(UnrealToMutablePixelFormat(UnrealTargetPlatformFormat, false));
+				UE::Mutable::Private::EImageFormat MutableFormatIfAlpha = QualityAndPerformanceFix(UnrealToMutablePixelFormat(UnrealTargetPlatformFormatAlpha, true));
+
+				// Temp hack to enable RG->LA 
+				if (GenerationContext.CompilationContext->Options.TargetPlatform)
+				{
+					bool bUseLA = GenerationContext.CompilationContext->Options.TargetPlatform->SupportsFeature(ETargetPlatformFeatures::NormalmapLAEncodingMode);
+					if (bUseLA)
+					{
+						// We'll have to trust the reference texture because the actual internal settings are opaque.
+						// See GetQualityFormat in TextureFormatASTC.cpp to understand how it works, but it depends on some inaccessible texture name constants like this:
+						//bool bIsNormalMapFormat = TextureFormatName == GTextureFormatNameASTC_NormalAG ||
+						//		TextureFormatName == GTextureFormatNameASTC_NormalRG ||
+						//		TextureFormatName == GTextureFormatNameASTC_NormalLA ||
+						//		TextureFormatName == GTextureFormatNameASTC_NormalRG_Precise;
+						bool bIsNormalMapFormat = ReferenceTexture->IsNormalMap();
+
+						if (bIsNormalMapFormat)
+						{
+							// Insert a channel swizzle
+							UE::Mutable::Private::Ptr<UE::Mutable::Private::NodeImageSwizzle> Swizzle = new UE::Mutable::Private::NodeImageSwizzle;
+							Swizzle->SetFormat(UE::Mutable::Private::EImageFormat::RGBA_UByte);
+							Swizzle->Sources[0] = InSource;
+							Swizzle->Sources[1] = InSource;
+							Swizzle->Sources[2] = InSource;
+							Swizzle->Sources[3] = InSource;
+							Swizzle->SourceChannels[0] = 0;
+							Swizzle->SourceChannels[1] = 0;
+							Swizzle->SourceChannels[2] = 0;
+							Swizzle->SourceChannels[3] = 1;
+
+							ImageFormat->Source = Swizzle;
+						}
+					}
+				}
+
+				// Unsupported format: look for something generic
+				if (MutableFormat == UE::Mutable::Private::EImageFormat::None)
+				{
+					const FString ReplacedImageFormatMsg = FString::Printf(TEXT("In object [%s] the unsupported image format %d is used, RGBA_UByte will be used instead."), *GenerationContext.GetObjectName(), UnrealTargetPlatformFormat);
+					const FText ReplacedImageFormatText = FText::FromString(ReplacedImageFormatMsg);
+					GenerationContext.Log(ReplacedImageFormatText, Node, EMessageSeverity::Info);
+					UE_LOG(LogMutable, Log, TEXT("%s"), *ReplacedImageFormatMsg);
+					MutableFormat = UE::Mutable::Private::EImageFormat::RGBA_UByte;
+				}
+				if (MutableFormatIfAlpha == UE::Mutable::Private::EImageFormat::None)
+				{
+					const FString ReplacedImageFormatMsg = FString::Printf(TEXT("In object [%s] the unsupported image format %d is used, RGBA_UByte will be used instead."), *GenerationContext.GetObjectName(), UnrealTargetPlatformFormatAlpha);
+					const FText ReplacedImageFormatText = FText::FromString(ReplacedImageFormatMsg);
+					GenerationContext.Log(ReplacedImageFormatText, Node, EMessageSeverity::Info);
+					UE_LOG(LogMutable, Log, TEXT("%s"), *ReplacedImageFormatMsg);
+					MutableFormatIfAlpha = UE::Mutable::Private::EImageFormat::RGBA_UByte;
+				}
+
+				ImageFormat->Format = MutableFormat;
+				ImageFormat->FormatIfAlpha = MutableFormatIfAlpha;
+			}
+		}
+	}
+
+	return ImageFormat;
+}
+
+
 #undef LOCTEXT_NAMESPACE
 
diff --git a/Engine/Plugins/Mutable/Source/CustomizableObjectEditor/Private/MuCOE/GenerateMutableSource/GenerateMutableSourceImage.h b/Engine/Plugins/Mutable/Source/CustomizableObjectEditor/Private/MuCOE/GenerateMutableSource/GenerateMutableSourceImage.h
--- a/Engine/Plugins/Mutable/Source/CustomizableObjectEditor/Private/MuCOE/GenerateMutableSource/GenerateMutableSourceImage.h
+++ b/Engine/Plugins/Mutable/Source/CustomizableObjectEditor/Private/MuCOE/GenerateMutableSource/GenerateMutableSourceImage.h
@@ -16,6 +16,7 @@ namespace UE::Mutable::Private
 {
 	class NodeImage;
 	class NodeImageParameter;
+	class NodeImageFormat;
 }
 
 
@@ -24,4 +25,10 @@ UE::Mutable::Private::Ptr<UE::Mutable::Private::NodeImage> ResizeTextureByNumMip
 /** Convert a CustomizableObject Source Graph from an Image pin into a mutable source graph. */
 UE::Mutable::Private::Ptr<UE::Mutable::Private::NodeImage> GenerateMutableSourceImage(const UEdGraphPin* Pin, FMutableGraphGenerationContext& GenerationContext, int32 ReferenceTextureSize);
 
-UE::Mutable::Private::Ptr<UE::Mutable::Private::NodeImageParameter> GenerateMutableSourceImageParameter(const UEdGraphPin* Pin, FMutableGraphGenerationContext& GenerationContext);
\ No newline at end of file
+UE::Mutable::Private::Ptr<UE::Mutable::Private::NodeImageParameter> GenerateMutableSourceImageParameter(const UEdGraphPin* Pin, FMutableGraphGenerationContext& GenerationContext);
+
+UE::Mutable::Private::Ptr<UE::Mutable::Private::NodeImageFormat> GenerateMutableImageFormat(
+	const UE::Mutable::Private::Ptr<UE::Mutable::Private::NodeImage>& InSource,
+	const UTexture2D* ReferenceTexture,
+	const FMutableGraphGenerationContext& GenerationContext,
+	const UCustomizableObjectNode* Node);
\ No newline at end of file
diff --git a/Engine/Plugins/Mutable/Source/CustomizableObjectEditor/Private/MuCOE/GenerateMutableSource/GenerateMutableSourceMaterial.cpp b/Engine/Plugins/Mutable/Source/CustomizableObjectEditor/Private/MuCOE/GenerateMutableSource/GenerateMutableSourceMaterial.cpp
--- a/Engine/Plugins/Mutable/Source/CustomizableObjectEditor/Private/MuCOE/GenerateMutableSource/GenerateMutableSourceMaterial.cpp
+++ b/Engine/Plugins/Mutable/Source/CustomizableObjectEditor/Private/MuCOE/GenerateMutableSource/GenerateMutableSourceMaterial.cpp
@@ -312,16 +312,9 @@ UE::Mutable::Private::Ptr<UE::Mutable::Private::NodeMaterial> GenerateMutableSou
 					UE::Mutable::Private::Ptr<UE::Mutable::Private::NodeImageFromMaterialParameter> TextureNode = new UE::Mutable::Private::NodeImageFromMaterialParameter();
 					TextureNode->ImageParameterName = GenerationContext.CurrentMaterialBreakParameter.ParameterName;
 
-					// Force the same format that the default texture if any.
-					UE::Mutable::Private::Ptr<UE::Mutable::Private::NodeImageFormat> FormatNode = new UE::Mutable::Private::NodeImageFormat();
-
-					// Force an "easy format" on the texture.
-					FormatNode->Format = UE::Mutable::Private::EImageFormat::RGBA_UByte;
-					FormatNode->Source = TextureNode;
-
 					// Resize image node.
 					UE::Mutable::Private::Ptr<UE::Mutable::Private::NodeImageResize> ResizeNode = new UE::Mutable::Private::NodeImageResize();
-					ResizeNode->Base = FormatNode;
+					ResizeNode->Base = TextureNode;
 					ResizeNode->bRelative = false;
 
 					UE::Mutable::Private::FImageSize TextureSize(1);
@@ -341,7 +334,10 @@ UE::Mutable::Private::Ptr<UE::Mutable::Private::NodeMaterial> GenerateMutableSou
 					ResizeNode->SizeX = TextureSize.X;
 					ResizeNode->SizeY = TextureSize.Y;
 
-					MaterialNode->ImageParameterNodes.Add(GenerationContext.CurrentMaterialBreakParameter.ParameterName.ToString(), ResizeNode);
+					// Force the same format that the default texture if any.
+					UE::Mutable::Private::Ptr<UE::Mutable::Private::NodeImageFormat> FormatNode = GenerateMutableImageFormat(ResizeNode, ReferenceTexture, GenerationContext, Node);
+
+					MaterialNode->ImageParameterNodes.Add(GenerationContext.CurrentMaterialBreakParameter.ParameterName.ToString(), FormatNode);
 
 					break;
 				}
diff --git a/Engine/Plugins/Mutable/Source/CustomizableObjectEditor/Private/MuCOE/GenerateMutableSource/GenerateMutableSourceSurface.cpp b/Engine/Plugins/Mutable/Source/CustomizableObjectEditor/Private/MuCOE/GenerateMutableSource/GenerateMutableSourceSurface.cpp
--- a/Engine/Plugins/Mutable/Source/CustomizableObjectEditor/Private/MuCOE/GenerateMutableSource/GenerateMutableSourceSurface.cpp
+++ b/Engine/Plugins/Mutable/Source/CustomizableObjectEditor/Private/MuCOE/GenerateMutableSource/GenerateMutableSourceSurface.cpp
@@ -7,9 +7,6 @@
 #include "Engine/TextureLODSettings.h"
 #include "GPUSkinPublicDefs.h"
 #include "Interfaces/ITargetPlatform.h"
-#include "Interfaces/ITextureFormat.h"
-#include "Interfaces/ITextureFormatModule.h"
-#include "Interfaces/ITextureFormatManagerModule.h"
 #include "Modules/ModuleManager.h"
 #include "TextureCompressorModule.h"
 #include "Materials/MaterialInstance.h"
@@ -57,7 +54,6 @@
 #include "MuT/NodeMaterialConstant.h"
 #include "MuT/NodeSurfaceSwitch.h"
 #include "MuT/NodeSurfaceVariation.h"
-#include "MuT/UnrealPixelFormatOverride.h"
 
 
 #define LOCTEXT_NAMESPACE "CustomizableObjectEditor"
@@ -547,120 +543,8 @@ UE::Mutable::Private::Ptr<UE::Mutable::Private::NodeSurface> GenerateMutableSour
 								LastImage = CompositedImage;
 							}
 
-							UE::Mutable::Private::Ptr<UE::Mutable::Private::NodeImage> FormatSource = LastImage;
-							UE::Mutable::Private::Ptr<UE::Mutable::Private::NodeImageFormat> FormatImage = new UE::Mutable::Private::NodeImageFormat();
-							FormatImage->Source = LastImage.get();
-							FormatImage->Format = UE::Mutable::Private::EImageFormat::RGBA_UByte;
-							FormatImage->SetMessageContext(Node);
-							LastImage = FormatImage;
-
-							TArray<TArray<FTextureBuildSettings>> BuildSettingsPerFormatPerLayer;
-							if (GenerationContext.CompilationContext->Options.TargetPlatform)
-							{
-								ReferenceTexture->GetTargetPlatformBuildSettings(GenerationContext.CompilationContext->Options.TargetPlatform, BuildSettingsPerFormatPerLayer);
-
-								bool bIsServerOnly = GenerationContext.CompilationContext->Options.TargetPlatform->IsServerOnly();
-								// Suppress the message for server only platforms. Images are discarded is ServerOnly. 
-								if (!bIsServerOnly && BuildSettingsPerFormatPerLayer.IsEmpty())
-								{
-									const FString ReplacedImageFormatMsg = FString::Printf(TEXT("In object [%s] for platform [%s] the unsupported image format of texture [%s] is used, RGBA_UByte will be used instead."),
-										*GenerationContext.GetObjectName(),
-										*GenerationContext.CompilationContext->Options.TargetPlatform->PlatformName(),
-										*ReferenceTexture->GetName());
-									const FText ReplacedImageFormatText = FText::FromString(ReplacedImageFormatMsg);
-									GenerationContext.Log(ReplacedImageFormatText, Node, EMessageSeverity::Info);
-									UE_LOG(LogMutable, Log, TEXT("%s"), *ReplacedImageFormatMsg);
-								}
-								else if (BuildSettingsPerFormatPerLayer.Num() > 1)
-								{
-									const FString ReplacedImageFormatMsg = FString::Printf(TEXT("In object [%s] for platform [%s] the image format of texture [%s] has multiple target formats. Only one will be used.."),
-										*GenerationContext.GetObjectName(),
-										*GenerationContext.CompilationContext->Options.TargetPlatform->PlatformName(),
-										*ReferenceTexture->GetName());
-									const FText ReplacedImageFormatText = FText::FromString(ReplacedImageFormatMsg);
-									GenerationContext.Log(ReplacedImageFormatText, Node, EMessageSeverity::Info);
-									UE_LOG(LogMutable, Log, TEXT("%s"), *ReplacedImageFormatMsg);
-								}
-							}
-
-							if (!BuildSettingsPerFormatPerLayer.IsEmpty())
-							{
-								const TArray<FTextureBuildSettings>& BuildSettingsPerLayer = BuildSettingsPerFormatPerLayer[0];
-
-								if (GenerationContext.CompilationContext->Options.TextureCompression!=ECustomizableObjectTextureCompression::None)
-								{
-									static ITextureFormatManagerModule* TextureFormatManager = nullptr;
-									if (!TextureFormatManager)
-									{
-										TextureFormatManager = &FModuleManager::LoadModuleChecked<ITextureFormatManagerModule>("TextureFormat");
-										check(TextureFormatManager);
-									}
-									const ITextureFormat* TextureFormat = TextureFormatManager->FindTextureFormat(BuildSettingsPerLayer[0].TextureFormatName);
-									check(TextureFormat);
-									EPixelFormat UnrealTargetPlatformFormat = TextureFormat->GetEncodedPixelFormat(BuildSettingsPerLayer[0], false);
-									EPixelFormat UnrealTargetPlatformFormatAlpha = TextureFormat->GetEncodedPixelFormat(BuildSettingsPerLayer[0], true);
-
-									// \TODO: The QualityFix filter is used while the internal mutable runtime compression doesn't provide enough quality for some large block formats.
-									UE::Mutable::Private::EImageFormat MutableFormat = QualityAndPerformanceFix(UnrealToMutablePixelFormat(UnrealTargetPlatformFormat,false));
-									UE::Mutable::Private::EImageFormat MutableFormatIfAlpha = QualityAndPerformanceFix(UnrealToMutablePixelFormat(UnrealTargetPlatformFormatAlpha,true));
-
-									// Temp hack to enable RG->LA 
-									if (GenerationContext.CompilationContext->Options.TargetPlatform)
-									{
-										bool bUseLA = GenerationContext.CompilationContext->Options.TargetPlatform->SupportsFeature(ETargetPlatformFeatures::NormalmapLAEncodingMode);
-										if (bUseLA)
-										{
-											// We'll have to trust the reference texture because the actual internal settings are opaque.
-											// See GetQualityFormat in TextureFormatASTC.cpp to understand how it works, but it depends on some inaccessible texture name constants like this:
-											//bool bIsNormalMapFormat = TextureFormatName == GTextureFormatNameASTC_NormalAG ||
-											//		TextureFormatName == GTextureFormatNameASTC_NormalRG ||
-											//		TextureFormatName == GTextureFormatNameASTC_NormalLA ||
-											//		TextureFormatName == GTextureFormatNameASTC_NormalRG_Precise;
-											bool bIsNormalMapFormat = ReferenceTexture->IsNormalMap();
-
-											if (bIsNormalMapFormat)
-											{
-												// Insert a channel swizzle
-												UE::Mutable::Private::Ptr<UE::Mutable::Private::NodeImageSwizzle> Swizzle = new UE::Mutable::Private::NodeImageSwizzle;
-												Swizzle->SetFormat( UE::Mutable::Private::EImageFormat::RGBA_UByte );
-												Swizzle->Sources[0] = FormatSource;
-												Swizzle->Sources[1] = FormatSource;
-												Swizzle->Sources[2] = FormatSource;
-												Swizzle->Sources[3] = FormatSource;
-												Swizzle->SourceChannels[0] = 0;
-												Swizzle->SourceChannels[1] = 0;
-												Swizzle->SourceChannels[2] = 0;
-												Swizzle->SourceChannels[3] = 1;
-
-												FormatImage->Source = Swizzle;
-											}
-										}
-									}
-
-									// Unsupported format: look for something generic
-									if (MutableFormat == UE::Mutable::Private::EImageFormat::None)
-									{
-										const FString ReplacedImageFormatMsg = FString::Printf(TEXT("In object [%s] the unsupported image format %d is used, RGBA_UByte will be used instead."), *GenerationContext.GetObjectName(), UnrealTargetPlatformFormat );
-										const FText ReplacedImageFormatText = FText::FromString(ReplacedImageFormatMsg);
-										GenerationContext.Log(ReplacedImageFormatText, Node, EMessageSeverity::Info);
-										UE_LOG(LogMutable, Log, TEXT("%s"), *ReplacedImageFormatMsg);
-										MutableFormat = UE::Mutable::Private::EImageFormat::RGBA_UByte;
-									}
-									if (MutableFormatIfAlpha == UE::Mutable::Private::EImageFormat::None)
-									{
-										const FString ReplacedImageFormatMsg = FString::Printf(TEXT("In object [%s] the unsupported image format %d is used, RGBA_UByte will be used instead."), *GenerationContext.GetObjectName(), UnrealTargetPlatformFormatAlpha);
-										const FText ReplacedImageFormatText = FText::FromString(ReplacedImageFormatMsg);
-										GenerationContext.Log(ReplacedImageFormatText, Node, EMessageSeverity::Info);
-										UE_LOG(LogMutable, Log, TEXT("%s"), *ReplacedImageFormatMsg);
-										MutableFormatIfAlpha = UE::Mutable::Private::EImageFormat::RGBA_UByte;
-									}
-
-									FormatImage->Format = MutableFormat;
-									FormatImage->FormatIfAlpha = MutableFormatIfAlpha;
-								}
-							}
-
-							ImageNode = LastImage;
+							// Apply Image format on top of the image tree to propagate the format down
+							ImageNode = GenerateMutableImageFormat(LastImage, ReferenceTexture, GenerationContext, Node);
 						}
 
 						SurfNode->Images[ImageIndex].Image = ImageNode;
diff --git a/Engine/Plugins/Mutable/Source/MutableRuntime/Internal/MuR/ImageDataStorage.h b/Engine/Plugins/Mutable/Source/MutableRuntime/Internal/MuR/ImageDataStorage.h
--- a/Engine/Plugins/Mutable/Source/MutableRuntime/Internal/MuR/ImageDataStorage.h
+++ b/Engine/Plugins/Mutable/Source/MutableRuntime/Internal/MuR/ImageDataStorage.h
@@ -67,7 +67,7 @@ public:
 
 	inline FImageDesc MakeImageDesc() const
 	{
-		return FImageDesc(ImageSize, ImageFormat, NumLODs);
+		return FImageDesc(ImageSize, ImageFormat, ImageFormat, NumLODs);
 	}
 
 	inline void InitVoid(const FImageDesc& Desc) 
diff --git a/Engine/Plugins/Mutable/Source/MutableRuntime/Internal/MuR/ImageTypes.h b/Engine/Plugins/Mutable/Source/MutableRuntime/Internal/MuR/ImageTypes.h
--- a/Engine/Plugins/Mutable/Source/MutableRuntime/Internal/MuR/ImageTypes.h
+++ b/Engine/Plugins/Mutable/Source/MutableRuntime/Internal/MuR/ImageTypes.h
@@ -78,8 +78,8 @@ namespace UE::Mutable::Private
 		{
 		}
 
-		FImageDesc(const FImageSize& InSize, EImageFormat InFormat, uint8 InLods)
-			: m_size(InSize), m_format(InFormat), m_lods(InLods) 
+		FImageDesc(const FImageSize& InSize, EImageFormat InFormat, EImageFormat InFormatIfAlpha, uint8 InLods)
+			: m_size(InSize), m_format(InFormat), FormatIfAlpha(InFormatIfAlpha), m_lods(InLods)
 		{
 		}
 
@@ -89,6 +89,9 @@ namespace UE::Mutable::Private
 		/** */
 		EImageFormat m_format = EImageFormat::None;
 
+		/** */
+		EImageFormat FormatIfAlpha = EImageFormat::None;
+
 		/** Levels of detail (mipmaps) */
 		uint8 m_lods = 0;
 
@@ -98,6 +101,7 @@ namespace UE::Mutable::Private
 			return
 				m_size == Other.m_size && 
 				m_format == Other.m_format && 
+				FormatIfAlpha == Other.FormatIfAlpha &&
 				m_lods == Other.m_lods;
 		}
 	};
diff --git a/Engine/Plugins/Mutable/Source/MutableRuntime/Private/MuR/Image.cpp b/Engine/Plugins/Mutable/Source/MutableRuntime/Private/MuR/Image.cpp
--- a/Engine/Plugins/Mutable/Source/MutableRuntime/Private/MuR/Image.cpp
+++ b/Engine/Plugins/Mutable/Source/MutableRuntime/Private/MuR/Image.cpp
@@ -27,7 +27,7 @@ namespace UE::Mutable::Private
 		MUTABLE_CPUPROFILER_SCOPE(NewImage)
 		LLM_SCOPE_BYNAME(TEXT("MutableRuntime"));
 
-		const FImageDesc ImageDesc {FImageSize(uint16(SizeX), uint16(SizeY)), Format, static_cast<uint8>(Lods)};
+		const FImageDesc ImageDesc {FImageSize(uint16(SizeX), uint16(SizeY)), Format, EImageFormat::None, static_cast<uint8>(Lods)};
 		DataStorage.Init(ImageDesc, InitType);
 	}
 
diff --git a/Engine/Plugins/Mutable/Source/MutableTools/Internal/MuT/AST.h b/Engine/Plugins/Mutable/Source/MutableTools/Internal/MuT/AST.h
--- a/Engine/Plugins/Mutable/Source/MutableTools/Internal/MuT/AST.h
+++ b/Engine/Plugins/Mutable/Source/MutableTools/Internal/MuT/AST.h
@@ -854,6 +854,7 @@ namespace UE::Mutable::Private
         {
         public:
             TMap<const ASTOp*, FImageDesc> m_results;
+			bool bIsFormatFromImageParameter = false;
         };
 
         //!
diff --git a/Engine/Plugins/Mutable/Source/MutableTools/Internal/MuT/NodeImageParameterConvert.h b/Engine/Plugins/Mutable/Source/MutableTools/Internal/MuT/NodeImageParameterConvert.h
--- a/Engine/Plugins/Mutable/Source/MutableTools/Internal/MuT/NodeImageParameterConvert.h
+++ b/Engine/Plugins/Mutable/Source/MutableTools/Internal/MuT/NodeImageParameterConvert.h
@@ -3,6 +3,7 @@
 #pragma once
 
 #include "MuR/Ptr.h"
+#include "MuR/ImageTypes.h"
 #include "MuT/Node.h"
 #include "MuT/NodeImageParameter.h"
 
@@ -15,6 +16,7 @@ namespace UE::Mutable::Private
 	{
 	public:
 		Ptr<NodeImageParameter> ImageParameter;
+		FImageDesc ImageDescriptor;
 		
 		// Node Interface
 		virtual const FNodeType* GetType() const override { return &StaticType; }
diff --git a/Engine/Plugins/Mutable/Source/MutableTools/Private/MuT/ASTOpConditional.cpp b/Engine/Plugins/Mutable/Source/MutableTools/Private/MuT/ASTOpConditional.cpp
--- a/Engine/Plugins/Mutable/Source/MutableTools/Private/MuT/ASTOpConditional.cpp
+++ b/Engine/Plugins/Mutable/Source/MutableTools/Private/MuT/ASTOpConditional.cpp
@@ -173,6 +173,7 @@ namespace UE::Mutable::Private
 			else
 			{
 				res.m_format = (yesDesc.m_format == noDesc.m_format) ? yesDesc.m_format : EImageFormat::None;
+				res.FormatIfAlpha = (yesDesc.FormatIfAlpha == noDesc.FormatIfAlpha) ? yesDesc.FormatIfAlpha : EImageFormat::None;
 				res.m_lods = (yesDesc.m_lods == noDesc.m_lods) ? yesDesc.m_lods : 0;
 				res.m_size = (yesDesc.m_size == noDesc.m_size) ? yesDesc.m_size : FImageSize(0, 0);
 			}
diff --git a/Engine/Plugins/Mutable/Source/MutableTools/Private/MuT/ASTOpImageCompose.cpp b/Engine/Plugins/Mutable/Source/MutableTools/Private/MuT/ASTOpImageCompose.cpp
--- a/Engine/Plugins/Mutable/Source/MutableTools/Private/MuT/ASTOpImageCompose.cpp
+++ b/Engine/Plugins/Mutable/Source/MutableTools/Private/MuT/ASTOpImageCompose.cpp
@@ -117,20 +117,20 @@ void ASTOpImageCompose::Link( FProgram& program, FLinkerOptions*)
 }
 
 
-FImageDesc ASTOpImageCompose::GetImageDesc( bool returnBestOption, FGetImageDescContext* context ) const
+FImageDesc ASTOpImageCompose::GetImageDesc( bool returnBestOption, FGetImageDescContext* Context) const
 {
-    FImageDesc res;
+    FImageDesc Result;
 
     // Local context in case it is necessary
-    FGetImageDescContext localContext;
-    if (!context)
+    FGetImageDescContext LocalContext;
+    if (!Context)
     {
-      context = &localContext;
+		Context = &LocalContext;
     }
     else
     {
         // Cached result?
-		FImageDesc* PtrValue = context->m_results.Find(this);
+		FImageDesc* PtrValue = Context->m_results.Find(this);
 		if (PtrValue)
 		{
 			return *PtrValue;
@@ -140,22 +140,23 @@ FImageDesc ASTOpImageCompose::GetImageDesc( bool returnBestOption, FGetImageDesc
     // Actual work
     if (Base)
     {
-        res = Base->GetImageDesc( returnBestOption, context );
+		Result = Base->GetImageDesc( returnBestOption, Context);
     }
 
 	if (BlockImage)
 	{
-		FImageDesc BlockDesc = BlockImage->GetImageDesc(returnBestOption, context);
-		res.m_format = GetMostGenericFormat(res.m_format,BlockDesc.m_format);
+		FImageDesc BlockDesc = BlockImage->GetImageDesc(returnBestOption, Context);
+		Result.m_format = GetMostGenericFormat(Result.m_format,BlockDesc.m_format);
+		Result.FormatIfAlpha = GetMostGenericFormat(Result.FormatIfAlpha,BlockDesc.FormatIfAlpha);
 	}
 
+	// Format overriden by current op.
+	Context->bIsFormatFromImageParameter = false;
+
     // Cache th result
-    if (context)
-    {
-		context->m_results.Add(this, res);
-	}
+	Context->m_results.Add(this, Result);
 
-    return res;
+    return Result;
 }
 
 
@@ -269,7 +270,7 @@ UE::Mutable::Private::Ptr<ASTOp> ASTOpImageCompose::OptimiseSemantic(const FMode
 			{
 				UE::Mutable::Private::Ptr<ASTOpImagePixelFormat> Reformat = new ASTOpImagePixelFormat;
 				Reformat->Format = BaseDesc.m_format;
-				Reformat->FormatIfAlpha = EImageFormat::None;
+				Reformat->FormatIfAlpha = BaseDesc.FormatIfAlpha;
 				Reformat->Source = at;
 				at = Reformat;
 			}
diff --git a/Engine/Plugins/Mutable/Source/MutableTools/Private/MuT/ASTOpImageLayer.cpp b/Engine/Plugins/Mutable/Source/MutableTools/Private/MuT/ASTOpImageLayer.cpp
--- a/Engine/Plugins/Mutable/Source/MutableTools/Private/MuT/ASTOpImageLayer.cpp
+++ b/Engine/Plugins/Mutable/Source/MutableTools/Private/MuT/ASTOpImageLayer.cpp
@@ -426,12 +426,14 @@ namespace UE::Mutable::Private
 					&&
 					TypedBlend->SourceChannels[0]==3)
 				{
-					Ptr<ASTOpImageLayer> nop = UE::Mutable::Private::Clone<ASTOpImageLayer>(this);
-					nop->Flags = Flags | OP::ImageLayerArgs::FLAGS::F_BLENDED_RGB_FROM_ALPHA;
-					nop->blend = TypedBlend->Sources[0].child();
-					auto Desc = nop->blend->GetImageDesc(true);
-					check(Desc.m_format==EImageFormat::RGBA_UByte);
-					at = nop;					
+					const FImageDesc SourceImageDesc = TypedBlend->Sources[0]->GetImageDesc(true);
+					if (ensure(SourceImageDesc.m_format == EImageFormat::RGBA_UByte))
+					{
+						Ptr<ASTOpImageLayer> nop = UE::Mutable::Private::Clone<ASTOpImageLayer>(this);
+						nop->Flags = Flags | OP::ImageLayerArgs::FLAGS::F_BLENDED_RGB_FROM_ALPHA;
+						nop->blend = TypedBlend->Sources[0].child();
+						at = nop;
+					}
 				}
 			}
 		}
diff --git a/Engine/Plugins/Mutable/Source/MutableTools/Private/MuT/ASTOpImageMipmap.cpp b/Engine/Plugins/Mutable/Source/MutableTools/Private/MuT/ASTOpImageMipmap.cpp
--- a/Engine/Plugins/Mutable/Source/MutableTools/Private/MuT/ASTOpImageMipmap.cpp
+++ b/Engine/Plugins/Mutable/Source/MutableTools/Private/MuT/ASTOpImageMipmap.cpp
@@ -186,20 +186,20 @@ namespace UE::Mutable::Private
 
 
 	//!
-	FImageDesc ASTOpImageMipmap::GetImageDesc(bool returnBestOption, FGetImageDescContext* context) const
+	FImageDesc ASTOpImageMipmap::GetImageDesc(bool returnBestOption, FGetImageDescContext* Context) const
 	{
-		FImageDesc res;
+		FImageDesc Result;
 
 		// Local context in case it is necessary
-		FGetImageDescContext localContext;
-		if (!context)
+		FGetImageDescContext LocalContext;
+		if (!Context)
 		{
-			context = &localContext;
+			Context = &LocalContext;
 		}
 		else
 		{
 			// Cached result?
-			FImageDesc* PtrValue = context->m_results.Find(this);
+			FImageDesc* PtrValue = Context->m_results.Find(this);
 			if (PtrValue)
 			{
 				return *PtrValue;
@@ -208,7 +208,7 @@ namespace UE::Mutable::Private
 
 		if (Source.child())
 		{
-			res = Source.child()->GetImageDesc(returnBestOption, context);
+			Result = Source.child()->GetImageDesc(returnBestOption, Context);
 		}
 
 		int mipLevels = Levels;
@@ -216,21 +216,17 @@ namespace UE::Mutable::Private
 		if (mipLevels == 0)
 		{
 			mipLevels = FMath::Max(
-				(int)ceilf(logf((float)res.m_size[0]) / logf(2.0f)),
-				(int)ceilf(logf((float)res.m_size[1]) / logf(2.0f))
+				(int)ceilf(logf((float)Result.m_size[0]) / logf(2.0f)),
+				(int)ceilf(logf((float)Result.m_size[1]) / logf(2.0f))
 			);
 		}
 
-		res.m_lods = FMath::Max(res.m_lods, (uint8)mipLevels);
-
+		Result.m_lods = FMath::Max(Result.m_lods, (uint8)mipLevels);
 
 		// Cache the result
-		if (context)
-		{
-			context->m_results.Add(this, res);
-		}
+		Context->m_results.Add(this, Result);
 
-		return res;
+		return Result;
 	}
 
 
diff --git a/Engine/Plugins/Mutable/Source/MutableTools/Private/MuT/ASTOpImageParameterConvert.cpp b/Engine/Plugins/Mutable/Source/MutableTools/Private/MuT/ASTOpImageParameterConvert.cpp
--- a/Engine/Plugins/Mutable/Source/MutableTools/Private/MuT/ASTOpImageParameterConvert.cpp
+++ b/Engine/Plugins/Mutable/Source/MutableTools/Private/MuT/ASTOpImageParameterConvert.cpp
@@ -32,7 +32,8 @@ namespace UE::Mutable::Private
 		{
 			const ASTOpImageParameterConvert* Other = static_cast<const ASTOpImageParameterConvert*>(&OtherUntyped);
 			return Type == Other->Type &&
-				ImageParameter == Other->ImageParameter;
+				ImageParameter == Other->ImageParameter &&
+				ImageDescriptor == Other->ImageDescriptor;
 		}
 		
 		return false;
@@ -44,6 +45,7 @@ namespace UE::Mutable::Private
 		Ptr<ASTOpImageParameterConvert> New = new ASTOpImageParameterConvert();
 		New->Type = Type;
 		New->ImageParameter =  ASTChild(New, MapChild(ImageParameter.child()));
+		New->ImageDescriptor = ImageDescriptor;
 
 		return New;
 	}
@@ -84,8 +86,13 @@ namespace UE::Mutable::Private
 	}
 
 	
-	FImageDesc ASTOpImageParameterConvert::GetImageDesc(bool, FGetImageDescContext*) const
+	FImageDesc ASTOpImageParameterConvert::GetImageDesc(bool, FGetImageDescContext* Context) const
 	{
-		return ImageParameter->GetImageDesc();
+		if (Context)
+		{
+			Context->bIsFormatFromImageParameter = true;
+		}
+
+		return ImageDescriptor;
 	}
 }
diff --git a/Engine/Plugins/Mutable/Source/MutableTools/Private/MuT/ASTOpImageParameterConvert.h b/Engine/Plugins/Mutable/Source/MutableTools/Private/MuT/ASTOpImageParameterConvert.h
--- a/Engine/Plugins/Mutable/Source/MutableTools/Private/MuT/ASTOpImageParameterConvert.h
+++ b/Engine/Plugins/Mutable/Source/MutableTools/Private/MuT/ASTOpImageParameterConvert.h
@@ -20,6 +20,8 @@ namespace UE::Mutable::Private
 		EOpType Type = EOpType::IM_PARAMETER_CONVERT;
 
 		ASTChild ImageParameter;
+
+		FImageDesc ImageDescriptor;
 		
 		int32 LinkedParameterIndex = -1;
 
diff --git a/Engine/Plugins/Mutable/Source/MutableTools/Private/MuT/ASTOpImagePatch.cpp b/Engine/Plugins/Mutable/Source/MutableTools/Private/MuT/ASTOpImagePatch.cpp
--- a/Engine/Plugins/Mutable/Source/MutableTools/Private/MuT/ASTOpImagePatch.cpp
+++ b/Engine/Plugins/Mutable/Source/MutableTools/Private/MuT/ASTOpImagePatch.cpp
@@ -93,20 +93,20 @@ namespace UE::Mutable::Private
 
 
 	//-------------------------------------------------------------------------------------------------
-	FImageDesc ASTOpImagePatch::GetImageDesc(bool returnBestOption, FGetImageDescContext* context) const
+	FImageDesc ASTOpImagePatch::GetImageDesc(bool returnBestOption, FGetImageDescContext* Context) const
 	{
-		FImageDesc res;
+		FImageDesc Result;
 
 		// Local context in case it is necessary
-		FGetImageDescContext localContext;
-		if (!context)
+		FGetImageDescContext LocalContext;
+		if (!Context)
 		{
-			context = &localContext;
+			Context = &LocalContext;
 		}
 		else
 		{
 			// Cached result?
-			FImageDesc* PtrValue = context->m_results.Find(this);
+			FImageDesc* PtrValue = Context->m_results.Find(this);
 			if (PtrValue)
 			{
 				return *PtrValue;
@@ -116,17 +116,13 @@ namespace UE::Mutable::Private
 		// Actual work
 		if (base)
 		{
-			res = base->GetImageDesc(returnBestOption, context);
+			Result = base->GetImageDesc(returnBestOption, Context);
 		}
 
-
 		// Cache the result
-		if (context)
-		{
-			context->m_results.Add(this, res);
-		}
+		Context->m_results.Add(this, Result);
 
-		return res;
+		return Result;
 	}
 
 
diff --git a/Engine/Plugins/Mutable/Source/MutableTools/Private/MuT/ASTOpImagePixelFormat.cpp b/Engine/Plugins/Mutable/Source/MutableTools/Private/MuT/ASTOpImagePixelFormat.cpp
--- a/Engine/Plugins/Mutable/Source/MutableTools/Private/MuT/ASTOpImagePixelFormat.cpp
+++ b/Engine/Plugins/Mutable/Source/MutableTools/Private/MuT/ASTOpImagePixelFormat.cpp
@@ -107,11 +107,14 @@ namespace UE::Mutable::Private
 	{
 		UE::Mutable::Private::Ptr<ASTOp> NewOp;
 
-		// Skip this operation if the source op format is already the one we want.
+		// Skip this operation if the source op format if fixed and is already the one we want.
 		if (Source)
 		{
-			FImageDesc SourceDesc = Source->GetImageDesc();
-			if (SourceDesc.m_format==Format)
+			FGetImageDescContext Context;
+			FImageDesc SourceDesc = Source->GetImageDesc(false, &Context);
+
+			const bool bIsFixedFormat = SourceDesc.m_format == SourceDesc.FormatIfAlpha || SourceDesc.FormatIfAlpha == EImageFormat::None;
+			if (!Context.bIsFormatFromImageParameter && bIsFixedFormat && SourceDesc.m_format == Format)
 			{
 				NewOp = Source.child();
 			}
@@ -271,8 +274,7 @@ namespace UE::Mutable::Private
 			res = Source.child()->GetImageDesc(returnBestOption, context);
 		}
 
-		if (FormatIfAlpha != EImageFormat::None
-			&&
+		if (FormatIfAlpha != EImageFormat::None &&
 			GetImageFormatData(res.m_format).Channels > 3)
 		{
 			res.m_format = FormatIfAlpha;
@@ -280,15 +282,24 @@ namespace UE::Mutable::Private
 		else
 		{
 			res.m_format = Format;
+			check(res.m_format != EImageFormat::None);
 		}
-		check(res.m_format != EImageFormat::None);
-
 
-		// Cache the result
-		if (context)
+		if (FormatIfAlpha != EImageFormat::None &&
+			res.FormatIfAlpha != EImageFormat::None)
 		{
-			context->m_results.Add(this, res);
+			res.FormatIfAlpha = FormatIfAlpha;
 		}
+		else
+		{
+			res.FormatIfAlpha = EImageFormat::None;
+		}
+
+		// Format overriden by current op.
+		context->bIsFormatFromImageParameter = false;
+
+		// Cache the result
+		context->m_results.Add(this, res);
 
 		return res;
 	}
diff --git a/Engine/Plugins/Mutable/Source/MutableTools/Private/MuT/ASTOpImagePlainColor.cpp b/Engine/Plugins/Mutable/Source/MutableTools/Private/MuT/ASTOpImagePlainColor.cpp
--- a/Engine/Plugins/Mutable/Source/MutableTools/Private/MuT/ASTOpImagePlainColor.cpp
+++ b/Engine/Plugins/Mutable/Source/MutableTools/Private/MuT/ASTOpImagePlainColor.cpp
@@ -130,10 +130,7 @@ namespace UE::Mutable::Private
 		check(Result.m_format != EImageFormat::None);
 
 		// Cache the result
-		if (Context)
-		{
-			Context->m_results.Add(this, Result);
-		}
+		Context->m_results.Add(this, Result);
 
 		return Result;
 	}
diff --git a/Engine/Plugins/Mutable/Source/MutableTools/Private/MuT/ASTOpImageRasterMesh.cpp b/Engine/Plugins/Mutable/Source/MutableTools/Private/MuT/ASTOpImageRasterMesh.cpp
--- a/Engine/Plugins/Mutable/Source/MutableTools/Private/MuT/ASTOpImageRasterMesh.cpp
+++ b/Engine/Plugins/Mutable/Source/MutableTools/Private/MuT/ASTOpImageRasterMesh.cpp
@@ -170,20 +170,20 @@ namespace UE::Mutable::Private
 
 
 	//-------------------------------------------------------------------------------------------------
-	FImageDesc ASTOpImageRasterMesh::GetImageDesc(bool returnBestOption, FGetImageDescContext* context) const 
+	FImageDesc ASTOpImageRasterMesh::GetImageDesc(bool returnBestOption, FGetImageDescContext* Context) const
 	{
-		FImageDesc res;
+		FImageDesc Result;
 
 		// Local context in case it is necessary
-		FGetImageDescContext localContext;
-		if (!context)
+		FGetImageDescContext LocalContext;
+		if (!Context)
 		{
-			context = &localContext;
+			Context = &LocalContext;
 		}
 		else
 		{
 			// Cached result?
-			FImageDesc* PtrValue = context->m_results.Find(this);
+			FImageDesc* PtrValue = Context->m_results.Find(this);
 			if (PtrValue)
 			{
 				return *PtrValue;
@@ -193,19 +193,16 @@ namespace UE::Mutable::Private
 		// Actual work
 		if (image)
 		{
-			res = image->GetImageDesc( returnBestOption, context);
-			res.m_size[0] = SizeX;
-			res.m_size[1] = SizeY;
+			Result = image->GetImageDesc( returnBestOption, Context);
+			Result.m_size[0] = SizeX;
+			Result.m_size[1] = SizeY;
 		}
 
 
 		// Cache the result
-		if (context)
-		{
-			context->m_results.Add(this, res);
-		}
+		Context->m_results.Add(this, Result);
 
-		return res;
+		return Result;
 	}
 
 
diff --git a/Engine/Plugins/Mutable/Source/MutableTools/Private/MuT/ASTOpImageResize.cpp b/Engine/Plugins/Mutable/Source/MutableTools/Private/MuT/ASTOpImageResize.cpp
--- a/Engine/Plugins/Mutable/Source/MutableTools/Private/MuT/ASTOpImageResize.cpp
+++ b/Engine/Plugins/Mutable/Source/MutableTools/Private/MuT/ASTOpImageResize.cpp
@@ -130,10 +130,7 @@ namespace UE::Mutable::Private
 		Result.m_size[1] = Size[1];
 
 		// Cache the result
-		if (Context)
-		{
-			Context->m_results.Add(this, Result);
-		}
+		Context->m_results.Add(this, Result);
 
 		return Result;
 	}
diff --git a/Engine/Plugins/Mutable/Source/MutableTools/Private/MuT/ASTOpImageSwizzle.cpp b/Engine/Plugins/Mutable/Source/MutableTools/Private/MuT/ASTOpImageSwizzle.cpp
--- a/Engine/Plugins/Mutable/Source/MutableTools/Private/MuT/ASTOpImageSwizzle.cpp
+++ b/Engine/Plugins/Mutable/Source/MutableTools/Private/MuT/ASTOpImageSwizzle.cpp
@@ -1049,11 +1049,11 @@ namespace UE::Mutable::Private
 			check(res.m_format != EImageFormat::None);
 		}
 
+		// Format overriden by current op.
+		context->bIsFormatFromImageParameter = false;
+
 		// Cache the result
-		if (context)
-		{
-			context->m_results.Add(this, res);
-		}
+		context->m_results.Add(this, res);
 
 		return res;
 	}
diff --git a/Engine/Plugins/Mutable/Source/MutableTools/Private/MuT/ASTOpImageTransform.cpp b/Engine/Plugins/Mutable/Source/MutableTools/Private/MuT/ASTOpImageTransform.cpp
--- a/Engine/Plugins/Mutable/Source/MutableTools/Private/MuT/ASTOpImageTransform.cpp
+++ b/Engine/Plugins/Mutable/Source/MutableTools/Private/MuT/ASTOpImageTransform.cpp
@@ -165,6 +165,7 @@ namespace UE::Mutable::Private
 			Result = Base->GetImageDesc(bReturnBestOption, Context);
 			
 			Result.m_format = GetUncompressedFormat(Result.m_format); 
+			Result.FormatIfAlpha = GetUncompressedFormat(Result.FormatIfAlpha); 
 			Result.m_lods = 1;
 			
 			if (!(SizeX == 0 && SizeY == 0))
@@ -173,12 +174,11 @@ namespace UE::Mutable::Private
 			}
 		}
 
+		// Format overriden by current op.
+		Context->bIsFormatFromImageParameter = false;
 
 		// Cache the result
-		if (Context)
-		{
-			Context->m_results.Add(this, Result);
-		}
+		Context->m_results.Add(this, Result);
 
 		return Result;
 	}
diff --git a/Engine/Plugins/Mutable/Source/MutableTools/Private/MuT/ASTOpSwitch.cpp b/Engine/Plugins/Mutable/Source/MutableTools/Private/MuT/ASTOpSwitch.cpp
--- a/Engine/Plugins/Mutable/Source/MutableTools/Private/MuT/ASTOpSwitch.cpp
+++ b/Engine/Plugins/Mutable/Source/MutableTools/Private/MuT/ASTOpSwitch.cpp
@@ -299,12 +299,14 @@ namespace UE::Mutable::Private
 		bool bSameFormat = true;
 		bool bSameLods = true;
 		bool bFirst = true;
+		bool bIsFormatFromImageParameter = false;
 
 		if (Default)
 		{
 			FImageDesc ChildDesc = Default->GetImageDesc(bReturnBestOption, Context);
 			Candidate = ChildDesc;
 			bFirst = false;
+			bIsFormatFromImageParameter = Context->bIsFormatFromImageParameter;
 		}
 
 		for (int32 CaseIndex = 0; CaseIndex < Cases.Num(); ++CaseIndex)
@@ -312,6 +314,8 @@ namespace UE::Mutable::Private
 			if (Cases[CaseIndex].Branch)
 			{
 				FImageDesc ChildDesc = Cases[CaseIndex].Branch->GetImageDesc(bReturnBestOption, Context);
+				bIsFormatFromImageParameter = bIsFormatFromImageParameter || Context->bIsFormatFromImageParameter;
+
 				if (bFirst)
 				{
 					Candidate = ChildDesc;
@@ -321,11 +325,13 @@ namespace UE::Mutable::Private
 				{
 					bSameSize = bSameSize && (Candidate.m_size == ChildDesc.m_size);
 					bSameFormat = bSameFormat && (Candidate.m_format == ChildDesc.m_format);
+					bSameFormat = bSameFormat && (Candidate.FormatIfAlpha == ChildDesc.FormatIfAlpha);
 					bSameLods = bSameLods && (Candidate.m_lods == ChildDesc.m_lods);
 
 					if (bReturnBestOption)
 					{
 						Candidate.m_format = GetMostGenericFormat(Candidate.m_format, ChildDesc.m_format);
+						Candidate.FormatIfAlpha = GetMostGenericFormat(Candidate.FormatIfAlpha, ChildDesc.FormatIfAlpha);
 
 						// Return the biggest size
 						Candidate.m_size[0] = FMath::Max(Candidate.m_size[0], ChildDesc.m_size[0]);
@@ -342,6 +348,7 @@ namespace UE::Mutable::Private
 		if (!bSameFormat && !bReturnBestOption)
 		{
 			Result.m_format = EImageFormat::None;
+			Result.FormatIfAlpha = EImageFormat::None;
 		}
 
 		if (!bSameSize && !bReturnBestOption)
@@ -354,11 +361,11 @@ namespace UE::Mutable::Private
 			Result.m_lods = 0;
 		}
 
+		// Format overriden by current op.
+		Context->bIsFormatFromImageParameter = bIsFormatFromImageParameter;
+
 		// Cache the result
-		if (Context)
-		{
-			Context->m_results.Add(this, Result);
-		}
+		Context->m_results.Add(this, Result);
 
 		return Result;
 	}
diff --git a/Engine/Plugins/Mutable/Source/MutableTools/Private/MuT/CodeGenerator.h b/Engine/Plugins/Mutable/Source/MutableTools/Private/MuT/CodeGenerator.h
--- a/Engine/Plugins/Mutable/Source/MutableTools/Private/MuT/CodeGenerator.h
+++ b/Engine/Plugins/Mutable/Source/MutableTools/Private/MuT/CodeGenerator.h
@@ -586,9 +586,6 @@ namespace UE::Mutable::Private
 		/** */
 		Ptr<ASTOp> GenerateImageFormat(Ptr<ASTOp>, EImageFormat);
 
-		/** */
-		Ptr<ASTOp> GenerateImageUncompressed(Ptr<ASTOp>);
-
 		/** */
 		Ptr<ASTOp> GenerateImageSize(Ptr<ASTOp>, UE::Math::TIntVector2<int32>);
 
diff --git a/Engine/Plugins/Mutable/Source/MutableTools/Private/MuT/CodeGenerator_Image.cpp b/Engine/Plugins/Mutable/Source/MutableTools/Private/MuT/CodeGenerator_Image.cpp
--- a/Engine/Plugins/Mutable/Source/MutableTools/Private/MuT/CodeGenerator_Image.cpp
+++ b/Engine/Plugins/Mutable/Source/MutableTools/Private/MuT/CodeGenerator_Image.cpp
@@ -253,7 +253,7 @@ namespace UE::Mutable::Private
 			// channel if not used.
 			// TODO: review this, probably the reference descriptor generation needs to check for alpha channels 
 			// as well.
-			ReferenceOp->ImageDesc = FImageDesc(pImage->GetSize(), EImageFormat::None, pImage->GetLODCount());
+			ReferenceOp->ImageDesc = FImageDesc(pImage->GetSize(), EImageFormat::None, EImageFormat::None, pImage->GetLODCount());
 			Result.op = ReferenceOp;
 		}
 		else
@@ -855,14 +855,27 @@ namespace UE::Mutable::Private
 				GenerateImage(Options, BaseResult, Node->Sources[SourceIndex]);
                 Ptr<ASTOp> Source = BaseResult.op;
 
-				Source = GenerateImageUncompressed(Source);
-
 				if (!Source)
 				{
 					// TODO: Warn?
 					Source = GenerateMissingImageCode(TEXT("Swizzle channel"), EImageFormat::L_UByte, Node->GetMessageContext(), Options);
 				}
 
+				const FImageDesc SourceImageDesc = Source->GetImageDesc();
+
+				const bool bUseFormatWithAlpha = Node->SourceChannels[SourceIndex] == 3 && SourceImageDesc.FormatIfAlpha != EImageFormat::None;
+				const EImageFormat SourceFormat = bUseFormatWithAlpha  ? SourceImageDesc.FormatIfAlpha : SourceImageDesc.m_format;
+				const EImageFormat TargetFormat = GetUncompressedFormat(SourceFormat);
+
+				if (TargetFormat != SourceFormat)
+				{
+					// Generate uncompressed format
+					Ptr<ASTOpImagePixelFormat> FormatOp = new ASTOpImagePixelFormat();
+					FormatOp->Source = Source;
+					FormatOp->Format = TargetFormat;
+					Source = FormatOp;
+				}
+				
                 Ptr<ASTOp> SizedSource;
 				if (FirstValid && FirstValidDesc.m_size[0])
 				{
@@ -1946,6 +1959,7 @@ namespace UE::Mutable::Private
 		GenerateImage(InOptions, PassthroughImageResult, InNode->ImageParameter);
 
 		ReferenceConvert->ImageParameter = PassthroughImageResult.op;
+		ReferenceConvert->ImageDescriptor = InNode->ImageDescriptor;
 		
 		Result.op = ReferenceConvert;
 		Result.Parameter = PassthroughImageResult.Parameter;
@@ -2091,29 +2105,6 @@ namespace UE::Mutable::Private
     }
 
 
-    Ptr<ASTOp> CodeGenerator::GenerateImageUncompressed( Ptr<ASTOp> at )
-    {
-        Ptr<ASTOp> Result = at;
-
-        if (at)
-        {
-            EImageFormat sourceFormat = at->GetImageDesc().m_format;
-            EImageFormat targetFormat = GetUncompressedFormat( sourceFormat );
-
-            if ( targetFormat != sourceFormat )
-            {
-                // Generate the format change code
-                Ptr<ASTOpImagePixelFormat> op = new ASTOpImagePixelFormat();
-                op->Source = at;
-                op->Format = targetFormat;
-                Result = op;
-            }
-        }
-
-        return Result;
-    }
-
-
     Ptr<ASTOp> CodeGenerator::GenerateImageSize( Ptr<ASTOp> at, UE::Math::TIntVector2<int32> size )
     {
         Ptr<ASTOp> Result = at;
