[Mutable] Remove disk compilation

#rnx
#Synced-CL 48501169

diff --git a/Engine/Plugins/Mutable/Source/CustomizableObject/Internal/MuCO/CustomizableObjectCompilerTypes.h b/Engine/Plugins/Mutable/Source/CustomizableObject/Internal/MuCO/CustomizableObjectCompilerTypes.h
--- a/Engine/Plugins/Mutable/Source/CustomizableObject/Internal/MuCO/CustomizableObjectCompilerTypes.h
+++ b/Engine/Plugins/Mutable/Source/CustomizableObject/Internal/MuCO/CustomizableObjectCompilerTypes.h
@@ -25,13 +25,6 @@ struct FCompilationOptions
 	// From 0 to UE_MUTABLE_MAX_OPTIMIZATION
 	int32 OptimizationLevel = UE_MUTABLE_MAX_OPTIMIZATION;
 
-	/** Use the disk to store intermediate compilation data.This slows down the object compilation
-	 * but it may be necessary for huge objects.
-	 *
-	 * This option does not modify the DDC Key.
-	 * */
-	bool bUseDiskCompilation = false;
-
 	/** High limit of the size in bytes of the packaged data when cooking this object.
 	 * This limit is before any pak or filesystem compression. This limit will be broken if a single piece of data is bigger because data is not fragmented for packaging purposes.
 	 *
diff --git a/Engine/Plugins/Mutable/Source/CustomizableObject/Internal/MuCO/CustomizableObjectPrivate.h b/Engine/Plugins/Mutable/Source/CustomizableObject/Internal/MuCO/CustomizableObjectPrivate.h
--- a/Engine/Plugins/Mutable/Source/CustomizableObject/Internal/MuCO/CustomizableObjectPrivate.h
+++ b/Engine/Plugins/Mutable/Source/CustomizableObject/Internal/MuCO/CustomizableObjectPrivate.h
@@ -2026,11 +2026,6 @@ public:
 	UPROPERTY()
 	int32 OptimizationLevel = UE_MUTABLE_MAX_OPTIMIZATION;
 
-	/** Use the disk to store intermediate compilation data. This slows down the object compilation
-	 * but it may be necessary for huge objects. */
-	UPROPERTY()
-	bool bUseDiskCompilation = false;
-
 	/** High limit of the size in bytes of the packaged data when cooking this object.
 	 * This limit is before any pak or filesystem compression. This limit will be broken if a single piece of data is bigger because data is not fragmented for packaging purposes.	*/
 	UPROPERTY()
diff --git a/Engine/Plugins/Mutable/Source/CustomizableObject/Private/MuCO/CustomizableObject.cpp b/Engine/Plugins/Mutable/Source/CustomizableObject/Private/MuCO/CustomizableObject.cpp
--- a/Engine/Plugins/Mutable/Source/CustomizableObject/Private/MuCO/CustomizableObject.cpp
+++ b/Engine/Plugins/Mutable/Source/CustomizableObject/Private/MuCO/CustomizableObject.cpp
@@ -392,7 +392,6 @@ void UCustomizableObjectPrivate::BackwardsCompatibleFixup(int32 CustomizableObje
 	{
 		OptimizationLevel = GetPublic()->CompileOptions_DEPRECATED.OptimizationLevel;
 		TextureCompression = GetPublic()->CompileOptions_DEPRECATED.TextureCompression;
-		bUseDiskCompilation = GetPublic()->CompileOptions_DEPRECATED.bUseDiskCompilation;
 		EmbeddedDataBytesLimit = GetPublic()->CompileOptions_DEPRECATED.EmbeddedDataBytesLimit;
 		PackagedDataBytesLimit = GetPublic()->CompileOptions_DEPRECATED.PackagedDataBytesLimit;
 	}
@@ -2436,7 +2435,6 @@ FCompilationOptions UCustomizableObjectPrivate::GetCompileOptions() const
 	FCompilationOptions Options;
 	Options.TextureCompression = TextureCompression;
 	Options.OptimizationLevel = OptimizationLevel;
-	Options.bUseDiskCompilation = bUseDiskCompilation;
 
 	Options.TargetPlatform = GetTargetPlatformManagerRef().GetRunningTargetPlatform();
 
diff --git a/Engine/Plugins/Mutable/Source/CustomizableObject/Public/MuCO/CustomizableObject_Deprecated.h b/Engine/Plugins/Mutable/Source/CustomizableObject/Public/MuCO/CustomizableObject_Deprecated.h
--- a/Engine/Plugins/Mutable/Source/CustomizableObject/Public/MuCO/CustomizableObject_Deprecated.h
+++ b/Engine/Plugins/Mutable/Source/CustomizableObject/Public/MuCO/CustomizableObject_Deprecated.h
@@ -18,9 +18,6 @@ struct FCompilationOptions_DEPRECATED
 	UPROPERTY()
 	int32 OptimizationLevel = 2; // UE_MUTABLE_MAX_OPTIMIZATION;
 
-	UPROPERTY()
-	bool bUseDiskCompilation = false;
-
 	UPROPERTY()
 	uint64 PackagedDataBytesLimit = 256 * 1024 * 1024;
 
diff --git a/Engine/Plugins/Mutable/Source/CustomizableObjectEditor/Private/MuCOE/CustomizableObjectCompileRunnable.cpp b/Engine/Plugins/Mutable/Source/CustomizableObjectEditor/Private/MuCOE/CustomizableObjectCompileRunnable.cpp
--- a/Engine/Plugins/Mutable/Source/CustomizableObjectEditor/Private/MuCOE/CustomizableObjectCompileRunnable.cpp
+++ b/Engine/Plugins/Mutable/Source/CustomizableObjectEditor/Private/MuCOE/CustomizableObjectCompileRunnable.cpp
@@ -37,13 +37,6 @@
 
 #define UE_MUTABLE_CORE_REGION	TEXT("Mutable Core")
 
-
-TAutoConsoleVariable<bool> CVarMutableCompilerDiskCache(
-	TEXT("mutable.ForceCompilerDiskCache"),
-	false,
-	TEXT("Force the use of disk cache to reduce memory usage when compiling CustomizableObjects both in editor and cook commandlets."),
-	ECVF_Default);
-
 TAutoConsoleVariable<bool> CVarMutableCompilerFastCompression(
 	TEXT("mutable.ForceFastTextureCompression"),
 	false,
@@ -103,14 +96,6 @@ uint32 FCustomizableObjectCompileRunnable::Run()
 	// Translate CO compile options into UE::Mutable::Private::CompilerOptions
 	UE::Mutable::Private::Ptr<UE::Mutable::Private::CompilerOptions> CompilerOptions = new UE::Mutable::Private::CompilerOptions();
 
-	bool bUseDiskCache = Options.bUseDiskCompilation;
-	if (CVarMutableCompilerDiskCache->GetBool())
-	{
-		bUseDiskCache = true;
-	}
-
-	CompilerOptions->SetUseDiskCache(bUseDiskCache);
-
 	if (Options.OptimizationLevel > UE_MUTABLE_MAX_OPTIMIZATION)
 	{
 		UE_LOG(LogMutable, Log, TEXT("Mutable compile optimization level out of range. Clamping to maximum."));
@@ -401,8 +386,6 @@ uint32 FCustomizableObjectCompileRunnable::Run()
 
 	UE_LOG(LogMutable, Verbose, TEXT("PROFILE: [ %16.8f ] FCustomizableObjectCompileRunnable::Run end."), FPlatformTime::Seconds());
 
-	CompilerOptions->LogStats();
-
 	TRACE_END_REGION(UE_MUTABLE_CORE_REGION);
 
 	return Result;
diff --git a/Engine/Plugins/Mutable/Source/CustomizableObjectEditor/Private/MuCOE/CustomizableObjectEditor.cpp b/Engine/Plugins/Mutable/Source/CustomizableObjectEditor/Private/MuCOE/CustomizableObjectEditor.cpp
--- a/Engine/Plugins/Mutable/Source/CustomizableObjectEditor/Private/MuCOE/CustomizableObjectEditor.cpp
+++ b/Engine/Plugins/Mutable/Source/CustomizableObjectEditor/Private/MuCOE/CustomizableObjectEditor.cpp
@@ -480,12 +480,6 @@ void FCustomizableObjectEditor::BindCommands()
 		FCanExecuteAction(),
 		FIsActionChecked());
 
-	ToolkitCommands->MapAction(
-		Commands.CompileOptions_UseDiskCompilation,
-		FExecuteAction::CreateSP(this, &FCustomizableObjectEditor::CompileOptions_UseDiskCompilation_Toggled),
-		FCanExecuteAction(),
-		FIsActionChecked::CreateSP(this, &FCustomizableObjectEditor::CompileOptions_UseDiskCompilation_IsChecked));
-
 	// Debug and options
 	ToolkitCommands->MapAction(
 		Commands.GraphViewer,
@@ -1137,8 +1131,6 @@ TSharedRef<SWidget> FCustomizableObjectEditor::GenerateCompileOptionsMenuContent
 					}
 				});
 		MenuBuilder.AddWidget(CompileTilingCombo.ToSharedRef(), LOCTEXT("MutableCompileImageTiling", "Image Tiling"));
-
-		MenuBuilder.AddMenuEntry(FCustomizableObjectEditorCommands::Get().CompileOptions_UseDiskCompilation);
 	}
 	MenuBuilder.EndSection();
 
@@ -1555,7 +1547,6 @@ void FCustomizableObjectEditor::ResetCompileOptions()
 	
 	UCustomizableObjectPrivate* DefaultObject = Cast<UCustomizableObjectPrivate>(CustomizableObject->GetPrivate()->StaticClass()->GetDefaultObject());
 	CustomizableObject->GetPrivate()->OptimizationLevel = DefaultObject->OptimizationLevel;
-	CustomizableObject->GetPrivate()->bUseDiskCompilation = DefaultObject->bUseDiskCompilation;
 	CustomizableObject->GetPrivate()->TextureCompression = DefaultObject->TextureCompression;
 	CustomizableObject->GetPrivate()->EmbeddedDataBytesLimit = DefaultObject->EmbeddedDataBytesLimit;
 	CustomizableObject->GetPrivate()->PackagedDataBytesLimit = DefaultObject->PackagedDataBytesLimit;
@@ -1570,20 +1561,6 @@ void FCustomizableObjectEditor::OnChangeCompileOptimizationLevel(TSharedPtr<FStr
 }
 
 
-void FCustomizableObjectEditor::CompileOptions_UseDiskCompilation_Toggled()
-{
-	const FScopedTransaction Transaction(LOCTEXT("ChangedEnableCompilingUsingTheDiskAsMemory", "Changed Enable compiling using the disk as memory"));
-	CustomizableObject->Modify();
-	CustomizableObject->GetPrivate()->bUseDiskCompilation = !CustomizableObject->GetPrivate()->bUseDiskCompilation;
-}
-
-
-bool FCustomizableObjectEditor::CompileOptions_UseDiskCompilation_IsChecked()
-{
-	return CustomizableObject->GetPrivate()->bUseDiskCompilation;
-}
-
-
 void FCustomizableObjectEditor::OnChangeCompileTextureCompressionType(TSharedPtr<FString> NewSelection, ESelectInfo::Type)
 {
 	const FScopedTransaction Transaction(LOCTEXT("ChangedTextureCompressionTransaction", "Changed Texture Compression Type"));
diff --git a/Engine/Plugins/Mutable/Source/CustomizableObjectEditor/Private/MuCOE/CustomizableObjectEditor.h b/Engine/Plugins/Mutable/Source/CustomizableObjectEditor/Private/MuCOE/CustomizableObjectEditor.h
--- a/Engine/Plugins/Mutable/Source/CustomizableObjectEditor/Private/MuCOE/CustomizableObjectEditor.h
+++ b/Engine/Plugins/Mutable/Source/CustomizableObjectEditor/Private/MuCOE/CustomizableObjectEditor.h
@@ -217,8 +217,6 @@ private:
 	TSharedPtr<SNumericDropDown<float>> EmbeddedDataLimitCombo;
 	TSharedPtr<SNumericDropDown<float>> PackagedDataLimitCombo;
 
-	void CompileOptions_UseDiskCompilation_Toggled();
-	bool CompileOptions_UseDiskCompilation_IsChecked();
 	void OnChangeCompileOptimizationLevel(TSharedPtr<FString> NewSelection, ESelectInfo::Type);
 	void OnChangeCompileTextureCompressionType(TSharedPtr<FString> NewSelection, ESelectInfo::Type);
 
diff --git a/Engine/Plugins/Mutable/Source/CustomizableObjectEditor/Private/MuCOE/CustomizableObjectEditorActions.cpp b/Engine/Plugins/Mutable/Source/CustomizableObjectEditor/Private/MuCOE/CustomizableObjectEditorActions.cpp
--- a/Engine/Plugins/Mutable/Source/CustomizableObjectEditor/Private/MuCOE/CustomizableObjectEditorActions.cpp
+++ b/Engine/Plugins/Mutable/Source/CustomizableObjectEditor/Private/MuCOE/CustomizableObjectEditorActions.cpp
@@ -25,7 +25,6 @@ void FCustomizableObjectEditorCommands::RegisterCommands()
 	UI_COMMAND(Compile, "Compile", "Compile the source graph of the customizable object and update the previews. \nActive if the CVar Mutable.Enabled is set to true.", EUserInterfaceActionType::Button, FInputChord());
 	UI_COMMAND(CompileOnlySelected, "Compile Only Selected", "Compile the source graph of the customizable object and update the previews, only for the selected options in the preview. The rest of options are discarded. If they are selected, press again this button to see their effect in the preview. \nActive if the CVar Mutable.Enabled is set to true.", EUserInterfaceActionType::Button, FInputChord());
 	UI_COMMAND(ResetCompileOptions, "Reset Compilation Options", "Set reasonable defaults for the compilation options.", EUserInterfaceActionType::Button, FInputChord());
-	UI_COMMAND(CompileOptions_UseDiskCompilation, "Enable compiling using the disk as memory.", "This is very slow but supports compiling huge objects. It requires a lot of free space in the OS disk.", EUserInterfaceActionType::ToggleButton, FInputChord());
 	UI_COMMAND(GraphViewer, "Graph Viewer", "Open the Customizable Object Graph Viewer.", EUserInterfaceActionType::Button, FInputChord());
 	UI_COMMAND(CodeViewer, "Code Viewer", "Open the Customizable Object Code Viewer.", EUserInterfaceActionType::Button, FInputChord());
 	UI_COMMAND(UpdateCookDataDistributionId, "Regenerate Cooked Data Distribution", "Optimize cooked data distribution based on current graph and generated data. Data distribution is stored in the DDC. \n\nWarning! Changing the data distribution may result in big patch sizes.", EUserInterfaceActionType::Button, FInputChord());
diff --git a/Engine/Plugins/Mutable/Source/CustomizableObjectEditor/Private/MuCOE/CustomizableObjectEditorActions.h b/Engine/Plugins/Mutable/Source/CustomizableObjectEditor/Private/MuCOE/CustomizableObjectEditorActions.h
--- a/Engine/Plugins/Mutable/Source/CustomizableObjectEditor/Private/MuCOE/CustomizableObjectEditorActions.h
+++ b/Engine/Plugins/Mutable/Source/CustomizableObjectEditor/Private/MuCOE/CustomizableObjectEditorActions.h
@@ -17,7 +17,6 @@ public:
 	TSharedPtr< FUICommandInfo > Compile;
 	TSharedPtr< FUICommandInfo > CompileOnlySelected;
 	TSharedPtr< FUICommandInfo > ResetCompileOptions;
-	TSharedPtr< FUICommandInfo > CompileOptions_UseDiskCompilation;
 	TSharedPtr< FUICommandInfo > GraphViewer;
 	TSharedPtr< FUICommandInfo > CodeViewer;
 	TSharedPtr< FUICommandInfo > UpdateCookDataDistributionId;
diff --git a/Engine/Plugins/Mutable/Source/CustomizableObjectEditor/Private/MuCOE/GenerateMutableSource/GenerateMutableSourceGroupProjector.cpp b/Engine/Plugins/Mutable/Source/CustomizableObjectEditor/Private/MuCOE/GenerateMutableSource/GenerateMutableSourceGroupProjector.cpp
--- a/Engine/Plugins/Mutable/Source/CustomizableObjectEditor/Private/MuCOE/GenerateMutableSource/GenerateMutableSourceGroupProjector.cpp
+++ b/Engine/Plugins/Mutable/Source/CustomizableObjectEditor/Private/MuCOE/GenerateMutableSource/GenerateMutableSourceGroupProjector.cpp
@@ -498,10 +498,8 @@ TOptional<FGroupProjectorTempData> GenerateMutableGroupProjector(UCustomizableOb
 		{
 			if (const TObjectPtr<UTexture2D>& Texture = ArrayOptionTexture[SelectorIndex].OptionTexture)
 			{
-				TSharedPtr<UE::Mutable::Private::FImage> ImageConstant = GenerateImageConstant(ArrayOptionTexture[SelectorIndex].OptionTexture, GenerationContext, false);
-
 				UE::Mutable::Private::Ptr<UE::Mutable::Private::NodeImageConstant> ImageNode = new UE::Mutable::Private::NodeImageConstant();
-				ImageNode->SetValue(ImageConstant);
+				ImageNode->Image = GenerateImageConstant(ArrayOptionTexture[SelectorIndex].OptionTexture, GenerationContext, false);
 
 				const uint32 MipsToSkip = ComputeLODBiasForTexture(GenerationContext, *Texture, ProjParamNode->ReferenceTexture);
 				SwitchNode->Options[SelectorIndex] = ResizeTextureByNumMips(ImageNode, MipsToSkip);
diff --git a/Engine/Plugins/Mutable/Source/CustomizableObjectEditor/Private/MuCOE/GenerateMutableSource/GenerateMutableSourceImage.cpp b/Engine/Plugins/Mutable/Source/CustomizableObjectEditor/Private/MuCOE/GenerateMutableSource/GenerateMutableSourceImage.cpp
--- a/Engine/Plugins/Mutable/Source/CustomizableObjectEditor/Private/MuCOE/GenerateMutableSource/GenerateMutableSourceImage.cpp
+++ b/Engine/Plugins/Mutable/Source/CustomizableObjectEditor/Private/MuCOE/GenerateMutableSource/GenerateMutableSourceImage.cpp
@@ -139,10 +139,8 @@ UE::Mutable::Private::Ptr<UE::Mutable::Private::NodeImage> GenerateMutableSource
 			}
 			else
 			{
-				TSharedPtr<UE::Mutable::Private::FImage> ImageConstant = GenerateImageConstant(BaseTexture, GenerationContext, false);
-
 				UE::Mutable::Private::Ptr<UE::Mutable::Private::NodeImageConstant> ConstantImageNode = new UE::Mutable::Private::NodeImageConstant;
-				ConstantImageNode->SetValue( ImageConstant );
+				ConstantImageNode->Image = GenerateImageConstant(BaseTexture, GenerationContext, false);
 				
 				const UTextureLODSettings& TextureLODSettings = GenerationContext.CompilationContext->Options.TargetPlatform->GetTextureLODSettings();
 				const FTextureLODGroup& LODGroupInfo = TextureLODSettings.GetTextureLODGroup(BaseTexture->LODGroup);
@@ -229,8 +227,7 @@ UE::Mutable::Private::Ptr<UE::Mutable::Private::NodeImage> GenerateMutableSource
 
 		if (UTexture2D* Texture = TypedNodeMesh->FindTextureForPin(Pin))
 		{
-			TSharedPtr<UE::Mutable::Private::FImage> ImageConstant = GenerateImageConstant(Texture, GenerationContext, false);
-			ImageNode->SetValue(ImageConstant);
+			ImageNode->Image = GenerateImageConstant(Texture, GenerationContext, false);
 
 			const uint32 MipsToSkip = ComputeLODBiasForTexture(GenerationContext, *Texture, nullptr, ReferenceTextureSize);
 			Result = ResizeTextureByNumMips(ImageNode, MipsToSkip);
@@ -907,9 +904,9 @@ UE::Mutable::Private::Ptr<UE::Mutable::Private::NodeImage> GenerateMutableSource
 		if (BaseTexture)
 		{
 			UE::Mutable::Private::Ptr<UE::Mutable::Private::NodeImageConstant> ImageNode = new UE::Mutable::Private::NodeImageConstant();
-			Result = ImageNode;
+			ImageNode->Image = GenerateImageConstant(BaseTexture, GenerationContext, true);
 
-			ImageNode->SetValue(GenerateImageConstant(BaseTexture, GenerationContext, true));
+			Result = ImageNode;
 		}
 		else
 		{
diff --git a/Engine/Plugins/Mutable/Source/CustomizableObjectEditor/Private/MuCOE/GenerateMutableSource/GenerateMutableSourceMaterial.cpp b/Engine/Plugins/Mutable/Source/CustomizableObjectEditor/Private/MuCOE/GenerateMutableSource/GenerateMutableSourceMaterial.cpp
--- a/Engine/Plugins/Mutable/Source/CustomizableObjectEditor/Private/MuCOE/GenerateMutableSource/GenerateMutableSourceMaterial.cpp
+++ b/Engine/Plugins/Mutable/Source/CustomizableObjectEditor/Private/MuCOE/GenerateMutableSource/GenerateMutableSourceMaterial.cpp
@@ -113,11 +113,8 @@ UE::Mutable::Private::Ptr<UE::Mutable::Private::NodeMaterial> GenerateMutableSou
 
 						if (bParameterFound && BaseTexture)
 						{
-							TSharedPtr<UE::Mutable::Private::FImage> ImageConstant = GenerateImageConstant(BaseTexture, GenerationContext, false);
-							const uint32 MipsToSkip = ComputeLODBiasForTexture(GenerationContext, *BaseTexture);
-
 							UE::Mutable::Private::Ptr<UE::Mutable::Private::NodeImageConstant> ConstantImageNode = new UE::Mutable::Private::NodeImageConstant;
-							ConstantImageNode->SetValue(ImageConstant);
+							ConstantImageNode->Image = GenerateImageConstant(BaseTexture, GenerationContext, false);
 
 							const UTextureLODSettings& TextureLODSettings = GenerationContext.CompilationContext->Options.TargetPlatform->GetTextureLODSettings();
 							const FTextureLODGroup& LODGroupInfo = TextureLODSettings.GetTextureLODGroup(BaseTexture->LODGroup);
@@ -130,6 +127,7 @@ UE::Mutable::Private::Ptr<UE::Mutable::Private::NodeMaterial> GenerateMutableSou
 							const FString TextureName = GetNameSafe(BaseTexture).ToLower();
 							ConstantImageNode->SourceDataDescriptor.SourceId = CityHash32(reinterpret_cast<const char*>(*TextureName), TextureName.Len() * sizeof(FString::ElementType));
 
+							const uint32 MipsToSkip = ComputeLODBiasForTexture(GenerationContext, *BaseTexture);
 							UE::Mutable::Private::Ptr<UE::Mutable::Private::NodeImage> FinalNodeImag = ResizeTextureByNumMips(ConstantImageNode, MipsToSkip);
 
 							if (FinalNodeImag)
diff --git a/Engine/Plugins/Mutable/Source/CustomizableObjectEditor/Private/MuCOE/GenerateMutableSource/GenerateMutableSourceSurface.cpp b/Engine/Plugins/Mutable/Source/CustomizableObjectEditor/Private/MuCOE/GenerateMutableSource/GenerateMutableSourceSurface.cpp
--- a/Engine/Plugins/Mutable/Source/CustomizableObjectEditor/Private/MuCOE/GenerateMutableSource/GenerateMutableSourceSurface.cpp
+++ b/Engine/Plugins/Mutable/Source/CustomizableObjectEditor/Private/MuCOE/GenerateMutableSource/GenerateMutableSourceSurface.cpp
@@ -416,8 +416,7 @@ UE::Mutable::Private::Ptr<UE::Mutable::Private::NodeSurface> GenerateMutableSour
 									if (Texture2D)
 									{
 										UE::Mutable::Private::Ptr<UE::Mutable::Private::NodeImageConstant> ConstImageNode = new UE::Mutable::Private::NodeImageConstant();
-										TSharedPtr<UE::Mutable::Private::FImage> ImageConstant = GenerateImageConstant(Texture2D, GenerationContext, false);
-										ConstImageNode->SetValue(ImageConstant);
+										ConstImageNode->Image = GenerateImageConstant(Texture2D, GenerationContext, false);
 
 										const uint32 MipsToSkip = ComputeLODBiasForTexture(GenerationContext, *Texture2D, nullptr, Props.TextureSize);
 										UE::Mutable::Private::Ptr<UE::Mutable::Private::NodeImage> Result =  ResizeTextureByNumMips(ConstImageNode, MipsToSkip);
@@ -514,8 +513,7 @@ UE::Mutable::Private::Ptr<UE::Mutable::Private::NodeSurface> GenerateMutableSour
 									// GenerationContext.ArrayTextureUnrealToMutableTask.Add(FTextureUnrealToMutableTask(CompositeNormalImage, ReferenceCompositeNormalTexture, Node, true));
 									// TODO: The normal composite part is not propagated, so it will be unsupported. Create a task that performs the required transforms at mutable image level, and add the right operations here
 									// instead of propagating the flag and doing them on unreal-convert.
-									TSharedPtr<UE::Mutable::Private::FImage> ImageConstant = GenerateImageConstant(ReferenceCompositeNormalTexture, GenerationContext, false);
-									CompositeNormalImage->SetValue(ImageConstant);
+									CompositeNormalImage->Image = GenerateImageConstant(ReferenceCompositeNormalTexture, GenerationContext, false);
 
 									UE::Mutable::Private::Ptr<UE::Mutable::Private::NodeImageMipmap> NormalCompositeMipmapImage = new UE::Mutable::Private::NodeImageMipmap();
 									const uint32 MipsToSkip = ComputeLODBiasForTexture(GenerationContext, *ReferenceCompositeNormalTexture, ReferenceTexture);
diff --git a/Engine/Plugins/Mutable/Source/CustomizableObjectEditor/Private/MuCOE/GenerateMutableSource/GenerateMutableSourceTable.cpp b/Engine/Plugins/Mutable/Source/CustomizableObjectEditor/Private/MuCOE/GenerateMutableSource/GenerateMutableSourceTable.cpp
--- a/Engine/Plugins/Mutable/Source/CustomizableObjectEditor/Private/MuCOE/GenerateMutableSource/GenerateMutableSourceTable.cpp
+++ b/Engine/Plugins/Mutable/Source/CustomizableObjectEditor/Private/MuCOE/GenerateMutableSource/GenerateMutableSourceTable.cpp
@@ -332,8 +332,7 @@ bool FillTableColumn(const UCustomizableObjectNodeTable* TableNode, UE::Mutable:
 			}
 
 			const bool bIsPassthroughTexture = TableNode->GetColumnImageMode(ColumnPropertyName) == ETableTextureType::PASSTHROUGH_TEXTURE;
-			UE::Mutable::Private::Ptr<UE::Mutable::Private::TResourceProxyMemory<UE::Mutable::Private::FImage>> Proxy = new UE::Mutable::Private::TResourceProxyMemory<UE::Mutable::Private::FImage>(GenerateImageConstant(Texture, GenerationContext, bIsPassthroughTexture));
-			MutableTable->SetCell(CurrentColumn, RowId, Proxy.get());
+			MutableTable->SetCell(CurrentColumn, RowId, GenerateImageConstant(Texture, GenerationContext, bIsPassthroughTexture));
 		}
 
 		else if (SoftObjectProperty->PropertyClass->IsChildOf(UMaterialInterface::StaticClass()))
@@ -464,8 +463,7 @@ bool FillTableColumn(const UCustomizableObjectNodeTable* TableNode, UE::Mutable:
 				}
 
 				bool bIsPassthroughTexture = false;
-				UE::Mutable::Private::Ptr<UE::Mutable::Private::TResourceProxyMemory<UE::Mutable::Private::FImage>> Proxy = new UE::Mutable::Private::TResourceProxyMemory<UE::Mutable::Private::FImage>(GenerateImageConstant(ParameterTexture, GenerationContext, bIsPassthroughTexture));
-				MutableTable->SetCell(ColumnIndex, RowId, Proxy.get());
+				MutableTable->SetCell(ColumnIndex, RowId, GenerateImageConstant(ParameterTexture, GenerationContext, bIsPassthroughTexture));
 
 				return true;
 			}
diff --git a/Engine/Plugins/Mutable/Source/MutableRuntime/Internal/MuR/Serialisation.h b/Engine/Plugins/Mutable/Source/MutableRuntime/Internal/MuR/Serialisation.h
--- a/Engine/Plugins/Mutable/Source/MutableRuntime/Internal/MuR/Serialisation.h
+++ b/Engine/Plugins/Mutable/Source/MutableRuntime/Internal/MuR/Serialisation.h
@@ -44,41 +44,7 @@ namespace UE::Mutable::Private
     void DLLEXPORT operator>>(UE::Mutable::Private::FInputArchive& Arch, Type& T);
 	
 	
-    /** This repesents a handle to Mutable resource. */
-    template<class R>
-    class TResourceProxy : public RefCounted
-    {
-    public:
-
-        virtual TSharedPtr<const R> Get() = 0;
-
-    };
-
-
-	/** Proxy implementation that always has the resource loaded in memory. */
-	template<class R>
-	class TResourceProxyMemory : public TResourceProxy<R>
-	{
-	private:
-
-		TSharedPtr<const R> Resource;
-
-	public:
-
-		TResourceProxyMemory(const TSharedPtr<const R>& InResource)
-		{
-			Resource = InResource;
-		}
-
-		// TResourceProxy interface
-		TSharedPtr<const R> Get() override
-		{
-			return Resource;
-		}
-	};
-
-
-	/** */
+    /** */
 	class FModelReader
     {
     public:
@@ -205,9 +171,6 @@ namespace UE::Mutable::Private
 		/** Ensure virtual destruction. */
 		virtual ~FInputArchive() = default;
 
-        //
-        UE_API virtual Ptr<TResourceProxy<FImage>> NewImageProxy();
-
 		/** Not owned. */
 		FInputStream* Stream = nullptr;
 
diff --git a/Engine/Plugins/Mutable/Source/MutableRuntime/Private/MuR/Serialisation.cpp b/Engine/Plugins/Mutable/Source/MutableRuntime/Private/MuR/Serialisation.cpp
--- a/Engine/Plugins/Mutable/Source/MutableRuntime/Private/MuR/Serialisation.cpp
+++ b/Engine/Plugins/Mutable/Source/MutableRuntime/Private/MuR/Serialisation.cpp
@@ -203,12 +203,6 @@ namespace UE::Mutable::Private
 	}
 
 
-    Ptr<TResourceProxy<FImage>> FInputArchive::NewImageProxy()
-    {
-        return nullptr;
-    }
-
-
 	FOutputArchive::FOutputArchive( FOutputStream* InStream )
 	{
 		Stream = InStream;
diff --git a/Engine/Plugins/Mutable/Source/MutableTools/Internal/MuT/AST.h b/Engine/Plugins/Mutable/Source/MutableTools/Internal/MuT/AST.h
--- a/Engine/Plugins/Mutable/Source/MutableTools/Internal/MuT/AST.h
+++ b/Engine/Plugins/Mutable/Source/MutableTools/Internal/MuT/AST.h
@@ -62,7 +62,6 @@ namespace UE::Mutable::Private
 	class ASTOpImageSwizzle;
 	class ASTOpImageMipmap;
 	class ASTOpMeshExtractLayoutBlocks;
-	struct FProxyFileContext;
 
 	template<typename T>
 	inline uint32 GetTypeHash(const Ptr<const T>& p)
@@ -206,9 +205,6 @@ namespace UE::Mutable::Private
         //! Maximum number of iterations when optimising models. If 0 as many as necessary will be performed.
         int32 MaxOptimisationLoopCount = 8;
 
-        /** If valied, store resource data in disk instead of memory. */
-		FProxyFileContext* DiskCacheContext = nullptr;
-
 		/** Compile optimizing for the generation of smaller mipmaps of every image. */
 		bool bEnableProgressiveImages = false;
 
diff --git a/Engine/Plugins/Mutable/Source/MutableTools/Internal/MuT/Compiler.h b/Engine/Plugins/Mutable/Source/MutableTools/Internal/MuT/Compiler.h
--- a/Engine/Plugins/Mutable/Source/MutableTools/Internal/MuT/Compiler.h
+++ b/Engine/Plugins/Mutable/Source/MutableTools/Internal/MuT/Compiler.h
@@ -61,10 +61,6 @@ namespace UE::Mutable::Private
         //!
         UE_API void SetIgnoreStates( bool bIgnore );
 
-        //! If enabled, the disk will be used as temporary memory. This will make the compilation
-        //! process very slow, but will be able to compile very large models.
-        UE_API void SetUseDiskCache( bool bEnabled );
-
 		//! Set the quality for the image compression algorithms. The level value is used internally
 		//! with FSystem::SetImagecompressionQuality
 		UE_API void SetImageCompressionQuality(int32 Quality);
@@ -105,9 +101,6 @@ namespace UE::Mutable::Private
         //! 
         static UE_API const char* GetTextureLayoutStrategyName( TextureLayoutStrategy s );
 
-		/** Output some stats about the complete compilation to the log. */
-		UE_API void LogStats() const;
-
         //-----------------------------------------------------------------------------------------
         // Interface pattern
         //-----------------------------------------------------------------------------------------
diff --git a/Engine/Plugins/Mutable/Source/MutableTools/Internal/MuT/NodeImageConstant.h b/Engine/Plugins/Mutable/Source/MutableTools/Internal/MuT/NodeImageConstant.h
--- a/Engine/Plugins/Mutable/Source/MutableTools/Internal/MuT/NodeImageConstant.h
+++ b/Engine/Plugins/Mutable/Source/MutableTools/Internal/MuT/NodeImageConstant.h
@@ -13,15 +13,13 @@
 namespace UE::Mutable::Private
 {
 
-	/** Node that outputs a constant image.
-	* This node also supports "image references".
-	*/
+	/** Node that outputs a constant image.	*/
 	class NodeImageConstant : public NodeImage
 	{
 	public:
 
 		/** */
-		Ptr<TResourceProxy<FImage>> Proxy;
+		TSharedPtr<FImage> Image;
 		
 		/** */
 		FSourceDataDescriptor SourceDataDescriptor;
@@ -32,14 +30,6 @@ namespace UE::Mutable::Private
 		virtual const FNodeType* GetType() const override { return GetStaticType(); }
 		static const FNodeType* GetStaticType() { return &StaticType; }
 
-		// Own Interface
-
-        //! Set the image to be output by this node
-        UE_API void SetValue(TSharedPtr<const FImage>);
-
-        //! Set the image proxy that will provide the image for this node when necessary
-        UE_API void SetValue( Ptr<TResourceProxy<FImage>> );
-
 	protected:
 
 		/** Forbidden. Manage with the Ptr<> template. */
diff --git a/Engine/Plugins/Mutable/Source/MutableTools/Internal/MuT/Table.h b/Engine/Plugins/Mutable/Source/MutableTools/Internal/MuT/Table.h
--- a/Engine/Plugins/Mutable/Source/MutableTools/Internal/MuT/Table.h
+++ b/Engine/Plugins/Mutable/Source/MutableTools/Internal/MuT/Table.h
@@ -60,7 +60,7 @@ namespace UE::Mutable::Private
         UE_API void SetCell( int32 Column, uint32 RowId, int32 IntValue, const void* ErrorContext = nullptr);
         UE_API void SetCell( int32 Column, uint32 RowId, float Value, const void* ErrorContext = nullptr);
         UE_API void SetCell( int32 Column, uint32 RowId, const FVector4f& Value, const void* ErrorContext = nullptr);
-		UE_API void SetCell( int32 Column, uint32 RowId, TResourceProxy<FImage>* Value, const void* ErrorContext = nullptr);
+		UE_API void SetCell( int32 Column, uint32 RowId, const TSharedPtr<FImage>& Value, const void* ErrorContext = nullptr);
 		UE_API void SetCell( int32 Column, uint32 RowId, const TSharedPtr<FMesh>& Value, const void* ErrorContext = nullptr);
         UE_API void SetCell( int32 Column, uint32 RowId, const FString& Value, const void* ErrorContext = nullptr);
 
diff --git a/Engine/Plugins/Mutable/Source/MutableTools/Internal/MuT/TablePrivate.h b/Engine/Plugins/Mutable/Source/MutableTools/Internal/MuT/TablePrivate.h
--- a/Engine/Plugins/Mutable/Source/MutableTools/Internal/MuT/TablePrivate.h
+++ b/Engine/Plugins/Mutable/Source/MutableTools/Internal/MuT/TablePrivate.h
@@ -25,7 +25,7 @@ namespace UE::Mutable::Private
 		int32 Int;
 		float Scalar;
 		FVector4f Color;
-		Ptr<TResourceProxy<FImage>> ProxyImage;
+		TSharedPtr<FImage> Image;
 		TSharedPtr<FMesh> Mesh;
 		FString String;
 
diff --git a/Engine/Plugins/Mutable/Source/MutableTools/Private/MuT/ASTOpConstantResource.cpp b/Engine/Plugins/Mutable/Source/MutableTools/Private/MuT/ASTOpConstantResource.cpp
--- a/Engine/Plugins/Mutable/Source/MutableTools/Private/MuT/ASTOpConstantResource.cpp
+++ b/Engine/Plugins/Mutable/Source/MutableTools/Private/MuT/ASTOpConstantResource.cpp
@@ -12,225 +12,14 @@
 #include "MuR/Skeleton.h"
 
 #include "Containers/Array.h"
-#include "HAL/PlatformMath.h"
-#include "HAL/PlatformFileManager.h"
 #include "Hash/CityHash.h"
 #include "Misc/AssertionMacros.h"
-#include "GenericPlatform/GenericPlatformFile.h"
-#include "Compression/OodleDataCompression.h"
 
 #include <inttypes.h> // Required for 64-bit printf macros
 
 namespace UE::Mutable::Private
 {
 
-	/** Proxy class for a temporary resource while compiling. 
-	* The resource may be stored in different ways:
-	* - as is, in memory with its own pointer.
-	* - in a compressed buffer
-	* - saved to a disk file compressed or uncompressed.
-	*/
-	template<class R>
-	class MUTABLETOOLS_API ResourceProxyTempFile : public TResourceProxy<R>
-	{
-	private:
-
-		/** Actual resource to store. If the pointer is valid, it wasn't worth dumping to disk or compressing. */
-		TSharedPtr<const R> Resource;
-
-		/** Temp filename used if it was necessary. */
-		FString FileName;
-
-		/** Size of the resource in memory. */
-		uint32 UncompressedSize = 0;
-
-		/** Size of the saved file. It may be the size of the resource in memory, or its compressed size. */
-		uint32 FileSize = 0;
-
-		/** Valid if the resource was compressed and stored in memory instead of dumped to disk. */
-		TArray<uint8> CompressedBuffer;
-
-		/** Shared context with cache settings and stats. */
-		FProxyFileContext& Options;
-
-		/** Prevent concurrent access to a signal resource. */
-		FCriticalSection Mutex;
-
-	public:
-
-		ResourceProxyTempFile(TSharedPtr<const R> InResource, FProxyFileContext& InOptions)
-			: Options(InOptions)
-		{
-			if (!InResource)
-			{
-				return;
-			}
-
-			IPlatformFile& PlatformFile = FPlatformFileManager::Get().GetPlatformFile();
-
-			FOutputMemoryStream Stream(128*1024);
-			FOutputArchive Arch(&Stream);
-			R::Serialise(InResource.Get(), Arch);
-
-			UncompressedSize = Stream.GetBufferSize();
-
-			if (Stream.GetBufferSize() <= Options.MinProxyFileSize)
-			{
-				// Not worth compressing or caching to disk
-				Resource = InResource;
-			}
-			else
-			{
-				// Compress
-				int64 CompressedSize = 0;
-				constexpr bool bEnableCompression = true;
-				if (bEnableCompression)
-				{
-					int64 CompressedBufferSize = FOodleDataCompression::CompressedBufferSizeNeeded(Stream.GetBufferSize());
-					CompressedBufferSize = FMath::Max(CompressedBufferSize, int64(Stream.GetBufferSize() / 2));
-					CompressedBuffer.SetNumUninitialized(CompressedBufferSize);
-
-					CompressedSize = FOodleDataCompression::CompressParallel(
-						CompressedBuffer.GetData(), CompressedBufferSize,
-						Stream.GetBuffer(), Stream.GetBufferSize(),
-						FOodleDataCompression::ECompressor::Kraken,
-						FOodleDataCompression::ECompressionLevel::SuperFast,
-						true // CompressIndependentChunks
-					);
-				}
-
-				bool bCompressed = CompressedSize != 0;
-
-				if (bCompressed && uint64(CompressedSize) <= Options.MinProxyFileSize)
-				{
-					// Keep the compressed data, and don't store to file
-					CompressedBuffer.SetNum(CompressedSize,EAllowShrinking::Yes);
-				}
-				else
-				{
-					// Save
-					FString Prefix = FPlatformProcess::UserTempDir();
-
-					uint32 PID = FPlatformProcess::GetCurrentProcessId();
-					Prefix += FString::Printf(TEXT("mut.temp.%u"), PID);
-
-					FString FinalTempPath;
-					IFileHandle* ResourceFile = nullptr;
-					uint64 AttemptCount = 0;
-					while (!ResourceFile && AttemptCount < Options.MaxFileCreateAttempts)
-					{
-						uint64 ThisThreadFileIndex = Options.CurrentFileIndex.load();
-						while (!Options.CurrentFileIndex.compare_exchange_strong(ThisThreadFileIndex, ThisThreadFileIndex + 1));
-
-						FinalTempPath = Prefix + FString::Printf(TEXT(".%.16" PRIx64), ThisThreadFileIndex);
-						ResourceFile = PlatformFile.OpenWrite(*FinalTempPath);
-						++AttemptCount;
-					}
-
-					if (!ResourceFile)
-					{
-						UE_LOG(LogMutableCore, Error, TEXT("Failed to create temporary file. Disk full?"));
-						check(false);
-					}
-
-					if (bCompressed)
-					{
-						FileSize = CompressedSize;
-						ResourceFile->Write(CompressedBuffer.GetData(), FileSize);
-					}
-					else
-					{
-						FileSize = UncompressedSize;
-						ResourceFile->Write(Stream.GetBuffer(), FileSize);
-					}
-
-					CompressedBuffer.SetNum(0, EAllowShrinking::Yes);
-
-					delete ResourceFile;
-
-					FileName = FinalTempPath;
-					Options.FilesWritten++;
-					Options.BytesWritten += FileSize;
-				}
-			}
-		}
-
-		~ResourceProxyTempFile()
-		{
-			FScopeLock Lock(&Mutex);
-
-			if (!FileName.IsEmpty())
-			{
-				// Delete temp file
-				FPlatformFileManager::Get().GetPlatformFile().DeleteFile(*FileName);
-				FileName.Empty();
-			}
-		}
-
-		TSharedPtr<const R> Get() override
-		{
-			FScopeLock Lock(&Mutex);
-
-			TSharedPtr<const R> Result;
-			if (Resource)
-			{
-				// Cached as is
-				Result = Resource;
-			}
-			else if (!CompressedBuffer.Num() && !FileName.IsEmpty())
-			{
-				IFileHandle* ResourceFile = FPlatformFileManager::Get().GetPlatformFile().OpenRead(*FileName);
-				check(ResourceFile);
-
-				CompressedBuffer.SetNumUninitialized(FileSize);
-				ResourceFile->Read(CompressedBuffer.GetData(), FileSize);
-				delete ResourceFile;
-
-				bool bCompressed = FileSize != UncompressedSize;
-
-				if (!bCompressed)
-				{
-					FInputMemoryStream stream(CompressedBuffer.GetData(), FileSize);
-					FInputArchive arch(&stream);
-					Result = R::StaticUnserialise(arch);
-
-					CompressedBuffer.SetNum(0, EAllowShrinking::Yes);
-				}
-
-				Options.FilesRead++;
-				Options.BytesRead += FileSize;
-			}
-
-			if (CompressedBuffer.Num())
-			{
-				// Cached compressed
-				TArray<uint8> UncompressedBuf;
-				UncompressedBuf.SetNumUninitialized(UncompressedSize);
-
-				bool bSuccess = FOodleDataCompression::DecompressParallel(
-					UncompressedBuf.GetData(), UncompressedSize,
-					CompressedBuffer.GetData(), CompressedBuffer.Num());
-				check(bSuccess);
-
-				if (bSuccess)
-				{
-					FInputMemoryStream stream(UncompressedBuf.GetData(), UncompressedSize);
-					FInputArchive arch(&stream);
-					Result = R::StaticUnserialise(arch);
-				}
-
-				if (!FileName.IsEmpty())
-				{
-					CompressedBuffer.SetNum(0, EAllowShrinking::Yes);
-				}
-			}
-
-			return Result;
-		}
-
-	};
-
-
 	//-------------------------------------------------------------------------------------------------
 	void ASTOpConstantResource::ForEachChild(const TFunctionRef<void(ASTChild&)> f)
 	{
@@ -249,9 +38,9 @@ namespace UE::Mutable::Private
 		{
 			const ASTOpConstantResource* Other = static_cast<const ASTOpConstantResource*>(&OtherUntyped);
 			return Type == Other->Type && ValueHash == Other->ValueHash &&
-				LoadedValue == Other->LoadedValue && Proxy == Other->Proxy
-				&& SourceDataDescriptor == Other->SourceDataDescriptor
-				&& ImageOperations == Other->ImageOperations;
+				Value == Other->Value &&
+				SourceDataDescriptor == Other->SourceDataDescriptor	&&
+				ImageOperations == Other->ImageOperations;
 		}
 		return false;
 	}
@@ -262,8 +51,7 @@ namespace UE::Mutable::Private
 	{
 		Ptr<ASTOpConstantResource> n = new ASTOpConstantResource();
 		n->Type = Type;
-		n->Proxy = Proxy;
-		n->LoadedValue = LoadedValue;
+		n->Value = Value;
 		n->ValueHash = ValueHash;
 		n->SourceDataDescriptor = SourceDataDescriptor;
 
@@ -621,7 +409,7 @@ namespace UE::Mutable::Private
 					OP::MeshConstantArgs Args;
 					FMemory::Memzero(Args);
 
-					TSharedPtr<FMesh> MeshValue = StaticCastSharedPtr<const FMesh>(GetValue())->Clone();
+					TSharedPtr<FMesh> MeshValue = StaticCastSharedPtr<const FMesh>(Value)->Clone();
 					check(MeshValue);
 
 					Args.Skeleton = -1;
@@ -650,7 +438,7 @@ namespace UE::Mutable::Private
 					OP::ResourceConstantArgs Args;
 					FMemory::Memzero(Args);
 					
-					TSharedPtr<const FImage> pTyped = StaticCastSharedPtr<const FImage>(GetValue());
+					TSharedPtr<const FImage> pTyped = StaticCastSharedPtr<const FImage>(Value);
 					check(pTyped);
 
 					if (pTyped->GetSizeX() * pTyped->GetSizeY() == 0)
@@ -683,7 +471,7 @@ namespace UE::Mutable::Private
 					OP::ResourceConstantArgs Args;
 					FMemory::Memzero(Args);
 					
-					TSharedPtr<const FLayout> pTyped = StaticCastSharedPtr<const FLayout>(GetValue());
+					TSharedPtr<const FLayout> pTyped = StaticCastSharedPtr<const FLayout>(Value);
 					check(pTyped);
 					Args.value = Program.AddConstant(pTyped);
 
@@ -700,7 +488,7 @@ namespace UE::Mutable::Private
 					OP::MaterialConstantArgs Args;
 					FMemory::Memzero(Args);
 
-					TSharedPtr<FMaterial> MaterialValue = StaticCastSharedPtr<const FMaterial>(GetValue())->Clone();
+					TSharedPtr<FMaterial> MaterialValue = StaticCastSharedPtr<const FMaterial>(Value)->Clone();
 					check(MaterialValue);
 
 					// Store linked child images into the material constant
@@ -729,8 +517,7 @@ namespace UE::Mutable::Private
 			}
 			
 			// Clear stored value to reduce memory usage.
-			LoadedValue = nullptr;
-			Proxy = nullptr;
+			Value = nullptr;
 		}
 	}
 
@@ -743,10 +530,13 @@ namespace UE::Mutable::Private
 		if (Type == EOpType::IM_CONSTANT)
 		{
 			// TODO: cache to avoid disk loading
-			TSharedPtr<const FImage> ConstImage = StaticCastSharedPtr<const FImage>(GetValue());
-			Result.m_format = ConstImage->GetFormat();
-			Result.m_lods = ConstImage->GetLODCount();
-			Result.m_size = ConstImage->GetSize();
+			TSharedPtr<const FImage> ConstImage = StaticCastSharedPtr<const FImage>(Value);
+			if (ensure(ConstImage))
+			{
+				Result.m_format = ConstImage->GetFormat();
+				Result.m_lods = ConstImage->GetLODCount();
+				Result.m_size = ConstImage->GetSize();
+			}
 		}
 		else if (Type == EOpType::MI_CONSTANT)
 		{
@@ -770,7 +560,7 @@ namespace UE::Mutable::Private
 		{
 		case EOpType::LA_CONSTANT:
 		{
-			TSharedPtr<const FLayout> pLayout = StaticCastSharedPtr<const FLayout>(GetValue());
+			TSharedPtr<const FLayout> pLayout = StaticCastSharedPtr<const FLayout>(Value);
 			check(pLayout);
 
 			if (pLayout)
@@ -822,7 +612,7 @@ namespace UE::Mutable::Private
 		if (Type == EOpType::IM_CONSTANT)
 		{
 			// TODO: cache
-			TSharedPtr<const FImage> pMask = StaticCastSharedPtr<const FImage>(GetValue());
+			TSharedPtr<const FImage> pMask = StaticCastSharedPtr<const FImage>(Value);
 			pMask->GetNonBlackRect(maskUsage);
 			return true;
 		}
@@ -840,7 +630,7 @@ namespace UE::Mutable::Private
 
 		case EOpType::IM_CONSTANT:
 		{
-			TSharedPtr<const FImage> pImage = StaticCastSharedPtr<const FImage>(GetValue());
+			TSharedPtr<const FImage> pImage = StaticCastSharedPtr<const FImage>(Value);
 			if (pImage->GetSizeX() <= 0 || pImage->GetSizeY() <= 0)
 			{
 				res = true;
@@ -895,34 +685,12 @@ namespace UE::Mutable::Private
 	//-------------------------------------------------------------------------------------------------
 	TSharedPtr<const FResource> ASTOpConstantResource::GetValue() const
 	{
-		if (LoadedValue)
-		{
-			return LoadedValue;
-		}
-		else
-		{
-			switch (Type)
-			{
-
-			case EOpType::IM_CONSTANT:
-			{
-				Ptr<TResourceProxy<FImage>> typedProxy = static_cast<TResourceProxy<FImage>*>(Proxy.get());
-				TSharedPtr<const FImage> Resource = typedProxy->Get();
-				return Resource;
-			}
-
-			default:
-				check(false);
-				break;
-			}
-		}
-
-		return nullptr;
+		return Value;
 	}
 
 
 	//-------------------------------------------------------------------------------------------------
-	void ASTOpConstantResource::SetValue(const TSharedPtr<const FResource>& Value, FProxyFileContext* DiskCacheContext)
+	void ASTOpConstantResource::SetValue(const TSharedPtr<const FResource>& InValue)
 	{
 		MUTABLE_CPUPROFILER_SCOPE(ASTOpConstantResource_SetValue);
 
@@ -930,7 +698,7 @@ namespace UE::Mutable::Private
 		{
 		case EOpType::IM_CONSTANT:
 		{
-			TSharedPtr<const FImage> Resource = StaticCastSharedPtr<const FImage>(Value);
+			TSharedPtr<const FImage> Resource = StaticCastSharedPtr<const FImage>(InValue);
 
 			FOutputHashStream stream;
 			{
@@ -940,21 +708,14 @@ namespace UE::Mutable::Private
 			}
 
 			ValueHash = stream.GetHash();
-
-			if (DiskCacheContext)
-			{
-				Proxy = new ResourceProxyTempFile<FImage>(Resource, *DiskCacheContext);
-			}
-			else
-			{
-				LoadedValue = Resource;
-			}
+			
+			Value = Resource;
 			break;
 		}
 
 		case EOpType::ME_CONSTANT:
 		{
-			TSharedPtr<const FMesh> Resource = StaticCastSharedPtr<const FMesh>(Value);
+			TSharedPtr<const FMesh> Resource = StaticCastSharedPtr<const FMesh>(InValue);
 
 			FOutputHashStream stream;
 			{
@@ -965,13 +726,13 @@ namespace UE::Mutable::Private
 
 			ValueHash = stream.GetHash();
 
-			LoadedValue = Resource;
+			Value = Resource;
 			break;
 		}
 
 		case EOpType::LA_CONSTANT:
 		{
-			TSharedPtr<const FLayout> Resource = StaticCastSharedPtr<const FLayout>(Value);
+			TSharedPtr<const FLayout> Resource = StaticCastSharedPtr<const FLayout>(InValue);
 
 			FOutputHashStream stream;
 			{
@@ -982,13 +743,13 @@ namespace UE::Mutable::Private
 
 			ValueHash = stream.GetHash();
 
-			LoadedValue = Resource;
+			Value = Resource;
 			break;
 		}
 		
 		case EOpType::MI_CONSTANT:
 		{
-			TSharedPtr<const FMaterial> Resource = StaticCastSharedPtr<const FMaterial>(Value);
+			TSharedPtr<const FMaterial> Resource = StaticCastSharedPtr<const FMaterial>(InValue);
 
 			FOutputHashStream stream;
 			{
@@ -998,13 +759,13 @@ namespace UE::Mutable::Private
 			}
 
 			ValueHash = stream.GetHash();
-			LoadedValue = Resource;
+			Value = Resource;
 
 			break;
 		}
 
 		default:
-			LoadedValue = Value;
+			Value = InValue;
 			break;
 		}
 	}
@@ -1016,7 +777,7 @@ namespace UE::Mutable::Private
 		{
 			Ptr<ImageSizeExpression> Res = new ImageSizeExpression;
 			Res->type = ImageSizeExpression::ISET_CONSTANT;
-			TSharedPtr<const FImage> Const = StaticCastSharedPtr<const FImage>(GetValue());
+			TSharedPtr<const FImage> Const = StaticCastSharedPtr<const FImage>(Value);
 			Res->size = Const->GetSize();
 			return Res;
 		}
@@ -1045,7 +806,7 @@ namespace UE::Mutable::Private
 		ASTOp::EClosedMeshTest Result = EClosedMeshTest::Unknown;
 		if (Type == EOpType::ME_CONSTANT)
 		{
-			TSharedPtr<const FMesh> Mesh = StaticCastSharedPtr<const FMesh>(GetValue());
+			TSharedPtr<const FMesh> Mesh = StaticCastSharedPtr<const FMesh>(Value);
 			if (Mesh)
 			{
 				if (Mesh->IsClosed())
diff --git a/Engine/Plugins/Mutable/Source/MutableTools/Private/MuT/ASTOpConstantResource.h b/Engine/Plugins/Mutable/Source/MutableTools/Private/MuT/ASTOpConstantResource.h
--- a/Engine/Plugins/Mutable/Source/MutableTools/Private/MuT/ASTOpConstantResource.h
+++ b/Engine/Plugins/Mutable/Source/MutableTools/Private/MuT/ASTOpConstantResource.h
@@ -13,7 +13,6 @@
 namespace UE::Mutable::Private
 {
 	struct FProgram;
-	struct FProxyFileContext;
 
 	/** A constant mesh, image, volume or layout. */
 	class ASTOpConstantResource final : public ASTOp
@@ -21,8 +20,7 @@ namespace UE::Mutable::Private
 	private:
 
 		//!
-		TSharedPtr<const FResource> LoadedValue;
-		Ptr<RefCounted> Proxy;
+		TSharedPtr<const FResource> Value;
 
 		//! Value hash
 		uint64 ValueHash;
@@ -55,10 +53,8 @@ namespace UE::Mutable::Private
 		//! Get a pointer to the stored value.
 		TSharedPtr<const FResource> GetValue() const;
 
-		/** Set the value to store in this op.
-		* If the DiskCacheContext is not null, the disk cache will be used.
-		*/
-		void SetValue(const TSharedPtr<const FResource>& Value, FProxyFileContext* DiskCacheContext);
+		/** Set the value to store in this op.	*/
+		void SetValue(const TSharedPtr<const FResource>& Value);
 
 
 		// ASTOp interface
diff --git a/Engine/Plugins/Mutable/Source/MutableTools/Private/MuT/ASTOpMeshFormat.cpp b/Engine/Plugins/Mutable/Source/MutableTools/Private/MuT/ASTOpMeshFormat.cpp
--- a/Engine/Plugins/Mutable/Source/MutableTools/Private/MuT/ASTOpMeshFormat.cpp
+++ b/Engine/Plugins/Mutable/Source/MutableTools/Private/MuT/ASTOpMeshFormat.cpp
@@ -358,7 +358,7 @@ UE::Mutable::Private::Ptr<ASTOp> Sink_MeshFormatAST::Visit(const UE::Mutable::Pr
 
 		UE::Mutable::Private::Ptr<ASTOpConstantResource> NewFormatConstant = new ASTOpConstantResource();
 		NewFormatConstant->Type = EOpType::ME_CONSTANT;
-		NewFormatConstant->SetValue(pTargetMorphFormat, nullptr);
+		NewFormatConstant->SetValue(pTargetMorphFormat);
 		NewFormatConstant->SourceDataDescriptor = at->GetSourceDataDescriptor();
 
 		if (NewOp->Target)
@@ -393,7 +393,7 @@ UE::Mutable::Private::Ptr<ASTOp> Sink_MeshFormatAST::Visit(const UE::Mutable::Pr
         
 		UE::Mutable::Private::Ptr<ASTOpConstantResource> NewFormatConstant = new ASTOpConstantResource();
 		NewFormatConstant->Type = EOpType::ME_CONSTANT;
-		NewFormatConstant->SetValue(TargetFormat, nullptr);
+		NewFormatConstant->SetValue(TargetFormat);
 		NewFormatConstant->SourceDataDescriptor = at->GetSourceDataDescriptor();
 	
         NewFormat->Flags = NewFormat->Flags | OP::MeshFormatArgs::OptimizeBuffers;
diff --git a/Engine/Plugins/Mutable/Source/MutableTools/Private/MuT/CodeGenerator.cpp b/Engine/Plugins/Mutable/Source/MutableTools/Private/MuT/CodeGenerator.cpp
--- a/Engine/Plugins/Mutable/Source/MutableTools/Private/MuT/CodeGenerator.cpp
+++ b/Engine/Plugins/Mutable/Source/MutableTools/Private/MuT/CodeGenerator.cpp
@@ -469,7 +469,7 @@ namespace UE::Mutable::Private
 			Ptr<ASTOp> RectConstantOp;
 			{
 				Ptr<NodeImageConstant> pNode = new NodeImageConstant();
-				pNode->SetValue(PatchMask);
+				pNode->Image = PatchMask;
 
 				FImageGenerationOptions ConstantOptions(-1,-1);
 				FImageGenerationResult ConstantResult;
@@ -1081,9 +1081,7 @@ namespace UE::Mutable::Private
 							Ptr<ASTOpConstantResource> ConstantLayoutOp = new ASTOpConstantResource();
 							ConstantLayoutOp->Type = EOpType::LA_CONSTANT;
 
-							ConstantLayoutOp->SetValue(
-								SurfaceReferenceLayouts[LayoutIndex].Layout,
-								CompilerOptions->OptimisationOptions.DiskCacheContext);
+							ConstantLayoutOp->SetValue(SurfaceReferenceLayouts[LayoutIndex].Layout);
 							SurfaceLayoutOps[LayoutIndex] = ConstantLayoutOp;
 						}
 					}
@@ -1131,9 +1129,7 @@ namespace UE::Mutable::Private
 							Ptr<ASTOpConstantResource> LayoutFragmentConstantOp = new ASTOpConstantResource();
 							LayoutFragmentConstantOp->Type = EOpType::LA_CONSTANT;
 
-							LayoutFragmentConstantOp->SetValue(
-								ExtraLayoutsData.GeneratedLayouts[LayoutIndex].Layout,
-								CompilerOptions->OptimisationOptions.DiskCacheContext);
+							LayoutFragmentConstantOp->SetValue(ExtraLayoutsData.GeneratedLayouts[LayoutIndex].Layout);
 
 							Ptr<ASTOpLayoutMerge> LayoutMergeOp = new ASTOpLayoutMerge();
 							// Base may be null if the base does not have  a mesh with a layout at LayoutIndex.
@@ -2506,7 +2502,7 @@ namespace UE::Mutable::Private
 								// Morph generation through mesh diff
 								Ptr<ASTOpConstantResource> TargetMeshOp = new ASTOpConstantResource;
 								TargetMeshOp->Type = EOpType::ME_CONSTANT;
-								TargetMeshOp->SetValue(TargetMesh->Clone(), CompilerOptions->OptimisationOptions.DiskCacheContext);
+								TargetMeshOp->SetValue(TargetMesh->Clone());
 								TargetMeshOp->SourceDataDescriptor = static_cast<const NodeMeshConstant*>(OriginalMeshCopy.get())->SourceDataDescriptor;
 				
 								Ptr<ASTOpMeshDifference> DiffOp;
@@ -2699,7 +2695,7 @@ namespace UE::Mutable::Private
 
 							Ptr<ASTOpConstantResource> UVMeshOp = new ASTOpConstantResource;
 							UVMeshOp->Type = EOpType::ME_CONSTANT;
-							UVMeshOp->SetValue(TargetMesh->Clone(), CompilerOptions->OptimisationOptions.DiskCacheContext);
+							UVMeshOp->SetValue(TargetMesh->Clone());
 							UVMeshOp->SourceDataDescriptor = static_cast<const NodeMeshConstant*>(OriginalMeshCopy.get())->SourceDataDescriptor;
 
 							Ptr<ASTOpMeshMaskClipUVMask> MeshMaskAt = new ASTOpMeshMaskClipUVMask();
@@ -2721,7 +2717,7 @@ namespace UE::Mutable::Private
 							{
 								Ptr<ASTOpConstantResource> LayoutOp = new ASTOpConstantResource();
 								LayoutOp->Type = EOpType::LA_CONSTANT;
-								LayoutOp->SetValue(Layout, CompilerOptions->OptimisationOptions.DiskCacheContext);
+								LayoutOp->SetValue(Layout);
 								MeshMaskAt->MaskLayout = LayoutOp;
 							}
 
diff --git a/Engine/Plugins/Mutable/Source/MutableTools/Private/MuT/CodeGenerator_Image.cpp b/Engine/Plugins/Mutable/Source/MutableTools/Private/MuT/CodeGenerator_Image.cpp
--- a/Engine/Plugins/Mutable/Source/MutableTools/Private/MuT/CodeGenerator_Image.cpp
+++ b/Engine/Plugins/Mutable/Source/MutableTools/Private/MuT/CodeGenerator_Image.cpp
@@ -223,11 +223,7 @@ namespace UE::Mutable::Private
 		const NodeImageConstant& node = *InNode;
 		
         // TODO: check duplicates
-        TSharedPtr<const FImage> pImage;
-		if (node.Proxy)
-		{
-			pImage = node.Proxy->Get();
-		}
+        TSharedPtr<const FImage> pImage = InNode->Image;
 
         if (!pImage)
         {
@@ -260,7 +256,7 @@ namespace UE::Mutable::Private
 		{
 			Ptr<ASTOpConstantResource> op = new ASTOpConstantResource();
 			op->Type = EOpType::IM_CONSTANT;
-			op->SetValue(pImage, CompilerOptions->OptimisationOptions.DiskCacheContext);
+			op->SetValue(pImage);
 			op->SourceDataDescriptor = InNode->SourceDataDescriptor;
 			Result.op = op;
 		}
@@ -1422,7 +1418,7 @@ namespace UE::Mutable::Private
 				TSharedPtr<FMesh> FormatMeshResult = MakeShared<FMesh>();
 				CreateMeshOptimisedForWrappingProjection(FormatMeshResult.Get(), Node->Layout);
 
-                cop->SetValue(FormatMeshResult, CompilerOptions->OptimisationOptions.DiskCacheContext);
+                cop->SetValue(FormatMeshResult);
 
                 Ptr<ASTOpMeshFormat> FormatOp = new ASTOpMeshFormat();
 				FormatOp->Flags = OP::MeshFormatArgs::Vertex
@@ -1454,7 +1450,7 @@ namespace UE::Mutable::Private
 				TSharedPtr<FMesh> FormatMeshResult = MakeShared<FMesh>();
                 CreateMeshOptimisedForProjection(FormatMeshResult.Get(), Node->Layout);
 
-                cop->SetValue(FormatMeshResult, CompilerOptions->OptimisationOptions.DiskCacheContext);
+                cop->SetValue(FormatMeshResult);
 
                 Ptr<ASTOpMeshFormat> FormatOp = new ASTOpMeshFormat();
 				FormatOp->Flags = OP::MeshFormatArgs::Vertex
@@ -1474,7 +1470,7 @@ namespace UE::Mutable::Private
             TSharedPtr<const FMesh> TempMesh = MakeShared<FMesh>();
             Ptr<ASTOpConstantResource> cop = new ASTOpConstantResource();
             cop->Type = EOpType::ME_CONSTANT;
-            cop->SetValue(TempMesh, CompilerOptions->OptimisationOptions.DiskCacheContext);
+            cop->SetValue(TempMesh);
 			ProjectOp->Mesh = cop;
             ErrorLog->Add( "Projector mesh not set.", ELMT_ERROR, Node->GetMessageContext() );
         }
@@ -1848,12 +1844,7 @@ namespace UE::Mutable::Private
 			[this, InNode, Options](const NodeImageTable& node, int32 colIndex, int32 row, FErrorLog* LocalErrorLog)
 			{
 				const FTableValue& CellData = node.Table->GetPrivate()->Rows[row].Values[colIndex];
-				TSharedPtr<const FImage> pImage = nullptr;
-
-				if (Ptr<TResourceProxy<FImage>> pProxyImage = CellData.ProxyImage)
-				{
-					pImage = pProxyImage->Get();
-				}
+				TSharedPtr<FImage> pImage = CellData.Image;
 
 				Ptr<ASTOp> ImageOp;
 
@@ -1870,7 +1861,7 @@ namespace UE::Mutable::Private
 				else
 				{
 					Ptr<NodeImageConstant> ImageConst = new NodeImageConstant();
-					ImageConst->SetValue(pImage);
+					ImageConst->Image = pImage;
 
 					// TODO: We probably want to get the data tags from the table row.
 					ImageConst->SourceDataDescriptor = InNode->SourceDataDescriptor;
@@ -2063,7 +2054,7 @@ namespace UE::Mutable::Private
 		TSharedPtr<FImage> GeneratedImage = GenerateMissingImage( format );
 
 		Ptr<NodeImageConstant> pNode = new NodeImageConstant();
-		pNode->SetValue(GeneratedImage);
+		pNode->Image = GeneratedImage;
 
 		FImageGenerationResult Result;
 		GenerateImage(Options, Result, pNode);
diff --git a/Engine/Plugins/Mutable/Source/MutableTools/Private/MuT/CodeGenerator_Material.cpp b/Engine/Plugins/Mutable/Source/MutableTools/Private/MuT/CodeGenerator_Material.cpp
--- a/Engine/Plugins/Mutable/Source/MutableTools/Private/MuT/CodeGenerator_Material.cpp
+++ b/Engine/Plugins/Mutable/Source/MutableTools/Private/MuT/CodeGenerator_Material.cpp
@@ -132,7 +132,7 @@ namespace UE::Mutable::Private
 		}
 
 		// Store the constant material
-		ConstantMaterial->SetValue(MaterialResource, CompilerOptions->OptimisationOptions.DiskCacheContext);
+		ConstantMaterial->SetValue(MaterialResource);
 
 		Result.op = ConstantMaterial;
 	}
@@ -265,7 +265,7 @@ namespace UE::Mutable::Private
 					ConstantMaterial->Type = EOpType::MI_CONSTANT;
 
 					// Store the constant material
-					ConstantMaterial->SetValue(MaterialResource, CompilerOptions->OptimisationOptions.DiskCacheContext);
+					ConstantMaterial->SetValue(MaterialResource);
 
 					Op = ConstantMaterial;
 				}
@@ -379,7 +379,7 @@ namespace UE::Mutable::Private
 		GenerateMaterial->PassthroughObject = {};
 
 		// Store the constant material
-		ConstantMaterial->SetValue(GenerateMaterial, CompilerOptions->OptimisationOptions.DiskCacheContext);
+		ConstantMaterial->SetValue(GenerateMaterial);
 
 		return ConstantMaterial;
 	}
diff --git a/Engine/Plugins/Mutable/Source/MutableTools/Private/MuT/CodeGenerator_Mesh.cpp b/Engine/Plugins/Mutable/Source/MutableTools/Private/MuT/CodeGenerator_Mesh.cpp
--- a/Engine/Plugins/Mutable/Source/MutableTools/Private/MuT/CodeGenerator_Mesh.cpp
+++ b/Engine/Plugins/Mutable/Source/MutableTools/Private/MuT/CodeGenerator_Mesh.cpp
@@ -727,7 +727,7 @@ namespace UE::Mutable::Private
 
 			Ptr<ASTOpConstantResource> LayoutOp = new ASTOpConstantResource;
 			LayoutOp->Type = EOpType::LA_CONSTANT;
-			LayoutOp->SetValue(Layout, CompilerOptions->OptimisationOptions.DiskCacheContext);
+			LayoutOp->SetValue(Layout);
 
 			Ptr<ASTOpMeshPrepareLayout> PrepareOp = new ASTOpMeshPrepareLayout;
 			PrepareOp->Mesh = LastMeshOp;
@@ -763,7 +763,7 @@ namespace UE::Mutable::Private
 
 			// This data is required
 			TSharedPtr<FMesh> EmptyMesh = MakeShared<FMesh>();
-			ConstantOp->SetValue(EmptyMesh, CompilerOptions->OptimisationOptions.DiskCacheContext);
+			ConstantOp->SetValue(EmptyMesh);
 			EmptyMesh->MeshIDPrefix = UniqueMeshIds.EnsureUnique(ConstantOp->GetValueHash());
 
 			// Log an error message
@@ -877,7 +877,7 @@ namespace UE::Mutable::Private
 					TSharedPtr<FMesh> ClonedMesh = Mesh->Clone();
 					ClonedMesh->EnsureSurfaceData();
 
-					ConstantOp->SetValue(ClonedMesh, CompilerOptions->OptimisationOptions.DiskCacheContext);
+					ConstantOp->SetValue(ClonedMesh);
 
 					// Add the unique vertex ID prefix in all cases, since it is free memory-wise
 					MeshIDPrefix = uint32(ConstantOp->GetValueHash());
@@ -984,7 +984,7 @@ namespace UE::Mutable::Private
 
 				Ptr<ASTOpConstantResource> cop = new ASTOpConstantResource();
 				cop->Type = EOpType::ME_CONSTANT;
-				cop->SetValue(FormatMesh, CompilerOptions->OptimisationOptions.DiskCacheContext);
+				cop->SetValue(FormatMesh);
 				if (baseResult.BaseMeshOp)
 				{
 					cop->SourceDataDescriptor = baseResult.BaseMeshOp->GetSourceDataDescriptor();
diff --git a/Engine/Plugins/Mutable/Source/MutableTools/Private/MuT/CodeOptimiser.cpp b/Engine/Plugins/Mutable/Source/MutableTools/Private/MuT/CodeOptimiser.cpp
--- a/Engine/Plugins/Mutable/Source/MutableTools/Private/MuT/CodeOptimiser.cpp
+++ b/Engine/Plugins/Mutable/Source/MutableTools/Private/MuT/CodeOptimiser.cpp
@@ -135,7 +135,6 @@ namespace UE::Mutable::Private
 
 		// input
 		Ptr<ASTOp> Source;
-		FProxyFileContext* DiskCacheContext = nullptr;
 		int32 ImageCompressionQuality = 0;
 		int32 OptimizationPass = 0;
 		FReferencedMeshResourceFunc ReferencedMeshResourceProvider;
@@ -153,7 +152,6 @@ namespace UE::Mutable::Private
 		{
 			OptimizationPass = InOptimizationPass;
 			Source = InSource;
-			DiskCacheContext = InOptions->OptimisationOptions.DiskCacheContext;
 			ImageCompressionQuality = InOptions->ImageCompressionQuality;
 			ReferencedMeshResourceProvider = InOptions->OptimisationOptions.ReferencedMeshResourceProvider;
 			ReferencedImageResourceProvider = InOptions->OptimisationOptions.ReferencedImageResourceProvider;
@@ -211,7 +209,7 @@ namespace UE::Mutable::Private
 					UE::Mutable::Private::Ptr<ASTOpConstantResource> ConstantOp = new ASTOpConstantResource();
 					ConstantOp->SourceDataDescriptor = SourceDataDescriptor;
 					ConstantOp->Type = EOpType::ME_CONSTANT;
-					ConstantOp->SetValue(MeshBuild, DiskCacheContext);
+					ConstantOp->SetValue(MeshBuild);
 					Result = ConstantOp;
 				  }
 				break;
@@ -228,7 +226,7 @@ namespace UE::Mutable::Private
 					UE::Mutable::Private::Ptr<ASTOpConstantResource> ConstantOp = new ASTOpConstantResource();
 					ConstantOp->SourceDataDescriptor = SourceDataDescriptor;
 					ConstantOp->Type = EOpType::IM_CONSTANT;
-					ConstantOp->SetValue( pImage, DiskCacheContext );
+					ConstantOp->SetValue(pImage);
 					Result = ConstantOp;
 				}
 				break;
@@ -244,7 +242,7 @@ namespace UE::Mutable::Private
 				{
 					UE::Mutable::Private::Ptr<ASTOpConstantResource> constantOp = new ASTOpConstantResource();
 					constantOp->Type = EOpType::LA_CONSTANT;
-					constantOp->SetValue( pLayout, DiskCacheContext);
+					constantOp->SetValue( pLayout);
 					Result = constantOp;
 				}
 				break;
@@ -554,7 +552,7 @@ namespace UE::Mutable::Private
 											MUTABLE_CPUPROFILER_SCOPE(GetSourceDataDescriptor);
 											ConstantOp->SourceDataDescriptor = SubgraphRoot->GetSourceDataDescriptor();
 										}
-										ConstantOp->SetValue(*ResolveImage, InOptions->OptimisationOptions.DiskCacheContext);
+										ConstantOp->SetValue(*ResolveImage);
 									}
 									{
 										MUTABLE_CPUPROFILER_SCOPE(MutableResolveComplete_Replace);
@@ -603,7 +601,7 @@ namespace UE::Mutable::Private
 											MUTABLE_CPUPROFILER_SCOPE(GetSourceDataDescriptor);
 											ConstantOp->SourceDataDescriptor = SubgraphRoot->GetSourceDataDescriptor();
 										}
-										ConstantOp->SetValue(*ResolveMesh, InOptions->OptimisationOptions.DiskCacheContext);
+										ConstantOp->SetValue(*ResolveMesh);
 									}
 									{
 										MUTABLE_CPUPROFILER_SCOPE(MutableResolveComplete_Replace);
@@ -944,7 +942,7 @@ namespace UE::Mutable::Private
 							NewMesh->CheckIntegrity();
 							UE::Mutable::Private::Ptr<ASTOpConstantResource> newOp = new ASTOpConstantResource();
 							newOp->Type = EOpType::ME_CONSTANT;
-							newOp->SetValue(NewMesh, options.DiskCacheContext);
+							newOp->SetValue(NewMesh);
 							newOp->SourceDataDescriptor = at->GetSourceDataDescriptor();
 
 							ASTOp::Replace(at, newOp);
@@ -1062,7 +1060,7 @@ namespace UE::Mutable::Private
 					ASTOpConstantResource* typed = static_cast<ASTOpConstantResource*>(n.get());
 					auto pMesh = StaticCastSharedPtr<const FMesh>(typed->GetValue());
 					pMesh->ResetStaticFormatFlags();
-					typed->SetValue( pMesh, Options->GetPrivate()->OptimisationOptions.DiskCacheContext);
+					typed->SetValue(pMesh);
 				}
 			});
 
diff --git a/Engine/Plugins/Mutable/Source/MutableTools/Private/MuT/Compiler.cpp b/Engine/Plugins/Mutable/Source/MutableTools/Private/MuT/Compiler.cpp
--- a/Engine/Plugins/Mutable/Source/MutableTools/Private/MuT/Compiler.cpp
+++ b/Engine/Plugins/Mutable/Source/MutableTools/Private/MuT/Compiler.cpp
@@ -216,15 +216,6 @@ namespace UE::Mutable::Private
 	}
 
 
-	//---------------------------------------------------------------------------------------------
-	FProxyFileContext::FProxyFileContext()
-	{
-		uint32 Seed = FPlatformTime::Cycles();
-		FRandomStream RandomStream = FRandomStream((int32)Seed);
-		CurrentFileIndex = RandomStream.GetUnsignedInt();
-	}
-
-
     //---------------------------------------------------------------------------------------------
     //---------------------------------------------------------------------------------------------
     //---------------------------------------------------------------------------------------------
@@ -275,13 +266,6 @@ namespace UE::Mutable::Private
     }
 
 
-	//---------------------------------------------------------------------------------------------
-	void CompilerOptions::SetUseDiskCache(bool bEnabled)
-	{
-		m_pD->OptimisationOptions.DiskCacheContext = bEnabled ? &m_pD->DiskCacheContext : nullptr;
-	}
-
-
     //---------------------------------------------------------------------------------------------
     void CompilerOptions::SetOptimisationMaxIteration( int32 MaxIterations )
     {
@@ -350,14 +334,6 @@ namespace UE::Mutable::Private
 		m_pD->OptimisationOptions.bDisableMeshGeneration = bDisabled;
     }
 
-	void CompilerOptions::LogStats() const
-	{
-		UE_LOG(LogMutableCore, Log, TEXT("   Cache Files Written : %" PRIu64), m_pD->DiskCacheContext.FilesWritten.load());
-		UE_LOG(LogMutableCore, Log, TEXT("   Cache Files Read    : %" PRIu64), m_pD->DiskCacheContext.FilesRead.load());
-		UE_LOG(LogMutableCore, Log, TEXT("   Cache MB Written    : %" PRIu64), m_pD->DiskCacheContext.BytesWritten.load() >> 20);
-		UE_LOG(LogMutableCore, Log, TEXT("   Cache MB Read       : %" PRIu64), m_pD->DiskCacheContext.BytesRead.load()>>20);
-	}
-	
 	
 	//---------------------------------------------------------------------------------------------
     //---------------------------------------------------------------------------------------------
diff --git a/Engine/Plugins/Mutable/Source/MutableTools/Private/MuT/CompilerPrivate.h b/Engine/Plugins/Mutable/Source/MutableTools/Private/MuT/CompilerPrivate.h
--- a/Engine/Plugins/Mutable/Source/MutableTools/Private/MuT/CompilerPrivate.h
+++ b/Engine/Plugins/Mutable/Source/MutableTools/Private/MuT/CompilerPrivate.h
@@ -25,28 +25,6 @@ namespace UE::Mutable::Private
 
 	extern const FOpToolsDesc& GetOpToolsDesc(EOpType type);
 
-	/** Statistics about the proxy file usage. */
-	struct FProxyFileContext
-	{
-		FProxyFileContext();
-
-		/** Minimum data size in bytes to dump it to the disk. */
-		uint64 MinProxyFileSize = 64 * 1024;
-
-		/** When creating temporary files, number of retries in case the OS-level call fails. */
-		uint64 MaxFileCreateAttempts = 256;
-
-		/** Statistics */
-		std::atomic<uint64> FilesWritten = 0;
-		std::atomic<uint64> FilesRead = 0;
-		std::atomic<uint64> BytesWritten = 0;
-		std::atomic<uint64> BytesRead = 0;
-
-		/** Internal data. */
-		std::atomic<uint64> CurrentFileIndex = 0;
-	};
-
-		
     //!
     class CompilerOptions::Private
     {
@@ -54,7 +32,6 @@ namespace UE::Mutable::Private
 
         //! Detailed optimization options
         FModelOptimizationOptions OptimisationOptions;
-		FProxyFileContext DiskCacheContext;
 
 		uint64 EmbeddedDataBytesLimit = 1024;
 		uint64 PackagedDataBytesLimit = 1024*1024*64;
diff --git a/Engine/Plugins/Mutable/Source/MutableTools/Private/MuT/DataPacker.cpp b/Engine/Plugins/Mutable/Source/MutableTools/Private/MuT/DataPacker.cpp
--- a/Engine/Plugins/Mutable/Source/MutableTools/Private/MuT/DataPacker.cpp
+++ b/Engine/Plugins/Mutable/Source/MutableTools/Private/MuT/DataPacker.cpp
@@ -548,7 +548,7 @@ namespace UE::Mutable::Private
 					int32 newSize = pNew->GetDataSize();
                     if (float(oldSize) > float(newSize) * OptimizeOptions.MinRLECompressionGain)
                     {
-						Typed->SetValue(pNew, OptimizeOptions.DiskCacheContext);
+						Typed->SetValue(pNew);
                     }
                 }
                 else if (ImageFormatAccumulator.SupportedFormats[Typed][(int32)EImageFormat::L_UByteRLE] )
@@ -560,7 +560,7 @@ namespace UE::Mutable::Private
 					int32 newSize = pNew->GetDataSize();
                     if (float(oldSize) > float(newSize) * OptimizeOptions.MinRLECompressionGain)
                     {
-                        Typed->SetValue(pNew, OptimizeOptions.DiskCacheContext);
+                        Typed->SetValue(pNew);
                     }
                 }
             }
@@ -577,7 +577,7 @@ namespace UE::Mutable::Private
 			ASTOpConstantResource* Op = Entry.Key.get();
 			TSharedPtr<FMesh> Mesh = static_cast<const FMesh*>(Op->GetValue().Get())->Clone();
 			MeshRemoveUnusedBufferSemantics(Mesh.Get(), Entry.Value);
-			Op->SetValue(Mesh, OptimizeOptions.DiskCacheContext);
+			Op->SetValue(Mesh);
 		}
 
 		// See if we can remove entire "prepare layout" operations
diff --git a/Engine/Plugins/Mutable/Source/MutableTools/Private/MuT/NodeImageConstant.cpp b/Engine/Plugins/Mutable/Source/MutableTools/Private/MuT/NodeImageConstant.cpp
--- a/Engine/Plugins/Mutable/Source/MutableTools/Private/MuT/NodeImageConstant.cpp
+++ b/Engine/Plugins/Mutable/Source/MutableTools/Private/MuT/NodeImageConstant.cpp
@@ -10,16 +10,4 @@
 
 namespace UE::Mutable::Private
 {
-
-	void NodeImageConstant::SetValue(TSharedPtr<const FImage> Value)
-	{
-		Proxy = new TResourceProxyMemory<FImage>(Value);
-	}
-
-
-	void NodeImageConstant::SetValue(Ptr<TResourceProxy<FImage>> InProxy)
-	{
-		Proxy = InProxy;
-	}
-
 }
diff --git a/Engine/Plugins/Mutable/Source/MutableTools/Private/MuT/Table.cpp b/Engine/Plugins/Mutable/Source/MutableTools/Private/MuT/Table.cpp
--- a/Engine/Plugins/Mutable/Source/MutableTools/Private/MuT/Table.cpp
+++ b/Engine/Plugins/Mutable/Source/MutableTools/Private/MuT/Table.cpp
@@ -139,7 +139,7 @@ namespace UE::Mutable::Private
 	}
 
 
-    void FTable::SetCell(int32 Column, uint32 RowId, TResourceProxy<FImage>* Value, const void* ErrorContext)
+    void FTable::SetCell(int32 Column, uint32 RowId, const TSharedPtr<FImage>& Value, const void* ErrorContext)
 	{
 		int32 Row = m_pD->FindRow(RowId);
 		check( Row>=0 );
@@ -147,7 +147,7 @@ namespace UE::Mutable::Private
 		check( Column < m_pD->Rows[Row].Values.Num() );
 		check( m_pD->Columns[Column].Type == ETableColumnType::Image );
 
-		m_pD->Rows[ Row ].Values[Column].ProxyImage = Value;
+		m_pD->Rows[ Row ].Values[Column].Image = Value;
 		m_pD->Rows[ Row ].Values[Column].ErrorContext = ErrorContext;
 	}
 
diff --git a/Engine/Plugins/Mutable/Source/MutableValidation/Private/MuV/CustomizableObjectBulkValidationCommandlet.cpp b/Engine/Plugins/Mutable/Source/MutableValidation/Private/MuV/CustomizableObjectBulkValidationCommandlet.cpp
--- a/Engine/Plugins/Mutable/Source/MutableValidation/Private/MuV/CustomizableObjectBulkValidationCommandlet.cpp
+++ b/Engine/Plugins/Mutable/Source/MutableValidation/Private/MuV/CustomizableObjectBulkValidationCommandlet.cpp
@@ -34,8 +34,7 @@ int32 UCustomizableObjectBulkValidationCommandlet::Main(const FString& Params)
 	{
 		return 1;
 	}
-	// Parse if we want to use disk compilation or not
-	const bool bUseDiskCompilation = GetDiskCompilationArg(Params);
+
 	// Get the amount of instances to generate if parameter was provided (it will get multiplied by the amount of states later so this is a minimun value)
 	const uint32 InstancesToGenerate = GetTargetAmountOfInstances(Params);
 
@@ -135,7 +134,7 @@ int32 UCustomizableObjectBulkValidationCommandlet::Main(const FString& Params)
 		}
 		
 		// Body of the test ----------------------------------------------------------------------------------------------------------------------
-		TestCustomizableObject(*TargetCustomizableObject, *TargetCompilationPlatform, InstancesToGenerate, bUseDiskCompilation);
+		TestCustomizableObject(*TargetCustomizableObject, *TargetCompilationPlatform, InstancesToGenerate);
 		// ---------------------------------------------------------------------------------------------------------------------------------------
 		
 		// Try to collect the garbage
diff --git a/Engine/Plugins/Mutable/Source/MutableValidation/Private/MuV/CustomizableObjectCompilationUtility.cpp b/Engine/Plugins/Mutable/Source/MutableValidation/Private/MuV/CustomizableObjectCompilationUtility.cpp
--- a/Engine/Plugins/Mutable/Source/MutableValidation/Private/MuV/CustomizableObjectCompilationUtility.cpp
+++ b/Engine/Plugins/Mutable/Source/MutableValidation/Private/MuV/CustomizableObjectCompilationUtility.cpp
@@ -55,7 +55,6 @@ bool FCustomizableObjectCompilationUtility::CompileCustomizableObject(UCustomiza
 		UE_LOG(LogMutable, Log, TEXT("(string) model_compile_options_overriden : %s "), bOverrideCompilationOptions ? TEXT("true") : TEXT("false"));
 		UE_LOG(LogMutable, Log, TEXT("(int) model_optimization_level : %d "), CompilationOptions.OptimizationLevel);
 		UE_LOG(LogMutable, Log, TEXT("(string) model_texture_compression : %s "), *UEnum::GetValueAsString(CompilationOptions.TextureCompression));
-		UE_LOG(LogMutable, Log, TEXT("(string) model_disk_compilation : %s "), CompilationOptions.bUseDiskCompilation ? TEXT("true") : TEXT("false"));
 		UE_LOG(LogMutable, Log, TEXT("(string) model_compile_platform_name : %s "), *CompilationOptions.TargetPlatform->PlatformName());
 		UE_LOG(LogMutable, Log, TEXT("(int) model_package_data_bytes_limit_bytes : %llu "), CompilationOptions.PackagedDataBytesLimit);
 	}
diff --git a/Engine/Plugins/Mutable/Source/MutableValidation/Private/MuV/CustomizableObjectValidationCommandlet.cpp b/Engine/Plugins/Mutable/Source/MutableValidation/Private/MuV/CustomizableObjectValidationCommandlet.cpp
--- a/Engine/Plugins/Mutable/Source/MutableValidation/Private/MuV/CustomizableObjectValidationCommandlet.cpp
+++ b/Engine/Plugins/Mutable/Source/MutableValidation/Private/MuV/CustomizableObjectValidationCommandlet.cpp
@@ -65,8 +65,7 @@ int32 UCustomizableObjectValidationCommandlet::Main(const FString& Params)
 	{
 		return 1;
 	}
-	// Parse if we want to use disk compilation or not
-	const bool bUseDiskCompilation = GetDiskCompilationArg(Params);
+
 	const uint32 InstancesToGenerate = GetTargetAmountOfInstances(Params);
 	
 	// Perform a blocking search to ensure all assets used by mutable are reachable using the AssetRegistry
@@ -85,7 +84,7 @@ int32 UCustomizableObjectValidationCommandlet::Main(const FString& Params)
 	LogGlobalSettings();
 
 	// Body of the test ------------------------------------------------------------------------------------------------------------------------------
-	TestCustomizableObject(*TargetCustomizableObject, *TargetCompilationPlatform, InstancesToGenerate, bUseDiskCompilation);
+	TestCustomizableObject(*TargetCustomizableObject, *TargetCompilationPlatform, InstancesToGenerate);
 	// -----------------------------------------------------------------------------------------------------------------------------------------------
 	
 	UE_LOG(LogMutable, Display, TEXT("Mutable commandlet finished."));
diff --git a/Engine/Plugins/Mutable/Source/MutableValidation/Private/MuV/ValidationUtils.cpp b/Engine/Plugins/Mutable/Source/MutableValidation/Private/MuV/ValidationUtils.cpp
--- a/Engine/Plugins/Mutable/Source/MutableValidation/Private/MuV/ValidationUtils.cpp
+++ b/Engine/Plugins/Mutable/Source/MutableValidation/Private/MuV/ValidationUtils.cpp
@@ -126,17 +126,6 @@ TArray<FAssetData> FindAllAssetsAtPath(FName SearchPath, const UClass* TargetObj
 }
 
 
-bool GetDiskCompilationArg(const FString& Params)
-{
-	bool bUseDiskCompilation = false;
-	if (! FParse::Bool(*Params,TEXT("UseDiskCompilation="),bUseDiskCompilation))
-	{
-		UE_LOG(LogMutable, Display, TEXT("Disk compilation setting for the compilation of the CO not specified. Using default value : %hhd"), bUseDiskCompilation) ;
-	}
-	return bUseDiskCompilation;
-}
-
-
 uint32 GetTargetAmountOfInstances(const FString& Params)
 {
 	// Get the amount of instances to generate if parameter was provided (it will get multiplied by the amount of states later so this is a minimun value)
@@ -186,7 +175,7 @@ ITargetPlatform* GetCompilationPlatform(const FString& Params)
 
 
 bool TestCustomizableObject(UCustomizableObject& InTargetCustomizableObject, const ITargetPlatform& TargetCompilationPlatform,
-	const uint32 InstancesToGenerate, const bool bUseDiskCompilation)
+	const uint32 InstancesToGenerate)
 {
 	const FScopedLogSection ObjectSection (EMutableLogSection::Object, FName( InTargetCustomizableObject.GetPathName()));
 	
@@ -202,8 +191,6 @@ bool TestCustomizableObject(UCustomizableObject& InTargetCustomizableObject, con
 		
 		// Set the target compilation platform based on what the caller wants
 		CompilationOptions.TargetPlatform = &TargetCompilationPlatform;
-		// Disk cache usage for compilation operation
-		CompilationOptions.bUseDiskCompilation = bUseDiskCompilation;
 		
 		TSharedRef<FCustomizableObjectCompilationUtility> CompilationUtility = MakeShared<FCustomizableObjectCompilationUtility>();
 		bWasCoCompilationSuccessful = CompilationUtility->CompileCustomizableObject(InTargetCustomizableObject, true, &CompilationOptions);
diff --git a/Engine/Plugins/Mutable/Source/MutableValidation/Private/MuV/ValidationUtils.h b/Engine/Plugins/Mutable/Source/MutableValidation/Private/MuV/ValidationUtils.h
--- a/Engine/Plugins/Mutable/Source/MutableValidation/Private/MuV/ValidationUtils.h
+++ b/Engine/Plugins/Mutable/Source/MutableValidation/Private/MuV/ValidationUtils.h
@@ -9,13 +9,6 @@ class UCustomizableObject;
 
 // ARGUMENT PARSING ----------------------------------------------------------------------------------------------------------------------------------
 
-/**
- * Get the type of compilation the caller wants to run. 
- * @param Params The parameters fed to this commandlet.
- * @return True if the use of the disk cache is wanted for the compilation, false otherwise.
- */
-bool GetDiskCompilationArg(const FString& Params);
-
 
 /**
  * Get the amount of instances we want to generate.
@@ -88,9 +81,8 @@ TArray<FAssetData> FindAllAssetsAtPath(FName SearchPath, const UClass* TargetObj
  * @param InTargetCustomizableObject The Customizable Object to be compiled and used as the base for the generated instances.
  * @param TargetCompilationPlatform The target platform used to compile the Customizable Object against.
  * @param InstancesToGenerate The amount of instances to generate per state.
- * @param bUseDiskCompilation Determines if a disk cache will be used for the compilation of the Customizable Object.
  * @return True if all went well and false if the testing operation did throw errors.
  */
 bool TestCustomizableObject(UCustomizableObject& InTargetCustomizableObject, const ITargetPlatform& TargetCompilationPlatform,
-	const uint32 InstancesToGenerate,const bool bUseDiskCompilation);
+	const uint32 InstancesToGenerate);
 	
