[Mutable] Add missing ASTC RG formats to the image provider engine to mutable conversion function.

- Remove stray inline in the format selection functions.

#rnx
#Synced-CL 48510493

diff --git a/Engine/Plugins/Mutable/Source/CustomizableObject/Private/MuCO/UnrealMutableImageProvider.cpp b/Engine/Plugins/Mutable/Source/CustomizableObject/Private/MuCO/UnrealMutableImageProvider.cpp
--- a/Engine/Plugins/Mutable/Source/CustomizableObject/Private/MuCO/UnrealMutableImageProvider.cpp
+++ b/Engine/Plugins/Mutable/Source/CustomizableObject/Private/MuCO/UnrealMutableImageProvider.cpp
@@ -79,10 +79,15 @@ UE::Mutable::Private::EImageFormat GetMutablePixelFormat(EPixelFormat InTextureF
 	case PF_BC5: return UE::Mutable::Private::EImageFormat::BC5;
 	case PF_G8: return UE::Mutable::Private::EImageFormat::L_UByte;
 	case PF_ASTC_4x4: return UE::Mutable::Private::EImageFormat::ASTC_4x4_RGBA_LDR;
+	case PF_ASTC_4x4_NORM_RG: return UE::Mutable::Private::EImageFormat::ASTC_4x4_RG_LDR;
 	case PF_ASTC_6x6: return UE::Mutable::Private::EImageFormat::ASTC_6x6_RGBA_LDR;
+	case PF_ASTC_6x6_NORM_RG: return UE::Mutable::Private::EImageFormat::ASTC_6x6_RG_LDR;
 	case PF_ASTC_8x8: return UE::Mutable::Private::EImageFormat::ASTC_8x8_RGBA_LDR;
+	case PF_ASTC_8x8_NORM_RG: return UE::Mutable::Private::EImageFormat::ASTC_8x8_RG_LDR;
 	case PF_ASTC_10x10: return UE::Mutable::Private::EImageFormat::ASTC_10x10_RGBA_LDR;
+	case PF_ASTC_10x10_NORM_RG: return UE::Mutable::Private::EImageFormat::ASTC_10x10_RG_LDR;
 	case PF_ASTC_12x12: return UE::Mutable::Private::EImageFormat::ASTC_12x12_RGBA_LDR;
+	case PF_ASTC_12x12_NORM_RG: return UE::Mutable::Private::EImageFormat::ASTC_12x12_RG_LDR;
 	default: return UE::Mutable::Private::EImageFormat::None;
 	}
 }
diff --git a/Engine/Plugins/Mutable/Source/MutableRuntime/Private/MuR/ImageFormatUtils.cpp b/Engine/Plugins/Mutable/Source/MutableRuntime/Private/MuR/ImageFormatUtils.cpp
--- a/Engine/Plugins/Mutable/Source/MutableRuntime/Private/MuR/ImageFormatUtils.cpp
+++ b/Engine/Plugins/Mutable/Source/MutableRuntime/Private/MuR/ImageFormatUtils.cpp
@@ -303,7 +303,7 @@ namespace
 		return DecompressionFuncRefType(DecompressionFuncNoOp);
 	}
 
-	inline miro::SubImageCompression::FuncRefType SelectCompressionFunction(EImageFormat DestFormat, EImageFormat SrcFormat)
+	miro::SubImageCompression::FuncRefType SelectCompressionFunction(EImageFormat DestFormat, EImageFormat SrcFormat)
 	{
 		using CompressionFuncRefType = miro::SubImageCompression::FuncRefType;
 
diff --git a/Engine/Plugins/Mutable/Source/MutableRuntime/Private/MuR/ImageFormatUtils.h b/Engine/Plugins/Mutable/Source/MutableRuntime/Private/MuR/ImageFormatUtils.h
--- a/Engine/Plugins/Mutable/Source/MutableRuntime/Private/MuR/ImageFormatUtils.h
+++ b/Engine/Plugins/Mutable/Source/MutableRuntime/Private/MuR/ImageFormatUtils.h
@@ -11,8 +11,8 @@ namespace UE::Mutable::Private
 namespace SubImageFormat
 {	
 
-	extern miro::SubImageDecompression::FuncRefType SelectDecompressionFunction(EImageFormat DestFormat, EImageFormat SrcFormat);
-	extern miro::SubImageCompression::FuncRefType SelectCompressionFunction(EImageFormat DestFormat, EImageFormat SrcFormat);	
+	miro::SubImageDecompression::FuncRefType SelectDecompressionFunction(EImageFormat DestFormat, EImageFormat SrcFormat);
+	miro::SubImageCompression::FuncRefType SelectCompressionFunction(EImageFormat DestFormat, EImageFormat SrcFormat);	
 
 	using UncompressedFormatFuncType = void(FImageSize, FImageSize, FImageSize, const uint8*, uint8*); 
 	UncompressedFormatFuncType* SelectUncompressedFormatFunction(EImageFormat DestFormat, EImageFormat SrcFormat);
