[Mutable] Improved material system and implemented a new node that can modify materials.
+ Implemented Modify Material Node.
+ Remode operations IN_ADDIMAGE, IN_ADDVECTOR and IN_ADDSCALAR. Now a surface node creates a ASTOpMaterialModify operation which creates a new FMaterial.

#rnx

#Synced-CL 48603048

diff --git a/Engine/Plugins/Mutable/Source/CustomizableObject/Internal/MuCO/CustomizableObjectPrivate.h b/Engine/Plugins/Mutable/Source/CustomizableObject/Internal/MuCO/CustomizableObjectPrivate.h
--- a/Engine/Plugins/Mutable/Source/CustomizableObject/Internal/MuCO/CustomizableObjectPrivate.h
+++ b/Engine/Plugins/Mutable/Source/CustomizableObject/Internal/MuCO/CustomizableObjectPrivate.h
@@ -1763,6 +1763,8 @@ enum class ECustomizableObjectVersions : int32
 
 	TexutreParametersFormatFromReferenceTexture,
 
+	AddMaterialModify,
+
 	// -----<new versions can be added above this line>--------
 	LastCustomizableObjectVersion
 };
diff --git a/Engine/Plugins/Mutable/Source/CustomizableObject/Internal/MuCO/CustomizableObjectSystemPrivate.h b/Engine/Plugins/Mutable/Source/CustomizableObject/Internal/MuCO/CustomizableObjectSystemPrivate.h
--- a/Engine/Plugins/Mutable/Source/CustomizableObject/Internal/MuCO/CustomizableObjectSystemPrivate.h
+++ b/Engine/Plugins/Mutable/Source/CustomizableObject/Internal/MuCO/CustomizableObjectSystemPrivate.h
@@ -322,7 +322,6 @@ struct FInstanceUpdateData
 {
 	struct FImage
 	{
-		FName Name; // TODO MATERIAL_PIN Remove this hack
 		FName ParameterName; // TODO MATERIAL_PIN This does not belong to the Image but the Material
 		UE::Mutable::Private::FImageId ImageID;
 		
@@ -338,6 +337,8 @@ struct FInstanceUpdateData
 
 		TArray<int32> ConstantImagesNeededToGenerate;
 
+		int32 ImagePropertiesIndex = INDEX_NONE; // TODO MATERIAL_PIN This does not belong to the Image but the Material
+
 		bool bIsPassThrough = false;
 		bool bIsNonProgressive = false;
 
@@ -373,9 +374,7 @@ struct FInstanceUpdateData
 		/** Range in the Scalar array */
 		uint16 FirstScalar = 0;
 		uint16 ScalarCount = 0;
-		
-		bool bGetParameters = true; // TODO MATERIAL_PIN Remove once the Material Pin is completed. Only used to avid GetImages/Vector/Scalar twice (Material and Surface).
-		
+				
 		UE::Mutable::Private::FMaterialId MaterialId;
 		TSharedPtr<const UE::Mutable::Private::FMaterial> Material;
 	};
diff --git a/Engine/Plugins/Mutable/Source/CustomizableObject/Internal/MuCO/UnrealToMutableTextureConversionUtils.h b/Engine/Plugins/Mutable/Source/CustomizableObject/Internal/MuCO/UnrealToMutableTextureConversionUtils.h
--- a/Engine/Plugins/Mutable/Source/CustomizableObject/Internal/MuCO/UnrealToMutableTextureConversionUtils.h
+++ b/Engine/Plugins/Mutable/Source/CustomizableObject/Internal/MuCO/UnrealToMutableTextureConversionUtils.h
@@ -4,6 +4,7 @@
 
 #include "HAL/Platform.h"
 #include "Engine/Texture.h"
+#include "MuR/Ptr.h"
 
 #define UE_API CUSTOMIZABLEOBJECT_API
 
@@ -140,20 +141,6 @@ inline uint32 GetTypeHash(const FMutableSourceMeshData& Key)
 }
 
 
-struct FSourceExternalOptions
-{
-	// GenerateMutableSourceMesh
-	FMutableSourceMeshData MeshOptions;
-	bool bLinkedToExtendMaterial = false;
-	bool bOnlyConnectedLOD = false;
-	
-	// GenerateMutableSourceImage
-	int32 ReferenceTextureSize = 0;
-
-	bool operator==(const FSourceExternalOptions& Other) const = default;
-};
-
-
 enum class EUnrealToMutableConversionError 
 {
     Success,
diff --git a/Engine/Plugins/Mutable/Source/CustomizableObject/Private/MuCO/CustomizableObjectInstance.cpp b/Engine/Plugins/Mutable/Source/CustomizableObject/Private/MuCO/CustomizableObjectInstance.cpp
--- a/Engine/Plugins/Mutable/Source/CustomizableObject/Private/MuCO/CustomizableObjectInstance.cpp
+++ b/Engine/Plugins/Mutable/Source/CustomizableObject/Private/MuCO/CustomizableObjectInstance.cpp
@@ -5598,7 +5598,6 @@ UTexture* BuildTexture(const TSharedRef<FUpdateContextPrivate>& Context, const i
 	check((bNeverStream && Context->MipsToSkip == 0) ||
 		(!bNeverStream && Context->MipsToSkip >= 0));
 	
-	FString KeyName = Image->Name.ToString();
 	TSharedPtr<const UE::Mutable::Private::FImage> MutableImage = Image->Image;
 
 	UTexture2D* Texture; // Texture generated by mutable
@@ -5618,7 +5617,7 @@ UTexture* BuildTexture(const TSharedRef<FUpdateContextPrivate>& Context, const i
 	else
 	{
 		// Find the additional information for this image
-    	int32 ImageKey = FCString::Atoi(*KeyName);
+    	int32 ImageKey = Image->ImagePropertiesIndex;
     	check (ImageKey >= 0 && ImageKey < ModelResources.ImageProperties.Num());
     	
     	const FMutableModelImageProperties& Props = ModelResources.ImageProperties[ImageKey];
diff --git a/Engine/Plugins/Mutable/Source/CustomizableObject/Private/MuCO/CustomizableObjectSystem.cpp b/Engine/Plugins/Mutable/Source/CustomizableObject/Private/MuCO/CustomizableObjectSystem.cpp
--- a/Engine/Plugins/Mutable/Source/CustomizableObject/Private/MuCO/CustomizableObjectSystem.cpp
+++ b/Engine/Plugins/Mutable/Source/CustomizableObject/Private/MuCO/CustomizableObjectSystem.cpp
@@ -2311,72 +2311,6 @@ namespace impl
 						TSharedRef<FInstanceUpdateData::FMaterial>& Material = OperationData->InstanceUpdateData.Materials.Add_GetRef(MakeShared<FInstanceUpdateData::FMaterial>());
 
 						Material->MaterialId = MaterialId;
-
-						// TODO MATERIAL_PIN Remove
-
-						Material->bGetParameters = false;
-
-						// Vectors
-						Material->FirstVector = OperationData->InstanceUpdateData.Vectors.Num();
-						Material->VectorCount = MutableInstance->GetVectorCount(InstanceComponentIndex, MutableLODIndex, InstanceSurfaceIndex);
-						for (int32 VectorIndex = 0; VectorIndex < Material->VectorCount; ++VectorIndex)
-						{
-							MUTABLE_CPUPROFILER_SCOPE(GetVector);
-							OperationData->InstanceUpdateData.Vectors.Push({});
-							FInstanceUpdateData::FVector& Vector = OperationData->InstanceUpdateData.Vectors.Last();
-							Vector.Name = MutableInstance->GetVectorName(InstanceComponentIndex, MutableLODIndex, InstanceSurfaceIndex, VectorIndex);
-							Vector.Vector = MutableInstance->GetVector(InstanceComponentIndex, MutableLODIndex, InstanceSurfaceIndex, VectorIndex);
-							Vector.MaterialLayer = MutableInstance->GetVectorLayerIndex(InstanceComponentIndex, MutableLODIndex, InstanceSurfaceIndex, VectorIndex);
-						}
-
-						// Scalars
-						Material->FirstScalar = OperationData->InstanceUpdateData.Scalars.Num();
-						Material->ScalarCount = MutableInstance->GetScalarCount(InstanceComponentIndex, MutableLODIndex, InstanceSurfaceIndex);
-						for (int32 ScalarIndex = 0; ScalarIndex < Material->ScalarCount; ++ScalarIndex)
-						{
-							MUTABLE_CPUPROFILER_SCOPE(GetScalar)
-
-							const FName ScalarName = MutableInstance->GetScalarName(InstanceComponentIndex, MutableLODIndex, InstanceSurfaceIndex, ScalarIndex);
-							const float ScalarValue = MutableInstance->GetScalar(InstanceComponentIndex, MutableLODIndex, InstanceSurfaceIndex, ScalarIndex);
-
-							const int32 MaterialLayer = MutableInstance->GetScalarLayerIndex(InstanceComponentIndex, MutableLODIndex, InstanceSurfaceIndex, ScalarIndex);
-
-							OperationData->InstanceUpdateData.Scalars.Push({ ScalarName, ScalarValue, MaterialLayer });
-						}
-
-						// Images
-						Material->FirstImage = OperationData->InstanceUpdateData.Images.Num();
-						Material->ImageCount = MutableInstance->GetImageCount(InstanceComponentIndex, MutableLODIndex, InstanceSurfaceIndex);
-						for (int32 ImageIndex = 0; ImageIndex < Material->ImageCount; ++ImageIndex)
-						{
-							MUTABLE_CPUPROFILER_SCOPE(GetImageId);
-							
-							TSharedRef<FInstanceUpdateData::FImage>& Image = OperationData->InstanceUpdateData.Images.Add_GetRef(MakeShared<FInstanceUpdateData::FImage>());
-
-							Image->Name = MutableInstance->GetImageName(InstanceComponentIndex, MutableLODIndex, InstanceSurfaceIndex, ImageIndex);
-							Image->ImageID = MutableInstance->GetImageId(InstanceComponentIndex, MutableLODIndex, InstanceSurfaceIndex, ImageIndex);
-							Image->MaterialLayer = MutableInstance->GetImageLayerIndex(InstanceComponentIndex, MutableLODIndex, InstanceSurfaceIndex, ImageIndex);
-							Image->FullImageSizeX = 0;
-							Image->FullImageSizeY = 0;
-							Image->BaseMip = 0;
-
-							FString KeyName = Image->Name.ToString();
-							int32 ImageKey = FCString::Atoi(*KeyName);
-
-							if (ImageKey >= 0 && ImageKey < ModelResources.ImageProperties.Num())
-							{
-								const FMutableModelImageProperties& Props = ModelResources.ImageProperties[ImageKey];
-
-								Image->ParameterName = FName(Props.TextureParameterName);
-								Image->bIsNonProgressive = Props.MipGenSettings == TMGS_NoMipmaps;
-								Image->bIsPassThrough = Props.IsPassThrough;
-							}
-							else
-							{
-								// This means the compiled model (maybe coming from derived data) has images that the asset doesn't know about.
-								UE_LOG(LogMutable, Error, TEXT("CustomizableObject derived data out of sync with asset for [%s]. Try recompiling it."), *CustomizableObject->GetName());
-							}
-						}
 					}
 				}
 			}
@@ -2406,54 +2340,51 @@ namespace impl
 			{
 				Material->Material = GetMaterialTask.GetResult();
 
-				if (Material->bGetParameters)
+				// Scalars
+				Material->FirstScalar = Context->InstanceUpdateData.Scalars.Num();
+				Material->ScalarCount = Material->Material->ScalarParameters.Num();
+				for (TTuple<UE::Mutable::Private::FParameterKey, float> ScalarParameter : Material->Material->ScalarParameters)
 				{
-					// Scalars
-					Material->FirstScalar = Context->InstanceUpdateData.Scalars.Num();
-					Material->ScalarCount = Material->Material->ScalarParameters.Num();
-					for (TTuple<FName, float> ScalarParameter : Material->Material->ScalarParameters)
-					{
-						FInstanceUpdateData::FScalar& Scalar = Context->InstanceUpdateData.Scalars.AddDefaulted_GetRef();
-						Scalar.Name = ScalarParameter.Key;
-						Scalar.Scalar = ScalarParameter.Value;
-						//Scalar.MaterialLayer = ScalarParameter.Key.Layer; // TODO MATERIAL_PIN
-					}
+					FInstanceUpdateData::FScalar& Scalar = Context->InstanceUpdateData.Scalars.AddDefaulted_GetRef();
+					Scalar.Name = ScalarParameter.Key.ParameterName;
+					Scalar.MaterialLayer = ScalarParameter.Key.LayerIndex;
+					Scalar.Scalar = ScalarParameter.Value;
+					
+				}
 
-					// Vectors
-					Material->FirstVector = Context->InstanceUpdateData.Vectors.Num();
-					Material->VectorCount = Material->Material->ColorParameters.Num();
-					for (TTuple<FName, UE::Math::TVector4<float>> VectorParameter : Material->Material->ColorParameters)
-					{
-						FInstanceUpdateData::FVector& Vector = Context->InstanceUpdateData.Vectors.AddDefaulted_GetRef();
-						Vector.Name = VectorParameter.Key;
-						Vector.Vector = VectorParameter.Value;
-						//Vector.MaterialLayer =  VectorParameter.Key.Layer; // TODO MATERIAL_PIN
-					}
+				// Vectors
+				Material->FirstVector = Context->InstanceUpdateData.Vectors.Num();
+				Material->VectorCount = Material->Material->ColorParameters.Num();
+				for (TTuple<UE::Mutable::Private::FParameterKey, UE::Math::TVector4<float>> VectorParameter : Material->Material->ColorParameters)
+				{
+					FInstanceUpdateData::FVector& Vector = Context->InstanceUpdateData.Vectors.AddDefaulted_GetRef();
+					Vector.Name = VectorParameter.Key.ParameterName;
+					Vector.MaterialLayer = VectorParameter.Key.LayerIndex;
+					Vector.Vector = VectorParameter.Value;
+				}
 
-					// Images
-					Material->FirstImage = Context->InstanceUpdateData.Images.Num();
-					Material->ImageCount = Material->Material->ImageParameters.Num();
-					for (TTuple<FName, TVariant<uint32, TSharedPtr<const UE::Mutable::Private::FImage>>> ImageParameter : Material->Material->ImageParameters)
-					{
-						const UE::Mutable::Private::OP::ADDRESS Address = ImageParameter.Value.Get<UE::Mutable::Private::OP::ADDRESS>();
-						
-						TSharedRef<FInstanceUpdateData::FImage>& Image = Context->InstanceUpdateData.Images.Add_GetRef(MakeShared<FInstanceUpdateData::FImage>());
-						Image->ParameterName = ImageParameter.Key;
-						// Image->Name = ImageParameter.Key; // TODO MATERIAL_PIN It should not have a name but is it currently used?
-						Image->ImageID = UE::Mutable::Private::FSystem::GetImageId(Context->ImageIdRegistry, Context->Model, Context->Parameters.Get(), Address);
-						//Image->MaterialLayer = ImageParameter.Key.Layer; // TODO MATERIAL_PIN
-						Image->FullImageSizeX = 0;
-						Image->FullImageSizeY = 0;
-						Image->BaseMip = 0;
-
-						FString KeyName = Image->Name.ToString();
-						int32 ImageKey = FCString::Atoi(*KeyName);
-
-						const FMutableModelImageProperties& Props = ModelResources->ImageProperties[ImageKey];
-						Image->bIsNonProgressive = Props.MipGenSettings == TMGS_NoMipmaps;
-						Image->bIsPassThrough = Props.IsPassThrough;
-					}
+				// Images
+				Material->FirstImage = Context->InstanceUpdateData.Images.Num();
+				Material->ImageCount = Material->Material->ImageParameters.Num();
+				for (TTuple<UE::Mutable::Private::FParameterKey, UE::Mutable::Private::FMaterial::FImageParameterData> ImageParameter : Material->Material->ImageParameters)
+				{
+					const UE::Mutable::Private::OP::ADDRESS Address = ImageParameter.Value.ImageParameter.Get<UE::Mutable::Private::OP::ADDRESS>();
+					
+					TSharedRef<FInstanceUpdateData::FImage>& Image = Context->InstanceUpdateData.Images.Add_GetRef(MakeShared<FInstanceUpdateData::FImage>());
+					Image->ParameterName = ImageParameter.Key.ParameterName;
+					Image->MaterialLayer = ImageParameter.Key.LayerIndex;
+					Image->ImageID = UE::Mutable::Private::FSystem::GetImageId(Context->ImageIdRegistry, Context->Model, Context->Parameters.Get(), Address);
+					Image->FullImageSizeX = 0;
+					Image->FullImageSizeY = 0;
+					Image->BaseMip = 0;
+					Image->ImagePropertiesIndex = ImageParameter.Value.ImagePropertyIndex;
+
+					check(ModelResources->ImageProperties.IsValidIndex(Image->ImagePropertiesIndex))
+					const FMutableModelImageProperties& Props = ModelResources->ImageProperties[Image->ImagePropertiesIndex];
+					Image->bIsNonProgressive = Props.MipGenSettings == TMGS_NoMipmaps;
+					Image->bIsPassThrough = Props.IsPassThrough;
 				}
+				
 			},
 			GetMaterialTask,
 			LowLevelTasks::ETaskPriority::Inherit);
diff --git a/Engine/Plugins/Mutable/Source/CustomizableObjectEditor/Private/MuCOE/CustomizableObjectCompileRunnable.cpp b/Engine/Plugins/Mutable/Source/CustomizableObjectEditor/Private/MuCOE/CustomizableObjectCompileRunnable.cpp
--- a/Engine/Plugins/Mutable/Source/CustomizableObjectEditor/Private/MuCOE/CustomizableObjectCompileRunnable.cpp
+++ b/Engine/Plugins/Mutable/Source/CustomizableObjectEditor/Private/MuCOE/CustomizableObjectCompileRunnable.cpp
@@ -318,10 +318,10 @@ uint32 FCustomizableObjectCompileRunnable::Run()
 	
 	CompilerOptions->SetImageTiling(Options.ImageTiling);
 
-	// On server builds we don't want the images to be generted.
+	// On server builds we don't want materials and images to be generted.
 	if (Options.TargetPlatform && Options.TargetPlatform->IsServerOnly())
 	{
-		CompilerOptions->SetDisableImageGeneration(true);
+		CompilerOptions->SetDisableMaterialGeneration(true);
 	}
 
 	TFunction<void()> WaitCallback = [WeakCompiler=WeakCompiler]()
diff --git a/Engine/Plugins/Mutable/Source/CustomizableObjectEditor/Private/MuCOE/GenerateMutableSource/GenerateMutableSource.cpp b/Engine/Plugins/Mutable/Source/CustomizableObjectEditor/Private/MuCOE/GenerateMutableSource/GenerateMutableSource.cpp
--- a/Engine/Plugins/Mutable/Source/CustomizableObjectEditor/Private/MuCOE/GenerateMutableSource/GenerateMutableSource.cpp
+++ b/Engine/Plugins/Mutable/Source/CustomizableObjectEditor/Private/MuCOE/GenerateMutableSource/GenerateMutableSource.cpp
@@ -1553,9 +1553,9 @@ void PopulateReferenceSkeletalMeshesData(FMutableGraphGenerationContext& Generat
 }
 
 
-uint32 GetBaseTextureSize(const FMutableGraphGenerationContext& GenerationContext, const UCustomizableObjectNodeMaterialBase* Material, uint32 ImageIndex)
+uint32 GetBaseTextureSize(const FMutableGraphGenerationContext& GenerationContext, const UCustomizableObjectNodeMaterialBase* Material, const UE::Mutable::Private::FParameterKey& ImageKey)
 {
-	const FGeneratedImageProperties* ImageProperties = GenerationContext.ImageProperties.Find({ Material, ImageIndex });
+	const FGeneratedImageProperties* ImageProperties = GenerationContext.ImageProperties.Find({ Material, ImageKey });
 	return ImageProperties ? ImageProperties->TextureSize : 0;
 }
 
diff --git a/Engine/Plugins/Mutable/Source/CustomizableObjectEditor/Private/MuCOE/GenerateMutableSource/GenerateMutableSource.h b/Engine/Plugins/Mutable/Source/CustomizableObjectEditor/Private/MuCOE/GenerateMutableSource/GenerateMutableSource.h
--- a/Engine/Plugins/Mutable/Source/CustomizableObjectEditor/Private/MuCOE/GenerateMutableSource/GenerateMutableSource.h
+++ b/Engine/Plugins/Mutable/Source/CustomizableObjectEditor/Private/MuCOE/GenerateMutableSource/GenerateMutableSource.h
@@ -181,6 +181,23 @@ public:
 
 uint32 GetTypeHash(const FGeneratedKey& Key);
 
+// Struct used for external options of the GenerateMutableSourceX functions.
+struct FSourceExternalOptions
+{
+	// GenerateMutableSourceMesh
+	FMutableSourceMeshData MeshOptions;
+	bool bLinkedToExtendMaterial = false;
+	bool bOnlyConnectedLOD = false;
+
+	// GenerateMutableSourceImage
+	int32 ReferenceTextureSize = 0;
+
+	// GenerateMutableSourceMaterial
+	UE::Mutable::Private::Ptr<UE::Mutable::Private::NodeSurfaceNew> SurfaceNode = nullptr;
+
+	bool operator==(const FSourceExternalOptions& Other) const = default;
+};
+
 
 struct FGeneratedSourceExternalKey
 {
@@ -223,20 +240,20 @@ struct FImageKey
 
 struct FGeneratedImagePropertiesKey
 {
-	FGeneratedImagePropertiesKey(const UCustomizableObjectNodeMaterialBase* InMaterial, uint32 InImageIndex)
+	FGeneratedImagePropertiesKey(const UCustomizableObjectNode* InMaterial, UE::Mutable::Private::FParameterKey InImageKey)
 	{
 		MaterialReferenceId = (PTRINT)InMaterial;
-		ImageIndex = InImageIndex;
+		ImageKey = InImageKey;
 	}
 
 	bool operator==(const FGeneratedImagePropertiesKey& Other) const
 	{
-		return MaterialReferenceId == Other.MaterialReferenceId && ImageIndex == Other.ImageIndex;
+		return MaterialReferenceId == Other.MaterialReferenceId && ImageKey == Other.ImageKey;
 	}
 
 
 	PTRINT MaterialReferenceId = 0;
-	uint32 ImageIndex = 0;
+	UE::Mutable::Private::FParameterKey ImageKey;
 };
 
 // Structure storing results to propagate up when generating mutable mesh node expressions.
@@ -299,7 +316,7 @@ inline uint32 GetTypeHash(const FImageKey& Key)
 
 inline uint32 GetTypeHash(const FGeneratedImagePropertiesKey& Key)
 {
-	uint32 GuidHash = HashCombineFast(GetTypeHash(Key.MaterialReferenceId), Key.ImageIndex);
+	uint32 GuidHash = HashCombineFast(GetTypeHash(Key.MaterialReferenceId), GetTypeHash(Key.ImageKey));
 
 	return GuidHash;
 }
@@ -375,7 +392,6 @@ struct FGroupProjectorImageInfo
 {
 	UE::Mutable::Private::Ptr<UE::Mutable::Private::NodeImage> ImageNode;
 	UE::Mutable::Private::Ptr<UE::Mutable::Private::NodeImage> ImageResizeNode;
-	UE::Mutable::Private::Ptr<UE::Mutable::Private::NodeSurfaceNew> SurfNode;
 	UCustomizableObjectNodeMaterialBase* TypedNodeMat;
 	FString TextureName;
 	FString RealTextureName;
@@ -387,13 +403,11 @@ struct FGroupProjectorImageInfo
 		const FString& InTextureName,
 		const FString& InRealTextureName,
 		UCustomizableObjectNodeMaterialBase* InTypedNodeMat,
-		UE::Mutable::Private::Ptr<UE::Mutable::Private::NodeSurfaceNew> InSurfNode,
 		int32 InUVLayout,
 		int32 LayerIndex)
 		: TypedNodeMat(InTypedNodeMat), TextureName(InTextureName), RealTextureName(InRealTextureName), UVLayout(InUVLayout), LayerIndex(LayerIndex)
 	{
 		ImageNode = InImageNode;
-		SurfNode = InSurfNode;
 	}
 
 	static FString GenerateId(const UCustomizableObjectNode* TypedNodeMat, int32 ImageIndex)
@@ -954,7 +968,8 @@ struct FMutableGraphGenerationContext
 	
 	struct FMaterialBreakParameter {
 		FName ParameterName;
-		EMaterialParameterType ParameterType = EMaterialParameterType::None;
+		FName ParameterType;
+		int32 LayerIndex;
 	};
 
 	/** Name of the material break Node that is being compiled.
@@ -1047,7 +1062,7 @@ FString GenerateAnimationInstanceTag(const int32 AnimInstanceIndex, const FName&
 
 FString GenerateGameplayTag(const FString& GameplayTag);
 
-uint32 GetBaseTextureSize(const FMutableGraphGenerationContext& GenerationContext, const UCustomizableObjectNodeMaterialBase* Material, uint32 ImageIndex);
+uint32 GetBaseTextureSize(const FMutableGraphGenerationContext& GenerationContext, const UCustomizableObjectNodeMaterialBase* Material, const UE::Mutable::Private::FParameterKey& ImageKey);
 
 // Computes the LOD bias for a texture given the current mesh LOD and automatic LOD settings, the reference texture settings
 // and whether it's being built for a server or not
@@ -1075,3 +1090,4 @@ uint32 AddUniqueSurfaceMetadata(const FMutableSourceSurfaceMetadata& MeshData, F
 /** Convert a mesh to Mutable format.
 */
 UE::Tasks::FTask ConvertSkeletalMeshToMutable(TSharedRef<FMeshConversionContext> MeshConversionContext);
+
diff --git a/Engine/Plugins/Mutable/Source/CustomizableObjectEditor/Private/MuCOE/GenerateMutableSource/GenerateMutableSourceColor.cpp b/Engine/Plugins/Mutable/Source/CustomizableObjectEditor/Private/MuCOE/GenerateMutableSource/GenerateMutableSourceColor.cpp
--- a/Engine/Plugins/Mutable/Source/CustomizableObjectEditor/Private/MuCOE/GenerateMutableSource/GenerateMutableSourceColor.cpp
+++ b/Engine/Plugins/Mutable/Source/CustomizableObjectEditor/Private/MuCOE/GenerateMutableSource/GenerateMutableSourceColor.cpp
@@ -398,14 +398,16 @@ UE::Mutable::Private::Ptr<UE::Mutable::Private::NodeColour> GenerateMutableSourc
 			UE::Mutable::Private::Ptr<UE::Mutable::Private::NodeColourMaterialBreak> MaterialBreak = new UE::Mutable::Private::NodeColourMaterialBreak();
 
 			FName ParameterName = TypedNodeMaterialBreak->GetPinParameterName(*Pin);
+			int32 LayerIndex = TypedNodeMaterialBreak->GetPinParameterLayerIndex(*Pin);
 			check(!ParameterName.IsNone());
 
 			// Set the name of the parameter that represents the pin
 			MaterialBreak->ParameterName = ParameterName;
+			MaterialBreak->LayerIndex = LayerIndex;
 
 			// Generate the material that is going to be broken into parameters
-			GenerationContext.CurrentMaterialBreakParameter = FMutableGraphGenerationContext::FMaterialBreakParameter{ ParameterName, EMaterialParameterType::Vector };
-			UE::Mutable::Private::Ptr<UE::Mutable::Private::NodeMaterial> LinkedNode = GenerateMutableSourceMaterial(LinkedMaterial, GenerationContext);
+			GenerationContext.CurrentMaterialBreakParameter = FMutableGraphGenerationContext::FMaterialBreakParameter{ ParameterName, Pin->PinType.PinCategory, LayerIndex };
+			UE::Mutable::Private::Ptr<UE::Mutable::Private::NodeMaterial> LinkedNode = GenerateMutableSourceMaterial(LinkedMaterial, GenerationContext, FGenerationMaterialOptions());
 			MaterialBreak->MaterialSource = LinkedNode;
 			GenerationContext.CurrentMaterialBreakParameter = FMutableGraphGenerationContext::FMaterialBreakParameter{ };
 
diff --git a/Engine/Plugins/Mutable/Source/CustomizableObjectEditor/Private/MuCOE/GenerateMutableSource/GenerateMutableSourceComponent.cpp b/Engine/Plugins/Mutable/Source/CustomizableObjectEditor/Private/MuCOE/GenerateMutableSource/GenerateMutableSourceComponent.cpp
--- a/Engine/Plugins/Mutable/Source/CustomizableObjectEditor/Private/MuCOE/GenerateMutableSource/GenerateMutableSourceComponent.cpp
+++ b/Engine/Plugins/Mutable/Source/CustomizableObjectEditor/Private/MuCOE/GenerateMutableSource/GenerateMutableSourceComponent.cpp
@@ -79,7 +79,7 @@ UE::Mutable::Private::Ptr<UE::Mutable::Private::NodeComponent> GenerateMutableSo
 		if (const UEdGraphPin* ConnectedPin = MaterialAssetPin ? FollowInputPin(*MaterialAssetPin) : nullptr)
 		{
 			GenerationContext.CurrentMaterialParameterId = ConnectedPin->PinId.ToString();
-			UE::Mutable::Private::Ptr<UE::Mutable::Private::NodeMaterial> MaterialNode = GenerateMutableSourceMaterial(ConnectedPin, GenerationContext);
+			UE::Mutable::Private::Ptr<UE::Mutable::Private::NodeMaterial> MaterialNode = GenerateMutableSourceMaterial(ConnectedPin, GenerationContext, FGenerationMaterialOptions());
 			NodeComponentNew->OverlayMaterial = MaterialNode;
 		}
 		else if(UMaterialInterface* OverlayMaterial = GenerationContext.LoadObject(TypedComponentMesh->GetOverlayMaterial()))
@@ -192,7 +192,7 @@ UE::Mutable::Private::Ptr<UE::Mutable::Private::NodeComponent> GenerateMutableSo
 			{
 				if (const UEdGraphPin* ConnectedPin = FollowInputPin(*MaterialOverridePin))
 				{
-					MaterialNode = GenerateMutableSourceMaterial(ConnectedPin, GenerationContext);
+					MaterialNode = GenerateMutableSourceMaterial(ConnectedPin, GenerationContext, FGenerationMaterialOptions());
 				}
 			}
 
diff --git a/Engine/Plugins/Mutable/Source/CustomizableObjectEditor/Private/MuCOE/GenerateMutableSource/GenerateMutableSourceExternal.cpp b/Engine/Plugins/Mutable/Source/CustomizableObjectEditor/Private/MuCOE/GenerateMutableSource/GenerateMutableSourceExternal.cpp
--- a/Engine/Plugins/Mutable/Source/CustomizableObjectEditor/Private/MuCOE/GenerateMutableSource/GenerateMutableSourceExternal.cpp
+++ b/Engine/Plugins/Mutable/Source/CustomizableObjectEditor/Private/MuCOE/GenerateMutableSource/GenerateMutableSourceExternal.cpp
@@ -1,4 +1,4 @@
-﻿// Copyright Epic Games, Inc. All Rights Reserved.
+// Copyright Epic Games, Inc. All Rights Reserved.
 
 
 #include "GenerateMutableSourceExternal.h"
@@ -127,7 +127,10 @@ UE::Mutable::Private::Ptr<UE::Mutable::Private::NodeExternal> GenerateMutableSou
             	}
             	else if (InputType == UE::Mutable::FMaterialAdapter::StaticStruct())
             	{
-            		UE::Mutable::Private::Ptr<UE::Mutable::Private::NodeMaterial> SourceNode = GenerateMutableSourceMaterial(ConnectedPin, GenerationContext);
+					FGenerationMaterialOptions MaterialOptions;
+					MaterialOptions.SurfaceNode = Options.SurfaceNode;
+
+            		UE::Mutable::Private::Ptr<UE::Mutable::Private::NodeMaterial> SourceNode = GenerateMutableSourceMaterial(ConnectedPin, GenerationContext, MaterialOptions);
             		ExternalNode->Inputs.Add(SourceNode);
             	}
             	else
diff --git a/Engine/Plugins/Mutable/Source/CustomizableObjectEditor/Private/MuCOE/GenerateMutableSource/GenerateMutableSourceFloat.cpp b/Engine/Plugins/Mutable/Source/CustomizableObjectEditor/Private/MuCOE/GenerateMutableSource/GenerateMutableSourceFloat.cpp
--- a/Engine/Plugins/Mutable/Source/CustomizableObjectEditor/Private/MuCOE/GenerateMutableSource/GenerateMutableSourceFloat.cpp
+++ b/Engine/Plugins/Mutable/Source/CustomizableObjectEditor/Private/MuCOE/GenerateMutableSource/GenerateMutableSourceFloat.cpp
@@ -50,7 +50,6 @@ UE::Mutable::Private::Ptr<UE::Mutable::Private::NodeScalar> GenerateMutableSourc
 
 	UCustomizableObjectNode* Node = CastChecked<UCustomizableObjectNode>(Pin->GetOwningNode());
 
-	//TODO(Max) UE-314814: if this node generation come from a break material node, ignore the cache otherwise we won't be able to add new parameters.
 	const FGeneratedKey Key(reinterpret_cast<void*>(&GenerateMutableSourceFloat), *Pin, *Node, GenerationContext);
 	if (const FGeneratedData* Generated = GenerationContext.Generated.Find(Key))
 	{
@@ -423,14 +422,16 @@ UE::Mutable::Private::Ptr<UE::Mutable::Private::NodeScalar> GenerateMutableSourc
 			UE::Mutable::Private::Ptr<UE::Mutable::Private::NodeScalarMaterialBreak> MaterialBreak = new UE::Mutable::Private::NodeScalarMaterialBreak();
 
 			FName ParameterName = TypedNodeMaterialBreak->GetPinParameterName(*Pin);
+			int32 LayerIndex = TypedNodeMaterialBreak->GetPinParameterLayerIndex(*Pin);
 			check(!ParameterName.IsNone());
 
 			// Set the name of the parameter that represents the pin
 			MaterialBreak->ParameterName = ParameterName;
+			MaterialBreak->LayerIndex = LayerIndex;
 
 			// Generate the material that is going to be broken into parameters
-			GenerationContext.CurrentMaterialBreakParameter = FMutableGraphGenerationContext::FMaterialBreakParameter{ ParameterName, EMaterialParameterType::Scalar };
-			UE::Mutable::Private::Ptr<UE::Mutable::Private::NodeMaterial> LinkedNode = GenerateMutableSourceMaterial(LinkedMaterial, GenerationContext);
+			GenerationContext.CurrentMaterialBreakParameter = FMutableGraphGenerationContext::FMaterialBreakParameter{ ParameterName, Pin->PinType.PinCategory, LayerIndex };
+			UE::Mutable::Private::Ptr<UE::Mutable::Private::NodeMaterial> LinkedNode = GenerateMutableSourceMaterial(LinkedMaterial, GenerationContext, FGenerationMaterialOptions());
 			MaterialBreak->MaterialSource = LinkedNode;
 			GenerationContext.CurrentMaterialBreakParameter = FMutableGraphGenerationContext::FMaterialBreakParameter{ };
 
diff --git a/Engine/Plugins/Mutable/Source/CustomizableObjectEditor/Private/MuCOE/GenerateMutableSource/GenerateMutableSourceImage.cpp b/Engine/Plugins/Mutable/Source/CustomizableObjectEditor/Private/MuCOE/GenerateMutableSource/GenerateMutableSourceImage.cpp
--- a/Engine/Plugins/Mutable/Source/CustomizableObjectEditor/Private/MuCOE/GenerateMutableSource/GenerateMutableSourceImage.cpp
+++ b/Engine/Plugins/Mutable/Source/CustomizableObjectEditor/Private/MuCOE/GenerateMutableSource/GenerateMutableSourceImage.cpp
@@ -1232,14 +1232,16 @@ UE::Mutable::Private::Ptr<UE::Mutable::Private::NodeImage> GenerateMutableSource
 			UE::Mutable::Private::Ptr<UE::Mutable::Private::NodeImageMaterialBreak> MaterialBreak = new UE::Mutable::Private::NodeImageMaterialBreak();
 
 			FName ParameterName = TypedNodeMaterialBreak->GetPinParameterName(*Pin);
+			int32 LayerIndex = TypedNodeMaterialBreak->GetPinParameterLayerIndex(*Pin);
 			check(!ParameterName.IsNone());
 
 			// Set the name of the parameter that represents the pin
 			MaterialBreak->ParameterName = ParameterName;
+			MaterialBreak->LayerIndex = LayerIndex;
 
 			// Generate the material that is going to be broken into parameters
-			GenerationContext.CurrentMaterialBreakParameter = FMutableGraphGenerationContext::FMaterialBreakParameter{ ParameterName, EMaterialParameterType::Texture };
-			UE::Mutable::Private::Ptr<UE::Mutable::Private::NodeMaterial> LinkedNode = GenerateMutableSourceMaterial(LinkedMaterial, GenerationContext);
+			GenerationContext.CurrentMaterialBreakParameter = FMutableGraphGenerationContext::FMaterialBreakParameter{ ParameterName, Pin->PinType.PinCategory, LayerIndex };
+			UE::Mutable::Private::Ptr<UE::Mutable::Private::NodeMaterial> LinkedNode = GenerateMutableSourceMaterial(LinkedMaterial, GenerationContext, FGenerationMaterialOptions());
 			MaterialBreak->MaterialSource = LinkedNode;
 			GenerationContext.CurrentMaterialBreakParameter = FMutableGraphGenerationContext::FMaterialBreakParameter{ };
 
diff --git a/Engine/Plugins/Mutable/Source/CustomizableObjectEditor/Private/MuCOE/GenerateMutableSource/GenerateMutableSourceMaterial.cpp b/Engine/Plugins/Mutable/Source/CustomizableObjectEditor/Private/MuCOE/GenerateMutableSource/GenerateMutableSourceMaterial.cpp
--- a/Engine/Plugins/Mutable/Source/CustomizableObjectEditor/Private/MuCOE/GenerateMutableSource/GenerateMutableSourceMaterial.cpp
+++ b/Engine/Plugins/Mutable/Source/CustomizableObjectEditor/Private/MuCOE/GenerateMutableSource/GenerateMutableSourceMaterial.cpp
@@ -2,6 +2,12 @@
 
 #include "MuCOE/GenerateMutableSource/GenerateMutableSourceMaterial.h"
 
+#include "Interfaces/ITargetPlatform.h"
+#include "Interfaces/ITextureFormat.h"
+#include "Interfaces/ITextureFormatModule.h"
+#include "Interfaces/ITextureFormatManagerModule.h"
+#include "TextureCompressorModule.h"
+
 #include "GenerateMutableSourceExternal.h"
 #include "GenerateMutableSourceMesh.h"
 #include "GenerateMutableSourceSkeletalMesh.h"
@@ -10,7 +16,9 @@
 #include "MuCOE/CustomizableObjectCompiler.h"
 #include "MuCOE/CustomizableObjectGraph.h"
 #include "MuCOE/EdGraphSchema_CustomizableObject.h"
+#include "MuCOE/GenerateMutableSource/GenerateMutableSourceColor.h"
 #include "MuCOE/GenerateMutableSource/GenerateMutableSourceFloat.h"
+#include "MuCOE/GenerateMutableSource/GenerateMutableSourceGroupProjector.h"
 #include "MuCOE/GenerateMutableSource/GenerateMutableSourceImage.h"
 #include "MuCOE/GenerateMutableSource/GenerateMutableSourceMacro.h"
 #include "MuCOE/GenerateMutableSource/GenerateMutableSourceTable.h"
@@ -22,30 +30,35 @@
 #include "MuCOE/Nodes/CustomizableObjectNodeTunnel.h"
 #include "MuCOE/Nodes/CONodeMaterialBreak.h"
 #include "MuCOE/Nodes/CONodeMaterialConstant.h"
+#include "MuCOE/Nodes/CONodeMaterialModify.h"
 #include "MuCOE/Nodes/CONodeMaterialSwitch.h"
 #include "MuCOE/Nodes/CONodeMaterialVariation.h"
 #include "MuCO/LoadUtils.h"
 #include "MuT/NodeImageFormat.h"
 #include "MuT/NodeImageFromMaterialParameter.h"
+#include "MuT/NodeImageMipmap.h"
+#include "MuT/NodeImageNormalComposite.h"
 #include "MuT/NodeImageResize.h"
+#include "MuT/NodeImageSwizzle.h"
 #include "MuT/NodeMaterialConstant.h"
 #include "MuT/NodeMaterialSwitch.h"
 #include "MuT/NodeMaterialTable.h"
 #include "MuT/NodeMaterialVariation.h"
 #include "MuT/NodeMaterialParameter.h"
-
+#include "MuT/NodeMaterialModify.h"
 #include "Materials/MaterialInterface.h"
 #include "Materials/MaterialInstance.h"
 #include "MuCOE/Nodes/CONodeExternalOperation.h"
 #include "MuCOE/Nodes/CONodeSkeletalMeshBreak.h"
 #include "MuT/NodeMaterialExternal.h"
 #include "MuT/NodeMaterialSkeletalMeshBreak.h"
+#include "MuT/UnrealPixelFormatOverride.h"
 
 
 #define LOCTEXT_NAMESPACE "CustomizableObjectEditor"
 
 
-UE::Mutable::Private::Ptr<UE::Mutable::Private::NodeMaterial> GenerateMutableSourceMaterial(const UEdGraphPin* Pin, FMutableGraphGenerationContext& GenerationContext)
+UE::Mutable::Private::Ptr<UE::Mutable::Private::NodeMaterial> GenerateMutableSourceMaterial(const UEdGraphPin* Pin, FMutableGraphGenerationContext& GenerationContext, const FGenerationMaterialOptions& MaterialOptions)
 {
 	check(Pin)
 	RETURN_ON_CYCLE(*Pin, GenerationContext)
@@ -72,8 +85,292 @@ UE::Mutable::Private::Ptr<UE::Mutable::Private::NodeMaterial> GenerateMutableSou
 	bool bCacheNode = true;
 
 	UE::Mutable::Private::Ptr<UE::Mutable::Private::NodeMaterial> Result;
+
+	if (UCustomizableObjectNodeMaterialBase* TypedNodeMaterialBase = Cast<UCustomizableObjectNodeMaterialBase>(Node))
+	{
+		check(MaterialOptions.SurfaceNode.get());
+
+		//Base Materials are a Material Modify
+		UE::Mutable::Private::Ptr<UE::Mutable::Private::NodeMaterialModify> BaseMaterial = new UE::Mutable::Private::NodeMaterialModify();
+		Result = BaseMaterial;
+
+		const UEdGraphPin* ConnectedMaterialAssetPin = nullptr;
+
+		// Generate the Material linked to the Mesh Section first, if any
+		if (const UEdGraphPin* MaterialAssetPin = TypedNodeMaterialBase->GetMaterialAssetPin())
+		{
+			ConnectedMaterialAssetPin = FollowInputPin(*MaterialAssetPin);
+
+			if (ConnectedMaterialAssetPin)
+			{
+				//Check if the pin goes through a macro or tunnel node
+				ConnectedMaterialAssetPin = GraphTraversal::FindIOPinSourceThroughMacroContext(*ConnectedMaterialAssetPin, &GenerationContext.MacroNodesStack);
+
+				GenerationContext.CurrentMaterialParameterId = MaterialAssetPin->PinId.ToString();
+				BaseMaterial->MaterialSource = GenerateMutableSourceMaterial(ConnectedMaterialAssetPin, GenerationContext, MaterialOptions);
+			}
+		}
+
+		if (!BaseMaterial->MaterialSource)
+		{
+			UE::Mutable::Private::Ptr<UE::Mutable::Private::NodeMaterialConstant> MaterialNode = new UE::Mutable::Private::NodeMaterialConstant();
+			MaterialNode->MaterialId = GenerationContext.AddPassthroughObject(*GenerationContext.CurrentReferencedMaterial.Get());
+
+			BaseMaterial->MaterialSource = MaterialNode;
+		}
+
+		// Generate images
+		int32 NumImages = TypedNodeMaterialBase->GetNumParameters(EMaterialParameterType::Texture);
+
+		for (int32 ImageIndex = 0; ImageIndex < NumImages; ++ImageIndex)
+		{
+			const UEdGraphPin* ImagePin = TypedNodeMaterialBase->GetParameterPin(EMaterialParameterType::Texture, ImageIndex);
+			const FName ImageName = TypedNodeMaterialBase->GetParameterName(EMaterialParameterType::Texture, ImageIndex);
+			const int32 LayerIndex = TypedNodeMaterialBase->GetParameterLayerIndex(EMaterialParameterType::Texture, ImageIndex);
+			const bool bIsImagePinLinked = ImagePin && FollowInputPin(*ImagePin);
+
+			if (bIsImagePinLinked && !TypedNodeMaterialBase->IsImageMutableMode(ImageIndex))
+			{
+				if (const UEdGraphPin* ConnectedPin = FollowInputPin(*ImagePin))
+				{
+					UE::Mutable::Private::FParameterKey ImageKey = { ImageName, LayerIndex };
+					FGeneratedImageProperties Props = GenerateImageProperties(TypedNodeMaterialBase, nullptr, ImageKey, true, GenerationContext);
+
+					// This is a connected pass-through texture that simply has to be passed to the core
+					UE::Mutable::Private::Ptr<UE::Mutable::Private::NodeImage> PassThroughImagePtr = GenerateMutableSourceImage(ConnectedPin, GenerationContext, 0);
+					UE::Mutable::Private::FImageParameterData Image;
+
+					check(Props.ImagePropertiesIndex != INDEX_NONE);
+
+					Image.ImageNode = PassThroughImagePtr;
+					Image.ImagePropertyIndex = Props.ImagePropertiesIndex;
+					Image.LayoutIndex = -1;
+
+					BaseMaterial->ImageParameters.Add(ImageKey, Image);
+				}
+			}
+			else
+			{
+				UE::Mutable::Private::NodeImagePtr GroupProjectionImg;
+				UTexture2D* GroupProjectionReferenceTexture = nullptr;
+				const FNodeMaterialParameterId ImageId = TypedNodeMaterialBase->GetParameterId(EMaterialParameterType::Texture, ImageIndex);
+
+				FString MaterialImageId = FGroupProjectorImageInfo::GenerateId(TypedNodeMaterialBase, ImageIndex);
+				bool bShareProjectionTexturesBetweenLODs = false;
+				FGroupProjectorImageInfo* ProjectorInfo = GenerationContext.GroupProjectorLODCache.Find(MaterialImageId);
+
+				if (!ProjectorInfo) // No previous LOD of this material generated the image.
+				{
+					bool bIsGroupProjectorImage = false;
+
+					GroupProjectionImg = GenerateMutableSourceGroupProjector(GenerationContext.CurrentLOD, ImageIndex, MaterialOptions.SurfaceNode->Mesh, GenerationContext,
+						TypedNodeMaterialBase, nullptr, bShareProjectionTexturesBetweenLODs, bIsGroupProjectorImage,
+						GroupProjectionReferenceTexture);
+
+					if (GroupProjectionImg.get() || TypedNodeMaterialBase->IsImageMutableMode(ImageIndex))
+					{
+						// Get the reference texture
+						UTexture2D* ReferenceTexture = nullptr;
+						{
+							//TODO(Max) UE-220247: Add support for multilayer materials
+							GenerationContext.CurrentMaterialTableParameter = ImageName.ToString();
+							GenerationContext.CurrentMaterialParameterId = ImageId.ParameterId.ToString();
+
+							ReferenceTexture = GroupProjectionImg.get() ? GroupProjectionReferenceTexture : nullptr;
+
+							if (!ReferenceTexture)
+							{
+								ReferenceTexture = TypedNodeMaterialBase->GetImageReferenceTexture(ImageIndex);
+							}
+
+							// In case of group projector, don't follow the pin to find the reference texture.
+							if (!GroupProjectionImg.get() && !ReferenceTexture && ImagePin)
+							{
+								if (const UEdGraphPin* ConnectedPin = FollowInputPin(*ImagePin))
+								{
+									ReferenceTexture = FindReferenceImage(ConnectedPin, GenerationContext);
+								}
+							}
+
+							if (!ReferenceTexture &&
+								ConnectedMaterialAssetPin && Cast<UCustomizableObjectNodeTable>(ConnectedMaterialAssetPin->GetOwningNode()) != nullptr)
+							{
+								ReferenceTexture = FindReferenceImage(ConnectedMaterialAssetPin, GenerationContext);
+							}
+
+							if (!ReferenceTexture)
+							{
+								ReferenceTexture = TypedNodeMaterialBase->GetImageValue(ImageIndex);
+							}
+						}
+
+						// Generate image key to identify it
+						UE::Mutable::Private::FParameterKey ImageKey = { ImageName, LayerIndex };
+						FGeneratedImageProperties Props = GenerateImageProperties(TypedNodeMaterialBase, ReferenceTexture, ImageKey, false, GenerationContext);
+
+						// Generate the texture nodes
+						UE::Mutable::Private::NodeImagePtr ImageNode = [&]()
+							{
+								if (TypedNodeMaterialBase->IsImageMutableMode(ImageIndex))
+								{
+									if (ImagePin)
+									{
+										if (const UEdGraphPin* ConnectedPin = FollowInputPin(*ImagePin))
+										{
+											return GenerateMutableSourceImage(ConnectedPin, GenerationContext, Props.TextureSize);
+										}
+									}
+
+									// If the table material pin is linked to a table node, get all the textures of the current material parameter (CurrentMaterialTableParameter) from the Material Instances of the specified data table column.
+									// Then Generate a mutable table column with all these textures.
+									if (ConnectedMaterialAssetPin && Cast<UCustomizableObjectNodeTable>(ConnectedMaterialAssetPin->GetOwningNode()) != nullptr)
+									{
+										return GenerateMutableSourceImage(ConnectedMaterialAssetPin, GenerationContext, Props.TextureSize);
+									}
+
+									// Else
+									{
+										UTexture2D* Texture2D = TypedNodeMaterialBase->GetImageValue(ImageIndex);
+
+										if (Texture2D)
+										{
+											UE::Mutable::Private::Ptr<UE::Mutable::Private::NodeImageConstant> ConstImageNode = new UE::Mutable::Private::NodeImageConstant();
+											ConstImageNode->Image = GenerateImageConstant(Texture2D, GenerationContext, false);
+
+											const uint32 MipsToSkip = ComputeLODBiasForTexture(GenerationContext, *Texture2D, nullptr, Props.TextureSize);
+											UE::Mutable::Private::Ptr<UE::Mutable::Private::NodeImage> Result = ResizeTextureByNumMips(ConstImageNode, MipsToSkip);
+
+											const UTextureLODSettings& TextureLODSettings = GenerationContext.CompilationContext->Options.TargetPlatform->GetTextureLODSettings();
+											const FTextureLODGroup& LODGroupInfo = TextureLODSettings.GetTextureLODGroup(Texture2D->LODGroup);
+
+											ConstImageNode->SourceDataDescriptor.OptionalMaxLODSize = LODGroupInfo.OptionalMaxLODSize;
+											ConstImageNode->SourceDataDescriptor.OptionalLODBias = LODGroupInfo.OptionalLODBias;
+											ConstImageNode->SourceDataDescriptor.NumNonOptionalLODs = GenerationContext.CompilationContext->Options.MinDiskMips;
+
+											const FString TextureName = GetNameSafe(Texture2D).ToLower();
+											ConstImageNode->SourceDataDescriptor.SourceId = CityHash32(reinterpret_cast<const char*>(*TextureName), TextureName.Len() * sizeof(FString::ElementType));
+
+											return Result;
+										}
+										else
+										{
+											return UE::Mutable::Private::NodeImagePtr();
+										}
+									}
+								}
+								else
+								{
+									return UE::Mutable::Private::NodeImagePtr();
+								}
+							}();
+
+						if (GroupProjectionImg.get())
+						{
+							ImageNode = GroupProjectionImg;
+						}
+
+						// NOTE: Important step to generate a texture with the correct format and size
+						ImageNode = GenerateMipMapAndFormatNodes(Node, GenerationContext, ImageNode, ReferenceTexture);
+
+						// If we are generating an implicit component (with a passthrough mesh) we don't apply any layout.
+						bool bGeneratingImplicitComponent = GenerationContext.ComponentMeshOverride.get() != nullptr;
+
+						int32 UVLayout = -1;
+						if (!bGeneratingImplicitComponent)
+						{
+							UVLayout = TypedNodeMaterialBase->GetImageUVLayout(ImageIndex);
+						}
+
+						// Generate image data
+						UE::Mutable::Private::FImageParameterData Image;
+						Image.ImageNode = ImageNode;
+						Image.ImagePropertyIndex = Props.ImagePropertiesIndex;
+						Image.LayoutIndex = UVLayout;
+
+						//Store the image inside the current generated material
+						BaseMaterial->ImageParameters.Add(ImageKey, Image);
+
+						if (bShareProjectionTexturesBetweenLODs && bIsGroupProjectorImage)
+						{
+							// Add to the GroupProjectorLODCache to potentially reuse this projection texture in higher LODs
+							ensure(GenerationContext.CurrentLOD == GenerationContext.FirstLODAvailable[GenerationContext.CurrentMeshComponent]);
+
+							GenerationContext.GroupProjectorLODCache.Add(MaterialImageId,
+								FGroupProjectorImageInfo(ImageNode, ImageName.ToString(), ImageName.ToString(), TypedNodeMaterialBase, UVLayout, LayerIndex));
+						}
+					}
+				}
+				else
+				{
+					ensure(GenerationContext.CurrentLOD > GenerationContext.FirstLODAvailable[GenerationContext.CurrentMeshComponent]);
+
+					UE::Mutable::Private::FParameterKey ImageKey = { FName(ProjectorInfo->TextureName), ProjectorInfo->LayerIndex };
+					UE::Mutable::Private::FImageParameterData Image;
+
+					Image.ImageNode = ProjectorInfo->ImageNode;
+					Image.LayoutIndex = ProjectorInfo->UVLayout;
+
+					BaseMaterial->ImageParameters.Add(ImageKey, Image);
+				}
+			}
+		}
+
+		// Generate color parameters
+		const int32 NumVectors = TypedNodeMaterialBase->GetNumParameters(EMaterialParameterType::Vector);
+		for (int32 VectorIndex = 0; VectorIndex < NumVectors; ++VectorIndex)
+		{
+			const UEdGraphPin* VectorPin = TypedNodeMaterialBase->GetParameterPin(EMaterialParameterType::Vector, VectorIndex);
+			bool bVectorPinConnected = VectorPin && FollowInputPin(*VectorPin);
+
+			if (bVectorPinConnected)
+			{
+				if (const UEdGraphPin* ConnectedPin = FollowInputPin(*VectorPin))
+				{
+					FName VectorName = TypedNodeMaterialBase->GetParameterName(EMaterialParameterType::Vector, VectorIndex);
+					FNodeMaterialParameterId VectorId = TypedNodeMaterialBase->GetParameterId(EMaterialParameterType::Vector, VectorIndex);
+					const int32 LayerIndex = TypedNodeMaterialBase->GetParameterLayerIndex(EMaterialParameterType::Vector, VectorIndex);
+
+					// Store info to generate colors from the material of a table node
+					GenerationContext.CurrentMaterialTableParameter = VectorName.ToString();
+					GenerationContext.CurrentMaterialParameterId = VectorId.ParameterId.ToString();
+
+					UE::Mutable::Private::Ptr<UE::Mutable::Private::NodeColour> ColorNode = GenerateMutableSourceColor(ConnectedPin, GenerationContext);
+					UE::Mutable::Private::FParameterKey ColorKey = { VectorName, LayerIndex };
+
+					BaseMaterial->ColorParameters.Add(ColorKey, ColorNode);
+				}
+			}
+		}
+
+		// Generate scalar parameters
+		const int32 NumScalar = TypedNodeMaterialBase->GetNumParameters(EMaterialParameterType::Scalar);
+		for (int32 ScalarIndex = 0; ScalarIndex < NumScalar; ++ScalarIndex)
+		{
+			const UEdGraphPin* ScalarPin = TypedNodeMaterialBase->GetParameterPin(EMaterialParameterType::Scalar, ScalarIndex);
+			bool bScalarPinConnected = ScalarPin && FollowInputPin(*ScalarPin);
+
+			if (bScalarPinConnected)
+			{
+				if (const UEdGraphPin* ConnectedPin = FollowInputPin(*ScalarPin))
+				{
+					FName ScalarName = TypedNodeMaterialBase->GetParameterName(EMaterialParameterType::Scalar, ScalarIndex);
+					FNodeMaterialParameterId ScalarId = TypedNodeMaterialBase->GetParameterId(EMaterialParameterType::Scalar, ScalarIndex);
+					const int32 LayerIndex = TypedNodeMaterialBase->GetParameterLayerIndex(EMaterialParameterType::Scalar, ScalarIndex);
+					
+					// Store info to generate scalars from the material of a table
+					GenerationContext.CurrentMaterialTableParameter = ScalarName.ToString();
+					GenerationContext.CurrentMaterialParameterId = ScalarId.ParameterId.ToString();
+
+					UE::Mutable::Private::Ptr<UE::Mutable::Private::NodeScalar> ScalarNode = GenerateMutableSourceFloat(ConnectedPin, GenerationContext);
+					UE::Mutable::Private::FParameterKey ScalarKey = { ScalarName, LayerIndex };
+
+					BaseMaterial->ScalarParameters.Add(ScalarKey, ScalarNode);
+				}
+			}
+		}
+	}
 	
-	if (const UCONodeMaterialConstant* TypedNodeMaterialConstant = Cast<UCONodeMaterialConstant>(Node))
+	else if (const UCONodeMaterialConstant* TypedNodeMaterialConstant = Cast<UCONodeMaterialConstant>(Node))
 	{
 		UE::Mutable::Private::Ptr<UE::Mutable::Private::NodeMaterialConstant> MaterialNode;
 		bCacheNode = false; //custom node cache
@@ -94,21 +391,24 @@ UE::Mutable::Private::Ptr<UE::Mutable::Private::NodeMaterial> GenerateMutableSou
 				GenerationContext.MaterialConstantNodesCache.Add(TypedNodeMaterialConstant, MaterialNode);
 			}
 
-			// Store the parameter that needs to be processed
-			if (GenerationContext.CurrentMaterialBreakParameter.ParameterType != EMaterialParameterType::None)
+			// Store parameters that needs to be processed
+			if (!GenerationContext.CurrentMaterialBreakParameter.ParameterType.IsNone())
 			{
+				FHashedMaterialParameterInfo ParameterInfo;
+				ParameterInfo.Name = FScriptName(GenerationContext.CurrentMaterialBreakParameter.ParameterName);
+				ParameterInfo.Index = GenerationContext.CurrentMaterialBreakParameter.LayerIndex;
+				ParameterInfo.Association = ParameterInfo.Index == INDEX_NONE ? GlobalParameter : LayerParameter;
+
 				if (!TypedNodeMaterialConstant->Material)
 				{
 					GenerationContext.Log(LOCTEXT("NoRefernceMaterialToBreak", "Could not find a Reference Material to break."), TypedNodeMaterialConstant);
 				}
 				else
 				{
-					switch (GenerationContext.CurrentMaterialBreakParameter.ParameterType)
-					{
-					case EMaterialParameterType::Texture:
+					if (GenerationContext.CurrentMaterialBreakParameter.ParameterType == UEdGraphSchema_CustomizableObject::PC_Texture)
 					{
 						UTexture* Texture = nullptr;
-						bool bParameterFound = Material->GetTextureParameterValue(GenerationContext.CurrentMaterialBreakParameter.ParameterName, Texture);
+						bool bParameterFound = Material->GetTextureParameterValue(ParameterInfo, Texture);
 						UTexture2D* BaseTexture = Cast<UTexture2D>(Texture);
 
 						if (bParameterFound && BaseTexture)
@@ -132,34 +432,41 @@ UE::Mutable::Private::Ptr<UE::Mutable::Private::NodeMaterial> GenerateMutableSou
 
 							if (FinalNodeImag)
 							{
-								MaterialNode->ImageValues.Add(GenerationContext.CurrentMaterialBreakParameter.ParameterName, FinalNodeImag);
+								MaterialNode->ImageValues.Add({ GenerationContext.CurrentMaterialBreakParameter.ParameterName, GenerationContext.CurrentMaterialBreakParameter.LayerIndex}, FinalNodeImag);
 							}
 						}
+					}
+					
+					else if (GenerationContext.CurrentMaterialBreakParameter.ParameterType == UEdGraphSchema_CustomizableObject::PC_PassthroughTexture)
+					{
+						UTexture* Texture = nullptr;
+						bool bParameterFound = Material->GetTextureParameterValue(ParameterInfo, Texture);
+						UTexture2D* BaseTexture = Cast<UTexture2D>(Texture);
 
-						break;
+						if (bParameterFound && BaseTexture)
+						{
+							UE::Mutable::Private::Ptr<UE::Mutable::Private::NodeImageConstant> ImageNode = new UE::Mutable::Private::NodeImageConstant();
+							ImageNode->Image = GenerateImageConstant(Texture, GenerationContext, true);
+							MaterialNode->ImageValues.Add({GenerationContext.CurrentMaterialBreakParameter.ParameterName, GenerationContext.CurrentMaterialBreakParameter.LayerIndex}, ImageNode);
+						}
 					}
-					case EMaterialParameterType::Vector:
+
+					else if(GenerationContext.CurrentMaterialBreakParameter.ParameterType == UEdGraphSchema_CustomizableObject::PC_Color)
 					{
 						FLinearColor Color;
-						if (Material->GetVectorParameterValue(GenerationContext.CurrentMaterialBreakParameter.ParameterName, Color))
+						if (Material->GetVectorParameterValue(ParameterInfo, Color))
 						{
-							MaterialNode->ColorValues.Add(GenerationContext.CurrentMaterialBreakParameter.ParameterName, Color);
+							MaterialNode->ColorValues.Add({GenerationContext.CurrentMaterialBreakParameter.ParameterName, GenerationContext.CurrentMaterialBreakParameter.LayerIndex}, Color);
 						}
-
-						break;
 					}
-					case EMaterialParameterType::Scalar:
+
+					else if(GenerationContext.CurrentMaterialBreakParameter.ParameterType == UEdGraphSchema_CustomizableObject::PC_Float)
 					{
 						float Scalar;
-						if (Material->GetScalarParameterValue(GenerationContext.CurrentMaterialBreakParameter.ParameterName, Scalar))
+						if (Material->GetScalarParameterValue(ParameterInfo, Scalar))
 						{
-							MaterialNode->ScalarValues.Add(GenerationContext.CurrentMaterialBreakParameter.ParameterName, Scalar);
+							MaterialNode->ScalarValues.Add({GenerationContext.CurrentMaterialBreakParameter.ParameterName, GenerationContext.CurrentMaterialBreakParameter.LayerIndex}, Scalar);
 						}
-
-						break;
-					}
-					default:
-						break;
 					}
 				}
 			}
@@ -213,7 +520,7 @@ UE::Mutable::Private::Ptr<UE::Mutable::Private::NodeMaterial> GenerateMutableSou
 				{
 					if (const UEdGraphPin* MaterialPin = FollowInputPin(*TypedNodeMaterialSwitch->GetElementPin(SelectorIndex)))
 					{
-						SwitchNode->Options[SelectorIndex] = GenerateMutableSourceMaterial(MaterialPin, GenerationContext);
+						SwitchNode->Options[SelectorIndex] = GenerateMutableSourceMaterial(MaterialPin, GenerationContext, MaterialOptions);
 					}
 					else
 					{
@@ -241,7 +548,7 @@ UE::Mutable::Private::Ptr<UE::Mutable::Private::NodeMaterial> GenerateMutableSou
 
 		if (const UEdGraphPin* ConnectedPin = FollowInputPin(*TypedNodeMaterialVariation->DefaultPin()))
 		{
-			UE::Mutable::Private::Ptr<UE::Mutable::Private::NodeMaterial> ChildNode = GenerateMutableSourceMaterial(ConnectedPin, GenerationContext);
+			UE::Mutable::Private::Ptr<UE::Mutable::Private::NodeMaterial> ChildNode = GenerateMutableSourceMaterial(ConnectedPin, GenerationContext, MaterialOptions);
 			if (ChildNode)
 			{
 				MaterialNode->DefaultMaterial = ChildNode;
@@ -262,7 +569,7 @@ UE::Mutable::Private::Ptr<UE::Mutable::Private::NodeMaterial> GenerateMutableSou
 
 				if (const UEdGraphPin* ConnectedPin = FollowInputPin(*VariationPin))
 				{
-					UE::Mutable::Private::Ptr<UE::Mutable::Private::NodeMaterial> ChildNode = GenerateMutableSourceMaterial(ConnectedPin, GenerationContext);
+					UE::Mutable::Private::Ptr<UE::Mutable::Private::NodeMaterial> ChildNode = GenerateMutableSourceMaterial(ConnectedPin, GenerationContext, MaterialOptions);
 					MaterialNode->Variations[VariationIndex].Material = ChildNode;
 				}
 			}
@@ -294,7 +601,7 @@ UE::Mutable::Private::Ptr<UE::Mutable::Private::NodeMaterial> GenerateMutableSou
 		}
 
 		// Store the parameter that needs to be processed
-		if (GenerationContext.CurrentMaterialBreakParameter.ParameterType != EMaterialParameterType::None)
+		if (!GenerationContext.CurrentMaterialBreakParameter.ParameterType.IsNone())
 		{
 			if (!MaterialParameterNode->ReferenceValue)
 			{
@@ -302,11 +609,15 @@ UE::Mutable::Private::Ptr<UE::Mutable::Private::NodeMaterial> GenerateMutableSou
 			}
 			else
 			{
-				switch (GenerationContext.CurrentMaterialBreakParameter.ParameterType)
-				{
-				case EMaterialParameterType::Texture:
+				// Data to find the parameter value inside the material
+				FHashedMaterialParameterInfo ParameterInfo;
+				ParameterInfo.Name = FScriptName(GenerationContext.CurrentMaterialBreakParameter.ParameterName);
+				ParameterInfo.Index = GenerationContext.CurrentMaterialBreakParameter.LayerIndex;
+				ParameterInfo.Association = ParameterInfo.Index == INDEX_NONE ? GlobalParameter : LayerParameter;
+
+				if (GenerationContext.CurrentMaterialBreakParameter.ParameterType == UEdGraphSchema_CustomizableObject::PC_Texture)
 				{
-					//TODO(Max): Move this to the generatemutablesourceimage:
+					//TODO(Max - GM): Move this to the generatemutablesourceimage?
 					UE::Mutable::Private::Ptr<UE::Mutable::Private::NodeImageFromMaterialParameter> TextureNode = new UE::Mutable::Private::NodeImageFromMaterialParameter();
 					TextureNode->ImageParameterName = GenerationContext.CurrentMaterialBreakParameter.ParameterName;
 
@@ -317,9 +628,9 @@ UE::Mutable::Private::Ptr<UE::Mutable::Private::NodeMaterial> GenerateMutableSou
 
 					UE::Mutable::Private::FImageSize TextureSize(1);
 
-					//Find reference Texture
+					//Find reference Texture from the material parameter
 					UTexture* Texture = nullptr;
-					bool bParameterFound = MaterialParameterNode->ReferenceValue->GetTextureParameterValue(GenerationContext.CurrentMaterialBreakParameter.ParameterName, Texture);
+					bool bParameterFound = MaterialParameterNode->ReferenceValue->GetTextureParameterValue(ParameterInfo, Texture);
 					UTexture2D* ReferenceTexture = Cast<UTexture2D>(Texture);
 
 					if (bParameterFound && ReferenceTexture)
@@ -334,23 +645,33 @@ UE::Mutable::Private::Ptr<UE::Mutable::Private::NodeMaterial> GenerateMutableSou
 
 					// Force the same format that the default texture if any.
 					UE::Mutable::Private::Ptr<UE::Mutable::Private::NodeImageFormat> FormatNode = GenerateMutableImageFormat(ResizeNode, ReferenceTexture, GenerationContext, Node);
+					MaterialNode->ImageParametersToCompile.Add({ FName(ParameterInfo.Name), ParameterInfo.Index }, FormatNode);
+				}
 
-					MaterialNode->ImageParameterNodes.Add(GenerationContext.CurrentMaterialBreakParameter.ParameterName.ToString(), FormatNode);
+				else if (GenerationContext.CurrentMaterialBreakParameter.ParameterType == UEdGraphSchema_CustomizableObject::PC_PassthroughTexture)
+				{
+					//TODO(Max - GM): Add support for Passthrough texture parameters
+					////Find reference Texture from the material parameter
+					//UTexture* Texture = nullptr;
+					//bool bParameterFound = MaterialParameterNode->ReferenceValue->GetTextureParameterValue(GenerationContext.CurrentMaterialBreakParameter.ParameterName, Texture);
+
+					//UE::Mutable::Private::Ptr<UE::Mutable::Private::NodeImageConstant> ImageNode = new UE::Mutable::Private::NodeImageConstant();
+					//ImageNode->SetValue(GenerateImageConstant(Texture, GenerationContext, true));
 
-					break;
+					//MaterialNode->ImageParameterNodes.Add({ FName(ParameterInfo.Name), ParameterInfo.Index });
+
+					FText Message = LOCTEXT("ParameterPassthroughUnsupported", "A Material Parameter node can't be break into passthrough textures right now.");
+					GenerationContext.Log(Message, Node, EMessageSeverity::Error);
 				}
-				case EMaterialParameterType::Vector:
+
+				else if (GenerationContext.CurrentMaterialBreakParameter.ParameterType == UEdGraphSchema_CustomizableObject::PC_Color)
 				{
-					MaterialNode->ColorParameterNames.Add(GenerationContext.CurrentMaterialBreakParameter.ParameterName.ToString());
-					break;
+					MaterialNode->ColorParametersToCompile.Add({ FName(ParameterInfo.Name), ParameterInfo.Index });
 				}
-				case EMaterialParameterType::Scalar:
+
+				else if (GenerationContext.CurrentMaterialBreakParameter.ParameterType == UEdGraphSchema_CustomizableObject::PC_Float)
 				{
-					MaterialNode->ScalarParameterNames.Add(GenerationContext.CurrentMaterialBreakParameter.ParameterName.ToString());
-					break;
-				}
-				default:
-					break;
+					MaterialNode->ScalarParametersToCompile.Add({ FName(ParameterInfo.Name), ParameterInfo.Index });
 				}
 			}
 		}
@@ -376,6 +697,148 @@ UE::Mutable::Private::Ptr<UE::Mutable::Private::NodeMaterial> GenerateMutableSou
 		}
 	}
 
+	else if (const UCONodeMaterialModify* TypedNodeMaterialModify = Cast<UCONodeMaterialModify>(Node))
+	{
+		UE::Mutable::Private::Ptr<UE::Mutable::Private::NodeMaterialModify> MaterialModifyNode = new UE::Mutable::Private::NodeMaterialModify();
+		Result = MaterialModifyNode;
+
+		check(TypedNodeMaterialModify->MaterialPinRef.Get())
+		const UEdGraphPin* MaterialPin = FollowInputPin(*TypedNodeMaterialModify->MaterialPinRef.Get());
+
+		if (!MaterialPin)
+		{
+			const FText ErrorMessage = LOCTEXT("ModifyMaterial_MissingMaterialConnection","Missing Material Connection. Modify Material Nodes require a material node.");
+			GenerationContext.Log(ErrorMessage, Node, EMessageSeverity::Error);
+
+			return Result;
+		}
+
+		UE::Mutable::Private::Ptr<UE::Mutable::Private::NodeMaterial> MaterialNode = GenerateMutableSourceMaterial(MaterialPin, GenerationContext, MaterialOptions);
+
+		if (!MaterialNode)
+		{
+			const FText ErrorMessage = LOCTEXT("ModifyMaterial_FailedMaterialGeneration", "The Generation of the material node has failed.");
+			GenerationContext.Log(ErrorMessage, Node, EMessageSeverity::Error);
+
+			return Result;
+		}
+
+		MaterialModifyNode->MaterialSource = MaterialNode;
+
+		// Generating the material parameters that will be modified
+		//-----------------------------------------------------------------
+		// TODOs: 
+		// - Add support for table node materials. This will be added when the current hack of the tables is removed. UE-314401
+		// - Add support for layout index once the pin viewer has all the needed properties. UE-317590
+		// - Add support for group projectors.
+		// - Remove the material editon part from the mesh section nodes. Talk about this with Gerard since we can lose the capacity to pass an image parameter of the reference material as mutable without linking any texture node.
+		// - Improve the dectection of the reference texture. 
+		//		|_ Find the first material constant (or parameter) node that contains the texture parameter and return it
+		//		|_ Or serach it in the pin data UE-318756
+		//		|_ Or use the reference image of the group projector generated (if any).
+
+		TArray<UEdGraphPin*> InputPins;
+		TypedNodeMaterialModify->GetInputPins(InputPins);
+
+		// Tracking the index of the image. Needed for the FGeneratedImagePropertyKey
+		int32 ImageIndex = 0;
+		
+		for (const UEdGraphPin* InputPin : InputPins)
+		{
+			if (InputPin == TypedNodeMaterialModify->MaterialPinRef)
+			{
+				continue;
+			}
+
+			FName PinType = InputPin->PinType.PinCategory;
+			FName ParameterName = TypedNodeMaterialModify->GetPinParameterName(*InputPin);
+			int32 LayerIndex = TypedNodeMaterialModify->GetPinParameterLayerIndex(*InputPin);
+
+			UE::Mutable::Private::FParameterKey ParameterKey = { ParameterName, LayerIndex };
+
+			if (PinType == UEdGraphSchema_CustomizableObject::PC_Texture)
+			{
+				UE::Mutable::Private::FImageParameterData ImageData;
+
+				if (const UEdGraphPin* FollowPin = FollowInputPin(*InputPin))
+				{
+					UTexture2D* ReferenceTexture = FindReferenceImage(FollowPin, GenerationContext); //TODO(Max - GM) UE - 318756
+					FGeneratedImageProperties Props = GenerateImageProperties(Node, ReferenceTexture, ParameterKey, false, GenerationContext);
+					UE::Mutable::Private::Ptr<UE::Mutable::Private::NodeImage> ImageNode = GenerateMutableSourceImage(FollowPin, GenerationContext, Props.TextureSize);
+
+					// NOTE: Important step to generate a texture with the correct format and size
+					ImageNode = GenerateMipMapAndFormatNodes(Node, GenerationContext, ImageNode, ReferenceTexture);
+
+					// Fill Image Data
+					ImageData.ImageNode = ImageNode;
+					ImageData.ImagePropertyIndex = Props.ImagePropertiesIndex;
+					ImageData.LayoutIndex = -1; // TODO(Max - GM) UE-317590
+				}
+
+				// We do not check Null images since it is possible to not have anything linked
+				MaterialModifyNode->ImageParameters.Add(ParameterKey, ImageData);
+
+				ImageIndex++;
+			}
+
+			else if (PinType == UEdGraphSchema_CustomizableObject::PC_PassthroughTexture)
+			{
+				UE::Mutable::Private::FImageParameterData ImageData;
+
+				if (const UEdGraphPin* ConnectedPin = FollowInputPin(*InputPin))
+				{
+					UE::Mutable::Private::FParameterKey ImageKey = { ParameterName, LayerIndex };
+					FGeneratedImageProperties Props = GenerateImageProperties(TypedNodeMaterialBase, nullptr, ImageKey, true, GenerationContext);
+
+					// This is a connected pass-through texture that simply has to be passed to the core
+					UE::Mutable::Private::Ptr<UE::Mutable::Private::NodeImage> PassThroughImagePtr = GenerateMutableSourceImage(ConnectedPin, GenerationContext, 0);
+
+					check(Props.ImagePropertiesIndex != INDEX_NONE);
+
+					ImageData.ImageNode = PassThroughImagePtr;
+					ImageData.ImagePropertyIndex = Props.ImagePropertiesIndex;
+					ImageData.LayoutIndex = -1;
+				}
+
+				// We do not check Null images since it is possible to not have anything linked
+				MaterialModifyNode->ImageParameters.Add(ParameterKey, ImageData);
+
+				ImageIndex++;
+			}
+
+			else if (PinType == UEdGraphSchema_CustomizableObject::PC_Color)
+			{
+				UE::Mutable::Private::Ptr<UE::Mutable::Private::NodeColour> ColorNode;
+
+				if (const UEdGraphPin* FollowPin = FollowInputPin(*InputPin))
+				{
+					ColorNode = GenerateMutableSourceColor(FollowPin, GenerationContext);
+				}
+
+				// We do not check Null images since it is possible to not have anything linked
+				MaterialModifyNode->ColorParameters.Add(ParameterKey, ColorNode);
+			}
+
+			else if (PinType == UEdGraphSchema_CustomizableObject::PC_Float)
+			{
+				UE::Mutable::Private::Ptr<UE::Mutable::Private::NodeScalar> ScalarNode;
+
+				if (const UEdGraphPin* FollowPin = FollowInputPin(*InputPin))
+				{
+					ScalarNode = GenerateMutableSourceFloat(FollowPin, GenerationContext);
+				}
+
+				// We do not check Null images since it is possible to not have anything linked
+				MaterialModifyNode->ScalarParameters.Add(ParameterKey, ScalarNode);
+			}
+
+			else
+			{
+				unimplemented();
+			}
+		}
+	}
+
 	else if (const UCustomizableObjectNodeTable* TypedNodeTable = Cast<UCustomizableObjectNodeTable>(Node))
 	{
 		//This node will add a default value in case of error
@@ -487,6 +950,8 @@ UE::Mutable::Private::Ptr<UE::Mutable::Private::NodeMaterial> GenerateMutableSou
 		UE::Mutable::Private::Ptr<UE::Mutable::Private::NodeMaterialExternal> MeshExtension = new UE::Mutable::Private::NodeMaterialExternal();
 
 		FSourceExternalOptions Options;
+		Options.SurfaceNode = MaterialOptions.SurfaceNode;
+
 		MeshExtension->Node = GenerateMutableSourceExternal(Pin, GenerationContext, Options);
 		
 		Result = MeshExtension;
@@ -496,13 +961,13 @@ UE::Mutable::Private::Ptr<UE::Mutable::Private::NodeMaterial> GenerateMutableSou
 	else if (const UCustomizableObjectNodeMacroInstance* TypedNodeMacro = Cast<UCustomizableObjectNodeMacroInstance>(Node))
 	{
 		bCacheNode = false;
-		Result = GenerateMutableSourceMacro<UE::Mutable::Private::NodeMaterial>(*Pin, GenerationContext, GenerateMutableSourceMaterial);
+		Result = GenerateMutableSourceMacro<UE::Mutable::Private::NodeMaterial>(*Pin, GenerationContext, GenerateMutableSourceMaterial, MaterialOptions);
 	}
 
 	else if (const UCustomizableObjectNodeTunnel* TypedNodeTunnel = Cast<UCustomizableObjectNodeTunnel>(Node))
 	{
 		bCacheNode = false;
-		Result = GenerateMutableSourceMacro<UE::Mutable::Private::NodeMaterial>(*Pin, GenerationContext, GenerateMutableSourceMaterial);
+		Result = GenerateMutableSourceMacro<UE::Mutable::Private::NodeMaterial>(*Pin, GenerationContext, GenerateMutableSourceMaterial, MaterialOptions);
 	}
 
 	else
@@ -524,5 +989,177 @@ UE::Mutable::Private::Ptr<UE::Mutable::Private::NodeMaterial> GenerateMutableSou
 	return Result;
 }
 
+
+FGeneratedImageProperties GenerateImageProperties(const UCustomizableObjectNode* Node, const UTexture2D* ReferenceTexture, const UE::Mutable::Private::FParameterKey& ImageKey, bool bIsPassthrough, FMutableGraphGenerationContext& GenerationContext)
+{
+	// Find or add Image properties
+	const FGeneratedImagePropertiesKey PropsKey(Node, ImageKey);
+	const bool bNewImageProps = !GenerationContext.ImageProperties.Contains(PropsKey);
+
+	FGeneratedImageProperties& Props = GenerationContext.ImageProperties.FindOrAdd(PropsKey);
+
+	if (!bNewImageProps)
+	{
+		return Props;
+	}
+
+	// Store properties for the generated images
+	Props.TextureParameterName = ImageKey.ParameterName.ToString();
+	Props.ImagePropertiesIndex = GenerationContext.ImageProperties.Num() - 1;
+	Props.bIsPassThrough = bIsPassthrough;
+
+	if (bIsPassthrough)
+	{
+		// We don't need more info for passthrough textures.
+		return Props;
+	}
+
+	if (ReferenceTexture)
+	{
+		Props.CompressionSettings = ReferenceTexture->CompressionSettings;
+		Props.Filter = ReferenceTexture->Filter;
+		Props.SRGB = ReferenceTexture->SRGB;
+		Props.LODBias = 0;
+		Props.MipGenSettings = ReferenceTexture->MipGenSettings;
+		Props.LODGroup = ReferenceTexture->LODGroup;
+		Props.AddressX = ReferenceTexture->AddressX;
+		Props.AddressY = ReferenceTexture->AddressY;
+		Props.bFlipGreenChannel = ReferenceTexture->bFlipGreenChannel;
+
+		// MaxTextureSize setting. Based on the ReferenceTexture and Platform settings.
+		const UTextureLODSettings& TextureLODSettings = GenerationContext.CompilationContext->Options.TargetPlatform->GetTextureLODSettings();
+		Props.MaxTextureSize = GetMaxTextureSize(*ReferenceTexture, TextureLODSettings);
+
+		// ReferenceTexture source size. Textures contributing to this Image should be equal to or smaller than TextureSize. 
+		// The LOD Bias applied to the root node will be applied on top of it.
+		Props.TextureSize = (int32)FMath::Max3(ReferenceTexture->Source.GetSizeX(), ReferenceTexture->Source.GetSizeY(), 1LL);
+
+		// TODO: MTBL-1081
+		// TextureGroup::TEXTUREGROUP_UI does not support streaming. If we generate a texture that requires streaming and set this group, it will crash when initializing the resource. 
+		// If LODGroup == TEXTUREGROUP_UI, UTexture::IsPossibleToStream() will return false and UE will assume all mips are loaded, when they're not, and crash.
+		if (Props.LODGroup == TEXTUREGROUP_UI)
+		{
+			Props.LODGroup = TextureGroup::TEXTUREGROUP_Character;
+
+			FString msg = FString::Printf(TEXT("The Reference texture [%s] is using TEXTUREGROUP_UI which does not support streaming. Please set a different TEXTURE group."),
+				*ReferenceTexture->GetName(), *ImageKey.ParameterName.ToString());
+			GenerationContext.Log(FText::FromString(msg), Node, EMessageSeverity::Info);
+		}
+	}
+	else
+	{
+		// warning!
+		FString msg = FString::Printf(TEXT("The Reference texture for material image [%s] is not set and it couldn't be found automatically."), *ImageKey.ParameterName.ToString());
+		GenerationContext.Log(FText::FromString(msg), Node);
+	}
+
+	return Props;
+}
+
+UE::Mutable::Private::Ptr<UE::Mutable::Private::NodeImage> GenerateMipMapAndFormatNodes(
+	const UCustomizableObjectNode* Node,
+	FMutableGraphGenerationContext& GenerationContext,
+	UE::Mutable::Private::Ptr<UE::Mutable::Private::NodeImage> ImageNode,
+	UTexture2D* ReferenceTexture)
+{
+	if (!ReferenceTexture)
+	{
+		return ImageNode;
+	}
+
+	// Apply base LODBias. It will be propagated to most images.
+	const uint32 BaseLODBias = ComputeLODBiasForTexture(GenerationContext, *ReferenceTexture);
+	UE::Mutable::Private::NodeImagePtr LastImage = ResizeTextureByNumMips(ImageNode, BaseLODBias);
+
+	if (ReferenceTexture->MipGenSettings != TextureMipGenSettings::TMGS_NoMipmaps)
+	{
+		UE::Mutable::Private::EMipmapFilterType MipGenerationFilterType = Invoke([&]()
+			{
+				if (ReferenceTexture)
+				{
+					switch (ReferenceTexture->MipGenSettings)
+					{
+					case TextureMipGenSettings::TMGS_SimpleAverage: return UE::Mutable::Private::EMipmapFilterType::SimpleAverage;
+					case TextureMipGenSettings::TMGS_Unfiltered:    return UE::Mutable::Private::EMipmapFilterType::Unfiltered;
+					default: return UE::Mutable::Private::EMipmapFilterType::SimpleAverage;
+					}
+				}
+
+				return UE::Mutable::Private::EMipmapFilterType::SimpleAverage;
+			});
+
+
+		UE::Mutable::Private::Ptr<UE::Mutable::Private::NodeImageMipmap> MipmapImage = new UE::Mutable::Private::NodeImageMipmap();
+		MipmapImage->Source = LastImage;
+		MipmapImage->Settings.FilterType = MipGenerationFilterType;
+		MipmapImage->Settings.AddressMode = UE::Mutable::Private::EAddressMode::None;
+
+		MipmapImage->SetMessageContext(Node);
+		LastImage = MipmapImage;
+	}
+
+	// Apply composite image. This needs to be computed after mipmaps generation. 	
+	if (ReferenceTexture && ReferenceTexture->GetCompositeTexture() && ReferenceTexture->CompositeTextureMode != CTM_Disabled)
+	{
+		UE::Mutable::Private::Ptr<UE::Mutable::Private::NodeImageNormalComposite> CompositedImage = new UE::Mutable::Private::NodeImageNormalComposite();
+		CompositedImage->Base = LastImage;
+		CompositedImage->Power = ReferenceTexture->CompositePower;
+
+		UE::Mutable::Private::ECompositeImageMode CompositeImageMode = [CompositeTextureMode = ReferenceTexture->CompositeTextureMode]() -> UE::Mutable::Private::ECompositeImageMode
+			{
+				switch (CompositeTextureMode)
+				{
+				case CTM_NormalRoughnessToRed: return UE::Mutable::Private::ECompositeImageMode::CIM_NormalRoughnessToRed;
+				case CTM_NormalRoughnessToGreen: return UE::Mutable::Private::ECompositeImageMode::CIM_NormalRoughnessToGreen;
+				case CTM_NormalRoughnessToBlue: return UE::Mutable::Private::ECompositeImageMode::CIM_NormalRoughnessToBlue;
+				case CTM_NormalRoughnessToAlpha: return UE::Mutable::Private::ECompositeImageMode::CIM_NormalRoughnessToAlpha;
+
+				default: return UE::Mutable::Private::ECompositeImageMode::CIM_Disabled;
+				}
+			}();
+
+		CompositedImage->Mode = CompositeImageMode;
+
+		UE::Mutable::Private::Ptr<UE::Mutable::Private::NodeImageConstant> CompositeNormalImage = new UE::Mutable::Private::NodeImageConstant();
+
+		UTexture2D* ReferenceCompositeNormalTexture = Cast<UTexture2D>(ReferenceTexture->GetCompositeTexture());
+		if (ReferenceCompositeNormalTexture)
+		{
+			// GenerationContext.ArrayTextureUnrealToMutableTask.Add(FTextureUnrealToMutableTask(CompositeNormalImage, ReferenceCompositeNormalTexture, Node, true));
+			// TODO: The normal composite part is not propagated, so it will be unsupported. Create a task that performs the required transforms at mutable image level, and add the right operations here
+			// instead of propagating the flag and doing them on unreal-convert.
+			CompositeNormalImage->Image = GenerateImageConstant(ReferenceCompositeNormalTexture, GenerationContext, false);
+
+			UE::Mutable::Private::Ptr<UE::Mutable::Private::NodeImageMipmap> NormalCompositeMipmapImage = new UE::Mutable::Private::NodeImageMipmap();
+			const uint32 MipsToSkip = ComputeLODBiasForTexture(GenerationContext, *ReferenceCompositeNormalTexture, ReferenceTexture);
+			NormalCompositeMipmapImage->Source = ResizeTextureByNumMips(CompositeNormalImage, MipsToSkip);
+			NormalCompositeMipmapImage->Settings.FilterType = UE::Mutable::Private::EMipmapFilterType::SimpleAverage;
+			NormalCompositeMipmapImage->Settings.AddressMode = UE::Mutable::Private::EAddressMode::None;
+
+			CompositedImage->Normal = NormalCompositeMipmapImage;
+
+			CompositeNormalImage->SourceDataDescriptor.OptionalMaxLODSize = 0;
+			if (GenerationContext.CompilationContext->Options.TargetPlatform)
+			{
+				const UTextureLODSettings& TextureLODSettings = GenerationContext.CompilationContext->Options.TargetPlatform->GetTextureLODSettings();
+				const FTextureLODGroup& LODGroupInfo = TextureLODSettings.GetTextureLODGroup(ReferenceCompositeNormalTexture->LODGroup);
+
+				CompositeNormalImage->SourceDataDescriptor.OptionalMaxLODSize = LODGroupInfo.OptionalMaxLODSize;
+				CompositeNormalImage->SourceDataDescriptor.OptionalLODBias = LODGroupInfo.OptionalLODBias;
+				CompositeNormalImage->SourceDataDescriptor.NumNonOptionalLODs = GenerationContext.CompilationContext->Options.MinDiskMips;
+			}
+
+			const FString TextureName = GetNameSafe(ReferenceCompositeNormalTexture).ToLower();
+			CompositeNormalImage->SourceDataDescriptor.SourceId = CityHash32(reinterpret_cast<const char*>(*TextureName), TextureName.Len() * sizeof(FString::ElementType));
+		}
+
+		LastImage = CompositedImage;
+	}
+
+	LastImage = GenerateMutableImageFormat(LastImage, ReferenceTexture, GenerationContext, Node);
+
+	return LastImage;
+}
+
 #undef LOCTEXT_NAMESPACE
 
diff --git a/Engine/Plugins/Mutable/Source/CustomizableObjectEditor/Private/MuCOE/GenerateMutableSource/GenerateMutableSourceMaterial.h b/Engine/Plugins/Mutable/Source/CustomizableObjectEditor/Private/MuCOE/GenerateMutableSource/GenerateMutableSourceMaterial.h
--- a/Engine/Plugins/Mutable/Source/CustomizableObjectEditor/Private/MuCOE/GenerateMutableSource/GenerateMutableSourceMaterial.h
+++ b/Engine/Plugins/Mutable/Source/CustomizableObjectEditor/Private/MuCOE/GenerateMutableSource/GenerateMutableSourceMaterial.h
@@ -2,10 +2,31 @@
 
 #pragma once
 
+#include "GenerateMutableSource.h"
 #include "MuT/NodeMaterial.h"
 
 class UEdGraphPin;
 struct FMutableGraphGenerationContext;
 
+// Add to this struct all the extra parameters that need to be passed to the GenerateMutableSourceMaterial function.
+struct FGenerationMaterialOptions 
+{
+	// Node surface. Only needed for MeshSection nodes.
+	UE::Mutable::Private::Ptr<UE::Mutable::Private::NodeSurfaceNew> SurfaceNode = nullptr;
+};
+
 /** Convert a CustomizableObject Source Graph into a mutable source graph. */
-UE::Mutable::Private::Ptr<UE::Mutable::Private::NodeMaterial> GenerateMutableSourceMaterial(const UEdGraphPin* Pin, FMutableGraphGenerationContext& GenerationContext);
+UE::Mutable::Private::Ptr<UE::Mutable::Private::NodeMaterial> GenerateMutableSourceMaterial(
+	const UEdGraphPin* Pin,
+	FMutableGraphGenerationContext& GenerationContext,
+	const FGenerationMaterialOptions& GenerationOptions);
+
+/** Generates an FGeneratedImageProperties from a Material Parameter and (if not null) its Reference Texture. */
+FGeneratedImageProperties GenerateImageProperties(const UCustomizableObjectNode* Node, const UTexture2D* ReferenceTexture, const UE::Mutable::Private::FParameterKey& ImageKey, bool bIsPassthrough, FMutableGraphGenerationContext& GenerationContext);
+
+// Generate additional nodes to generete a texture with the correct format and size.
+UE::Mutable::Private::Ptr<UE::Mutable::Private::NodeImage> GenerateMipMapAndFormatNodes(
+	const UCustomizableObjectNode* Node,
+	FMutableGraphGenerationContext& GenerationContext,
+	UE::Mutable::Private::Ptr<UE::Mutable::Private::NodeImage> ImageNode,
+	UTexture2D* ReferenceTexture);
diff --git a/Engine/Plugins/Mutable/Source/CustomizableObjectEditor/Private/MuCOE/GenerateMutableSource/GenerateMutableSourceModifier.cpp b/Engine/Plugins/Mutable/Source/CustomizableObjectEditor/Private/MuCOE/GenerateMutableSource/GenerateMutableSourceModifier.cpp
--- a/Engine/Plugins/Mutable/Source/CustomizableObjectEditor/Private/MuCOE/GenerateMutableSource/GenerateMutableSourceModifier.cpp
+++ b/Engine/Plugins/Mutable/Source/CustomizableObjectEditor/Private/MuCOE/GenerateMutableSource/GenerateMutableSourceModifier.cpp
@@ -321,7 +321,8 @@ UE::Mutable::Private::Ptr<UE::Mutable::Private::NodeModifier> GenerateMutableSou
 			for (int32 ImageIndex = 0; ImageIndex < NumImages; ++ImageIndex)
 			{
 				UE::Mutable::Private::NodeImagePtr ImageNode;
-				FString MaterialParameterName;
+				FName MaterialParameterName = TypedNodeExt->GetParameterName(EMaterialParameterType::Texture, ImageIndex);
+				int32 MaterialParameterLayerIndex = TypedNodeExt->GetParameterLayerIndex(EMaterialParameterType::Texture, ImageIndex);
 
 				if (!ImageNode) // If
 				{
@@ -331,7 +332,7 @@ UE::Mutable::Private::Ptr<UE::Mutable::Private::NodeModifier> GenerateMutableSou
 					if (ProjectorInfo)
 					{
 						ensure(LOD > GenerationContext.FirstLODAvailable[GenerationContext.CurrentMeshComponent]);
-						check(ProjectorInfo->SurfNode->Images[ImageIndex].Image == ProjectorInfo->ImageNode);
+						//check(ProjectorInfo->SurfNode->Images[ImageIndex].Image == ProjectorInfo->ImageNode); TODO(Max): Review this
 						ImageNode = ProjectorInfo->ImageNode;
 					}
 				}
@@ -362,15 +363,13 @@ UE::Mutable::Private::Ptr<UE::Mutable::Private::NodeModifier> GenerateMutableSou
 							const int32 ReferenceTextureSize = 0; // TODO GetBaseTextureSize(GenerationContext, TypedNodeExt, ImageIndex);
 
 							ImageNode = GenerateMutableSourceImage(ConnectedPin, GenerationContext, ReferenceTextureSize);
-							MaterialParameterName = TypedNodeExt->GetParameterName(EMaterialParameterType::Texture, ImageIndex).ToString();
-
 						}
 					}
 				}
 
 
 				SurfNode->LODs[LODIndex].Textures[ImageIndex].Extend = ImageNode;
-				SurfNode->LODs[LODIndex].Textures[ImageIndex].MaterialParameterName = MaterialParameterName;
+				SurfNode->LODs[LODIndex].Textures[ImageIndex].MaterialParameterKey = UE::Mutable::Private::FParameterKey(MaterialParameterName, MaterialParameterLayerIndex);
 			}
 		}
 
@@ -479,10 +478,12 @@ UE::Mutable::Private::Ptr<UE::Mutable::Private::NodeModifier> GenerateMutableSou
 					const UEdGraphPin* ConnectedImagePin = FollowInputPin(*TypedNodeEdit->GetUsedImagePin(ImageId));
 
 					UE::Mutable::Private::NodeModifierSurfaceEdit::FTexture& ImagePatch = SurfNode->LODs[LODIndex].Textures[ImageIndex];
-
-					ImagePatch.MaterialParameterName = TypedNodeEdit->GetParameterName(EMaterialParameterType::Texture, ImageIndex).ToString();
-
+					
+					FName MaterialParameterName = TypedNodeEdit->GetParameterName(EMaterialParameterType::Texture, ImageIndex);
+					int32 MaterialParameterLayerIndex = TypedNodeEdit->GetParameterLayerIndex(EMaterialParameterType::Texture, ImageIndex);
+					
 					// \todo: expose these two options?
+					ImagePatch.MaterialParameterKey = UE::Mutable::Private::FParameterKey(MaterialParameterName, MaterialParameterLayerIndex);
 					ImagePatch.PatchBlendType = UE::Mutable::Private::EBlendType::BT_BLEND;
 					ImagePatch.bPatchApplyToAlpha = true;
 
diff --git a/Engine/Plugins/Mutable/Source/CustomizableObjectEditor/Private/MuCOE/GenerateMutableSource/GenerateMutableSourceSurface.cpp b/Engine/Plugins/Mutable/Source/CustomizableObjectEditor/Private/MuCOE/GenerateMutableSource/GenerateMutableSourceSurface.cpp
--- a/Engine/Plugins/Mutable/Source/CustomizableObjectEditor/Private/MuCOE/GenerateMutableSource/GenerateMutableSourceSurface.cpp
+++ b/Engine/Plugins/Mutable/Source/CustomizableObjectEditor/Private/MuCOE/GenerateMutableSource/GenerateMutableSourceSurface.cpp
@@ -6,9 +6,8 @@
 #include "Engine/StaticMesh.h"
 #include "Engine/TextureLODSettings.h"
 #include "GPUSkinPublicDefs.h"
-#include "Interfaces/ITargetPlatform.h"
 #include "Modules/ModuleManager.h"
-#include "TextureCompressorModule.h"
+
 #include "Materials/MaterialInstance.h"
 #include "GPUSkinVertexFactory.h"
 #include "Rendering/SkeletalMeshLODModel.h"
@@ -52,6 +51,7 @@
 #include "MuT/NodeMeshFormat.h"
 #include "MuT/NodeMeshFragment.h"
 #include "MuT/NodeMaterialConstant.h"
+#include "MuT/NodeMaterialModify.h"
 #include "MuT/NodeSurfaceSwitch.h"
 #include "MuT/NodeSurfaceVariation.h"
 
@@ -82,8 +82,6 @@ UE::Mutable::Private::Ptr<UE::Mutable::Private::NodeSurface> GenerateMutableSour
 	
 	UE::Mutable::Private::Ptr<UE::Mutable::Private::NodeSurface> Result;
 
-	const int32 LOD = Node->IsAffectedByLOD() ? GenerationContext.CurrentLOD : 0;
-	
 	if (UCustomizableObjectNode* CustomObjNode = Cast<UCustomizableObjectNode>(Node))
 	{
 		if (CustomObjNode->IsNodeOutDatedAndNeedsRefresh())
@@ -179,7 +177,6 @@ UE::Mutable::Private::Ptr<UE::Mutable::Private::NodeSurface> GenerateMutableSour
 		{
 			if (const UEdGraphPin* ConnectedPin = FollowInputPin(*TypedNodeMat->GetMeshPin()))
 			{
-
 				// Flags to know which UV channels need layout
 				FLayoutGenerationFlags LayoutGenerationFlags;
 				
@@ -215,437 +212,18 @@ UE::Mutable::Private::Ptr<UE::Mutable::Private::NodeSurface> GenerateMutableSour
 			}
 		}
 
-		bool bMaterialPinLinked = false;
-		const UEdGraphPin* ConnectedMaterialAssetPin = nullptr;
-
-		// Checking if we should not use the material of the table node even if it is linked to the material node
-		if (const UEdGraphPin* MaterialAssetPin = TypedNodeMat->GetMaterialAssetPin())
-		{
-			ConnectedMaterialAssetPin = FollowInputPin(*MaterialAssetPin);
-
-			if (ConnectedMaterialAssetPin)
-			{
-				//Check if the pin goes through a macro or tunnel node
-				ConnectedMaterialAssetPin = GraphTraversal::FindIOPinSourceThroughMacroContext(*ConnectedMaterialAssetPin, &GenerationContext.MacroNodesStack);
-
-				GenerationContext.CurrentMaterialParameterId = MaterialAssetPin->PinId.ToString();
-				SurfNode->Material = GenerateMutableSourceMaterial(ConnectedMaterialAssetPin, GenerationContext);
-			}
-		}
-
-		if (!SurfNode->Material)
-		{
-			UE::Mutable::Private::Ptr<UE::Mutable::Private::NodeMaterialConstant> MaterialNode = new UE::Mutable::Private::NodeMaterialConstant();
-			MaterialNode->MaterialId = GenerationContext.AddPassthroughObject(*GenerationContext.CurrentReferencedMaterial.Get());
-			SurfNode->Material = MaterialNode;
-		}
-
-		int32 NumImages = TypedNodeMat->GetNumParameters(EMaterialParameterType::Texture);
-		SurfNode->Images.SetNum(NumImages);
-
-		for (int32 ImageIndex = 0; ImageIndex < NumImages; ++ImageIndex)
-		{
-			const UEdGraphPin* ImagePin = TypedNodeMat->GetParameterPin(EMaterialParameterType::Texture, ImageIndex);
-
-			const bool bIsImagePinLinked = ImagePin && FollowInputPin(*ImagePin);
-
-			if (bIsImagePinLinked && !TypedNodeMat->IsImageMutableMode(ImageIndex))
-			{
-				if (const UEdGraphPin* ConnectedPin = FollowInputPin(*ImagePin))
-				{
-					// Find or add Image properties
-					const FGeneratedImagePropertiesKey PropsKey(TypedNodeMat, (uint32)ImageIndex);
-					const bool bNewImageProps = !GenerationContext.ImageProperties.Contains(PropsKey);
-
-					FGeneratedImageProperties& Props = GenerationContext.ImageProperties.FindOrAdd(PropsKey);
-					if (bNewImageProps)
-					{
-						// We don't need a reference texture or props here, but we do need the parameter name.
-						Props.TextureParameterName = TypedNodeMat->GetParameterName(EMaterialParameterType::Texture, ImageIndex).ToString();
-						Props.ImagePropertiesIndex = GenerationContext.ImageProperties.Num() - 1;
-						Props.bIsPassThrough = true;
-					}
-
-					// This is a connected pass-through texture that simply has to be passed to the core
-					UE::Mutable::Private::Ptr<UE::Mutable::Private::NodeImage> PassThroughImagePtr = GenerateMutableSourceImage(ConnectedPin, GenerationContext, 0);
-					SurfNode->Images[ImageIndex].Image = PassThroughImagePtr;
-
-					check(Props.ImagePropertiesIndex != INDEX_NONE);
-					const FString SurfNodeImageName = FString::Printf(TEXT("%d"), Props.ImagePropertiesIndex);
-					const int32 LayerIndex = TypedNodeMat->GetParameterLayerIndex(EMaterialParameterType::Texture, ImageIndex);
-					SurfNode->Images[ImageIndex].ParameterInfo.ParameterName = SurfNodeImageName;
-					SurfNode->Images[ImageIndex].ParameterInfo.MaterialLayerIndex = LayerIndex;
-					SurfNode->Images[ImageIndex].LayoutIndex = -1;
-					SurfNode->Images[ImageIndex].MaterialName = Material->GetName();
-					SurfNode->Images[ImageIndex].MaterialParameterName = Props.TextureParameterName;
-
-				}
-			}
-			else
-			{
-				UE::Mutable::Private::NodeImagePtr GroupProjectionImg;
-				UTexture2D* GroupProjectionReferenceTexture = nullptr;
-				const FString ImageName = TypedNodeMat->GetParameterName(EMaterialParameterType::Texture, ImageIndex).ToString();
-				const FNodeMaterialParameterId ImageId = TypedNodeMat->GetParameterId(EMaterialParameterType::Texture, ImageIndex);
-
-				FString MaterialImageId = FGroupProjectorImageInfo::GenerateId(TypedNodeMat, ImageIndex);
-				bool bShareProjectionTexturesBetweenLODs = false;
-				FGroupProjectorImageInfo* ProjectorInfo = GenerationContext.GroupProjectorLODCache.Find(MaterialImageId);
-
-				if (!ProjectorInfo) // No previous LOD of this material generated the image.
-				{
-					bool bIsGroupProjectorImage = false;
-
-					GroupProjectionImg = GenerateMutableSourceGroupProjector(LOD, ImageIndex, MeshNode, GenerationContext,
-						TypedNodeMat, nullptr, bShareProjectionTexturesBetweenLODs, bIsGroupProjectorImage,
-						GroupProjectionReferenceTexture);
-
-					if (GroupProjectionImg.get() || TypedNodeMat->IsImageMutableMode(ImageIndex))
-					{
-						// Get the reference texture
-						UTexture2D* ReferenceTexture = nullptr;
-						{
-							//TODO(Max) UE-220247: Add support for multilayer materials
-							GenerationContext.CurrentMaterialTableParameter = ImageName;
-							GenerationContext.CurrentMaterialParameterId = ImageId.ParameterId.ToString();
-
-							ReferenceTexture = GroupProjectionImg.get() ? GroupProjectionReferenceTexture : nullptr;
-						
-							if (!ReferenceTexture)
-							{
-								ReferenceTexture = TypedNodeMat->GetImageReferenceTexture(ImageIndex);
-							}
-
-							// In case of group projector, don't follow the pin to find the reference texture.
-							if (!GroupProjectionImg.get() && !ReferenceTexture && ImagePin)
-							{
-								if (const UEdGraphPin* ConnectedPin = FollowInputPin(*ImagePin))
-								{
-									ReferenceTexture = FindReferenceImage(ConnectedPin, GenerationContext);
-								}
-							}
-
-							if (!ReferenceTexture && 
-								ConnectedMaterialAssetPin && Cast<UCustomizableObjectNodeTable>(ConnectedMaterialAssetPin->GetOwningNode()) != nullptr)
-							{
-								ReferenceTexture = FindReferenceImage(ConnectedMaterialAssetPin, GenerationContext);
-							}	
-
-							if (!ReferenceTexture)
-							{
-								ReferenceTexture = TypedNodeMat->GetImageValue(ImageIndex);
-							}
-						}
-
-						const FGeneratedImagePropertiesKey PropsKey(TypedNodeMat, ImageIndex);
-						const bool bNewImageProps = !GenerationContext.ImageProperties.Contains(PropsKey);
-
-						FGeneratedImageProperties& Props = GenerationContext.ImageProperties.FindOrAdd(PropsKey);
-
-						if (bNewImageProps)
-						{
-							if (ReferenceTexture)
-							{
-								// Store properties for the generated images
-								Props.TextureParameterName = ImageName;
-								Props.ImagePropertiesIndex = GenerationContext.ImageProperties.Num() - 1;
-
-								Props.CompressionSettings = ReferenceTexture->CompressionSettings;
-								Props.Filter = ReferenceTexture->Filter;
-								Props.SRGB = ReferenceTexture->SRGB;
-								Props.LODBias = 0;
-								Props.MipGenSettings = ReferenceTexture->MipGenSettings;
-								Props.LODGroup = ReferenceTexture->LODGroup;
-								Props.AddressX = ReferenceTexture->AddressX;
-								Props.AddressY = ReferenceTexture->AddressY;
-								Props.bFlipGreenChannel = ReferenceTexture->bFlipGreenChannel;
-
-
-								// MaxTextureSize setting. Based on the ReferenceTexture and Platform settings.
-								const UTextureLODSettings& TextureLODSettings = GenerationContext.CompilationContext->Options.TargetPlatform->GetTextureLODSettings();
-								Props.MaxTextureSize = GetMaxTextureSize(*ReferenceTexture, TextureLODSettings);
-
-								// ReferenceTexture source size. Textures contributing to this Image should be equal to or smaller than TextureSize. 
-								// The LOD Bias applied to the root node will be applied on top of it.
-								Props.TextureSize = (int32)FMath::Max3(ReferenceTexture->Source.GetSizeX(), ReferenceTexture->Source.GetSizeY(), 1LL);
-
-								// TODO: MTBL-1081
-								// TextureGroup::TEXTUREGROUP_UI does not support streaming. If we generate a texture that requires streaming and set this group, it will crash when initializing the resource. 
-								// If LODGroup == TEXTUREGROUP_UI, UTexture::IsPossibleToStream() will return false and UE will assume all mips are loaded, when they're not, and crash.
-								if (Props.LODGroup == TEXTUREGROUP_UI)
-								{
-									Props.LODGroup = TextureGroup::TEXTUREGROUP_Character;
-
-									FString msg = FString::Printf(TEXT("The Reference texture [%s] is using TEXTUREGROUP_UI which does not support streaming. Please set a different TEXTURE group."),
-										*ReferenceTexture->GetName(), *ImageName);
-									GenerationContext.Log(FText::FromString(msg), Node, EMessageSeverity::Info);
-								}
-							}
-							else
-							{
-								// warning!
-								FString msg = FString::Printf(TEXT("The Reference texture for material image [%s] is not set and it couldn't be found automatically."), *ImageName);
-								GenerationContext.Log(FText::FromString(msg), Node);
-							}
-						}
-
-						// Generate the texture nodes
-						UE::Mutable::Private::NodeImagePtr ImageNode = [&]()
-						{
-							if (TypedNodeMat->IsImageMutableMode(ImageIndex))
-							{
-								if (ImagePin)
-								{
-									if (const UEdGraphPin* ConnectedPin = FollowInputPin(*ImagePin))
-									{
-										return GenerateMutableSourceImage(ConnectedPin, GenerationContext, Props.TextureSize);
-									}
-								}
-
-								// If the table material pin is linked to a table node, get all the textures of the current material parameter (CurrentMaterialTableParameter) from the Material Instances of the specified data table column.
-								// Then Generate a mutable table column with all these textures.
-								if (ConnectedMaterialAssetPin && Cast<UCustomizableObjectNodeTable>(ConnectedMaterialAssetPin->GetOwningNode()) != nullptr)
-								{
-									return GenerateMutableSourceImage(ConnectedMaterialAssetPin, GenerationContext, Props.TextureSize);
-								}
-
-								// Else
-								{
-									UTexture2D* Texture2D = TypedNodeMat->GetImageValue(ImageIndex);
-
-									if (Texture2D)
-									{
-										UE::Mutable::Private::Ptr<UE::Mutable::Private::NodeImageConstant> ConstImageNode = new UE::Mutable::Private::NodeImageConstant();
-										ConstImageNode->Image = GenerateImageConstant(Texture2D, GenerationContext, false);
-
-										const uint32 MipsToSkip = ComputeLODBiasForTexture(GenerationContext, *Texture2D, nullptr, Props.TextureSize);
-										UE::Mutable::Private::Ptr<UE::Mutable::Private::NodeImage> Result =  ResizeTextureByNumMips(ConstImageNode, MipsToSkip);
-	
-										const UTextureLODSettings& TextureLODSettings = GenerationContext.CompilationContext->Options.TargetPlatform->GetTextureLODSettings();
-										const FTextureLODGroup& LODGroupInfo = TextureLODSettings.GetTextureLODGroup(Texture2D->LODGroup);
-
-										ConstImageNode->SourceDataDescriptor.OptionalMaxLODSize = LODGroupInfo.OptionalMaxLODSize;
-										ConstImageNode->SourceDataDescriptor.OptionalLODBias = LODGroupInfo.OptionalLODBias;
-										ConstImageNode->SourceDataDescriptor.NumNonOptionalLODs = GenerationContext.CompilationContext->Options.MinDiskMips;
-
-										const FString TextureName = GetNameSafe(Texture2D).ToLower();
-										ConstImageNode->SourceDataDescriptor.SourceId = CityHash32(reinterpret_cast<const char*>(*TextureName), TextureName.Len() * sizeof(FString::ElementType));
-
-										return Result;
-									}
-									else
-									{
-										return UE::Mutable::Private::NodeImagePtr();
-									}
-								}
-							}
-							else
-							{
-								return UE::Mutable::Private::NodeImagePtr();
-							}
-						}();
-
-						if (GroupProjectionImg.get())
-						{
-							ImageNode = GroupProjectionImg;
-						}
-
-						if (ReferenceTexture)
-						{
-							// Apply base LODBias. It will be propagated to most images.
-							const uint32 BaseLODBias = ComputeLODBiasForTexture(GenerationContext, *ReferenceTexture);
-							UE::Mutable::Private::NodeImagePtr LastImage = ResizeTextureByNumMips(ImageNode, BaseLODBias);
-
-							if (ReferenceTexture->MipGenSettings != TextureMipGenSettings::TMGS_NoMipmaps)
-							{
-								UE::Mutable::Private::EMipmapFilterType MipGenerationFilterType = Invoke([&]()
-								{
-									if (ReferenceTexture)
-									{
-										switch (ReferenceTexture->MipGenSettings)
-										{
-											case TextureMipGenSettings::TMGS_SimpleAverage: return UE::Mutable::Private::EMipmapFilterType::SimpleAverage;
-											case TextureMipGenSettings::TMGS_Unfiltered:    return UE::Mutable::Private::EMipmapFilterType::Unfiltered;
-											default: return UE::Mutable::Private::EMipmapFilterType::SimpleAverage;
-										}
-									}
-
-									return UE::Mutable::Private::EMipmapFilterType::SimpleAverage;
-								});
-
-
-								UE::Mutable::Private::Ptr<UE::Mutable::Private::NodeImageMipmap> MipmapImage = new UE::Mutable::Private::NodeImageMipmap();
-								MipmapImage->Source = LastImage;
-								MipmapImage->Settings.FilterType = MipGenerationFilterType;
-								MipmapImage->Settings.AddressMode = UE::Mutable::Private::EAddressMode::None;
-
-								MipmapImage->SetMessageContext(Node);
-								LastImage = MipmapImage;
-							}
-
-							// Apply composite image. This needs to be computed after mipmaps generation. 	
-							if (ReferenceTexture && ReferenceTexture->GetCompositeTexture() && ReferenceTexture->CompositeTextureMode != CTM_Disabled)
-							{
-								UE::Mutable::Private::Ptr<UE::Mutable::Private::NodeImageNormalComposite> CompositedImage = new UE::Mutable::Private::NodeImageNormalComposite();
-								CompositedImage->Base = LastImage;
-								CompositedImage->Power = ReferenceTexture->CompositePower;
-
-								UE::Mutable::Private::ECompositeImageMode CompositeImageMode = [CompositeTextureMode = ReferenceTexture->CompositeTextureMode]() -> UE::Mutable::Private::ECompositeImageMode
-								{
-									switch (CompositeTextureMode)
-									{
-										case CTM_NormalRoughnessToRed: return UE::Mutable::Private::ECompositeImageMode::CIM_NormalRoughnessToRed;
-										case CTM_NormalRoughnessToGreen: return UE::Mutable::Private::ECompositeImageMode::CIM_NormalRoughnessToGreen;
-										case CTM_NormalRoughnessToBlue: return UE::Mutable::Private::ECompositeImageMode::CIM_NormalRoughnessToBlue;
-										case CTM_NormalRoughnessToAlpha: return UE::Mutable::Private::ECompositeImageMode::CIM_NormalRoughnessToAlpha;
-
-										default: return UE::Mutable::Private::ECompositeImageMode::CIM_Disabled;
-									}
-								}();
-
-								CompositedImage->Mode = CompositeImageMode;
-
-								UE::Mutable::Private::Ptr<UE::Mutable::Private::NodeImageConstant> CompositeNormalImage = new UE::Mutable::Private::NodeImageConstant();
-
-								UTexture2D* ReferenceCompositeNormalTexture = Cast<UTexture2D>(ReferenceTexture->GetCompositeTexture());
-								if (ReferenceCompositeNormalTexture)
-								{
-									// GenerationContext.ArrayTextureUnrealToMutableTask.Add(FTextureUnrealToMutableTask(CompositeNormalImage, ReferenceCompositeNormalTexture, Node, true));
-									// TODO: The normal composite part is not propagated, so it will be unsupported. Create a task that performs the required transforms at mutable image level, and add the right operations here
-									// instead of propagating the flag and doing them on unreal-convert.
-									CompositeNormalImage->Image = GenerateImageConstant(ReferenceCompositeNormalTexture, GenerationContext, false);
-
-									UE::Mutable::Private::Ptr<UE::Mutable::Private::NodeImageMipmap> NormalCompositeMipmapImage = new UE::Mutable::Private::NodeImageMipmap();
-									const uint32 MipsToSkip = ComputeLODBiasForTexture(GenerationContext, *ReferenceCompositeNormalTexture, ReferenceTexture);
-									NormalCompositeMipmapImage->Source = ResizeTextureByNumMips(CompositeNormalImage, MipsToSkip);
-									NormalCompositeMipmapImage->Settings.FilterType = UE::Mutable::Private::EMipmapFilterType::SimpleAverage;
-									NormalCompositeMipmapImage->Settings.AddressMode = UE::Mutable::Private::EAddressMode::None;
-
-									CompositedImage->Normal = NormalCompositeMipmapImage;
-
-									CompositeNormalImage->SourceDataDescriptor.OptionalMaxLODSize = 0;
-									if (GenerationContext.CompilationContext->Options.TargetPlatform)
-									{
-										const UTextureLODSettings& TextureLODSettings = GenerationContext.CompilationContext->Options.TargetPlatform->GetTextureLODSettings();
-										const FTextureLODGroup& LODGroupInfo = TextureLODSettings.GetTextureLODGroup(ReferenceCompositeNormalTexture->LODGroup);
-
-										CompositeNormalImage->SourceDataDescriptor.OptionalMaxLODSize = LODGroupInfo.OptionalMaxLODSize;
-										CompositeNormalImage->SourceDataDescriptor.OptionalLODBias = LODGroupInfo.OptionalLODBias;
-										CompositeNormalImage->SourceDataDescriptor.NumNonOptionalLODs = GenerationContext.CompilationContext->Options.MinDiskMips;
-									}
-
-									const FString TextureName = GetNameSafe(ReferenceCompositeNormalTexture).ToLower();
-									CompositeNormalImage->SourceDataDescriptor.SourceId = CityHash32(reinterpret_cast<const char*>(*TextureName), TextureName.Len() * sizeof(FString::ElementType));
-								}
-
-								LastImage = CompositedImage;
-							}
-
-							// Apply Image format on top of the image tree to propagate the format down
-							ImageNode = GenerateMutableImageFormat(LastImage, ReferenceTexture, GenerationContext, Node);
-						}
-
-						SurfNode->Images[ImageIndex].Image = ImageNode;
-
-						check(Props.ImagePropertiesIndex != INDEX_NONE);
-						const FString SurfNodeImageName = FString::Printf(TEXT("%d"), Props.ImagePropertiesIndex);
-
-						// Name and material layer
-						const int32 LayerIndex = TypedNodeMat->GetParameterLayerIndex(EMaterialParameterType::Texture, ImageIndex);
-						SurfNode->Images[ImageIndex].ParameterInfo.ParameterName = SurfNodeImageName;
-						SurfNode->Images[ImageIndex].ParameterInfo.MaterialLayerIndex = LayerIndex;
-
-						// If we are generating an implicit component (with a passthrough mesh) we don't apply any layout.
-						int32 UVLayout = -1;
-						if (!bGeneratingImplicitComponent)
-						{
-							UVLayout = TypedNodeMat->GetImageUVLayout(ImageIndex);;
-						}
-						SurfNode->Images[ImageIndex].LayoutIndex = UVLayout;
-						SurfNode->Images[ImageIndex].MaterialName = Material->GetName();
-						SurfNode->Images[ImageIndex].MaterialParameterName = ImageName;
-
-						if (bShareProjectionTexturesBetweenLODs && bIsGroupProjectorImage)
-						{
-							// Add to the GroupProjectorLODCache to potentially reuse this projection texture in higher LODs
-							ensure(LOD == GenerationContext.FirstLODAvailable[GenerationContext.CurrentMeshComponent]);
-							GenerationContext.GroupProjectorLODCache.Add(MaterialImageId,
-								FGroupProjectorImageInfo(ImageNode, ImageName, ImageName, TypedNodeMat, SurfNode, UVLayout, LayerIndex));
-						}
-					}
-				}
-				else
-				{
-					ensure(LOD > GenerationContext.FirstLODAvailable[GenerationContext.CurrentMeshComponent]);
-					check(ProjectorInfo->SurfNode->Images[ImageIndex].Image == ProjectorInfo->ImageNode);
-					SurfNode->Images[ImageIndex].Image = ProjectorInfo->ImageNode;
-					SurfNode->Images[ImageIndex].ParameterInfo.ParameterName = ProjectorInfo->TextureName;
-					SurfNode->Images[ImageIndex].ParameterInfo.MaterialLayerIndex = ProjectorInfo->LayerIndex;
-					SurfNode->Images[ImageIndex].LayoutIndex = ProjectorInfo->UVLayout;
-				}
-			}
-		}
-
-		const int32 NumVectors = TypedNodeMat->GetNumParameters(EMaterialParameterType::Vector);
-		SurfNode->Vectors.SetNum(NumVectors);
-		for (int32 VectorIndex = 0; VectorIndex < NumVectors; ++VectorIndex)
-		{
-			const UEdGraphPin* VectorPin = TypedNodeMat->GetParameterPin(EMaterialParameterType::Vector, VectorIndex);
-			bool bVectorPinConnected = VectorPin && FollowInputPin(*VectorPin);
-
-			FString VectorName = TypedNodeMat->GetParameterName(EMaterialParameterType::Vector, VectorIndex).ToString();
-			FNodeMaterialParameterId VectorId = TypedNodeMat->GetParameterId(EMaterialParameterType::Vector, VectorIndex);
-
-			GenerationContext.CurrentMaterialTableParameter = VectorName;
-			GenerationContext.CurrentMaterialParameterId = VectorId.ParameterId.ToString();
-
-			if (bVectorPinConnected)
-			{				
-				if (const UEdGraphPin* ConnectedPin = FollowInputPin(*VectorPin))
-				{
-					UE::Mutable::Private::Ptr<UE::Mutable::Private::NodeColour> ColorNode = GenerateMutableSourceColor(ConnectedPin, GenerationContext);
-
-					const int32 LayerIndex = TypedNodeMat->GetParameterLayerIndex(EMaterialParameterType::Vector, VectorIndex);
-					SurfNode->Vectors[VectorIndex].ParameterInfo.ParameterName = VectorName;
-					SurfNode->Vectors[VectorIndex].ParameterInfo.MaterialLayerIndex = LayerIndex;
-					SurfNode->Vectors[VectorIndex].Vector = ColorNode;
-				}
-			}
-		}
-
-		const int32 NumScalar = TypedNodeMat->GetNumParameters(EMaterialParameterType::Scalar);
-		SurfNode->Scalars.SetNum(NumScalar);
-		for (int32 ScalarIndex = 0; ScalarIndex < NumScalar; ++ScalarIndex)
-		{
-			const UEdGraphPin* ScalarPin = TypedNodeMat->GetParameterPin(EMaterialParameterType::Scalar, ScalarIndex);
-			bool bScalarPinConnected = ScalarPin && FollowInputPin(*ScalarPin);
-
-			FString ScalarName = TypedNodeMat->GetParameterName(EMaterialParameterType::Scalar, ScalarIndex).ToString();
-			FNodeMaterialParameterId ScalarId = TypedNodeMat->GetParameterId(EMaterialParameterType::Scalar, ScalarIndex);
-
-			GenerationContext.CurrentMaterialTableParameter = ScalarName;
-			GenerationContext.CurrentMaterialParameterId = ScalarId.ParameterId.ToString();
-
-			if (bScalarPinConnected)
-			{
-				if (const UEdGraphPin* ConnectedPin = FollowInputPin(*ScalarPin))
-				{
-					UE::Mutable::Private::Ptr<UE::Mutable::Private::NodeScalar> ScalarNode = GenerateMutableSourceFloat(ConnectedPin, GenerationContext);
-
-					const int32 LayerIndex = TypedNodeMat->GetParameterLayerIndex(EMaterialParameterType::Scalar, ScalarIndex);
-					SurfNode->Scalars[ScalarIndex].Scalar = ScalarNode;
-					SurfNode->Scalars[ScalarIndex].ParameterInfo.ParameterName = ScalarName;
-					SurfNode->Scalars[ScalarIndex].ParameterInfo.MaterialLayerIndex = LayerIndex;
-				}
-			}
-		}
+		// Generate the material of this surface. Mesh section nodes contain their own material.
+		FGenerationMaterialOptions Options;
+		Options.SurfaceNode = SurfNode;
+		SurfNode->Material = GenerateMutableSourceMaterial(Pin, GenerationContext, Options);
 
+		// Add the tags that this mesh section enables.
 		for (const FString& Tag : TypedNodeMat->GetEnableTags(&GenerationContext.MacroNodesStack))
 		{
 			SurfNode->Tags.AddUnique(Tag);
 		}
 
-		SurfNode->Tags.AddUnique( TypedNodeMat->GetInternalTag() );
+		SurfNode->Tags.AddUnique(TypedNodeMat->GetInternalTag());
 	}
 
 	else if (const UCustomizableObjectNodeMaterialVariation* TypedNodeVar = Cast<UCustomizableObjectNodeMaterialVariation>(Node))
@@ -810,5 +388,4 @@ UE::Mutable::Private::Ptr<UE::Mutable::Private::NodeSurface> GenerateMutableSour
 	return Result;
 }
 
-
 #undef LOCTEXT_NAMESPACE
diff --git a/Engine/Plugins/Mutable/Source/CustomizableObjectEditor/Private/MuCOE/Nodes/CONodeMaterialBreak.cpp b/Engine/Plugins/Mutable/Source/CustomizableObjectEditor/Private/MuCOE/Nodes/CONodeMaterialBreak.cpp
--- a/Engine/Plugins/Mutable/Source/CustomizableObjectEditor/Private/MuCOE/Nodes/CONodeMaterialBreak.cpp
+++ b/Engine/Plugins/Mutable/Source/CustomizableObjectEditor/Private/MuCOE/Nodes/CONodeMaterialBreak.cpp
@@ -59,6 +59,15 @@ void UCONodeMaterialBreak::AllocateDefaultPins(UCustomizableObjectNodeRemapPins*
 		{
 			// Use the same pin data as the old pin since the new one is a copy
 			UEdGraphPin* NewPin = CustomCreatePin(EGPD_Output, PinData->VariableData.Type, PinData->VariableData.Name, PinData);
+
+			if (PinData->LayerIndex == 0)
+			{
+				NewPin->PinFriendlyName = FText::Format(LOCTEXT("BackgroundLayerText", "{0} - Background"), FText::FromName(PinData->VariableData.Name));
+			}
+			else if (PinData->LayerIndex > 0)
+			{
+				NewPin->PinFriendlyName = FText::Format(LOCTEXT("LayerText", "{0} - Layer {1}"), FText::FromName(PinData->VariableData.Name), FText::FromString(FString::FromInt(PinData->LayerIndex)));
+			}
 		}
 	}
 }
@@ -163,12 +172,13 @@ void UCONodeMaterialBreak::AddPinFromUI()
 
 FName UCONodeMaterialBreak::GetPinParameterName(const UEdGraphPin& Pin) const
 {
-	if (const UCONodeMaterialBreakPinData* PinData = Cast<UCONodeMaterialBreakPinData>(GetPinData(Pin)))
-	{
-		return PinData->VariableData.Name;
-	}
+	return GetPinData<UCONodeMaterialBreakPinData>(Pin).VariableData.Name;
+}
+
 
-	return NAME_None;
+int32 UCONodeMaterialBreak::GetPinParameterLayerIndex(const UEdGraphPin& Pin) const
+{
+	return GetPinData<UCONodeMaterialBreakPinData>(Pin).LayerIndex;
 }
 
 
diff --git a/Engine/Plugins/Mutable/Source/CustomizableObjectEditor/Private/MuCOE/Nodes/CONodeMaterialBreak.h b/Engine/Plugins/Mutable/Source/CustomizableObjectEditor/Private/MuCOE/Nodes/CONodeMaterialBreak.h
--- a/Engine/Plugins/Mutable/Source/CustomizableObjectEditor/Private/MuCOE/Nodes/CONodeMaterialBreak.h
+++ b/Engine/Plugins/Mutable/Source/CustomizableObjectEditor/Private/MuCOE/Nodes/CONodeMaterialBreak.h
@@ -34,6 +34,11 @@ public:
 	UPROPERTY(EditAnywhere, Category = NoCategory)
 	FCOVariable VariableData;
 
+	/** Layer index of this parameter inside a Layered Material.
+	-1 if the material does not contain layers.*/
+	UPROPERTY(EditAnywhere, Category = NoCategory)
+	int32 LayerIndex = -1;
+
 };
 
 
@@ -62,6 +67,7 @@ public:
 
 	// Own interface
 	FName GetPinParameterName(const UEdGraphPin& Pin) const;
+	int32 GetPinParameterLayerIndex(const UEdGraphPin& Pin) const;
 
 	UPROPERTY()
 	FEdGraphPinReference MaterialPinRef;
diff --git a/Engine/Plugins/Mutable/Source/CustomizableObjectEditor/Private/MuCOE/Nodes/CONodeMaterialModify.cpp b/Engine/Plugins/Mutable/Source/CustomizableObjectEditor/Private/MuCOE/Nodes/CONodeMaterialModify.cpp
--- a/Engine/Plugins/Mutable/Source/CustomizableObjectEditor/Private/MuCOE/Nodes/CONodeMaterialModify.cpp
+++ b/Engine/Plugins/Mutable/Source/CustomizableObjectEditor/Private/MuCOE/Nodes/CONodeMaterialModify.cpp
@@ -73,6 +73,15 @@ void UCONodeMaterialModify::AllocateDefaultPins(UCustomizableObjectNodeRemapPins
 			{
 				// Use the same pin data as the old pin since the new one is a copy
 				UEdGraphPin* NewPin = CustomCreatePin(EGPD_Input, PinData->VariableData.Type, PinData->VariableData.Name, PinData);
+
+				if (PinData->LayerIndex == 0)
+				{
+					NewPin->PinFriendlyName = FText::Format(LOCTEXT("BackgroundLayerText","{0} - Background"), FText::FromName(PinData->VariableData.Name));
+				}
+				else if (PinData->LayerIndex > 0)
+				{
+					NewPin->PinFriendlyName = FText::Format(LOCTEXT("LayerText", "{0} - Layer {1}"), FText::FromName(PinData->VariableData.Name), FText::FromString(FString::FromInt(PinData->LayerIndex)));
+				}
 			}
 		}
 	}
@@ -149,14 +158,21 @@ void UCONodeMaterialModify::AddPinFromUI()
 }
 
 
+bool UCONodeMaterialModify::CanPinBeRemoved(const UEdGraphPin& Pin) const
+{
+	return Pin.Direction == EGPD_Input && !Pin.LinkedTo.Num() && !Pin.bOrphanedPin && &Pin != MaterialPinRef.Get();
+}
+
+
 FName UCONodeMaterialModify::GetPinParameterName(const UEdGraphPin& Pin) const
 {
-	if (const UCONodeMaterialModifyPinData* PinData = Cast<UCONodeMaterialModifyPinData>(GetPinData(Pin)))
-	{
-		return PinData->VariableData.Name;
-	}
+	return GetPinData<UCONodeMaterialModifyPinData>(Pin).VariableData.Name;
+}
 
-	return NAME_None;
+
+int32 UCONodeMaterialModify::GetPinParameterLayerIndex(const UEdGraphPin& Pin) const
+{
+	return GetPinData<UCONodeMaterialModifyPinData>(Pin).LayerIndex;
 }
 
 
diff --git a/Engine/Plugins/Mutable/Source/CustomizableObjectEditor/Private/MuCOE/Nodes/CONodeMaterialModify.h b/Engine/Plugins/Mutable/Source/CustomizableObjectEditor/Private/MuCOE/Nodes/CONodeMaterialModify.h
--- a/Engine/Plugins/Mutable/Source/CustomizableObjectEditor/Private/MuCOE/Nodes/CONodeMaterialModify.h
+++ b/Engine/Plugins/Mutable/Source/CustomizableObjectEditor/Private/MuCOE/Nodes/CONodeMaterialModify.h
@@ -33,6 +33,11 @@ public:
 
 	UPROPERTY(EditAnywhere, Category = NoCategory)
 	FCOVariable VariableData;
+	
+	/** Layer index of this parameter inside a Layered Material.
+	-1 if the material does not contain layers.*/
+	UPROPERTY(EditAnywhere, Category = NoCategory)
+	int32 LayerIndex = -1;
 };
 
 UCLASS()
@@ -53,12 +58,14 @@ public:
 	virtual bool HasPinViewer() const override;
 	virtual EAddPinNodeButtonLocation GetAddPinButtonNodeSide() const override;
 	virtual void AddPinFromUI() override;
+	virtual bool CanPinBeRemoved(const UEdGraphPin& Pin) const override;
 	virtual UCustomizableObjectNodeRemapPins* CreateRemapPinsDefault() const override;
 	virtual TArray<FName> GetAllowedPinViewerCreationTypes() const override;
 	virtual bool IsExperimental() const override;
 
 	// Own interface
 	FName GetPinParameterName(const UEdGraphPin& Pin) const;
+	int32 GetPinParameterLayerIndex(const UEdGraphPin& Pin) const;
 
 public:
 	
diff --git a/Engine/Plugins/Mutable/Source/CustomizableObjectEditor/Private/MuCOE/Nodes/CustomizableObjectNodeMaterial.cpp b/Engine/Plugins/Mutable/Source/CustomizableObjectEditor/Private/MuCOE/Nodes/CustomizableObjectNodeMaterial.cpp
--- a/Engine/Plugins/Mutable/Source/CustomizableObjectEditor/Private/MuCOE/Nodes/CustomizableObjectNodeMaterial.cpp
+++ b/Engine/Plugins/Mutable/Source/CustomizableObjectEditor/Private/MuCOE/Nodes/CustomizableObjectNodeMaterial.cpp
@@ -987,9 +987,15 @@ UTexture2D* UCustomizableObjectNodeMaterial::GetImageReferenceTexture(const int3
 UTexture2D* UCustomizableObjectNodeMaterial::GetImageValue(const int32 ImageIndex) const
 {
 	FName TextureName = GetParameterName(EMaterialParameterType::Texture, ImageIndex);
+	int32 LayerIndex = GetParameterLayerIndex(EMaterialParameterType::Texture, ImageIndex);
 	
+	FHashedMaterialParameterInfo TextureInfo;
+	TextureInfo.Name = FScriptName(TextureName);
+	TextureInfo.Index = LayerIndex;
+	TextureInfo.Association = LayerIndex == INDEX_NONE ? GlobalParameter : LayerParameter; // TODO: Blend Parameters
+
 	UTexture* Texture;
-	Material->GetTextureParameterValue(TextureName, Texture);
+	Material->GetTextureParameterValue(TextureInfo, Texture);
 
 	return Cast<UTexture2D>(Texture);
 }
diff --git a/Engine/Plugins/Mutable/Source/CustomizableObjectEditor/Private/MuCOE/SMutableCodeViewer.cpp b/Engine/Plugins/Mutable/Source/CustomizableObjectEditor/Private/MuCOE/SMutableCodeViewer.cpp
--- a/Engine/Plugins/Mutable/Source/CustomizableObjectEditor/Private/MuCOE/SMutableCodeViewer.cpp
+++ b/Engine/Plugins/Mutable/Source/CustomizableObjectEditor/Private/MuCOE/SMutableCodeViewer.cpp
@@ -1691,9 +1691,6 @@ void SMutableCodeViewer::GenerateElementRecursive(const int32& InStateIndex, UE:
 	}
 
 	case UE::Mutable::Private::EOpType::IN_ADDMESH:
-	case UE::Mutable::Private::EOpType::IN_ADDIMAGE:
-	case UE::Mutable::Private::EOpType::IN_ADDVECTOR:
-	case UE::Mutable::Private::EOpType::IN_ADDSCALAR:
 	case UE::Mutable::Private::EOpType::IN_ADDSTRING:
 	case UE::Mutable::Private::EOpType::IN_ADDCOMPONENT:
 	case UE::Mutable::Private::EOpType::IN_ADDSURFACE:
@@ -2614,18 +2611,6 @@ bool SMutableCodeViewer::IsConstantResourceUsedByOperation(const int32 IndexOnCo
 				{
 					bResourceLocated = IndexOnConstantsArray == InProgram.GetOpArgs<UE::Mutable::Private::OP::InstanceAddArgs>(OperationAddress).name;
 				}
-				else if(OperationType == UE::Mutable::Private::EOpType::IN_ADDIMAGE)
-				{
-					bResourceLocated = IndexOnConstantsArray == InProgram.GetOpArgs<UE::Mutable::Private::OP::InstanceAddArgs>(OperationAddress).name;
-				}
-				else if (OperationType == UE::Mutable::Private::EOpType::IN_ADDVECTOR)
-				{
-					bResourceLocated = IndexOnConstantsArray == InProgram.GetOpArgs<UE::Mutable::Private::OP::InstanceAddArgs>(OperationAddress).name;
-				}
-				else if (OperationType == UE::Mutable::Private::EOpType::IN_ADDSCALAR)
-				{
-					bResourceLocated = IndexOnConstantsArray == InProgram.GetOpArgs<UE::Mutable::Private::OP::InstanceAddArgs>(OperationAddress).name;
-				}
 				else if (OperationType == UE::Mutable::Private::EOpType::IN_ADDCOMPONENT)
 				{
 					bResourceLocated = IndexOnConstantsArray == InProgram.GetOpArgs<UE::Mutable::Private::OP::InstanceAddArgs>(OperationAddress).name;
@@ -3321,17 +3306,6 @@ void FMutableCodeTreeElement::GenerateLabelText()
 		break;
 	}
 
-	case UE::Mutable::Private::EOpType::IN_ADDIMAGE:
-	{
-		UE::Mutable::Private::OP::InstanceAddArgs Args = Program.GetOpArgs<UE::Mutable::Private::OP::InstanceAddArgs>(MutableOperation);
-		if (Program.ConstantStrings.IsValidIndex(Args.name))
-		{
-			OpName += TEXT(" name: ");
-			OpName += Program.ConstantStrings[Args.name];
-		}
-		break;
-	}
-
 	default:
 		break;
 	}
diff --git a/Engine/Plugins/Mutable/Source/CustomizableObjectEditor/Private/MuCOE/SMutableGraphViewer.cpp b/Engine/Plugins/Mutable/Source/CustomizableObjectEditor/Private/MuCOE/SMutableGraphViewer.cpp
--- a/Engine/Plugins/Mutable/Source/CustomizableObjectEditor/Private/MuCOE/SMutableGraphViewer.cpp
+++ b/Engine/Plugins/Mutable/Source/CustomizableObjectEditor/Private/MuCOE/SMutableGraphViewer.cpp
@@ -43,6 +43,7 @@
 #include "MuT/NodeSurfaceSwitch.h"
 #include "MuT/NodeSurfaceVariation.h"
 #include "MuT/NodeLOD.h"
+#include "MuT/NodeMaterialModify.h"
 #include "MuT/NodeMeshConstant.h"
 #include "MuT/NodeMeshFormat.h"
 #include "MuT/NodeMeshFragment.h"
@@ -232,6 +233,7 @@ static const char* MutableNodeNames[] =
 	"MaterialParameter",
 	"MaterialExternal",
 	"MaterialSkeletalMeshBreak",
+	"MaterialModify",
 
 	"ImageMaterialBreak",
 	"ScalarMaterialBreak",
@@ -466,20 +468,7 @@ void SMutableGraphViewer::GetChildrenForInfo(TSharedPtr<FMutableGraphTreeElement
 		UE::Mutable::Private::NodeSurfaceNew* SurfaceNew = StaticCast<UE::Mutable::Private::NodeSurfaceNew*>(ParentNode);
 		AddChildFunc(SurfaceNew->Mesh.get(), TEXT("MESH"));
 
-		for (int32 l = 0; l < SurfaceNew->Images.Num(); ++l)
-		{
-			AddChildFunc(SurfaceNew->Images[l].Image.get(), FString::Printf(TEXT("IMAGE [%s]"), *SurfaceNew->Images[l].ParameterInfo.ParameterName));
-		}
-
-		for (int32 l = 0; l < SurfaceNew->Vectors.Num(); ++l)
-		{
-			AddChildFunc(SurfaceNew->Vectors[l].Vector.get(), FString::Printf(TEXT("VECTOR [%s]"), *SurfaceNew->Vectors[l].ParameterInfo.ParameterName));
-		}
-
-		for (int32 l = 0; l < SurfaceNew->Scalars.Num(); ++l)
-		{
-			AddChildFunc(SurfaceNew->Scalars[l].Scalar.get(), FString::Printf(TEXT("SCALAR [%s]"), *SurfaceNew->Scalars[l].ParameterInfo.ParameterName));
-		}
+		AddChildFunc(SurfaceNew->Material.get(), TEXT("MATERIAL"));
 
 		for (int32 l = 0; l < SurfaceNew->Strings.Num(); ++l)
 		{
@@ -888,6 +877,26 @@ void SMutableGraphViewer::GetChildrenForInfo(TSharedPtr<FMutableGraphTreeElement
 		AddChildFunc(ScalarTableVar->Z.get(), TEXT("Z"));
 		AddChildFunc(ScalarTableVar->W.get(), TEXT("W"));
 	}
+
+	else if (ParentNode->GetType() == UE::Mutable::Private::NodeMaterialModify::GetStaticType())
+	{
+		UE::Mutable::Private::NodeMaterialModify* MaterialModify = StaticCast<UE::Mutable::Private::NodeMaterialModify*>(ParentNode);
+
+		for (const TPair<UE::Mutable::Private::FParameterKey, UE::Mutable::Private::FImageParameterData>& ImageData : MaterialModify->ImageParameters)
+		{
+			AddChildFunc(ImageData.Value.ImageNode.get(), FString::Printf(TEXT("IMAGE [%s]"), *ImageData.Key.ParameterName.ToString()));
+		}
+
+		for (const TPair<UE::Mutable::Private::FParameterKey, UE::Mutable::Private::Ptr<UE::Mutable::Private::NodeColour>>& ColorData : MaterialModify->ColorParameters)
+		{
+			AddChildFunc(ColorData.Value.get(), FString::Printf(TEXT("VECTOR [%s]"), *ColorData.Key.ParameterName.ToString()));
+		}
+
+		for (const TPair<UE::Mutable::Private::FParameterKey, UE::Mutable::Private::Ptr<UE::Mutable::Private::NodeScalar>>& ScalarData : MaterialModify->ScalarParameters)
+		{
+			AddChildFunc(ScalarData.Value.get(), FString::Printf(TEXT("SCALAR [%s]"), *ScalarData.Key.ParameterName.ToString()));
+		}
+	}
 	
 	else
 	{
diff --git a/Engine/Plugins/Mutable/Source/CustomizableObjectEditor/Private/MuCOE/SMutableInstanceViewer.cpp b/Engine/Plugins/Mutable/Source/CustomizableObjectEditor/Private/MuCOE/SMutableInstanceViewer.cpp
--- a/Engine/Plugins/Mutable/Source/CustomizableObjectEditor/Private/MuCOE/SMutableInstanceViewer.cpp
+++ b/Engine/Plugins/Mutable/Source/CustomizableObjectEditor/Private/MuCOE/SMutableInstanceViewer.cpp
@@ -342,12 +342,13 @@ void SMutableInstanceViewer::SetInstance(const TSharedPtr<const UE::Mutable::Pri
 					{
 						TSharedPtr<FMutableInstanceViewerSurfaceElement> NewSurface = MakeShareable(new FMutableInstanceViewerSurfaceElement());
 
-						NewSurface->Id = FText::AsNumber(MutableInstance->GetSurfaceId(ComponentIndex, LODIndex, SurfaceIndex));
-						NewSurface->CustomId = FText::AsNumber(MutableInstance->GetSurfaceCustomId(ComponentIndex, LODIndex, SurfaceIndex));
-						NewSurface->ImageCount = FText::AsNumber(MutableInstance->GetImageCount(ComponentIndex, LODIndex, SurfaceIndex));
-						NewSurface->VectorCount = FText::AsNumber(MutableInstance->GetVectorCount(ComponentIndex, LODIndex, SurfaceIndex));
-						NewSurface->ScalarCount = FText::AsNumber(MutableInstance->GetScalarCount(ComponentIndex, LODIndex, SurfaceIndex));
-						NewSurface->StringCount = FText::AsNumber(MutableInstance->GetStringCount(ComponentIndex, LODIndex, SurfaceIndex));
+						//TODO(Max - GMT): fix this
+						//NewSurface->Id = FText::AsNumber(MutableInstance->GetSurfaceId(ComponentIndex, LODIndex, SurfaceIndex));
+						//NewSurface->CustomId = FText::AsNumber(MutableInstance->GetSurfaceCustomId(ComponentIndex, LODIndex, SurfaceIndex));
+						//NewSurface->ImageCount = FText::AsNumber(MutableInstance->GetImageCount(ComponentIndex, LODIndex, SurfaceIndex));
+						//NewSurface->VectorCount = FText::AsNumber(MutableInstance->GetVectorCount(ComponentIndex, LODIndex, SurfaceIndex));
+						//NewSurface->ScalarCount = FText::AsNumber(MutableInstance->GetScalarCount(ComponentIndex, LODIndex, SurfaceIndex));
+						//NewSurface->StringCount = FText::AsNumber(MutableInstance->GetStringCount(ComponentIndex, LODIndex, SurfaceIndex));
 
 						SurfaceArray->Add(NewSurface);
 					}
diff --git a/Engine/Plugins/Mutable/Source/MutableRuntime/Internal/MuR/CodeVisitor.h b/Engine/Plugins/Mutable/Source/MutableRuntime/Internal/MuR/CodeVisitor.h
--- a/Engine/Plugins/Mutable/Source/MutableRuntime/Internal/MuR/CodeVisitor.h
+++ b/Engine/Plugins/Mutable/Source/MutableRuntime/Internal/MuR/CodeVisitor.h
@@ -19,8 +19,7 @@ namespace UE::Mutable::Private
     /** Decide what operations are an "add resource" since they are handled differently sometimes. */
     inline bool VisitorIsAddResource(const EOpType& type)
     {
-        return  type == EOpType::IN_ADDIMAGE ||
-                type == EOpType::IN_ADDMESH;
+        return  type == EOpType::IN_ADDMESH || type == EOpType::IN_ADDMATERIAL;
     }
 
 
@@ -420,25 +419,24 @@ namespace UE::Mutable::Private
                 break;
             }
 
+			case EOpType::IN_ADDMATERIAL:
+			{
+				OP::InstanceAddMaterialArgs Args = Program.GetOpArgs<OP::InstanceAddMaterialArgs>(Address);
+				bRecurse = false;
+				const uint8* Data = Program.GetOpArgsPointer(Address);
 
-            case EOpType::IN_ADDIMAGE:
-            {
-				OP::InstanceAddArgs Args = Program.GetOpArgs<OP::InstanceAddArgs>(Address);
-
-                bRecurse = false;
-
-                PARENT::RecurseWithCurrentState(Args.instance);
-
-                STATE NewState = PARENT::GetCurrentState();
-				NewState.m_underResourceCount=1;
+				PARENT::RecurseWithCurrentState(Args.Instance);
+				STATE NewState = PARENT::GetCurrentState();
+				NewState.m_underResourceCount = 1;
+				OP::ADDRESS MaterialAddress = Args.Material;
+				
+				if (MaterialAddress)
+				{
+					PARENT::RecurseWithState(MaterialAddress, NewState);
+				}
 
-                OP::ADDRESS ImageAddress = Args.value;
-                if (ImageAddress)
-                {
-                    PARENT::RecurseWithState(ImageAddress, NewState);
-                }
-                break;
-            }
+				break;
+			}
 
             default:
                 break;
diff --git a/Engine/Plugins/Mutable/Source/MutableRuntime/Internal/MuR/Instance.h b/Engine/Plugins/Mutable/Source/MutableRuntime/Internal/MuR/Instance.h
--- a/Engine/Plugins/Mutable/Source/MutableRuntime/Internal/MuR/Instance.h
+++ b/Engine/Plugins/Mutable/Source/MutableRuntime/Internal/MuR/Instance.h
@@ -24,33 +24,6 @@ namespace UE::Mutable::Private
 
 		FMaterialId MaterialId;
 
-		struct FInstanceImage
-		{
-			FImageId Id;
-			FName Name;
-			int32 LayerIndex;
-		};
-
-		TArray<FInstanceImage> Images;
-
-		struct FInstanceVector
-		{
-			FVector4f Value;
-			FName Name;
-			int32 LayerIndex;
-		};
-
-		TArray<FInstanceVector> Vectors;
-
-		struct FInstanceScalar
-		{
-			float Value;
-			FName Name;
-			int32 LayerIndex;
-		};
-
-		TArray<FInstanceScalar> Scalars;
-
 		struct FInstanceString
 		{
 			FString Value;
@@ -166,87 +139,6 @@ namespace UE::Mutable::Private
 		//! \param LODIndex Index of the level of detail, from 0 to GetLODCount()-1
   		UE_API FMeshId GetMeshId( int32 ComponentIndex, int32 LODIndex ) const;
 
-		//! Get the number of images in a component
-		//! \param ComponentIndex Index of the component, from 0 to GetComponentCount()-1
-		//! \param LODIndex Index of the level of detail, from 0 to GetLODCount()-1
-		//! \param SurfaceIndex Index of the surface, from 0 to GetSurfaceCount()-1
-        UE_API int32 GetImageCount( int32 ComponentIndex, int32 LODIndex, int32 SurfaceIndex) const;
-
-        //! Get an image resource id from a component
-		//! \param ComponentIndex Index of the component, from 0 to GetComponentCount()-1
-		//! \param LODIndex Index of the level of detail, from 0 to GetLODCount()-1
-		//! \param SurfaceIndex Index of the surface, from 0 to GetSurfaceCount()-1
-        //! \param ImageIndex Index of the image, from 0 to GetImageCount()-1
- 		UE_API FImageId GetImageId( int32 ComponentIndex, int32 LODIndex, int32 SurfaceIndex, int32 ImageIndex) const;
-
-		//! Get the name of an image in a component
-		//! \param ComponentIndex Index of the component, from 0 to GetComponentCount()-1
-		//! \param LODIndex Index of the level of detail, from 0 to GetLODCount()-1
-		//! \param SurfaceIndex Index of the surface, from 0 to GetSurfaceCount()-1
-        //! \param ImageIndex Index of the image, from 0 to GetImageCount()-1
-		UE_API FName GetImageName( int32 ComponentIndex, int32 LODIndex, int32 SurfaceIndex, int32 ImageIndex ) const;
-
-		//! Get the Layer Index of an image in a Material
-		//! \param ComponentIndex Index of the component, from 0 to GetComponentCount()-1
-		//! \param LODIndex Index of the level of detail, from 0 to GetLODCount()-1
-		//! \param SurfaceIndex Index of the surface, from 0 to GetSurfaceCount()-1
-        //! \param ImageIndex Index of the image, from 0 to GetImageCount()-1
-		UE_API int32 GetImageLayerIndex(int32 ComponentIndex, int32 LODIndex, int32 SurfaceIndex, int32 ImageIndex) const;
-
-		//! Get the number of vectors in a component
-		//! \param ComponentIndex Index of the component, from 0 to GetComponentCount()-1
-		//! \param LODIndex Index of the level of detail, from 0 to GetLODCount()-1
-		//! \param SurfaceIndex Index of the surface, from 0 to GetSurfaceCount()-1
-        UE_API int32 GetVectorCount( int32 ComponentIndex, int32 LODIndex, int32 SurfaceIndex) const;
-
-		//! Get a vector from a component
-		//! \param ComponentIndex Index of the component, from 0 to GetComponentCount()-1
-		//! \param LODIndex Index of the level of detail, from 0 to GetLODCount()-1
-		//! \param SurfaceIndex Index of the surface, from 0 to GetSurfaceCount()-1
-        //! \param VectorIndex Index of the vector, from 0 to GetVectorCount()-1
-        UE_API FVector4f GetVector( int32 ComponentIndex, int32 LODIndex, int32 SurfaceIndex, int32 VectorIndex) const;
-
-		//! Get the name of a vector in a component
-		//! \param ComponentIndex Index of the component, from 0 to GetComponentCount()-1
-		//! \param LODIndex Index of the level of detail, from 0 to GetLODCount()-1
-		//! \param SurfaceIndex Index of the surface, from 0 to GetSurfaceCount()-1
-        //! \param VectorIndex Index of the vector, from 0 to GetVectorCount()-1
-		UE_API FName GetVectorName( int32 ComponentIndex, int32 LODIndex, int32 SurfaceIndex, int32 VectorIndex ) const;
-
-		//! Get the Layer Index of a vector in a Material
-		//! \param ComponentIndex Index of the component, from 0 to GetComponentCount()-1
-		//! \param LODIndex Index of the level of detail, from 0 to GetLODCount()-1
-		//! \param SurfaceIndex Index of the surface, from 0 to GetSurfaceCount()-1
-		//! \param VectorIndex Index of the vector, from 0 to GetVectorCount()-1
-		UE_API int32 GetVectorLayerIndex(int32 ComponentIndex, int32 LODIndex, int32 SurfaceIndex, int32 VectorIndex) const;
-
-        //! Get the number of scalar values in a component
-        //! \param ComponentIndex Index of the component, from 0 to GetComponentCount()-1
-		//! \param LODIndex Index of the level of detail, from 0 to GetLODCount()-1
-		//! \param SurfaceIndex Index of the surface, from 0 to GetSurfaceCount()-1
-        UE_API int32 GetScalarCount( int32 ComponentIndex, int32 LODIndex, int32 SurfaceIndex) const;
-
-        //! Get a scalar value from a component
-        //! \param ComponentIndex Index of the component, from 0 to GetComponentCount()-1
-		//! \param LODIndex Index of the level of detail, from 0 to GetLODCount()-1
-		//! \param SurfaceIndex Index of the surface, from 0 to GetSurfaceCount()-1
-        //! \param ScalarIndex Index of the scalar, from 0 to GetScalarCount()-1
-        UE_API float GetScalar( int32 ComponentIndex, int32 LODIndex, int32 SurfaceIndex, int32 ScalarIndex) const;
-
-        //! Get the name of a scalar from a component
-        //! \param ComponentIndex Index of the component, from 0 to GetComponentCount()-1
-		//! \param LODIndex Index of the level of detail, from 0 to GetLODCount()-1
-		//! \param SurfaceIndex Index of the surface, from 0 to GetSurfaceCount()-1
-        //! \param ScalarIndex Index of the scalar, from 0 to GetScalarCount()-1
-		UE_API FName GetScalarName( int32 ComponentIndex, int32 LODIndex, int32 SurfaceIndex, int32 ScalarIndex ) const;
-
-		//! Get the Layer Index of a scalar in a Material
-		//! \param ComponentIndex Index of the component, from 0 to GetComponentCount()-1
-		//! \param LODIndex Index of the level of detail, from 0 to GetLODCount()-1
-		//! \param SurfaceIndex Index of the surface, from 0 to GetSurfaceCount()-1
-		//! \param ScalarIndex Index of the scalar, from 0 to GetScalarCount()-1
-		UE_API int32 GetScalarLayerIndex(int32 ComponentIndex, int32 LODIndex, int32 SurfaceIndex, int32 ScalarIndex) const;
-
         //! Get the number of string values in a component
         //! \param ComponentIndex Index of the component, from 0 to GetComponentCount()-1
 		//! \param LODIndex Index of the level of detail, from 0 to GetLODCount()-1
@@ -300,12 +192,6 @@ namespace UE::Mutable::Private
 
 		void SetSurfaceName(int32 ComponentIndex, int32 LODIndex, int32 SurfaceIndex, FName Name);
 
-		int32 AddImage(int32 ComponentIndex, int32 LODIndex, int32 SurfaceIndex, const FImageId& Image, FName Name, int32 LayerIndex);
-
-		int32 AddVector(int32 ComponentIndex, int32 LODIndex, int32 SurfaceIndex, const FVector4f& Vector, FName Name, int32 LayerIndex);
-
-		int32 AddScalar(int32 ComponentIndex, int32 LODIndex, int32 SurfaceIndex, float Scalar, FName Name, int32 LayerIndex);
-
 		int32 AddString(int32 ComponentIndex, int32 LODIndex, int32 SurfaceIndex, const FString& Value, FName Name);
 
 		void SetMaterialId(int32 ComponentIndex, int32 LODIndex, int32 SurfaceIndex, const FMaterialId& MaterialId);
diff --git a/Engine/Plugins/Mutable/Source/MutableRuntime/Internal/MuR/Material.h b/Engine/Plugins/Mutable/Source/MutableRuntime/Internal/MuR/Material.h
--- a/Engine/Plugins/Mutable/Source/MutableRuntime/Internal/MuR/Material.h
+++ b/Engine/Plugins/Mutable/Source/MutableRuntime/Internal/MuR/Material.h
@@ -4,6 +4,7 @@
 
 #include "MuR/Ptr.h"
 #include "MuR/RefCounted.h"
+#include "MuR/ResourceID.h"
 #include "MuR/Types.h"
 #include "MuR/Operations.h"
 #include "MuR/TVariant.h"
@@ -11,6 +12,7 @@
 #include "UObject/StrongObjectPtr.h"
 #include "Materials/MaterialInterface.h"
 #include "MuR/PassthroughObject.h"
+#include "Templates/TypeHash.h"
 
 #define UE_API MUTABLERUNTIME_API
 
@@ -18,6 +20,26 @@ namespace UE::Mutable::Private
 {
 	class FImage;
 
+	struct FParameterKey
+	{
+		FName ParameterName;
+		int32 LayerIndex = -1;
+
+		inline void Serialise(FOutputArchive& arch) const;
+		inline void Unserialise(FInputArchive& arch);
+
+		bool operator==(const FParameterKey& Other) const = default;
+	};
+
+
+	inline uint32 GetTypeHash(const FParameterKey& Key)
+	{
+		uint32 GuidHash = ::GetTypeHash(Key.LayerIndex);
+		GuidHash = HashCombineFast(GuidHash, GetTypeHash(Key.ParameterName));
+
+		return GuidHash;
+	}
+
     /** Material type resource */
 	class FMaterial : public FResource
 	{
@@ -47,9 +69,20 @@ namespace UE::Mutable::Private
 		/** Passthrough Id or TStrongObjectPtr. */
 		TPassthroughObjectPtr<UMaterialInterface> PassthroughObject;
 
-		TMap<FName, TVariant<OP::ADDRESS, TSharedPtr<const FImage>>> ImageParameters;
-		TMap<FName, FVector4f> ColorParameters;
-		TMap<FName, float> ScalarParameters;
+		struct FImageParameterData 
+		{
+			TVariant<OP::ADDRESS, TSharedPtr<const FImage>> ImageParameter;
+			int32 ImagePropertyIndex;
+
+			inline void Serialise(FOutputArchive& arch) const;
+			inline void Unserialise(FInputArchive& arch);
+
+			bool operator==(const FImageParameterData& Other) const = default;
+		};
+
+		TMap<FParameterKey, FImageParameterData> ImageParameters;
+		TMap<FParameterKey, FVector4f> ColorParameters;
+		TMap<FParameterKey, float> ScalarParameters;
 	};
 }
 
diff --git a/Engine/Plugins/Mutable/Source/MutableRuntime/Internal/MuR/Operations.h b/Engine/Plugins/Mutable/Source/MutableRuntime/Internal/MuR/Operations.h
--- a/Engine/Plugins/Mutable/Source/MutableRuntime/Internal/MuR/Operations.h
+++ b/Engine/Plugins/Mutable/Source/MutableRuntime/Internal/MuR/Operations.h
@@ -82,6 +82,10 @@ namespace UE::Mutable::Private
 		ED_SWITCH,
 		MI_SWITCH,
 
+		//-----------------------------------------------------------------------------------------
+	   // Material operations
+	   //-----------------------------------------------------------------------------------------
+
 		//! Selects a parameter value from a Material
 		SC_MATERIAL_BREAK,
 		CO_MATERIAL_BREAK,
@@ -90,6 +94,9 @@ namespace UE::Mutable::Private
 		//! Converts a texture of a material parameter into a texture parameter to process it at runtime.
 		IM_PARAMETER_FROM_MATERIAL,
 
+    	//! Modifies the parameter values of a material
+		MI_MODIFY,
+
         //-----------------------------------------------------------------------------------------
         // Boolean operations
         //-----------------------------------------------------------------------------------------
@@ -324,15 +331,6 @@ namespace UE::Mutable::Private
         //! Add a mesh to an instance
         IN_ADDMESH,
 
-        //! Add an image to an instance
-        IN_ADDIMAGE,
-
-        //! Add a vector to an instance
-        IN_ADDVECTOR,
-
-        //! Add a scalar to an instance
-        IN_ADDSCALAR,
-
         //! Add a string to an instance
         IN_ADDSTRING,
 
@@ -1112,7 +1110,6 @@ namespace UE::Mutable::Private
             uint32 ExternalId;
 
             ADDRESS name;
-			int32 LayerIndex; // Will be removed in UE-351597.
         };
 
 		struct InstanceAddExtensionDataArgs
@@ -1174,6 +1171,7 @@ namespace UE::Mutable::Private
 		{
 			ADDRESS Material;
 			STRING_ADDRESS ParameterName;
+			int32 LayerIndex;
 		};
 
 		struct MaterialBreakImageParameterArgs
@@ -1191,6 +1189,12 @@ namespace UE::Mutable::Private
     	{
     		ADDRESS ImageParameter;
     	};
+
+    	struct MaterialModifyArgs
+    	{
+    		ADDRESS Material;
+    		uint32 NumParameters;
+    	};
     };
 
 
diff --git a/Engine/Plugins/Mutable/Source/MutableRuntime/Private/MuR/CodeRunner.cpp b/Engine/Plugins/Mutable/Source/MutableRuntime/Private/MuR/CodeRunner.cpp
--- a/Engine/Plugins/Mutable/Source/MutableRuntime/Private/MuR/CodeRunner.cpp
+++ b/Engine/Plugins/Mutable/Source/MutableRuntime/Private/MuR/CodeRunner.cpp
@@ -526,98 +526,6 @@ namespace UE::Mutable::Private
 		EOpType type = Program.GetOpType(Item.At);
         switch (type)
         {
-
-        case EOpType::IN_ADDVECTOR:
-        {
-			OP::InstanceAddArgs Args = Program.GetOpArgs<OP::InstanceAddArgs>(Item.At);
-
-            switch (Item.Stage)
-            {
-            case 0:
-                    AddOp( FScheduledOp( Item.At, Item, 1),
-                           FScheduledOp( Args.instance, Item),
-                           FScheduledOp( Args.value, Item) );
-                break;
-
-            case 1:
-            {
-                TSharedPtr<const FInstance> Base = LoadInstance( FCacheAddress(Args.instance,Item) );
-				TSharedPtr<FInstance> Result;
-                if (!Base)
-                {
-                    Result = MakeShared<FInstance>();
-                }
-                else
-                {
-					Result = UE::Mutable::Private::CloneOrTakeOver<FInstance>(Base);
-                }
-
-                if ( Args.value )
-                {
-					FVector4f value = LoadColor( FCacheAddress(Args.value,Item) );
-
-                    OP::ADDRESS nameAd = Args.name;
-                    check(  nameAd < (uint32)Program.ConstantStrings.Num() );
-                    const FString& Name = Program.ConstantStrings[ nameAd ];
-
-                    Result->AddVector( 0, 0, 0, value, FName(Name), Args.LayerIndex );
-                }
-                StoreInstance( Item, Result );
-                break;
-            }
-
-            default:
-                check(false);
-            }
-
-            break;
-        }
-
-        case EOpType::IN_ADDSCALAR:
-        {
-			OP::InstanceAddArgs Args = Program.GetOpArgs<OP::InstanceAddArgs>(Item.At);
-            switch (Item.Stage)
-            {
-            case 0:
-                    AddOp( FScheduledOp( Item.At, Item, 1),
-                           FScheduledOp( Args.instance, Item),
-                           FScheduledOp( Args.value, Item) );
-                break;
-
-            case 1:
-            {
-                TSharedPtr<const FInstance> Base = LoadInstance( FCacheAddress(Args.instance,Item) );
-                TSharedPtr<FInstance> Result;
-                if (!Base)
-                {
-                    Result = MakeShared<FInstance>();
-                }
-                else
-                {
-                    Result = UE::Mutable::Private::CloneOrTakeOver<FInstance>(Base);
-                }
-
-                if ( Args.value )
-                {
-                    float value = LoadScalar( FCacheAddress(Args.value,Item) );
-
-                    OP::ADDRESS nameAd = Args.name;
-                    check(  nameAd < (uint32)Program.ConstantStrings.Num() );
-                    const FString& Name = Program.ConstantStrings[ nameAd ];
-
-                    Result->AddScalar( 0, 0, 0, value, FName(Name), Args.LayerIndex);
-                }
-                StoreInstance( Item, Result );
-                break;
-            }
-
-            default:
-                check(false);
-            }
-
-            break;
-        }
-
         case EOpType::IN_ADDSTRING:
         {
 			OP::InstanceAddArgs Args = Program.GetOpArgs<OP::InstanceAddArgs>( Item.At );
@@ -972,48 +880,6 @@ namespace UE::Mutable::Private
             break;
         }
 
-        case EOpType::IN_ADDIMAGE:
-        {
-			OP::InstanceAddArgs Args = Program.GetOpArgs<OP::InstanceAddArgs>(Item.At);
-            switch (Item.Stage)
-            {
-            case 0:
-				// We don't build the resources when building instance: just store ids for them.
-				AddOp( FScheduledOp( Item.At, Item, 1), FScheduledOp( Args.instance, Item) );
-                break;
-
-            case 1:
-            {
-                TSharedPtr<const FInstance> Base = LoadInstance( FCacheAddress(Args.instance,Item) );
-                TSharedPtr<FInstance> Result;
-                if (!Base)
-                {
-                    Result = MakeShared<FInstance>();
-                }
-                else
-                {
-					Result = UE::Mutable::Private::CloneOrTakeOver<FInstance>(Base);
-				}
-
-                if ( Args.value )
-                {
-					FImageId ImageId = FSystem::GetImageId(System->WorkingMemoryManager.CurrentImageIdRegistry, InModel, InParams, Args.value);
-					OP::ADDRESS NameAd = Args.name;
-					check(NameAd < (uint32)Program.ConstantStrings.Num());
-					const FString& Name = Program.ConstantStrings[NameAd];
-					Result->AddImage(0, 0, 0, ImageId, FName(Name), Args.LayerIndex);
-                }
-                StoreInstance( Item, Result );
-                break;
-            }
-
-            default:
-                check(false);
-            }
-
-            break;
-        }
-
         case EOpType::IN_ADDOVERLAYMATERIAL:
         	{
         		OP::InstanceAddMaterialArgs Args = Program.GetOpArgs<OP::InstanceAddMaterialArgs>(Item.At);
@@ -1117,6 +983,7 @@ namespace UE::Mutable::Private
         		case 1:
         			{
         				TSharedPtr<const FInstance> Base = LoadInstance(FCacheAddress(Args.Instance, Item));
+						TSharedPtr<const FMaterial> Material = LoadMaterial(FCacheAddress(Args.Material, Item));
 
         				TSharedPtr<FInstance> Result;
         				if (Base)
@@ -1131,7 +998,7 @@ namespace UE::Mutable::Private
         				if (Args.Material)
         				{
         					FMaterialId MaterialId = FSystem::GetMaterialId(System->WorkingMemoryManager.CurrentMaterialIdRegistry, InModel, InParams, Args.Material);
-        					Result->SetMaterialId(0, 0, 0, MaterialId);
+							Result->SetMaterialId(0, 0, 0, MaterialId);
         				}
 
         				StoreInstance(Item, Result);
@@ -5620,25 +5487,28 @@ namespace UE::Mutable::Private
 
 				check(Args.ParameterName < (uint32)Program.ConstantStrings.Num());
 				const FString& ParameterName = Program.ConstantStrings[Args.ParameterName];
-				const TVariant<OP::ADDRESS, TSharedPtr<const FImage>>* Image = Material->ImageParameters.Find(FName(ParameterName));
+				const FParameterKey ParameterKey = { FName(ParameterName), Args.LayerIndex };
+				const FMaterial::FImageParameterData* ImageParameterData = Material->ImageParameters.Find(ParameterKey);
 				
-				if (!Image)
+				if (!ImageParameterData)
 				{
 					StoreImage(Item, nullptr);
 				}
 				else
 				{
-					if (Image->IsType<OP::ADDRESS>())
+					const TVariant<OP::ADDRESS, TSharedPtr<const FImage>> ImageParameter = ImageParameterData->ImageParameter;
+
+					if (ImageParameter.IsType<OP::ADDRESS>())
 					{
 						FScheduledOp ItemNextStage(Item.At, Item, 2);
-						ItemNextStage.CustomState = Image->Get<OP::ADDRESS>();
+						ItemNextStage.CustomState = ImageParameter.Get<OP::ADDRESS>();
 
 						AddOp(ItemNextStage,
-							FScheduledOp(Image->Get<OP::ADDRESS>(), Item));
+							FScheduledOp(ImageParameter.Get<OP::ADDRESS>(), Item));
 					}
 					else
 					{
-						StoreImage(Item, Image->Get<TSharedPtr<const FImage>>());
+						StoreImage(Item, ImageParameter.Get<TSharedPtr<const FImage>>());
 					}
 				}
 				break;
@@ -5649,10 +5519,21 @@ namespace UE::Mutable::Private
 
 				check(Args.ParameterName < (uint32)Program.ConstantStrings.Num());
 				const FString& ParameterName = Program.ConstantStrings[Args.ParameterName];
-				const TVariant<OP::ADDRESS, TSharedPtr<const FImage>>* Image = Material->ImageParameters.Find(FName(ParameterName));
+				const FParameterKey ParameterKey = { FName(ParameterName), Args.LayerIndex };
+				const FMaterial::FImageParameterData* ImageParameterData = Material->ImageParameters.Find(ParameterKey);
+
+				if (!ImageParameterData)
+				{
+					StoreImage(Item, nullptr);
+				}
+				else
+				{
+					const TVariant<OP::ADDRESS, TSharedPtr<const FImage>>& ImageParameter = ImageParameterData->ImageParameter;
+					TSharedPtr<const FImage> pImage = LoadImage(FCacheAddress(ImageParameter.Get<OP::ADDRESS>(), Item));
+
+					StoreImage(Item, pImage);
+				}
 
-				TSharedPtr<const FImage> pImage = LoadImage(FCacheAddress(Image->Get<OP::ADDRESS>(), Item));
-				StoreImage(Item, pImage);
 				break;
 			}
 
@@ -6122,7 +6003,8 @@ namespace UE::Mutable::Private
 					
 				check(Args.ParameterName < (uint32)Program.ConstantStrings.Num());
 				const FString& ParameterName = Program.ConstantStrings[Args.ParameterName];
-				const float* Value = Material->ScalarParameters.Find(FName(ParameterName));
+				const FParameterKey ParameterKey = { FName(ParameterName), Args.LayerIndex };
+				const float* Value = Material->ScalarParameters.Find(ParameterKey);
 
 				StoreScalar(Item, Value ? *Value : 0.0f);
 
@@ -6437,7 +6319,8 @@ namespace UE::Mutable::Private
 
 				check(Args.ParameterName < (uint32)Program.ConstantStrings.Num());
 				const FString& ParameterName = Program.ConstantStrings[Args.ParameterName];
-				const FVector4f* Value = Material->ColorParameters.Find(FName(ParameterName));
+				const FParameterKey ParameterKey = { FName(ParameterName), Args.LayerIndex };
+				const FVector4f* Value = Material->ColorParameters.Find(ParameterKey);
 
 				StoreColor(Item, Value ? *Value : FVector4f::Zero());
 
@@ -6553,61 +6436,282 @@ namespace UE::Mutable::Private
  
 			if (MaterialInterface)
 			{
-				// Images
+				// Get Name indexes for each type of parameters
+				//Images
 				const TArray<uint32>& ImageParameterNames = Program.ConstantUInt32Lists[Args.ImageParameterNames];
 				const TArray<uint32>& ImageParameterOperations = Program.ConstantUInt32Lists[Args.ImageParameterAddress];
+				// Scalars
+				const TArray<uint32>& ScalarParameterNames = Program.ConstantUInt32Lists[Args.ScalarParameterNames];
+				// Vectors
+				const TArray<uint32>& ColorParameterNames = Program.ConstantUInt32Lists[Args.ColorParameterNames];
+
+				// Get Parameter layers from program bytecode
+				const uint8* Data = Program.GetOpArgsPointer(Item.At);
+				Data += sizeof(OP::MaterialParameterArgs);
+
+				// Get parameter values from the material. Order maters since we store layer indexes directly to the program (Scalar, Colors, Textures)
+				// Scalars
+				for (int32 ParameterIndex = 0; ParameterIndex < ScalarParameterNames.Num(); ++ParameterIndex)
+				{
+					// Get Scalar Parameter Name
+					const FString& ParameterName = Program.ConstantStrings[ScalarParameterNames[ParameterIndex]];
+
+					// Get Scalar Parameter Layer
+					int32 LayerIndex;
+					FMemory::Memcpy(&LayerIndex, Data, sizeof(int32));
+					Data += sizeof(int32);
 
+					// Construct Scalar parameter info
+					FHashedMaterialParameterInfo ScalarParameterInfo;
+					ScalarParameterInfo.Name = FScriptName(FName(ParameterName));
+					ScalarParameterInfo.Index = LayerIndex;
+					ScalarParameterInfo.Association = ScalarParameterInfo.Index == INDEX_NONE ? GlobalParameter : LayerParameter;
+
+					// Find Texture using the parameter info
+					float ScalarValue;
+					bool bParameterFound = MaterialInterface->GetScalarParameterValue(ScalarParameterInfo, ScalarValue);
+
+					if (bParameterFound)
+					{
+						Material->ScalarParameters.Add({ FName(ScalarParameterInfo.Name), ScalarParameterInfo.Index }, ScalarValue);
+					}
+				}
+
+				// Colors
+				for (int32 ParameterIndex = 0; ParameterIndex < ColorParameterNames.Num(); ++ParameterIndex)
+				{
+					// Get Color Parameter Name
+					const FString& ParameterName = Program.ConstantStrings[ColorParameterNames[ParameterIndex]];
+
+					// Get Color Parameter Layer
+					int32 LayerIndex;
+					FMemory::Memcpy(&LayerIndex, Data, sizeof(int32));
+					Data += sizeof(int32);
+
+					// Construct Scalar parameter info
+					FHashedMaterialParameterInfo ColorParameterInfo;
+					ColorParameterInfo.Name = FScriptName(FName(ParameterName));
+					ColorParameterInfo.Index = LayerIndex;
+					ColorParameterInfo.Association = ColorParameterInfo.Index == INDEX_NONE ? GlobalParameter : LayerParameter;
+
+					// Find Texture using the parameter info
+					FLinearColor VectorValue;
+					bool bParameterFound = MaterialInterface->GetVectorParameterValue(ColorParameterInfo, VectorValue);
+
+					if (bParameterFound)
+					{
+						Material->ColorParameters.Add({FName(ColorParameterInfo.Name), ColorParameterInfo.Index}, VectorValue);
+					}
+				}
+
+				// Textures
 				for (int32 ParameterIndex = 0; ParameterIndex < ImageParameterNames.Num(); ++ParameterIndex)
 				{
+					// Get Image Parameter Name
 					const FString& ParameterName = Program.ConstantStrings[ImageParameterNames[ParameterIndex]];
+					
+					// Get Image Parameter Layer
+					int32 LayerIndex;
+					FMemory::Memcpy(&LayerIndex, Data, sizeof(int32));
+					Data += sizeof(int32);
+
+					// Construct image parameter info
+					FHashedMaterialParameterInfo ImageParameterInfo;
+					ImageParameterInfo.Name = FScriptName(FName(ParameterName));
+					ImageParameterInfo.Index = LayerIndex;
+					ImageParameterInfo.Association = ImageParameterInfo.Index == INDEX_NONE ? GlobalParameter : LayerParameter;
 
+					// Find Texture using the parameter info
 					UTexture* Texture;
-					bool bParameterFound = MaterialInterface->GetTextureParameterValue(FName(ParameterName), Texture);
+					bool bParameterFound = MaterialInterface->GetTextureParameterValue(ImageParameterInfo, Texture);
 
+					// Store Texture into the FMaterial
 					if (bParameterFound)
 					{
 						TVariant<OP::ADDRESS, TSharedPtr<const FImage>> NewParameterImage;
 						NewParameterImage.Set<OP::ADDRESS>(ImageParameterOperations[ParameterIndex]);
-						
-						Material->ImageParameters.Add(FName(ParameterName), NewParameterImage);
+
+						FMaterial::FImageParameterData NewImageParameter;
+						NewImageParameter.ImageParameter = NewParameterImage;
+
+						Material->ImageParameters.Add({FName(ImageParameterInfo.Name), ImageParameterInfo.Index}, NewImageParameter);
 					}
 				}
+			}
 
-				// Scalars
-				const TArray<uint32>& ScalarParameterNames = Program.ConstantUInt32Lists[Args.ScalarParameterNames];
+			StoreMaterial(Item, Material);
 
-				for (int32 ParameterIndex = 0; ParameterIndex < ScalarParameterNames.Num(); ++ParameterIndex)
+			break;
+		}
+			
+		case EOpType::MI_MODIFY:
+		{
+			OP::MaterialModifyArgs Args = Program.GetOpArgs<OP::MaterialModifyArgs>(Item.At);
+			int32 ParameterCount = Args.NumParameters;
+
+			//Dynamic data can't be stored into arg structures
+			const uint8* Data = Program.GetOpArgsPointer(Item.At);
+			Data += sizeof(OP::MaterialModifyArgs);
+
+			switch (Item.Stage)
+			{
+			case 0:
+			{
+				TArray<FScheduledOp> Deps;
+
+				// First, Schedule the material
+				Deps.Emplace(Args.Material, Item);
+
+				// Then, Schedule parameter operations
+				for (int32 ParameterIndex = 0; ParameterIndex < ParameterCount; ++ParameterIndex)
 				{
-					const FString& ParameterName = Program.ConstantStrings[ScalarParameterNames[ParameterIndex]];
+					OP::ADDRESS NewName;
+					FMemory::Memcpy(&NewName, Data, sizeof(OP::ADDRESS));
+					Data += sizeof(OP::ADDRESS);
 
-					float ScalarValue;
-					bool bParameterFound = MaterialInterface->GetScalarParameterValue(FName(ParameterName), ScalarValue);
-					
-					if (bParameterFound)
+					uint32 LayerIndex;
+					FMemory::Memcpy(&LayerIndex, Data, sizeof(int32));
+					Data += sizeof(int32);
+
+					OP::ADDRESS NewParameterOperation;
+					FMemory::Memcpy(&NewParameterOperation, Data, sizeof(OP::ADDRESS));
+					Data += sizeof(OP::ADDRESS);
+
+					// Evaluate non lazy branches here
+					if (GetOpDataType(Program.GetOpType(NewParameterOperation)) != EDataType::Image)
 					{
-						Material->ScalarParameters.Add(FName(ParameterName), ScalarValue);
+						Deps.Emplace(NewParameterOperation, Item);
+					}
+					else
+					{
+						// Jump Image Parameter Index
+						Data += sizeof(int32);
 					}
 				}
 
-				// Vectors
-				const TArray<uint32>& ColorParameterNames = Program.ConstantUInt32Lists[Args.ColorParameterNames];
+				AddOp(FScheduledOp(Item.At, Item, 1), Deps);
+				break;
+			}
+			case 1:
+			{
+				TArray<OP::ADDRESS> ParameterNames;
+				TArray<int32> ParameterLayerIndex;
+				TArray<OP::ADDRESS> ParameterOperations;
+				TArray<int32> ImageParameterIndexArray;
 
-				for (int32 ParameterIndex = 0; ParameterIndex < ColorParameterNames.Num(); ++ParameterIndex)
+				for (int32 ParameterIndex = 0; ParameterIndex < ParameterCount; ++ParameterIndex)
 				{
-					const FString& ParameterName = Program.ConstantStrings[ColorParameterNames[ParameterIndex]];
+					OP::ADDRESS NewName;
+					FMemory::Memcpy(&NewName, Data, sizeof(OP::ADDRESS));
+					ParameterNames.Add(NewName);
+					Data += sizeof(OP::ADDRESS);
 
-					FLinearColor VectorValue;
-					bool bParameterFound = MaterialInterface->GetVectorParameterValue(FName(ParameterName), VectorValue);
-					
-					if (bParameterFound)
+					uint32 LayerIndex;
+					FMemory::Memcpy(&LayerIndex, Data, sizeof(int32));
+					ParameterLayerIndex.Add(LayerIndex);
+					Data += sizeof(int32);
+
+					OP::ADDRESS NewParameterOperation;
+					FMemory::Memcpy(&NewParameterOperation, Data, sizeof(OP::ADDRESS));
+					ParameterOperations.Add(NewParameterOperation);
+					Data += sizeof(OP::ADDRESS);
+
+					// Evaluate lazy branches here
+					if (GetOpDataType(Program.GetOpType(NewParameterOperation)) == EDataType::Image)
 					{
-						Material->ColorParameters.Add(FName(ParameterName), VectorValue);
+						int32 ImageParameterIndex;
+						FMemory::Memcpy(&ImageParameterIndex, Data, sizeof(int32));
+						ImageParameterIndexArray.Add(ImageParameterIndex);
+						Data += sizeof(int32);
 					}
 				}
-			}
 
-			StoreMaterial(Item, Material);
+				TSharedPtr<const FMaterial> Material = LoadMaterial(FCacheAddress(Args.Material, Item));
+				TSharedPtr<FMaterial> Result;
+
+				if (Material.IsValid())
+				{
+					Result = UE::Mutable::Private::CloneOrTakeOver<FMaterial>(Material);
+					int32 ImageIndex = 0;
+				
+					// ParameterIndex 0 is reserver for the material operation
+					for (int32 ParameterIndex = 0; ParameterIndex < ParameterCount; ++ParameterIndex)
+					{
+						// Get the parameter name from the constant string
+						const FName ParameterName = FName(Program.ConstantStrings[ParameterNames[ParameterIndex]]);
+
+						// Construct the parameter key
+						FParameterKey ParameterKey = { ParameterName, ParameterLayerIndex[ParameterIndex] };						
+						EDataType OperationType = GetOpDataType(Program.GetOpType(ParameterOperations[ParameterIndex]));
+
+						switch (OperationType)
+						{
+						case EDataType::Image:
+						{
+							//TODO: Evaluate the non-lazy textures
+							TVariant<OP::ADDRESS, TSharedPtr<const FImage>> NewParameterImage;
+							NewParameterImage.Set<OP::ADDRESS>(ParameterOperations[ParameterIndex]);
+
+							FMaterial::FImageParameterData ImageParameterData;
+							ImageParameterData.ImageParameter = NewParameterImage;
+							ImageParameterData.ImagePropertyIndex = ImageParameterIndexArray[ImageIndex];
+							
+							FMaterial::FImageParameterData& Value = Result->ImageParameters.FindOrAdd(ParameterKey);
+							Value = ImageParameterData;
+
+							ImageIndex++;
+							break;
+						}
+						case EDataType::Color:
+						{
+							const FVector4f Color = LoadColor(FCacheAddress(ParameterOperations[ParameterIndex], Item));
+							FVector4f& Value = Result->ColorParameters.FindOrAdd(ParameterKey);
+							Value = Color;
+							
+							break;
+						}
+						case EDataType::Scalar:
+						{
+							float Scalar = LoadScalar(FCacheAddress(ParameterOperations[ParameterIndex], Item));
+							float& Value = Result->ScalarParameters.FindOrAdd(ParameterKey);
+							Value = Scalar;
+							
+							break;
+						}
+						case EDataType::None:
+						{
+							// If the Address is invalid, we remove the parameter from the material.
 
+							if (Result->ImageParameters.Contains(ParameterKey))
+							{
+								Result->ImageParameters.Remove(ParameterKey);
+								ImageIndex++;
+							}
+							else if (Result->ColorParameters.Contains(ParameterKey))
+							{
+								Result->ColorParameters.Remove(ParameterKey);
+							}
+							else if (Result->ScalarParameters.Contains(ParameterKey))
+							{
+								Result->ScalarParameters.Remove(ParameterKey);
+							}
+
+							break;
+						}
+						default:
+							unimplemented();
+						}
+					}
+				}
+
+				StoreMaterial(Item, Result);
+				break;
+			}
+
+			default:
+				break;
+			}
+				
 			break;
 		}
 
@@ -6903,7 +7007,6 @@ namespace UE::Mutable::Private
             break;
 
         case EOpType::IN_ADDMESH:
-        case EOpType::IN_ADDIMAGE:
         case EOpType::IN_ADDMATERIAL:
         case EOpType::IN_ADDOVERLAYMATERIAL:
         case EOpType::IN_ADDOVERRIDEMATERIAL:
@@ -8304,25 +8407,22 @@ namespace UE::Mutable::Private
 				CacheMaterialOp.Type = FScheduledOp::EType::Full;
 
 				TSharedPtr<const FMaterial> Material = LoadMaterial(CacheMaterialOp);
-					
+			
 				check(Args.ParameterName < (uint32)Program.ConstantStrings.Num());
 				const FString& ParameterName = Program.ConstantStrings[Args.ParameterName];
-				const TVariant<OP::ADDRESS, TSharedPtr<const FImage>>* Image = Material->ImageParameters.Find(FName(ParameterName));
+				const FParameterKey ParameterKey = { FName(ParameterName), Args.LayerIndex };
+				const FMaterial::FImageParameterData* ImageParameterData = Material->ImageParameters.Find(ParameterKey);
 
-				if (!Image)
+				if (ImageParameterData)
 				{
-					ImageDescResults.FindOrAdd(Item.At);
-					StoreValidDesc(Item);
-				}
-				else
-				{
-					if (Image->IsType<OP::ADDRESS>())
+					const TVariant<OP::ADDRESS, TSharedPtr<const FImage>> ImageParameter = ImageParameterData->ImageParameter;
+					if (ImageParameter.IsType<OP::ADDRESS>())
 					{
 						FScheduledOp ItemNextStage(Item.At, Item, 2);
-						ItemNextStage.CustomState = Image->Get<OP::ADDRESS>();
+						ItemNextStage.CustomState = ImageParameter.Get<OP::ADDRESS>();
 
 						AddOp(ItemNextStage,
-							FScheduledOp(Image->Get<OP::ADDRESS>(), Item));
+								FScheduledOp(ImageParameter.Get<OP::ADDRESS>(), Item));
 					}
 					else
 					{
@@ -8330,6 +8430,11 @@ namespace UE::Mutable::Private
 						StoreValidDesc(Item);
 					}
 				}
+				else
+				{
+					ImageDescResults.FindOrAdd(Item.At);
+					StoreValidDesc(Item);
+				}
 
 				break;
 			}
@@ -8342,10 +8447,16 @@ namespace UE::Mutable::Private
 
 				check(Args.ParameterName < (uint32)Program.ConstantStrings.Num());
 				const FString& ParameterName = Program.ConstantStrings[Args.ParameterName];
-				const TVariant<OP::ADDRESS, TSharedPtr<const FImage>>* Image = Material->ImageParameters.Find(FName(ParameterName));
+				const FParameterKey ParameterKey = { FName(ParameterName), Args.LayerIndex };
+				const FMaterial::FImageParameterData* ImageParameterData = Material->ImageParameters.Find(ParameterKey);
+					
+				if (ImageParameterData)
+				{
+					const TVariant<OP::ADDRESS, TSharedPtr<const FImage>> ImageParameter = ImageParameterData->ImageParameter;
 
-				FExtendedImageDesc& Result = ImageDescResults.FindOrAdd(Item.At);
-				Result = ImageDescResults.FindChecked(Image->Get<OP::ADDRESS>());
+					FExtendedImageDesc& Result = ImageDescResults.FindOrAdd(Item.At);
+					Result = ImageDescResults.FindChecked(ImageParameter.Get<OP::ADDRESS>());
+				}
 
 				StoreValidDesc(Item);
 				break;
diff --git a/Engine/Plugins/Mutable/Source/MutableRuntime/Private/MuR/CodeRunnerBegin.cpp b/Engine/Plugins/Mutable/Source/MutableRuntime/Private/MuR/CodeRunnerBegin.cpp
--- a/Engine/Plugins/Mutable/Source/MutableRuntime/Private/MuR/CodeRunnerBegin.cpp
+++ b/Engine/Plugins/Mutable/Source/MutableRuntime/Private/MuR/CodeRunnerBegin.cpp
@@ -702,49 +702,6 @@ namespace UE::Mutable::Private
 					break;
 				}
 
-			case EOpType::IN_ADDVECTOR:
-			case EOpType::IN_ADDSCALAR:
-				{
-					OP::InstanceAddArgs Args = Program.GetOpArgs<OP::InstanceAddArgs>(Item.At);
-					switch (Item.Stage)
-					{
-					case 0:
-						{
-							PushOp(FScheduledOpInline(Item, 1));
-							PushOp(FScheduledOpInline(Args.instance, Item));
-
-							break;
-						}
-
-					case 1:
-						{
-							TSharedPtr<const FInstance> InInstance;
-							if (Args.instance)
-							{
-								InInstance = LoadInstance(Args.instance);
-							}
-
-							TSharedPtr<FInstance> Instance;
-							if (InInstance)
-							{
-								Instance = UE::Mutable::Private::CloneOrTakeOver<FInstance>(InInstance); // TODO GMT It will always clone due to a second reference being in the ResultsInstance
-							}
-							else
-							{
-								Instance = MakeShared<FInstance>();
-							}
-
-							StoreInstance( Item, Instance);
-							break;
-						}
-
-					default:
-						unimplemented();
-					}
-
-					break;
-				}
-        		
 			case EOpType::IN_ADDMESH:
 				{
 					OP::InstanceAddArgs Args = Program.GetOpArgs<OP::InstanceAddArgs>(Item.At);
@@ -797,57 +754,6 @@ namespace UE::Mutable::Private
 					break;
 				}
 
-			case EOpType::IN_ADDIMAGE:
-				{
-					OP::InstanceAddArgs Args = Program.GetOpArgs<OP::InstanceAddArgs>(Item.At);
-					switch (Item.Stage)
-					{
-					case 0:
-						{
-							PushOp(FScheduledOpInline(Item, 1));
-							PushOp(FScheduledOpInline(Args.instance, Item));
-							PushOp(FScheduledOpInline(Args.value, Item));
-
-							break;
-						}
-						
-					case 1:
-						{
-							TSharedPtr<const FInstance> InInstance;
-							if (Args.instance)
-							{
-								InInstance = LoadInstance(Args.instance);
-							}
-
-							TSharedPtr<FInstance> Instance;
-							if (InInstance)
-							{
-								Instance = UE::Mutable::Private::CloneOrTakeOver<FInstance>(InInstance);
-							}
-							else
-							{
-								Instance = MakeShared<FInstance>();
-							}
-							
-							if (Args.value)
-							{
-								FImageId ImageId = FSystem::GetImageId(System.GetPrivate()->WorkingMemoryManager.CurrentImageIdRegistry, Model, Params.Get(), Args.value);
-								OP::ADDRESS NameAd = Args.name;
-								check(NameAd < (uint32)Program.ConstantStrings.Num());
-								const FString& Name = Program.ConstantStrings[NameAd];
-								Instance->AddImage(0, 0, 0, ImageId, FName(Name), Args.LayerIndex);
-							}
-							
-							StoreInstance(Item, Instance);
-							break;
-						}
-
-					default:
-						unimplemented();
-					}
-
-					break;
-				}
 			case EOpType::IN_ADDMATERIAL:
 				{
 					OP::InstanceAddMaterialArgs Args = Program.GetOpArgs<OP::InstanceAddMaterialArgs>(Item.At);
@@ -1044,7 +950,7 @@ namespace UE::Mutable::Private
 								}
 							}
 
-							StoreInstance(Item, Instance);							
+							StoreInstance(Item, Instance);
 							break;
 						}
 
@@ -2553,6 +2459,111 @@ namespace UE::Mutable::Private
 					break;
 				}
 			
+			case EOpType::MI_MODIFY:
+			{
+				OP::MaterialModifyArgs Args = Program.GetOpArgs<OP::MaterialModifyArgs>(Item.At);
+				int32 ParameterCount = Args.NumParameters;
+
+				switch (Item.Stage)
+				{
+					case 0:
+					{
+						// Push next stage
+						PushOp(FScheduledOpInline(Item, 1));
+
+						//Dynamic data can't be stored into arg structures
+						const uint8* Data = Program.GetOpArgsPointer(Item.At);
+						Data += sizeof(OP::MaterialModifyArgs);
+
+						// First, Schedule the material
+						PushOp(FScheduledOpInline(Args.Material, Item));
+
+						// Store parameter names and schedule the parameter operations
+						for (int32 ParameterIndex = 0; ParameterIndex < ParameterCount; ++ParameterIndex)
+						{
+							OP::ADDRESS NewName;
+							FMemory::Memcpy(&NewName, Data, sizeof(OP::ADDRESS));
+							Data += sizeof(OP::ADDRESS);
+
+							uint32 LayerIndex;
+							FMemory::Memcpy(&LayerIndex, Data, sizeof(int32));
+							Data += sizeof(int32);
+
+							OP::ADDRESS NewParameterOperation;
+							FMemory::Memcpy(&NewParameterOperation, Data, sizeof(OP::ADDRESS));
+							Data += sizeof(OP::ADDRESS);
+
+							// Skip 0 addresses since do not need to load anything
+							if (NewParameterOperation == 0)
+							{
+								continue;
+							}
+
+							PushOp(FScheduledOpInline(NewParameterOperation, Item));
+
+							// Evaluate non lazy branches here
+							if (GetOpDataType(Program.GetOpType(NewParameterOperation)) == EDataType::Image)
+							{
+								// Jump Image Parameter Index
+								Data += sizeof(int32);
+							}
+						}
+
+						break;
+					}
+					case 1:
+					{
+						//Dynamic data can't be stored into arg structures
+						const uint8* Data = Program.GetOpArgsPointer(Item.At);
+						
+						//Material
+						Data += sizeof(OP::MaterialModifyArgs);
+
+						// Store parameter names and schedule the parameter operations
+						for (int32 ParameterIndex = 0; ParameterIndex < ParameterCount; ++ParameterIndex)
+						{
+							OP::ADDRESS NewName;
+							FMemory::Memcpy(&NewName, Data, sizeof(OP::ADDRESS));
+							Data += sizeof(OP::ADDRESS);
+
+							uint32 LayerIndex;
+							FMemory::Memcpy(&LayerIndex, Data, sizeof(int32));
+							Data += sizeof(int32);
+
+							OP::ADDRESS ParameterOperation;
+							FMemory::Memcpy(&ParameterOperation, Data, sizeof(OP::ADDRESS));
+							Data += sizeof(OP::ADDRESS);
+
+							// Skip 0 addresses since do not need to load anything
+							if (ParameterOperation == 0)
+							{
+								continue;
+							}
+
+							EOpType OpType = Program.GetOpType(ParameterOperation);
+							EDataType DataType = GetOpDataType(OpType);
+
+							// Evaluate non lazy branches here
+							if ( DataType == EDataType::Image)
+							{
+								// Jump Image Parameter Index
+								Data += sizeof(int32);
+							}
+							else if (DataType == EDataType::Scalar)
+							{
+								// Pop float value
+								LoadFloat(ParameterOperation);
+							}
+						}
+
+						StoreNone(Item);
+						break;
+					}
+				}
+
+				break;
+			}
+
 			default:
 				unimplemented();
 			}
diff --git a/Engine/Plugins/Mutable/Source/MutableRuntime/Private/MuR/External/MaterialAdapter.cpp b/Engine/Plugins/Mutable/Source/MutableRuntime/Private/MuR/External/MaterialAdapter.cpp
--- a/Engine/Plugins/Mutable/Source/MutableRuntime/Private/MuR/External/MaterialAdapter.cpp
+++ b/Engine/Plugins/Mutable/Source/MutableRuntime/Private/MuR/External/MaterialAdapter.cpp
@@ -1,4 +1,4 @@
-﻿// Copyright Epic Games, Inc. All Rights Reserved.
+// Copyright Epic Games, Inc. All Rights Reserved.
 
 #include "MuR/External/MaterialAdapter.h"
 
@@ -31,9 +31,9 @@ namespace UE::Mutable
 	}
 
 
-	TOptional<TValueConst<FFloatAdapter>> FMaterialAdapter::GetFloat(const FName& ParameterName) const
+	TOptional<TValueConst<FFloatAdapter>> FMaterialAdapter::GetFloat(const Private::FParameterKey& ParameterKey) const
 	{
-		float* Result = Material->ScalarParameters.Find(ParameterName);
+		float* Result = Material->ScalarParameters.Find(ParameterKey);
 		if (!Result)
 		{
 			return {};
@@ -48,29 +48,29 @@ namespace UE::Mutable
 	}
 
 
-	bool FMaterialAdapter::SetFloat(const FName& ParameterName, const TValueConst<FFloatAdapter>* Value)
+	bool FMaterialAdapter::SetFloat(const Private::FParameterKey& ParameterKey, const TValueConst<FFloatAdapter>* Value)
 	{
-		if (!Material->ScalarParameters.Contains(ParameterName))
+		if (!Material->ScalarParameters.Contains(ParameterKey))
 		{
 			return false;
 		}
 
 		if (Value)
 		{
-			Material->ScalarParameters.Add(ParameterName, Value->Get().GetValue());
+			Material->ScalarParameters.Add(ParameterKey, Value->Get().GetValue());
 		}
 		else
 		{
-			Material->ScalarParameters.Remove(ParameterName);
+			Material->ScalarParameters.Remove(ParameterKey);
 		}
 		
 		return true;
 	}
 
 
-	TOptional<TValueConst<FVectorAdapter>> FMaterialAdapter::GetVector(const FName& ParameterName) const
+	TOptional<TValueConst<FVectorAdapter>> FMaterialAdapter::GetVector(const Private::FParameterKey& ParameterKey) const
 	{
-		FVector4f* Result = Material->ColorParameters.Find(ParameterName);
+		FVector4f* Result = Material->ColorParameters.Find(ParameterKey);
 		if (!Result)
 		{
 			return {};
@@ -85,40 +85,42 @@ namespace UE::Mutable
 	}
 
 
-	bool FMaterialAdapter::SetVector(const FName& ParameterName, const TValueConst<FVectorAdapter>* Value)
+	bool FMaterialAdapter::SetVector(const Private::FParameterKey& ParameterKey, const TValueConst<FVectorAdapter>* Value)
 	{
-		if (!Material->ColorParameters.Contains(ParameterName))
+		if (!Material->ColorParameters.Contains(ParameterKey))
 		{
 			return false;
 		}
 
 		if (Value)
 		{
-			Material->ColorParameters.Add(ParameterName, Value->Get().GetValue());
+			Material->ColorParameters.Add(ParameterKey, Value->Get().GetValue());
 		}
 		else
 		{
-			Material->ColorParameters.Remove(ParameterName);
+			Material->ColorParameters.Remove(ParameterKey);
 		}
 		
 		return true;
 	}
 
 
-	TOptional<TValueConst<FTextureAdapter>> FMaterialAdapter::GetTexture(const FName& ParameterName) const
+	TOptional<TValueConst<FTextureAdapter>> FMaterialAdapter::GetTexture(const Private::FParameterKey& ParameterKey) const
 	{
-		TVariant<Private::OP::ADDRESS, TSharedPtr<const Private::FImage>>* Result = Material->ImageParameters.Find(ParameterName);
+		UE::Mutable::Private::FMaterial::FImageParameterData* Result = Material->ImageParameters.Find(ParameterKey);
 		if (!Result)
 		{
 			return {};
 		}
 
-		if (ensure(Result->IsType<TSharedPtr<const Private::FImage>>()))
+		TVariant<Private::OP::ADDRESS, TSharedPtr<const Private::FImage>> Variant = Result->ImageParameter;
+
+		if (ensure(Variant.IsType<TSharedPtr<const Private::FImage>>()))
 		{
 			TValue<FTextureAdapter> Value;
 			Value.Ptr = MakeShared<FInstancedStruct>();
 			Value.Ptr->InitializeAs(FTextureAdapter::StaticStruct());
-			Value.Ptr->GetMutable<FTextureAdapter>().Image = ConstCastSharedPtr<Private::FImage>(Result->Get<TSharedPtr<const Private::FImage>>()); // TODO GMT If lazy wait
+			Value.Ptr->GetMutable<FTextureAdapter>().Image = ConstCastSharedPtr<Private::FImage>(Variant.Get<TSharedPtr<const Private::FImage>>()); // TODO GMT If lazy wait
 			
 			return TOptional<TValueConst<FTextureAdapter>>(MoveTemp(Value));
 		}
@@ -128,10 +130,11 @@ namespace UE::Mutable
 		}
 	}
 
-	
-	bool FMaterialAdapter::SetTexture(const FName& ParameterName, const TValueConst<FTextureAdapter>* Value)
+	bool FMaterialAdapter::SetTexture(const Private::FParameterKey& ParameterKey, const TValueConst<FTextureAdapter>* Value)
 	{
-		if (!Material->ImageParameters.Contains(ParameterName))
+		UE::Mutable::Private::FMaterial::FImageParameterData* ImageData = Material->ImageParameters.Find(ParameterKey);
+
+		if (!ImageData)
 		{
 			return false;
 		}
@@ -141,11 +144,11 @@ namespace UE::Mutable
 			TVariant<Private::OP::ADDRESS, TSharedPtr<const Private::FImage>> Image;
 			Image.Set<TSharedPtr<const Private::FImage>>(Value->Get().Image);
 			
-			Material->ImageParameters.Add(ParameterName, Image);
+			Material->ImageParameters.Add(ParameterKey, { Image, ImageData->ImagePropertyIndex });
 		}
 		else
 		{
-			Material->ImageParameters.Remove(ParameterName);
+			Material->ImageParameters.Remove(ParameterKey);
 		}
 		
 		return true;
diff --git a/Engine/Plugins/Mutable/Source/MutableRuntime/Private/MuR/Instance.cpp b/Engine/Plugins/Mutable/Source/MutableRuntime/Private/MuR/Instance.cpp
--- a/Engine/Plugins/Mutable/Source/MutableRuntime/Private/MuR/Instance.cpp
+++ b/Engine/Plugins/Mutable/Source/MutableRuntime/Private/MuR/Instance.cpp
@@ -128,27 +128,8 @@ namespace UE::Mutable::Private
 
         return 0;
     }
-	
-	
-    int32 FInstance::GetImageCount(int32 ComponentIndex, int32 LODIndex, int32 SurfaceIndex ) const
-	{
-		return Components[ComponentIndex].LODs[LODIndex].Surfaces[SurfaceIndex].Images.Num();
-	}
-
-	
-    int32 FInstance::GetVectorCount(int32 ComponentIndex, int32 LODIndex, int32 SurfaceIndex ) const
-	{
-		return Components[ComponentIndex].LODs[LODIndex].Surfaces[SurfaceIndex].Vectors.Num();
-	}
 
-	
-    int32 FInstance::GetScalarCount(int32 ComponentIndex, int32 LODIndex, int32 SurfaceIndex ) const
-	{
-		return Components[ComponentIndex].LODs[LODIndex].Surfaces[SurfaceIndex].Scalars.Num();
-	}
-
-    
-    int32 FInstance::GetStringCount(int32 ComponentIndex, int32 LODIndex, int32 SurfaceIndex ) const
+int32 FInstance::GetStringCount(int32 ComponentIndex, int32 LODIndex, int32 SurfaceIndex ) const
     {
 		return Components[ComponentIndex].LODs[LODIndex].Surfaces[SurfaceIndex].Strings.Num();
 	}
@@ -159,61 +140,7 @@ namespace UE::Mutable::Private
 		return Components[ComponentIndex].LODs[LODIndex].MeshId;
     }
 
-	
-	FImageId FInstance::GetImageId(int32 ComponentIndex, int32 LODIndex, int32 SurfaceIndex, int32 ImageIndex) const
-	{
-		return Components[ComponentIndex].LODs[LODIndex].Surfaces[SurfaceIndex].Images[ImageIndex].Id;
-	}
 
-	
-    FName FInstance::GetImageName(int32 ComponentIndex, int32 LODIndex, int32 SurfaceIndex, int32 ImageIndex ) const
-	{
-        return Components[ComponentIndex].LODs[LODIndex].Surfaces[SurfaceIndex].Images[ImageIndex].Name;
-	}
-
-	
-	int32 FInstance::GetImageLayerIndex(int32 ComponentIndex, int32 LODIndex, int32 SurfaceIndex, int32 ImageIndex) const
-	{
-		return Components[ComponentIndex].LODs[LODIndex].Surfaces[SurfaceIndex].Images[ImageIndex].LayerIndex;
-	}
-
-	
-	FVector4f FInstance::GetVector(int32 ComponentIndex, int32 LODIndex, int32 SurfaceIndex, int32 VectorIndex) const
-	{
-        return Components[ComponentIndex].LODs[LODIndex].Surfaces[SurfaceIndex].Vectors[VectorIndex].Value;
-	}
-	
-	
-	FName FInstance::GetVectorName(int32 ComponentIndex, int32 LODIndex, int32 SurfaceIndex, int32 VectorIndex) const
-	{
-        return Components[ComponentIndex].LODs[LODIndex].Surfaces[SurfaceIndex].Vectors[VectorIndex].Name;
-	}
-
-
-	int32 FInstance::GetVectorLayerIndex(int32 ComponentIndex, int32 LODIndex, int32 SurfaceIndex, int32 VectorIndex) const
-	{
-		return Components[ComponentIndex].LODs[LODIndex].Surfaces[SurfaceIndex].Vectors[VectorIndex].LayerIndex;
-	}
-	
-	
-    float FInstance::GetScalar(int32 ComponentIndex, int32 LODIndex, int32 SurfaceIndex, int32 ScalarIndex ) const
-	{
-        return Components[ComponentIndex].LODs[LODIndex].Surfaces[SurfaceIndex].Scalars[ScalarIndex].Value;
-	}
-
-	
-	FName FInstance::GetScalarName(int32 ComponentIndex, int32 LODIndex, int32 SurfaceIndex, int32 ScalarIndex) const
-	{
-        return Components[ComponentIndex].LODs[LODIndex].Surfaces[SurfaceIndex].Scalars[ScalarIndex].Name;
-	}
-
-
-	int32 FInstance::GetScalarLayerIndex(int32 ComponentIndex, int32 LODIndex, int32 SurfaceIndex, int32 ScalarIndex) const
-	{
-        return Components[ComponentIndex].LODs[LODIndex].Surfaces[SurfaceIndex].Scalars[ScalarIndex].LayerIndex;
-	}
-
-	
     FString FInstance::GetString(int32 ComponentIndex, int32 LODIndex, int32 SurfaceIndex, int32 StringIndex ) const
     {
 		bool bValid = Components[ComponentIndex].LODs[LODIndex].Surfaces[SurfaceIndex].Strings.IsValidIndex(StringIndex);
@@ -364,69 +291,6 @@ namespace UE::Mutable::Private
 	}
 
 
-    int32 FInstance::AddImage( int32 ComponentIndex, int32 LODIndex, int32 SurfaceIndex, const FImageId& ImageId, FName Name, int32 LayerIndex)
-	{
-		// Automatically create the necessary lods and components
-		while (ComponentIndex >= Components.Num())
-		{
-			AddComponent();
-		}
-		while (LODIndex >= Components[ComponentIndex].LODs.Num())
-		{
-			AddLOD(ComponentIndex);
-		}
-		while (SurfaceIndex >= Components[ComponentIndex].LODs[LODIndex].Surfaces.Num())
-		{
-			AddSurface(LODIndex, ComponentIndex);
-		}
-
-		FInstanceSurface& Surface = Components[ComponentIndex].LODs[LODIndex].Surfaces[SurfaceIndex];
-		return Surface.Images.Add({ ImageId, Name, LayerIndex });
-	}
-	
-	
-    int32 FInstance::AddVector( int32 ComponentIndex, int32 LODIndex, int32 SurfaceIndex, const FVector4f& Vector, FName Name, int32 LayerIndex)
-	{
-		// Automatically create the necessary lods and components
-		while (ComponentIndex >= Components.Num())
-		{
-			AddComponent();
-		}
-		while (LODIndex >= Components[ComponentIndex].LODs.Num())
-		{
-			AddLOD(ComponentIndex);
-		}
-		while (SurfaceIndex >= Components[ComponentIndex].LODs[LODIndex].Surfaces.Num())
-		{
-			AddSurface(LODIndex, ComponentIndex);
-		}
-
-		FInstanceSurface& Surface = Components[ComponentIndex].LODs[LODIndex].Surfaces[SurfaceIndex];
-		return Surface.Vectors.Add({ Vector, Name , LayerIndex } );
-	}
-
-	
-    int32 FInstance::AddScalar( int32 ComponentIndex, int32 LODIndex, int32 SurfaceIndex, float Scalar, FName Name, int32 LayerIndex)
-    {
-        // Automatically create the necessary lods and components
-		while (ComponentIndex >= Components.Num())
-		{
-			AddComponent();
-		}
-		while (LODIndex >= Components[ComponentIndex].LODs.Num())
-		{
-			AddLOD(ComponentIndex);
-		}
-		while (SurfaceIndex >= Components[ComponentIndex].LODs[LODIndex].Surfaces.Num())
-		{
-			AddSurface(LODIndex, ComponentIndex);
-		}
-
-		FInstanceSurface& Surface = Components[ComponentIndex].LODs[LODIndex].Surfaces[SurfaceIndex];
-		return Surface.Scalars.Add({ Scalar, Name, LayerIndex });
-    }
-	
-    
     int32 FInstance::AddString( int32 ComponentIndex, int32 LODIndex, int32 SurfaceIndex, const FString& Value, FName Name)
     {
         // Automatically create the necessary lods and components
diff --git a/Engine/Plugins/Mutable/Source/MutableRuntime/Private/MuR/Material.cpp b/Engine/Plugins/Mutable/Source/MutableRuntime/Private/MuR/Material.cpp
--- a/Engine/Plugins/Mutable/Source/MutableRuntime/Private/MuR/Material.cpp
+++ b/Engine/Plugins/Mutable/Source/MutableRuntime/Private/MuR/Material.cpp
@@ -66,4 +66,32 @@ namespace UE::Mutable::Private
 			ColorParameters == Other.ColorParameters &&
 			ScalarParameters == Other.ScalarParameters;
 	};
+
+
+	void FMaterial::FImageParameterData::Serialise(FOutputArchive& arch) const
+	{
+		arch << ImageParameter;
+		arch << ImagePropertyIndex;
+	}
+
+
+	void FMaterial::FImageParameterData::Unserialise(FInputArchive& arch)
+	{
+		arch >> ImageParameter;
+		arch >> ImagePropertyIndex;
+	}
+
+
+	void FParameterKey::Serialise(FOutputArchive& arch) const
+	{
+		arch << ParameterName;
+		arch << LayerIndex;
+	}
+
+
+	void FParameterKey::Unserialise(FInputArchive& arch)
+	{
+		arch >> ParameterName;
+		arch >> LayerIndex;
+	}
 }
diff --git a/Engine/Plugins/Mutable/Source/MutableRuntime/Private/MuR/Operations.cpp b/Engine/Plugins/Mutable/Source/MutableRuntime/Private/MuR/Operations.cpp
--- a/Engine/Plugins/Mutable/Source/MutableRuntime/Private/MuR/Operations.cpp
+++ b/Engine/Plugins/Mutable/Source/MutableRuntime/Private/MuR/Operations.cpp
@@ -71,6 +71,7 @@ namespace UE::Mutable::Private
 		{ EDataType::Color,			},	// CO_MATERIAL_BREAK
 		{ EDataType::Image,			},	// IM_MATERIAL_BREAK
 		{ EDataType::Image,			},	// IM_PARAMETER_FROM_MATERIAL
+		{ EDataType::Material,		},	// MI_MODIFY
 
 		{ EDataType::Bool,			},	// BO_EQUAL_SC_CONST
 		{ EDataType::Bool,			},	// BO_AND
@@ -145,9 +146,6 @@ namespace UE::Mutable::Private
 		{ EDataType::Mesh,			},	// ME_SKELETALMESH_BREAK
 
 		{ EDataType::Instance,		},	// IN_ADDMESH
-		{ EDataType::Instance,		},	// IN_ADDIMAGE
-		{ EDataType::Instance,		},	// IN_ADDVECTOR
-		{ EDataType::Instance,		},	// IN_ADDSCALAR
 		{ EDataType::Instance,		},	// IN_ADDSTRING
 		{ EDataType::Instance,		},	// IN_ADDSURFACE
 		{ EDataType::Instance,		},	// IN_ADDCOMPONENT
@@ -311,6 +309,45 @@ namespace UE::Mutable::Private
 			f(Args.ParameterName);
 			break;
 		}
+		
+		case EOpType::MI_MODIFY:
+		{
+			OP::MaterialModifyArgs Args = program.GetOpArgs<OP::MaterialModifyArgs>(at);
+			int32 ParameterCount = Args.NumParameters;
+
+			//Dynamic data can't be stored into arg structures
+			const uint8* Data = program.GetOpArgsPointer(at);
+			Data += sizeof(OP::MaterialModifyArgs);
+
+			// First, the material
+			f(Args.Material);
+
+			// Store parameter names and schedule the parameter operations
+			for (int32 ParameterIndex = 0; ParameterIndex < ParameterCount; ++ParameterIndex)
+			{
+				OP::ADDRESS NewName;
+				FMemory::Memcpy(&NewName, Data, sizeof(OP::ADDRESS));
+				Data += sizeof(OP::ADDRESS);
+
+				uint32 LayerIndex;
+				FMemory::Memcpy(&LayerIndex, Data, sizeof(int32));
+				Data += sizeof(int32);
+
+				OP::ADDRESS NewParameterOperation;
+				FMemory::Memcpy(&NewParameterOperation, Data, sizeof(OP::ADDRESS));
+				f(NewParameterOperation);
+				Data += sizeof(OP::ADDRESS);
+
+				// Evaluate non lazy branches here
+				if (GetOpDataType(program.GetOpType(NewParameterOperation)) == EDataType::Image)
+				{
+					// Jump Image Parameter Index
+					Data += sizeof(int32);
+				}
+			}
+
+			break;
+		}
 
         //-------------------------------------------------------------------------------------
         case EOpType::BO_EQUAL_INT_CONST:
@@ -852,9 +889,6 @@ namespace UE::Mutable::Private
         
         //-------------------------------------------------------------------------------------
         case EOpType::IN_ADDMESH:
-        case EOpType::IN_ADDIMAGE:
-        case EOpType::IN_ADDVECTOR:
-        case EOpType::IN_ADDSCALAR:
         case EOpType::IN_ADDSTRING:
         case EOpType::IN_ADDCOMPONENT:
         case EOpType::IN_ADDSURFACE:
diff --git a/Engine/Plugins/Mutable/Source/MutableRuntime/Public/MuR/External/MaterialAdapter.h b/Engine/Plugins/Mutable/Source/MutableRuntime/Public/MuR/External/MaterialAdapter.h
--- a/Engine/Plugins/Mutable/Source/MutableRuntime/Public/MuR/External/MaterialAdapter.h
+++ b/Engine/Plugins/Mutable/Source/MutableRuntime/Public/MuR/External/MaterialAdapter.h
@@ -1,4 +1,4 @@
-﻿// Copyright Epic Games, Inc. All Rights Reserved.
+// Copyright Epic Games, Inc. All Rights Reserved.
 
 #pragma once
 
@@ -17,6 +17,8 @@ namespace UE::Mutable
 	{
 		class FMaterial;
 		class CodeRunner;
+
+		struct FParameterKey;
 	}
 
 	
@@ -37,17 +39,17 @@ namespace UE::Mutable
 		
 		FMaterialAdapter& operator=(FMaterialAdapter&& Other) = delete;
 
-		MUTABLERUNTIME_API TOptional<TValueConst<FFloatAdapter>> GetFloat(const FName& ParameterName) const;
+		MUTABLERUNTIME_API TOptional<TValueConst<FFloatAdapter>> GetFloat(const Private::FParameterKey& ParameterKey) const;
 
-		MUTABLERUNTIME_API bool SetFloat(const FName& ParameterName, const TValueConst<FFloatAdapter>* Value);
+		MUTABLERUNTIME_API bool SetFloat(const Private::FParameterKey& ParameterKey, const TValueConst<FFloatAdapter>* Value);
 		
-		MUTABLERUNTIME_API TOptional<TValueConst<FVectorAdapter>> GetVector(const FName& ParameterName) const;
+		MUTABLERUNTIME_API TOptional<TValueConst<FVectorAdapter>> GetVector(const Private::FParameterKey& ParameterKey) const;
 		
-		MUTABLERUNTIME_API bool SetVector(const FName& ParameterName, const TValueConst<FVectorAdapter>* Value);
+		MUTABLERUNTIME_API bool SetVector(const Private::FParameterKey& ParameterKey, const TValueConst<FVectorAdapter>* Value);
 		
-		MUTABLERUNTIME_API TOptional<TValueConst<FTextureAdapter>> GetTexture(const FName& ParameterName) const;
+		MUTABLERUNTIME_API TOptional<TValueConst<FTextureAdapter>> GetTexture(const Private::FParameterKey& ParameterKey) const;
 
-		MUTABLERUNTIME_API bool SetTexture(const FName& ParameterName, const TValueConst<FTextureAdapter>* Value);
+		MUTABLERUNTIME_API bool SetTexture(const Private::FParameterKey& ParameterKey, const TValueConst<FTextureAdapter>* Value);
 	
 	private:
 		TSharedPtr<Private::FMaterial> Material;
diff --git a/Engine/Plugins/Mutable/Source/MutableTools/Internal/MuT/AST.h b/Engine/Plugins/Mutable/Source/MutableTools/Internal/MuT/AST.h
--- a/Engine/Plugins/Mutable/Source/MutableTools/Internal/MuT/AST.h
+++ b/Engine/Plugins/Mutable/Source/MutableTools/Internal/MuT/AST.h
@@ -224,7 +224,7 @@ namespace UE::Mutable::Private
 		FReferencedImageResourceFunc ReferencedImageResourceProvider;
 		FReferencedMeshResourceFunc ReferencedMeshResourceProvider;
 
-        bool bDisableImageGeneration = false;
+        bool bDisableMaterialGeneration = false;
         bool bDisableMeshGeneration = false;
 	};
 
diff --git a/Engine/Plugins/Mutable/Source/MutableTools/Internal/MuT/Compiler.h b/Engine/Plugins/Mutable/Source/MutableTools/Internal/MuT/Compiler.h
--- a/Engine/Plugins/Mutable/Source/MutableTools/Internal/MuT/Compiler.h
+++ b/Engine/Plugins/Mutable/Source/MutableTools/Internal/MuT/Compiler.h
@@ -82,7 +82,7 @@ namespace UE::Mutable::Private
 		/** */
 		UE_API void SetReferencedResourceCallback(const FReferencedImageResourceFunc&, const FReferencedMeshResourceFunc&);
 
-        UE_API void SetDisableImageGeneration(bool bDisabled);
+        UE_API void SetDisableMaterialGeneration(bool bDisabled);
         UE_API void SetDisableMeshGeneration(bool bDisabled);
 
         //! Different data packing strategies
diff --git a/Engine/Plugins/Mutable/Source/MutableTools/Internal/MuT/Node.h b/Engine/Plugins/Mutable/Source/MutableTools/Internal/MuT/Node.h
--- a/Engine/Plugins/Mutable/Source/MutableTools/Internal/MuT/Node.h
+++ b/Engine/Plugins/Mutable/Source/MutableTools/Internal/MuT/Node.h
@@ -177,11 +177,12 @@ namespace UE::Mutable::Private
 			MaterialParameter,
 			MaterialExternal,
 			MaterialSkeletalMeshBreak,
-
+			MaterialModify,
+			
 			ImageMaterialBreak,
 			ScalarMaterialBreak,
 			ColourMaterialBreak,
-
+			
 			Count
 		};
 
diff --git a/Engine/Plugins/Mutable/Source/MutableTools/Internal/MuT/NodeColourMaterialBreak.h b/Engine/Plugins/Mutable/Source/MutableTools/Internal/MuT/NodeColourMaterialBreak.h
--- a/Engine/Plugins/Mutable/Source/MutableTools/Internal/MuT/NodeColourMaterialBreak.h
+++ b/Engine/Plugins/Mutable/Source/MutableTools/Internal/MuT/NodeColourMaterialBreak.h
@@ -18,6 +18,7 @@ namespace UE::Mutable::Private
 
 		Ptr<NodeMaterial> MaterialSource;
 		FName ParameterName;
+		int32 LayerIndex;
 
 	public:
 
diff --git a/Engine/Plugins/Mutable/Source/MutableTools/Internal/MuT/NodeImageMaterialBreak.h b/Engine/Plugins/Mutable/Source/MutableTools/Internal/MuT/NodeImageMaterialBreak.h
--- a/Engine/Plugins/Mutable/Source/MutableTools/Internal/MuT/NodeImageMaterialBreak.h
+++ b/Engine/Plugins/Mutable/Source/MutableTools/Internal/MuT/NodeImageMaterialBreak.h
@@ -18,6 +18,7 @@ namespace UE::Mutable::Private
 
 		Ptr<NodeMaterial> MaterialSource;
 		FName ParameterName;
+		int32 LayerIndex;
 
 	public:
 
diff --git a/Engine/Plugins/Mutable/Source/MutableTools/Internal/MuT/NodeMaterialConstant.h b/Engine/Plugins/Mutable/Source/MutableTools/Internal/MuT/NodeMaterialConstant.h
--- a/Engine/Plugins/Mutable/Source/MutableTools/Internal/MuT/NodeMaterialConstant.h
+++ b/Engine/Plugins/Mutable/Source/MutableTools/Internal/MuT/NodeMaterialConstant.h
@@ -3,6 +3,7 @@
 #pragma once
 
 #include "HAL/Platform.h"
+#include "MuR/Material.h"
 #include "MuR/Ptr.h"
 #include "MuR/RefCounted.h"
 #include "MuR/PassthroughObject.h"
@@ -26,9 +27,9 @@ namespace UE::Mutable::Private
 
 		PASSTHROUGH_ID MaterialId = PASSTHROUGH_ID_INVALID;
 
-		TMap<FName, FVector4f> ColorValues;
-		TMap<FName, float> ScalarValues;
-		TMap<FName, Ptr<NodeImage>> ImageValues;
+		TMap<FParameterKey, FVector4f> ColorValues;
+		TMap<FParameterKey, float> ScalarValues;
+		TMap<FParameterKey, Ptr<NodeImage>> ImageValues;
 
 	protected:
 
diff --git a/Engine/Plugins/Mutable/Source/MutableTools/Internal/MuT/NodeMaterialParameter.h b/Engine/Plugins/Mutable/Source/MutableTools/Internal/MuT/NodeMaterialParameter.h
--- a/Engine/Plugins/Mutable/Source/MutableTools/Internal/MuT/NodeMaterialParameter.h
+++ b/Engine/Plugins/Mutable/Source/MutableTools/Internal/MuT/NodeMaterialParameter.h
@@ -4,6 +4,7 @@
 
 #include "Containers/Map.h"
 #include "HAL/Platform.h"
+#include "MuR/Material.h"
 #include "MuR/Ptr.h"
 #include "MuR/RefCounted.h"
 #include "MuT/NodeRange.h"
@@ -30,9 +31,9 @@ namespace UE::Mutable::Private
 		/** Ranges for the parameter in case it is multidimensional. */
 		TArray<Ptr<NodeRange>> Ranges;
 
-		TArray<FString> ColorParameterNames;
-		TArray<FString> ScalarParameterNames;
-		TMap<FString, Ptr<NodeImage>> ImageParameterNodes;
+		TArray<FParameterKey> ColorParametersToCompile;
+		TArray<FParameterKey> ScalarParametersToCompile;
+		TMap<FParameterKey, Ptr<NodeImage>> ImageParametersToCompile;
 
 	public:
 
diff --git a/Engine/Plugins/Mutable/Source/MutableTools/Internal/MuT/NodeModifierSurfaceEdit.h b/Engine/Plugins/Mutable/Source/MutableTools/Internal/MuT/NodeModifierSurfaceEdit.h
--- a/Engine/Plugins/Mutable/Source/MutableTools/Internal/MuT/NodeModifierSurfaceEdit.h
+++ b/Engine/Plugins/Mutable/Source/MutableTools/Internal/MuT/NodeModifierSurfaceEdit.h
@@ -7,6 +7,7 @@
 #include "MuT/NodeImage.h"
 #include "MuT/NodeScalar.h"
 #include "MuR/ImageTypes.h"
+#include "MuR/Material.h"
 #include "MuR/RefCounted.h"
 #include "MuR/Ptr.h"
 
@@ -27,10 +28,10 @@ namespace UE::Mutable::Private
 		/** Data for every modified texture. */
 		struct FTexture
 		{
-			/** Name used to match the image with the original one being modified. 
-			* This should match the MaterialParameterName in a NodeSurfaceNew::FImageData.
+			/** Key used to match the image with the original one being modified. 
+			* This should match the MaterialParameterKey in a NodeSurfaceNew::FImageData.
 			*/
-			FString MaterialParameterName;
+			FParameterKey MaterialParameterKey;
 
 			/** Image to add if extgending. */
 			Ptr<NodeImage> Extend;
diff --git a/Engine/Plugins/Mutable/Source/MutableTools/Internal/MuT/NodeScalarMaterialBreak.h b/Engine/Plugins/Mutable/Source/MutableTools/Internal/MuT/NodeScalarMaterialBreak.h
--- a/Engine/Plugins/Mutable/Source/MutableTools/Internal/MuT/NodeScalarMaterialBreak.h
+++ b/Engine/Plugins/Mutable/Source/MutableTools/Internal/MuT/NodeScalarMaterialBreak.h
@@ -18,6 +18,7 @@ namespace UE::Mutable::Private
 
 		Ptr<NodeMaterial> MaterialSource;
 		FName ParameterName;
+		int32 LayerIndex;
 
 	public:
 
diff --git a/Engine/Plugins/Mutable/Source/MutableTools/Internal/MuT/NodeSurfaceNew.h b/Engine/Plugins/Mutable/Source/MutableTools/Internal/MuT/NodeSurfaceNew.h
--- a/Engine/Plugins/Mutable/Source/MutableTools/Internal/MuT/NodeSurfaceNew.h
+++ b/Engine/Plugins/Mutable/Source/MutableTools/Internal/MuT/NodeSurfaceNew.h
@@ -39,47 +39,6 @@ namespace UE::Mutable::Private
 
 		Ptr<NodeMaterial> Material;
 
-		struct FParameterIdentifier
-		{
-			/** Name of the parameter inside the material. */
-			FString ParameterName;
-
-			/** Index of the Material Layer. If it isn't a layered material then -1. */
-			int32 MaterialLayerIndex = -1;
-
-			//TODO: Add the possibility to modify blend parameters not just layer parameters.
-			//ParameterType: Global, Blend or Layer
-		};
-
-		struct FImageData
-		{
-			FParameterIdentifier ParameterInfo;
-			FString MaterialName;
-			FString MaterialParameterName;
-			Ptr<NodeImage> Image;
-
-			/** Index of the layout transform to apply to this image. It could be negative, to indicate no layout transform. */
-			int8 LayoutIndex = 0;
-		};
-
-		TArray<FImageData> Images;
-
-		struct FVectorData
-		{
-			FParameterIdentifier ParameterInfo;
-			Ptr<NodeColour> Vector;
-		};
-
-		TArray<FVectorData> Vectors;
-
-		struct FScalar
-		{
-			FParameterIdentifier ParameterInfo;
-			Ptr<NodeScalar> Scalar;
-		};
-
-		TArray<FScalar> Scalars;
-
 		struct FStringData
 		{
 			FString Name;
diff --git a/Engine/Plugins/Mutable/Source/MutableTools/Private/MuT/ASTOpConstantResource.cpp b/Engine/Plugins/Mutable/Source/MutableTools/Private/MuT/ASTOpConstantResource.cpp
--- a/Engine/Plugins/Mutable/Source/MutableTools/Private/MuT/ASTOpConstantResource.cpp
+++ b/Engine/Plugins/Mutable/Source/MutableTools/Private/MuT/ASTOpConstantResource.cpp
@@ -24,7 +24,7 @@ namespace UE::Mutable::Private
 	void ASTOpConstantResource::ForEachChild(const TFunctionRef<void(ASTChild&)> f)
 	{
 		/** Image Parameters are evaluated at runtime. */
-		for (TPair<FName, ASTChild >& Element : ImageOperations)
+		for (TPair<FParameterKey, ASTChild >& Element : ImageOperations)
 		{
 			f(Element.Value);
 		}
@@ -55,7 +55,7 @@ namespace UE::Mutable::Private
 		n->ValueHash = ValueHash;
 		n->SourceDataDescriptor = SourceDataDescriptor;
 
-		for (const TPair<FName, ASTChild>& ImageOperation : ImageOperations)
+		for (const TPair<FParameterKey, ASTChild>& ImageOperation : ImageOperations)
 		{
 			n->ImageOperations.Add(ImageOperation.Key, ASTChild(n, MapChild(ImageOperation.Value.child())));
 		}
@@ -492,12 +492,15 @@ namespace UE::Mutable::Private
 					check(MaterialValue);
 
 					// Store linked child images into the material constant
-					for (const TPair<FName, ASTChild>& ImageOperation : ImageOperations)
+					for (const TPair<FParameterKey, ASTChild>& ImageOperation : ImageOperations)
 					{
 						TVariant<OP::ADDRESS, TSharedPtr<const FImage>> NewImage;
 						NewImage.Set<OP::ADDRESS>(ImageOperation.Value->linkedAddress);
 
-						MaterialValue->ImageParameters.Add(FName(ImageOperation.Key), NewImage);
+						FMaterial::FImageParameterData ImageParameterData;
+						ImageParameterData.ImageParameter = NewImage;
+
+						MaterialValue->ImageParameters.Add(ImageOperation.Key, ImageParameterData);
 					}
 
 					Args.ID = MaterialValue->PassthroughObject.GetId();
@@ -540,7 +543,7 @@ namespace UE::Mutable::Private
 		}
 		else if (Type == EOpType::MI_CONSTANT)
 		{
-			if (TMap<FName, ASTChild>::TConstIterator Iterator = ImageOperations.CreateConstIterator())
+			if (TMap<FParameterKey, ASTChild>::TConstIterator Iterator = ImageOperations.CreateConstIterator())
 			{
 				Result = Iterator.Value()->GetImageDesc();
 			}
diff --git a/Engine/Plugins/Mutable/Source/MutableTools/Private/MuT/ASTOpConstantResource.h b/Engine/Plugins/Mutable/Source/MutableTools/Private/MuT/ASTOpConstantResource.h
--- a/Engine/Plugins/Mutable/Source/MutableTools/Private/MuT/ASTOpConstantResource.h
+++ b/Engine/Plugins/Mutable/Source/MutableTools/Private/MuT/ASTOpConstantResource.h
@@ -39,7 +39,7 @@ namespace UE::Mutable::Private
 		PASSTHROUGH_ID PassthroughId = PASSTHROUGH_ID_INVALID;
 		
 		/** Image Parameters are evaluated at runtime. */
-		TMap<FName, ASTChild> ImageOperations;
+		TMap<FParameterKey, ASTChild> ImageOperations;
 
 	public:
 
diff --git a/Engine/Plugins/Mutable/Source/MutableTools/Private/MuT/ASTOpInstanceAdd.cpp b/Engine/Plugins/Mutable/Source/MutableTools/Private/MuT/ASTOpInstanceAdd.cpp
--- a/Engine/Plugins/Mutable/Source/MutableTools/Private/MuT/ASTOpInstanceAdd.cpp
+++ b/Engine/Plugins/Mutable/Source/MutableTools/Private/MuT/ASTOpInstanceAdd.cpp
@@ -38,8 +38,7 @@ namespace UE::Mutable::Private
 				value == other->value &&
 				id == other->id &&
 				ExternalId == other->ExternalId &&
-				name == other->name &&
-				LayerIndex == other->LayerIndex;
+				name == other->name;
 		}
 		return false;
 	}
@@ -54,7 +53,6 @@ namespace UE::Mutable::Private
 		n->id = id;
 		n->ExternalId = ExternalId;
 		n->name = name;
-		n->LayerIndex = LayerIndex;
 		return n;
 	}
 
@@ -73,9 +71,6 @@ namespace UE::Mutable::Private
 		switch (type)
 		{
 		case EOpType::IN_ADDMESH:
-		case EOpType::IN_ADDIMAGE:
-		case EOpType::IN_ADDVECTOR:
-		case EOpType::IN_ADDSCALAR:
 		case EOpType::IN_ADDSTRING:
 		case EOpType::IN_ADDSURFACE:
 		case EOpType::IN_ADDCOMPONENT:
@@ -110,7 +105,6 @@ namespace UE::Mutable::Private
 			Args.id = id;
 			Args.ExternalId = ExternalId;
 			Args.name = Program.AddConstant(name);
-			Args.LayerIndex = LayerIndex;
 
 			if (instance)
 			{
@@ -121,7 +115,7 @@ namespace UE::Mutable::Private
 			{
 				Args.value = value->linkedAddress;
 
-				if (type == EOpType::IN_ADDIMAGE || type == EOpType::IN_ADDMESH)
+				if (type == EOpType::IN_ADDMESH)
 				{
 					// Find out relevant parameters. \todo: this may be optimised by reusing partial
 					// values in a LINK_CONTEXT or similar
diff --git a/Engine/Plugins/Mutable/Source/MutableTools/Private/MuT/ASTOpInstanceAdd.h b/Engine/Plugins/Mutable/Source/MutableTools/Private/MuT/ASTOpInstanceAdd.h
--- a/Engine/Plugins/Mutable/Source/MutableTools/Private/MuT/ASTOpInstanceAdd.h
+++ b/Engine/Plugins/Mutable/Source/MutableTools/Private/MuT/ASTOpInstanceAdd.h
@@ -28,7 +28,6 @@ struct FProgram;
 		uint32_t id = 0;
 		uint32_t ExternalId = 0;
 		FString name;
-		int32 LayerIndex = -1; // TODO(Max) UE-351597:This will be moved to the materials so meshes won't have this value anymore.
 
 		ASTOpInstanceAdd();
 		ASTOpInstanceAdd(const ASTOpInstanceAdd&) = delete;
diff --git a/Engine/Plugins/Mutable/Source/MutableTools/Private/MuT/ASTOpMaterialBreak.cpp b/Engine/Plugins/Mutable/Source/MutableTools/Private/MuT/ASTOpMaterialBreak.cpp
--- a/Engine/Plugins/Mutable/Source/MutableTools/Private/MuT/ASTOpMaterialBreak.cpp
+++ b/Engine/Plugins/Mutable/Source/MutableTools/Private/MuT/ASTOpMaterialBreak.cpp
@@ -38,7 +38,8 @@ namespace UE::Mutable::Private
 			const ASTOpMaterialBreak* Other = static_cast<const ASTOpMaterialBreak*>(&OtherUntyped);
 			return Type == Other->Type &&
 				Material == Other->Material &&
-				ParameterName == Other->ParameterName;
+				ParameterName == Other->ParameterName &&
+				LayerIndex == Other->LayerIndex;
 		}
 		return false;
 	}
@@ -50,6 +51,7 @@ namespace UE::Mutable::Private
 		Result->Type = Type;
 		Result->ParameterName = ParameterName;
 		Result->Material = MapChild(Material.child());
+		Result->LayerIndex = LayerIndex;
 
 		return Result;
 	}
@@ -66,6 +68,7 @@ namespace UE::Mutable::Private
 		uint64 res = std::hash<uint64>()(uint64(Type));
 		hash_combine(res, GetTypeHash(ParameterName));
 		hash_combine(res, Material.child().get());
+		hash_combine(res, GetTypeHash(LayerIndex));
 		
 		return res;
 	}
@@ -95,6 +98,7 @@ namespace UE::Mutable::Private
 			}
 
 			Args.ParameterName = Program.AddConstant(ParameterName.ToString());
+			Args.LayerIndex = LayerIndex;
 
 			AppendCode(Program.ByteCode, Args);
 		}
diff --git a/Engine/Plugins/Mutable/Source/MutableTools/Private/MuT/ASTOpMaterialBreak.h b/Engine/Plugins/Mutable/Source/MutableTools/Private/MuT/ASTOpMaterialBreak.h
--- a/Engine/Plugins/Mutable/Source/MutableTools/Private/MuT/ASTOpMaterialBreak.h
+++ b/Engine/Plugins/Mutable/Source/MutableTools/Private/MuT/ASTOpMaterialBreak.h
@@ -24,6 +24,8 @@ namespace UE::Mutable::Private
 
 		FName ParameterName;
 
+		int32 LayerIndex;
+
 	public:
 
 		ASTOpMaterialBreak();
diff --git a/Engine/Plugins/Mutable/Source/MutableTools/Private/MuT/ASTOpMaterialModify.cpp b/Engine/Plugins/Mutable/Source/MutableTools/Private/MuT/ASTOpMaterialModify.cpp
new file mode 100644
--- /dev/null
+++ b/Engine/Plugins/Mutable/Source/MutableTools/Private/MuT/ASTOpMaterialModify.cpp
@@ -0,0 +1,217 @@
+// Copyright Epic Games, Inc. All Rights Reserved.
+
+#include "MuT/ASTOpMaterialModify.h"
+
+#include "MuT/ASTOpBoolEqualIntConst.h"
+#include "HAL/PlatformMath.h"
+#include "Materials/MaterialExpressionLandscapeVisibilityMask.h"
+#include "Misc/AssertionMacros.h"
+#include "MuR/ModelPrivate.h"
+#include "MuR/Parameters.h"
+#include "MuR/RefCounted.h"
+#include "MuR/Types.h"
+#include "MuT/CodeOptimiser.h"
+
+
+namespace UE::Mutable::Private
+{
+	ASTOpMaterialModify::ASTOpMaterialModify()
+		: Material(this)
+	{
+	}
+
+	ASTOpMaterialModify::~ASTOpMaterialModify()
+	{
+		// Explicit call needed to avoid recursive destruction
+		ASTOp::RemoveChildren();
+	}
+
+
+	void ASTOpMaterialModify::ForEachChild(const TFunctionRef<void(ASTChild&)> f)
+	{
+		f(Material);
+
+		for (TPair<FParameterKey, ASTChild>& Parameter : ParametersToModify)
+		{
+			f(Parameter.Value);
+		}
+	}
+
+
+	bool ASTOpMaterialModify::IsEqual(const ASTOp& OtherUntyped) const
+	{
+		if (OtherUntyped.GetOpType() == GetOpType())
+		{
+			const ASTOpMaterialModify* Other = static_cast<const ASTOpMaterialModify*>(&OtherUntyped);
+			return GetOpType() == Other->GetOpType() &&
+				Material == Other->Material &&
+				ParametersToModify == Other->ParametersToModify &&
+				ImagePropertyIndexMap == Other->ImagePropertyIndexMap;
+		}
+		return false;
+	}
+
+
+	UE::Mutable::Private::Ptr<ASTOp> ASTOpMaterialModify::Clone(MapChildFuncRef MapChild) const
+	{
+		Ptr<ASTOpMaterialModify> Result = new ASTOpMaterialModify();
+		Result->Material = MapChild(Material.child());
+
+		//Clone array of parameters
+		Result->ParametersToModify.Reserve(ParametersToModify.Num());
+		for (const TPair<FParameterKey, ASTChild>& Parameter : ParametersToModify)
+		{
+			Result->ParametersToModify.Add(Parameter.Key, ASTChild(Result, MapChild(Parameter.Value.child())));
+		}
+
+		Result->ImagePropertyIndexMap = ImagePropertyIndexMap;
+		
+		return Result;
+	}
+
+
+	EOpType ASTOpMaterialModify::GetOpType() const
+	{ 
+		return EOpType::MI_MODIFY;
+	}
+
+
+	uint64 ASTOpMaterialModify::Hash() const
+	{
+		uint64 res = std::hash<uint64>()(uint64(GetOpType()));
+		hash_combine(res, Material.child().get());
+
+		for (const TPair<FParameterKey, ASTChild>& Parameter : ParametersToModify)
+		{
+			hash_combine(res, GetTypeHash(Parameter.Key));
+			hash_combine(res, Parameter.Value.child().get());
+		}
+
+		for (const TPair<FParameterKey, int32>& Parameter : ImagePropertyIndexMap)
+		{
+			hash_combine(res, GetTypeHash(Parameter.Key));
+			hash_combine(res, GetTypeHash(Parameter.Value));
+		}
+
+		return res;
+	}
+
+
+	void ASTOpMaterialModify::Assert()
+	{
+		ASTOp::Assert();
+	}
+
+
+	void ASTOpMaterialModify::Link(FProgram& Program, FLinkerOptions*)
+	{
+		// Already linked?
+		if (!linkedAddress)
+		{
+			linkedAddress = (OP::ADDRESS)Program.OpAddress.Num();
+			Program.OpAddress.Add((uint32)Program.ByteCode.Num());
+			AppendCode(Program.ByteCode, GetOpType());
+
+			OP::MaterialModifyArgs Args;
+			FMemory::Memzero(Args);
+
+			if (Material)
+			{
+				Args.Material = Material->linkedAddress;
+			}
+
+			Args.NumParameters = ParametersToModify.Num();
+			AppendCode(Program.ByteCode, Args);
+
+			//Store parameters info directly into the bytecode
+			for (const TPair<FParameterKey, ASTChild>& Parameter : ParametersToModify)
+			{
+				//Parameter Name
+				AppendCode(Program.ByteCode,Program.AddConstant(Parameter.Key.ParameterName.ToString()));
+				
+				// Layer Index
+				AppendCode(Program.ByteCode, Parameter.Key.LayerIndex);
+
+				// Check if the exist a child operation
+				if (Parameter.Value)
+				{
+					AppendCode(Program.ByteCode, Parameter.Value->linkedAddress);
+
+					// Only images need this step
+					if (GetOpDataType(Parameter.Value.Child->GetOpType()) == EDataType::Image)
+					{
+						int32* ImagePropertyIndex = ImagePropertyIndexMap.Find(Parameter.Key);
+						check(ImagePropertyIndex) // All images must have this index
+
+						AppendCode(Program.ByteCode, *ImagePropertyIndex);
+
+						// Code from the removed operation IN_ADDIMAGE:
+						// Find out relevant parameters. \todo: this may be optimised by reusing partial
+						// values in a LINK_CONTEXT or similar
+						SubtreeRelevantParametersVisitorAST visitor;
+						visitor.Run(Parameter.Value.child());
+
+						TArray<uint16> params;
+						for (const FString& paramName : visitor.Parameters)
+						{
+							for (int32 i = 0; i < Program.Parameters.Num(); ++i)
+							{
+								const auto& param = Program.Parameters[i];
+								if (param.Name == paramName)
+								{
+									params.Add(uint16(i));
+									break;
+								}
+							}
+						}
+
+						params.Sort();
+
+						auto it = Program.ParameterLists.Find(params);
+
+						if (it != INDEX_NONE)
+						{
+							Program.RelevantParameterList.Add(Parameter.Value->linkedAddress, it);
+						}
+						else
+						{
+							Program.RelevantParameterList.Add(Parameter.Value->linkedAddress, Program.ParameterLists.Num());
+							Program.ParameterLists.Add(params);
+						}
+					}
+				}
+				else // If the child is not valid, means that we want to delete the already existing value of the parameter
+				{
+					AppendCode(Program.ByteCode, 0);
+				}
+			}
+		}
+	}
+
+
+	FImageDesc ASTOpMaterialModify::GetImageDesc(bool returnBestOption, FGetImageDescContext* Context) const
+	{
+		check(GetOpType() == EOpType::MI_MODIFY);
+
+		FImageDesc Result;
+
+		if (Material)
+		{
+			Result = Material->GetImageDesc();
+		}
+
+		return Result;
+	}
+
+
+	FSourceDataDescriptor ASTOpMaterialModify::GetSourceDataDescriptor(FGetSourceDataDescriptorContext* Context) const
+	{
+		if (Material)
+		{
+			return Material->GetSourceDataDescriptor(Context);
+		}
+
+		return {};
+	}
+}
+
diff --git a/Engine/Plugins/Mutable/Source/MutableTools/Private/MuT/ASTOpMaterialModify.h b/Engine/Plugins/Mutable/Source/MutableTools/Private/MuT/ASTOpMaterialModify.h
new file mode 100644
--- /dev/null
+++ b/Engine/Plugins/Mutable/Source/MutableTools/Private/MuT/ASTOpMaterialModify.h
@@ -0,0 +1,41 @@
+// Copyright Epic Games, Inc. All Rights Reserved.
+
+#pragma once
+
+#include "MuT/AST.h"
+
+namespace UE::Mutable::Private
+{
+	class ASTOpMaterialModify final : public ASTOp
+	{
+	public:
+
+		/** Parent Material */
+		ASTChild Material;
+
+		/** Parameters to modify. A nullptr value means that, if this operation is linked to a previous Modify operation,
+			the value set in that operation will be reset to its default value (the value set in the material)	*/
+		TMap<FParameterKey, ASTChild> ParametersToModify;
+
+		/** Index to the ImageProperties array of the model resources */
+		TMap<FParameterKey, int32> ImagePropertyIndexMap; //Cannot create a structure due to the ASTChild constructor restrictions
+
+	public:
+
+		ASTOpMaterialModify();
+		ASTOpMaterialModify(const ASTOpMaterialModify&) = delete;
+		~ASTOpMaterialModify() override;
+
+		// ASTOp interface.
+		virtual EOpType GetOpType() const override;
+		virtual uint64 Hash() const override;
+		virtual void ForEachChild(const TFunctionRef<void(ASTChild&)>) override;
+		virtual bool IsEqual(const ASTOp&) const override;
+		virtual Ptr<ASTOp> Clone(MapChildFuncRef ) const override;
+		virtual void Assert() override;
+		virtual void Link(FProgram&, FLinkerOptions*) override;
+		virtual FImageDesc GetImageDesc(bool, FGetImageDescContext*) const override;
+		virtual FSourceDataDescriptor GetSourceDataDescriptor(FGetSourceDataDescriptorContext*) const override;
+	
+	};
+}
diff --git a/Engine/Plugins/Mutable/Source/MutableTools/Private/MuT/ASTOpParameter.cpp b/Engine/Plugins/Mutable/Source/MutableTools/Private/MuT/ASTOpParameter.cpp
--- a/Engine/Plugins/Mutable/Source/MutableTools/Private/MuT/ASTOpParameter.cpp
+++ b/Engine/Plugins/Mutable/Source/MutableTools/Private/MuT/ASTOpParameter.cpp
@@ -26,7 +26,7 @@ namespace UE::Mutable::Private
 			f(c.rangeSize);
 		}
 		
-		for (TPair<FString, ASTChild>& Element : ImageOperations)
+		for (TPair<FParameterKey, ASTChild>& Element : ImageOperations)
 		{
 			f(Element.Value);
 		}
@@ -41,8 +41,8 @@ namespace UE::Mutable::Private
 			return Type == Other->Type &&
 				Parameter == Other->Parameter &&
 				Ranges == Other->Ranges &&
-				ColorParameterNames == Other->ColorParameterNames &&
-				ScalarParameterNames == Other->ScalarParameterNames &&
+				ColorParameterKeys == Other->ColorParameterKeys &&
+				ScalarParameterKeys == Other->ScalarParameterKeys &&
 				ImageOperations == Other->ImageOperations;
 		}
 		return false;
@@ -60,10 +60,10 @@ namespace UE::Mutable::Private
 			n->Ranges.Emplace(n.get(), MapChild(c.rangeSize.child()), c.rangeName, c.rangeUID);
 		}
 
-		n->ColorParameterNames = ColorParameterNames;
-		n->ScalarParameterNames = ScalarParameterNames;
+		n->ColorParameterKeys = ColorParameterKeys;
+		n->ScalarParameterKeys = ScalarParameterKeys;
 		
-		for (const TPair<FString, ASTChild>& ImageOperation : ImageOperations)
+		for (const TPair<FParameterKey, ASTChild>& ImageOperation : ImageOperations)
 		{
 			n->ImageOperations.Add(ImageOperation.Key, ASTChild(n, MapChild(ImageOperation.Value.child())));
 		}
@@ -117,15 +117,28 @@ namespace UE::Mutable::Private
 				Args.Variable = (OP::ADDRESS)LinkedParameterIndex;
 				
 				// Material Parameters
+				// Scalars
+				TArray<FString> ScalarParameterNames;
+				for (const FParameterKey& ParameterKey : ScalarParameterKeys)
+				{
+					ScalarParameterNames.Add(ParameterKey.ParameterName.ToString());
+				}
 				Args.ScalarParameterNames = Program.AddConstant(ScalarParameterNames);
+
+				// Colors
+				TArray<FString> ColorParameterNames;
+				for (const FParameterKey& ParameterKey : ColorParameterKeys)
+				{
+					ColorParameterNames.Add(ParameterKey.ParameterName.ToString());
+				}
 				Args.ColorParameterNames = Program.AddConstant(ColorParameterNames);
 
+				// Images
 				TArray<FString> ImageParameterNames;
 				TArray<uint32> ImageParameterOperations;
-
-				for (const TPair<FString, ASTChild>& Element : ImageOperations)
+				for (const TPair<FParameterKey, ASTChild>& Element : ImageOperations)
 				{
-					ImageParameterNames.Add(Element.Key);
+					ImageParameterNames.Add(Element.Key.ParameterName.ToString());
 					ImageParameterOperations.Add(Element.Value->linkedAddress);
 				}
 
@@ -142,6 +155,22 @@ namespace UE::Mutable::Private
 				}
 
 				AppendCode(Program.ByteCode, Args);
+
+				// Store layer indexes
+				for (const FParameterKey& ParameterKey : ScalarParameterKeys)
+				{
+					AppendCode(Program.ByteCode, ParameterKey.LayerIndex);
+				}
+
+				for (const FParameterKey& ParameterKey : ColorParameterKeys)
+				{
+					AppendCode(Program.ByteCode, ParameterKey.LayerIndex);
+				}
+
+				for (const TPair<FParameterKey, ASTChild>& Element : ImageOperations)
+				{
+					AppendCode(Program.ByteCode, Element.Key.LayerIndex);
+				}
 			}
 
 			else
@@ -206,12 +235,12 @@ namespace UE::Mutable::Private
 
 		if (Type == EOpType::MI_PARAMETER)
 		{
-			if (TMap<FString, ASTChild>::TConstIterator Iterator = ImageOperations.CreateConstIterator())
+			if (TMap<FParameterKey, ASTChild>::TConstIterator Iterator = ImageOperations.CreateConstIterator())
 			{
 				Result = Iterator.Value()->GetImageDesc();
 			}
 		}
-		
+
 		return Result;
 	}
 
diff --git a/Engine/Plugins/Mutable/Source/MutableTools/Private/MuT/ASTOpParameter.h b/Engine/Plugins/Mutable/Source/MutableTools/Private/MuT/ASTOpParameter.h
--- a/Engine/Plugins/Mutable/Source/MutableTools/Private/MuT/ASTOpParameter.h
+++ b/Engine/Plugins/Mutable/Source/MutableTools/Private/MuT/ASTOpParameter.h
@@ -30,11 +30,11 @@ namespace UE::Mutable::Private
 		TArray<FRangeData> Ranges;
 
 		/** List of all the parameters that this material can be broken into. */
-		TArray<FString> ScalarParameterNames;
-		TArray<FString> ColorParameterNames;
+		TArray<FParameterKey> ScalarParameterKeys;
+		TArray<FParameterKey> ColorParameterKeys;
 
 		/** Image Parameters are evaluated at runtime. */
-		TMap<FString, ASTChild> ImageOperations;
+		TMap<FParameterKey, ASTChild> ImageOperations;
 
 		/** Index of the parameter in the final program parameter list. 
 		* This is generated at link time, because the parameters may be reordered. 
diff --git a/Engine/Plugins/Mutable/Source/MutableTools/Private/MuT/CodeGenerator.cpp b/Engine/Plugins/Mutable/Source/MutableTools/Private/MuT/CodeGenerator.cpp
--- a/Engine/Plugins/Mutable/Source/MutableTools/Private/MuT/CodeGenerator.cpp
+++ b/Engine/Plugins/Mutable/Source/MutableTools/Private/MuT/CodeGenerator.cpp
@@ -33,6 +33,7 @@
 #include "MuT/ASTOpImageBlankLayout.h"
 #include "MuT/ASTOpImagePlainColor.h"
 #include "MuT/ASTOpInstanceAdd.h"
+#include "MuT/ASTOpMaterialModify.h"
 #include "MuT/ASTOpMeshBindShape.h"
 #include "MuT/ASTOpMeshClipDeform.h"
 #include "MuT/ASTOpMeshClipMorphPlane.h"
@@ -944,8 +945,9 @@ namespace UE::Mutable::Private
 		int32 ComponentId = Options.Component ? Options.Component->Id : -1;
 		GetModifiersFor(ComponentId, SurfaceNode->Tags, bModifiersForBeforeOperations, Modifiers);
 
+		// TODO (Max): UE-354034 Fix this function and uncomment it
 		// This pass on the modifiers is only to detect errors that cannot be detected at the point they are applied.
-		CheckModifiersForSurface(*SurfaceNode, Modifiers, Options.LODIndex);
+		//CheckModifiersForSurface(*SurfaceNode, Modifiers, Options.LODIndex);
 		
         // Generate the mesh
 		FMeshGenerationStaticOptions MeshStaticOptions(ComponentId, Options.LODIndex);
@@ -1211,457 +1213,37 @@ namespace UE::Mutable::Private
 				}
 
 				// Build a series of operations to assemble the surface
+				//-------------------------------------------------------
+
 				Ptr<ASTOp> LastSurfOp;
 
-				// Create the expression for each texture, if we are not reusing the surface from another LOD.
+				// Create the expression for each surface parameter, if we are not reusing the surface from another LOD.
 				//------------------------------------------------------------------------
 				if (!SharedSurfaceResults)
 				{
-					for (int32 ImageIndex = 0; ImageIndex < SurfaceNode->Images.Num(); ++ImageIndex)
-					{
-						MUTABLE_CPUPROFILER_SCOPE(SurfaceTexture);
-
-						// Any image-specific format or mipmapping needs to be applied at the end
-						Ptr<NodeImageMipmap> mipmapNode;
-						Ptr<NodeImageFormat> formatNode;
-						Ptr<NodeImageSwizzle> swizzleNode;
-
-						bool bFound = false;
-						Ptr<NodeImage> pImageNode = SurfaceNode->Images[ImageIndex].Image;
-
-						while (!bFound && pImageNode)
-						{
-							if (pImageNode->GetType() == NodeImageMipmap::GetStaticType())
-							{
-								NodeImageMipmap* tm = static_cast<NodeImageMipmap*>(pImageNode.get());
-								if (!mipmapNode) mipmapNode = tm;
-								pImageNode = tm->Source;
-							}
-							else if (pImageNode->GetType() == NodeImageFormat::GetStaticType())
-							{
-								NodeImageFormat* tf = static_cast<NodeImageFormat*>(pImageNode.get());
-								if (!formatNode) formatNode = tf;
-								pImageNode = tf->Source;
-							}
-							else if (pImageNode->GetType() == NodeImageSwizzle::GetStaticType())
-							{
-								NodeImageSwizzle* ts = static_cast<NodeImageSwizzle*>(pImageNode.get());
-
-								if (!ts->Sources.IsEmpty())
-								{
-									NodeImage* Source = ts->Sources[0].get();
-
-									bool bAllSourcesAreTheSame = true;
-									for (int32 SourceIndex = 1; SourceIndex < ts->Sources.Num(); ++SourceIndex)
-									{
-										bAllSourcesAreTheSame = bAllSourcesAreTheSame && (Source == ts->Sources[SourceIndex]);
-									}
-
-									if (!swizzleNode && bAllSourcesAreTheSame)
-									{
-										swizzleNode = ts;
-										pImageNode = Source;
-									}
-									else
-									{
-										bFound = true;
-									}
-								}
-								else
-								{
-									// break loop if swizzle has no sources.
-									bFound = true;
-								}
-							}
-							else
-							{
-								bFound = true;
-							}
-						}
-
-						if (bFound)
-						{
-							const NodeSurfaceNew::FImageData& ImageData = SurfaceNode->Images[ImageIndex];
-
-							const int32 LayoutIndex = ImageData.LayoutIndex;
-
-							// If the layout index has been set to negative, it means we should ignore the layout for this image.
-							CompilerOptions::TextureLayoutStrategy ImageLayoutStrategy = (LayoutIndex < 0)
-								? CompilerOptions::TextureLayoutStrategy::None
-								: CompilerOptions::TextureLayoutStrategy::Pack
-								;
-
-							if (ImageLayoutStrategy == CompilerOptions::TextureLayoutStrategy::None)
-							{
-								// Generate the image
-								FImageGenerationOptions ImageOptions(ComponentId, Options.LODIndex);
-								ImageOptions.State = Options.State;
-								ImageOptions.ImageLayoutStrategy = ImageLayoutStrategy;
-								ImageOptions.ActiveTags = SurfaceNode->Tags;
-								ImageOptions.RectSize = { 0, 0 };
-
-								// TODO: To tasks
-								FImageGenerationResult Result;
-								GenerateImage(ImageOptions, Result, pImageNode);
-								Ptr<ASTOp> imageAd = Result.op;
-
-								// Placeholder block. Ideally this should be the actual image size
-								constexpr int32 FakeLayoutSize = 256;
-								FIntPoint GridSize(FakeLayoutSize, FakeLayoutSize);
-								FLayoutBlockDesc LayoutBlockDesc;
-								LayoutBlockDesc.BlockPixelsX = 1;
-								LayoutBlockDesc.BlockPixelsY = 1;
-								box< FIntVector2 > RectInCells;
-								RectInCells.min = { 0,0 };
-								RectInCells.size = { FakeLayoutSize ,FakeLayoutSize };
-
-								imageAd = ApplyImageBlockModifiers(Modifiers, ImageOptions, imageAd, ImageData, GridSize, LayoutBlockDesc, RectInCells, SurfaceNode->GetMessageContext());
-
-								check(imageAd);
-
-								if (swizzleNode)
-								{
-									Ptr<ASTOpImageSwizzle> fop = new ASTOpImageSwizzle();
-									fop->Format = swizzleNode->NewFormat;
-									fop->Sources[0] = imageAd;
-									fop->Sources[1] = imageAd;
-									fop->Sources[2] = imageAd;
-									fop->Sources[3] = imageAd;
-									fop->SourceChannels[0] = swizzleNode->SourceChannels[0];
-									fop->SourceChannels[1] = swizzleNode->SourceChannels[1];
-									fop->SourceChannels[2] = swizzleNode->SourceChannels[2];
-									fop->SourceChannels[3] = swizzleNode->SourceChannels[3];
-									check(fop->Format != EImageFormat::None);
-									imageAd = fop;
-								}
-
-								if (mipmapNode)
-								{
-									Ptr<ASTOpImageMipmap> op = new ASTOpImageMipmap();
-									op->Levels = 0;
-									op->Source = imageAd;
-									op->BlockLevels = 0;
-
-									op->AddressMode = mipmapNode->Settings.AddressMode;
-									op->FilterType = mipmapNode->Settings.FilterType;
-									imageAd = op;
-								}
-
-								if (formatNode)
-								{
-									Ptr<ASTOpImagePixelFormat> fop = new ASTOpImagePixelFormat();
-									fop->Format = formatNode->Format;
-									fop->FormatIfAlpha = formatNode->FormatIfAlpha;
-									fop->Source = imageAd;
-									check(fop->Format != EImageFormat::None);
-									imageAd = fop;
-								}
-
-								Ptr<ASTOpInstanceAdd> op = new ASTOpInstanceAdd();
-								op->type = EOpType::IN_ADDIMAGE;
-								op->instance = LastSurfOp;
-								op->value = imageAd;
-								op->name = SurfaceNode->Images[ImageIndex].ParameterInfo.ParameterName;
-								op->LayerIndex = SurfaceNode->Images[ImageIndex].ParameterInfo.MaterialLayerIndex;
-
-								LastSurfOp = op;
-							}
-
-							else if (ImageLayoutStrategy == CompilerOptions::TextureLayoutStrategy::Pack) //-V547
-							{
-								if (LayoutIndex >= MeshResults.GeneratedLayouts.Num() ||
-									LayoutIndex >= MeshResults.LayoutOps.Num())
-								{
-									ErrorLog->Add("Missing layout in object, or its parent.", ELMT_ERROR, SurfaceNode->GetMessageContext());
-								}
-								else
-								{
-									TSharedPtr<const FLayout> pLayout = MeshResults.GeneratedLayouts[LayoutIndex].Layout;
-									check(pLayout);
-
-									Ptr<ASTOpInstanceAdd> op = new ASTOpInstanceAdd();
-									op->type = EOpType::IN_ADDIMAGE;
-									op->instance = LastSurfOp;
-
-									// Image
-									//-------------------------------------
-
-									// Size of a layout block in pixels
-									FIntPoint GridSize = pLayout->GetGridSize();
-
-									// Try to guess the layout block description from the first valid block that is generated.
-									FLayoutBlockDesc LayoutBlockDesc;
-									if (formatNode)
-									{
-										LayoutBlockDesc.FinalFormat = formatNode->FormatIfAlpha;
-										if (LayoutBlockDesc.FinalFormat == EImageFormat::None)
-										{
-											LayoutBlockDesc.FinalFormat = formatNode->Format;
-										}
-									}
-
-									bool bImageSizeWarning = false;
-
-									// Start with a blank image. It will be completed later with the blockSize, format and mips information
-									Ptr<ASTOpImageBlankLayout> BlankImageOp;
-									Ptr<ASTOp> imageAd;
-									{
-										BlankImageOp = new ASTOpImageBlankLayout();
-										BlankImageOp->Layout = MeshResults.LayoutOps[LayoutIndex];
-										// The rest ok the op will be completed below
-										BlankImageOp->MipmapCount = 0;
-										imageAd = BlankImageOp;
-									}
-
-									// Skip the block addition for this image if the layout was from a extension.
-									if (!LayoutFromExtension[LayoutIndex])
-									{
-										for (int32 BlockIndex = 0; BlockIndex < pLayout->GetBlockCount(); ++BlockIndex)
-										{
-											// Generate the image
-											FImageGenerationOptions ImageOptions(ComponentId, Options.LODIndex);
-											ImageOptions.State = Options.State;
-											ImageOptions.ImageLayoutStrategy = ImageLayoutStrategy;
-											ImageOptions.RectSize = { 0,0 };
-											ImageOptions.ActiveTags = SurfaceNode->Tags;
-											ImageOptions.LayoutToApply = pLayout;
-											ImageOptions.LayoutBlockId = pLayout->Blocks[BlockIndex].Id;
-											FImageGenerationResult ImageResult;
-											GenerateImage(ImageOptions, ImageResult, pImageNode);
-											Ptr<ASTOp> blockAd = ImageResult.op;
-
-											if (!blockAd)
-											{
-												// The GenerateImage(...) above has failed, skip this block
-												continue;
-											}
-
-											// Calculate the desc of the generated block.
-											constexpr bool bReturnBestOption = true;
-											FImageDesc BlockDesc = blockAd->GetImageDesc(bReturnBestOption, nullptr);
-
-											// Block in layout grid units (cells)
-											box< FIntVector2 > RectInCells;
-											RectInCells.min = pLayout->Blocks[BlockIndex].Min;
-											RectInCells.size = pLayout->Blocks[BlockIndex].Size;
-
-											// Try to update the layout block desc if we don't know it yet.
-											UpdateLayoutBlockDesc(LayoutBlockDesc, BlockDesc, RectInCells.size);
-
-											// Even if we force the size afterwards, we need some size hint in some cases, like image projections.
-											ImageOptions.RectSize = UE::Math::TIntVector2<int32>(BlockDesc.m_size);
-
-											blockAd = ApplyImageBlockModifiers(Modifiers, ImageOptions, blockAd, ImageData, GridSize, LayoutBlockDesc, RectInCells, SurfaceNode->GetMessageContext());
-
-											// Enforce block size and optimizations
-											blockAd = GenerateImageSize(blockAd, FIntVector2(BlockDesc.m_size));
-
-											EImageFormat baseFormat = imageAd->GetImageDesc().m_format;
-											// Actually don't do it, it will be propagated from the top format operation.
-											//Ptr<ASTOp> blockAd = GenerateImageFormat(blockAd, baseFormat);
-
-											// Apply tiling to avoid generating chunks of image that are too big.
-											blockAd = ApplyTiling(blockAd, ImageOptions.RectSize, LayoutBlockDesc.FinalFormat);
-
-											// Compose layout operation
-											Ptr<ASTOpImageCompose> composeOp = new ASTOpImageCompose();
-											composeOp->Layout = MeshResults.LayoutOps[LayoutIndex];
-											composeOp->Base = imageAd;
-											composeOp->BlockImage = blockAd;
-
-											// Set the absolute block index.
-											check(pLayout->Blocks[BlockIndex].Id != FLayoutBlock::InvalidBlockId);
-											composeOp->BlockId = pLayout->Blocks[BlockIndex].Id;
-
-											imageAd = composeOp;
-										}
-									}
-									check(imageAd);
-
-									FMeshGenerationStaticOptions ModifierOptions(ComponentId, Options.LODIndex);
-									ModifierOptions.State = Options.State;
-									ModifierOptions.ActiveTags = SurfaceNode->Tags;
-									imageAd = ApplyImageExtendModifiers(Modifiers, ModifierOptions, MeshResults, imageAd, ImageLayoutStrategy,
-										LayoutIndex, ImageData, GridSize, LayoutBlockDesc,
-										SurfaceNode->GetMessageContext());
-
-									// Complete the base op
-									BlankImageOp->BlockSize[0] = uint16(LayoutBlockDesc.BlockPixelsX);
-									BlankImageOp->BlockSize[1] = uint16(LayoutBlockDesc.BlockPixelsY);
-									BlankImageOp->Format = GetUncompressedFormat(LayoutBlockDesc.FinalFormat);
-									BlankImageOp->GenerateMipmaps = LayoutBlockDesc.bBlocksHaveMips;
-									BlankImageOp->MipmapCount = 0;
-
-									if (swizzleNode)
-									{
-										Ptr<ASTOpImageSwizzle> fop = new ASTOpImageSwizzle();
-										fop->Format = swizzleNode->NewFormat;
-
-										for (int32 ChannelIndex = 0; ChannelIndex < swizzleNode->SourceChannels.Num(); ++ChannelIndex)
-										{
-											fop->Sources[ChannelIndex] = imageAd;
-											fop->SourceChannels[ChannelIndex] = swizzleNode->SourceChannels[ChannelIndex];
-										}
-										check(fop->Format != EImageFormat::None);
-										imageAd = fop;
-									}
-
-									// Apply mipmap and format if necessary, skip if format is None (possibly because a block was skipped above)
-									bool bNeedsMips =
-										(mipmapNode && LayoutBlockDesc.FinalFormat != EImageFormat::None)
-										||
-										LayoutBlockDesc.bBlocksHaveMips;
-
-									if (bNeedsMips)
-									{
-										Ptr<ASTOpImageMipmap> mop = new ASTOpImageMipmap();
-
-										// At the end of the day, we want all the mipmaps. Maybe the code
-										// optimiser will split the process later.
-										mop->Levels = 0;
-										mop->bOnlyTail = false;
-										mop->Source = imageAd;
-
-										// We have to avoid mips smaller than the image format block size, so
-										// we will devide the layout block by the format block
-										const FImageFormatData& PixelFormatInfo = GetImageFormatData(LayoutBlockDesc.FinalFormat);
-
-										int32 mipsX = FMath::CeilLogTwo(LayoutBlockDesc.BlockPixelsX / PixelFormatInfo.PixelsPerBlockX);
-										int32 mipsY = FMath::CeilLogTwo(LayoutBlockDesc.BlockPixelsY / PixelFormatInfo.PixelsPerBlockY);
-										mop->BlockLevels = (uint8)FMath::Max(mipsX, mipsY);
-
-										if (LayoutBlockDesc.BlockPixelsX < PixelFormatInfo.PixelsPerBlockX || LayoutBlockDesc.BlockPixelsY < PixelFormatInfo.PixelsPerBlockY)
-										{
-											// In this case, the mipmap will never be useful for blocks, so we indicate that
-											// it should make the mips at the root of the expression.
-											mop->bOnlyTail = true;
-										}
-
-										mop->AddressMode = EAddressMode::ClampToEdge;
-										mop->FilterType = EMipmapFilterType::SimpleAverage;
-
-										if (mipmapNode)
-										{
-											mop->AddressMode = mipmapNode->Settings.AddressMode;
-											mop->FilterType = mipmapNode->Settings.FilterType;
-										}
-
-										imageAd = mop;
-									}
-
-									if (formatNode)
-									{
-										Ptr<ASTOpImagePixelFormat> fop = new ASTOpImagePixelFormat();
-										fop->Format = formatNode->Format;
-										fop->FormatIfAlpha = formatNode->FormatIfAlpha;
-										fop->Source = imageAd;
-										check(fop->Format != EImageFormat::None);
-										imageAd = fop;
-									}
-
-									op->value = imageAd;
-
-									// Name
-									op->name = SurfaceNode->Images[ImageIndex].ParameterInfo.ParameterName;
-									op->LayerIndex = SurfaceNode->Images[ImageIndex].ParameterInfo.MaterialLayerIndex;
-
-									LastSurfOp = op;
-								}
-							}
-
-							else
-							{
-								// Unimplemented texture layout strategy
-								check(false);
-							}
-						}
-					}
-
-					// Create the expression for each vector
+					// Create the expression for the Material. This material will be the base material of the instance
 					//------------------------------------------------------------------------
-					for (int32 t = 0; t < SurfaceNode->Vectors.Num(); ++t)
-					{
-						//MUTABLE_CPUPROFILER_SCOPE(SurfaceVector);
-
-						if (Ptr<NodeColour> pVectorNode = SurfaceNode->Vectors[t].Vector)
-						{
-							Ptr<ASTOpInstanceAdd> op = new ASTOpInstanceAdd();
-							op->type = EOpType::IN_ADDVECTOR;
-							op->instance = LastSurfOp;
-
-							// Vector
-							FColorGenerationResult VectorResult;
-							GenerateColor(VectorResult, Options, pVectorNode);
-							op->value = VectorResult.op;
+					check(SurfaceNode->Material)
 
-							// Name
-							op->name = SurfaceNode->Vectors[t].ParameterInfo.ParameterName;
-							op->LayerIndex = SurfaceNode->Vectors[t].ParameterInfo.MaterialLayerIndex;
-
-							LastSurfOp = op;
-						}
-					}
-
-					// Create the expression for each scalar
-					//------------------------------------------------------------------------
-					for (int32 t = 0; t < SurfaceNode->Scalars.Num(); ++t)
-					{
-						// MUTABLE_CPUPROFILER_SCOPE(SurfaceScalar);
-
-						if (Ptr<NodeScalar> ScalarNode = SurfaceNode->Scalars[t].Scalar)
-						{
-							Ptr<ASTOpInstanceAdd> op = new ASTOpInstanceAdd();
-							op->type = EOpType::IN_ADDSCALAR;
-							op->instance = LastSurfOp;
-
-							// Scalar
-							FScalarGenerationResult ScalarResult;
-							GenerateScalar(ScalarResult, Options, ScalarNode);
-							op->value = ScalarResult.op;
-
-							// Name
-							op->name = SurfaceNode->Scalars[t].ParameterInfo.ParameterName;
-							op->LayerIndex = SurfaceNode->Scalars[t].ParameterInfo.MaterialLayerIndex;
-
-							LastSurfOp = op;
-						}
-					}
-
-					// Create the expression for each string
-					//------------------------------------------------------------------------
-					for (int32 t = 0; t < SurfaceNode->Strings.Num(); ++t)
-					{
-						if (Ptr<NodeString> StringNode = SurfaceNode->Strings[t].String)
-						{
-							Ptr<ASTOpInstanceAdd> op = new ASTOpInstanceAdd();
-							op->type = EOpType::IN_ADDSTRING;
-							op->instance = LastSurfOp;
-
-							FStringGenerationResult StringResult;
-							GenerateString(StringResult, Options, StringNode);
-							op->value = StringResult.op;
+					FMaterialGenerationOptions MaterialOptions;
+					MaterialOptions.State = Options.State;
+					MaterialOptions.ComponentId = ComponentId;
+					MaterialOptions.LODIndex = Options.LODIndex;
+					MaterialOptions.Tags = SurfaceNode->Tags;
+					MaterialOptions.Modifiers = Modifiers;
+					MaterialOptions.MeshResults = MeshResults;
+					MaterialOptions.LayoutFromExtension = LayoutFromExtension;
 
-							// Name
-							op->name = SurfaceNode->Strings[t].Name;
+					FMaterialGenerationResult MaterialResult;
+					GenerateMaterial(MaterialOptions, MaterialResult, SurfaceNode->Material);
 
-							LastSurfOp = op;
-						}
-					}
+					// Operation to add the material to the instance
+					Ptr<ASTOpAddMaterial> AddMaterial = new ASTOpAddMaterial();
+					AddMaterial->Type = EOpType::IN_ADDMATERIAL;
+					AddMaterial->Instance = nullptr; // This will force the generation of an FMaterial
+					AddMaterial->Material = MaterialResult.op;
 
-					// Create the expression for the Material
-					//------------------------------------------------------------------------
-					if (SurfaceNode->Material)
-					{
-						Ptr<ASTOpAddMaterial> Op = new ASTOpAddMaterial();
-						Op->Type = EOpType::IN_ADDMATERIAL;
-						Op->Instance = LastSurfOp;
-
-						FMaterialGenerationOptions MaterialOptions;
-						FMaterialGenerationResult MaterialResult;
-						GenerateMaterial(MaterialOptions, MaterialResult, SurfaceNode->Material);
-						Op->Material = MaterialResult.op;
-						LastSurfOp = Op;
-					}
+					LastSurfOp = AddMaterial;
 				}
 				else
 				{
@@ -3023,7 +2605,7 @@ namespace UE::Mutable::Private
 	Ptr<ASTOp> CodeGenerator::ApplyImageBlockModifiers(
 		const TArray<FirstPassGenerator::FModifier>& Modifiers,
 		const FImageGenerationOptions& Options, Ptr<ASTOp> BaseImageOp, 
-		const NodeSurfaceNew::FImageData& ImageData,
+		const FParameterKey& ImageParameterKey,
 		FIntPoint GridSize,
 		const FLayoutBlockDesc& LayoutBlockDesc,
 		box< FIntVector2 > RectInCells,
@@ -3051,7 +2633,7 @@ namespace UE::Mutable::Private
 				const NodeModifierSurfaceEdit::FTexture* MatchingEdit = Edit->LODs[CurrentLOD].Textures.FindByPredicate(
 					[&](const NodeModifierSurfaceEdit::FTexture& Candidate)
 					{
-						return (Candidate.MaterialParameterName == ImageData.MaterialParameterName);
+						return (Candidate.MaterialParameterKey == ImageParameterKey);
 					});
 
 				if ( !MatchingEdit)
@@ -3105,7 +2687,7 @@ namespace UE::Mutable::Private
 		Ptr<ASTOp> BaseImageOp, 
 		CompilerOptions::TextureLayoutStrategy ImageLayoutStrategy,
 		int32 LayoutIndex, 
-		const NodeSurfaceNew::FImageData& ImageData,
+		const FParameterKey& ImageParameterKey,
 		FIntPoint GridSize, 
 		CodeGenerator::FLayoutBlockDesc& InOutLayoutBlockDesc,
 		const void* ModifiedNodeErrorContext)
@@ -3135,7 +2717,7 @@ namespace UE::Mutable::Private
 				const NodeModifierSurfaceEdit::FTexture* MatchingEdit = Edit->LODs[CurrentLOD].Textures.FindByPredicate(
 					[&](const NodeModifierSurfaceEdit::FTexture& Candidate)
 					{
-						return (Candidate.MaterialParameterName == ImageData.MaterialParameterName);
+						return (Candidate.MaterialParameterKey == ImageParameterKey);
 					});
 
 				if (!MatchingEdit || (MatchingEdit && !MatchingEdit->Extend) )
@@ -3143,7 +2725,13 @@ namespace UE::Mutable::Private
 					if (Edit->LODs[CurrentLOD].MeshAdd)
 					{
 						// When extending a mesh section it is mandatory to provide textures for all section textures handled by Mutable.
-						FString Msg = FString::Printf(TEXT("Required texture [%s] is missing when trying to extend a mesh section."), *ImageData.MaterialParameterName);
+						FString Msg = FString::Printf(TEXT("Required texture [%s] is missing when trying to extend a mesh section."), *ImageParameterKey.ParameterName.ToString());
+
+						if (ImageParameterKey.LayerIndex >= 0)
+						{
+							Msg = FString::Printf(TEXT("Required texture [%s - Layer %d] is missing when trying to extend a mesh section."), *ImageParameterKey.ParameterName.ToString(), ImageParameterKey.LayerIndex);
+						}
+
 						ErrorLog->Add(Msg, ELMT_INFO, Edit->GetMessageContext(), ModifiedNodeErrorContext);
 					}
 
@@ -3246,9 +2834,10 @@ namespace UE::Mutable::Private
 		for (const FirstPassGenerator::FModifier& Mod : Modifiers)
 		{
 			// A mistake in the surface edit modifier usually results in no change visible. Try to detect it.
-			if (Mod.Node->GetType() == NodeModifierSurfaceEdit::GetStaticType())
+			if (Mod.Node->GetType() == NodeModifierSurfaceEdit::GetStaticType() && Node.Material->GetType() == NodeMaterialModify::GetStaticType())
 			{
 				const NodeModifierSurfaceEdit* Edit = static_cast<const NodeModifierSurfaceEdit*>(Mod.Node);
+				const NodeMaterialModify* Material = static_cast<const NodeMaterialModify*>(Node.Material.get());
 
 				bool bAffectsCurrentLOD = Edit->LODs.IsValidIndex(CurrentLOD);
 				if (!bAffectsCurrentLOD)
@@ -3256,19 +2845,19 @@ namespace UE::Mutable::Private
 					continue;
 				}
 
-				if (Node.Images.IsEmpty() || Edit->LODs[CurrentLOD].Textures.IsEmpty())
+				if (Material->ImageParameters.IsEmpty() || Edit->LODs[CurrentLOD].Textures.IsEmpty())
 				{
 					continue;
 				}
 
 				bool bAtLeastSomeTexture = false;
 
-				for (NodeSurfaceNew::FImageData Data : Node.Images)
+				for (TPair<FParameterKey, FImageParameterData> Data : Material->ImageParameters)
 				{
 					const NodeModifierSurfaceEdit::FTexture* MatchingEdit = Edit->LODs[CurrentLOD].Textures.FindByPredicate(
 						[&](const NodeModifierSurfaceEdit::FTexture& Candidate)
 						{
-							return (Candidate.MaterialParameterName == Data.MaterialParameterName);
+							return (Candidate.MaterialParameterKey == Data.Key);
 						});
 
 					if (MatchingEdit)
diff --git a/Engine/Plugins/Mutable/Source/MutableTools/Private/MuT/CodeGenerator.h b/Engine/Plugins/Mutable/Source/MutableTools/Private/MuT/CodeGenerator.h
--- a/Engine/Plugins/Mutable/Source/MutableTools/Private/MuT/CodeGenerator.h
+++ b/Engine/Plugins/Mutable/Source/MutableTools/Private/MuT/CodeGenerator.h
@@ -37,6 +37,7 @@
 #include "MuT/NodeProjector.h"
 #include "MuT/NodeString.h"
 #include "MuT/NodeSurfaceNew.h"
+#include "MuT/NodeMaterialModify.h"
 #include "MuT/NodeModifierSurfaceEdit.h"
 #include "MuT/Table.h"
 #include "MuT/TablePrivate.h"
@@ -60,6 +61,7 @@ namespace UE::Mutable::Private
 {
 	class NodeImageExternal;
 	class NodeMeshExternal;
+	class ASTOpMaterialModify;
 	class ASTOpParameter;
 	class FLayout;
 	class NodeColourArithmeticOperation;
@@ -98,6 +100,7 @@ namespace UE::Mutable::Private
 	class NodeMaterial;
 	class NodeMaterialBreak;
 	class NodeMaterialConstant;
+	class NodeMaterialModify;
 	class NodeMaterialTable;
 	class NodeMaterialSwitch;
 	class NodeMaterialVariation;
@@ -610,6 +613,16 @@ namespace UE::Mutable::Private
 			/** The LOD being generated. */
 			int32 LODIndex = -1;
 
+			/** Array with all the tags for the modifiers. */
+			TArray<FString> Tags;
+
+			/** Modifiers to apply. */
+			TArray<FirstPassGenerator::FModifier> Modifiers;
+
+			FMeshGenerationResult MeshResults;
+			
+			TBitArray<> LayoutFromExtension;
+
 			/** */
 			CompilerOptions::TextureLayoutStrategy ImageLayoutStrategy = CompilerOptions::TextureLayoutStrategy::None;
 
@@ -661,6 +674,7 @@ namespace UE::Mutable::Private
 		void GenerateMaterial_Parameter(const FMaterialGenerationOptions&, FMaterialGenerationResult& result, const NodeMaterialParameter* node);
 		void GenerateMaterial_External(const FMaterialGenerationOptions&, FMaterialGenerationResult& result, const NodeMaterialExternal* node);
     	void GenerateMaterial_BreakSkeletalMesh(const FMaterialGenerationOptions&, FMaterialGenerationResult& OutResult, const NodeMaterialSkeletalMeshBreak* InNode);
+		void GenerateMaterial_Modify(const FMaterialGenerationOptions&, FMaterialGenerationResult& result, const NodeMaterialModify* node);
 
 		/** */
 		Ptr<ASTOp> GenerateMissingMaterialCode(const TCHAR* strWhere, const void* ErrorContext);
@@ -1167,7 +1181,7 @@ namespace UE::Mutable::Private
 			const TArray<FirstPassGenerator::FModifier>&, 
 			const FImageGenerationOptions&, 
 			Ptr<ASTOp> BaseImageOp,
-			const NodeSurfaceNew::FImageData& ImageData,
+			const FParameterKey& ImageParameterKey,
 			FIntPoint GridSize,
 			const FLayoutBlockDesc& LayoutBlockDesc,
 			box< FIntVector2 > RectInCells,
@@ -1178,15 +1192,28 @@ namespace UE::Mutable::Private
 			const FMeshGenerationStaticOptions& Options,
 			const FMeshGenerationResult& BaseMeshResults,
 			Ptr<ASTOp> ImageAd, 
-			CompilerOptions::TextureLayoutStrategy ImageLayoutStrategy, 
+			CompilerOptions::TextureLayoutStrategy ImageLayoutStrategy,
 			int32 LayoutIndex, 
-			const NodeSurfaceNew::FImageData& ImageData,
+			const FParameterKey& ImageParameterKey,
 			FIntPoint GridSize,
 			CodeGenerator::FLayoutBlockDesc& InOutLayoutBlockDesc,
 			const void* ModifiedNodeErrorContext);
 
 		void CheckModifiersForSurface(const NodeSurfaceNew&, const TArray<FirstPassGenerator::FModifier>&, int32 LODIndex);
 
+
+		void GenerateMaterialImage(
+			const TPair<FParameterKey, FImageParameterData>& ImageData,
+			FImageGenerationResult& OutImageGenerated,
+			const TArray<FirstPassGenerator::FModifier>& Modifiers,
+			const int32 State,
+			const int32 ComponentId,
+			const int32 LODIndex,
+			const TArray<FString>& Tags,
+			const FMeshGenerationResult& MeshResults,
+			const TBitArray<>& LayoutFromExtension,
+			const void* ErrorContext);
+
     };
 
 	
diff --git a/Engine/Plugins/Mutable/Source/MutableTools/Private/MuT/CodeGenerator_Color.cpp b/Engine/Plugins/Mutable/Source/MutableTools/Private/MuT/CodeGenerator_Color.cpp
--- a/Engine/Plugins/Mutable/Source/MutableTools/Private/MuT/CodeGenerator_Color.cpp
+++ b/Engine/Plugins/Mutable/Source/MutableTools/Private/MuT/CodeGenerator_Color.cpp
@@ -560,6 +560,7 @@ namespace UE::Mutable::Private
 
 		// Set the parameter name that this break node will generate
 		MaterialOp->ParameterName = Typed->ParameterName;
+		MaterialOp->LayerIndex = Typed->LayerIndex;
 
 		// Generate Material Source
 		FMaterialGenerationOptions MaterialSourceOptions;
diff --git a/Engine/Plugins/Mutable/Source/MutableTools/Private/MuT/CodeGenerator_Image.cpp b/Engine/Plugins/Mutable/Source/MutableTools/Private/MuT/CodeGenerator_Image.cpp
--- a/Engine/Plugins/Mutable/Source/MutableTools/Private/MuT/CodeGenerator_Image.cpp
+++ b/Engine/Plugins/Mutable/Source/MutableTools/Private/MuT/CodeGenerator_Image.cpp
@@ -1900,13 +1900,14 @@ namespace UE::Mutable::Private
 
 		// Set the parameter name that this break node will generate
 		MaterialOp->ParameterName = Node->ParameterName;
+		MaterialOp->LayerIndex = Node->LayerIndex;
 
 		// Generate Material Source
 		FMaterialGenerationResult MaterialSourceResult;
 		FMaterialGenerationOptions MaterialSourceOptions;
-		// Temporary fix
 		MaterialSourceOptions.ComponentId = InOptions.ComponentId;
 		MaterialSourceOptions.LODIndex = InOptions.LODIndex;
+		MaterialSourceOptions.State = InOptions.State;
 		MaterialSourceOptions.ImageLayoutStrategy = InOptions.ImageLayoutStrategy;
 		MaterialSourceOptions.RectSize = InOptions.RectSize;
 		MaterialSourceOptions.LayoutBlockId = InOptions.LayoutBlockId;
diff --git a/Engine/Plugins/Mutable/Source/MutableTools/Private/MuT/CodeGenerator_Material.cpp b/Engine/Plugins/Mutable/Source/MutableTools/Private/MuT/CodeGenerator_Material.cpp
--- a/Engine/Plugins/Mutable/Source/MutableTools/Private/MuT/CodeGenerator_Material.cpp
+++ b/Engine/Plugins/Mutable/Source/MutableTools/Private/MuT/CodeGenerator_Material.cpp
@@ -12,8 +12,12 @@
 #include "MuT/ASTOpConditional.h"
 #include "MuT/ASTOpConstantResource.h"
 #include "MuT/ASTOpImagePixelFormat.h"
+#include "MuT/ASTOpImageBlankLayout.h"
+#include "MuT/ASTOpImageCompose.h"
+#include "MuT/ASTOpImageMipmap.h"
+#include "MuT/ASTOpImageSwizzle.h"
 #include "MuT/ASTOpImageResize.h"
-#include "MuT/ASTOpMaterialBreak.h"
+#include "MuT/ASTOpMaterialModify.h"
 #include "MuT/ASTOpSwitch.h"
 #include "MuT/CodeGenerator.h"
 #include "MuT/CodeGenerator_FirstPass.h"
@@ -28,6 +32,10 @@
 #include "MuT/NodeMaterialTable.h"
 #include "MuT/NodeMaterialParameter.h"
 #include "MuT/NodeMaterialSkeletalMeshBreak.h"
+#include "MuT/NodeMaterialModify.h"
+#include "MuT/NodeImageFormat.h"
+#include "MuT/NodeImageMipmap.h"
+#include "MuT/NodeImageSwizzle.h"
 #include "MuT/Table.h"
 #include "MuT/TablePrivate.h"
 
@@ -87,6 +95,10 @@ namespace UE::Mutable::Private
 			GenerateMaterial_External(InOptions, Result, static_cast<const NodeMaterialExternal*>(Node));
 			break;
 
+		case Node::EType::MaterialModify:
+			GenerateMaterial_Modify(InOptions, Result, static_cast<const NodeMaterialModify*>(Node));
+			break;
+
 		case Node::EType::MaterialSkeletalMeshBreak:
 			GenerateMaterial_BreakSkeletalMesh(InOptions, Result, static_cast<const NodeMaterialSkeletalMeshBreak*>(Node));
 			break;
@@ -109,26 +121,24 @@ namespace UE::Mutable::Private
 		Ptr<ASTOpConstantResource> ConstantMaterial = new ASTOpConstantResource();
 		ConstantMaterial->Type = EOpType::MI_CONSTANT;
 
-		//TODO: each image should have its own generation options
 		FImageGenerationOptions ImageOptions(Options.ComponentId, Options.LODIndex);
 		ImageOptions.ImageLayoutStrategy = Options.ImageLayoutStrategy;
 		ImageOptions.RectSize = Options.RectSize;
 		ImageOptions.LayoutBlockId = Options.LayoutBlockId;
 		ImageOptions.LayoutToApply = Options.LayoutToApply;
 
-		const TArray<TPair<FName, Ptr<NodeImage>>>& ImageValuesArray = Node->ImageValues.Array();
+		const TArray<TPair<FParameterKey, Ptr<NodeImage>>>& ImageValuesArray = Node->ImageValues.Array();
 
 		//Store the images into the operation so that we can link them later
 		for (int32 ImageIndex = 0; ImageIndex < ImageValuesArray.Num(); ++ImageIndex)
 		{
-			FName ImageParameterName = ImageValuesArray[ImageIndex].Key;
 			Ptr<NodeImage> ImageNode = ImageValuesArray[ImageIndex].Value;
 			FImageGenerationResult ImageResult;
 
 			GenerateImage(ImageOptions, ImageResult, ImageNode);
 
 			// Store lazy branches
-			ConstantMaterial->ImageOperations.Add(ImageParameterName, ASTChild(ConstantMaterial, ImageResult.op));
+			ConstantMaterial->ImageOperations.Add(ImageValuesArray[ImageIndex].Key, ASTChild(ConstantMaterial, ImageResult.op));
 		}
 
 		// Store the constant material
@@ -304,8 +314,8 @@ namespace UE::Mutable::Private
 		}
 		
 		// Copy Scalar and Color Names
-		MaterialOp->ColorParameterNames = InNode->ColorParameterNames;
-		MaterialOp->ScalarParameterNames = InNode->ScalarParameterNames;
+		MaterialOp->ColorParameterKeys = InNode->ColorParametersToCompile;
+		MaterialOp->ScalarParameterKeys = InNode->ScalarParametersToCompile;
 
 		// Generate Image Nodes
 		FImageGenerationResult ImageResult;
@@ -316,7 +326,7 @@ namespace UE::Mutable::Private
 		ImageOptions.LayoutToApply = InOptions.LayoutToApply;
 		ImageOptions.MaterialParameter = MaterialOp->Parameter;
 
-		for (const TPair<FString, Ptr<NodeImage>>& ImageNode : InNode->ImageParameterNodes)
+		for (const TPair<FParameterKey, Ptr<NodeImage>>& ImageNode : InNode->ImageParametersToCompile)
 		{
 			GenerateImage(ImageOptions, ImageResult, ImageNode.Value);
 
@@ -329,6 +339,56 @@ namespace UE::Mutable::Private
 		FirstPass.ParameterNodes.GenericParametersCache.Add(InNode, MaterialOp);
 	}
 
+	
+	void CodeGenerator::GenerateMaterial_Modify(const FMaterialGenerationOptions& Options, FMaterialGenerationResult& Result, const NodeMaterialModify* Node)
+	{
+		Ptr<ASTOpMaterialModify> MaterialModify = new ASTOpMaterialModify();
+		
+		FMaterialGenerationResult SourceResult;
+
+		if (Node->MaterialSource)
+		{
+			GenerateMaterial(Options, SourceResult, Node->MaterialSource);
+			MaterialModify->Material = SourceResult.op;
+
+			// Add all the parameters that will be modified
+			// Textures
+			for (const TPair<FParameterKey, FImageParameterData>& ImageParameterData : Node->ImageParameters)
+			{
+				FImageGenerationResult ImageResult;
+				GenerateMaterialImage(ImageParameterData, ImageResult, Options.Modifiers, Options.State, Options.ComponentId, Options.LODIndex, Options.Tags, Options.MeshResults, Options.LayoutFromExtension, Node->GetMessageContext());
+				MaterialModify->ParametersToModify.Add(ImageParameterData.Key, ASTChild(MaterialModify, ImageResult.op));
+				MaterialModify->ImagePropertyIndexMap.Add(ImageParameterData.Key, ImageParameterData.Value.ImagePropertyIndex);
+			}
+
+			// Colors
+			for (const TPair<FParameterKey, Ptr<NodeColour>>& ColorParameterData : Node->ColorParameters)
+			{
+				FColorGenerationResult ColorResult;
+				FGenericGenerationOptions ColorOptions;
+				GenerateColor(ColorResult, ColorOptions, ColorParameterData.Value);
+
+				MaterialModify->ParametersToModify.Add(ColorParameterData.Key, ASTChild(MaterialModify, ColorResult.op));
+			}
+
+			// Scalars
+			for (const TPair<FParameterKey, Ptr<NodeScalar>>& ScalarParameterData : Node->ScalarParameters)
+			{
+				FScalarGenerationResult ScalarResult;
+				FGenericGenerationOptions ScalarOptions;
+				GenerateScalar(ScalarResult, ScalarOptions, ScalarParameterData.Value);
+
+				MaterialModify->ParametersToModify.Add(ScalarParameterData.Key, ASTChild(MaterialModify, ScalarResult.op));
+			}
+		}
+		else
+		{
+			MaterialModify->Material = GenerateMissingMaterialCode(TEXT("Make Material"), Node->GetMessageContext());
+		}
+
+		Result.op = MaterialModify;
+	}
+
 
 	void CodeGenerator::GenerateMaterial_External(const FMaterialGenerationOptions& InOptions, FMaterialGenerationResult& Result, const NodeMaterialExternal* InNode)
 	{
@@ -374,7 +434,6 @@ namespace UE::Mutable::Private
 		Ptr<ASTOpConstantResource> ConstantMaterial = new ASTOpConstantResource();
 		ConstantMaterial->Type = EOpType::MI_CONSTANT;
 
-		// Make a checkered debug image
 		TSharedPtr<FMaterial> GenerateMaterial = MakeShared<FMaterial>();
 		GenerateMaterial->PassthroughObject = {};
 
@@ -383,4 +442,359 @@ namespace UE::Mutable::Private
 
 		return ConstantMaterial;
 	}
+
+
+	void CodeGenerator::GenerateMaterialImage(
+		const TPair<FParameterKey, FImageParameterData>& ImageData,
+		FImageGenerationResult& OutImageGenerated,
+		const TArray<FirstPassGenerator::FModifier>& Modifiers,
+		const int32 State,
+		const int32 ComponentId,
+		const int32 LODIndex,
+		const TArray<FString>& Tags,
+		const FMeshGenerationResult& MeshResults,
+		const TBitArray<>& LayoutFromExtension,
+		const void* ErrorContext)
+	{
+		MUTABLE_CPUPROFILER_SCOPE(SurfaceTexture);
+
+		// Any image-specific format or mipmapping needs to be applied at the end
+		Ptr<NodeImageMipmap> mipmapNode;
+		Ptr<NodeImageFormat> formatNode;
+		Ptr<NodeImageSwizzle> swizzleNode;
+
+		bool bFound = false;
+		Ptr<NodeImage> pImageNode = ImageData.Value.ImageNode;
+
+		while (!bFound && pImageNode)
+		{
+			if (pImageNode->GetType() == NodeImageMipmap::GetStaticType())
+			{
+				NodeImageMipmap* tm = static_cast<NodeImageMipmap*>(pImageNode.get());
+				if (!mipmapNode) mipmapNode = tm;
+				pImageNode = tm->Source;
+			}
+			else if (pImageNode->GetType() == NodeImageFormat::GetStaticType())
+			{
+				NodeImageFormat* tf = static_cast<NodeImageFormat*>(pImageNode.get());
+				if (!formatNode) formatNode = tf;
+				pImageNode = tf->Source;
+			}
+			else if (pImageNode->GetType() == NodeImageSwizzle::GetStaticType())
+			{
+				NodeImageSwizzle* ts = static_cast<NodeImageSwizzle*>(pImageNode.get());
+
+				if (!ts->Sources.IsEmpty())
+				{
+					NodeImage* Source = ts->Sources[0].get();
+
+					bool bAllSourcesAreTheSame = true;
+					for (int32 SourceIndex = 1; SourceIndex < ts->Sources.Num(); ++SourceIndex)
+					{
+						bAllSourcesAreTheSame = bAllSourcesAreTheSame && (Source == ts->Sources[SourceIndex]);
+					}
+
+					if (!swizzleNode && bAllSourcesAreTheSame)
+					{
+						swizzleNode = ts;
+						pImageNode = Source;
+					}
+					else
+					{
+						bFound = true;
+					}
+				}
+				else
+				{
+					// break loop if swizzle has no sources.
+					bFound = true;
+				}
+			}
+			else
+			{
+				bFound = true;
+			}
+		}
+
+		if (bFound)
+		{
+			const int32 LayoutIndex = ImageData.Value.LayoutIndex;
+
+			// If the layout index has been set to negative, it means we should ignore the layout for this image.
+			CompilerOptions::TextureLayoutStrategy ImageLayoutStrategy = (LayoutIndex < 0)
+				? CompilerOptions::TextureLayoutStrategy::None
+				: CompilerOptions::TextureLayoutStrategy::Pack
+				;
+
+			if (ImageLayoutStrategy == CompilerOptions::TextureLayoutStrategy::None)
+			{
+				// Generate the image
+				FImageGenerationOptions ImageOptions(ComponentId, LODIndex);
+				ImageOptions.State = State;
+				ImageOptions.ImageLayoutStrategy = ImageLayoutStrategy;
+				ImageOptions.ActiveTags = Tags;
+				ImageOptions.RectSize = { 0, 0 };
+
+				// TODO: To tasks
+				FImageGenerationResult Result;
+				GenerateImage(ImageOptions, Result, pImageNode);
+				Ptr<ASTOp> imageAd = Result.op;
+
+				// Placeholder block. Ideally this should be the actual image size
+				constexpr int32 FakeLayoutSize = 256;
+				FIntPoint GridSize(FakeLayoutSize, FakeLayoutSize);
+				FLayoutBlockDesc LayoutBlockDesc;
+				LayoutBlockDesc.BlockPixelsX = 1;
+				LayoutBlockDesc.BlockPixelsY = 1;
+				box< FIntVector2 > RectInCells;
+				RectInCells.min = { 0,0 };
+				RectInCells.size = { FakeLayoutSize ,FakeLayoutSize };
+
+
+				imageAd = ApplyImageBlockModifiers(Modifiers, ImageOptions, imageAd, ImageData.Key, GridSize, LayoutBlockDesc, RectInCells, ErrorContext);
+
+				check(imageAd);
+
+				if (swizzleNode)
+				{
+					Ptr<ASTOpImageSwizzle> fop = new ASTOpImageSwizzle();
+					fop->Format = swizzleNode->NewFormat;
+					fop->Sources[0] = imageAd;
+					fop->Sources[1] = imageAd;
+					fop->Sources[2] = imageAd;
+					fop->Sources[3] = imageAd;
+					fop->SourceChannels[0] = swizzleNode->SourceChannels[0];
+					fop->SourceChannels[1] = swizzleNode->SourceChannels[1];
+					fop->SourceChannels[2] = swizzleNode->SourceChannels[2];
+					fop->SourceChannels[3] = swizzleNode->SourceChannels[3];
+					check(fop->Format != EImageFormat::None);
+					imageAd = fop;
+				}
+
+				if (mipmapNode)
+				{
+					Ptr<ASTOpImageMipmap> op = new ASTOpImageMipmap();
+					op->Levels = 0;
+					op->Source = imageAd;
+					op->BlockLevels = 0;
+
+					op->AddressMode = mipmapNode->Settings.AddressMode;
+					op->FilterType = mipmapNode->Settings.FilterType;
+					imageAd = op;
+				}
+
+				if (formatNode)
+				{
+					Ptr<ASTOpImagePixelFormat> fop = new ASTOpImagePixelFormat();
+					fop->Format = formatNode->Format;
+					fop->FormatIfAlpha = formatNode->FormatIfAlpha;
+					fop->Source = imageAd;
+					check(fop->Format != EImageFormat::None);
+					imageAd = fop;
+				}
+
+				// Store the image parameter into the modify operation
+				OutImageGenerated.op = imageAd;
+			}
+
+			else if (ImageLayoutStrategy == CompilerOptions::TextureLayoutStrategy::Pack) //-V547
+			{
+				if (LayoutIndex >= MeshResults.GeneratedLayouts.Num() ||
+					LayoutIndex >= MeshResults.LayoutOps.Num())
+				{
+					ErrorLog->Add("Missing layout in object, or its parent.", ELMT_ERROR, ErrorContext);
+				}
+				else
+				{
+					TSharedPtr<const FLayout> pLayout = MeshResults.GeneratedLayouts[LayoutIndex].Layout;
+					check(pLayout);
+
+					// Image
+					//-------------------------------------
+
+					// Size of a layout block in pixels
+					FIntPoint GridSize = pLayout->GetGridSize();
+
+					// Try to guess the layout block description from the first valid block that is generated.
+					FLayoutBlockDesc LayoutBlockDesc;
+					if (formatNode)
+					{
+						LayoutBlockDesc.FinalFormat = formatNode->FormatIfAlpha;
+						if (LayoutBlockDesc.FinalFormat == EImageFormat::None)
+						{
+							LayoutBlockDesc.FinalFormat = formatNode->Format;
+						}
+					}
+
+					bool bImageSizeWarning = false;
+
+					// Start with a blank image. It will be completed later with the blockSize, format and mips information
+					Ptr<ASTOpImageBlankLayout> BlankImageOp;
+					Ptr<ASTOp> imageAd;
+					{
+						BlankImageOp = new ASTOpImageBlankLayout();
+						BlankImageOp->Layout = MeshResults.LayoutOps[LayoutIndex];
+						// The rest ok the op will be completed below
+						BlankImageOp->MipmapCount = 0;
+						imageAd = BlankImageOp;
+					}
+
+					// Skip the block addition for this image if the layout was from a extension.
+					if (!LayoutFromExtension[LayoutIndex])
+					{
+						for (int32 BlockIndex = 0; BlockIndex < pLayout->GetBlockCount(); ++BlockIndex)
+						{
+							// Generate the image
+							FImageGenerationOptions ImageOptions(ComponentId, LODIndex);
+							ImageOptions.State = State;
+							ImageOptions.ImageLayoutStrategy = ImageLayoutStrategy;
+							ImageOptions.RectSize = { 0,0 };
+							ImageOptions.ActiveTags = Tags;
+							ImageOptions.LayoutToApply = pLayout;
+							ImageOptions.LayoutBlockId = pLayout->Blocks[BlockIndex].Id;
+							FImageGenerationResult ImageResult;
+							GenerateImage(ImageOptions, ImageResult, pImageNode);
+							Ptr<ASTOp> blockAd = ImageResult.op;
+
+							if (!blockAd)
+							{
+								// The GenerateImage(...) above has failed, skip this block
+								continue;
+							}
+
+							// Calculate the desc of the generated block.
+							constexpr bool bReturnBestOption = true;
+							FImageDesc BlockDesc = blockAd->GetImageDesc(bReturnBestOption, nullptr);
+
+							// Block in layout grid units (cells)
+							box< FIntVector2 > RectInCells;
+							RectInCells.min = pLayout->Blocks[BlockIndex].Min;
+							RectInCells.size = pLayout->Blocks[BlockIndex].Size;
+
+							// Try to update the layout block desc if we don't know it yet.
+							UpdateLayoutBlockDesc(LayoutBlockDesc, BlockDesc, RectInCells.size);
+
+							// Even if we force the size afterwards, we need some size hint in some cases, like image projections.
+							ImageOptions.RectSize = UE::Math::TIntVector2<int32>(BlockDesc.m_size);
+
+							blockAd = ApplyImageBlockModifiers(Modifiers, ImageOptions, blockAd, ImageData.Key, GridSize, LayoutBlockDesc, RectInCells, ErrorContext);
+
+							// Enforce block size and optimizations
+							blockAd = GenerateImageSize(blockAd, FIntVector2(BlockDesc.m_size));
+
+							EImageFormat baseFormat = imageAd->GetImageDesc().m_format;
+							// Actually don't do it, it will be propagated from the top format operation.
+							//Ptr<ASTOp> blockAd = GenerateImageFormat(blockAd, baseFormat);
+
+							// Apply tiling to avoid generating chunks of image that are too big.
+							blockAd = ApplyTiling(blockAd, ImageOptions.RectSize, LayoutBlockDesc.FinalFormat);
+
+							// Compose layout operation
+							Ptr<ASTOpImageCompose> composeOp = new ASTOpImageCompose();
+							composeOp->Layout = MeshResults.LayoutOps[LayoutIndex];
+							composeOp->Base = imageAd;
+							composeOp->BlockImage = blockAd;
+
+							// Set the absolute block index.
+							check(pLayout->Blocks[BlockIndex].Id != FLayoutBlock::InvalidBlockId);
+							composeOp->BlockId = pLayout->Blocks[BlockIndex].Id;
+
+							imageAd = composeOp;
+						}
+					}
+					check(imageAd);
+
+					FMeshGenerationStaticOptions ModifierOptions(ComponentId, LODIndex);
+					ModifierOptions.State = State;
+					ModifierOptions.ActiveTags = Tags;
+					imageAd = ApplyImageExtendModifiers(Modifiers, ModifierOptions, MeshResults, imageAd, ImageLayoutStrategy,
+						LayoutIndex, ImageData.Key, GridSize, LayoutBlockDesc,
+						ErrorContext);
+
+					// Complete the base op
+					BlankImageOp->BlockSize[0] = uint16(LayoutBlockDesc.BlockPixelsX);
+					BlankImageOp->BlockSize[1] = uint16(LayoutBlockDesc.BlockPixelsY);
+					BlankImageOp->Format = GetUncompressedFormat(LayoutBlockDesc.FinalFormat);
+					BlankImageOp->GenerateMipmaps = LayoutBlockDesc.bBlocksHaveMips;
+					BlankImageOp->MipmapCount = 0;
+
+					if (swizzleNode)
+					{
+						Ptr<ASTOpImageSwizzle> fop = new ASTOpImageSwizzle();
+						fop->Format = swizzleNode->NewFormat;
+
+						for (int32 ChannelIndex = 0; ChannelIndex < swizzleNode->SourceChannels.Num(); ++ChannelIndex)
+						{
+							fop->Sources[ChannelIndex] = imageAd;
+							fop->SourceChannels[ChannelIndex] = swizzleNode->SourceChannels[ChannelIndex];
+						}
+						check(fop->Format != EImageFormat::None);
+						imageAd = fop;
+					}
+
+					// Apply mipmap and format if necessary, skip if format is None (possibly because a block was skipped above)
+					bool bNeedsMips =
+						(mipmapNode && LayoutBlockDesc.FinalFormat != EImageFormat::None)
+						||
+						LayoutBlockDesc.bBlocksHaveMips;
+
+					if (bNeedsMips)
+					{
+						Ptr<ASTOpImageMipmap> mop = new ASTOpImageMipmap();
+
+						// At the end of the day, we want all the mipmaps. Maybe the code
+						// optimiser will split the process later.
+						mop->Levels = 0;
+						mop->bOnlyTail = false;
+						mop->Source = imageAd;
+
+						// We have to avoid mips smaller than the image format block size, so
+						// we will devide the layout block by the format block
+						const FImageFormatData& PixelFormatInfo = GetImageFormatData(LayoutBlockDesc.FinalFormat);
+
+						int32 mipsX = FMath::CeilLogTwo(LayoutBlockDesc.BlockPixelsX / PixelFormatInfo.PixelsPerBlockX);
+						int32 mipsY = FMath::CeilLogTwo(LayoutBlockDesc.BlockPixelsY / PixelFormatInfo.PixelsPerBlockY);
+						mop->BlockLevels = (uint8)FMath::Max(mipsX, mipsY);
+
+						if (LayoutBlockDesc.BlockPixelsX < PixelFormatInfo.PixelsPerBlockX || LayoutBlockDesc.BlockPixelsY < PixelFormatInfo.PixelsPerBlockY)
+						{
+							// In this case, the mipmap will never be useful for blocks, so we indicate that
+							// it should make the mips at the root of the expression.
+							mop->bOnlyTail = true;
+						}
+
+						mop->AddressMode = EAddressMode::ClampToEdge;
+						mop->FilterType = EMipmapFilterType::SimpleAverage;
+
+						if (mipmapNode)
+						{
+							mop->AddressMode = mipmapNode->Settings.AddressMode;
+							mop->FilterType = mipmapNode->Settings.FilterType;
+						}
+
+						imageAd = mop;
+					}
+
+					if (formatNode)
+					{
+						Ptr<ASTOpImagePixelFormat> fop = new ASTOpImagePixelFormat();
+						fop->Format = formatNode->Format;
+						fop->FormatIfAlpha = formatNode->FormatIfAlpha;
+						fop->Source = imageAd;
+						check(fop->Format != EImageFormat::None);
+						imageAd = fop;
+					}
+
+					// Add image parameter
+					OutImageGenerated.op = imageAd;
+				}
+			}
+
+			else
+			{
+				// Unimplemented texture layout strategy
+				check(false);
+			}
+		}
+	}
 }
diff --git a/Engine/Plugins/Mutable/Source/MutableTools/Private/MuT/CodeGenerator_Scalar.cpp b/Engine/Plugins/Mutable/Source/MutableTools/Private/MuT/CodeGenerator_Scalar.cpp
--- a/Engine/Plugins/Mutable/Source/MutableTools/Private/MuT/CodeGenerator_Scalar.cpp
+++ b/Engine/Plugins/Mutable/Source/MutableTools/Private/MuT/CodeGenerator_Scalar.cpp
@@ -441,6 +441,7 @@ namespace UE::Mutable::Private
 
 		// Set the parameter name that this break node will generate
 		MaterialOp->ParameterName = Node->ParameterName;
+		MaterialOp->LayerIndex = Node->LayerIndex;
 
 		// Generate Material Source
 		FMaterialGenerationOptions MaterialSourceOptions;
diff --git a/Engine/Plugins/Mutable/Source/MutableTools/Private/MuT/CodeOptimiser_Parameters.cpp b/Engine/Plugins/Mutable/Source/MutableTools/Private/MuT/CodeOptimiser_Parameters.cpp
--- a/Engine/Plugins/Mutable/Source/MutableTools/Private/MuT/CodeOptimiser_Parameters.cpp
+++ b/Engine/Plugins/Mutable/Source/MutableTools/Private/MuT/CodeOptimiser_Parameters.cpp
@@ -2,6 +2,7 @@
 
 #include "MuT/AST.h"
 #include "MuT/ASTOpAddLOD.h"
+#include "MuT/ASTOpAddMaterial.h"
 #include "MuT/ASTOpConditional.h"
 #include "MuT/ASTOpImageCompose.h"
 #include "MuT/ASTOpImageMipmap.h"
@@ -11,6 +12,7 @@
 #include "MuT/ASTOpImageResizeLike.h"
 #include "MuT/ASTOpImagePlainColor.h"
 #include "MuT/ASTOpImageDisplace.h"
+#include "MuT/ASTOpMaterialModify.h"
 #include "MuT/ASTOpMeshMorph.h"
 #include "MuT/ASTOpInstanceAdd.h"
 #include "MuT/ASTOpLayoutFromMesh.h"
@@ -1687,10 +1689,7 @@ namespace UE::Mutable::Private
             {
                 switch(Type)
                 {
-                case EOpType::IN_ADDIMAGE:
                 case EOpType::IN_ADDMESH:
-                case EOpType::IN_ADDVECTOR:
-                case EOpType::IN_ADDSCALAR:
                 case EOpType::IN_ADDSTRING:
                 {
 					const ASTOpInstanceAdd* TypedOp = static_cast<const ASTOpInstanceAdd*>(At.get());
@@ -1708,6 +1707,23 @@ namespace UE::Mutable::Private
                     }
                     return false;
                 }
+				case EOpType::IN_ADDMATERIAL:
+				{
+					const ASTOpAddMaterial* TypedOp = static_cast<const ASTOpAddMaterial*>(At.get());
+
+					TPair<bool, bool> NewState;
+					NewState.Key = false; //resource root
+					NewState.Value = bThisIsRuntime;
+					RecurseWithState(TypedOp->Instance.child(), NewState);
+
+					if (TypedOp->Material)
+					{
+						NewState.Key = true; //resource root
+						NewState.Value = bThisIsRuntime;
+						RecurseWithState(TypedOp->Material.child(), NewState);
+					}
+					return false;
+				}
 
                 default:
                 {
@@ -1818,39 +1834,86 @@ namespace UE::Mutable::Private
 
         // TODO: Finer grained: what if the runtime parameter just selects between compressed
         // textures? We don't want them uncompressed.
-        if( type==EOpType::IN_ADDIMAGE )
+
+        if(type == EOpType::MI_MODIFY)
         {
-			ASTOpInstanceAdd* TypedOp = static_cast<ASTOpInstanceAdd*>(At.get());
-			
-			if (TypedOp && TypedOp->value)
-			{			
-				Ptr<ASTOp> ImageOp = TypedOp->value.child();
+			ASTOpMaterialModify* TypedOp = static_cast<ASTOpMaterialModify*>(At.get());
 
-				// Does it have a runtime parameter in its subtree?
-				bool hasRuntimeParameter = HasRuntimeParamVisitor.HasAny(ImageOp);
+			if (TypedOp)
+			{
+				bool bHasChanged = false;
 
-				if (bAlwaysUncompress || hasRuntimeParameter)
+				// Check if we need to add a compress pixel format to uncompress any texture of the material modify
+				for (TPair<FParameterKey, ASTChild>& Parameter : TypedOp->ParametersToModify)
 				{
-					FImageDesc imageDesc = ImageOp->GetImageDesc( true );
+					if (!Parameter.Value.Child || !Parameter.Value.Child.get() || GetOpDataType(Parameter.Value.Child->GetOpType()) != EDataType::Image)
+					{
+						continue;
+					}
+
+					Ptr<ASTOp> ImageOp = Parameter.Value.child();
+
+					// Does it have a runtime parameter in its subtree?
+					bool hasRuntimeParameter = HasRuntimeParamVisitor.HasAny(ImageOp);
+
+					if (bAlwaysUncompress || hasRuntimeParameter)
+					{
+						FImageDesc imageDesc = ImageOp->GetImageDesc(true);
+
+						// Is it a compressed format?
+						EImageFormat format = imageDesc.m_format;
+						EImageFormat uncompressedFormat = GetUncompressedFormat(format);
+						bool isCompressedFormat = (uncompressedFormat != format);
+
+						if (isCompressedFormat)
+						{
+							bHasChanged = true;
+							break;
+						}
+					}
+				}
 
-					// Is it a compressed format?
-					EImageFormat format = imageDesc.m_format;
-					EImageFormat uncompressedFormat = GetUncompressedFormat( format );
-					bool isCompressedFormat = (uncompressedFormat != format);
+				// Clone the operation and add a compressed pixel format operation if needed
+				if (bHasChanged)
+				{
+					Ptr<ASTOpMaterialModify> NewOp = UE::Mutable::Private::Clone<ASTOpMaterialModify>(TypedOp);
 
-					if (isCompressedFormat)
+					// If needed, add a a compressed pixel format operation for all images 
+					for (TPair<FParameterKey, ASTChild>& Parameter : NewOp->ParametersToModify)
 					{
-						Ptr<ASTOpInstanceAdd> NewOp = UE::Mutable::Private::Clone<ASTOpInstanceAdd>(At);
+						if (GetOpDataType(Parameter.Value.Child->GetOpType()) != EDataType::Image)
+						{
+							continue;
+						}
+
+						Ptr<ASTOp> ImageOp = Parameter.Value.child();
 
-						// Add a new format operation to uncompress the image
-						Ptr<ASTOpImagePixelFormat> fop = new ASTOpImagePixelFormat;
-						fop->Format = uncompressedFormat;
-						fop->FormatIfAlpha = uncompressedFormat;
-						fop->Source = ImageOp;
+						// Does it have a runtime parameter in its subtree?
+						bool hasRuntimeParameter = HasRuntimeParamVisitor.HasAny(ImageOp);
 
-						NewOp->value = fop;
-						At = NewOp;
+						if (bAlwaysUncompress || hasRuntimeParameter)
+						{
+							FImageDesc imageDesc = ImageOp->GetImageDesc(true);
+
+							// Is it a compressed format?
+							EImageFormat format = imageDesc.m_format;
+							EImageFormat uncompressedFormat = GetUncompressedFormat(format);
+							bool isCompressedFormat = (uncompressedFormat != format);
+
+							if (isCompressedFormat)
+							{
+								// Add a new format operation to uncompress the image
+								Ptr<ASTOpImagePixelFormat> fop = new ASTOpImagePixelFormat;
+								fop->Format = uncompressedFormat;
+								fop->FormatIfAlpha = uncompressedFormat;
+								fop->Source = ImageOp;
+
+								Parameter.Value = ASTChild(NewOp, fop);
+							}
+						}
 					}
+
+					At = NewOp;
 				}
 			}
         }
diff --git a/Engine/Plugins/Mutable/Source/MutableTools/Private/MuT/Compiler.cpp b/Engine/Plugins/Mutable/Source/MutableTools/Private/MuT/Compiler.cpp
--- a/Engine/Plugins/Mutable/Source/MutableTools/Private/MuT/Compiler.cpp
+++ b/Engine/Plugins/Mutable/Source/MutableTools/Private/MuT/Compiler.cpp
@@ -6,6 +6,7 @@
 #include "MuT/AST.h"
 #include "MuT/ASTOpParameter.h"
 #include "MuT/ASTOpInstanceAdd.h"
+#include "MuT/ASTOpAddMaterial.h"
 #include "MuT/CodeGenerator.h"
 #include "MuT/CodeOptimiser.h"
 #include "MuT/ErrorLog.h"
@@ -108,6 +109,7 @@ namespace UE::Mutable::Private
 		{ false,	{ 0, 0, 0, 0, 0, 0, 0, 0, 0, 0, 0, 0, 0, 0, 0, 0, 0, 0, 0, 0, 0, 0, 0, 0, 0, 0, 0, 0, 0 }		},	// CO_MATERIAL_BREAK
 		{ false,	{ 0, 0, 0, 0, 0, 0, 0, 0, 0, 0, 0, 0, 0, 0, 0, 0, 0, 0, 0, 0, 0, 0, 0, 0, 0, 0, 0, 0, 0 }		},	// IM_MATERIAL_BREAK
 		{ false,	{ 0, 0, 0, 0, 0, 0, 0, 0, 0, 0, 0, 0, 0, 0, 0, 0, 0, 0, 0, 0, 0, 0, 0, 0, 0, 0, 0, 0, 0 }		},	// IM_PARAMETER_FROM_MATERIAL
+		{ false,	{ 0, 0, 0, 0, 0, 0, 0, 0, 0, 0, 0, 0, 0, 0, 0, 0, 0, 0, 0, 0, 0, 0, 0, 0, 0, 0, 0, 0, 0 }		},	// MI_MODIFY
 
 		{ false,	{ 0, 0, 0, 0, 0, 0, 0, 0, 0, 0, 0, 0, 0, 0, 0, 0, 0, 0, 0, 0, 0, 0, 0, 0, 0, 0, 0, 0, 0 }		},	// BO_EQUAL_SC_CONST
 		{ false,	{ 0, 0, 0, 0, 0, 0, 0, 0, 0, 0, 0, 0, 0, 0, 0, 0, 0, 0, 0, 0, 0, 0, 0, 0, 0, 0, 0, 0, 0 }		},	// BO_AND
@@ -182,9 +184,6 @@ namespace UE::Mutable::Private
 		{ true,		{ 0, 0, 0, 0, 0, 0, 0, 0, 0, 0, 0, 0, 0, 0, 0, 0, 0, 0, 0, 0, 0, 0, 0, 0, 0, 0, 0, 0, 0 }		},	// ME_SKELETALMESH_BREAK
 
 		{ false,	{ 0, 0, 0, 0, 0, 0, 0, 0, 0, 0, 0, 0, 0, 0, 0, 0, 0, 0, 0, 0, 0, 0, 0, 0, 0, 0, 0, 0, 0 }		},	// IN_ADDMESH
-		{ false,	{ 0, 0, 0, 0, 0, 0, 0, 0, 0, 0, 0, 0, 0, 0, 0, 0, 0, 0, 0, 0, 0, 0, 0, 0, 0, 0, 0, 0, 0 }		},	// IN_ADDIMAGE
-		{ false,	{ 0, 0, 0, 0, 0, 0, 0, 0, 0, 0, 0, 0, 0, 0, 0, 0, 0, 0, 0, 0, 0, 0, 0, 0, 0, 0, 0, 0, 0 }		},	// IN_ADDVECTOR
-		{ false,	{ 0, 0, 0, 0, 0, 0, 0, 0, 0, 0, 0, 0, 0, 0, 0, 0, 0, 0, 0, 0, 0, 0, 0, 0, 0, 0, 0, 0, 0 }		},	// IN_ADDSCALAR
 		{ false,	{ 0, 0, 0, 0, 0, 0, 0, 0, 0, 0, 0, 0, 0, 0, 0, 0, 0, 0, 0, 0, 0, 0, 0, 0, 0, 0, 0, 0, 0 }		},	// IN_ADDSTRING
 		{ false,	{ 0, 0, 0, 0, 0, 0, 0, 0, 0, 0, 0, 0, 0, 0, 0, 0, 0, 0, 0, 0, 0, 0, 0, 0, 0, 0, 0, 0, 0 }		},	// IN_ADDSURFACE
 		{ false,	{ 0, 0, 0, 0, 0, 0, 0, 0, 0, 0, 0, 0, 0, 0, 0, 0, 0, 0, 0, 0, 0, 0, 0, 0, 0, 0, 0, 0, 0 }		},	// IN_ADDCOMPONENT
@@ -324,9 +323,9 @@ namespace UE::Mutable::Private
 		m_pD->OptimisationOptions.ReferencedMeshResourceProvider = Mesh;
 	}
 
-    void CompilerOptions::SetDisableImageGeneration(bool bDisabled)
+    void CompilerOptions::SetDisableMaterialGeneration(bool bDisabled)
     {
-		m_pD->OptimisationOptions.bDisableImageGeneration = bDisabled;
+		m_pD->OptimisationOptions.bDisableMaterialGeneration = bDisabled;
     }
 
     void CompilerOptions::SetDisableMeshGeneration(bool bDisabled)
@@ -370,59 +369,63 @@ namespace UE::Mutable::Private
                 Roots.Add(State.root);
             }
 
-            TArray<Ptr<ASTOpInstanceAdd>> InstanceAddOpsToRemove;
+            TArray<Ptr<ASTOp>> InstanceAddOpsToRemove;
 
             const bool bRemoveMeshes = Options->GetPrivate()->OptimisationOptions.bDisableMeshGeneration;
-            const bool bRemoveImages = Options->GetPrivate()->OptimisationOptions.bDisableImageGeneration;
+            const bool bRemoveMaterial = Options->GetPrivate()->OptimisationOptions.bDisableMaterialGeneration;
 
             // Gather all InstanceAddOps its value needs to be removed.
-            if (bRemoveMeshes || bRemoveImages)
+            if (bRemoveMeshes || bRemoveMaterial)
             {
                 ASTOp::Traverse_TopDown_Unique_Imprecise(Roots, [&](Ptr<ASTOp> Node)
                 //ASTOp::Traverse_TopRandom_Unique_NonReentrant(Roots, [&](Ptr<ASTOp> Node)
                 {
-			        
                     bool bProcessChildren = GetOpDataType(Node->GetOpType()) == EDataType::Instance;
                     if (!bProcessChildren)
                     {
                         return false;
                     }
 
-                    switch (Node->GetOpType())
-                    {
-
-                    case EOpType::IN_ADDMESH:
-                    {
-                        Ptr<ASTOpInstanceAdd> TypedNode = static_cast<ASTOpInstanceAdd*>(Node.get());
-                        if (bRemoveMeshes && TypedNode)
-                        {
-                            InstanceAddOpsToRemove.Add(TypedNode);
-                        }
-                        break;
-                    }
-                    case EOpType::IN_ADDIMAGE:
-                    {
-                        Ptr<ASTOpInstanceAdd> TypedNode = static_cast<ASTOpInstanceAdd*>(Node.get());
-                        if (bRemoveImages && TypedNode)
-                        {
-                            InstanceAddOpsToRemove.Add(TypedNode);
-                        }
-                        break;
-                    }
-
-                    default:
-                        break;
-                    }
+					if (Node->GetOpType() == EOpType::IN_ADDMESH)
+					{
+						if (bRemoveMeshes)
+						{
+						    InstanceAddOpsToRemove.Add(Node);
+						}
+					}
+					else if (Node->GetOpType() == EOpType::IN_ADDMATERIAL)
+					{
+						if (bRemoveMaterial)
+						{
+							InstanceAddOpsToRemove.Add(Node);
+						}
+					}
 
-                    return true;
-                });
-            }
+					return true;
+				});
+			}
 
             // Remove the ops dangling from the InstanceAddOp but keep the InstanceAddOp for simplicity and parity with a full AST.
             // The case of missing value is contemplated in all usages.
-            for (Ptr<ASTOpInstanceAdd>& InstanceAddOp : InstanceAddOpsToRemove)
+            for (Ptr<ASTOp>& InstanceAddOp : InstanceAddOpsToRemove)
             {
-                InstanceAddOp->value = nullptr;
+				switch (InstanceAddOp->GetOpType())
+				{
+					case EOpType::IN_ADDMESH:
+					{
+						Ptr<ASTOpInstanceAdd> TypedNode = static_cast<ASTOpInstanceAdd*>(InstanceAddOp.get());
+						TypedNode->value = nullptr;
+
+						break;
+					}
+
+					case EOpType::IN_ADDMATERIAL:
+					{
+						Ptr<ASTOpAddMaterial> TypedNode = static_cast<ASTOpAddMaterial*>(InstanceAddOp.get());
+						TypedNode->Material = nullptr;
+						break;
+					}
+				}
             }
         }
     }
diff --git a/Engine/Plugins/Mutable/Source/MutableTools/Private/MuT/ErrorLog.cpp b/Engine/Plugins/Mutable/Source/MutableTools/Private/MuT/ErrorLog.cpp
--- a/Engine/Plugins/Mutable/Source/MutableTools/Private/MuT/ErrorLog.cpp
+++ b/Engine/Plugins/Mutable/Source/MutableTools/Private/MuT/ErrorLog.cpp
@@ -253,6 +253,7 @@ namespace UE::Mutable::Private
 		TEXT("CO_MATERIAL_BREAK"),
 		TEXT("IM_MATERIAL_BREAK"),
 		TEXT("IM_PARAMETER_FROM_MATERIAL"),
+		TEXT("MI_MODIFY        "),
 
 		TEXT("BO_EQUAL_SC_CONST"),
 		TEXT("BO_AND           "),
@@ -327,9 +328,6 @@ namespace UE::Mutable::Private
 		TEXT("ME_SKELETALMESH_BREAK"),
 
 		TEXT("IN_ADDMESH       "),
-		TEXT("IN_ADDIMAGE      "),
-		TEXT("IN_ADDVECTOR     "),
-		TEXT("IN_ADDSCALAR     "),
 		TEXT("IN_ADDSTRING     "),
         TEXT("IN_ADDSURFACE    "),
         TEXT("IN_ADDCOMPONENT  "),
diff --git a/Engine/Plugins/Mutable/Source/MutableTools/Private/MuT/Node.cpp b/Engine/Plugins/Mutable/Source/MutableTools/Private/MuT/Node.cpp
--- a/Engine/Plugins/Mutable/Source/MutableTools/Private/MuT/Node.cpp
+++ b/Engine/Plugins/Mutable/Source/MutableTools/Private/MuT/Node.cpp
@@ -117,6 +117,7 @@
 #include "MuT/NodeSkeletalMesh.h"
 #include "MuT/NodeSkeletalMeshNew.h"
 #include "MuT/NodeSkeletalMeshParameter.h"
+#include "MuT/NodeMaterialModify.h"
 
 
 namespace UE::Mutable::Private
@@ -260,7 +261,8 @@ namespace UE::Mutable::Private
 	FNodeType NodeMaterialParameter::StaticType = FNodeType(Node::EType::MaterialParameter, NodeMaterial::GetStaticType());
 	FNodeType NodeMaterialExternal::StaticType = FNodeType(Node::EType::MaterialExternal, NodeMaterial::GetStaticType());
 	FNodeType NodeMaterialSkeletalMeshBreak::StaticType = FNodeType(Node::EType::MaterialSkeletalMeshBreak, NodeMaterial::GetStaticType());
-
+	FNodeType NodeMaterialModify::StaticType = FNodeType(Node::EType::MaterialModify, NodeMaterial::GetStaticType());
+	
 	FNodeType NodeImageMaterialBreak::StaticType = FNodeType(Node::EType::ImageMaterialBreak, NodeImage::GetStaticType());
 	FNodeType NodeScalarMaterialBreak::StaticType = FNodeType(Node::EType::ScalarMaterialBreak, NodeScalar::GetStaticType());
 	FNodeType NodeColourMaterialBreak::StaticType = FNodeType(Node::EType::ColourMaterialBreak, NodeColour::GetStaticType());
diff --git a/Engine/Plugins/Mutable/Source/MutableTools/Public/MuT/NodeMaterialModify.h b/Engine/Plugins/Mutable/Source/MutableTools/Public/MuT/NodeMaterialModify.h
new file mode 100644
--- /dev/null
+++ b/Engine/Plugins/Mutable/Source/MutableTools/Public/MuT/NodeMaterialModify.h
@@ -0,0 +1,56 @@
+﻿// Copyright Epic Games, Inc. All Rights Reserved.
+
+#pragma once
+
+#include "MuT/NodeMaterial.h"
+#include "MuT/NodeColour.h"
+#include "MuT/NodeScalar.h"
+#include "MuT/NodeImage.h"
+#include "MuR/Material.h"
+#include "MuR/Ptr.h"
+#include "MuR/Types.h"
+
+#define UE_API MUTABLETOOLS_API
+
+namespace UE::Mutable::Private
+{
+	struct FImageParameterData
+	{
+		Ptr<NodeImage> ImageNode;
+		int32 ImagePropertyIndex = 0;
+		int8 LayoutIndex = 0;
+
+		bool operator==(const FImageParameterData& Other) const
+		{
+			return ImageNode == Other.ImageNode && ImagePropertyIndex == Other.ImagePropertyIndex && LayoutIndex == Other.LayoutIndex;
+		}
+	};
+
+	class NodeMaterialModify : public NodeMaterial
+	{
+	public:
+
+		Ptr<NodeMaterial> MaterialSource;
+		TMap<FParameterKey, FImageParameterData> ImageParameters;
+		TMap<FParameterKey, Ptr<NodeColour>> ColorParameters;
+		TMap<FParameterKey, Ptr<NodeScalar>> ScalarParameters;
+
+	public:
+		
+		// Node Interface
+		virtual const FNodeType* GetType() const override { return &StaticType; }
+		static const FNodeType* GetStaticType() { return &StaticType; }
+
+	protected:
+
+		/** Forbidden. Manage with the Ptr<> template. */
+		~NodeMaterialModify() {}
+
+	private:
+
+		static UE_API FNodeType StaticType;
+	
+	};
+}
+
+#undef UE_API
\ No newline at end of file
