[Mutable] Add image multicompose code generation and coderunner logic.

#Synced-CL 48603947

diff --git a/Engine/Plugins/Mutable/Source/CustomizableObject/Internal/MuCO/CustomizableObjectPrivate.h b/Engine/Plugins/Mutable/Source/CustomizableObject/Internal/MuCO/CustomizableObjectPrivate.h
--- a/Engine/Plugins/Mutable/Source/CustomizableObject/Internal/MuCO/CustomizableObjectPrivate.h
+++ b/Engine/Plugins/Mutable/Source/CustomizableObject/Internal/MuCO/CustomizableObjectPrivate.h
@@ -1765,6 +1765,8 @@ enum class ECustomizableObjectVersions : int32
 
 	AddMaterialModify,
 
+	AddImageMultiCompose,
+
 	// -----<new versions can be added above this line>--------
 	LastCustomizableObjectVersion
 };
diff --git a/Engine/Plugins/Mutable/Source/CustomizableObjectEditor/Private/MuCOE/SMutableCodeViewer.cpp b/Engine/Plugins/Mutable/Source/CustomizableObjectEditor/Private/MuCOE/SMutableCodeViewer.cpp
--- a/Engine/Plugins/Mutable/Source/CustomizableObjectEditor/Private/MuCOE/SMutableCodeViewer.cpp
+++ b/Engine/Plugins/Mutable/Source/CustomizableObjectEditor/Private/MuCOE/SMutableCodeViewer.cpp
@@ -1729,10 +1729,21 @@ void SMutableCodeViewer::GenerateElementRecursive(const int32& InStateIndex, UE:
 	case UE::Mutable::Private::EOpType::IM_COMPOSE:
 	{
 		UE::Mutable::Private::OP::ImageComposeArgs Args = InProgram.GetOpArgs<UE::Mutable::Private::OP::ImageComposeArgs>(InParentAddress);
-		AddOpFunc(Args.layout, TEXT("layout "));
-		AddOpFunc(Args.base, TEXT("base "));
-		AddOpFunc(Args.blockImage, TEXT("blockImage "));
-		AddOpFunc(Args.mask, TEXT("mask "));
+		AddOpFunc(Args.layout, TEXT("Layout "));
+		AddOpFunc(Args.base, TEXT("Base "));
+		AddOpFunc(Args.blockImage, TEXT("BlockImage "));
+		AddOpFunc(Args.mask, TEXT("Mask "));
+		break;
+	}
+
+	case UE::Mutable::Private::EOpType::IM_MULTICOMPOSE:
+	{
+		UE::Mutable::Private::OP::ImageMultiComposeArgs Args = InProgram.GetOpArgs<UE::Mutable::Private::OP::ImageMultiComposeArgs>(InParentAddress);
+		AddOpFunc(Args.Layout, TEXT("Layout "));
+		AddOpFunc(Args.Base, TEXT("Base "));
+		AddOpFunc(Args.SourceLayout, TEXT("SourceLayout"));
+		AddOpFunc(Args.SourceImage, TEXT("SourceImage"));
+		
 		break;
 	}
 		
diff --git a/Engine/Plugins/Mutable/Source/MutableRuntime/Internal/MuR/Operations.h b/Engine/Plugins/Mutable/Source/MutableRuntime/Internal/MuR/Operations.h
--- a/Engine/Plugins/Mutable/Source/MutableRuntime/Internal/MuR/Operations.h
+++ b/Engine/Plugins/Mutable/Source/MutableRuntime/Internal/MuR/Operations.h
@@ -181,6 +181,9 @@ namespace UE::Mutable::Private
         //! Copy an image into a rect of another one.
         IM_COMPOSE,
 
+		//! Copy blocks from an image into a another one.
+		IM_MULTICOMPOSE,
+
         //! Interpolate between 2 images taken from a row of targets (2 consecutive targets).
         IM_INTERPOLATE,
 
@@ -744,6 +747,23 @@ namespace UE::Mutable::Private
             uint64 BlockId;
         };
 
+		struct ImageMultiComposeArgs
+		{
+			ADDRESS Layout;
+			ADDRESS SourceLayout;
+			ADDRESS Base;
+			ADDRESS SourceImage;
+
+			//ADDRESS Mask;
+
+			/** Size of the image to create. If 0, reuse size from base.*/
+			uint16 LayoutBlockSizeInPixelsX = 0;
+			uint16 LayoutBlockSizeInPixelsY = 0;
+
+			uint16 SourceSizeX = 0;
+			uint16 SourceSizeY = 0;
+		};
+
         struct ImageInterpolateArgs
         {
             ADDRESS Factor;
diff --git a/Engine/Plugins/Mutable/Source/MutableRuntime/Private/MuR/CodeRunner.cpp b/Engine/Plugins/Mutable/Source/MutableRuntime/Private/MuR/CodeRunner.cpp
--- a/Engine/Plugins/Mutable/Source/MutableRuntime/Private/MuR/CodeRunner.cpp
+++ b/Engine/Plugins/Mutable/Source/MutableRuntime/Private/MuR/CodeRunner.cpp
@@ -32,6 +32,7 @@
 #include "MuR/OpImageProject.h"
 #include "MuR/OpImageRasterMesh.h"
 #include "MuR/OpImageTransform.h"
+#include "MuR/OpImageMultiCompose.h"
 #include "MuR/OpLayoutPack.h"
 #include "MuR/OpLayoutRemoveBlocks.h"
 #include "MuR/OpMeshApplyLayout.h"
@@ -4191,6 +4192,262 @@ namespace UE::Mutable::Private
 
             break;
         }
+		
+		case EOpType::IM_MULTICOMPOSE:
+		{
+			OP::ImageMultiComposeArgs Args = Program.GetOpArgs<OP::ImageMultiComposeArgs>(Item.At);
+			constexpr float LODRoundTolerance = 0.1f;
+			
+			switch (Item.Stage)
+			{
+			case 0:
+			{
+				AddOp(FScheduledOp(Item.At, Item, 1), FScheduledOp::FromOpAndOptions(Args.Layout, Item, 0), FScheduledOp::FromOpAndOptions(Args.SourceLayout, Item, 0));
+				break;
+			}
+			case 1:
+			{
+				TSharedPtr<const FLayout> ComposeLayout =
+					LoadLayout(FCacheAddress(Args.Layout, FScheduledOp::FromOpAndOptions(Args.Layout, Item, 0)));
+
+
+				TSharedPtr<const FLayout> SourceLayout =
+					LoadLayout(FCacheAddress(Args.SourceLayout, FScheduledOp::FromOpAndOptions(Args.SourceLayout, Item, 0)));
+
+				bool bHasBlocksToCompose = false;
+				for (const FLayoutBlock& Block : SourceLayout->Blocks)
+				{
+					if (ComposeLayout->FindBlock(Block.Id) >= 0)
+					{
+						bHasBlocksToCompose = true;
+						break;
+					}
+				}
+
+				FScheduledOpData SourceData;
+				SourceData.Resource = SourceLayout;
+
+				FScheduledOpData Data;
+				Data.Resource = ComposeLayout;
+				Data.MultiCompose.NextHeapIndex = HeapData.Add(SourceData);
+				Data.MultiCompose.bHasBlocksToCompose = bHasBlocksToCompose;
+				int32 DataPos = HeapData.Add(Data);
+
+				if (bHasBlocksToCompose)
+				{
+					AddOp(FScheduledOp(Item.At, Item, 2, DataPos), FScheduledOp(Args.Base, Item));
+				}
+				else
+				{
+					// Jump directly to stage 3, no need to load the source image.
+					AddOp(FScheduledOp(Item.At, Item, 3, DataPos), FScheduledOp(Args.Base, Item));
+				}
+
+				break;
+			}
+			case 2:
+			{
+				FScheduledOpData& Data = HeapData[Item.CustomState];
+				TSharedPtr<const FLayout> ComposeLayout = StaticCastSharedPtr<const FLayout>(Data.Resource);
+
+				FScheduledOpData& SourceData = HeapData[Data.MultiCompose.NextHeapIndex];
+				TSharedPtr<const FLayout> SourceLayout = StaticCastSharedPtr<const FLayout>(SourceData.Resource);
+
+				FIntVector2 SourceLayoutGridSize = SourceLayout->GetGridSize();
+				FIntVector2 ComposeLayoutGridSize = ComposeLayout->GetGridSize();
+				
+				FIntVector2 DstImageSize = FIntVector2(
+						Args.LayoutBlockSizeInPixelsX*ComposeLayoutGridSize.X, 
+						Args.LayoutBlockSizeInPixelsY*ComposeLayoutGridSize.Y);
+
+				const int32 NumLODsToSkip = Item.ExecutionOptions;
+				for (int32 LODIndex = 0; LODIndex < NumLODsToSkip; ++LODIndex)
+				{
+					DstImageSize.X = FMath::DivideAndRoundUp(DstImageSize.X, 2);	
+					DstImageSize.Y = FMath::DivideAndRoundUp(DstImageSize.Y, 2);	
+				}
+
+				FIntVector2 SrcImageSize = FIntVector2(Args.SourceSizeX, Args.SourceSizeY);
+
+				FIntVector2 SourceLayoutGridSizeInPixels = FIntVector2(
+						SrcImageSize.X / SourceLayoutGridSize.X,
+						SrcImageSize.Y / SourceLayoutGridSize.Y);
+
+				FIntVector2 DstLayoutGridSizeInPixels = FIntVector2(
+						DstImageSize.X / SourceLayoutGridSize.X,
+						DstImageSize.Y / SourceLayoutGridSize.Y);
+
+				float BestMaxLODForSourceImage = 0.0f;
+				float BestMinLODForSourceImage = std::numeric_limits<float>::max();
+
+				for (const FLayoutBlock& SrcBlock : SourceLayout->Blocks)
+				{
+					int32 FoundDstBlockIndex = ComposeLayout->FindBlock(SrcBlock.Id);
+
+					if (FoundDstBlockIndex >= 0)
+					{
+						const FLayoutBlock& DstBlock = ComposeLayout->Blocks[FoundDstBlockIndex];
+						
+						FIntRect SrcRectInPixels = FIntRect(
+								SrcBlock.Min.X * SourceLayoutGridSizeInPixels.X, 
+								SrcBlock.Min.Y * SourceLayoutGridSizeInPixels.Y, 
+								(SrcBlock.Min.X + SrcBlock.Size.X) * SourceLayoutGridSizeInPixels.X, 
+								(SrcBlock.Min.Y + SrcBlock.Size.Y) * SourceLayoutGridSizeInPixels.Y);
+
+						FIntRect DstRectInPixels = FIntRect(
+								DstBlock.Min.X * DstLayoutGridSizeInPixels.X, 
+								DstBlock.Min.Y * DstLayoutGridSizeInPixels.Y, 
+								(DstBlock.Min.X + DstBlock.Size.X) * DstLayoutGridSizeInPixels.X, 
+								(DstBlock.Min.Y + DstBlock.Size.Y) * DstLayoutGridSizeInPixels.Y);
+	
+						float BestLODForRect = FMath::Max(0.0f, 
+								ImageMultiComposeComputeBestLODForSrcRect(DstRectInPixels, SrcRectInPixels)) + LODRoundTolerance;
+
+						BestMinLODForSourceImage = FMath::Min(BestLODForRect, BestMinLODForSourceImage);
+						BestMaxLODForSourceImage = FMath::Max(BestLODForRect, BestMaxLODForSourceImage);
+					}
+				}
+
+				Data.MultiCompose.SourceLOD = FMath::FloorToInt32(BestMinLODForSourceImage);
+				Data.MultiCompose.SourceNumLODs = FMath::FloorToInt32(BestMaxLODForSourceImage) - FMath::FloorToInt32(BestMinLODForSourceImage) + 1;
+
+				AddOp(FScheduledOp(Item.At, Item, 3, Item.CustomState),	
+						FScheduledOp::FromOpAndOptions(Args.SourceImage, Item, Data.MultiCompose.SourceLOD));
+				break;
+			}
+			case 3:
+			{
+				MUTABLE_CPUPROFILER_SCOPE(RunCode_ImageMultiCompose);
+
+				TSharedPtr<const FImage> BaseImage = LoadImage(FCacheAddress(Args.Base, Item));
+
+				FScheduledOpData& Data = HeapData[Item.CustomState];
+				if (BaseImage && Data.MultiCompose.bHasBlocksToCompose)
+				{
+					TSharedPtr<const FLayout> ComposeLayout = StaticCastSharedPtr<const FLayout>(Data.Resource);
+
+					FScheduledOpData& SourceData = HeapData[Data.MultiCompose.NextHeapIndex];
+					TSharedPtr<const FLayout> SourceLayout = StaticCastSharedPtr<const FLayout>(SourceData.Resource);
+
+					int32 SourceImageLOD = Data.MultiCompose.SourceLOD;
+
+					TSharedPtr<const FImage> SourceImage = 
+							LoadImage(FCacheAddress(Args.SourceImage, Item.ExecutionIndex, SourceImageLOD));
+
+					if (SourceImage)
+					{
+						const FImageSize ExpectedSourceSize = FImageSize(
+								FMath::Max<uint16>(Args.SourceSizeX >> SourceImageLOD, 1), 
+								FMath::Max<uint16>(Args.SourceSizeY >> SourceImageLOD, 1));
+
+						if (SourceImage->GetSize() != ExpectedSourceSize)
+						{
+							MUTABLE_CPUPROFILER_SCOPE(RunCode_ImageMultiCompose_SizeFixup);
+							
+							TSharedPtr<FImage> Resized = CreateImage(ExpectedSourceSize.X, ExpectedSourceSize.Y, 1, SourceImage->GetFormat(), EInitializationType::NotInitialized);
+							ImOp.ImageResizeLinear(Resized.Get(), Settings.ImageCompressionQuality, SourceImage.Get());
+							
+							Release(SourceImage);
+							SourceImage = Resized;	
+						}
+
+						TSharedPtr<FImage> Result = CloneOrTakeOver(BaseImage);
+
+						FIntVector2 SourceLayoutGridSize = SourceLayout->GetGridSize();
+						FIntVector2 ComposeLayoutGridSize = ComposeLayout->GetGridSize();
+
+						FIntVector2 SourceLayoutGridSizeInPixels = FIntVector2(
+								SourceImage->GetSizeX() / SourceLayoutGridSize.X,
+								SourceImage->GetSizeY() / SourceLayoutGridSize.Y);
+
+						FIntVector2 ComposeLayoutGridSizeInPixels = FIntVector2(
+								Result->GetSizeX() / ComposeLayoutGridSize.X,
+								Result->GetSizeY() / ComposeLayoutGridSize.Y);
+
+						TArray<FIntRect, TInlineAllocator<16>> SrcRects;
+						TArray<FIntRect, TInlineAllocator<16>> DstRects;
+						TArray<int32,    TInlineAllocator<16>> SelectedLODForRect;
+
+						int32 NumSrcLODsNeeded = 1;
+						for (const FLayoutBlock& SrcBlock : SourceLayout->Blocks)
+						{
+							int32 FoundDstBlockIndex = ComposeLayout->FindBlock(SrcBlock.Id);
+
+							if (FoundDstBlockIndex >= 0)
+							{
+								const FLayoutBlock& DstBlock = ComposeLayout->Blocks[FoundDstBlockIndex];
+								FIntRect& SrcRectInPixels = SrcRects.Emplace_GetRef(
+										SrcBlock.Min.X * SourceLayoutGridSizeInPixels.X, 
+										SrcBlock.Min.Y * SourceLayoutGridSizeInPixels.Y, 
+										(SrcBlock.Min.X + SrcBlock.Size.X) * SourceLayoutGridSizeInPixels.X, 
+										(SrcBlock.Min.Y + SrcBlock.Size.Y) * SourceLayoutGridSizeInPixels.Y);
+
+								FIntRect& DstRectInPixels = DstRects.Emplace_GetRef(
+										DstBlock.Min.X * ComposeLayoutGridSizeInPixels.X, 
+										DstBlock.Min.Y * ComposeLayoutGridSizeInPixels.Y, 
+										(DstBlock.Min.X + DstBlock.Size.X) * ComposeLayoutGridSizeInPixels.X, 
+										(DstBlock.Min.Y + DstBlock.Size.Y) * ComposeLayoutGridSizeInPixels.Y);
+
+								float BestLODForRect = FMath::Max(0.0f,
+										ImageMultiComposeComputeBestLODForSrcRect(DstRectInPixels, SrcRectInPixels));
+								
+								int32 SelectedLOD = FMath::FloorToInt32(BestLODForRect + LODRoundTolerance);
+
+								NumSrcLODsNeeded = FMath::Max(SelectedLOD + 1, NumSrcLODsNeeded);
+								SelectedLODForRect.Emplace(SelectedLOD);
+							}
+						}
+
+						// Fix src rect sizes for the selected LOD.
+						for (int32 RectIndex = SrcRects.Num() - 1; RectIndex >= 0; --RectIndex)
+						{
+							int32 SelectedLOD = SelectedLODForRect[RectIndex];
+
+							for (int32 LODIndex = 0; LODIndex < SelectedLOD; ++LODIndex)
+							{
+								SrcRects[RectIndex] = FIntRect::DivideAndRoundUp(SrcRects[RectIndex], 2);
+							}
+						}
+
+						if (SourceImage->GetLODCount() < NumSrcLODsNeeded)
+						{
+							MUTABLE_CPUPROFILER_SCOPE(RunCode_ImageTransform_BilinearMipGen);
+
+							TSharedPtr<FImage> OwnedSource = CloneOrTakeOver(SourceImage);
+							OwnedSource->DataStorage.SetNumLODs(NumSrcLODsNeeded);
+
+							FMipmapGenerationSettings MipSettings{};
+							ImOp.ImageMipmap(Settings.ImageCompressionQuality, OwnedSource.Get(), OwnedSource.Get(), 0, NumSrcLODsNeeded, MipSettings);
+							//ImageMipmapInPlace(0, OwnedSource.Get(), FMipmapGenerationSettings{});
+
+							SourceImage = OwnedSource;
+						}	
+
+						ImageMultiCompose(*Result, *SourceImage, DstRects, SrcRects, SelectedLODForRect, Settings.ImageCompressionQuality);
+						Result->DataStorage.SetNumLODs(1);
+
+						Release(SourceImage);
+						StoreImage(Item, Result);
+					}
+					else
+					{
+						StoreImage(Item, BaseImage);
+					}
+				}
+				else
+				{
+					StoreImage(Item, BaseImage);
+				}
+
+				break;
+			}
+
+			default:
+				check(false);
+			}	
+
+			break;
+		}
 
         case EOpType::IM_INTERPOLATE:
         {
@@ -7970,6 +8227,38 @@ namespace UE::Mutable::Private
 			break;
 		}
 
+		case EOpType::IM_MULTICOMPOSE:
+		{
+			MUTABLE_CPUPROFILER_SCOPE(RunCodeImageDesc_IM_MULTICOMPOSE);
+			OP::ImageMultiComposeArgs Args = Program.GetOpArgs<OP::ImageMultiComposeArgs>(Item.At);
+			switch (Item.Stage)
+			{
+			case 0:
+			{
+				AddOp(FScheduledOp(Item.At, Item, 1),
+					FScheduledOp(Args.Base, Item, 0),
+					FScheduledOp(Args.SourceImage, Item, 0));
+				break;
+			}
+			case 1:
+			{
+				FExtendedImageDesc& Result = ImageDescResults.FindOrAdd(Item.At);
+				Result = ImageDescResults.FindChecked(Args.Base);
+
+				if (Args.SourceImage)
+				{
+					const FExtendedImageDesc& BlockResult = ImageDescResults.FindChecked(Args.SourceImage);
+					Result.FirstLODAvailable = FMath::Max(Result.FirstLODAvailable, BlockResult.FirstLODAvailable);
+				}
+
+				StoreValidDesc(Item);
+				break;
+			}
+			default: check(false);
+			}
+			break;
+		}
+
 		case EOpType::IM_INTERPOLATE:
 		{
 			MUTABLE_CPUPROFILER_SCOPE(RunCodeImageDesc_IM_INTERPOLATE);
diff --git a/Engine/Plugins/Mutable/Source/MutableRuntime/Private/MuR/CodeRunner.h b/Engine/Plugins/Mutable/Source/MutableRuntime/Private/MuR/CodeRunner.h
--- a/Engine/Plugins/Mutable/Source/MutableRuntime/Private/MuR/CodeRunner.h
+++ b/Engine/Plugins/Mutable/Source/MutableRuntime/Private/MuR/CodeRunner.h
@@ -109,6 +109,15 @@ namespace  UE::Mutable::Private
 					uint16 ScaleYEncodedHalf;
 					float MipValue;
 				} ImageTransform;
+
+				struct
+				{
+					int32 NextHeapIndex;
+					bool bHasBlocksToCompose;
+					uint8 SourceLOD;
+					uint8 SourceNumLODs;
+				} MultiCompose;
+
 			};
 
 			TSharedPtr<const FResource> Resource;
diff --git a/Engine/Plugins/Mutable/Source/MutableRuntime/Private/MuR/CodeRunnerBegin.cpp b/Engine/Plugins/Mutable/Source/MutableRuntime/Private/MuR/CodeRunnerBegin.cpp
--- a/Engine/Plugins/Mutable/Source/MutableRuntime/Private/MuR/CodeRunnerBegin.cpp
+++ b/Engine/Plugins/Mutable/Source/MutableRuntime/Private/MuR/CodeRunnerBegin.cpp
@@ -2146,6 +2146,16 @@ namespace UE::Mutable::Private
 					StoreNone(Item);
 					break;
 				}
+
+			case EOpType::IM_MULTICOMPOSE:
+			{
+				OP::ImageMultiComposeArgs Args = Program.GetOpArgs<OP::ImageMultiComposeArgs>(Item.At);
+				PushOp(FScheduledOpInline(Args.Base, Item));
+				PushOp(FScheduledOpInline(Args.SourceImage, Item));
+
+				StoreNone(Item);
+				break;
+			}
 	
 			case EOpType::IM_INTERPOLATE:
 				{
diff --git a/Engine/Plugins/Mutable/Source/MutableRuntime/Private/MuR/OpImageMultiCompose.cpp b/Engine/Plugins/Mutable/Source/MutableRuntime/Private/MuR/OpImageMultiCompose.cpp
--- a/Engine/Plugins/Mutable/Source/MutableRuntime/Private/MuR/OpImageMultiCompose.cpp
+++ b/Engine/Plugins/Mutable/Source/MutableRuntime/Private/MuR/OpImageMultiCompose.cpp
@@ -376,6 +376,8 @@ namespace OpImageMultiComposeInternal
 			TConstArrayView<int32> SourceMipPerRect,
 			int32 CompressionQuality)
 	{
+		MUTABLE_CPUPROFILER_SCOPE(ImageMultiCompose);
+
 		using namespace OpImageMultiComposeInternal;
 
 		check(DstRects.Num() == SrcRects.Num());
diff --git a/Engine/Plugins/Mutable/Source/MutableRuntime/Private/MuR/Operations.cpp b/Engine/Plugins/Mutable/Source/MutableRuntime/Private/MuR/Operations.cpp
--- a/Engine/Plugins/Mutable/Source/MutableRuntime/Private/MuR/Operations.cpp
+++ b/Engine/Plugins/Mutable/Source/MutableRuntime/Private/MuR/Operations.cpp
@@ -98,6 +98,7 @@ namespace UE::Mutable::Private
 		{ EDataType::Image,			},	// IM_RESIZEREL
 		{ EDataType::Image,			},	// IM_BLANKLAYOUT
 		{ EDataType::Image,			},	// IM_COMPOSE
+		{ EDataType::Image,			},	// IM_MULTICOMPOSE
 		{ EDataType::Image,			},	// IM_INTERPOLATE
 		{ EDataType::Image,			},	// IM_SATURATE
 		{ EDataType::Image,			},	// IM_LUMINANCE
@@ -516,6 +517,16 @@ namespace UE::Mutable::Private
             break;
         }
 
+		case EOpType::IM_MULTICOMPOSE:
+		{
+			OP::ImageMultiComposeArgs args = program.GetOpArgs<OP::ImageMultiComposeArgs>(at);
+			f(args.Layout);
+			f(args.Base);
+			f(args.SourceLayout);
+			f(args.SourceImage);
+			break;
+		}
+
         case EOpType::IM_INTERPOLATE:
         {
 			OP::ImageInterpolateArgs Args = program.GetOpArgs<OP::ImageInterpolateArgs>(at);
diff --git a/Engine/Plugins/Mutable/Source/MutableTools/Private/MuT/ASTOpImageMipmap.cpp b/Engine/Plugins/Mutable/Source/MutableTools/Private/MuT/ASTOpImageMipmap.cpp
--- a/Engine/Plugins/Mutable/Source/MutableTools/Private/MuT/ASTOpImageMipmap.cpp
+++ b/Engine/Plugins/Mutable/Source/MutableTools/Private/MuT/ASTOpImageMipmap.cpp
@@ -13,6 +13,7 @@
 #include "MuR/Types.h"
 #include "MuT/ASTOpConditional.h"
 #include "MuT/ASTOpImageCompose.h"
+#include "MuT/ASTOpImageMultiCompose.h"
 #include "MuT/ASTOpImagePatch.h"
 #include "MuT/ASTOpImagePixelFormat.h"
 #include "MuT/ASTOpImageBlankLayout.h"
@@ -399,6 +400,30 @@ namespace UE::Mutable::Private
 			break;
 		}
 
+		case EOpType::IM_MULTICOMPOSE:
+		{
+			const ASTOpImageMultiCompose* typedAt = static_cast<const ASTOpImageMultiCompose*>(at.get());
+			if (!currentMipmapOp->bOnlyTail
+				//&&
+				// Don't move the mipmapping if we are composing with a mask.
+				// TODO: allow mipmapping in the masks, RLE formats, etc.
+				//!typedAt->Mask
+				)
+			{
+				Ptr<ASTOpImageMultiCompose> NewOp = UE::Mutable::Private::Clone<ASTOpImageMultiCompose>(at);
+
+				Ptr<ASTOp> BaseOp = NewOp->Base.child();
+				NewOp->Base = Visit(BaseOp, currentMipmapOp);
+
+				Ptr<ASTOp> SourceImageOp = NewOp->SourceImage.child();
+				NewOp->SourceImage = Visit(SourceImageOp, currentMipmapOp);
+
+				newAt = NewOp;
+			}
+
+			break;
+		}
+
 		case EOpType::IM_PATCH:
 		{
 			if (!currentMipmapOp->bOnlyTail)
diff --git a/Engine/Plugins/Mutable/Source/MutableTools/Private/MuT/ASTOpImageMultiCompose.cpp b/Engine/Plugins/Mutable/Source/MutableTools/Private/MuT/ASTOpImageMultiCompose.cpp
new file mode 100644
--- /dev/null
+++ b/Engine/Plugins/Mutable/Source/MutableTools/Private/MuT/ASTOpImageMultiCompose.cpp
@@ -0,0 +1,274 @@
+// Copyright Epic Games, Inc. All Rights Reserved.
+
+#include "MuT/ASTOpImageMultiCompose.h"
+
+#include "Containers/Array.h"
+#include "Containers/Map.h"
+#include "HAL/PlatformMath.h"
+#include "MuR/Layout.h"
+#include "MuR/ModelPrivate.h"
+#include "MuR/MutableMath.h"
+#include "MuR/MutableTrace.h"
+#include "MuR/ImagePrivate.h"
+#include "MuR/RefCounted.h"
+#include "MuR/Types.h"
+#include "MuT/ASTOpConstantResource.h"
+#include "MuT/ASTOpImagePixelFormat.h"
+#include "MuT/ASTOpImageResize.h"
+
+
+namespace UE::Mutable::Private
+{
+
+ASTOpImageMultiCompose::ASTOpImageMultiCompose()
+    : Layout(this)
+    , Base(this)
+	, SourceLayout(this)
+	, SourceImage(this)
+{
+}
+
+
+ASTOpImageMultiCompose::~ASTOpImageMultiCompose()
+{
+    // Explicit call needed to avoid recursive destruction
+    ASTOp::RemoveChildren();
+}
+
+
+bool ASTOpImageMultiCompose::IsEqual(const ASTOp& OtherUntyped) const
+{
+	if (OtherUntyped.GetOpType() == GetOpType())
+    {
+		const ASTOpImageMultiCompose* Other = static_cast<const ASTOpImageMultiCompose*>(&OtherUntyped);
+		return 
+			Layout == Other->Layout &&
+			Base == Other->Base &&
+			SourceLayout == Other->SourceLayout &&
+			SourceImage == Other->SourceImage &&
+			LayoutBlockSizeInPixelsX == Other->LayoutBlockSizeInPixelsX &&
+			LayoutBlockSizeInPixelsY == Other->LayoutBlockSizeInPixelsY &&
+			SourceSizeX == Other->SourceSizeX &&
+			SourceSizeY == Other->SourceSizeY;
+			
+    }
+
+    return false;
+}
+
+
+uint64 ASTOpImageMultiCompose::Hash() const
+{
+	uint64 Res = std::hash<EOpType>()(EOpType::IM_MULTICOMPOSE);
+    hash_combine(Res, Layout.child().get());
+	hash_combine(Res, Base.child().get());
+	hash_combine(Res, SourceLayout.child().get());
+	hash_combine(Res, SourceImage.child().get());
+	return Res;
+}
+
+
+UE::Mutable::Private::Ptr<ASTOp> ASTOpImageMultiCompose::Clone(MapChildFuncRef MapChild) const
+{
+	UE::Mutable::Private::Ptr<ASTOpImageMultiCompose> Clone = new ASTOpImageMultiCompose();
+    Clone->Layout = MapChild(Layout.child());
+	Clone->Base = MapChild(Base.child());
+	Clone->SourceLayout = MapChild(SourceLayout.child());
+	Clone->SourceImage = MapChild(SourceImage.child());
+
+	Clone->LayoutBlockSizeInPixelsX = LayoutBlockSizeInPixelsX;
+	Clone->LayoutBlockSizeInPixelsY = LayoutBlockSizeInPixelsY;
+
+	Clone->SourceSizeX = SourceSizeX;
+	Clone->SourceSizeY = SourceSizeY;
+
+    return Clone;
+}
+
+
+void ASTOpImageMultiCompose::ForEachChild(const TFunctionRef<void(ASTChild&)> Function)
+{
+	Function(Layout);
+	Function(Base);
+	Function(SourceLayout);
+	Function(SourceImage);
+}
+
+
+void ASTOpImageMultiCompose::Link( FProgram& Program, FLinkerOptions*)
+{
+    // Already linked?
+    if (!linkedAddress)
+    {
+        OP::ImageMultiComposeArgs Args;
+		FMemory::Memzero(Args);
+
+		if (Layout) 
+		{
+			Args.Layout = Layout->linkedAddress;
+		}
+
+		if (Base) 
+		{
+			Args.Base = Base->linkedAddress;
+		}
+
+		if (SourceLayout) 
+		{
+			Args.SourceLayout = SourceLayout->linkedAddress;
+		}
+
+		if (SourceImage) 
+		{
+			Args.SourceImage = SourceImage->linkedAddress;
+		}
+
+		Args.LayoutBlockSizeInPixelsX = LayoutBlockSizeInPixelsX;
+		Args.LayoutBlockSizeInPixelsY = LayoutBlockSizeInPixelsY;
+		
+		Args.SourceSizeX = SourceSizeX;
+		Args.SourceSizeY = SourceSizeY;
+
+        linkedAddress = (OP::ADDRESS)Program.OpAddress.Num();
+		Program.OpAddress.Add((uint32)Program.ByteCode.Num());
+        AppendCode(Program.ByteCode,EOpType::IM_MULTICOMPOSE);
+        AppendCode(Program.ByteCode, Args);
+    }
+
+}
+
+
+FImageDesc ASTOpImageMultiCompose::GetImageDesc( bool bReturnBestOption, FGetImageDescContext* Context ) const
+{
+    FImageDesc Result;
+
+    // Local context in case it is necessary
+    FGetImageDescContext LocalContext;
+    if (!Context)
+    {
+		Context = &LocalContext;
+    }
+    else
+    {
+        // Cached result?
+		FImageDesc* PtrValue = Context->m_results.Find(this);
+		if (PtrValue)
+		{
+			return *PtrValue;
+		}
+	}
+
+    // Actual work
+    if (Base)
+    {
+		Result = Base->GetImageDesc(bReturnBestOption, Context);
+    }
+
+	if (SourceImage)
+	{
+		FImageDesc SourceImageDesc = SourceImage->GetImageDesc(bReturnBestOption, Context);
+		Result.m_format = GetMostGenericFormat(Result.m_format, SourceImageDesc.m_format);
+	}
+
+    // Cache th result
+    if (Context)
+    {
+		Context->m_results.Add(this, Result);
+	}
+
+    return Result;
+}
+
+
+UE::Mutable::Private::Ptr<ImageSizeExpression> ASTOpImageMultiCompose::GetImageSizeExpression() const
+{
+    if ( Base )
+    {
+        return Base->GetImageSizeExpression();
+    }
+
+    return nullptr;
+}
+
+
+bool ASTOpImageMultiCompose::IsImagePlainConstant(FVector4f& Color) const
+{
+	bool bIsImagePlainConstant = false;
+
+	if (SourceImage.child())
+	{
+		bIsImagePlainConstant = SourceImage->IsImagePlainConstant(Color);
+	}
+
+	if (bIsImagePlainConstant && Base.child())
+	{
+		FVector4f BaseColor;
+		bIsImagePlainConstant = Base->IsImagePlainConstant(BaseColor);
+		bIsImagePlainConstant &= (Color == BaseColor);
+	}
+
+	return bIsImagePlainConstant;
+}
+
+
+void ASTOpImageMultiCompose::GetLayoutBlockSize(int32* BlockX, int32* BlockY)
+{
+	// Try to follow the base image of the compose, which is the most stable.
+	if (Base)
+	{
+		Base->GetLayoutBlockSize(BlockX, BlockY);
+	}
+
+	// We can only follow the block if the base is empty, since the first block will set the block size.	
+	if (*BlockX == 0)
+	{
+		*BlockX = static_cast<int32>(LayoutBlockSizeInPixelsX);
+		*BlockY = static_cast<int32>(LayoutBlockSizeInPixelsY);
+	}
+}
+
+
+FSourceDataDescriptor ASTOpImageMultiCompose::GetSourceDataDescriptor(FGetSourceDataDescriptorContext* Context) const
+{
+	// Cache management
+	TUniquePtr<FGetSourceDataDescriptorContext> LocalContext;
+	if (!Context)
+	{
+		LocalContext.Reset(new FGetSourceDataDescriptorContext);
+		Context = LocalContext.Get();
+	}
+
+	FSourceDataDescriptor* Found = Context->Cache.Find(this);
+	if (Found)
+	{
+		return *Found;
+	}
+
+	// Not cached: calculate
+	FSourceDataDescriptor Result;
+
+	if (Base)
+	{
+		FSourceDataDescriptor SourceDesc = Base->GetSourceDataDescriptor(Context);
+		Result.CombineWith(SourceDesc);
+	}
+
+	if (SourceImage)
+	{
+		FSourceDataDescriptor SourceDesc = SourceImage->GetSourceDataDescriptor(Context);
+		Result.CombineWith(SourceDesc);
+	}
+
+	//if (Mask)
+	//{
+	//	FSourceDataDescriptor SourceDesc = Mask->GetSourceDataDescriptor(Context);
+	//	Result.CombineWith(SourceDesc);
+	//}
+
+	Context->Cache.Add(this, Result);
+
+	return Result;
+}
+
+
+}
diff --git a/Engine/Plugins/Mutable/Source/MutableTools/Private/MuT/ASTOpImageMultiCompose.h b/Engine/Plugins/Mutable/Source/MutableTools/Private/MuT/ASTOpImageMultiCompose.h
new file mode 100644
--- /dev/null
+++ b/Engine/Plugins/Mutable/Source/MutableTools/Private/MuT/ASTOpImageMultiCompose.h
@@ -0,0 +1,56 @@
+// Copyright Epic Games, Inc. All Rights Reserved.
+
+#pragma once
+
+#include "HAL/Platform.h"
+#include "MuR/Image.h"
+#include "MuR/Operations.h"
+#include "MuR/Ptr.h"
+#include "MuT/AST.h"
+
+
+namespace UE::Mutable::Private
+{
+	struct FProgram;
+
+	class ASTOpImageMultiCompose final : public ASTOp
+	{
+	public:
+
+		ASTChild Layout;
+		ASTChild Base;
+		ASTChild SourceLayout;
+		ASTChild SourceImage;
+		//ASTChild Mask;
+		
+		/** Used to find the size of the image to create. */
+		uint16 LayoutBlockSizeInPixelsX = 0;
+		uint16 LayoutBlockSizeInPixelsY = 0;
+
+		uint16 SourceSizeX = 0;
+		uint16 SourceSizeY = 0;
+
+	public:
+
+		ASTOpImageMultiCompose();
+		ASTOpImageMultiCompose(const ASTOpImageMultiCompose&) = delete;
+		~ASTOpImageMultiCompose();
+
+		virtual EOpType GetOpType() const override { return EOpType::IM_MULTICOMPOSE; }
+		virtual uint64 Hash() const override;
+		virtual bool IsEqual(const ASTOp& otherUntyped) const override;
+		virtual Ptr<ASTOp> Clone(MapChildFuncRef mapChild) const override;
+		virtual void ForEachChild(const TFunctionRef<void(ASTChild&)>) override;
+		virtual void Link(FProgram& program, FLinkerOptions* Options) override;
+		//virtual Ptr<ASTOp> OptimiseSink(const FModelOptimizationOptions& options, FOptimizeSinkContext& context) const override;
+		//virtual Ptr<ASTOp> OptimiseSemantic(const FModelOptimizationOptions& options, int32 Pass) const override;
+		virtual FImageDesc GetImageDesc(bool returnBestOption, FGetImageDescContext* context) const override;
+		virtual void GetLayoutBlockSize(int32* OutBlockX, int32* OutBlockY) override;
+		virtual bool IsImagePlainConstant(FVector4f& OutColor) const override;
+		virtual Ptr<ImageSizeExpression> GetImageSizeExpression() const override;
+		virtual FSourceDataDescriptor GetSourceDataDescriptor(FGetSourceDataDescriptorContext*) const override;
+
+	};
+
+}
+
diff --git a/Engine/Plugins/Mutable/Source/MutableTools/Private/MuT/ASTOpImagePixelFormat.cpp b/Engine/Plugins/Mutable/Source/MutableTools/Private/MuT/ASTOpImagePixelFormat.cpp
--- a/Engine/Plugins/Mutable/Source/MutableTools/Private/MuT/ASTOpImagePixelFormat.cpp
+++ b/Engine/Plugins/Mutable/Source/MutableTools/Private/MuT/ASTOpImagePixelFormat.cpp
@@ -12,6 +12,7 @@
 #include "MuR/Types.h"
 #include "MuT/ASTOpConditional.h"
 #include "MuT/ASTOpImageCompose.h"
+#include "MuT/ASTOpImageMultiCompose.h"
 #include "MuT/ASTOpImageMipmap.h"
 #include "MuT/ASTOpImagePatch.h"
 #include "MuT/ASTOpImageLayer.h"
@@ -445,6 +446,31 @@ namespace UE::Mutable::Private
 			break;
 		}
 
+		case EOpType::IM_MULTICOMPOSE:
+		{
+			if (isBlockFormat)
+			{
+				// Since blocks can be resized at runtime anyway, push the format down and rely on reformatting on the fly if necessary.
+
+				// We move the format down the two paths
+				Ptr<ASTOpImageMultiCompose> newOp = UE::Mutable::Private::Clone<ASTOpImageMultiCompose>(at);
+
+				// TODO: We have to make sure we don't end up with two different formats if
+				// there is an formatIfAlpha
+
+				Ptr<ASTOp> BaseOp = newOp->Base.child();
+				newOp->Base = Visit(BaseOp, currentFormatOp);
+
+				// TODO: Do not force a format to the Source Image it's from a TextureParameter
+				Ptr<ASTOp> SourceImageOp = newOp->SourceImage.child();
+				newOp->SourceImage = Visit(SourceImageOp, currentFormatOp);
+
+				newAt = newOp;
+			}
+
+			break;
+		}
+
 		case EOpType::IM_PATCH:
 		{
 			if (isBlockFormat)
diff --git a/Engine/Plugins/Mutable/Source/MutableTools/Private/MuT/ASTOpImageResize.cpp b/Engine/Plugins/Mutable/Source/MutableTools/Private/MuT/ASTOpImageResize.cpp
--- a/Engine/Plugins/Mutable/Source/MutableTools/Private/MuT/ASTOpImageResize.cpp
+++ b/Engine/Plugins/Mutable/Source/MutableTools/Private/MuT/ASTOpImageResize.cpp
@@ -17,6 +17,7 @@
 #include "MuT/ASTOpImageTransform.h"
 #include "MuT/ASTOpImageSwizzle.h"
 #include "MuT/ASTOpImageCompose.h"
+#include "MuT/ASTOpImageMultiCompose.h"
 #include "MuT/ASTOpImageMakeGrowMap.h"
 #include "MuT/ASTOpImageInterpolate.h"
 #include "MuT/ASTOpImagePlainColor.h"
@@ -356,6 +357,29 @@ namespace UE::Mutable::Private
 			break;
 		}
 
+		case EOpType::IM_MULTICOMPOSE:
+		{
+			Ptr<ASTOpImageMultiCompose> NewOp = UE::Mutable::Private::Clone<ASTOpImageMultiCompose>(SourceAt);
+
+			Ptr<ASTOpImageResize> BaseOp = UE::Mutable::Private::Clone<ASTOpImageResize>(this);
+			BaseOp->Source = NewOp->Base.child();
+			NewOp->Base = BaseOp;
+
+			Ptr<ASTOpImageResize> SourceImageOp = UE::Mutable::Private::Clone<ASTOpImageResize>(this);
+			SourceImageOp->Source = NewOp->SourceImage.child();
+			NewOp->SourceImage = SourceImageOp;
+
+			//if (NewOp->Mask)
+			//{
+			//	Ptr<ASTOpImageResize> MaskOp = UE::Mutable::Private::Clone<ASTOpImageResize>(this);
+			//	MaskOp->Source = NewOp->Mask.child();
+			//	NewOp->Mask = MaskOp;
+			//}
+
+			Result = NewOp;
+			break;
+		}
+
 		case EOpType::IM_RASTERMESH:
 		{
 			Ptr<ASTOpImageRasterMesh> NewOp = UE::Mutable::Private::Clone<ASTOpImageRasterMesh>(SourceAt);
diff --git a/Engine/Plugins/Mutable/Source/MutableTools/Private/MuT/ASTOpImageResizeRel.cpp b/Engine/Plugins/Mutable/Source/MutableTools/Private/MuT/ASTOpImageResizeRel.cpp
--- a/Engine/Plugins/Mutable/Source/MutableTools/Private/MuT/ASTOpImageResizeRel.cpp
+++ b/Engine/Plugins/Mutable/Source/MutableTools/Private/MuT/ASTOpImageResizeRel.cpp
@@ -17,6 +17,7 @@
 #include "MuT/ASTOpImageTransform.h"
 #include "MuT/ASTOpImageSwizzle.h"
 #include "MuT/ASTOpImageCompose.h"
+#include "MuT/ASTOpImageMultiCompose.h"
 #include "MuT/ASTOpImageMakeGrowMap.h"
 #include "MuT/ASTOpImageResize.h"
 #include "MuT/ASTOpImageMultiLayer.h"
@@ -390,6 +391,48 @@ namespace UE::Mutable::Private
 				break;
 			}
 
+			case EOpType::IM_MULTICOMPOSE:
+			{
+				// We can only optimise if the layout grid blocks size in pixels is
+				// still an integer after relative scale
+				bool bIsAcceptable = false;
+				{
+					const ASTOpImageMultiCompose* TypedAt = static_cast<const ASTOpImageMultiCompose*>(at.get());
+					Ptr<ASTOp> OriginalBaseOp = TypedAt->Base.child();
+
+					// \todo: recursion-proof cache?
+					int32 LayoutBlockPixelsX = 0;
+					int32 LayoutBlockPixelsY = 0;
+					OriginalBaseOp->GetLayoutBlockSize(&LayoutBlockPixelsX, &LayoutBlockPixelsY);
+
+					int32 ScaledLayoutBlockPixelsX = int32(LayoutBlockPixelsX * ScaleX);
+					int32 ScaledLayoutBlockPixelsY = int32(LayoutBlockPixelsY * ScaleY);
+					int32 UnscaledLayoutBlockPixelsX = int32(ScaledLayoutBlockPixelsX / ScaleX);
+					int32 UnscaledLayoutBlockPixelsY = int32(ScaledLayoutBlockPixelsY / ScaleY);
+					bIsAcceptable =
+						(LayoutBlockPixelsX != 0 && LayoutBlockPixelsY != 0)
+						&&
+						(LayoutBlockPixelsX == UnscaledLayoutBlockPixelsX)
+						&&
+						(LayoutBlockPixelsY == UnscaledLayoutBlockPixelsY);
+				}
+
+				if (bIsAcceptable)
+				{
+					Ptr<ASTOpImageMultiCompose> NewOp = UE::Mutable::Private::Clone<ASTOpImageMultiCompose>(at);
+
+					Ptr<ASTOp> baseOp = NewOp->Base.child();
+					NewOp->Base = Visit(baseOp, CurrentSinkingOp);
+
+					Ptr<ASTOp> SourceImageOp = NewOp->SourceImage.child();
+					NewOp->SourceImage = Visit(SourceImageOp, CurrentSinkingOp);
+
+					Result = NewOp;
+				}
+
+				break;
+			}
+
 			case EOpType::IM_PATCH:
 			{
 				Ptr<ASTOpImagePatch> NewOp = UE::Mutable::Private::Clone<ASTOpImagePatch>(at);
diff --git a/Engine/Plugins/Mutable/Source/MutableTools/Private/MuT/ASTOpImageSwizzle.cpp b/Engine/Plugins/Mutable/Source/MutableTools/Private/MuT/ASTOpImageSwizzle.cpp
--- a/Engine/Plugins/Mutable/Source/MutableTools/Private/MuT/ASTOpImageSwizzle.cpp
+++ b/Engine/Plugins/Mutable/Source/MutableTools/Private/MuT/ASTOpImageSwizzle.cpp
@@ -13,6 +13,7 @@
 #include "MuR/Types.h"
 #include "MuT/ASTOpConditional.h"
 #include "MuT/ASTOpImageCompose.h"
+#include "MuT/ASTOpImageMultiCompose.h"
 #include "MuT/ASTOpImagePatch.h"
 #include "MuT/ASTOpImagePixelFormat.h"
 #include "MuT/ASTOpImageMultiLayer.h"
@@ -1225,6 +1226,15 @@ namespace UE::Mutable::Private
 			break;
 		}
 
+		case EOpType::IM_MULTICOMPOSE:
+		{
+			Ptr<ASTOpImageMultiCompose> newOp = UE::Mutable::Private::Clone<ASTOpImageMultiCompose>(at);
+			newOp->Base = Visit(newOp->Base.child(), CurrentSwizzleOp);
+			newOp->SourceImage = Visit(newOp->SourceImage.child(), CurrentSwizzleOp);
+			newAt = newOp;
+			break;
+		}
+
 		case EOpType::IM_PATCH:
 		{
 			Ptr<ASTOpImagePatch> newOp = UE::Mutable::Private::Clone<ASTOpImagePatch>(at);
diff --git a/Engine/Plugins/Mutable/Source/MutableTools/Private/MuT/CodeGenerator.cpp b/Engine/Plugins/Mutable/Source/MutableTools/Private/MuT/CodeGenerator.cpp
--- a/Engine/Plugins/Mutable/Source/MutableTools/Private/MuT/CodeGenerator.cpp
+++ b/Engine/Plugins/Mutable/Source/MutableTools/Private/MuT/CodeGenerator.cpp
@@ -103,6 +103,11 @@
 
 namespace UE::Mutable::Private
 {
+	TAutoConsoleVariable<bool> CVarAllowMulticompose(
+		TEXT("mutable.AllowImageMulticompose"),
+		false,
+		TEXT("If set to true, enables the use of multicompose."),
+		ECVF_Default);
 
 	CodeGenerator::CodeGenerator(CompilerOptions::Private* Options, TFunction<void()>& InWaitCallback)
 		: WaitCallback(InWaitCallback),
@@ -2780,6 +2785,12 @@ namespace UE::Mutable::Private
 					// If we don't know the size of a layout block in pixels, calculate it
 					UpdateLayoutBlockDesc(InOutLayoutBlockDesc, ExtendDesc, RectInCells.size);
 
+					if (CVarAllowMulticompose.GetValueOnAnyThread() && ExtendResult.bCanMulticompose)
+					{
+						LastBase = GenerateImageMulticompose(LastBase, BaseMeshResults.LayoutOps[LayoutIndex], ImageOptions, MatchingEdit->Extend, ExtendLayout, InOutLayoutBlockDesc);
+						break;
+					}
+
 					// Adjust the format and size of the block to be added
 					// Actually don't do it, it will be propagated from the top format operation.
 					//FragmentOp = GenerateImageFormat(FragmentOp, FinalImageFormat);
@@ -2874,4 +2885,44 @@ namespace UE::Mutable::Private
 			}
 		}
 	}
+
+
+	bool CodeGenerator::HasImageBlockModifiers(const FParameterKey& ImageKey, const TArray<FirstPassGenerator::FModifier>& Modifiers, int32 LODIndex) const
+	{
+		check(LODIndex >= 0);
+
+		// Process patch image modifiers (from edit modifiers)
+		for (const FirstPassGenerator::FModifier& Modifier : Modifiers)
+		{
+			if (Modifier.Node->GetType() == NodeModifierSurfaceEdit::GetStaticType())
+			{
+				const NodeModifierSurfaceEdit* Edit = static_cast<const NodeModifierSurfaceEdit*>(Modifier.Node);
+
+				bool bAffectsLOD = Edit->LODs.IsValidIndex(LODIndex);
+
+				if (!bAffectsLOD)
+				{
+					continue;
+				}
+
+				const NodeModifierSurfaceEdit::FTexture* MatchingEdit = Edit->LODs[LODIndex].Textures.FindByPredicate(
+					[&](const NodeModifierSurfaceEdit::FTexture& Candidate)
+					{
+						return (Candidate.MaterialParameterKey == ImageKey);
+					});
+
+				if (!MatchingEdit)
+				{
+					continue;
+				}
+
+				if (MatchingEdit->PatchImage.get())
+				{
+					return true;
+				}
+			}
+		}
+
+		return false;
+	}
 }
diff --git a/Engine/Plugins/Mutable/Source/MutableTools/Private/MuT/CodeGenerator.h b/Engine/Plugins/Mutable/Source/MutableTools/Private/MuT/CodeGenerator.h
--- a/Engine/Plugins/Mutable/Source/MutableTools/Private/MuT/CodeGenerator.h
+++ b/Engine/Plugins/Mutable/Source/MutableTools/Private/MuT/CodeGenerator.h
@@ -44,9 +44,10 @@
 #include "Templates/TypeHash.h"
 #include "Tasks/Pipe.h"
 
-
 namespace UE::Mutable::Private
 {
+	extern TAutoConsoleVariable<bool> CVarAllowMulticompose;
+
 	class NodeImageParameterConvert;
 	class NodeExternalParameter;
 	class NodeExternalOperation;
@@ -55,10 +56,6 @@ namespace UE::Mutable::Private
 	class NodeMaterialExternal;
 	class NodeScalarExternal;
 	class NodeMaterialSkeletalMeshBreak;
-}
-
-namespace UE::Mutable::Private
-{
 	class NodeImageExternal;
 	class NodeMeshExternal;
 	class ASTOpMaterialModify;
@@ -144,6 +141,8 @@ namespace UE::Mutable::Private
 		
 		friend class FirstPassGenerator;
 
+		struct FLayoutBlockDesc;
+
     public:
 
         CodeGenerator( CompilerOptions::Private*, TFunction<void()>& InWaitCallback );
@@ -513,6 +512,7 @@ namespace UE::Mutable::Private
 		{
 			Ptr<ASTOp> op;
 			bool Parameter = false;
+			bool bCanMulticompose = true;
 		};
 
 		/** */
@@ -592,6 +592,9 @@ namespace UE::Mutable::Private
 		/** */
 		Ptr<ASTOp> GenerateImageSize(Ptr<ASTOp>, UE::Math::TIntVector2<int32>);
 
+		/** */
+		Ptr<ASTOp> GenerateImageMulticompose(Ptr<ASTOp> BaseOp, Ptr<ASTOp> LayoutOp, const FImageGenerationOptions& Options, const NodeImagePtrConst& NodeImage, TSharedPtr<const FLayout> SourceLayout, const CodeGenerator::FLayoutBlockDesc& InLayoutBlockDesc );
+
 		/** Evaluate if the image to generate is big enough to be split in separate operations and tiled afterwards. */
 		Ptr<ASTOp> ApplyTiling(Ptr<ASTOp> Source, UE::Math::TIntVector2<int32> Size, EImageFormat Format);
 
@@ -1201,6 +1204,7 @@ namespace UE::Mutable::Private
 
 		void CheckModifiersForSurface(const NodeSurfaceNew&, const TArray<FirstPassGenerator::FModifier>&, int32 LODIndex);
 
+		bool HasImageBlockModifiers(const FParameterKey& ImageKey, const TArray<FirstPassGenerator::FModifier>&, int32 LODIndex) const;
 
 		void GenerateMaterialImage(
 			const TPair<FParameterKey, FImageParameterData>& ImageData,
@@ -1213,7 +1217,6 @@ namespace UE::Mutable::Private
 			const FMeshGenerationResult& MeshResults,
 			const TBitArray<>& LayoutFromExtension,
 			const void* ErrorContext);
-
     };
 
 	
diff --git a/Engine/Plugins/Mutable/Source/MutableTools/Private/MuT/CodeGenerator_Image.cpp b/Engine/Plugins/Mutable/Source/MutableTools/Private/MuT/CodeGenerator_Image.cpp
--- a/Engine/Plugins/Mutable/Source/MutableTools/Private/MuT/CodeGenerator_Image.cpp
+++ b/Engine/Plugins/Mutable/Source/MutableTools/Private/MuT/CodeGenerator_Image.cpp
@@ -39,6 +39,7 @@
 #include "MuT/ASTOpImagePlainColor.h"
 #include "MuT/ASTOpImageDisplace.h"
 #include "MuT/ASTOpImageInvert.h"
+#include "MuT/ASTOpImageMultiCompose.h"
 #include "MuT/ASTOpMaterialBreak.h"
 #include "MuT/ASTOpMaterialBreakImageParameter.h"
 #include "MuT/ASTOpMeshExtractLayoutBlocks.h"
@@ -222,6 +223,9 @@ namespace UE::Mutable::Private
     {
 		const NodeImageConstant& node = *InNode;
 		
+		// Image not suitable for multicompose operations
+		Result.bCanMulticompose = false;
+
         // TODO: check duplicates
         TSharedPtr<const FImage> pImage = InNode->Image;
 
@@ -312,6 +316,9 @@ namespace UE::Mutable::Private
 	{
 		MUTABLE_CPUPROFILER_SCOPE(NodeImageLayer);
 
+		// Image not suitable for multicompose operations
+		Result.bCanMulticompose = false;
+
         Ptr<ASTOpImageLayer> op = new ASTOpImageLayer();
 
         op->blendType = Node->Type;
@@ -386,16 +393,19 @@ namespace UE::Mutable::Private
 	{		
 		MUTABLE_CPUPROFILER_SCOPE(NodeImageLayerColour);
 
+		bool bCanMulticompose = true;
+
         Ptr<ASTOpImageLayerColor> op = new ASTOpImageLayerColor();
 		op->blendType = Node->Type;
 
         // Base image
         Ptr<ASTOp> base;
-        if (Node->Base )
+        if (Node->Base)
         {
 			FImageGenerationResult BaseResult;
 			GenerateImage(Options, BaseResult, Node->Base);
 			base = BaseResult.op;
+			bCanMulticompose = BaseResult.bCanMulticompose;
         }
         else
         {
@@ -418,11 +428,12 @@ namespace UE::Mutable::Private
 
         // Mask of the effect
         Ptr<ASTOp> mask;
-        if (Node->Mask )
+        if (Node->Mask)
         {
 			FImageGenerationResult MaskResult;
 			GenerateImage(Options, MaskResult, Node->Mask);
 			mask = MaskResult.op;
+			bCanMulticompose = bCanMulticompose && MaskResult.bCanMulticompose;
 			
 			mask = GenerateImageFormat( mask, EImageFormat::L_UByte );
             mask = GenerateImageSize( mask, TargetSize);
@@ -445,6 +456,7 @@ namespace UE::Mutable::Private
         op->color = colour;
 
         Result.op = op;
+		Result.bCanMulticompose = bCanMulticompose;
     }
 
 
@@ -452,6 +464,9 @@ namespace UE::Mutable::Private
 	{
 		MUTABLE_CPUPROFILER_SCOPE(NodeImageMultiLayer);
 
+		// Image not suitable for multicompose operations
+		Result.bCanMulticompose = false;
+
         Ptr<ASTOpImageMultiLayer> op = new ASTOpImageMultiLayer();
 
 		op->blendType = Node->Type;
@@ -533,6 +548,9 @@ namespace UE::Mutable::Private
 	{
 		MUTABLE_CPUPROFILER_SCOPE(NodeImageNormalComposite);
 
+		// Image not suitable for multicompose operations
+		Result.bCanMulticompose = false;
+
         Ptr<ASTOpImageNormalComposite> op = new ASTOpImageNormalComposite();
 
 		op->Mode = Node->Mode;
@@ -568,7 +586,7 @@ namespace UE::Mutable::Private
         op->Base = base;
 
         Ptr<ASTOp> normal;
-        if (Node->Normal )
+        if (Node->Normal)
         {
 			FImageGenerationResult BaseResult;
 			GenerateImage(Options, BaseResult, Node->Normal);
@@ -594,6 +612,9 @@ namespace UE::Mutable::Private
     {
         MUTABLE_CPUPROFILER_SCOPE(NodeImageTransform);
 
+		// Image not suitable for multicompose operations
+		Result.bCanMulticompose = false;
+
         Ptr<ASTOpImageTransform> Op = new ASTOpImageTransform();
 
 		Ptr<ASTOp> OffsetX;
@@ -723,6 +744,9 @@ namespace UE::Mutable::Private
 	{
 		MUTABLE_CPUPROFILER_SCOPE(NodeImageInterpolate);
 
+		// Image not suitable for multicompose operations
+		Result.bCanMulticompose = false;
+
         Ptr<ASTOpImageInterpolate> op = new ASTOpImageInterpolate();
 
         // Factor
@@ -784,6 +808,8 @@ namespace UE::Mutable::Private
 
 	void CodeGenerator::GenerateImage_Swizzle(const FImageGenerationOptions& Options, FImageGenerationResult& Result, const NodeImageSwizzle* Node)
 	{
+		bool bCanMulticompose = true;
+
 		if (Node->NewFormat == EImageFormat::None)
 		{
 			Result.op = GenerateMissingImageCode(TEXT("Make Texture channel."), EImageFormat::L_UByte, Node->GetMessageContext(), Options);
@@ -850,6 +876,7 @@ namespace UE::Mutable::Private
 				FImageGenerationResult BaseResult;
 				GenerateImage(Options, BaseResult, Node->Sources[SourceIndex]);
                 Ptr<ASTOp> Source = BaseResult.op;
+				bCanMulticompose = bCanMulticompose && BaseResult.bCanMulticompose;
 
 				if (!Source)
 				{
@@ -908,6 +935,7 @@ namespace UE::Mutable::Private
 		}
 
         Result.op = ResultOp;
+		Result.bCanMulticompose = bCanMulticompose;
 	}
 
 
@@ -915,9 +943,11 @@ namespace UE::Mutable::Private
 	{
 		MUTABLE_CPUPROFILER_SCOPE(NodeImageFormat);
 
-        check(Node->Format != EImageFormat::None);
-
-        Ptr<ASTOpImagePixelFormat> fop = new ASTOpImagePixelFormat();
+		bool bCanMulticompose = true; 
+        
+		check(Node->Format != EImageFormat::None);
+        
+		Ptr<ASTOpImagePixelFormat> fop = new ASTOpImagePixelFormat();
         fop->Format = Node->Format;
         fop->FormatIfAlpha = Node->FormatIfAlpha;
 
@@ -932,23 +962,28 @@ namespace UE::Mutable::Private
 			GenerateImage(Options, BaseResult, Node->Source);
 			fop->Source = BaseResult.op;
 			Result.Parameter = BaseResult.Parameter;
+			bCanMulticompose = BaseResult.bCanMulticompose;
 		}
 
         Result.op = fop;
+		Result.bCanMulticompose = bCanMulticompose;
 	}
 
 
 	void CodeGenerator::GenerateImage_Saturate(const FImageGenerationOptions& Options, FImageGenerationResult& Result, const NodeImageSaturate* Node)
 	{
-		Ptr<ASTOpImageSaturate> op = new ASTOpImageSaturate();
+		bool bCanMulticompose = true; 
 
-        // Source image
+		Ptr<ASTOpImageSaturate> op = new ASTOpImageSaturate();
+        
+		// Source image
         Ptr<ASTOp> Base;
         if (Node->Source )
         {
 			FImageGenerationResult BaseResult;
 			GenerateImage(Options, BaseResult, Node->Source);
 			Base = BaseResult.op;
+			bCanMulticompose = BaseResult.bCanMulticompose;
 		}
         else
         {
@@ -975,20 +1010,24 @@ namespace UE::Mutable::Private
         }
 
         Result.op = op;
+		Result.bCanMulticompose = bCanMulticompose;
     }
 
 
 	void CodeGenerator::GenerateImage_Luminance(const FImageGenerationOptions& Options, FImageGenerationResult& Result, const NodeImageLuminance* Node)
 	{
+		bool bCanMulticompose = true; 
+		
 		Ptr<ASTOpImageLuminance> op = new ASTOpImageLuminance();
-
-        // Source image
+        
+		// Source image
         Ptr<ASTOp> base;
         if (Node->Source )
         {
 			FImageGenerationResult BaseResult;
 			GenerateImage(Options, BaseResult, Node->Source);
 			base = BaseResult.op;
+			bCanMulticompose = BaseResult.bCanMulticompose;
         }
         else
         {
@@ -999,11 +1038,15 @@ namespace UE::Mutable::Private
         op->Base = base;
 
         Result.op = op;
+		Result.bCanMulticompose = bCanMulticompose;
     }
 
 
 	void CodeGenerator::GenerateImage_ColourMap(const FImageGenerationOptions& Options, FImageGenerationResult& Result, const NodeImageColourMap* Node)
 	{
+		// Image not suitable for multicompose operations
+		Result.bCanMulticompose = false;
+
 		Ptr<ASTOpImageColorMap> op = new ASTOpImageColorMap();
 
         // Base image
@@ -1070,6 +1113,8 @@ namespace UE::Mutable::Private
 	{
 		Ptr<ASTOpImageBinarize> op = new ASTOpImageBinarize();
 
+		bool bCanMulticompose = true; 
+
         // A image
         Ptr<ASTOp> BaseOp;
         if (Node->Base )
@@ -1077,6 +1122,7 @@ namespace UE::Mutable::Private
 			FImageGenerationResult BaseResult;
 			GenerateImage(Options, BaseResult, Node->Base);
 			BaseOp = BaseResult.op;
+			bCanMulticompose = BaseResult.bCanMulticompose;
 		}
         else
         {
@@ -1103,6 +1149,7 @@ namespace UE::Mutable::Private
         op->Threshold = ThresholdOp;
 
         Result.op = op;
+		Result.bCanMulticompose = bCanMulticompose;
     }
 
 
@@ -1113,7 +1160,8 @@ namespace UE::Mutable::Private
         // Source image
         Ptr<ASTOp> base;
 		bool bParameter = false;
-        if (Node->Base )
+		bool bCanMulticompose = true;
+        if (Node->Base)
         {
 			FImageGenerationOptions NewOptions = Options;
 			
@@ -1127,6 +1175,7 @@ namespace UE::Mutable::Private
 			GenerateImage(NewOptions, BaseResult, Node->Base);
         	base = BaseResult.op;
         	bParameter = BaseResult.Parameter;
+			bCanMulticompose = BaseResult.bCanMulticompose;
         }
         else
         {
@@ -1135,7 +1184,7 @@ namespace UE::Mutable::Private
         }
 
         // Size
-        if (Node->bRelative )
+        if (Node->bRelative)
         {
             Ptr<ASTOpImageResizeRel> NewOp = new ASTOpImageResizeRel();
 			NewOp->Factor[0] = Node->SizeX;
@@ -1177,6 +1226,8 @@ namespace UE::Mutable::Private
         		GenerateCrop(Options, Result, *Node);
         	}
         }
+
+		Result.bCanMulticompose = bCanMulticompose;
     }
 
 
@@ -1244,6 +1295,8 @@ namespace UE::Mutable::Private
 	{
 		MUTABLE_CPUPROFILER_SCOPE(NodeImageSwitch);
 
+		bool bCanMulticompose = true;
+
         if (Node->Options.Num() == 0)
 		{
 			// No options in the switch!
@@ -1277,6 +1330,7 @@ namespace UE::Mutable::Private
 				FImageGenerationResult BaseResult;
 				GenerateImage(Options, BaseResult, Node->Options[OptionIndex]);
 				branch = BaseResult.op;
+				bCanMulticompose = bCanMulticompose && BaseResult.bCanMulticompose;
 			}
             else
             {
@@ -1290,6 +1344,7 @@ namespace UE::Mutable::Private
         Ptr<ASTOp> switchAt = op;
 
         Result.op = switchAt;
+		Result.bCanMulticompose = bCanMulticompose;
     }
 
 
@@ -1299,7 +1354,7 @@ namespace UE::Mutable::Private
         op->type = EOpType::IM_CONDITIONAL;
 
         // Condition
-        if (Node->Parameter )
+        if (Node->Parameter)
         {
 			FBoolGenerationResult ParamResult;
 			GenerateBool(ParamResult, Options, Node->Parameter.get());
@@ -1321,6 +1376,7 @@ namespace UE::Mutable::Private
 		op->no = NoResult.op;
 		
         Result.op = op;
+		Result.bCanMulticompose = NoResult.bCanMulticompose && YesResult.bCanMulticompose;
     }
 
 
@@ -1328,6 +1384,9 @@ namespace UE::Mutable::Private
 	{
 		MUTABLE_CPUPROFILER_SCOPE(NodeImageProject);
 
+		// Image not suitable for multicompose operations
+		Result.bCanMulticompose = false;
+
         // Mesh project operation
         //------------------------------
         Ptr<ASTOpMeshProject> ProjectOp = new ASTOpMeshProject();
@@ -1665,6 +1724,8 @@ namespace UE::Mutable::Private
 	{
 		MUTABLE_CPUPROFILER_SCOPE(NodeImageMipmap);
 
+		bool bCanMulticompose = true;
+
         Ptr<ASTOp> res;
 
         Ptr<ASTOpImageMipmap> op = new ASTOpImageMipmap();
@@ -1680,6 +1741,7 @@ namespace UE::Mutable::Private
 			FImageGenerationResult BaseResult;
 			GenerateImage(Options, BaseResult, Node->Source);
 			base = BaseResult.op;
+			bCanMulticompose = BaseResult.bCanMulticompose;
 		}
         else
         {
@@ -1693,9 +1755,7 @@ namespace UE::Mutable::Private
         // optimisation operations. Scan the source image code looking for this info
         int32 blockX = 0;
         int32 blockY = 0;
-        if ( Options.ImageLayoutStrategy
-             !=
-             CompilerOptions::TextureLayoutStrategy::None )
+        if (Options.ImageLayoutStrategy != CompilerOptions::TextureLayoutStrategy::None)
         {
             MUTABLE_CPUPROFILER_SCOPE(GetLayoutBlockSize);
             op->Source->GetLayoutBlockSize( &blockX, &blockY );
@@ -1719,11 +1779,13 @@ namespace UE::Mutable::Private
         res = op;
 
         Result.op = res;
+		Result.bCanMulticompose = bCanMulticompose;
     }
 
 
 	void CodeGenerator::GenerateImage_Invert(const FImageGenerationOptions& Options, FImageGenerationResult& Result, const NodeImageInvert* Node)
 	{
+		bool bCanMulticompose = true;
 		Ptr<ASTOpImageInvert> op = new ASTOpImageInvert();
 
 		Ptr<ASTOp> a;
@@ -1732,6 +1794,7 @@ namespace UE::Mutable::Private
 			FImageGenerationResult BaseResult;
 			GenerateImage(Options, BaseResult, Node->Base);
 			a = BaseResult.op;
+			bCanMulticompose = BaseResult.bCanMulticompose;
 		}
 		else
 		{
@@ -1743,11 +1806,13 @@ namespace UE::Mutable::Private
 		op->Base = a;
 
 		Result.op = op;
+		Result.bCanMulticompose = bCanMulticompose;
 	}
 
 
 	void CodeGenerator::GenerateImage_Variation(const FImageGenerationOptions& Options, FImageGenerationResult& Result, const NodeImageVariation* Node)
 	{
+		bool bCanMulticompose = true;
 		Ptr<ASTOp> currentOp;
 
         // Default case
@@ -1785,6 +1850,7 @@ namespace UE::Mutable::Private
 				FImageGenerationResult VariationResult;
 				GenerateImage(Options, VariationResult, Node->Variations[t].Image);
 				variationOp = VariationResult.op;
+				bCanMulticompose = bCanMulticompose && VariationResult.bCanMulticompose;
             }
             else
             {
@@ -1833,6 +1899,7 @@ namespace UE::Mutable::Private
         }
 
         Result.op = currentOp;
+		Result.bCanMulticompose = bCanMulticompose;
     }
 
 	
@@ -1840,8 +1907,9 @@ namespace UE::Mutable::Private
 	{
 		const NodeImageTable& node = *InNode;
 
+		bool bCanMulticompose = true;
 		Result.op = GenerateTableSwitch<NodeImageTable, ETableColumnType::Image, EOpType::IM_SWITCH>(node,
-			[this, InNode, Options](const NodeImageTable& node, int32 colIndex, int32 row, FErrorLog* LocalErrorLog)
+			[this, InNode, Options, &bCanMulticompose](const NodeImageTable& node, int32 colIndex, int32 row, FErrorLog* LocalErrorLog)
 			{
 				const FTableValue& CellData = node.Table->GetPrivate()->Rows[row].Values[colIndex];
 				TSharedPtr<FImage> pImage = CellData.Image;
@@ -1874,6 +1942,7 @@ namespace UE::Mutable::Private
 					FImageGenerationResult Result;
 					GenerateImage(Options, Result, ImageConst);
 					ImageOp = Result.op;
+					bCanMulticompose = bCanMulticompose && Result.bCanMulticompose;
 
 					int32 MaxTextureSize = FMath::Max(node.ReferenceImageDesc.m_size[0], node.ReferenceImageDesc.m_size[1]);
 
@@ -1891,6 +1960,7 @@ namespace UE::Mutable::Private
 
 				return ImageOp;
 			});
+		Result.bCanMulticompose = bCanMulticompose;
 	}
 
 
@@ -2116,4 +2186,58 @@ namespace UE::Mutable::Private
         return Result;
     }
 
+
+	Ptr<ASTOp> CodeGenerator::GenerateImageMulticompose(
+			Ptr<ASTOp> InBaseImageOp, 
+			Ptr<ASTOp> InLayoutOp, 
+			const FImageGenerationOptions& InBaseImageOptions,
+			const NodeImagePtrConst& NodeImage, 
+			TSharedPtr<const FLayout> SourceLayout, 
+			const CodeGenerator::FLayoutBlockDesc& InLayoutBlockDesc)
+	{
+		Ptr<ASTOp> Result = InBaseImageOp;
+
+		// Single block layout to generate the entire SourceImage
+		TSharedPtr<FLayout> AuxLayout = SourceLayout->Clone();
+		AuxLayout->Blocks.SetNum(1);
+		AuxLayout->Blocks[0].Min = FIntVector2::ZeroValue;
+		AuxLayout->Blocks[0].Size = AuxLayout->Size;
+
+		// Generate the SourceImage
+		FImageGenerationOptions ImageOptions = InBaseImageOptions;
+		ImageOptions.LayoutToApply = AuxLayout;
+		ImageOptions.LayoutBlockId = AuxLayout->Blocks[0].Id;
+		FImageGenerationResult ImageResult;
+		GenerateImage(ImageOptions, ImageResult, NodeImage);
+
+		Ptr<ASTOp> SourceImage = ImageResult.op;
+		if (SourceImage)
+		{
+			// SourceImage Layout
+			Ptr<ASTOpConstantResource> SourceLayoutOp = new ASTOpConstantResource();
+			SourceLayoutOp->Type = EOpType::LA_CONSTANT;
+			SourceLayoutOp->SetValue(SourceLayout);
+
+			// Compose layout operation
+			Ptr<ASTOpImageMultiCompose> MultiComposeOp = new ASTOpImageMultiCompose();
+			MultiComposeOp->Layout = InLayoutOp;
+			MultiComposeOp->Base = InBaseImageOp;
+			MultiComposeOp->SourceLayout = SourceLayoutOp;
+			MultiComposeOp->SourceImage = SourceImage;
+
+			MultiComposeOp->LayoutBlockSizeInPixelsX = InLayoutBlockDesc.BlockPixelsX;
+			MultiComposeOp->LayoutBlockSizeInPixelsY = InLayoutBlockDesc.BlockPixelsY;
+
+			FImageDesc SourceDesc = SourceImage->GetImageDesc(true);
+			MultiComposeOp->SourceSizeX = SourceDesc.m_size.X;
+			MultiComposeOp->SourceSizeY = SourceDesc.m_size.Y;
+			
+			Result = MultiComposeOp;
+		}
+
+		return Result;
+	}
+
+	
+
 }
diff --git a/Engine/Plugins/Mutable/Source/MutableTools/Private/MuT/CodeGenerator_Material.cpp b/Engine/Plugins/Mutable/Source/MutableTools/Private/MuT/CodeGenerator_Material.cpp
--- a/Engine/Plugins/Mutable/Source/MutableTools/Private/MuT/CodeGenerator_Material.cpp
+++ b/Engine/Plugins/Mutable/Source/MutableTools/Private/MuT/CodeGenerator_Material.cpp
@@ -674,6 +674,14 @@ namespace UE::Mutable::Private
 							// Try to update the layout block desc if we don't know it yet.
 							UpdateLayoutBlockDesc(LayoutBlockDesc, BlockDesc, RectInCells.size);
 
+							// TODO: Allow using multicompose on images with image block modifiers
+							if (CVarAllowMulticompose.GetValueOnAnyThread() && ImageResult.bCanMulticompose && !HasImageBlockModifiers(ImageData.Key, Modifiers, ImageOptions.LODIndex))
+							{
+								// Replace IM_COMPOSE with a IM_MULTICOMPOSE operation.
+								imageAd = GenerateImageMulticompose(imageAd, MeshResults.LayoutOps[LayoutIndex], ImageOptions, pImageNode, pLayout, LayoutBlockDesc);
+								break;
+							}
+
 							// Even if we force the size afterwards, we need some size hint in some cases, like image projections.
 							ImageOptions.RectSize = UE::Math::TIntVector2<int32>(BlockDesc.m_size);
 
diff --git a/Engine/Plugins/Mutable/Source/MutableTools/Private/MuT/CodeOptimiser.cpp b/Engine/Plugins/Mutable/Source/MutableTools/Private/MuT/CodeOptimiser.cpp
--- a/Engine/Plugins/Mutable/Source/MutableTools/Private/MuT/CodeOptimiser.cpp
+++ b/Engine/Plugins/Mutable/Source/MutableTools/Private/MuT/CodeOptimiser.cpp
@@ -355,6 +355,7 @@ namespace UE::Mutable::Private
 					{
 					case EOpType::IM_BLANKLAYOUT:
 					case EOpType::IM_COMPOSE:
+					case EOpType::IM_MULTICOMPOSE:
 					case EOpType::ME_MERGE:
 					case EOpType::ME_CLIPWITHMESH:
 					case EOpType::ME_CLIPMORPHPLANE:
diff --git a/Engine/Plugins/Mutable/Source/MutableTools/Private/MuT/CodeOptimiser_Parameters.cpp b/Engine/Plugins/Mutable/Source/MutableTools/Private/MuT/CodeOptimiser_Parameters.cpp
--- a/Engine/Plugins/Mutable/Source/MutableTools/Private/MuT/CodeOptimiser_Parameters.cpp
+++ b/Engine/Plugins/Mutable/Source/MutableTools/Private/MuT/CodeOptimiser_Parameters.cpp
@@ -282,6 +282,8 @@ namespace UE::Mutable::Private
         // Only process children if there are runtime parameters in the subtree
         processChildren = HasRuntimeParamVisitor.HasAny(At);
 
+		// TODO IM_MULTICOMPOSE
+
         EOpType type = At->GetOpType();
         switch ( type )
         {
diff --git a/Engine/Plugins/Mutable/Source/MutableTools/Private/MuT/Compiler.cpp b/Engine/Plugins/Mutable/Source/MutableTools/Private/MuT/Compiler.cpp
--- a/Engine/Plugins/Mutable/Source/MutableTools/Private/MuT/Compiler.cpp
+++ b/Engine/Plugins/Mutable/Source/MutableTools/Private/MuT/Compiler.cpp
@@ -136,6 +136,7 @@ namespace UE::Mutable::Private
 		{ true,		{ 0, 1, 1, 1, 0, 0, 0, 0, 0, 0, 0, 0, 0, 0, 0, 0, 0, 0, 0, 0, 0, 0, 0, 0, 0, 0, 0, 0, 0 }		},	// IM_RESIZEREL
 		{ true,		{ 1, 1, 1, 1, 1, 1, 1, 1, 1, 1, 1, 1, 0, 0, 0, 0, 0, 0, 0, 0, 0, 0, 0, 0, 0, 0, 0, 0, 0 }		},	// IM_BLANKLAYOUT
 		{ true,		{ 1, 1, 1, 1, 1, 1, 1, 1, 1, 1, 1, 0, 0, 0, 1, 1, 1, 1, 1, 1, 1, 1, 0, 0, 1, 1, 0, 0, 0 }		},	// IM_COMPOSE
+		{ true,		{ 1, 1, 1, 1, 1, 1, 1, 1, 1, 1, 1, 0, 0, 0, 1, 1, 1, 1, 1, 1, 1, 1, 0, 0, 1, 1, 0, 0, 0 }		},	// IM_MULTICOMPOSE
 		{ true,		{ 0, 1, 1, 1, 0, 0, 0, 0, 0, 0, 0, 0, 0, 0, 0, 0, 0, 0, 0, 0, 0, 0, 0, 0, 0, 0, 0, 0, 0 }		},	// IM_INTERPOLATE
 		{ true,		{ 0, 1, 1, 0, 0, 0, 0, 0, 0, 0, 0, 0, 0, 0, 0, 0, 0, 0, 0, 0, 0, 0, 0, 0, 0, 0, 0, 0, 0 }		},	// IM_SATURATE
 		{ true,		{ 0, 1, 1, 0, 0, 0, 0, 0, 0, 0, 0, 0, 0, 0, 0, 0, 0, 0, 0, 0, 0, 0, 0, 0, 0, 0, 0, 0, 0 }		},	// IM_LUMINANCE
diff --git a/Engine/Plugins/Mutable/Source/MutableTools/Private/MuT/DataPacker.cpp b/Engine/Plugins/Mutable/Source/MutableTools/Private/MuT/DataPacker.cpp
--- a/Engine/Plugins/Mutable/Source/MutableTools/Private/MuT/DataPacker.cpp
+++ b/Engine/Plugins/Mutable/Source/MutableTools/Private/MuT/DataPacker.cpp
@@ -18,6 +18,7 @@
 #include "MuT/ASTOpConstantResource.h"
 #include "MuT/ASTOpConstantColor.h"
 #include "MuT/ASTOpImageCompose.h"
+#include "MuT/ASTOpImageMultiCompose.h"
 #include "MuT/ASTOpImageLayer.h"
 #include "MuT/ASTOpImageLayerColor.h"
 #include "MuT/ASTOpImageMultiLayer.h"
@@ -125,6 +126,26 @@ namespace UE::Mutable::Private
                 break;
             }
 
+			case EOpType::IM_MULTICOMPOSE:
+			{
+				recurse = false;
+				const ASTOpImageMultiCompose* op = static_cast<const ASTOpImageMultiCompose*>(node.get());
+
+				TArray<bool> NewState;
+				NewState.Init(false, int32(EImageFormat::Count));
+				RecurseWithState(op->Layout.child(), NewState);
+				RecurseWithState(op->Base.child(), NewState);
+				RecurseWithState(op->SourceLayout.child(), NewState);
+				RecurseWithState(op->SourceImage.child(), NewState);
+
+				//if (op->Mask)
+				//{
+				//	NewState[(int32)EImageFormat::L_UBitRLE] = true;
+				//	RecurseWithState(op->Mask.child(), NewState);
+				//}
+				break;
+			}
+
             case EOpType::IM_LAYERCOLOUR:
             {
                 recurse = false;
diff --git a/Engine/Plugins/Mutable/Source/MutableTools/Private/MuT/ErrorLog.cpp b/Engine/Plugins/Mutable/Source/MutableTools/Private/MuT/ErrorLog.cpp
--- a/Engine/Plugins/Mutable/Source/MutableTools/Private/MuT/ErrorLog.cpp
+++ b/Engine/Plugins/Mutable/Source/MutableTools/Private/MuT/ErrorLog.cpp
@@ -280,6 +280,7 @@ namespace UE::Mutable::Private
 		TEXT("IM_RESIZEREL     "),
 		TEXT("IM_BLANKLAYOUT   "),
 		TEXT("IM_COMPOSE       "),
+		TEXT("IM_MULTICOMPOSE  "),
 		TEXT("IM_INTERPOLATE   "),
 		TEXT("IM_SATURATE      "),
 		TEXT("IM_LUMINANCE     "),
