[Mutable] Remove compilation with image tiling

#rnx
#Synced-CL 48603947

diff --git a/Engine/Plugins/Mutable/Source/CustomizableObject/Internal/MuCO/CustomizableObjectCompilerTypes.h b/Engine/Plugins/Mutable/Source/CustomizableObject/Internal/MuCO/CustomizableObjectCompilerTypes.h
--- a/Engine/Plugins/Mutable/Source/CustomizableObject/Internal/MuCO/CustomizableObjectCompilerTypes.h
+++ b/Engine/Plugins/Mutable/Source/CustomizableObject/Internal/MuCO/CustomizableObjectCompilerTypes.h
@@ -77,9 +77,6 @@ struct FCompilationOptions
 	bool bForceLargeLODBias = false;
 	int32 DebugBias = 0;
 
-	// Control image tiled generation
-	int32 ImageTiling = 0;
-
 	/** If true, gather all game asset references and save them in the Customizable Object. 
 	 *
 	 * This option does not modify the DDC Key. 
diff --git a/Engine/Plugins/Mutable/Source/CustomizableObject/Internal/MuCO/CustomizableObjectPrivate.h b/Engine/Plugins/Mutable/Source/CustomizableObject/Internal/MuCO/CustomizableObjectPrivate.h
--- a/Engine/Plugins/Mutable/Source/CustomizableObject/Internal/MuCO/CustomizableObjectPrivate.h
+++ b/Engine/Plugins/Mutable/Source/CustomizableObject/Internal/MuCO/CustomizableObjectPrivate.h
@@ -2039,9 +2039,6 @@ public:
 	UPROPERTY()
 	uint64 EmbeddedDataBytesLimit = 1024;
 
-	UPROPERTY()
-	int32 ImageTiling = 0;
-
 	UPROPERTY()
 	bool bDisableTableMissingDataWarning = false;
 
diff --git a/Engine/Plugins/Mutable/Source/CustomizableObject/Private/MuCO/CustomizableObject.cpp b/Engine/Plugins/Mutable/Source/CustomizableObject/Private/MuCO/CustomizableObject.cpp
--- a/Engine/Plugins/Mutable/Source/CustomizableObject/Private/MuCO/CustomizableObject.cpp
+++ b/Engine/Plugins/Mutable/Source/CustomizableObject/Private/MuCO/CustomizableObject.cpp
@@ -2457,7 +2457,6 @@ FCompilationOptions UCustomizableObjectPrivate::GetCompileOptions() const
 	Options.bSkinWeightProfilesEnabled = IsAltSkinWeightProfilesEnabled();
 	Options.bPhysicsAssetMergeEnabled = IsPhysicsAssetMergeEnabled();
 	Options.bAnimBpPhysicsManipulationEnabled = IsEnabledAnimBpPhysicsAssetsManipulation();
-	Options.ImageTiling = ImageTiling;
 
 	return Options;
 }
@@ -3184,7 +3183,6 @@ void SerializeCompilationOptionsForDDC(FArchive& Ar, FCompilationOptions& Option
 	Ar << Options.bSkinWeightProfilesEnabled;
 	Ar << Options.bPhysicsAssetMergeEnabled;
 	Ar << Options.bAnimBpPhysicsManipulationEnabled;
-	Ar << Options.ImageTiling;
 	Ar << Options.ParamNamesToSelectedOptions;
 }
 
diff --git a/Engine/Plugins/Mutable/Source/CustomizableObjectEditor/Private/MuCOE/CustomizableObjectCompileRunnable.cpp b/Engine/Plugins/Mutable/Source/CustomizableObjectEditor/Private/MuCOE/CustomizableObjectCompileRunnable.cpp
--- a/Engine/Plugins/Mutable/Source/CustomizableObjectEditor/Private/MuCOE/CustomizableObjectCompileRunnable.cpp
+++ b/Engine/Plugins/Mutable/Source/CustomizableObjectEditor/Private/MuCOE/CustomizableObjectCompileRunnable.cpp
@@ -315,8 +315,6 @@ uint32 FCustomizableObjectCompileRunnable::Run()
 
 	// We always compile for progressive image generation.
 	CompilerOptions->SetEnableProgressiveImages(true);
-	
-	CompilerOptions->SetImageTiling(Options.ImageTiling);
 
 	// On server builds we don't want materials and images to be generted.
 	if (Options.TargetPlatform && Options.TargetPlatform->IsServerOnly())
diff --git a/Engine/Plugins/Mutable/Source/CustomizableObjectEditor/Private/MuCOE/CustomizableObjectEditor.cpp b/Engine/Plugins/Mutable/Source/CustomizableObjectEditor/Private/MuCOE/CustomizableObjectEditor.cpp
--- a/Engine/Plugins/Mutable/Source/CustomizableObjectEditor/Private/MuCOE/CustomizableObjectEditor.cpp
+++ b/Engine/Plugins/Mutable/Source/CustomizableObjectEditor/Private/MuCOE/CustomizableObjectEditor.cpp
@@ -1107,30 +1107,6 @@ TSharedRef<SWidget> FCustomizableObjectEditor::GenerateCompileOptionsMenuContent
 			MenuBuilder.AddWidget(CompileTextureCompressionCombo.ToSharedRef(), LOCTEXT("MutableCompileTextureCompressionType", "Texture Compression"));
 		}
 
-		// Image tiling
-		// Unfortunately SNumericDropDown doesn't work with integers at the time of writing.
-		TArray<SNumericDropDown<float>::FNamedValue> TilingOptions;
-		TilingOptions.Add(SNumericDropDown<float>::FNamedValue(0, FText::FromString(TEXT("0")), FText::FromString(TEXT("Disabled"))));
-		TilingOptions.Add(SNumericDropDown<float>::FNamedValue(64, FText::FromString(TEXT("64")), FText::FromString(TEXT("64"))));
-		TilingOptions.Add(SNumericDropDown<float>::FNamedValue(128, FText::FromString(TEXT("128")), FText::FromString(TEXT("128"))));
-		TilingOptions.Add(SNumericDropDown<float>::FNamedValue(256, FText::FromString(TEXT("256")), FText::FromString(TEXT("256"))));
-		TilingOptions.Add(SNumericDropDown<float>::FNamedValue(512, FText::FromString(TEXT("512")), FText::FromString(TEXT("512"))));
-
-		CompileTilingCombo = SNew(SNumericDropDown<float>)
-			.DropDownValues(TilingOptions)
-			.Value_Lambda([this]() 
-				{ 
-					return CustomizableObject ? float(CustomizableObject->GetPrivate()->ImageTiling) : 0.0f;
-				})
-			.OnValueChanged_Lambda([this](float Value) 
-				{ 
-					if (CustomizableObject)
-					{
-						CustomizableObject->GetPrivate()->ImageTiling = int32(Value);
-						CustomizableObject->Modify();
-					}
-				});
-		MenuBuilder.AddWidget(CompileTilingCombo.ToSharedRef(), LOCTEXT("MutableCompileImageTiling", "Image Tiling"));
 	}
 	MenuBuilder.EndSection();
 
@@ -1550,7 +1526,6 @@ void FCustomizableObjectEditor::ResetCompileOptions()
 	CustomizableObject->GetPrivate()->TextureCompression = DefaultObject->TextureCompression;
 	CustomizableObject->GetPrivate()->EmbeddedDataBytesLimit = DefaultObject->EmbeddedDataBytesLimit;
 	CustomizableObject->GetPrivate()->PackagedDataBytesLimit = DefaultObject->PackagedDataBytesLimit;
-	CustomizableObject->GetPrivate()->ImageTiling = DefaultObject->ImageTiling;
 }
 
 void FCustomizableObjectEditor::OnChangeCompileOptimizationLevel(TSharedPtr<FString> NewSelection, ESelectInfo::Type SelectInfo)
diff --git a/Engine/Plugins/Mutable/Source/CustomizableObjectEditor/Private/MuCOE/CustomizableObjectEditor.h b/Engine/Plugins/Mutable/Source/CustomizableObjectEditor/Private/MuCOE/CustomizableObjectEditor.h
--- a/Engine/Plugins/Mutable/Source/CustomizableObjectEditor/Private/MuCOE/CustomizableObjectEditor.h
+++ b/Engine/Plugins/Mutable/Source/CustomizableObjectEditor/Private/MuCOE/CustomizableObjectEditor.h
@@ -213,7 +213,6 @@ private:
 	TArray< TSharedPtr<FString> > CompileOptimizationStrings;
 	TSharedPtr<STextComboBox> CompileTextureCompressionCombo;
 	TArray< TSharedPtr<FString> > CompileTextureCompressionStrings;
-	TSharedPtr<SNumericDropDown<float>> CompileTilingCombo;
 	TSharedPtr<SNumericDropDown<float>> EmbeddedDataLimitCombo;
 	TSharedPtr<SNumericDropDown<float>> PackagedDataLimitCombo;
 
diff --git a/Engine/Plugins/Mutable/Source/MutableTools/Internal/MuT/Compiler.h b/Engine/Plugins/Mutable/Source/MutableTools/Internal/MuT/Compiler.h
--- a/Engine/Plugins/Mutable/Source/MutableTools/Internal/MuT/Compiler.h
+++ b/Engine/Plugins/Mutable/Source/MutableTools/Internal/MuT/Compiler.h
@@ -65,11 +65,6 @@ namespace UE::Mutable::Private
 		//! with FSystem::SetImagecompressionQuality
 		UE_API void SetImageCompressionQuality(int32 Quality);
 
-		/** Set the image tiling strategy :
-		 * If 0 (default) there is no tiling. Otherwise, images will be generated in tiles of the given size or less, and assembled afterwards as a final step.
-		 */
-		UE_API void SetImageTiling(int32 Tiling);
-
         //! 
         UE_API void SetDataPackingStrategy( int32 MinTextureResidentMipCount, uint64 EmbeddedDataBytesLimit, uint64 PackagedDataBytesLimit);
 
diff --git a/Engine/Plugins/Mutable/Source/MutableTools/Private/MuT/CodeGenerator.cpp b/Engine/Plugins/Mutable/Source/MutableTools/Private/MuT/CodeGenerator.cpp
--- a/Engine/Plugins/Mutable/Source/MutableTools/Private/MuT/CodeGenerator.cpp
+++ b/Engine/Plugins/Mutable/Source/MutableTools/Private/MuT/CodeGenerator.cpp
@@ -791,61 +791,6 @@ namespace UE::Mutable::Private
 	}
 
 
-	Ptr<ASTOp> CodeGenerator::ApplyTiling(Ptr<ASTOp> Source, UE::Math::TIntVector2<int32> Size, EImageFormat Format)
-	{
-		// For now always apply tiling
-		if (CompilerOptions->ImageTiling==0)
-		{
-			return Source;
-		}
-
-		int32 TileSize = CompilerOptions->ImageTiling;
-
-		int32 TilesX = FMath::DivideAndRoundUp<int32>(Size[0], TileSize);
-		int32 TilesY = FMath::DivideAndRoundUp<int32>(Size[1], TileSize);
-		if (TilesX * TilesY <= 2)
-		{
-			return Source;
-		}
-
-		Ptr<ASTOpImagePlainColor> BaseImage = new ASTOpImagePlainColor;
-		BaseImage->Size[0] = Size[0];
-		BaseImage->Size[1] = Size[1];
-		BaseImage->Format = Format;
-		BaseImage->LODs = 1;
-
-		Ptr<ASTOp> CurrentImage = BaseImage;
-
-		for (int32 Y = 0; Y < TilesY; ++Y)
-		{
-			for (int32 X = 0; X < TilesX; ++X)
-			{
-				int32 MinX = X * TileSize;
-				int32 MinY = Y * TileSize;
-				int32 TileSizeX = FMath::Min(TileSize, Size[0] - MinX);
-				int32 TileSizeY = FMath::Min(TileSize, Size[1] - MinY);
-
-				Ptr<ASTOpImageCrop> TileImage = new ASTOpImageCrop();
-				TileImage->Source = Source;
-				TileImage->Min[0] = MinX;
-				TileImage->Min[1] = MinY;
-				TileImage->Size[0] = TileSizeX;
-				TileImage->Size[1] = TileSizeY;
-
-				Ptr<ASTOpImagePatch> PatchedImage = new ASTOpImagePatch();
-				PatchedImage->base = CurrentImage;
-				PatchedImage->patch = TileImage;
-				PatchedImage->location[0] = MinX;
-				PatchedImage->location[1] = MinY;
-
-				CurrentImage = PatchedImage;
-			}
-		}
-
-		return CurrentImage;
-	}
-
-
 	TSharedPtr<FImage> CodeGenerator::GenerateImageBlockPatchMask(const NodeModifierSurfaceEdit::FTexture& Patch, FIntPoint GridSize, int32 BlockPixelsX, int32 BlockPixelsY, box<FIntVector2> RectInCells )
 	{
 		// Create a patching mask for the block
@@ -2800,9 +2745,6 @@ namespace UE::Mutable::Private
 					ExpectedSize[1] = InOutLayoutBlockDesc.BlockPixelsY * RectInCells.size[1];
 					FragmentOp = GenerateImageSize(FragmentOp, ExpectedSize);
 
-					// Apply tiling to avoid generating chunks of image that are too big.
-					FragmentOp = ApplyTiling(FragmentOp, ExpectedSize, InOutLayoutBlockDesc.FinalFormat);
-
 					// Compose operation
 					Ptr<ASTOpImageCompose> ComposeOp = new ASTOpImageCompose();
 					ComposeOp->Layout = BaseMeshResults.LayoutOps[LayoutIndex];
diff --git a/Engine/Plugins/Mutable/Source/MutableTools/Private/MuT/CodeGenerator.h b/Engine/Plugins/Mutable/Source/MutableTools/Private/MuT/CodeGenerator.h
--- a/Engine/Plugins/Mutable/Source/MutableTools/Private/MuT/CodeGenerator.h
+++ b/Engine/Plugins/Mutable/Source/MutableTools/Private/MuT/CodeGenerator.h
@@ -595,9 +595,6 @@ namespace UE::Mutable::Private
 		/** */
 		Ptr<ASTOp> GenerateImageMulticompose(Ptr<ASTOp> BaseOp, Ptr<ASTOp> LayoutOp, const FImageGenerationOptions& Options, const NodeImagePtrConst& NodeImage, TSharedPtr<const FLayout> SourceLayout, const CodeGenerator::FLayoutBlockDesc& InLayoutBlockDesc );
 
-		/** Evaluate if the image to generate is big enough to be split in separate operations and tiled afterwards. */
-		Ptr<ASTOp> ApplyTiling(Ptr<ASTOp> Source, UE::Math::TIntVector2<int32> Size, EImageFormat Format);
-
 		/** Generate a layout block-sized image with a mask including all pixels in the blocks defined in the patch node. */
 		TSharedPtr<FImage> GenerateImageBlockPatchMask(const NodeModifierSurfaceEdit::FTexture&, FIntPoint GridSize, int32 BlockPixelsX, int32 BlockPixelsY, box<FIntVector2> RectInCells);
 
diff --git a/Engine/Plugins/Mutable/Source/MutableTools/Private/MuT/CodeGenerator_Material.cpp b/Engine/Plugins/Mutable/Source/MutableTools/Private/MuT/CodeGenerator_Material.cpp
--- a/Engine/Plugins/Mutable/Source/MutableTools/Private/MuT/CodeGenerator_Material.cpp
+++ b/Engine/Plugins/Mutable/Source/MutableTools/Private/MuT/CodeGenerator_Material.cpp
@@ -690,13 +690,10 @@ namespace UE::Mutable::Private
 							// Enforce block size and optimizations
 							blockAd = GenerateImageSize(blockAd, FIntVector2(BlockDesc.m_size));
 
-							EImageFormat baseFormat = imageAd->GetImageDesc().m_format;
+							//EImageFormat baseFormat = imageAd->GetImageDesc().m_format;
 							// Actually don't do it, it will be propagated from the top format operation.
 							//Ptr<ASTOp> blockAd = GenerateImageFormat(blockAd, baseFormat);
 
-							// Apply tiling to avoid generating chunks of image that are too big.
-							blockAd = ApplyTiling(blockAd, ImageOptions.RectSize, LayoutBlockDesc.FinalFormat);
-
 							// Compose layout operation
 							Ptr<ASTOpImageCompose> composeOp = new ASTOpImageCompose();
 							composeOp->Layout = MeshResults.LayoutOps[LayoutIndex];
diff --git a/Engine/Plugins/Mutable/Source/MutableTools/Private/MuT/Compiler.cpp b/Engine/Plugins/Mutable/Source/MutableTools/Private/MuT/Compiler.cpp
--- a/Engine/Plugins/Mutable/Source/MutableTools/Private/MuT/Compiler.cpp
+++ b/Engine/Plugins/Mutable/Source/MutableTools/Private/MuT/Compiler.cpp
@@ -287,13 +287,6 @@ namespace UE::Mutable::Private
 	}
 
 
-	//---------------------------------------------------------------------------------------------
-	void CompilerOptions::SetImageTiling(int32 Tiling)
-	{
-		m_pD->ImageTiling = Tiling;
-	}
-
-
 	//---------------------------------------------------------------------------------------------
 	void CompilerOptions::SetDataPackingStrategy(int32 MinTextureResidentMipCount, uint64 EmbeddedDataBytesLimit, uint64 PackagedDataBytesLimit)
 	{
diff --git a/Engine/Plugins/Mutable/Source/MutableTools/Private/MuT/CompilerPrivate.h b/Engine/Plugins/Mutable/Source/MutableTools/Private/MuT/CompilerPrivate.h
--- a/Engine/Plugins/Mutable/Source/MutableTools/Private/MuT/CompilerPrivate.h
+++ b/Engine/Plugins/Mutable/Source/MutableTools/Private/MuT/CompilerPrivate.h
@@ -40,7 +40,6 @@ namespace UE::Mutable::Private
 		int32 MinTextureResidentMipCount = 3;
 
         int32 ImageCompressionQuality = 0;
-		int32 ImageTiling = 0 ;
 
 		bool bIgnoreStates = false;
 		bool bLog = false;
