[mutable] Added the Material Break node to the standard CO compilation node caches.
- Removed the "STRING_ADDRESS" typedef as its meaning is incorrect. I have replaced all its usages by the more correct "CONSTANT_STRING_INDEX".
- Removed from "ASTOpConstantResource" the "PASSTHROUGH_ID PassthroughId" as it is not being used and may be confused by the actual ID contained within the value of material-related "ASTOpConstantResources".
- Other small tweaks and changes.

#rnx
#Synced-CL 48643977

diff --git a/Engine/Plugins/Mutable/Source/CustomizableObject/Internal/MuCO/CustomizableObjectPrivate.h b/Engine/Plugins/Mutable/Source/CustomizableObject/Internal/MuCO/CustomizableObjectPrivate.h
--- a/Engine/Plugins/Mutable/Source/CustomizableObject/Internal/MuCO/CustomizableObjectPrivate.h
+++ b/Engine/Plugins/Mutable/Source/CustomizableObject/Internal/MuCO/CustomizableObjectPrivate.h
@@ -1766,6 +1766,8 @@ enum class ECustomizableObjectVersions : int32
 	AddMaterialModify,
 
 	AddImageMultiCompose,
+	
+	DeduplicateFMaterialsSharingPassthroughObject,
 
 	// -----<new versions can be added above this line>--------
 	LastCustomizableObjectVersion
diff --git a/Engine/Plugins/Mutable/Source/CustomizableObjectEditor/Private/MuCOE/GenerateMutableSource/GenerateMutableSource.cpp b/Engine/Plugins/Mutable/Source/CustomizableObjectEditor/Private/MuCOE/GenerateMutableSource/GenerateMutableSource.cpp
--- a/Engine/Plugins/Mutable/Source/CustomizableObjectEditor/Private/MuCOE/GenerateMutableSource/GenerateMutableSource.cpp
+++ b/Engine/Plugins/Mutable/Source/CustomizableObjectEditor/Private/MuCOE/GenerateMutableSource/GenerateMutableSource.cpp
@@ -13,7 +13,6 @@
 #include "MuCO/ICustomizableObjectModule.h"
 #include "MuCOE/CustomizableObjectCompiler.h"
 #include "MuCOE/EdGraphSchema_CustomizableObject.h"
-#include "MuCOE/GenerateMutableSource/GenerateMutableSourceGroupProjector.h"
 #include "MuCOE/GenerateMutableSource/GenerateMutableSourceMacro.h"
 #include "MuCOE/GenerateMutableSource/GenerateMutableSourceModifier.h"
 #include "MuCOE/GraphTraversal.h"
@@ -21,9 +20,7 @@
 #include "MuCOE/Nodes/CustomizableObjectNodeMacroInstance.h"
 #include "MuCOE/Nodes/CustomizableObjectNodeMaterial.h"
 #include "MuCOE/Nodes/CustomizableObjectNodeMeshInterface.h"
-#include "MuCOE/Nodes/CustomizableObjectNodeModifierClipWithMesh.h"
 #include "MuCOE/Nodes/CustomizableObjectNodeObjectGroup.h"
-#include "MuCOE/Nodes/CustomizableObjectNodeStaticString.h"
 #include "MuCOE/Nodes/CustomizableObjectNodeTable.h"
 #include "MuCOE/Nodes/CustomizableObjectNodeTexture.h"
 #include "MuCOE/Nodes/CustomizableObjectNodeParameter.h"
@@ -388,14 +385,29 @@ UObject* FMutableGraphGenerationContext::LoadObject(const FSoftObjectPtr& SoftOb
 }
 
 
-FGeneratedKey::FGeneratedKey(void* InFunctionAddress, const UEdGraphPin& InPin, const UCustomizableObjectNode& Node, FMutableGraphGenerationContext& GenerationContext, const bool UseMesh, const bool InbOnlyConnectedLOD)
+FMaterialBreakParameter FMutableGraphGenerationContext::GetCurrentMaterialBreakParameter()
+{
+	if (MaterialBreakParameterStack.Num() == 0)
+	{
+		return FMaterialBreakParameter();
+	}
+	return MaterialBreakParameterStack.Last();
+}
+
+
+FGeneratedKey::FGeneratedKey(void* InFunctionAddress, const UEdGraphPin& InPin, const UCustomizableObjectNode& Node, FMutableGraphGenerationContext& GenerationContext,  const bool UseMesh, const bool InbOnlyConnectedLOD, const bool bUseMaterialBreakStack)
 {
 	FunctionAddress = InFunctionAddress;
 	Pin = &InPin;
 	LOD = Node.IsAffectedByLOD() ? GenerationContext.CurrentLOD : 0;
 
+	if (bUseMaterialBreakStack)
+	{
+		MaterialBreakParameterStack = GenerationContext.MaterialBreakParameterStack;
+	}
+	
 	// Do not use the Stack of macros as Key if the Node is a Parameter Node. 
-	// By doing so, we can define parameters inside mutabel macros and reuse them without having repeated nodes.
+	// By doing so, we can define parameters inside mutable macros and reuse them without having repeated nodes.
 	MacroContext = Node.IsA(UCustomizableObjectNodeParameter::StaticClass()) ? TArray<const UCustomizableObjectNodeMacroInstance*>{} : GenerationContext.MacroNodesStack;
 
 	if (UseMesh)
@@ -419,6 +431,7 @@ uint32 GetTypeHash(const FGeneratedKey& Key)
 	Hash = HashCombine(Hash, GetTypeHash(Key.bOnlyConnectedLOD));
 	Hash = HashCombine(Hash, GetTypeHash(Key.CurrentMeshComponent));
 	Hash = HashCombine(Hash, GetTypeHash(Key.ReferenceTextureSize));
+	Hash = HashCombine(Hash, GetTypeHash(Key.MaterialBreakParameterStack));
 
 	for (int32 MacroIndex = 0; MacroIndex < Key.MacroContext.Num(); ++MacroIndex)
 	{
@@ -428,6 +441,16 @@ uint32 GetTypeHash(const FGeneratedKey& Key)
 	return Hash;
 }
 
+
+uint32 GetTypeHash(const FMaterialBreakParameter& MaterialBreakParameter)
+{
+	uint32 Hash = GetTypeHash(MaterialBreakParameter.ParameterName.ToString());
+	Hash = HashCombine(Hash, GetTypeHash(MaterialBreakParameter.ParameterType));
+	Hash = HashCombine(Hash, GetTypeHash(MaterialBreakParameter.LayerIndex));
+	return Hash;
+}
+
+
 uint32 GetTypeHash(const FGeneratedSourceExternalKey& Key)
 {
 	uint32 Hash = GetTypeHash(Key.Pin);
diff --git a/Engine/Plugins/Mutable/Source/CustomizableObjectEditor/Private/MuCOE/GenerateMutableSource/GenerateMutableSource.h b/Engine/Plugins/Mutable/Source/CustomizableObjectEditor/Private/MuCOE/GenerateMutableSource/GenerateMutableSource.h
--- a/Engine/Plugins/Mutable/Source/CustomizableObjectEditor/Private/MuCOE/GenerateMutableSource/GenerateMutableSource.h
+++ b/Engine/Plugins/Mutable/Source/CustomizableObjectEditor/Private/MuCOE/GenerateMutableSource/GenerateMutableSource.h
@@ -141,13 +141,29 @@ struct FMorphNodeData
 };
 
 
+/** Representation of the root of a Material Break node compilation branch. */
+struct FMaterialBreakParameter 
+{
+	/** The name of the material parameter. */
+	FName ParameterName;
+	
+	/** The type of data this parameter relates to. */
+	FName ParameterType;
+	
+	/** The index of the parameter layer. */
+	int32 LayerIndex;
+
+	bool operator==(const FMaterialBreakParameter& OtherItem) const = default;
+};
+
+
 // Key for the data stored for each processed unreal graph node.
 class FGeneratedKey
 {
 	friend uint32 GetTypeHash(const FGeneratedKey& Key);
 
 public:
-	FGeneratedKey(void* InFunctionAddress, const UEdGraphPin& InPin, const UCustomizableObjectNode& Node, FMutableGraphGenerationContext& GenerationContext, const bool UseMesh = false, bool bOnlyConnectedLOD = false);
+	FGeneratedKey(void* InFunctionAddress, const UEdGraphPin& InPin, const UCustomizableObjectNode& Node, FMutableGraphGenerationContext& GenerationContext, const bool UseMesh = false, bool bOnlyConnectedLOD = false, const bool bUseMaterialBreakStack = false);
 	
 	bool operator==(const FGeneratedKey& Other) const = default;
 
@@ -175,10 +191,17 @@ public:
 	/** Pointer to control if this is a node inside a Mutable Macro. */
 	TArray<const UCustomizableObjectNodeMacroInstance*> MacroContext;
 
+	/** Stack of Material Break Parameters. It may contain one or multiple entries where the last one is the Material Break Parameter being processed
+	 * while compiling the branches of the graph.
+	 */
+	TArray<FMaterialBreakParameter> MaterialBreakParameterStack;
+	
 	int32 ReferenceTextureSize = -1;
 };
 
 
+uint32 GetTypeHash(const FMaterialBreakParameter& Key);
+
 uint32 GetTypeHash(const FGeneratedKey& Key);
 
 // Struct used for external options of the GenerateMutableSourceX functions.
@@ -259,19 +282,19 @@ struct FGeneratedImagePropertiesKey
 // Structure storing results to propagate up when generating mutable mesh node expressions.
 struct FMutableGraphMeshGenerationData
 {
-	// Did we find any mesh with vertex colours in the expression?
+	// Did we find any mesh with vertex colors in the expression?
 	bool bHasVertexColors = false;
 	bool bHasRealTimeMorphs = false;
 	bool bHasClothing = false;
 
 	// Maximum number of texture channels found in the expression.
-	int NumTexCoordChannels = 0;
+	int32 NumTexCoordChannels = 0;
 
 	// Maximum number of bones per vertex found in the expression.
-	int MaxNumBonesPerVertex = 0;
+	int32 MaxNumBonesPerVertex = 0;
 
 	// Maximum size of the vertex bone index type in the expression.
-	int MaxBoneIndexTypeSizeBytes = 0;
+	int32 MaxBoneIndexTypeSizeBytes = 0;
 
 	int32 MaxNumTriangles = 0;
 	int32 MinNumTriangles = TNumericLimits<int32>::Max();
@@ -696,6 +719,8 @@ struct FMutableGraphGenerationContext
 	/** Key is the Node Uid. */
 	TMap<FGuid, UE::Mutable::Private::Ptr<UE::Mutable::Private::NodeSkeletalMeshParameter>> GeneratedSkeletalMeshParameters;
 
+	/** Key is the Node Uid. */
+	TMap<FGuid, UE::Mutable::Private::Ptr<UE::Mutable::Private::NodeMaterialParameter>> GeneratedMaterialParameters;
 	
 	struct FGeneratedCompositeDataTablesData
 	{
@@ -772,7 +797,7 @@ struct FMutableGraphGenerationContext
 			FString Tags;
 
 			/**
-			* SkeletalMeshNode is needed to disambiguate realtime morph selection from diferent nodes.
+			* SkeletalMeshNode is needed to disambiguate realtime morph selection from different nodes.
 			* TODO: Consider using the actual selection.
 			*/
 			const UCustomizableObjectNode* SkeletalMeshNode = nullptr;
@@ -851,7 +876,13 @@ struct FMutableGraphGenerationContext
 		return UE::Mutable::Private::LoadClass<T>(SoftClass);
 	}
 
-	/** Only compiled components. All components types. Index is the ObjectComponentIndex. */
+	/**
+	 * Get the Current Material Break Parameter that roots the current CO node graph branch. 
+	 * @return The current Material Break Parameter. Will return a default object if there is no Material Break Parameter in the current compilation scope.
+	 */
+	FMaterialBreakParameter GetCurrentMaterialBreakParameter();
+	
+	/** Only compiled components. All component types. Index is the ObjectComponentIndex. */
 	TArray<FName> ComponentNames;
 	
 	TArray<FMutableRefSkeletalMeshData> ReferenceSkeletalMeshesData;
@@ -877,7 +908,7 @@ struct FMutableGraphGenerationContext
 	uint8 CurrentLOD = 0;
 	FName CurrentMeshComponent;
 
-	/** If this is set, we are genreating materials for a "passthrough" component, with a fixed mesh. */
+	/** If this is set, we are generating materials for a "passthrough" component, with a fixed mesh. */
 	UE::Mutable::Private::Ptr<UE::Mutable::Private::NodeMesh> ComponentMeshOverride;
 
 	TMap<FName, uint8> NumLODs;
@@ -966,18 +997,9 @@ struct FMutableGraphGenerationContext
 
 	TMap<FString, TObjectPtr<const UScriptStruct>> ExternalTypeParameterTypes;
 	
-	struct FMaterialBreakParameter {
-		FName ParameterName;
-		FName ParameterType;
-		int32 LayerIndex;
-	};
-
 	/** Name of the material break Node that is being compiled.
-	NOTE: This is set when the material pin of a material breake node is compiled and then it must be reset to empty. */
-	FMaterialBreakParameter CurrentMaterialBreakParameter;
-
-	TMap<const UCustomizableObjectNodeMaterialParameter*, UE::Mutable::Private::Ptr<UE::Mutable::Private::NodeMaterialParameter>> MaterialParameterNodesCache;
-	TMap<const UCONodeMaterialConstant*, UE::Mutable::Private::Ptr<UE::Mutable::Private::NodeMaterialConstant>> MaterialConstantNodesCache;
+	NOTE: This is set when the material pin of a material break node is compiled and then it must be reset to empty. */
+	TArray<FMaterialBreakParameter> MaterialBreakParameterStack;
 };
 
 
diff --git a/Engine/Plugins/Mutable/Source/CustomizableObjectEditor/Private/MuCOE/GenerateMutableSource/GenerateMutableSourceColor.cpp b/Engine/Plugins/Mutable/Source/CustomizableObjectEditor/Private/MuCOE/GenerateMutableSource/GenerateMutableSourceColor.cpp
--- a/Engine/Plugins/Mutable/Source/CustomizableObjectEditor/Private/MuCOE/GenerateMutableSource/GenerateMutableSourceColor.cpp
+++ b/Engine/Plugins/Mutable/Source/CustomizableObjectEditor/Private/MuCOE/GenerateMutableSource/GenerateMutableSourceColor.cpp
@@ -37,7 +37,6 @@
 #include "MuT/NodeColorToSRGB.h"
 #include "MuT/NodeColourExternal.h"
 #include "MuT/NodeColourMaterialBreak.h"
-#include "MuT/NodeImageBinarise.h"
 
 
 #define LOCTEXT_NAMESPACE "CustomizableObjectEditor"
@@ -406,10 +405,10 @@ UE::Mutable::Private::Ptr<UE::Mutable::Private::NodeColour> GenerateMutableSourc
 			MaterialBreak->LayerIndex = LayerIndex;
 
 			// Generate the material that is going to be broken into parameters
-			GenerationContext.CurrentMaterialBreakParameter = FMutableGraphGenerationContext::FMaterialBreakParameter{ ParameterName, Pin->PinType.PinCategory, LayerIndex };
-			UE::Mutable::Private::Ptr<UE::Mutable::Private::NodeMaterial> LinkedNode = GenerateMutableSourceMaterial(LinkedMaterial, GenerationContext, FGenerationMaterialOptions());
+			GenerationContext.MaterialBreakParameterStack.Add(FMaterialBreakParameter{ParameterName, Pin->PinType.PinCategory, LayerIndex });
+			UE::Mutable::Private::Ptr<UE::Mutable::Private::NodeMaterial> LinkedNode = GenerateMutableSourceMaterial(LinkedMaterial, GenerationContext,FGenerationMaterialOptions());
 			MaterialBreak->MaterialSource = LinkedNode;
-			GenerationContext.CurrentMaterialBreakParameter = FMutableGraphGenerationContext::FMaterialBreakParameter{ };
+			GenerationContext.MaterialBreakParameterStack.Pop();
 
 			if (LinkedNode)
 			{
@@ -424,7 +423,7 @@ UE::Mutable::Private::Ptr<UE::Mutable::Private::NodeColour> GenerateMutableSourc
 		{
 			GenerationContext.Log(LOCTEXT("MaterialBreakPinError", "Material Break Nodes must be linked to an input Material Node."), Node);
 		}
-		}
+	}
 
 	else
 	{
diff --git a/Engine/Plugins/Mutable/Source/CustomizableObjectEditor/Private/MuCOE/GenerateMutableSource/GenerateMutableSourceFloat.cpp b/Engine/Plugins/Mutable/Source/CustomizableObjectEditor/Private/MuCOE/GenerateMutableSource/GenerateMutableSourceFloat.cpp
--- a/Engine/Plugins/Mutable/Source/CustomizableObjectEditor/Private/MuCOE/GenerateMutableSource/GenerateMutableSourceFloat.cpp
+++ b/Engine/Plugins/Mutable/Source/CustomizableObjectEditor/Private/MuCOE/GenerateMutableSource/GenerateMutableSourceFloat.cpp
@@ -430,10 +430,10 @@ UE::Mutable::Private::Ptr<UE::Mutable::Private::NodeScalar> GenerateMutableSourc
 			MaterialBreak->LayerIndex = LayerIndex;
 
 			// Generate the material that is going to be broken into parameters
-			GenerationContext.CurrentMaterialBreakParameter = FMutableGraphGenerationContext::FMaterialBreakParameter{ ParameterName, Pin->PinType.PinCategory, LayerIndex };
+			GenerationContext.MaterialBreakParameterStack.Add(FMaterialBreakParameter{ ParameterName, Pin->PinType.PinCategory, LayerIndex });
 			UE::Mutable::Private::Ptr<UE::Mutable::Private::NodeMaterial> LinkedNode = GenerateMutableSourceMaterial(LinkedMaterial, GenerationContext, FGenerationMaterialOptions());
 			MaterialBreak->MaterialSource = LinkedNode;
-			GenerationContext.CurrentMaterialBreakParameter = FMutableGraphGenerationContext::FMaterialBreakParameter{ };
+			GenerationContext.MaterialBreakParameterStack.Pop();
 
 			if (LinkedNode)
 			{
diff --git a/Engine/Plugins/Mutable/Source/CustomizableObjectEditor/Private/MuCOE/GenerateMutableSource/GenerateMutableSourceImage.cpp b/Engine/Plugins/Mutable/Source/CustomizableObjectEditor/Private/MuCOE/GenerateMutableSource/GenerateMutableSourceImage.cpp
--- a/Engine/Plugins/Mutable/Source/CustomizableObjectEditor/Private/MuCOE/GenerateMutableSource/GenerateMutableSourceImage.cpp
+++ b/Engine/Plugins/Mutable/Source/CustomizableObjectEditor/Private/MuCOE/GenerateMutableSource/GenerateMutableSourceImage.cpp
@@ -6,7 +6,6 @@
 #include "Engine/TextureLODSettings.h"
 
 #include "MuCOE/CustomizableObjectCompiler.h"
-#include "MuCOE/CustomizableObjectGraph.h"
 #include "MuCOE/EdGraphSchema_CustomizableObject.h"
 #include "MuCOE/GenerateMutableSource/GenerateMutableSource.h"
 #include "MuCOE/GenerateMutableSource/GenerateMutableSourceColor.h"
@@ -1240,10 +1239,10 @@ UE::Mutable::Private::Ptr<UE::Mutable::Private::NodeImage> GenerateMutableSource
 			MaterialBreak->LayerIndex = LayerIndex;
 
 			// Generate the material that is going to be broken into parameters
-			GenerationContext.CurrentMaterialBreakParameter = FMutableGraphGenerationContext::FMaterialBreakParameter{ ParameterName, Pin->PinType.PinCategory, LayerIndex };
+			GenerationContext.MaterialBreakParameterStack.Add(FMaterialBreakParameter{ ParameterName, Pin->PinType.PinCategory, LayerIndex });
 			UE::Mutable::Private::Ptr<UE::Mutable::Private::NodeMaterial> LinkedNode = GenerateMutableSourceMaterial(LinkedMaterial, GenerationContext, FGenerationMaterialOptions());
 			MaterialBreak->MaterialSource = LinkedNode;
-			GenerationContext.CurrentMaterialBreakParameter = FMutableGraphGenerationContext::FMaterialBreakParameter{ };
+			GenerationContext.MaterialBreakParameterStack.Pop();
 
 			if (LinkedNode)
 			{
diff --git a/Engine/Plugins/Mutable/Source/CustomizableObjectEditor/Private/MuCOE/GenerateMutableSource/GenerateMutableSourceMaterial.cpp b/Engine/Plugins/Mutable/Source/CustomizableObjectEditor/Private/MuCOE/GenerateMutableSource/GenerateMutableSourceMaterial.cpp
--- a/Engine/Plugins/Mutable/Source/CustomizableObjectEditor/Private/MuCOE/GenerateMutableSource/GenerateMutableSourceMaterial.cpp
+++ b/Engine/Plugins/Mutable/Source/CustomizableObjectEditor/Private/MuCOE/GenerateMutableSource/GenerateMutableSourceMaterial.cpp
@@ -58,6 +58,54 @@
 #define LOCTEXT_NAMESPACE "CustomizableObjectEditor"
 
 
+UE::Mutable::Private::Ptr<UE::Mutable::Private::NodeMaterialParameter> GenerateMutableSourceMaterialParameter(const UEdGraphPin* Pin,
+	FMutableGraphGenerationContext& GenerationContext)
+{
+	check(Pin)
+	RETURN_ON_CYCLE(*Pin, GenerationContext)
+
+	CheckNumOutputs(*Pin, GenerationContext);
+	
+	UCustomizableObjectNodeParameter* ParameterNode = CastChecked<UCustomizableObjectNodeParameter>(Pin->GetOwningNode());
+
+	UE::Mutable::Private::Ptr<UE::Mutable::Private::NodeMaterialParameter> Result;
+
+	const FGuid NodeGUID = ParameterNode->NodeGuid;
+	if (const UE::Mutable::Private::Ptr<UE::Mutable::Private::NodeMaterialParameter>* GeneratedMaterialParameter = GenerationContext.GeneratedMaterialParameters.Find(NodeGUID))
+	{
+		return GeneratedMaterialParameter->get();
+	}
+	
+	if (const UCustomizableObjectNodeMaterialParameter* MaterialParameterNode = Cast<UCustomizableObjectNodeMaterialParameter>(ParameterNode))
+	{
+		UE::Mutable::Private::Ptr<UE::Mutable::Private::NodeMaterialParameter> MaterialNode = new UE::Mutable::Private::NodeMaterialParameter();
+		MaterialNode->Name = MaterialParameterNode->GetParameterName(&GenerationContext.MacroNodesStack);
+		MaterialNode->UID = GenerationContext.GetNodeIdUnique(ParameterNode).ToString();
+		
+		GenerationContext.MaterialParameterDefaultValues.Add(MaterialParameterNode->GetParameterName(), UE::Mutable::Private::LoadObject(MaterialParameterNode->DefaultValue));
+		GenerationContext.ParameterUIDataMap.Add(MaterialParameterNode->GetParameterName(&GenerationContext.MacroNodesStack), FMutableParameterData(
+			MaterialParameterNode->ParamUIMetadata,
+			EMutableParameterType::Material));
+
+		Result = MaterialNode;
+	}
+	else
+	{
+		GenerationContext.Log(LOCTEXT("UnimplementedNode", "Node type not implemented yet."), ParameterNode);
+	}
+
+	if (Result)
+	{
+		Result->SetMessageContext(ParameterNode);
+	}
+	
+	GenerationContext.GeneratedMaterialParameters.Add(NodeGUID, Result);
+	GenerationContext.GeneratedNodes.Add(ParameterNode);
+	
+	return Result;
+}
+
+
 UE::Mutable::Private::Ptr<UE::Mutable::Private::NodeMaterial> GenerateMutableSourceMaterial(const UEdGraphPin* Pin, FMutableGraphGenerationContext& GenerationContext, const FGenerationMaterialOptions& MaterialOptions)
 {
 	check(Pin)
@@ -69,7 +117,7 @@ UE::Mutable::Private::Ptr<UE::Mutable::Private::NodeMaterial> GenerateMutableSou
 
 	UCustomizableObjectNode* Node = CastChecked<UCustomizableObjectNode>(Pin->GetOwningNode());
 
-	const FGeneratedKey Key(reinterpret_cast<void*>(&GenerateMutableSourceMaterial), *Pin, *Node, GenerationContext);
+	const FGeneratedKey Key(reinterpret_cast<void*>(&GenerateMutableSourceMaterial), *Pin, *Node, GenerationContext, false, false, true);
 	if (const FGeneratedData* Generated = GenerationContext.Generated.Find(Key))
 	{
 		return static_cast<UE::Mutable::Private::NodeMaterial*>(Generated->Node.get());
@@ -81,11 +129,11 @@ UE::Mutable::Private::Ptr<UE::Mutable::Private::NodeMaterial> GenerateMutableSou
 	}
 
 	// Bool that determines if a node can be added to the cache of nodes.
-	// Most nodes need to be added to the cache but there are some that don't. For exampel, MacroInstanceNodes
+	// Most nodes need to be added to the cache but there are some that don't. For example, MacroInstanceNodes
 	bool bCacheNode = true;
 
 	UE::Mutable::Private::Ptr<UE::Mutable::Private::NodeMaterial> Result;
-
+	
 	if (UCustomizableObjectNodeMaterialBase* TypedNodeMaterialBase = Cast<UCustomizableObjectNodeMaterialBase>(Node))
 	{
 		check(MaterialOptions.SurfaceNode.get());
@@ -373,39 +421,29 @@ UE::Mutable::Private::Ptr<UE::Mutable::Private::NodeMaterial> GenerateMutableSou
 	else if (const UCONodeMaterialConstant* TypedNodeMaterialConstant = Cast<UCONodeMaterialConstant>(Node))
 	{
 		UE::Mutable::Private::Ptr<UE::Mutable::Private::NodeMaterialConstant> MaterialNode;
-		bCacheNode = false; //custom node cache
 
 		if (UMaterialInterface* Material = GenerationContext.LoadObject(TypedNodeMaterialConstant->Material))
 		{
-			if (GenerationContext.MaterialConstantNodesCache.Contains(TypedNodeMaterialConstant))
-			{
-				MaterialNode = GenerationContext.MaterialConstantNodesCache[TypedNodeMaterialConstant];
-			}
-			else
-			{
-				MaterialNode = new UE::Mutable::Private::NodeMaterialConstant();
-
-				GenerationContext.CurrentReferencedMaterial = TStrongObjectPtr(Material);
-				MaterialNode->MaterialId = GenerationContext.AddPassthroughObject(*Material);;
-
-				GenerationContext.MaterialConstantNodesCache.Add(TypedNodeMaterialConstant, MaterialNode);
-			}
+			MaterialNode = new UE::Mutable::Private::NodeMaterialConstant();
+			GenerationContext.CurrentReferencedMaterial = TStrongObjectPtr(Material);
+			MaterialNode->MaterialId = GenerationContext.AddPassthroughObject(*Material);
 
-			// Store parameters that needs to be processed
-			if (!GenerationContext.CurrentMaterialBreakParameter.ParameterType.IsNone())
+			// Store the parameter that needs to be processed
+			const FMaterialBreakParameter CurrentBreakParameter = GenerationContext.GetCurrentMaterialBreakParameter();
+			if (!CurrentBreakParameter.ParameterType.IsNone())
 			{
 				FHashedMaterialParameterInfo ParameterInfo;
-				ParameterInfo.Name = FScriptName(GenerationContext.CurrentMaterialBreakParameter.ParameterName);
-				ParameterInfo.Index = GenerationContext.CurrentMaterialBreakParameter.LayerIndex;
+				ParameterInfo.Name = FScriptName(CurrentBreakParameter.ParameterName);
+				ParameterInfo.Index = CurrentBreakParameter.LayerIndex;
 				ParameterInfo.Association = ParameterInfo.Index == INDEX_NONE ? GlobalParameter : LayerParameter;
 
 				if (!TypedNodeMaterialConstant->Material)
 				{
-					GenerationContext.Log(LOCTEXT("NoRefernceMaterialToBreak", "Could not find a Reference Material to break."), TypedNodeMaterialConstant);
+					GenerationContext.Log(LOCTEXT("NoReferenceMaterialToBreak", "Could not find a Reference Material to break."), TypedNodeMaterialConstant);
 				}
 				else
 				{
-					if (GenerationContext.CurrentMaterialBreakParameter.ParameterType == UEdGraphSchema_CustomizableObject::PC_Texture)
+					if (CurrentBreakParameter.ParameterType == UEdGraphSchema_CustomizableObject::PC_Texture)
 					{
 						UTexture* Texture = nullptr;
 						bool bParameterFound = Material->GetTextureParameterValue(ParameterInfo, Texture);
@@ -428,44 +466,43 @@ UE::Mutable::Private::Ptr<UE::Mutable::Private::NodeMaterial> GenerateMutableSou
 							ConstantImageNode->SourceDataDescriptor.SourceId = CityHash32(reinterpret_cast<const char*>(*TextureName), TextureName.Len() * sizeof(FString::ElementType));
 
 							const uint32 MipsToSkip = ComputeLODBiasForTexture(GenerationContext, *BaseTexture);
-							UE::Mutable::Private::Ptr<UE::Mutable::Private::NodeImage> FinalNodeImag = ResizeTextureByNumMips(ConstantImageNode, MipsToSkip);
 
-							if (FinalNodeImag)
+							if (UE::Mutable::Private::Ptr<UE::Mutable::Private::NodeImage> FinalNodeImag = ResizeTextureByNumMips(ConstantImageNode, MipsToSkip))
 							{
-								MaterialNode->ImageValues.Add({ GenerationContext.CurrentMaterialBreakParameter.ParameterName, GenerationContext.CurrentMaterialBreakParameter.LayerIndex}, FinalNodeImag);
+								MaterialNode->ImageValues.Add({CurrentBreakParameter.ParameterName, CurrentBreakParameter.LayerIndex}, FinalNodeImag);
 							}
 						}
 					}
 					
-					else if (GenerationContext.CurrentMaterialBreakParameter.ParameterType == UEdGraphSchema_CustomizableObject::PC_PassthroughTexture)
+					else if (CurrentBreakParameter.ParameterType == UEdGraphSchema_CustomizableObject::PC_PassthroughTexture)
 					{
-						UTexture* Texture = nullptr;
-						bool bParameterFound = Material->GetTextureParameterValue(ParameterInfo, Texture);
-						UTexture2D* BaseTexture = Cast<UTexture2D>(Texture);
+						UTexture* PassthroughTexture = nullptr;
+						const bool bPassthroughParameterFound = Material->GetTextureParameterValue(ParameterInfo, PassthroughTexture);
+						const UTexture2D* PassthroughBaseTexture = Cast<UTexture2D>(PassthroughTexture);
 
-						if (bParameterFound && BaseTexture)
+						if (bPassthroughParameterFound && PassthroughBaseTexture)
 						{
 							UE::Mutable::Private::Ptr<UE::Mutable::Private::NodeImageConstant> ImageNode = new UE::Mutable::Private::NodeImageConstant();
-							ImageNode->Image = GenerateImageConstant(Texture, GenerationContext, true);
-							MaterialNode->ImageValues.Add({GenerationContext.CurrentMaterialBreakParameter.ParameterName, GenerationContext.CurrentMaterialBreakParameter.LayerIndex}, ImageNode);
+							ImageNode->Image = GenerateImageConstant(PassthroughTexture, GenerationContext, true);
+							MaterialNode->ImageValues.Add({CurrentBreakParameter.ParameterName, CurrentBreakParameter.LayerIndex}, ImageNode);
 						}
 					}
 
-					else if(GenerationContext.CurrentMaterialBreakParameter.ParameterType == UEdGraphSchema_CustomizableObject::PC_Color)
+					else if(CurrentBreakParameter.ParameterType == UEdGraphSchema_CustomizableObject::PC_Color)
 					{
 						FLinearColor Color;
 						if (Material->GetVectorParameterValue(ParameterInfo, Color))
 						{
-							MaterialNode->ColorValues.Add({GenerationContext.CurrentMaterialBreakParameter.ParameterName, GenerationContext.CurrentMaterialBreakParameter.LayerIndex}, Color);
+							MaterialNode->ColorValues.Add({CurrentBreakParameter.ParameterName, CurrentBreakParameter.LayerIndex}, Color);
 						}
 					}
 
-					else if(GenerationContext.CurrentMaterialBreakParameter.ParameterType == UEdGraphSchema_CustomizableObject::PC_Float)
+					else if(CurrentBreakParameter.ParameterType == UEdGraphSchema_CustomizableObject::PC_Float)
 					{
 						float Scalar;
 						if (Material->GetScalarParameterValue(ParameterInfo, Scalar))
 						{
-							MaterialNode->ScalarValues.Add({GenerationContext.CurrentMaterialBreakParameter.ParameterName, GenerationContext.CurrentMaterialBreakParameter.LayerIndex}, Scalar);
+							MaterialNode->ScalarValues.Add({CurrentBreakParameter.ParameterName, CurrentBreakParameter.LayerIndex}, Scalar);
 						}
 					}
 				}
@@ -478,67 +515,67 @@ UE::Mutable::Private::Ptr<UE::Mutable::Private::NodeMaterial> GenerateMutableSou
 	else if (const UCONodeMaterialSwitch* TypedNodeMaterialSwitch = Cast<UCONodeMaterialSwitch>(Node))
 	{
 		Result = [&]()
-		{
-			const UEdGraphPin* SwitchParameter = TypedNodeMaterialSwitch->SwitchParameter();
-
-			// Check Switch Parameter arity preconditions.
-			if (const UEdGraphPin* ConnectedPin = FollowInputPin(*SwitchParameter))
 			{
-				UE::Mutable::Private::Ptr<UE::Mutable::Private::NodeScalar> SwitchParam = GenerateMutableSourceFloat(ConnectedPin, GenerationContext);
+				const UEdGraphPin* SwitchParameter = TypedNodeMaterialSwitch->SwitchParameter();
 
-				// Switch Param not generated
-				if (!SwitchParam)
+				// Check Switch Parameter arity preconditions.
+				if (const UEdGraphPin* ConnectedPin = FollowInputPin(*SwitchParameter))
 				{
-					const FText Message = LOCTEXT("FailedToGenerateSwitchParam", "Could not generate switch enum parameter. Please refesh the switch node and connect an enum.");
-					GenerationContext.Log(Message, Node);
+					UE::Mutable::Private::Ptr<UE::Mutable::Private::NodeScalar> SwitchParam = GenerateMutableSourceFloat(ConnectedPin, GenerationContext);
 
-					return Result;
-				}
-
-				if (SwitchParam->GetType() != UE::Mutable::Private::NodeScalarEnumParameter::GetStaticType())
-				{
-					const FText Message = LOCTEXT("WrongSwitchParamType", "Switch parameter of incorrect type.");
-					GenerationContext.Log(Message, Node);
+					// Switch Param not generated
+					if (!SwitchParam)
+					{
+						const FText Message = LOCTEXT("FailedToGenerateSwitchParam", "Could not generate switch enum parameter. Please refresh the switch node and connect an enum.");
+						GenerationContext.Log(Message, Node);
 
-					return Result;
-				}
+						return Result;
+					}
 
-				const int32 NumSwitchOptions = TypedNodeMaterialSwitch->GetNumElements();
+					if (SwitchParam->GetType() != UE::Mutable::Private::NodeScalarEnumParameter::GetStaticType())
+					{
+						const FText Message = LOCTEXT("WrongSwitchParamType", "Switch parameter of incorrect type.");
+						GenerationContext.Log(Message, Node);
 
-				UE::Mutable::Private::NodeScalarEnumParameter* EnumParameter = static_cast<UE::Mutable::Private::NodeScalarEnumParameter*>(SwitchParam.get());
-				if (NumSwitchOptions != EnumParameter->Options.Num())
-				{
-					const FText Message = LOCTEXT("MismatchedSwitch", "Switch enum and switch node have different number of options. Please refresh the switch node to make sure the outcomes are labeled properly.");
-					GenerationContext.Log(Message, Node);
-				}
+						return Result;
+					}
 
-				UE::Mutable::Private::Ptr<UE::Mutable::Private::NodeMaterialSwitch> SwitchNode = new UE::Mutable::Private::NodeMaterialSwitch;
-				SwitchNode->Parameter = SwitchParam;
-				SwitchNode->Options.SetNum(NumSwitchOptions);
+					const int32 NumSwitchOptions = TypedNodeMaterialSwitch->GetNumElements();
 
-				for (int SelectorIndex = 0; SelectorIndex < NumSwitchOptions; ++SelectorIndex)
-				{
-					if (const UEdGraphPin* MaterialPin = FollowInputPin(*TypedNodeMaterialSwitch->GetElementPin(SelectorIndex)))
+					UE::Mutable::Private::NodeScalarEnumParameter* EnumParameter = static_cast<UE::Mutable::Private::NodeScalarEnumParameter*>(SwitchParam.get());
+					if (NumSwitchOptions != EnumParameter->Options.Num())
 					{
-						SwitchNode->Options[SelectorIndex] = GenerateMutableSourceMaterial(MaterialPin, GenerationContext, MaterialOptions);
+						const FText Message = LOCTEXT("MismatchedSwitch", "Switch enum and switch node have different number of options. Please refresh the switch node to make sure the outcomes are labeled properly.");
+						GenerationContext.Log(Message, Node);
 					}
-					else
+
+					UE::Mutable::Private::Ptr<UE::Mutable::Private::NodeMaterialSwitch> SwitchNode = new UE::Mutable::Private::NodeMaterialSwitch;
+					SwitchNode->Parameter = SwitchParam;
+					SwitchNode->Options.SetNum(NumSwitchOptions);
+
+					for (int SelectorIndex = 0; SelectorIndex < NumSwitchOptions; ++SelectorIndex)
 					{
-						const FText Message = LOCTEXT("MissingMaterial", "Unable to generate material switch node. Required connection not found.");
-						GenerationContext.Log(Message, Node);
-						return Result;
+						if (const UEdGraphPin* MaterialPin = FollowInputPin(*TypedNodeMaterialSwitch->GetElementPin(SelectorIndex)))
+						{
+							SwitchNode->Options[SelectorIndex] = GenerateMutableSourceMaterial(MaterialPin, GenerationContext, MaterialOptions);
+						}
+						else
+						{
+							const FText Message = LOCTEXT("MissingMaterial", "Unable to generate material switch node. Required connection not found.");
+							GenerationContext.Log(Message, Node);
+							return Result;
+						}
 					}
-				}
 
-				Result = SwitchNode;
-				return Result;
-			}
-			else
-			{
-				GenerationContext.Log(LOCTEXT("NoEnumParamInSwitch", "Switch nodes must have an enum switch parameter. Please connect an enum and refesh the switch node."), Node);
-				return Result;
-			}
-		}(); // invoke lambda.
+					Result = SwitchNode;
+					return Result;
+				}
+				else
+				{
+					GenerationContext.Log(LOCTEXT("NoEnumParamInSwitch", "Switch nodes must have an enum switch parameter. Please connect an enum and refesh the switch node."), Node);
+					return Result;
+				}
+			}(); // invoke lambda.
 	}
 
 	else if (const UCONodeMaterialVariation* TypedNodeMaterialVariation = Cast<UCONodeMaterialVariation>(Node))
@@ -578,30 +615,11 @@ UE::Mutable::Private::Ptr<UE::Mutable::Private::NodeMaterial> GenerateMutableSou
 
 	else if (const UCustomizableObjectNodeMaterialParameter* MaterialParameterNode = Cast<UCustomizableObjectNodeMaterialParameter>(Node))
 	{
-		bCacheNode = false;
-		UE::Mutable::Private::Ptr<UE::Mutable::Private::NodeMaterialParameter> MaterialNode;
-
-		if (GenerationContext.MaterialParameterNodesCache.Contains(MaterialParameterNode))
-		{
-			MaterialNode = GenerationContext.MaterialParameterNodesCache[MaterialParameterNode];
-		}
-		else
-		{
-			MaterialNode = new UE::Mutable::Private::NodeMaterialParameter();
-			MaterialNode->Name = MaterialParameterNode->GetParameterName(&GenerationContext.MacroNodesStack);
-			MaterialNode->UID = GenerationContext.GetNodeIdUnique(Node).ToString();
-
-			GenerationContext.MaterialParameterDefaultValues.Add(MaterialParameterNode->GetParameterName(), UE::Mutable::Private::LoadObject(MaterialParameterNode->DefaultValue));
-
-			GenerationContext.ParameterUIDataMap.Add(MaterialParameterNode->GetParameterName(&GenerationContext.MacroNodesStack), FMutableParameterData(
-				MaterialParameterNode->ParamUIMetadata,
-				EMutableParameterType::Material));
-
-			GenerationContext.MaterialParameterNodesCache.Add(MaterialParameterNode, MaterialNode);
-		}
+		UE::Mutable::Private::Ptr<UE::Mutable::Private::NodeMaterialParameter> MaterialNode = GenerateMutableSourceMaterialParameter(Pin, GenerationContext);
 
 		// Store the parameter that needs to be processed
-		if (!GenerationContext.CurrentMaterialBreakParameter.ParameterType.IsNone())
+		const FMaterialBreakParameter CurrentBreakParameter = GenerationContext.GetCurrentMaterialBreakParameter();
+		if (!CurrentBreakParameter.ParameterType.IsNone())
 		{
 			if (!MaterialParameterNode->ReferenceValue)
 			{
@@ -611,15 +629,15 @@ UE::Mutable::Private::Ptr<UE::Mutable::Private::NodeMaterial> GenerateMutableSou
 			{
 				// Data to find the parameter value inside the material
 				FHashedMaterialParameterInfo ParameterInfo;
-				ParameterInfo.Name = FScriptName(GenerationContext.CurrentMaterialBreakParameter.ParameterName);
-				ParameterInfo.Index = GenerationContext.CurrentMaterialBreakParameter.LayerIndex;
+				ParameterInfo.Name = FScriptName(CurrentBreakParameter.ParameterName);
+				ParameterInfo.Index = CurrentBreakParameter.LayerIndex;
 				ParameterInfo.Association = ParameterInfo.Index == INDEX_NONE ? GlobalParameter : LayerParameter;
 
-				if (GenerationContext.CurrentMaterialBreakParameter.ParameterType == UEdGraphSchema_CustomizableObject::PC_Texture)
+				if (CurrentBreakParameter.ParameterType == UEdGraphSchema_CustomizableObject::PC_Texture)
 				{
 					//TODO(Max - GM): Move this to the generatemutablesourceimage?
 					UE::Mutable::Private::Ptr<UE::Mutable::Private::NodeImageFromMaterialParameter> TextureNode = new UE::Mutable::Private::NodeImageFromMaterialParameter();
-					TextureNode->ImageParameterName = GenerationContext.CurrentMaterialBreakParameter.ParameterName;
+					TextureNode->ImageParameterName = CurrentBreakParameter.ParameterName;
 
 					// Resize image node.
 					UE::Mutable::Private::Ptr<UE::Mutable::Private::NodeImageResize> ResizeNode = new UE::Mutable::Private::NodeImageResize();
@@ -648,7 +666,7 @@ UE::Mutable::Private::Ptr<UE::Mutable::Private::NodeMaterial> GenerateMutableSou
 					MaterialNode->ImageParametersToCompile.Add({ FName(ParameterInfo.Name), ParameterInfo.Index }, FormatNode);
 				}
 
-				else if (GenerationContext.CurrentMaterialBreakParameter.ParameterType == UEdGraphSchema_CustomizableObject::PC_PassthroughTexture)
+				else if (CurrentBreakParameter.ParameterType == UEdGraphSchema_CustomizableObject::PC_PassthroughTexture)
 				{
 					//TODO(Max - GM): Add support for Passthrough texture parameters
 					////Find reference Texture from the material parameter
@@ -664,12 +682,12 @@ UE::Mutable::Private::Ptr<UE::Mutable::Private::NodeMaterial> GenerateMutableSou
 					GenerationContext.Log(Message, Node, EMessageSeverity::Error);
 				}
 
-				else if (GenerationContext.CurrentMaterialBreakParameter.ParameterType == UEdGraphSchema_CustomizableObject::PC_Color)
+				else if (CurrentBreakParameter.ParameterType == UEdGraphSchema_CustomizableObject::PC_Color)
 				{
 					MaterialNode->ColorParametersToCompile.Add({ FName(ParameterInfo.Name), ParameterInfo.Index });
 				}
 
-				else if (GenerationContext.CurrentMaterialBreakParameter.ParameterType == UEdGraphSchema_CustomizableObject::PC_Float)
+				else if (CurrentBreakParameter.ParameterType == UEdGraphSchema_CustomizableObject::PC_Float)
 				{
 					MaterialNode->ScalarParametersToCompile.Add({ FName(ParameterInfo.Name), ParameterInfo.Index });
 				}
@@ -742,10 +760,10 @@ UE::Mutable::Private::Ptr<UE::Mutable::Private::NodeMaterial> GenerateMutableSou
 
 		// Tracking the index of the image. Needed for the FGeneratedImagePropertyKey
 		int32 ImageIndex = 0;
-		
+	
 		for (const UEdGraphPin* InputPin : InputPins)
 		{
-			if (InputPin == TypedNodeMaterialModify->MaterialPinRef)
+			if (InputPin == TypedNodeMaterialModify->MaterialPinRef.Get())
 			{
 				continue;
 			}
@@ -853,11 +871,9 @@ UE::Mutable::Private::Ptr<UE::Mutable::Private::NodeMaterial> GenerateMutableSou
 			bCacheNode = false;
 		}
 
-		bool bSuccess = true;
-		UDataTable* DataTable = GetDataTable(TypedNodeTable, GenerationContext);
-
-		if (DataTable)
+		if (UDataTable* DataTable = GetDataTable(TypedNodeTable, GenerationContext))
 		{
+			bool bSuccess = true;
 			FString ColumnName = TypedNodeTable->GetPinColumnName(Pin);
 			FProperty* Property = TypedNodeTable->FindPinProperty(*Pin);
 
@@ -953,9 +969,8 @@ UE::Mutable::Private::Ptr<UE::Mutable::Private::NodeMaterial> GenerateMutableSou
 		Options.SurfaceNode = MaterialOptions.SurfaceNode;
 
 		MeshExtension->Node = GenerateMutableSourceExternal(Pin, GenerationContext, Options);
-		
+	
 		Result = MeshExtension;
-
 	}
 
 	else if (const UCustomizableObjectNodeMacroInstance* TypedNodeMacro = Cast<UCustomizableObjectNodeMacroInstance>(Node))
@@ -1056,6 +1071,7 @@ FGeneratedImageProperties GenerateImageProperties(const UCustomizableObjectNode*
 	return Props;
 }
 
+
 UE::Mutable::Private::Ptr<UE::Mutable::Private::NodeImage> GenerateMipMapAndFormatNodes(
 	const UCustomizableObjectNode* Node,
 	FMutableGraphGenerationContext& GenerationContext,
@@ -1161,5 +1177,6 @@ UE::Mutable::Private::Ptr<UE::Mutable::Private::NodeImage> GenerateMipMapAndForm
 	return LastImage;
 }
 
+
 #undef LOCTEXT_NAMESPACE
 
diff --git a/Engine/Plugins/Mutable/Source/CustomizableObjectEditor/Private/MuCOE/GenerateMutableSource/GenerateMutableSourceMaterial.h b/Engine/Plugins/Mutable/Source/CustomizableObjectEditor/Private/MuCOE/GenerateMutableSource/GenerateMutableSourceMaterial.h
--- a/Engine/Plugins/Mutable/Source/CustomizableObjectEditor/Private/MuCOE/GenerateMutableSource/GenerateMutableSourceMaterial.h
+++ b/Engine/Plugins/Mutable/Source/CustomizableObjectEditor/Private/MuCOE/GenerateMutableSource/GenerateMutableSourceMaterial.h
@@ -4,10 +4,18 @@
 
 #include "GenerateMutableSource.h"
 #include "MuT/NodeMaterial.h"
+#include "MuT/NodeMaterialParameter.h"
 
 class UEdGraphPin;
 struct FMutableGraphGenerationContext;
 
+/** Generates or retrieves an already generated Material Parameter based on the Pin and Generation Context provided. 
+ * @param Pin Does represent the pin from the node we want to use as source for the generation or retrieval of the Mutable Material Parameter Node.
+ * @param GenerationContext The compilation context of the current compilation operation.
+ * @return A pointer to the generated/retrieved Mutable NodeMaterialParameter
+ */
+UE::Mutable::Private::Ptr<UE::Mutable::Private::NodeMaterialParameter> GenerateMutableSourceMaterialParameter(const UEdGraphPin* Pin, FMutableGraphGenerationContext& GenerationContext);
+
 // Add to this struct all the extra parameters that need to be passed to the GenerateMutableSourceMaterial function.
 struct FGenerationMaterialOptions 
 {
@@ -24,7 +32,7 @@ UE::Mutable::Private::Ptr<UE::Mutable::Private::NodeMaterial> GenerateMutableSou
 /** Generates an FGeneratedImageProperties from a Material Parameter and (if not null) its Reference Texture. */
 FGeneratedImageProperties GenerateImageProperties(const UCustomizableObjectNode* Node, const UTexture2D* ReferenceTexture, const UE::Mutable::Private::FParameterKey& ImageKey, bool bIsPassthrough, FMutableGraphGenerationContext& GenerationContext);
 
-// Generate additional nodes to generete a texture with the correct format and size.
+// Generate additional nodes to generate a texture with the correct format and size.
 UE::Mutable::Private::Ptr<UE::Mutable::Private::NodeImage> GenerateMipMapAndFormatNodes(
 	const UCustomizableObjectNode* Node,
 	FMutableGraphGenerationContext& GenerationContext,
diff --git a/Engine/Plugins/Mutable/Source/MutableRuntime/Internal/MuR/Operations.h b/Engine/Plugins/Mutable/Source/MutableRuntime/Internal/MuR/Operations.h
--- a/Engine/Plugins/Mutable/Source/MutableRuntime/Internal/MuR/Operations.h
+++ b/Engine/Plugins/Mutable/Source/MutableRuntime/Internal/MuR/Operations.h
@@ -440,7 +440,6 @@ namespace UE::Mutable::Private
     {
         typedef ADDRESS_TYPE ADDRESS;
         typedef ADDRESS_TYPE CONSTANT_STRING_INDEX;
-        typedef ADDRESS_TYPE STRING_ADDRESS;
 		typedef ADDRESS_TYPE LIST_UINT32_ADDRESS;
 
         /** Arguments for every operation type. */
@@ -1190,14 +1189,14 @@ namespace UE::Mutable::Private
 		struct MaterialBreakArgs
 		{
 			ADDRESS Material;
-			STRING_ADDRESS ParameterName;
+			CONSTANT_STRING_INDEX ParameterName;
 			int32 LayerIndex;
 		};
 
 		struct MaterialBreakImageParameterArgs
 		{
 			ADDRESS MaterialParameter;
-			STRING_ADDRESS ParameterName;
+			CONSTANT_STRING_INDEX ParameterName;
 		};
 
     	struct ExternalArgs
diff --git a/Engine/Plugins/Mutable/Source/MutableRuntime/Internal/MuR/PassthroughObject.h b/Engine/Plugins/Mutable/Source/MutableRuntime/Internal/MuR/PassthroughObject.h
--- a/Engine/Plugins/Mutable/Source/MutableRuntime/Internal/MuR/PassthroughObject.h
+++ b/Engine/Plugins/Mutable/Source/MutableRuntime/Internal/MuR/PassthroughObject.h
@@ -58,6 +58,12 @@ namespace UE::Mutable::Private
 			return Resource.Get() == Other.Resource.Get();
 		}
 
+		friend uint32 GetTypeHash(const TPassthroughObjectPtr& Object)
+		{
+			uint32 Hash = GetTypeHash(Object.Resource);
+			return Hash;
+		}
+
 		bool IsResolved() const
 		{
 			return Resource->IsType<TStrongObjectPtr<UObject>>();
diff --git a/Engine/Plugins/Mutable/Source/MutableRuntime/Private/MuR/Operations.cpp b/Engine/Plugins/Mutable/Source/MutableRuntime/Private/MuR/Operations.cpp
--- a/Engine/Plugins/Mutable/Source/MutableRuntime/Private/MuR/Operations.cpp
+++ b/Engine/Plugins/Mutable/Source/MutableRuntime/Private/MuR/Operations.cpp
@@ -307,7 +307,6 @@ namespace UE::Mutable::Private
 		{
 			OP::MaterialBreakArgs Args = program.GetOpArgs<OP::MaterialBreakArgs>(at);
 			f(Args.Material);
-			f(Args.ParameterName);
 			break;
 		}
 		
diff --git a/Engine/Plugins/Mutable/Source/MutableTools/Internal/MuT/AST.h b/Engine/Plugins/Mutable/Source/MutableTools/Internal/MuT/AST.h
--- a/Engine/Plugins/Mutable/Source/MutableTools/Internal/MuT/AST.h
+++ b/Engine/Plugins/Mutable/Source/MutableTools/Internal/MuT/AST.h
@@ -890,7 +890,7 @@ namespace UE::Mutable::Private
 			Unknown
 		};
 
-		/** This can be overriden to help detect sugraph mesh properties.  */
+		/** This can be overriden to help detect subgraph mesh properties.  */
 		virtual EClosedMeshTest IsClosedMesh(TMap<const ASTOp*, EClosedMeshTest>* Cache=nullptr ) const;
 
         // Logic expression evaluation
diff --git a/Engine/Plugins/Mutable/Source/MutableTools/Private/MuT/ASTOpConstantResource.h b/Engine/Plugins/Mutable/Source/MutableTools/Private/MuT/ASTOpConstantResource.h
--- a/Engine/Plugins/Mutable/Source/MutableTools/Private/MuT/ASTOpConstantResource.h
+++ b/Engine/Plugins/Mutable/Source/MutableTools/Private/MuT/ASTOpConstantResource.h
@@ -35,8 +35,6 @@ namespace UE::Mutable::Private
 
 		/** Source data descriptor. */
 		FSourceDataDescriptor SourceDataDescriptor;
-
-		PASSTHROUGH_ID PassthroughId = PASSTHROUGH_ID_INVALID;
 		
 		/** Image Parameters are evaluated at runtime. */
 		TMap<FParameterKey, ASTChild> ImageOperations;
diff --git a/Engine/Plugins/Mutable/Source/MutableTools/Private/MuT/CodeOptimiser.cpp b/Engine/Plugins/Mutable/Source/MutableTools/Private/MuT/CodeOptimiser.cpp
--- a/Engine/Plugins/Mutable/Source/MutableTools/Private/MuT/CodeOptimiser.cpp
+++ b/Engine/Plugins/Mutable/Source/MutableTools/Private/MuT/CodeOptimiser.cpp
@@ -37,14 +37,15 @@ namespace UE::Mutable::Private
 
 		bool bModified = false;
 
-		TMap<FResourceEntry<FMesh>, Ptr<ASTOp>> MeshOps;
-		TMap<FResourceEntry<FImage>, Ptr<ASTOp>> ImageOps;
-		TMap<FResourceEntry<FLayout>, Ptr<ASTOp>> LayoutOps;
-		TMap<FResourceEntry<FMaterial>, Ptr<ASTOp>> MaterialOps;
+		TMap<FResourceEntry<FMesh>, Ptr<ASTOpConstantResource>> MeshOps;
+		TMap<FResourceEntry<FImage>, Ptr<ASTOpConstantResource>> ImageOps;
+		TMap<FResourceEntry<FLayout>, Ptr<ASTOpConstantResource>> LayoutOps;
+		TMap<TPassthroughObjectPtr<UMaterialInterface>, Ptr<ASTOpConstantResource>> MaterialOps;
 
 		ASTOp::Traverse_TopRandom_Unique_NonReentrant(Roots, [&](Ptr<ASTOp> Op)
 		{
-			switch (Op->GetOpType())
+			const EOpType OperationType = Op->GetOpType();
+			switch (OperationType)
 			{
 			case EOpType::ME_CONSTANT:
 				{
@@ -70,7 +71,7 @@ namespace UE::Mutable::Private
 			case EOpType::MI_CONSTANT:
 				{
 					const Ptr<ASTOpConstantResource> OpResource = static_pointer_cast<ASTOpConstantResource>(Op);
-					bModified |= Deduplicate<FMaterial>(MaterialOps, OpResource);
+					bModified |= DeduplicateFMaterial(MaterialOps, OpResource);
 					break;
 				}
 			}
diff --git a/Engine/Plugins/Mutable/Source/MutableTools/Private/MuT/CodeOptimiserPrivate.cpp b/Engine/Plugins/Mutable/Source/MutableTools/Private/MuT/CodeOptimiserPrivate.cpp
new file mode 100644
--- /dev/null
+++ b/Engine/Plugins/Mutable/Source/MutableTools/Private/MuT/CodeOptimiserPrivate.cpp
@@ -0,0 +1,40 @@
+﻿// Copyright Epic Games, Inc. All Rights Reserved.
+
+#include "CodeOptimiserPrivate.h"
+
+namespace UE::Mutable::Private
+{
+	bool DeduplicateFMaterial(TMap<TPassthroughObjectPtr<UMaterialInterface>, Ptr<ASTOpConstantResource>>& Ops, const Ptr<ASTOpConstantResource>& OpResource)
+	{
+		TSharedPtr<const FMaterial> CurrentValue = StaticCastSharedPtr<const FMaterial>(OpResource->GetValue());
+		
+		if (Ptr<ASTOpConstantResource>* CachedOpResource = Ops.Find(CurrentValue->PassthroughObject))
+		{
+			TSharedPtr<const FMaterial> CachedOpResourceValue = StaticCastSharedPtr<const FMaterial>(CachedOpResource->get()->GetValue());
+			TSharedPtr<FMaterial> ClonedCachedOpResourceValue = CachedOpResourceValue->Clone();
+			
+			// Append onto the already processed ASTOp the parameters found in the ASTOp being processed
+			ClonedCachedOpResourceValue->ImageParameters.Append(CurrentValue->ImageParameters);
+			ClonedCachedOpResourceValue->ColorParameters.Append(CurrentValue->ColorParameters);
+			ClonedCachedOpResourceValue->ScalarParameters.Append(CurrentValue->ScalarParameters);
+			
+			CachedOpResource->get()->SetValue(ClonedCachedOpResourceValue);
+
+			// Add the Child ASTOps of the OpResource being processed as children of already processed ASTOp.
+			for (TTuple<FParameterKey, ASTChild>& NewImageOperation : OpResource->ImageOperations)
+			{
+				// The children will have as parent the cached ASTOp resource as now they will hang from it
+				(*CachedOpResource)->ImageOperations.Add(NewImageOperation.Key, ASTChild(*CachedOpResource, NewImageOperation.Value.Child));
+			}
+			
+			ASTOp::Replace(OpResource, *CachedOpResource);
+			return true;
+		}
+		else
+		{
+			Ops.Add(CurrentValue->PassthroughObject, OpResource);
+			return false;
+		}
+	}	
+}
+
diff --git a/Engine/Plugins/Mutable/Source/MutableTools/Private/MuT/CodeOptimiserPrivate.h b/Engine/Plugins/Mutable/Source/MutableTools/Private/MuT/CodeOptimiserPrivate.h
--- a/Engine/Plugins/Mutable/Source/MutableTools/Private/MuT/CodeOptimiserPrivate.h
+++ b/Engine/Plugins/Mutable/Source/MutableTools/Private/MuT/CodeOptimiserPrivate.h
@@ -29,14 +29,14 @@ namespace UE::Mutable::Private
 
 
 	template<typename T>
-	bool Deduplicate(TMap<FResourceEntry<T>, Ptr<ASTOp>>& Ops, const Ptr<ASTOpConstantResource>& OpResource)
+	bool Deduplicate(TMap<FResourceEntry<T>, Ptr<ASTOpConstantResource>>& Ops, const Ptr<ASTOpConstantResource>& OpResource)
 	{
 		uint32 Hash = OpResource->GetValueHash();
 		TSharedPtr<const T> Value = StaticCastSharedPtr<const T>(OpResource->GetValue());
 
 		FResourceEntry<T> Entry(Hash, Value);
 		
-		if (Ptr<ASTOp>* Found = Ops.Find(Entry))
+		if (Ptr<ASTOpConstantResource>* Found = Ops.Find(Entry))
 		{
 			ASTOp::Replace(OpResource, *Found);
 			return true;
@@ -46,5 +46,8 @@ namespace UE::Mutable::Private
 			Ops.Add(Entry, OpResource);
 			return false;
 		}
-	}	
+	}
+	
+	
+	bool DeduplicateFMaterial(TMap<TPassthroughObjectPtr<UMaterialInterface>, Ptr<ASTOpConstantResource>>& Ops, const Ptr<ASTOpConstantResource>& OpResource);
 }
