[Mutable] Fixed multiple problems.
- Fixed texture size when using break material nodes.
- Fixed names of dataless parameter nodes.

#rnx

#Synced-CL 48644664

diff --git a/Engine/Plugins/Mutable/Source/CustomizableObjectEditor/Private/MuCOE/GenerateMutableSource/GenerateMutableSourceImage.cpp b/Engine/Plugins/Mutable/Source/CustomizableObjectEditor/Private/MuCOE/GenerateMutableSource/GenerateMutableSourceImage.cpp
--- a/Engine/Plugins/Mutable/Source/CustomizableObjectEditor/Private/MuCOE/GenerateMutableSource/GenerateMutableSourceImage.cpp
+++ b/Engine/Plugins/Mutable/Source/CustomizableObjectEditor/Private/MuCOE/GenerateMutableSource/GenerateMutableSourceImage.cpp
@@ -1238,9 +1238,13 @@ UE::Mutable::Private::Ptr<UE::Mutable::Private::NodeImage> GenerateMutableSource
 			MaterialBreak->ParameterName = ParameterName;
 			MaterialBreak->LayerIndex = LayerIndex;
 
+			// Store the reference texture into the material generation options to generate all the textures with the correct size
+			FGenerationMaterialOptions GenerationOptions;
+			GenerationOptions.ReferenceTextureSize = ReferenceTextureSize;
+
 			// Generate the material that is going to be broken into parameters
 			GenerationContext.MaterialBreakParameterStack.Add(FMaterialBreakParameter{ ParameterName, Pin->PinType.PinCategory, LayerIndex });
-			UE::Mutable::Private::Ptr<UE::Mutable::Private::NodeMaterial> LinkedNode = GenerateMutableSourceMaterial(LinkedMaterial, GenerationContext, FGenerationMaterialOptions());
+			UE::Mutable::Private::Ptr<UE::Mutable::Private::NodeMaterial> LinkedNode = GenerateMutableSourceMaterial(LinkedMaterial, GenerationContext, GenerationOptions);
 			MaterialBreak->MaterialSource = LinkedNode;
 			GenerationContext.MaterialBreakParameterStack.Pop();
 
@@ -1303,7 +1307,7 @@ UE::Mutable::Private::Ptr<UE::Mutable::Private::NodeImageParameter> GenerateMuta
 		TextureNode->Name = TypedNodeParam->GetParameterName(&GenerationContext.MacroNodesStack);
 		TextureNode->UID = GenerationContext.GetNodeIdUnique(ParameterNode).ToString();
 		
-		GenerationContext.TextureParameterDefaultValues.Add(TypedNodeParam->GetParameterName(), TypedNodeParam->DefaultValue);
+		GenerationContext.TextureParameterDefaultValues.Add(TypedNodeParam->GetParameterName(&GenerationContext.MacroNodesStack), TypedNodeParam->DefaultValue);
 
 		GenerationContext.ParameterUIDataMap.Add(TypedNodeParam->GetParameterName(&GenerationContext.MacroNodesStack), FMutableParameterData(
 			TypedNodeParam->ParamUIMetadata,
diff --git a/Engine/Plugins/Mutable/Source/CustomizableObjectEditor/Private/MuCOE/GenerateMutableSource/GenerateMutableSourceMaterial.cpp b/Engine/Plugins/Mutable/Source/CustomizableObjectEditor/Private/MuCOE/GenerateMutableSource/GenerateMutableSourceMaterial.cpp
--- a/Engine/Plugins/Mutable/Source/CustomizableObjectEditor/Private/MuCOE/GenerateMutableSource/GenerateMutableSourceMaterial.cpp
+++ b/Engine/Plugins/Mutable/Source/CustomizableObjectEditor/Private/MuCOE/GenerateMutableSource/GenerateMutableSourceMaterial.cpp
@@ -465,7 +465,7 @@ UE::Mutable::Private::Ptr<UE::Mutable::Private::NodeMaterial> GenerateMutableSou
 							const FString TextureName = GetNameSafe(BaseTexture).ToLower();
 							ConstantImageNode->SourceDataDescriptor.SourceId = CityHash32(reinterpret_cast<const char*>(*TextureName), TextureName.Len() * sizeof(FString::ElementType));
 
-							const uint32 MipsToSkip = ComputeLODBiasForTexture(GenerationContext, *BaseTexture);
+							const uint32 MipsToSkip = ComputeLODBiasForTexture(GenerationContext, *BaseTexture, nullptr, MaterialOptions.ReferenceTextureSize);
 
 							if (UE::Mutable::Private::Ptr<UE::Mutable::Private::NodeImage> FinalNodeImag = ResizeTextureByNumMips(ConstantImageNode, MipsToSkip))
 							{
@@ -653,7 +653,7 @@ UE::Mutable::Private::Ptr<UE::Mutable::Private::NodeMaterial> GenerateMutableSou
 
 					if (bParameterFound && ReferenceTexture)
 					{
-						const uint32 LODBias = ComputeLODBiasForTexture(GenerationContext, *ReferenceTexture);
+						const uint32 LODBias = ComputeLODBiasForTexture(GenerationContext, *ReferenceTexture, nullptr, MaterialOptions.ReferenceTextureSize);
 						TextureSize.X = FMath::Max(ReferenceTexture->Source.GetSizeX() >> LODBias, 1);
 						TextureSize.Y = FMath::Max(ReferenceTexture->Source.GetSizeY() >> LODBias, 1);
 					}
diff --git a/Engine/Plugins/Mutable/Source/CustomizableObjectEditor/Private/MuCOE/GenerateMutableSource/GenerateMutableSourceMaterial.h b/Engine/Plugins/Mutable/Source/CustomizableObjectEditor/Private/MuCOE/GenerateMutableSource/GenerateMutableSourceMaterial.h
--- a/Engine/Plugins/Mutable/Source/CustomizableObjectEditor/Private/MuCOE/GenerateMutableSource/GenerateMutableSourceMaterial.h
+++ b/Engine/Plugins/Mutable/Source/CustomizableObjectEditor/Private/MuCOE/GenerateMutableSource/GenerateMutableSourceMaterial.h
@@ -21,6 +21,9 @@ struct FGenerationMaterialOptions
 {
 	// Node surface. Only needed for MeshSection nodes.
 	UE::Mutable::Private::Ptr<UE::Mutable::Private::NodeSurfaceNew> SurfaceNode = nullptr;
+
+	// Size of the reference Texture Size
+	int32 ReferenceTextureSize = 0;
 };
 
 /** Convert a CustomizableObject Source Graph into a mutable source graph. */
diff --git a/Engine/Plugins/Mutable/Source/CustomizableObjectEditor/Private/MuCOE/GenerateMutableSource/GenerateMutableSourceSkeletalMesh.cpp b/Engine/Plugins/Mutable/Source/CustomizableObjectEditor/Private/MuCOE/GenerateMutableSource/GenerateMutableSourceSkeletalMesh.cpp
--- a/Engine/Plugins/Mutable/Source/CustomizableObjectEditor/Private/MuCOE/GenerateMutableSource/GenerateMutableSourceSkeletalMesh.cpp
+++ b/Engine/Plugins/Mutable/Source/CustomizableObjectEditor/Private/MuCOE/GenerateMutableSource/GenerateMutableSourceSkeletalMesh.cpp
@@ -40,7 +40,7 @@ namespace UE::Mutable::Private
 			SkeletalMeshNode->Name = TypedNodeParam->GetParameterName(&GenerationContext.MacroNodesStack);
 			SkeletalMeshNode->UID = GenerationContext.GetNodeIdUnique(ParameterNode).ToString();
 		
-			GenerationContext.SkeletalMeshParameterDefaultValues.Add(TypedNodeParam->GetParameterName(), LoadObject(TypedNodeParam->DefaultValue));
+			GenerationContext.SkeletalMeshParameterDefaultValues.Add(TypedNodeParam->GetParameterName(&GenerationContext.MacroNodesStack), LoadObject(TypedNodeParam->DefaultValue));
 
 			GenerationContext.ParameterUIDataMap.Add(TypedNodeParam->GetParameterName(&GenerationContext.MacroNodesStack), FMutableParameterData(
 				TypedNodeParam->ParamUIMetadata,
