[Mutable] Memory managment rework.

- Added own smart pointer TManagedPtr with persistent weak references. Weak references invalidation is controlled by an external system.
- Simplified cache eviction using the new smart pointer.

#rnx
#Synced-CL 48713935

diff --git a/Engine/Plugins/Mutable/Source/CustomizableObject/Internal/MuCO/CustomizableObjectInstancePrivate.h b/Engine/Plugins/Mutable/Source/CustomizableObject/Internal/MuCO/CustomizableObjectInstancePrivate.h
--- a/Engine/Plugins/Mutable/Source/CustomizableObject/Internal/MuCO/CustomizableObjectInstancePrivate.h
+++ b/Engine/Plugins/Mutable/Source/CustomizableObject/Internal/MuCO/CustomizableObjectInstancePrivate.h
@@ -49,7 +49,7 @@ CUSTOMIZABLEOBJECT_API extern const FString POSE_PARAMETER_POSTFIX;
 
 /** \param OnlyLOD: If not 0, extract and convert only one single LOD from the source image.
   * \param ExtractChannel: If different than -1, extract a single-channel image with the specified source channel data. */
-CUSTOMIZABLEOBJECT_API void ConvertImage(UTexture2D* Texture, TSharedPtr<const UE::Mutable::Private::FImage> MutableImage, const FMutableModelImageProperties& Props, int32 OnlyLOD = -1, int32 ExtractChannel = -1);
+CUSTOMIZABLEOBJECT_API void ConvertImage(UTexture2D* Texture, UE::Mutable::Private::TManagedPtr<const UE::Mutable::Private::FImage> MutableImage, const FMutableModelImageProperties& Props, int32 OnlyLOD = -1, int32 ExtractChannel = -1);
 
 
 /** CustomizableObject Instance flags for internal use  */
@@ -426,7 +426,7 @@ UE_API void BuildOrCopyClothingData(const TSharedRef<FUpdateContextPrivate>& Con
 
 UE_API bool BuildSkeletonData(const TSharedRef<FUpdateContextPrivate>& Context, USkeletalMesh& SkeletalMesh, FCustomizableObjectInstanceComponentIndex InstanceComponentIndex);
 
-UE_API void BuildMeshSockets(USkeletalMesh* SkeletalMesh, const UModelResources& ModelResources, const FMutableRefSkeletalMeshData& RefSkeletalMeshData, TSharedPtr<const UE::Mutable::Private::FMesh> MutableMesh);
+UE_API void BuildMeshSockets(USkeletalMesh* SkeletalMesh, const UModelResources& ModelResources, const FMutableRefSkeletalMeshData& RefSkeletalMeshData, UE::Mutable::Private::TManagedPtr<const UE::Mutable::Private::FMesh> MutableMesh);
 
 UE_API bool BuildOrCopyRenderData(const TSharedRef<FUpdateContextPrivate>& Context, USkeletalMesh* SkeletalMesh, UCustomizableObjectInstance* CustomizableObjectInstance, FCustomizableObjectInstanceComponentIndex InstanceComponentIndex);
 
diff --git a/Engine/Plugins/Mutable/Source/CustomizableObject/Internal/MuCO/CustomizableObjectMipDataProvider.h b/Engine/Plugins/Mutable/Source/CustomizableObject/Internal/MuCO/CustomizableObjectMipDataProvider.h
--- a/Engine/Plugins/Mutable/Source/CustomizableObject/Internal/MuCO/CustomizableObjectMipDataProvider.h
+++ b/Engine/Plugins/Mutable/Source/CustomizableObject/Internal/MuCO/CustomizableObjectMipDataProvider.h
@@ -94,7 +94,7 @@ struct FMutableImageOperationData
 
 	TSharedPtr<FMutableUpdateImageContext> UpdateContext;
 
-	TSharedPtr<const UE::Mutable::Private::FImage> Result;
+	UE::Mutable::Private::TManagedPtr<const UE::Mutable::Private::FImage> Result;
 
 	TArray<FMutableMipUpdateLevel> Levels;
 
diff --git a/Engine/Plugins/Mutable/Source/CustomizableObject/Internal/MuCO/CustomizableObjectSystemPrivate.h b/Engine/Plugins/Mutable/Source/CustomizableObject/Internal/MuCO/CustomizableObjectSystemPrivate.h
--- a/Engine/Plugins/Mutable/Source/CustomizableObject/Internal/MuCO/CustomizableObjectSystemPrivate.h
+++ b/Engine/Plugins/Mutable/Source/CustomizableObject/Internal/MuCO/CustomizableObjectSystemPrivate.h
@@ -107,16 +107,16 @@ void ReconstructMorphTargets(const UE::Mutable::Private::FMesh& Mesh, const TArr
  *  @param MutableDataStreamer Streamer to request the blocks.
  *  @param Mesh Mutable mesh to obtain which blocks to load.
  *  @param StreamingResult Key is the block. Value is the container where the data will be deserialized once streamed. */
-void LoadMorphTargetsData(FMutableStreamRequest& MutableDataStreamer, const TSharedRef<const UE::Mutable::Private::FMesh>& Mesh, TMap<uint32, FMorphTargetMeshData>& StreamingResult);
+void LoadMorphTargetsData(FMutableStreamRequest& MutableDataStreamer, const UE::Mutable::Private::TManagedRef<const UE::Mutable::Private::FMesh>& Mesh, TMap<uint32, FMorphTargetMeshData>& StreamingResult);
 
-void LoadMorphTargetsMetadata(const FMutableStreamRequest& MutableDataStreamer, const TSharedRef<const UE::Mutable::Private::FMesh>& Mesh, TMap<uint32, FMorphTargetMeshData>& StreamingResult);
+void LoadMorphTargetsMetadata(const FMutableStreamRequest& MutableDataStreamer, const UE::Mutable::Private::TManagedRef<const UE::Mutable::Private::FMesh>& Mesh, TMap<uint32, FMorphTargetMeshData>& StreamingResult);
 
 
 /** Request the Mutable Data Streamer to load Cloth blocks. Does not stream them.
  *  @param MutableDataStreamer Streamer to request the blocks.
  *  @param Mesh Mutable mesh to obtain which blocks to load.
  *  @param StreamingResult Key is the block. Value is the container where the data will be deserialized once streamed. */
-void LoadClothing(FMutableStreamRequest& MutableDataStreamer, const TSharedRef<const UE::Mutable::Private::FMesh>& Mesh, TMap<uint32, FClothingMeshData>& StreamingResult);
+void LoadClothing(FMutableStreamRequest& MutableDataStreamer, const UE::Mutable::Private::TManagedRef<const UE::Mutable::Private::FMesh>& Mesh, TMap<uint32, FClothingMeshData>& StreamingResult);
 
 
 /** End a Customizable Object Instance Update. All code paths of an update have to end here. */
@@ -333,7 +333,7 @@ struct FInstanceUpdateData
 		
 		UE::Mutable::Private::EImageFormat Format = UE::Mutable::Private::EImageFormat::None;
 
-		TSharedPtr<const UE::Mutable::Private::FImage> Image;
+		UE::Mutable::Private::TManagedPtr<const UE::Mutable::Private::FImage> Image;
 
 		TArray<int32> ConstantImagesNeededToGenerate;
 
@@ -376,7 +376,7 @@ struct FInstanceUpdateData
 		uint16 ScalarCount = 0;
 				
 		UE::Mutable::Private::FMaterialId MaterialId;
-		TSharedPtr<const UE::Mutable::Private::FMaterial> Material;
+		UE::Mutable::Private::TManagedPtr<const UE::Mutable::Private::FMaterial> Material;
 	};
 	
 	struct FSurface
@@ -393,7 +393,7 @@ struct FInstanceUpdateData
 	struct FLOD
 	{		
 		UE::Mutable::Private::FMeshId MeshId;
-		TSharedPtr<const UE::Mutable::Private::FMesh> Mesh;
+		UE::Mutable::Private::TManagedPtr<const UE::Mutable::Private::FMesh> Mesh;
 
 		/** Range in the Surfaces array */
 		uint16 FirstSurface = 0;
@@ -429,8 +429,7 @@ struct FInstanceUpdateData
 		uint32 OverlayMaterialIndex = INDEX_NONE;
 		
 		UE::Mutable::Private::FSkeletalMeshId SkeletalMeshId;
-		
-		TSharedPtr<const UE::Mutable::Private::FSkeletalMesh> SkeletalMesh;
+		UE::Mutable::Private::TManagedPtr<const UE::Mutable::Private::FSkeletalMesh> SkeletalMesh;
 	};
 
 	TArray<TSharedRef<FComponent>> Components;
@@ -478,7 +477,7 @@ struct FInstanceUpdateData
 
 	struct FNamedExtensionData
 	{
-		TSharedPtr<const UE::Mutable::Private::FExtensionData> Data;
+		UE::Mutable::Private::TManagedPtr<const UE::Mutable::Private::FExtensionData> Data;
 		FName Name;
 	};
 	TArray<FNamedExtensionData> ExtendedInputPins;
@@ -611,7 +610,7 @@ public:
 	int32 MipsToSkip = 0;
 
 	UE::Mutable::Private::FInstance::FID InstanceID = 0; // Redundant
-	TSharedPtr<const UE::Mutable::Private::FInstance> MutableInstance;
+	UE::Mutable::Private::TManagedPtr<const UE::Mutable::Private::FInstance> MutableInstance;
 
 	TSharedPtr<UE::Mutable::Private::FModel> Model;
 	TSharedPtr<UE::Mutable::Private::FMeshIdRegistry> MeshIdRegistry;
@@ -980,7 +979,7 @@ private:
 
 
 /** Set OnlyLOD to -1 to generate all mips */
-CUSTOMIZABLEOBJECT_API TUniquePtr<FTexturePlatformData> MutableCreateImagePlatformData(TSharedPtr<const UE::Mutable::Private::FImage> MutableImage, int32 OnlyLOD, uint16 FullSizeX, uint16 FullSizeY);
+CUSTOMIZABLEOBJECT_API TUniquePtr<FTexturePlatformData> MutableCreateImagePlatformData(UE::Mutable::Private::TManagedPtr<const UE::Mutable::Private::FImage> MutableImage, int32 OnlyLOD, uint16 FullSizeX, uint16 FullSizeY);
 
 
 /** Return true if Streaming is enabled for the given Object. */
diff --git a/Engine/Plugins/Mutable/Source/CustomizableObject/Internal/MuCO/UnrealConversionUtils.h b/Engine/Plugins/Mutable/Source/CustomizableObject/Internal/MuCO/UnrealConversionUtils.h
--- a/Engine/Plugins/Mutable/Source/CustomizableObject/Internal/MuCO/UnrealConversionUtils.h
+++ b/Engine/Plugins/Mutable/Source/CustomizableObject/Internal/MuCO/UnrealConversionUtils.h
@@ -8,11 +8,13 @@
 #define MUTABLE_VERTEXBUFFER_TEXCOORDS	2
 
 #include "MuR/Ptr.h"
+#include "MuR/ManagedPointer.h"
 #include "Containers/Array.h"
 #include "Containers/ArrayView.h"
 #include "Containers/Map.h"
 #include "Tasks/Task.h"
 
+
 class FBoneNames;
 struct FClothBufferIndexMapping;
 struct FMeshToMeshVertData;
@@ -156,7 +158,7 @@ namespace UnrealConversionUtils
 
 	void MorphTargetVertexInfoBuffers(FSkeletalMeshLODRenderData& LODResource, const USkeletalMesh& Owner, const UE::Mutable::Private::FMesh& MutableMesh, const TMap<uint32, FMorphTargetMeshData>& MorphTargetMeshData, int32 LODIndex, bool bGenerateCPUMorphTargets);
 
-	CUSTOMIZABLEOBJECT_API UE::Tasks::FTask ConvertSkeletalMeshFromRuntimeData(USkeletalMesh* SkeletalMesh, int32 LODIndex, int32 SectionIndex, TSharedPtr<UE::Mutable::Private::FMesh> Result, const TSharedRef<FBoneNames>& BoneNames);
+	CUSTOMIZABLEOBJECT_API UE::Tasks::FTask ConvertSkeletalMeshFromRuntimeData(USkeletalMesh* SkeletalMesh, int32 LODIndex, int32 SectionIndex, UE::Mutable::Private::TManagedPtr<UE::Mutable::Private::FMesh> Result, const TSharedRef<FBoneNames>& BoneNames);
 
 	void GetSectionClothData(const UE::Mutable::Private::FMesh& MutableMesh, int32 LODIndex, const TMap<uint32, FClothingMeshData>& ClothingMeshData, TArray<FSectionClothData>& SectionsClothData, int32& NumClothingDataNotFound);
 	
diff --git a/Engine/Plugins/Mutable/Source/CustomizableObject/Internal/MuCO/UnrealMutableImageProvider.h b/Engine/Plugins/Mutable/Source/CustomizableObject/Internal/MuCO/UnrealMutableImageProvider.h
--- a/Engine/Plugins/Mutable/Source/CustomizableObject/Internal/MuCO/UnrealMutableImageProvider.h
+++ b/Engine/Plugins/Mutable/Source/CustomizableObject/Internal/MuCO/UnrealMutableImageProvider.h
@@ -19,10 +19,10 @@ public:
 	
 	// UE::Mutable::Private::ExternalResourceProvider interface
 	// Thread: worker
-	CUSTOMIZABLEOBJECT_API virtual TTuple<UE::Tasks::FTask, TFunction<void()>> GetImageAsync(UTexture* Texture, uint8 MipmapsToSkip, TFunction<void(TSharedPtr<UE::Mutable::Private::FImage>)>& ResultCallback) override;
-	CUSTOMIZABLEOBJECT_API virtual TTuple<UE::Tasks::FTask, TFunction<void()>> GetReferencedImageAsync( int32 Id, uint8 MipmapsToSkip, TFunction<void(TSharedPtr<UE::Mutable::Private::FImage>)>& ResultCallback) override;
+	CUSTOMIZABLEOBJECT_API virtual TTuple<UE::Tasks::FTask, TFunction<void()>> GetImageAsync(UTexture* Texture, uint8 MipmapsToSkip, TFunction<void(UE::Mutable::Private::TManagedPtr<UE::Mutable::Private::FImage>)>& ResultCallback) override;
+	CUSTOMIZABLEOBJECT_API virtual TTuple<UE::Tasks::FTask, TFunction<void()>> GetReferencedImageAsync( int32 Id, uint8 MipmapsToSkip, TFunction<void(UE::Mutable::Private::TManagedPtr<UE::Mutable::Private::FImage>)>& ResultCallback) override;
 	CUSTOMIZABLEOBJECT_API virtual UE::Mutable::Private::FExtendedImageDesc GetImageDesc(UTexture* Texture) override;
-	CUSTOMIZABLEOBJECT_API virtual TTuple<UE::Tasks::FTask, TFunction<void()>> GetMeshAsync(USkeletalMesh* SkeletalMesh, int32 LODIndex, int32 SectionIndex, TFunction<void(TSharedPtr<UE::Mutable::Private::FMesh>)>& ResultCallback) override;
+	CUSTOMIZABLEOBJECT_API virtual TTuple<UE::Tasks::FTask, TFunction<void()>> GetMeshAsync(USkeletalMesh* SkeletalMesh, int32 LODIndex, int32 SectionIndex, TFunction<void(UE::Mutable::Private::TManagedPtr<UE::Mutable::Private::FMesh>)>& ResultCallback) override;
 
 #if WITH_EDITOR
 	void CacheRuntimeReferencedImages(const TArray<TSoftObjectPtr<UTexture2D>>& RuntimeReferencedTextures);
@@ -33,7 +33,7 @@ private:
 			UE::Mutable::Private::FImageDesc {UE::Mutable::Private::FImageSize(32, 32), UE::Mutable::Private::EImageFormat::RGBA_UByte, UE::Mutable::Private::EImageFormat::RGBA_UByte, 1};
 
 	/** This will be called if an image Id has been requested by Mutable core but it has not been provided by any provider. */
-	static TSharedPtr<UE::Mutable::Private::FImage> CreateDummy();
+	static UE::Mutable::Private::TManagedPtr<UE::Mutable::Private::FImage> CreateDummy();
 	static UE::Mutable::Private::FExtendedImageDesc CreateDummyDesc();
 
 #if WITH_EDITOR
diff --git a/Engine/Plugins/Mutable/Source/CustomizableObject/Internal/MuCO/UnrealToMutableTextureConversionUtils.h b/Engine/Plugins/Mutable/Source/CustomizableObject/Internal/MuCO/UnrealToMutableTextureConversionUtils.h
--- a/Engine/Plugins/Mutable/Source/CustomizableObject/Internal/MuCO/UnrealToMutableTextureConversionUtils.h
+++ b/Engine/Plugins/Mutable/Source/CustomizableObject/Internal/MuCO/UnrealToMutableTextureConversionUtils.h
@@ -4,7 +4,7 @@
 
 #include "HAL/Platform.h"
 #include "Engine/Texture.h"
-#include "MuR/Ptr.h"
+#include "MuR/ManagedPointer.h"
 
 #define UE_API CUSTOMIZABLEOBJECT_API
 
@@ -169,7 +169,7 @@ public:
 	// Keep strong references to meshes and other UObjects involved in the conversion
 	TArray<TStrongObjectPtr<UObject>> ReferencedObjects;
 
-	TSharedPtr<UE::Mutable::Private::FMesh> Result;
+	UE::Mutable::Private::TManagedPtr<UE::Mutable::Private::FMesh> Result;
 };
 
 #endif
diff --git a/Engine/Plugins/Mutable/Source/CustomizableObject/Private/MuCO/CustomizableObject.cpp b/Engine/Plugins/Mutable/Source/CustomizableObject/Private/MuCO/CustomizableObject.cpp
--- a/Engine/Plugins/Mutable/Source/CustomizableObject/Private/MuCO/CustomizableObject.cpp
+++ b/Engine/Plugins/Mutable/Source/CustomizableObject/Private/MuCO/CustomizableObject.cpp
@@ -2252,7 +2252,7 @@ void UCustomizableObjectPrivate::SetModel(const TSharedPtr<UE::Mutable::Private:
 				break;
 				
 			case UE::Mutable::Private::EParameterType::InstancedStruct:
-				Parameter.DefaultValue.Set<UE::Mutable::Private::FParamInstancedStructType>(MakeShared<FInstancedStruct>(LocalModelResources->ExternalTypeParameterDefaultValues[Parameter.Name]));
+				Parameter.DefaultValue.Set<UE::Mutable::Private::FParamInstancedStructType>(UE::Mutable::Private::MakeManaged<FInstancedStruct>(LocalModelResources->ExternalTypeParameterDefaultValues[Parameter.Name]));
 				break;
 				
 			case UE::Mutable::Private::EParameterType::Matrix:
@@ -3366,7 +3366,7 @@ void UCustomizableObjectPrivate::LogMemory()
 		int32 ConstantImageLODStreamed = Program.ConstantImageLODsStreamed.GetAllocatedSize();
 
 		int32 PermanentImages = 0;
-		for (const TSharedPtr<const UE::Mutable::Private::FImage>& Entry : Program.ConstantImageLODsPermanent)
+		for (const UE::Mutable::Private::TManagedPtr<const UE::Mutable::Private::FImage>& Entry : Program.ConstantImageLODsPermanent)
 		{
 			if (Entry)
 			{
@@ -3381,7 +3381,7 @@ void UCustomizableObjectPrivate::LogMemory()
 		int32 ConstantMeshesStreamed = Program.ConstantMeshesStreamed.GetAllocatedSize();
 
 		int32 PermanentMeshes = 0;
-		for (const TSharedPtr<const UE::Mutable::Private::FMesh>& Entry : Program.ConstantMeshesPermanent)
+		for (const UE::Mutable::Private::TManagedPtr<const UE::Mutable::Private::FMesh>& Entry : Program.ConstantMeshesPermanent)
 		{
 			if (Entry)
 			{
@@ -3397,13 +3397,13 @@ void UCustomizableObjectPrivate::LogMemory()
 		}
 
 		int32 SkeletonsTotal = Program.ConstantSkeletons.GetAllocatedSize();
-		for (const TSharedPtr<const UE::Mutable::Private::FSkeleton>& Entry : Program.ConstantSkeletons)
+		for (const UE::Mutable::Private::TManagedPtr<const UE::Mutable::Private::FSkeleton>& Entry : Program.ConstantSkeletons)
 		{
 			SkeletonsTotal += sizeof(UE::Mutable::Private::FSkeleton) + Entry->DebugBoneNames.GetAllocatedSize() + Entry->BoneIds.GetAllocatedSize() + Entry->BoneParents.GetAllocatedSize();
 		}
 
 		int32 PhysicsTotal = Program.ConstantPhysicsBodies.GetAllocatedSize();
-		for (const TSharedPtr<const UE::Mutable::Private::FPhysicsBody>& Entry : Program.ConstantPhysicsBodies)
+		for (const UE::Mutable::Private::TManagedPtr<const UE::Mutable::Private::FPhysicsBody>& Entry : Program.ConstantPhysicsBodies)
 		{
 			PhysicsTotal += sizeof(UE::Mutable::Private::FPhysicsBody) + Entry->BoneIds.GetAllocatedSize() + Entry->Bodies.GetAllocatedSize() + Entry->BodiesCustomIds.GetAllocatedSize();
 			for (const UE::Mutable::Private::FPhysicsBodyAggregate& Body : Entry->Bodies)
diff --git a/Engine/Plugins/Mutable/Source/CustomizableObject/Private/MuCO/CustomizableObjectInstance.cpp b/Engine/Plugins/Mutable/Source/CustomizableObject/Private/MuCO/CustomizableObjectInstance.cpp
--- a/Engine/Plugins/Mutable/Source/CustomizableObject/Private/MuCO/CustomizableObjectInstance.cpp
+++ b/Engine/Plugins/Mutable/Source/CustomizableObject/Private/MuCO/CustomizableObjectInstance.cpp
@@ -2187,7 +2187,7 @@ bool UCustomizableInstancePrivate::UpdateSkeletalMesh_PostBeginUpdate0(UCustomiz
 		}
 
 		// We need the first valid mesh. get it from the component, considering that some LOSs may have been skipped.
-		TSharedPtr<const UE::Mutable::Private::FMesh> ComponentMesh;
+		UE::Mutable::Private::TManagedPtr<const UE::Mutable::Private::FMesh> ComponentMesh;
 		for (int32 LODIndex = 0; LODIndex < Component->LODCount && !ComponentMesh; ++LODIndex)
 		{
 			ComponentMesh = Context->InstanceUpdateData.LODs[Component->FirstLOD + LODIndex]->Mesh;
@@ -2271,7 +2271,7 @@ bool UCustomizableInstancePrivate::UpdateSkeletalMesh_PostBeginUpdate0(UCustomiz
 		}
 
 		// Build PhysicsAsset merging physics assets coming from SubMeshes of the newly generated Mesh
-		if (TSharedPtr<const UE::Mutable::Private::FPhysicsBody> MutablePhysics = ComponentMesh->GetPhysicsBody())
+		if (UE::Mutable::Private::TManagedPtr<const UE::Mutable::Private::FPhysicsBody> MutablePhysics = ComponentMesh->GetPhysicsBody())
 		{
 			constexpr bool bDisallowCollisionBetweenAssets = true;
 			UPhysicsAsset* PhysicsAssetResult = GetOrBuildMainPhysicsAsset(
@@ -2291,7 +2291,7 @@ bool UCustomizableInstancePrivate::UpdateSkeletalMesh_PostBeginUpdate0(UCustomiz
 		const int32 NumAdditionalPhysicsNum = ComponentMesh->AdditionalPhysicsBodies.Num();
 		for (int32 I = 0; I < NumAdditionalPhysicsNum; ++I)
 		{
-			TSharedPtr<const UE::Mutable::Private::FPhysicsBody> AdditionalPhysicsBody = ComponentMesh->AdditionalPhysicsBodies[I];
+			UE::Mutable::Private::TManagedPtr<const UE::Mutable::Private::FPhysicsBody> AdditionalPhysicsBody = ComponentMesh->AdditionalPhysicsBodies[I];
 			
 			check(AdditionalPhysicsBody);
 			if (!AdditionalPhysicsBody->bBodiesModified)
@@ -3246,7 +3246,7 @@ void FMutableUpdateCandidate::ApplyLODUpdateParamsToInstance(FUpdateContextPriva
 
 // The memory allocated in the function and pointed by the returned pointer is owned by the caller and must be freed. 
 // If assigned to a UTexture2D, it will be freed by that UTexture2D
-TUniquePtr<FTexturePlatformData> MutableCreateImagePlatformData(TSharedPtr<const UE::Mutable::Private::FImage> MutableImage, int32 OnlyLOD, uint16 FullSizeX, uint16 FullSizeY)
+TUniquePtr<FTexturePlatformData> MutableCreateImagePlatformData(UE::Mutable::Private::TManagedPtr<const UE::Mutable::Private::FImage> MutableImage, int32 OnlyLOD, uint16 FullSizeX, uint16 FullSizeY)
 {
 	int32 SizeX = FMath::Max(MutableImage->GetSize()[0], FullSizeX);
 	int32 SizeY = FMath::Max(MutableImage->GetSize()[1], FullSizeY);
@@ -3519,7 +3519,7 @@ TUniquePtr<FTexturePlatformData> MutableCreateImagePlatformData(TSharedPtr<const
 }
 
 
-void ConvertImage(UTexture2D* Texture, TSharedPtr<const UE::Mutable::Private::FImage> MutableImage, const FMutableModelImageProperties& Props, int OnlyLOD, int32 ExtractChannel)
+void ConvertImage(UTexture2D* Texture, UE::Mutable::Private::TManagedPtr<const UE::Mutable::Private::FImage> MutableImage, const FMutableModelImageProperties& Props, int OnlyLOD, int32 ExtractChannel)
 {
 	MUTABLE_CPUPROFILER_SCOPE(UCustomizableInstancePrivate::ConvertImage);
 
@@ -3545,7 +3545,8 @@ void ConvertImage(UTexture2D* Texture, TSharedPtr<const UE::Mutable::Private::FI
 		UE_LOG(LogMutable, Warning, TEXT("Building instance: a texture was generated in RGB format, which is slow to convert to Unreal."));
 
 		// Expand the image.
-		TSharedPtr<UE::Mutable::Private::FImage> Converted = MakeShared<UE::Mutable::Private::FImage>(MutableImage->GetSizeX(), MutableImage->GetSizeY(), MutableImage->GetLODCount(), UE::Mutable::Private::EImageFormat::RGBA_UByte, UE::Mutable::Private::EInitializationType::NotInitialized);
+		UE::Mutable::Private::TManagedPtr<UE::Mutable::Private::FImage> Converted = 
+				UE::Mutable::Private::MakeManaged<UE::Mutable::Private::FImage>(MutableImage->GetSizeX(), MutableImage->GetSizeY(), MutableImage->GetLODCount(), UE::Mutable::Private::EImageFormat::RGBA_UByte, UE::Mutable::Private::EInitializationType::NotInitialized);
 
 		for (int32 LODIndex = 0; LODIndex < Converted->GetLODCount(); ++LODIndex)
 		{
@@ -3570,7 +3571,8 @@ void ConvertImage(UTexture2D* Texture, TSharedPtr<const UE::Mutable::Private::FI
 		MUTABLE_CPUPROFILER_SCOPE(Swizzle);
 		// Swizzle the image.
 		// \TODO: Raise a warning?
-		TSharedPtr<UE::Mutable::Private::FImage> Converted = MakeShared<UE::Mutable::Private::FImage>(MutableImage->GetSizeX(), MutableImage->GetSizeY(), 1, UE::Mutable::Private::EImageFormat::RGBA_UByte, UE::Mutable::Private::EInitializationType::NotInitialized);
+		UE::Mutable::Private::TManagedPtr<UE::Mutable::Private::FImage> Converted = 
+				UE::Mutable::Private::MakeManaged<UE::Mutable::Private::FImage>(MutableImage->GetSizeX(), MutableImage->GetSizeY(), 1, UE::Mutable::Private::EImageFormat::RGBA_UByte, UE::Mutable::Private::EInitializationType::NotInitialized);
 		int32 PixelCount = MutableImage->GetSizeX() * MutableImage->GetSizeY();
 
 		const uint8* pSource = MutableImage->GetLODData(0);
@@ -3913,7 +3915,7 @@ bool BuildSkeletonData(const TSharedRef<FUpdateContextPrivate>& Context, USkelet
 }
 
 
-void BuildMeshSockets(USkeletalMesh* SkeletalMesh, const UModelResources& ModelResources, const FMutableRefSkeletalMeshData& RefSkeletalMeshData, TSharedPtr<const UE::Mutable::Private::FMesh> MutableMesh)
+void BuildMeshSockets(USkeletalMesh* SkeletalMesh, const UModelResources& ModelResources, const FMutableRefSkeletalMeshData& RefSkeletalMeshData, UE::Mutable::Private::TManagedPtr<const UE::Mutable::Private::FMesh> MutableMesh)
 {
 	// Build mesh sockets.
 	MUTABLE_CPUPROFILER_SCOPE(UCustomizableInstancePrivate::BuildMeshSockets);
@@ -4209,7 +4211,7 @@ void BuildOrCopyClothingData(
 		{
 			const TSharedRef<FInstanceUpdateData::FLOD>& LOD = Context->InstanceUpdateData.LODs[Component->FirstLOD+LODIndex];
 
-			if (TSharedPtr<const UE::Mutable::Private::FMesh> MutableMesh = LOD->Mesh)
+			if (UE::Mutable::Private::TManagedPtr<const UE::Mutable::Private::FMesh> MutableMesh = LOD->Mesh)
 			{
 				for (int32 SectionIndex = 0; SectionIndex < MutableMesh->GetSurfaceCount(); ++SectionIndex)
 				{
@@ -4318,7 +4320,7 @@ void BuildOrCopyClothingData(
 		{
 			const TSharedRef<FInstanceUpdateData::FLOD>& LOD = Context->InstanceUpdateData.LODs[Component->FirstLOD+LODIndex];
 
-			if (TSharedPtr<const UE::Mutable::Private::FMesh> MutableMesh = LOD->Mesh)
+			if (UE::Mutable::Private::TManagedPtr<const UE::Mutable::Private::FMesh> MutableMesh = LOD->Mesh)
 			{
 				GetSectionClothData(*MutableMesh, LODIndex, Context->InstanceUpdateData.ClothingMeshData, SectionsClothData, NumClothingDataNotFound);
 			}
@@ -4498,7 +4500,7 @@ void BuildOrCopyClothingData(
 	{
 		const TSharedRef<FInstanceUpdateData::FLOD>& LOD = Context->InstanceUpdateData.LODs[Component->FirstLOD + LODIndex];
 		
-		TSharedPtr<const UE::Mutable::Private::FMesh> MutableMesh = LOD->Mesh;
+		UE::Mutable::Private::TManagedPtr<const UE::Mutable::Private::FMesh> MutableMesh = LOD->Mesh;
 		if (!MutableMesh)
 		{
 			continue;
@@ -4964,13 +4966,13 @@ UE::Tasks::FTask LoadAdditionalAssetsAndData(const TSharedRef<FUpdateContextPriv
 					{
 						const TSharedRef<FInstanceUpdateData::FLOD>& LOD = Context->InstanceUpdateData.LODs[Component->FirstLOD + LODIndex];
 
-						TSharedPtr<const UE::Mutable::Private::FMesh> MutableMesh = LOD->Mesh;
+						UE::Mutable::Private::TManagedPtr<const UE::Mutable::Private::FMesh> MutableMesh = LOD->Mesh;
 						if (!MutableMesh)
 						{
 							continue;
 						}
 
-						LoadMorphTargetsData(StreamRequest, MutableMesh.ToSharedRef(), Context->InstanceUpdateData.RealTimeMorphTargetMeshData);
+						LoadMorphTargetsData(StreamRequest, MutableMesh.ToStrongRef(), Context->InstanceUpdateData.RealTimeMorphTargetMeshData);
 					}
 
 					// Metadata
@@ -4982,13 +4984,13 @@ UE::Tasks::FTask LoadAdditionalAssetsAndData(const TSharedRef<FUpdateContextPriv
 					{
 						const TSharedRef<FInstanceUpdateData::FLOD>& LOD = Context->InstanceUpdateData.LODs[Component->FirstLOD + LODIndex];
 
-						TSharedPtr<const UE::Mutable::Private::FMesh> MutableMesh = LOD->Mesh;
+						UE::Mutable::Private::TManagedPtr<const UE::Mutable::Private::FMesh> MutableMesh = LOD->Mesh;
 						if (!MutableMesh)
 						{
 							continue;
 						}
 
-						LoadMorphTargetsMetadata(StreamRequest, MutableMesh.ToSharedRef(), Context->InstanceUpdateData.RealTimeMorphTargetMeshData);
+						LoadMorphTargetsMetadata(StreamRequest, MutableMesh.ToStrongRef(), Context->InstanceUpdateData.RealTimeMorphTargetMeshData);
 					}
 				}
 
@@ -5000,7 +5002,7 @@ UE::Tasks::FTask LoadAdditionalAssetsAndData(const TSharedRef<FUpdateContextPriv
 					{
 						const TSharedRef<FInstanceUpdateData::FLOD>& LOD = Context->InstanceUpdateData.LODs[Component->FirstLOD + LODIndex];
 
-						TSharedPtr<const UE::Mutable::Private::FMesh> MutableMesh = LOD->Mesh;
+						UE::Mutable::Private::TManagedPtr<const UE::Mutable::Private::FMesh> MutableMesh = LOD->Mesh;
 						if (!MutableMesh)
 						{
 							continue;
@@ -5008,7 +5010,7 @@ UE::Tasks::FTask LoadAdditionalAssetsAndData(const TSharedRef<FUpdateContextPriv
 
 						if (bUpdateMeshes)
 						{
-							LoadClothing(StreamRequest, MutableMesh.ToSharedRef(), Context->InstanceUpdateData.ClothingMeshData);
+							LoadClothing(StreamRequest, MutableMesh.ToStrongRef(), Context->InstanceUpdateData.ClothingMeshData);
 						}
 					}
 
@@ -5017,7 +5019,7 @@ UE::Tasks::FTask LoadAdditionalAssetsAndData(const TSharedRef<FUpdateContextPriv
 					{
 						const TSharedRef<FInstanceUpdateData::FLOD>& LOD = Context->InstanceUpdateData.LODs[Component->FirstLOD + LODIndex];
 
-						TSharedPtr<const UE::Mutable::Private::FMesh> MutableMesh = LOD->Mesh;
+						UE::Mutable::Private::TManagedPtr<const UE::Mutable::Private::FMesh> MutableMesh = LOD->Mesh;
 						if (!MutableMesh)
 						{
 							continue;
@@ -5078,7 +5080,7 @@ UE::Tasks::FTask LoadAdditionalAssetsAndData(const TSharedRef<FUpdateContextPriv
 			{
 				const TSharedRef<FInstanceUpdateData::FLOD>& LOD = Context->InstanceUpdateData.LODs[Component->FirstLOD + LODIndex];
 
-				TSharedPtr<const UE::Mutable::Private::FMesh> MutableMesh = LOD->Mesh;
+				UE::Mutable::Private::TManagedPtr<const UE::Mutable::Private::FMesh> MutableMesh = LOD->Mesh;
 				if (!MutableMesh)
 				{
 					continue;
@@ -5233,7 +5235,7 @@ UE::Tasks::FTask LoadAdditionalAssetsAndData(const TSharedRef<FUpdateContextPriv
 	// Copy FExtensionData Object node input from the Instance to the InstanceUpdateData
 	for (int32 ExtensionDataIndex = 0; ExtensionDataIndex < Context->MutableInstance->GetExtensionDataCount(); ExtensionDataIndex++)
 	{
-		TSharedPtr<const UE::Mutable::Private::FExtensionData> ExtensionData;
+		UE::Mutable::Private::TManagedPtr<const UE::Mutable::Private::FExtensionData> ExtensionData;
 		FName Name;
 		Context->MutableInstance->GetExtensionData(ExtensionDataIndex, ExtensionData, Name);
 
@@ -5598,7 +5600,7 @@ UTexture* BuildTexture(const TSharedRef<FUpdateContextPrivate>& Context, const i
 	check((bNeverStream && Context->MipsToSkip == 0) ||
 		(!bNeverStream && Context->MipsToSkip >= 0));
 	
-	TSharedPtr<const UE::Mutable::Private::FImage> MutableImage = Image->Image;
+	UE::Mutable::Private::TManagedPtr<const UE::Mutable::Private::FImage> MutableImage = Image->Image;
 
 	UTexture2D* Texture; // Texture generated by mutable
 
diff --git a/Engine/Plugins/Mutable/Source/CustomizableObject/Private/MuCO/CustomizableObjectInstanceDescriptor.cpp b/Engine/Plugins/Mutable/Source/CustomizableObject/Private/MuCO/CustomizableObjectInstanceDescriptor.cpp
--- a/Engine/Plugins/Mutable/Source/CustomizableObject/Private/MuCO/CustomizableObjectInstanceDescriptor.cpp
+++ b/Engine/Plugins/Mutable/Source/CustomizableObject/Private/MuCO/CustomizableObjectInstanceDescriptor.cpp
@@ -878,12 +878,12 @@ TSharedPtr<UE::Mutable::Private::FParameters> FCustomizableObjectInstanceDescrip
 						for (int32 RangeIndex = 0; RangeIndex < ExternalTypeParameter.ParameterRangeValues.Num(); ++RangeIndex)
 						{
 							RangeIdxPtr->SetPosition(0, RangeIndex);
-							MutableParameters->SetInstancedStructValue(ParamIndex, MakeShared<const FInstancedStruct>(ExternalTypeParameter.ParameterRangeValues[RangeIndex]), RangeIdxPtr.Get());
+							MutableParameters->SetInstancedStructValue(ParamIndex, UE::Mutable::Private::MakeManaged<const FInstancedStruct>(ExternalTypeParameter.ParameterRangeValues[RangeIndex]), RangeIdxPtr.Get());
 						}
 					}
 					else
 					{
-						MutableParameters->SetInstancedStructValue(ParamIndex, MakeShared<const FInstancedStruct>(ExternalTypeParameter.ParameterValue));
+						MutableParameters->SetInstancedStructValue(ParamIndex, UE::Mutable::Private::MakeManaged<const FInstancedStruct>(ExternalTypeParameter.ParameterValue));
 					}
 
 					break;
@@ -3092,4 +3092,4 @@ void FCustomizableObjectInstanceDescriptor::UpdateLayer(const FName& ProjectorPa
 }
  
  
-#undef RETURN_ON_CO_UNCOMPILED
\ No newline at end of file
+#undef RETURN_ON_CO_UNCOMPILED
diff --git a/Engine/Plugins/Mutable/Source/CustomizableObject/Private/MuCO/CustomizableObjectMeshUpdate.cpp b/Engine/Plugins/Mutable/Source/CustomizableObject/Private/MuCO/CustomizableObjectMeshUpdate.cpp
--- a/Engine/Plugins/Mutable/Source/CustomizableObject/Private/MuCO/CustomizableObjectMeshUpdate.cpp
+++ b/Engine/Plugins/Mutable/Source/CustomizableObject/Private/MuCO/CustomizableObjectMeshUpdate.cpp
@@ -213,7 +213,7 @@ namespace impl
 		const UE::Mutable::Private::FMeshId& MeshID = OperationData->MeshIDs[LODIndex];
 
 		constexpr UE::Mutable::Private::EMeshContentFlags MeshContentFilter = UE::Mutable::Private::EMeshContentFlags::AllFlags;
-		UE::Tasks::TTask<TSharedPtr<const UE::Mutable::Private::FMesh>> GetMeshTask = OperationData->System->GetMesh({}, InstanceID, MeshID, MeshContentFilter);
+		UE::Tasks::TTask<UE::Mutable::Private::TManagedPtr<const UE::Mutable::Private::FMesh>> GetMeshTask = OperationData->System->GetMesh({}, InstanceID, MeshID, MeshContentFilter);
 
 		UE::Tasks::AddNested(UE::Tasks::Launch(TEXT("Task_MutableGetMeshes_GetMesh_Post"), [=]() mutable
 			{
@@ -257,7 +257,7 @@ namespace impl
 		const uint32 LODMask = 0xFFFFFFFF;
 
 		// Main instance generation step
-		TSharedPtr<const UE::Mutable::Private::FInstance> Instance = System->BeginUpdate_MutableThread(InstanceID,
+		UE::Mutable::Private::TManagedPtr<const UE::Mutable::Private::FInstance> Instance = System->BeginUpdate_MutableThread(InstanceID,
 			OperationData->Parameters,
 			OperationData->MeshIdRegistry,
 			OperationData->ImageIdRegistry,
@@ -323,17 +323,17 @@ void FCustomizableObjectMeshStreamIn::RequestMeshUpdate(const FContext& Context)
 			FMutableStreamRequest StreamRequest(RefThis->OperationData->ModelStreamableBulkData);
 
 			// Get morphs to be streamed.
-			for (TSharedPtr<const UE::Mutable::Private::FMesh>& Mesh : RefThis->OperationData->Meshes)
+			for (UE::Mutable::Private::TManagedPtr<const UE::Mutable::Private::FMesh>& Mesh : RefThis->OperationData->Meshes)
 			{
 				if (!Mesh)
 				{
 					continue;
 				}
 
-				LoadMorphTargetsData(StreamRequest, Mesh.ToSharedRef(), RefThis->OperationData->MorphTargetMeshData);
-				LoadMorphTargetsMetadata(StreamRequest, Mesh.ToSharedRef(), RefThis->OperationData->MorphTargetMeshData);
+				LoadMorphTargetsData(StreamRequest, Mesh.ToStrongRef(), RefThis->OperationData->MorphTargetMeshData);
+				LoadMorphTargetsMetadata(StreamRequest, Mesh.ToStrongRef(), RefThis->OperationData->MorphTargetMeshData);
 				
-				LoadClothing(StreamRequest, Mesh.ToSharedRef(), RefThis->OperationData->ClothingMeshData);
+				LoadClothing(StreamRequest, Mesh.ToStrongRef(), RefThis->OperationData->ClothingMeshData);
 			}
 
 			// Stream data.
@@ -406,7 +406,7 @@ void FCustomizableObjectMeshStreamIn::ConvertMesh(const FContext& Context, bool&
 	{
 		const int32 LODIndexAbsolute = LODIndex + Context.AssetLODBias;
 
-		TSharedPtr<const UE::Mutable::Private::FMesh> MutableMesh = OperationData->Meshes[LODIndexAbsolute];
+		UE::Mutable::Private::TManagedPtr<const UE::Mutable::Private::FMesh> MutableMesh = OperationData->Meshes[LODIndexAbsolute];
 
 		if (!MutableMesh)
 		{
diff --git a/Engine/Plugins/Mutable/Source/CustomizableObject/Private/MuCO/CustomizableObjectMeshUpdate.h b/Engine/Plugins/Mutable/Source/CustomizableObject/Private/MuCO/CustomizableObjectMeshUpdate.h
--- a/Engine/Plugins/Mutable/Source/CustomizableObject/Private/MuCO/CustomizableObjectMeshUpdate.h
+++ b/Engine/Plugins/Mutable/Source/CustomizableObject/Private/MuCO/CustomizableObjectMeshUpdate.h
@@ -41,7 +41,7 @@ struct FMutableUpdateMeshContext
 	TArray<UE::Mutable::Private::FMeshId> MeshIDs;
 
 	// Meshes generated per LOD
-	TArray<TSharedPtr<const UE::Mutable::Private::FMesh>> Meshes;
+	TArray<UE::Mutable::Private::TManagedPtr<const UE::Mutable::Private::FMesh>> Meshes;
 
 	/** Serialized data per block. */
 	TMap<uint32, FMorphTargetMeshData> MorphTargetMeshData;
diff --git a/Engine/Plugins/Mutable/Source/CustomizableObject/Private/MuCO/CustomizableObjectMipDataProvider.cpp b/Engine/Plugins/Mutable/Source/CustomizableObject/Private/MuCO/CustomizableObjectMipDataProvider.cpp
--- a/Engine/Plugins/Mutable/Source/CustomizableObject/Private/MuCO/CustomizableObjectMipDataProvider.cpp
+++ b/Engine/Plugins/Mutable/Source/CustomizableObject/Private/MuCO/CustomizableObjectMipDataProvider.cpp
@@ -134,7 +134,7 @@ namespace Impl
 		UE::Mutable::Private::FInstance::FID InstanceID = System->NewInstance(Model, OperationData->UpdateContext->ExternalResourceProvider, nullptr);
 		UE_LOG(LogMutable, Verbose, TEXT("Creating Mutable instance with id [%d] for a single UpdateImage"), InstanceID)
 
-		TSharedPtr<const UE::Mutable::Private::FInstance> Instance;
+		UE::Mutable::Private::TManagedPtr<const UE::Mutable::Private::FInstance> Instance;
 
 		// Main instance generation step
 		{
@@ -155,14 +155,14 @@ namespace Impl
 
 		const FMutableImageReference& ImageRef = OperationData->RequestedImage;
 		
-		UE::Tasks::TTask<TSharedPtr<const UE::Mutable::Private::FImage>> GetImageTask = System->GetImage({}, InstanceID, ImageRef.ImageID, ImageRef.BaseMip + OperationData->MipsToSkip);
+		UE::Tasks::TTask<UE::Mutable::Private::TManagedPtr<const UE::Mutable::Private::FImage>> GetImageTask = System->GetImage({}, InstanceID, ImageRef.ImageID, ImageRef.BaseMip + OperationData->MipsToSkip);
 
 		UE::Tasks::AddNested(UE::Tasks::Launch(TEXT("MipDataProvider_EndUpdateImagesTask"),
 				[System, OperationData, InstanceID, StartTime, GetImageTask, EndUpdateImage]() mutable
 				{
 					check(GetImageTask.IsCompleted());
 
-					TSharedPtr<const UE::Mutable::Private::FImage> ResultImage = GetImageTask.GetResult();
+					UE::Mutable::Private::TManagedPtr<const UE::Mutable::Private::FImage> ResultImage = GetImageTask.GetResult();
 
 					check(ResultImage);
 
@@ -179,7 +179,8 @@ namespace Impl
 						UE_LOG(LogMutable, Warning, TEXT("Mutable generated an incomplete mip chain for image."));
 
 						// Force the right number of mips. The missing data will be black.
-						TSharedPtr<UE::Mutable::Private::FImage> NewImage = MakeShared<UE::Mutable::Private::FImage>(ResultImage->GetSizeX(), ResultImage->GetSizeY(), FullMipCount, ResultImage->GetFormat(), UE::Mutable::Private::EInitializationType::Black);
+						UE::Mutable::Private::TManagedPtr<UE::Mutable::Private::FImage> NewImage = 
+								UE::Mutable::Private::MakeManaged<UE::Mutable::Private::FImage>(ResultImage->GetSizeX(), ResultImage->GetSizeY(), FullMipCount, ResultImage->GetFormat(), UE::Mutable::Private::EInitializationType::Black);
 
 						// Formats with BytesPerBlock == 0 will not allocate memory. This type of images are not expected here.
 						check(!NewImage->DataStorage.IsEmpty());
@@ -476,7 +477,7 @@ bool FMutableTextureMipDataProvider::PollMips(const FTextureUpdateSyncOptions& S
 		// The counter must be zero meaning the Mutable image operation has finished
 		check(SyncOptions.Counter->GetValue() == 0);
 
-		TSharedPtr<const UE::Mutable::Private::FImage> Image = OperationData->Result;
+		UE::Mutable::Private::TManagedPtr<const UE::Mutable::Private::FImage> Image = OperationData->Result;
 		
 		int32 ImageLODCount = 0;
 		if (Image)
diff --git a/Engine/Plugins/Mutable/Source/CustomizableObject/Private/MuCO/CustomizableObjectSystem.cpp b/Engine/Plugins/Mutable/Source/CustomizableObject/Private/MuCO/CustomizableObjectSystem.cpp
--- a/Engine/Plugins/Mutable/Source/CustomizableObject/Private/MuCO/CustomizableObjectSystem.cpp
+++ b/Engine/Plugins/Mutable/Source/CustomizableObject/Private/MuCO/CustomizableObjectSystem.cpp
@@ -2041,7 +2041,7 @@ namespace impl
 			{
 				TSharedRef<FInstanceUpdateData::FLOD>& LOD = OperationData->InstanceUpdateData.LODs[Component->FirstLOD + LODIndex];
 
-				TSharedPtr<const UE::Mutable::Private::FMesh> Mesh = LOD->Mesh;
+				UE::Mutable::Private::TManagedPtr<const UE::Mutable::Private::FMesh> Mesh = LOD->Mesh;
 				if (!Mesh || Mesh->IsReference())
 				{
 					continue;
@@ -2332,7 +2332,7 @@ namespace impl
 		{
 			TSharedRef<FInstanceUpdateData::FMaterial> Material = Context->InstanceUpdateData.Materials[MaterialIndex];
 			
-			UE::Tasks::TTask<TSharedPtr<const UE::Mutable::Private::FMaterial>> GetMaterialTask = Context->MutableSystem->GetMaterial(LastTask, Context->InstanceID, Material->MaterialId);
+			UE::Tasks::TTask<UE::Mutable::Private::TManagedPtr<const UE::Mutable::Private::FMaterial>> GetMaterialTask = Context->MutableSystem->GetMaterial(LastTask, Context->InstanceID, Material->MaterialId);
 			UE::Tasks::AddNested(GetMaterialTask);
 			LastTask = GetMaterialTask;
 			
@@ -2411,7 +2411,7 @@ namespace impl
 			OperationData->MutableSystem->SetImagePixelConversionOverride( OperationData->PixelFormatOverride );
 		}
 		
-		TSharedPtr<const UE::Mutable::Private::FInstance> Instance = OperationData->MutableSystem->BeginUpdate_MutableThread(OperationData->InstanceID,
+		UE::Mutable::Private::TManagedPtr<const UE::Mutable::Private::FInstance> Instance = OperationData->MutableSystem->BeginUpdate_MutableThread(OperationData->InstanceID,
 			OperationData->Parameters,
 			OperationData->MeshIdRegistry,
 			OperationData->ImageIdRegistry,
@@ -2446,7 +2446,7 @@ namespace impl
 
 			Component->SkeletalMeshId = Instance->GetSkeletalMeshId(InstanceComponentIndex);
 
-			UE::Tasks::TTask<TSharedPtr<const UE::Mutable::Private::FSkeletalMesh>> GetSkeletalMeshTask = OperationData->MutableSystem->GetSkeletalMesh(LastTask, OperationData->InstanceID, Component->SkeletalMeshId, UE::Mutable::Private::EMeshContentFlags::AllFlags);
+			UE::Tasks::TTask<UE::Mutable::Private::TManagedPtr<const UE::Mutable::Private::FSkeletalMesh>> GetSkeletalMeshTask = OperationData->MutableSystem->GetSkeletalMesh(LastTask, OperationData->InstanceID, Component->SkeletalMeshId, UE::Mutable::Private::EMeshContentFlags::AllFlags);
 			UE::Tasks::AddNested(GetSkeletalMeshTask);
 			LastTask = GetSkeletalMeshTask;
 			
@@ -2518,7 +2518,7 @@ namespace impl
 					}
 				}
 				
-				UE::Tasks::TTask<TSharedPtr<const UE::Mutable::Private::FMesh>> GetMeshTask = OperationData->MutableSystem->GetMesh(LastTask, OperationData->InstanceID, LOD->MeshId, MeshContentFilter);
+				UE::Tasks::TTask<UE::Mutable::Private::TManagedPtr<const UE::Mutable::Private::FMesh>> GetMeshTask = OperationData->MutableSystem->GetMesh(LastTask, OperationData->InstanceID, LOD->MeshId, MeshContentFilter);
 				UE::Tasks::AddNested(GetMeshTask);
 				LastTask = GetMeshTask;
 				
@@ -2604,7 +2604,7 @@ namespace impl
 
 			if (Image->bIsPassThrough)
 			{
-				UE::Tasks::TTask<TSharedPtr<const UE::Mutable::Private::FImage>> GetImageTask = OperationData->MutableSystem->GetImage(LastTask, OperationData->InstanceID, Image->ImageID, 0);
+				UE::Tasks::TTask<UE::Mutable::Private::TManagedPtr<const UE::Mutable::Private::FImage>> GetImageTask = OperationData->MutableSystem->GetImage(LastTask, OperationData->InstanceID, Image->ImageID, 0);
 				UE::Tasks::AddNested(GetImageTask);
 				LastTask = GetImageTask;
 					
@@ -2644,7 +2644,7 @@ namespace impl
 				const int32 MipSizeX = FMath::Max(Image->FullImageSizeX >> MipsToSkip, 1);
 				const int32 MipSizeY = FMath::Max(Image->FullImageSizeY >> MipsToSkip, 1);
 
-				UE::Tasks::TTask<TSharedPtr<const UE::Mutable::Private::FImage>> GetImageTask = OperationData->MutableSystem->GetImage(LastTask, OperationData->InstanceID, Image->ImageID, Image->BaseMip + MipsToSkip);
+				UE::Tasks::TTask<UE::Mutable::Private::TManagedPtr<const UE::Mutable::Private::FImage>> GetImageTask = OperationData->MutableSystem->GetImage(LastTask, OperationData->InstanceID, Image->ImageID, Image->BaseMip + MipsToSkip);
 				UE::Tasks::AddNested(GetImageTask);
 				LastTask = GetImageTask;
 
@@ -2659,7 +2659,7 @@ namespace impl
 					{
 						// Generate a correctly-sized but empty image instead, to avoid crashes.
 						UE_LOG(LogMutable, Warning, TEXT("Mutable generated a wrongly-sized image %s."), *Image->ImageID.ToString());
-						Image->Image = MakeShared<UE::Mutable::Private::FImage>(MipSizeX, MipSizeY, FullLODCount - MipsToSkip, Image->Image->GetFormat(), UE::Mutable::Private::EInitializationType::Black);
+						Image->Image = UE::Mutable::Private::MakeManaged<UE::Mutable::Private::FImage>(MipSizeX, MipSizeY, FullLODCount - MipsToSkip, Image->Image->GetFormat(), UE::Mutable::Private::EInitializationType::Black);
 					}
 
 					// We need one mip or the complete chain. Otherwise there was a bug.
@@ -2677,7 +2677,7 @@ namespace impl
 						UE_LOG(LogMutable, Warning, TEXT("Mutable generated an incomplete mip chain for image %s."), *Image->ImageID.ToString());
 
 						// Force the right number of mips. The missing data will be black.
-						const TSharedPtr<UE::Mutable::Private::FImage> NewImage = MakeShared<UE::Mutable::Private::FImage>(Image->Image->GetSizeX(), Image->Image->GetSizeY(), FullMipCount, Image->Image->GetFormat(), UE::Mutable::Private::EInitializationType::Black);
+						const UE::Mutable::Private::TManagedPtr<UE::Mutable::Private::FImage> NewImage = UE::Mutable::Private::MakeManaged<UE::Mutable::Private::FImage>(Image->Image->GetSizeX(), Image->Image->GetSizeY(), FullMipCount, Image->Image->GetFormat(), UE::Mutable::Private::EInitializationType::Black);
 						check(NewImage);	
 						// Formats with BytesPerBlock == 0 will not allocate memory. This type of images are not expected here.
 						check(!NewImage->DataStorage.IsEmpty());
@@ -4416,7 +4416,7 @@ int32 UCustomizableObjectSystemPrivate::BlockTillAllRequestsFinished(float TimeL
 }
 
 
-void LoadMorphTargetsData(FMutableStreamRequest& MutableDataStreamer, const TSharedRef<const UE::Mutable::Private::FMesh>& Mesh, TMap<uint32, FMorphTargetMeshData>& StreamingResult)
+void LoadMorphTargetsData(FMutableStreamRequest& MutableDataStreamer, const UE::Mutable::Private::TManagedRef<const UE::Mutable::Private::FMesh>& Mesh, TMap<uint32, FMorphTargetMeshData>& StreamingResult)
 {
 	MUTABLE_CPUPROFILER_SCOPE(LoadMorphTargetsData);
 
@@ -4477,7 +4477,7 @@ void LoadMorphTargetsData(FMutableStreamRequest& MutableDataStreamer, const TSha
 }
 
 
-void LoadMorphTargetsMetadata(const FMutableStreamRequest& MutableDataStreamer, const TSharedRef<const UE::Mutable::Private::FMesh>& Mesh, TMap<uint32, FMorphTargetMeshData>& StreamingResult)
+void LoadMorphTargetsMetadata(const FMutableStreamRequest& MutableDataStreamer, const UE::Mutable::Private::TManagedRef<const UE::Mutable::Private::FMesh>& Mesh, TMap<uint32, FMorphTargetMeshData>& StreamingResult)
 {
 	MUTABLE_CPUPROFILER_SCOPE(LoadMorphTargetsMetadata);
 
@@ -4512,7 +4512,7 @@ void LoadMorphTargetsMetadata(const FMutableStreamRequest& MutableDataStreamer,
 }
 
 
-void LoadClothing(FMutableStreamRequest& MutableDataStreamer, const TSharedRef<const UE::Mutable::Private::FMesh>& Mesh, TMap<uint32, FClothingMeshData>& StreamingResult)
+void LoadClothing(FMutableStreamRequest& MutableDataStreamer, const UE::Mutable::Private::TManagedRef<const UE::Mutable::Private::FMesh>& Mesh, TMap<uint32, FClothingMeshData>& StreamingResult)
 {
 	MUTABLE_CPUPROFILER_SCOPE(LoadClothing);
 
diff --git a/Engine/Plugins/Mutable/Source/CustomizableObject/Private/MuCO/UnrealConversionUtils.cpp b/Engine/Plugins/Mutable/Source/CustomizableObject/Private/MuCO/UnrealConversionUtils.cpp
--- a/Engine/Plugins/Mutable/Source/CustomizableObject/Private/MuCO/UnrealConversionUtils.cpp
+++ b/Engine/Plugins/Mutable/Source/CustomizableObject/Private/MuCO/UnrealConversionUtils.cpp
@@ -1027,13 +1027,13 @@ namespace UnrealConversionUtils
 	}
 
 	
-	UE::Tasks::FTask ConvertSkeletalMeshFromRuntimeData(USkeletalMesh* SkeletalMesh, int32 LODIndex, int32 SectionIndex, TSharedPtr<UE::Mutable::Private::FMesh> Result, const TSharedRef<FBoneNames>& BoneNames)
+	UE::Tasks::FTask ConvertSkeletalMeshFromRuntimeData(USkeletalMesh* SkeletalMesh, int32 LODIndex, int32 SectionIndex, UE::Mutable::Private::TManagedPtr<UE::Mutable::Private::FMesh> Result, const TSharedRef<FBoneNames>& BoneNames)
 	{
 		MUTABLE_CPUPROFILER_SCOPE(ConvertSkeletalMeshFromRuntimeData);
 
 		if (!SkeletalMesh)
 		{
-			return UE::Tasks::MakeCompletedTask<TSharedPtr<UE::Mutable::Private::FMesh>>();
+			return UE::Tasks::MakeCompletedTask<UE::Mutable::Private::TManagedPtr<UE::Mutable::Private::FMesh>>();
 		}
 		
 		struct FMeshConversionContext
@@ -1198,7 +1198,8 @@ namespace UnrealConversionUtils
 					}
 
 					// Create the skeleton, poses, and BoneMap for this mesh
-					TSharedPtr<UE::Mutable::Private::FSkeleton> MutableSkeleton = MakeShared<UE::Mutable::Private::FSkeleton>();
+					UE::Mutable::Private::TManagedPtr<UE::Mutable::Private::FSkeleton> MutableSkeleton = 
+							UE::Mutable::Private::MakeManaged<UE::Mutable::Private::FSkeleton>();
 					Result->SetSkeleton(MutableSkeleton);
 
 					const int32 NumRequiredBones = RequiredBones.Num();
diff --git a/Engine/Plugins/Mutable/Source/CustomizableObject/Private/MuCO/UnrealMutableImageProvider.cpp b/Engine/Plugins/Mutable/Source/CustomizableObject/Private/MuCO/UnrealMutableImageProvider.cpp
--- a/Engine/Plugins/Mutable/Source/CustomizableObject/Private/MuCO/UnrealMutableImageProvider.cpp
+++ b/Engine/Plugins/Mutable/Source/CustomizableObject/Private/MuCO/UnrealMutableImageProvider.cpp
@@ -10,6 +10,7 @@
 #include "MuCO/LoadUtils.h"
 #include "MuCO/MutableMeshBufferUtils.h"
 #include "MuR/Parameters.h"
+#include "MuR/ManagedPointer.h"
 #include "MuR/ImageTypes.h"
 #include "MuR/Model.h"
 #include "MuR/Mesh.h"
@@ -99,7 +100,7 @@ FUnrealMutableResourceProvider::FUnrealMutableResourceProvider(const TSharedRef<
 
 
 //-------------------------------------------------------------------------------------------------
-TTuple<UE::Tasks::FTask, TFunction<void()>> FUnrealMutableResourceProvider::GetImageAsync(UTexture* Texture, uint8 MipmapsToSkip, TFunction<void(TSharedPtr<UE::Mutable::Private::FImage>)>& ResultCallback)
+TTuple<UE::Tasks::FTask, TFunction<void()>> FUnrealMutableResourceProvider::GetImageAsync(UTexture* Texture, uint8 MipmapsToSkip, TFunction<void(UE::Mutable::Private::TManagedPtr<UE::Mutable::Private::FImage>)>& ResultCallback)
 {
 	// Thread: worker
 	MUTABLE_CPUPROFILER_SCOPE(FUnrealMutableImageProvider::GetImageAsync);
@@ -141,7 +142,7 @@ TTuple<UE::Tasks::FTask, TFunction<void()>> FUnrealMutableResourceProvider::GetI
 	check(MipIndex >= 0);
 		
 	// In the editor the src data can be directly accessed
-	TSharedPtr<UE::Mutable::Private::FImage> Image = MakeShared<UE::Mutable::Private::FImage>();
+	UE::Mutable::Private::TManagedPtr<UE::Mutable::Private::FImage> Image = UE::Mutable::Private::MakeManaged<UE::Mutable::Private::FImage>();
 
 	FMutableSourceTextureData Tex(*Texture2D);
 	EUnrealToMutableConversionError Error = ConvertTextureUnrealSourceToMutable(Image.Get(), Tex, MipIndex);
@@ -188,7 +189,7 @@ TTuple<UE::Tasks::FTask, TFunction<void()>> FUnrealMutableResourceProvider::GetI
 	int32 SizeY = Texture2D->GetSizeY() >> MipIndex;
 
 	check(LODs == 1);
-	TSharedPtr<UE::Mutable::Private::FImage> Image = MakeShared<UE::Mutable::Private::FImage>(SizeX, SizeY, LODs, MutImageFormat, UE::Mutable::Private::EInitializationType::NotInitialized);
+	UE::Mutable::Private::TManagedPtr<UE::Mutable::Private::FImage> Image = MakeManaged<UE::Mutable::Private::FImage>(SizeX, SizeY, LODs, MutImageFormat, UE::Mutable::Private::EInitializationType::NotInitialized);
 	TArrayView<uint8> MutImageDataView = Image->DataStorage.GetLOD(0);
 
 	// In a packaged game the bulk data has to be loaded
@@ -331,7 +332,7 @@ TTuple<UE::Tasks::FTask, TFunction<void()>> FUnrealMutableResourceProvider::GetI
 
 
 //-------------------------------------------------------------------------------------------------
-TTuple<UE::Tasks::FTask, TFunction<void()>> FUnrealMutableResourceProvider::GetReferencedImageAsync(int32 Id, uint8 MipmapsToSkip, TFunction<void(TSharedPtr<UE::Mutable::Private::FImage>)>& ResultCallback)
+TTuple<UE::Tasks::FTask, TFunction<void()>> FUnrealMutableResourceProvider::GetReferencedImageAsync(int32 Id, uint8 MipmapsToSkip, TFunction<void(UE::Mutable::Private::TManagedPtr<UE::Mutable::Private::FImage>)>& ResultCallback)
 {
 	// Thread: worker
 	MUTABLE_CPUPROFILER_SCOPE(FUnrealMutableImageProvider::GetReferencedImageAsync);
@@ -372,7 +373,7 @@ TTuple<UE::Tasks::FTask, TFunction<void()>> FUnrealMutableResourceProvider::GetR
 	const int32 MipIndex = FMath::Min(static_cast<int32>(MipmapsToSkip), Texture2D->Source.CopyTornOff().GetNumMips() - 1);
 	check(MipIndex >= 0);
 
-	TSharedPtr<UE::Mutable::Private::FImage> Image = MakeShared<UE::Mutable::Private::FImage>();
+	UE::Mutable::Private::TManagedPtr<UE::Mutable::Private::FImage> Image = UE::Mutable::Private::MakeManaged<UE::Mutable::Private::FImage>();
 	
 	FMutableSourceTextureData Tex(*Texture2D);
 	EUnrealToMutableConversionError Error = ConvertTextureUnrealSourceToMutable(Image.Get(), Tex, MipIndex);
@@ -462,7 +463,7 @@ UE::Mutable::Private::FExtendedImageDesc FUnrealMutableResourceProvider::GetImag
 }
 
 
-TTuple<UE::Tasks::FTask, TFunction<void()>> FUnrealMutableResourceProvider::GetMeshAsync(USkeletalMesh* SkeletalMesh, int32 LODIndex, int32 SectionIndex, TFunction<void(TSharedPtr<UE::Mutable::Private::FMesh>)>& ResultCallback)
+TTuple<UE::Tasks::FTask, TFunction<void()>> FUnrealMutableResourceProvider::GetMeshAsync(USkeletalMesh* SkeletalMesh, int32 LODIndex, int32 SectionIndex, TFunction<void(UE::Mutable::Private::TManagedPtr<UE::Mutable::Private::FMesh>)>& ResultCallback)
 {
 	// Thread: worker
 	MUTABLE_CPUPROFILER_SCOPE(FUnrealMutableImageProvider::GetMeshAsync);
@@ -472,7 +473,7 @@ TTuple<UE::Tasks::FTask, TFunction<void()>> FUnrealMutableResourceProvider::GetM
 		return MakeTuple(UE::Tasks::MakeCompletedTask<void>(), []() -> void {});
 	};
 
-	TSharedPtr<UE::Mutable::Private::FMesh> Result = MakeShared<UE::Mutable::Private::FMesh>();
+	UE::Mutable::Private::TManagedPtr<UE::Mutable::Private::FMesh> Result = UE::Mutable::Private::MakeManaged<UE::Mutable::Private::FMesh>();
 
 	if (!SkeletalMesh)
 	{
@@ -530,7 +531,7 @@ void FUnrealMutableResourceProvider::CacheRuntimeReferencedImages(const TArray<T
 #endif
 
 
-TSharedPtr<UE::Mutable::Private::FImage> FUnrealMutableResourceProvider::CreateDummy()
+UE::Mutable::Private::TManagedPtr<UE::Mutable::Private::FImage> FUnrealMutableResourceProvider::CreateDummy()
 {
 	// Create a dummy image
 	const int32 Size = DUMMY_IMAGE_DESC.m_size[0];
@@ -543,7 +544,8 @@ TSharedPtr<UE::Mutable::Private::FImage> FUnrealMutableResourceProvider::CreateD
 	uint8 Colors[CheckerTileCount][4] = {{255, 255, 0, 0}, {0, 0, 255, 0}};
 #endif
 
-	TSharedPtr<UE::Mutable::Private::FImage> pResult = MakeShared<UE::Mutable::Private::FImage>(Size, Size, DUMMY_IMAGE_DESC.m_lods, DUMMY_IMAGE_DESC.m_format, UE::Mutable::Private::EInitializationType::NotInitialized);
+	UE::Mutable::Private::TManagedPtr<UE::Mutable::Private::FImage> pResult = 
+			UE::Mutable::Private::MakeManaged<UE::Mutable::Private::FImage>(Size, Size, DUMMY_IMAGE_DESC.m_lods, DUMMY_IMAGE_DESC.m_format, UE::Mutable::Private::EInitializationType::NotInitialized);
 
 	check(pResult->GetLODCount() == 1);
 	check(pResult->GetFormat() == UE::Mutable::Private::EImageFormat::RGBA_UByte || pResult->GetFormat() == UE::Mutable::Private::EImageFormat::BGRA_UByte);
diff --git a/Engine/Plugins/Mutable/Source/CustomizableObjectEditor/Private/MuCOE/CustomizableObjectCompileRunnable.cpp b/Engine/Plugins/Mutable/Source/CustomizableObjectEditor/Private/MuCOE/CustomizableObjectCompileRunnable.cpp
--- a/Engine/Plugins/Mutable/Source/CustomizableObjectEditor/Private/MuCOE/CustomizableObjectCompileRunnable.cpp
+++ b/Engine/Plugins/Mutable/Source/CustomizableObjectEditor/Private/MuCOE/CustomizableObjectCompileRunnable.cpp
@@ -53,11 +53,11 @@ FCustomizableObjectCompileRunnable::FCustomizableObjectCompileRunnable(UE::Mutab
 }
 
 
-TSharedPtr<UE::Mutable::Private::FImage> FCustomizableObjectCompileRunnable::LoadImageResourceReferenced(int32 ID)
+UE::Mutable::Private::TManagedPtr<UE::Mutable::Private::FImage> FCustomizableObjectCompileRunnable::LoadImageResourceReferenced(int32 ID)
 {
 	MUTABLE_CPUPROFILER_SCOPE(LoadResourceReferenced);
 
-	TSharedPtr<UE::Mutable::Private::FImage> Image;
+	UE::Mutable::Private::TManagedPtr<UE::Mutable::Private::FImage> Image;
 	if (!ReferencedTextures.IsValidIndex(ID))
 	{
 		// The id is not valid for this CO
@@ -69,7 +69,7 @@ TSharedPtr<UE::Mutable::Private::FImage> FCustomizableObjectCompileRunnable::Loa
 	FMutableSourceTextureData& TextureData = ReferencedTextures[ID];
 
 	// In the editor the src data can be directly accessed
-	Image = MakeShared<UE::Mutable::Private::FImage>();
+	Image = UE::Mutable::Private::MakeManaged<UE::Mutable::Private::FImage>();
 	int32 MipmapsToSkip = 0;
 	EUnrealToMutableConversionError Error = ConvertTextureUnrealSourceToMutable(Image.Get(), TextureData, MipmapsToSkip);
 
@@ -142,12 +142,12 @@ uint32 FCustomizableObjectCompileRunnable::Run()
 	}
 
 	CompilerOptions->SetReferencedResourceCallback(
-		[this](int32 ID, TSharedPtr<TSharedPtr<UE::Mutable::Private::FImage>> ResolvedImage, bool bRunImmediatlyIfPossible)
+		[this](int32 ID, TSharedPtr<UE::Mutable::Private::TManagedPtr<UE::Mutable::Private::FImage>> ResolvedImage, bool bRunImmediatlyIfPossible)
 		{
 			UE::Tasks::FTask LaunchTask = UE::Tasks::Launch(TEXT("LoadImageReferenceTasks"),
 				[ID,ResolvedImage,this]()
 				{
-					TSharedPtr<UE::Mutable::Private::FImage> Result = LoadImageResourceReferenced(ID);
+					UE::Mutable::Private::TManagedPtr<UE::Mutable::Private::FImage> Result = LoadImageResourceReferenced(ID);
 					*ResolvedImage = Result;
 				},
 				LowLevelTasks::ETaskPriority::BackgroundLow
@@ -156,7 +156,7 @@ uint32 FCustomizableObjectCompileRunnable::Run()
 			return LaunchTask;
 		},
 
-		[this](int32 ID, const FString& MorphNameRef, TSharedPtr<TSharedPtr<UE::Mutable::Private::FMesh>> ResolvedMesh, bool bRunImmediatlyIfPossible) -> UE::Tasks::FTask
+		[this](int32 ID, const FString& MorphNameRef, TSharedPtr<UE::Mutable::Private::TManagedPtr<UE::Mutable::Private::FMesh>> ResolvedMesh, bool bRunImmediatlyIfPossible) -> UE::Tasks::FTask
 		{
 			MUTABLE_CPUPROFILER_SCOPE(LoadMeshReferenceTasks);
 
@@ -184,7 +184,7 @@ uint32 FCustomizableObjectCompileRunnable::Run()
 					// If we are shutting down, we are not allowed to try to load anything.
 					if (IsEngineExitRequested())
 					{
-						*ResolvedMesh = MakeShared<UE::Mutable::Private::FMesh>();
+						*ResolvedMesh = UE::Mutable::Private::MakeManaged<UE::Mutable::Private::FMesh>();
 						CompletionEvent.Trigger();
 						return;
 					}
@@ -192,7 +192,7 @@ uint32 FCustomizableObjectCompileRunnable::Run()
 					TSharedPtr<FCustomizableObjectCompiler> Compiler = WeakCompiler.Pin();
 					if (!Compiler)
 					{
-						*ResolvedMesh = MakeShared<UE::Mutable::Private::FMesh>();
+						*ResolvedMesh = UE::Mutable::Private::MakeManaged<UE::Mutable::Private::FMesh>();
 						CompletionEvent.Trigger();
 						return;
 					}
@@ -300,7 +300,7 @@ uint32 FCustomizableObjectCompileRunnable::Run()
 				}
 				else
 				{
-					*ResolvedMesh = MakeShared<UE::Mutable::Private::FMesh>();
+					*ResolvedMesh = UE::Mutable::Private::MakeManaged<UE::Mutable::Private::FMesh>();
 					CompletionEvent.Trigger();
 				}
 			}
diff --git a/Engine/Plugins/Mutable/Source/CustomizableObjectEditor/Private/MuCOE/CustomizableObjectCompileRunnable.h b/Engine/Plugins/Mutable/Source/CustomizableObjectEditor/Private/MuCOE/CustomizableObjectCompileRunnable.h
--- a/Engine/Plugins/Mutable/Source/CustomizableObjectEditor/Private/MuCOE/CustomizableObjectCompileRunnable.h
+++ b/Engine/Plugins/Mutable/Source/CustomizableObjectEditor/Private/MuCOE/CustomizableObjectCompileRunnable.h
@@ -50,7 +50,7 @@ private:
 	UE::Mutable::Private::Ptr<UE::Mutable::Private::Node> MutableRoot;
 	TArray<FError> ArrayErrors;
 
-	TSharedPtr<UE::Mutable::Private::FImage> LoadImageResourceReferenced(int32 ID);
+	UE::Mutable::Private::TManagedPtr<UE::Mutable::Private::FImage> LoadImageResourceReferenced(int32 ID);
 
 public:
 
diff --git a/Engine/Plugins/Mutable/Source/CustomizableObjectEditor/Private/MuCOE/CustomizableObjectLayout.cpp b/Engine/Plugins/Mutable/Source/CustomizableObjectEditor/Private/MuCOE/CustomizableObjectLayout.cpp
--- a/Engine/Plugins/Mutable/Source/CustomizableObjectEditor/Private/MuCOE/CustomizableObjectLayout.cpp
+++ b/Engine/Plugins/Mutable/Source/CustomizableObjectEditor/Private/MuCOE/CustomizableObjectLayout.cpp
@@ -155,7 +155,7 @@ void UCustomizableObjectLayout::GenerateAutomaticBlocksFromUVs()
 		}
 	else // Convert Mesh and generate blocks
 	{
-		TSharedPtr<UE::Mutable::Private::FMesh> MutableMesh;
+		UE::Mutable::Private::TManagedPtr<UE::Mutable::Private::FMesh> MutableMesh;
 
 		if (USkeletalMesh* SkeletalMesh = Cast<USkeletalMesh>(GenerationContext.LoadObject(GetMesh())))
 		{
@@ -228,7 +228,7 @@ void UCustomizableObjectLayout::GenerateAutomaticBlocksFromUVs()
 		Block.bIsAutomatic = true;
 		Block.Id = FGuid::NewGuid();
 
-		TSharedPtr<UE::Mutable::Private::FImage> Mask = LayoutNode->Blocks[i].Mask;
+		UE::Mutable::Private::TManagedPtr<UE::Mutable::Private::FImage> Mask = LayoutNode->Blocks[i].Mask;
 		if (Mask)
 		{
 			UTexture2D* UnrealImage = NewObject<UTexture2D>(UTexture2D::StaticClass());
diff --git a/Engine/Plugins/Mutable/Source/CustomizableObjectEditor/Private/MuCOE/GenerateMutableSource/GenerateMutableSource.cpp b/Engine/Plugins/Mutable/Source/CustomizableObjectEditor/Private/MuCOE/GenerateMutableSource/GenerateMutableSource.cpp
--- a/Engine/Plugins/Mutable/Source/CustomizableObjectEditor/Private/MuCOE/GenerateMutableSource/GenerateMutableSource.cpp
+++ b/Engine/Plugins/Mutable/Source/CustomizableObjectEditor/Private/MuCOE/GenerateMutableSource/GenerateMutableSource.cpp
@@ -292,7 +292,7 @@ void FMutableGraphGenerationContext::Log(const FText& Message, const UObject* Co
 }
 
 
-TSharedPtr<UE::Mutable::Private::FMesh> FMutableGraphGenerationContext::FindGeneratedMesh( const FGeneratedMeshData::FKey& Key )
+UE::Mutable::Private::TManagedPtr<UE::Mutable::Private::FMesh> FMutableGraphGenerationContext::FindGeneratedMesh( const FGeneratedMeshData::FKey& Key )
 {
 	for (const FGeneratedMeshData& d : GeneratedMeshes)
 	{
@@ -743,11 +743,14 @@ UE::Mutable::Private::Ptr<UE::Mutable::Private::NodeMeshApplyPose> CreateNodeMes
 {
 	check(ArrayBoneName.Num() == ArrayTransform.Num());
 
-	TSharedPtr<UE::Mutable::Private::FMesh> MutableMesh = MakeShared<UE::Mutable::Private::FMesh>();
+	UE::Mutable::Private::TManagedPtr<UE::Mutable::Private::FMesh> MutableMesh = 
+			UE::Mutable::Private::MakeManaged<UE::Mutable::Private::FMesh>();
 	UE::Mutable::Private::Ptr<UE::Mutable::Private::NodeMeshConstant> PoseNodeMesh = new UE::Mutable::Private::NodeMeshConstant;
 	PoseNodeMesh->Value = MutableMesh;
 
-	TSharedPtr<UE::Mutable::Private::FSkeleton> MutableSkeleton = MakeShared<UE::Mutable::Private::FSkeleton>();
+	UE::Mutable::Private::TManagedPtr<UE::Mutable::Private::FSkeleton> MutableSkeleton = 
+			UE::Mutable::Private::MakeManaged<UE::Mutable::Private::FSkeleton>();
+
 	MutableMesh->SetSkeleton(MutableSkeleton);
 	MutableMesh->SetBonePoseCount(ArrayBoneName.Num());
 	MutableSkeleton->SetBoneCount(ArrayBoneName.Num());
@@ -1724,7 +1727,7 @@ UE::Mutable::Private::FImageDesc GenerateImageDescriptor(UTexture* Texture)
 }
 
 
-TSharedPtr<UE::Mutable::Private::FImage> GenerateImageConstant(UTexture* Texture, FMutableGraphGenerationContext& GenerationContext, bool bIsReference)
+UE::Mutable::Private::TManagedPtr<UE::Mutable::Private::FImage> GenerateImageConstant(UTexture* Texture, FMutableGraphGenerationContext& GenerationContext, bool bIsReference)
 {
 	MUTABLE_CPUPROFILER_SCOPE(GenerateImageConstant);
 
@@ -1791,12 +1794,12 @@ TSharedPtr<UE::Mutable::Private::FImage> GenerateImageConstant(UTexture* Texture
 	}
 	
 	// Compile-time references that are left should be resolved immediately (should only happen in editor).
-	TSharedPtr<UE::Mutable::Private::FImage> Result = UE::Mutable::Private::FImage::CreateAsReference(UE::Mutable::Private::TPassthroughObjectPtr<UTexture>(Id), ImageDesc, bForceLoad); // TODO GMT If bIsCompileTime it should not create a TPassthroughObjectPtr since it is not a Passthrough
+	UE::Mutable::Private::TManagedPtr<UE::Mutable::Private::FImage> Result = UE::Mutable::Private::FImage::CreateAsReference(UE::Mutable::Private::TPassthroughObjectPtr<UTexture>(Id), ImageDesc, bForceLoad); // TODO GMT If bIsCompileTime it should not create a TPassthroughObjectPtr since it is not a Passthrough
 	return Result;
 }
 
 
-TSharedPtr<UE::Mutable::Private::FMesh> GenerateMeshConstant(
+UE::Mutable::Private::TManagedPtr<UE::Mutable::Private::FMesh> GenerateMeshConstant(
 	const FMutableSourceMeshData& Source,
 	FMutableGraphGenerationContext& GenerationContext)
 {
@@ -1822,7 +1825,7 @@ TSharedPtr<UE::Mutable::Private::FMesh> GenerateMeshConstant(
 
 	Id = Entry->ID;
 
-	TSharedPtr<UE::Mutable::Private::FMesh> Result = UE::Mutable::Private::FMesh::CreateAsReference(UE::Mutable::Private::TPassthroughObjectPtr<USkeletalMesh>(Id));
+	UE::Mutable::Private::TManagedPtr<UE::Mutable::Private::FMesh> Result = UE::Mutable::Private::FMesh::CreateAsReference(UE::Mutable::Private::TPassthroughObjectPtr<USkeletalMesh>(Id));
 	return Result;
 }
 
@@ -1968,4 +1971,4 @@ UE::Mutable::Private::PASSTHROUGH_ID FMutableGraphGenerationContext::AddPassthro
 	PassthroughObjects.Add(Id, SoftObjectPtr);
 
 	return Id;
-}
\ No newline at end of file
+}
diff --git a/Engine/Plugins/Mutable/Source/CustomizableObjectEditor/Private/MuCOE/GenerateMutableSource/GenerateMutableSource.h b/Engine/Plugins/Mutable/Source/CustomizableObjectEditor/Private/MuCOE/GenerateMutableSource/GenerateMutableSource.h
--- a/Engine/Plugins/Mutable/Source/CustomizableObjectEditor/Private/MuCOE/GenerateMutableSource/GenerateMutableSource.h
+++ b/Engine/Plugins/Mutable/Source/CustomizableObjectEditor/Private/MuCOE/GenerateMutableSource/GenerateMutableSource.h
@@ -812,7 +812,7 @@ struct FMutableGraphGenerationContext
 		FKey Key;
 
 		/** Generated mesh. */
-		TSharedPtr<UE::Mutable::Private::FMesh> Generated;
+		UE::Mutable::Private::TManagedPtr<UE::Mutable::Private::FMesh> Generated;
 	};
 	TArray<FGeneratedMeshData> GeneratedMeshes;
 
@@ -847,7 +847,7 @@ struct FMutableGraphGenerationContext
 	void Log(const FText& Message, const UObject* Context = nullptr, const EMessageSeverity::Type MessageSeverity = EMessageSeverity::Warning, const bool bAddBaseObjectInfo = true, const ELoggerSpamBin SpamBin = ELoggerSpamBin::ShowAll) const;
 
 	/** Find a mesh if already generated for a given source and flags. */
-	TSharedPtr<UE::Mutable::Private::FMesh> FindGeneratedMesh(const FGeneratedMeshData::FKey& Key);
+	UE::Mutable::Private::TManagedPtr<UE::Mutable::Private::FMesh> FindGeneratedMesh(const FGeneratedMeshData::FKey& Key);
 
 	// Check if the Id of the node Node already exists, if it's new adds it to NodeIds array, otherwise, returns new Id
 	const FGuid GetNodeIdUnique(const UCustomizableObjectNode* Node);
@@ -1096,8 +1096,8 @@ int32 GetMaxTextureSize(const UTexture2D& ReferenceTexture, const UTextureLODSet
 // Max texture size of the texture with per platform MaxTextureSize and LODBias applied.
 int32 GetTextureSizeInGame(const UTexture2D& Texture, const UTextureLODSettings& LODSettings);
 
-TSharedPtr<UE::Mutable::Private::FImage> GenerateImageConstant(UTexture*, FMutableGraphGenerationContext&, bool bIsReference);
-TSharedPtr<UE::Mutable::Private::FMesh> GenerateMeshConstant(const FMutableSourceMeshData& Source, FMutableGraphGenerationContext&);
+UE::Mutable::Private::TManagedPtr<UE::Mutable::Private::FImage> GenerateImageConstant(UTexture*, FMutableGraphGenerationContext&, bool bIsReference);
+UE::Mutable::Private::TManagedPtr<UE::Mutable::Private::FMesh> GenerateMeshConstant(const FMutableSourceMeshData& Source, FMutableGraphGenerationContext&);
 
 /** Generates a mutable image descriptor from an unreal engine texture */
 UE::Mutable::Private::FImageDesc GenerateImageDescriptor(UTexture* Texture);
diff --git a/Engine/Plugins/Mutable/Source/CustomizableObjectEditor/Private/MuCOE/GenerateMutableSource/GenerateMutableSourceExternal.cpp b/Engine/Plugins/Mutable/Source/CustomizableObjectEditor/Private/MuCOE/GenerateMutableSource/GenerateMutableSourceExternal.cpp
--- a/Engine/Plugins/Mutable/Source/CustomizableObjectEditor/Private/MuCOE/GenerateMutableSource/GenerateMutableSourceExternal.cpp
+++ b/Engine/Plugins/Mutable/Source/CustomizableObjectEditor/Private/MuCOE/GenerateMutableSource/GenerateMutableSourceExternal.cpp
@@ -150,7 +150,7 @@ UE::Mutable::Private::Ptr<UE::Mutable::Private::NodeExternal> GenerateMutableSou
 
 		ExternalParameterNode->Name = TypeNodeExternalParameter->GetParameterName(&GenerationContext.MacroNodesStack);
 		ExternalParameterNode->UID = GenerationContext.GetNodeIdUnique(Node).ToString();
-		ExternalParameterNode->DefaultValue = MakeShared<FInstancedStruct>(TypeNodeExternalParameter->DefaultValue);
+		ExternalParameterNode->DefaultValue = UE::Mutable::Private::MakeManaged<FInstancedStruct>(TypeNodeExternalParameter->DefaultValue);
 
 		GenerationContext.ExternalTypeParameterDefaultValues.Add(TypeNodeExternalParameter->GetParameterName(), TypeNodeExternalParameter->DefaultValue);
 		GenerationContext.ExternalTypeParameterTypes.Add(TypeNodeExternalParameter->GetParameterName(), TypeNodeExternalParameter->DefaultValue.GetScriptStruct());
@@ -190,4 +190,4 @@ UE::Mutable::Private::Ptr<UE::Mutable::Private::NodeExternal> GenerateMutableSou
 }
 
 
-#undef LOCTEXT_NAMESPACE
\ No newline at end of file
+#undef LOCTEXT_NAMESPACE
diff --git a/Engine/Plugins/Mutable/Source/CustomizableObjectEditor/Private/MuCOE/GenerateMutableSource/GenerateMutableSourceLayout.cpp b/Engine/Plugins/Mutable/Source/CustomizableObjectEditor/Private/MuCOE/GenerateMutableSource/GenerateMutableSourceLayout.cpp
--- a/Engine/Plugins/Mutable/Source/CustomizableObjectEditor/Private/MuCOE/GenerateMutableSource/GenerateMutableSourceLayout.cpp
+++ b/Engine/Plugins/Mutable/Source/CustomizableObjectEditor/Private/MuCOE/GenerateMutableSource/GenerateMutableSourceLayout.cpp
@@ -82,7 +82,7 @@ UE::Mutable::Private::FSourceLayoutBlock ToMutable(const FCustomizableObjectLayo
 	if (UnrealBlock.Mask)
 	{
 		// In the editor the src data can be directly accessed
-		TSharedPtr<UE::Mutable::Private::FImage> MaskImage = MakeShared<UE::Mutable::Private::FImage>();
+		UE::Mutable::Private::TManagedPtr<UE::Mutable::Private::FImage> MaskImage = UE::Mutable::Private::MakeManaged<UE::Mutable::Private::FImage>();
 
 		FMutableSourceTextureData Tex(*UnrealBlock.Mask);
 		EUnrealToMutableConversionError Error = ConvertTextureUnrealSourceToMutable(MaskImage.Get(), Tex, 0);
diff --git a/Engine/Plugins/Mutable/Source/CustomizableObjectEditor/Private/MuCOE/GenerateMutableSource/GenerateMutableSourceMesh.cpp b/Engine/Plugins/Mutable/Source/CustomizableObjectEditor/Private/MuCOE/GenerateMutableSource/GenerateMutableSourceMesh.cpp
--- a/Engine/Plugins/Mutable/Source/CustomizableObjectEditor/Private/MuCOE/GenerateMutableSource/GenerateMutableSourceMesh.cpp
+++ b/Engine/Plugins/Mutable/Source/CustomizableObjectEditor/Private/MuCOE/GenerateMutableSource/GenerateMutableSourceMesh.cpp
@@ -772,7 +772,7 @@ TArray<TTuple<UPhysicsAsset*, int32>> GetPhysicsAssetsFromAnimInstance(FMutableG
 	return Result;
 }
 
-TArray<uint8> MakePhysicsAssetBodySetupRelevancyMap(const FMutableCompilationContext& Context, UPhysicsAsset* Asset, const TSharedPtr<UE::Mutable::Private::FMesh>& Mesh)
+TArray<uint8> MakePhysicsAssetBodySetupRelevancyMap(const FMutableCompilationContext& Context, UPhysicsAsset* Asset, const UE::Mutable::Private::TManagedPtr<UE::Mutable::Private::FMesh>& Mesh)
 {
 	const int32 BodySetupsNum = Asset->SkeletalBodySetups.Num();
 
@@ -793,7 +793,7 @@ TArray<uint8> MakePhysicsAssetBodySetupRelevancyMap(const FMutableCompilationCon
 	return RelevancyMap;
 }
 
-TSharedPtr<UE::Mutable::Private::FPhysicsBody> MakePhysicsBodyFromAsset(FMutableCompilationContext& Context, UPhysicsAsset* Asset, const TArray<uint8>& BodySetupRelevancyMap)
+UE::Mutable::Private::TManagedPtr<UE::Mutable::Private::FPhysicsBody> MakePhysicsBodyFromAsset(FMutableCompilationContext& Context, UPhysicsAsset* Asset, const TArray<uint8>& BodySetupRelevancyMap)
 {
 	check(Asset);
 	check(Asset->SkeletalBodySetups.Num() == BodySetupRelevancyMap.Num());
@@ -803,7 +803,7 @@ TSharedPtr<UE::Mutable::Private::FPhysicsBody> MakePhysicsBodyFromAsset(FMutable
 	
 	const int32 NumRelevantSetups = Algo::CountIf(BodySetupRelevancyMap, [](uint8 V) { return V; });
 
-	TSharedPtr<UE::Mutable::Private::FPhysicsBody> PhysicsBody = MakeShared<UE::Mutable::Private::FPhysicsBody>();
+	UE::Mutable::Private::TManagedPtr<UE::Mutable::Private::FPhysicsBody> PhysicsBody = UE::Mutable::Private::MakeManaged<UE::Mutable::Private::FPhysicsBody>();
 	
 	PhysicsBody->SetBodyCount(NumRelevantSetups);
 
@@ -922,7 +922,7 @@ TSharedPtr<UE::Mutable::Private::FPhysicsBody> MakePhysicsBodyFromAsset(FMutable
 
 namespace UE::Mutable::Private
 {
-	UE::Tasks::FTask ConvertSkeletalMeshDefaultBuffers(TSharedPtr<UE::Mutable::Private::FMesh> MutableMesh, const FSkeletalMeshLODModel& LODModel, const FSkelMeshSection& MeshSection, TSharedPtr<FMutableConvertSkeletalMeshContext>& MeshContext)
+	UE::Tasks::FTask ConvertSkeletalMeshDefaultBuffers(UE::Mutable::Private::TManagedPtr<UE::Mutable::Private::FMesh> MutableMesh, const FSkeletalMeshLODModel& LODModel, const FSkelMeshSection& MeshSection, TSharedPtr<FMutableConvertSkeletalMeshContext>& MeshContext)
 	{
 		TArrayView<const FSoftSkinVertex> Vertices(MeshSection.SoftVertices);
 
@@ -1208,7 +1208,7 @@ UE::Tasks::FTask ConvertSkeletalMeshToMutable(TSharedRef<FMeshConversionContext>
 	{
 		if (!bMeshMustExist)
 		{
-			MeshConversionContext->Result = MakeShared<UE::Mutable::Private::FMesh>(); // Return empty mesh to preserve the layouts
+			MeshConversionContext->Result = UE::Mutable::Private::MakeManaged<UE::Mutable::Private::FMesh>(); // Return empty mesh to preserve the layouts
 			return UE::Tasks::MakeCompletedTask<void>();
 		}
 		else
@@ -1232,7 +1232,7 @@ UE::Tasks::FTask ConvertSkeletalMeshToMutable(TSharedRef<FMeshConversionContext>
 	{
 		if (!bMeshMustExist)
 		{
-			MeshConversionContext->Result = MakeShared<UE::Mutable::Private::FMesh>(); // Return empty mesh to preserve the layouts
+			MeshConversionContext->Result = UE::Mutable::Private::MakeManaged<UE::Mutable::Private::FMesh>(); // Return empty mesh to preserve the layouts
 			return UE::Tasks::MakeCompletedTask<void>();
 		}
 		else
@@ -1263,7 +1263,7 @@ UE::Tasks::FTask ConvertSkeletalMeshToMutable(TSharedRef<FMeshConversionContext>
 	const bool bIgnoreTexCoords = EnumHasAnyFlags(CurrentFlags, EMutableMeshConversionFlags::IgnoreTexCoords);
 	const bool bDoNotCreateMeshMetadata = EnumHasAnyFlags(CurrentFlags, EMutableMeshConversionFlags::DoNotCreateMeshMetadata);
 
-	TSharedPtr<UE::Mutable::Private::FMesh> MutableMesh = MakeShared<UE::Mutable::Private::FMesh>();
+	UE::Mutable::Private::TManagedPtr<UE::Mutable::Private::FMesh> MutableMesh = UE::Mutable::Private::MakeManaged<UE::Mutable::Private::FMesh>();
 	TSharedPtr<FMutableConvertSkeletalMeshContext> MeshContext = MakeShared<FMutableConvertSkeletalMeshContext>();
 
 	if (UCustomizableObjectSystem::GetInstance()->IsMutableAnimInfoDebuggingEnabled())
@@ -1408,7 +1408,8 @@ UE::Tasks::FTask ConvertSkeletalMeshToMutable(TSharedRef<FMeshConversionContext>
 		MeshContext->bShouldRemoveBonesInfluences = bBoneMapModified;
 
 		// Create the skeleton, poses, and BoneMap for this mesh
-		TSharedPtr<UE::Mutable::Private::FSkeleton> MutableSkeleton = MakeShared<UE::Mutable::Private::FSkeleton>();
+		UE::Mutable::Private::TManagedPtr<UE::Mutable::Private::FSkeleton> MutableSkeleton = 
+				UE::Mutable::Private::MakeManaged<UE::Mutable::Private::FSkeleton>();
 		MutableMesh->SetSkeleton(MutableSkeleton);
 
 		const int32 NumRequiredBones = RequiredBones.Num();
@@ -2387,7 +2388,8 @@ UE::Tasks::FTask ConvertSkeletalMeshToMutable(TSharedRef<FMeshConversionContext>
 			Context.Log(FText::FromString(PhysicsSetupsRemovedMsg), Source.MessageContext, EMessageSeverity::Warning);
 		}
 
-		TSharedPtr<UE::Mutable::Private::FPhysicsBody> PhysicsBody = MakeShared<UE::Mutable::Private::FPhysicsBody>();
+		UE::Mutable::Private::TManagedPtr<UE::Mutable::Private::FPhysicsBody> PhysicsBody = 
+				UE::Mutable::Private::MakeManaged<UE::Mutable::Private::FPhysicsBody>();
 
 		const int32 NumBodySetups = RelevantBodySetups.Num();
 		PhysicsBody->SetBodyCount(NumBodySetups);
@@ -2567,7 +2569,7 @@ UE::Tasks::FTask ConvertSkeletalMeshToMutable(TSharedRef<FMeshConversionContext>
 
 			const int32 PhysicsAssetId = Context.AnimBpOverridePhysicsAssetsInfo.AddUnique(Info);
 
-			TSharedPtr<UE::Mutable::Private::FPhysicsBody> MutableBody = MakePhysicsBodyFromAsset(Context,
+			UE::Mutable::Private::TManagedPtr<UE::Mutable::Private::FPhysicsBody> MutableBody = MakePhysicsBodyFromAsset(Context,
 				PropertyAsset,
 				MakePhysicsAssetBodySetupRelevancyMap(Context, PropertyAsset, MutableMesh));
 			MutableBody->CustomId = PhysicsAssetId;
@@ -2631,7 +2633,7 @@ UE::Tasks::FTask ConvertSkeletalMeshToMutable(TSharedRef<FMeshConversionContext>
 }
 
 
-TSharedPtr<UE::Mutable::Private::FMesh> ConvertSkeletalMeshToMutable(
+UE::Mutable::Private::TManagedPtr<UE::Mutable::Private::FMesh> ConvertSkeletalMeshToMutable(
 	FMutableSourceMeshData& Source,
 	int32 LODIndex, int32 SectionIndex,
 	FMutableGraphGenerationContext& GenerationContext, 
@@ -2666,7 +2668,7 @@ TSharedPtr<UE::Mutable::Private::FMesh> ConvertSkeletalMeshToMutable(
 		}
 	}
 
-	TSharedPtr<UE::Mutable::Private::FMesh> MutableMesh;
+	UE::Mutable::Private::TManagedPtr<UE::Mutable::Private::FMesh> MutableMesh;
 	if (!bForceImmediateConversion)
 	{
 		MutableMesh = GenerateMeshConstant(Source, GenerationContext);
@@ -2693,7 +2695,7 @@ TSharedPtr<UE::Mutable::Private::FMesh> ConvertSkeletalMeshToMutable(
 	return MutableMesh;
 }
 
-TSharedPtr<UE::Mutable::Private::FMesh> ConvertStaticMeshToMutable(const UStaticMesh* StaticMesh, int32 LODIndex, int32 SectionIndex, FMutableGraphGenerationContext& GenerationContext, const UCustomizableObjectNode* CurrentNode)
+UE::Mutable::Private::TManagedPtr<UE::Mutable::Private::FMesh> ConvertStaticMeshToMutable(const UStaticMesh* StaticMesh, int32 LODIndex, int32 SectionIndex, FMutableGraphGenerationContext& GenerationContext, const UCustomizableObjectNode* CurrentNode)
 {
 	if (!StaticMesh->GetRenderData() ||
 		!StaticMesh->GetRenderData()->LODResources.IsValidIndex(LODIndex) ||
@@ -2704,7 +2706,7 @@ TSharedPtr<UE::Mutable::Private::FMesh> ConvertStaticMeshToMutable(const UStatic
 		return nullptr;
 	}
 
-	TSharedPtr<UE::Mutable::Private::FMesh> MutableMesh = MakeShared<UE::Mutable::Private::FMesh>();
+	UE::Mutable::Private::TManagedPtr<UE::Mutable::Private::FMesh> MutableMesh = UE::Mutable::Private::MakeManaged<UE::Mutable::Private::FMesh>();
 
 	// Vertices
 	int32 VertexStart = StaticMesh->GetRenderData()->LODResources[LODIndex].Sections[SectionIndex].MinVertexIndex;
@@ -2881,7 +2883,7 @@ TSharedPtr<UE::Mutable::Private::FMesh> ConvertStaticMeshToMutable(const UStatic
 
 // Convert a Mesh constant to a mutable format. UniqueTags are the tags that make this Mesh unique that cannot be merged in the cache 
 // with the exact same Mesh with other tags
-TSharedPtr<UE::Mutable::Private::FMesh> GenerateMutableSkeletalMesh(FMutableSourceMeshData& Source, int32 LODIndexConnected, int32 SectionIndexConnected,
+UE::Mutable::Private::TManagedPtr<UE::Mutable::Private::FMesh> GenerateMutableSkeletalMesh(FMutableSourceMeshData& Source, int32 LODIndexConnected, int32 SectionIndexConnected,
 										const FString& UniqueTags, FMutableGraphGenerationContext & GenerationContext, 
 										const UCustomizableObjectNode* CurrentNode,
 										const bool bOnlyConnectedLOD)
@@ -2890,7 +2892,7 @@ TSharedPtr<UE::Mutable::Private::FMesh> GenerateMutableSkeletalMesh(FMutableSour
 	EMutableMeshConversionFlags CurrentFlags = GenerationContext.MeshGenerationFlags.Last();
 
 	FMutableGraphGenerationContext::FGeneratedMeshData::FKey Key = { Source.Mesh, LODIndexConnected, GenerationContext.CurrentLOD, SectionIndexConnected, CurrentFlags, UniqueTags, CurrentNode };
-	TSharedPtr<UE::Mutable::Private::FMesh> MutableMesh = GenerationContext.FindGeneratedMesh(Key);
+	UE::Mutable::Private::TManagedPtr<UE::Mutable::Private::FMesh> MutableMesh = GenerationContext.FindGeneratedMesh(Key);
 	if (MutableMesh)
 	{
 		return MutableMesh;
@@ -2910,7 +2912,7 @@ TSharedPtr<UE::Mutable::Private::FMesh> GenerateMutableSkeletalMesh(FMutableSour
 
 // Convert a Mesh constant to a mutable format. UniqueTags are the tags that make this Mesh unique that cannot be merged in the cache 
 // with the exact same Mesh with other tags
-TSharedPtr<UE::Mutable::Private::FMesh> GenerateMutableStaticMesh(TSoftObjectPtr<UStreamableRenderAsset> Mesh, const TSoftClassPtr<UAnimInstance>& AnimInstance,
+UE::Mutable::Private::TManagedPtr<UE::Mutable::Private::FMesh> GenerateMutableStaticMesh(TSoftObjectPtr<UStreamableRenderAsset> Mesh, const TSoftClassPtr<UAnimInstance>& AnimInstance,
 	int32 LODIndex, int32 SectionIndex, const FString& UniqueTags, FMutableGraphGenerationContext& GenerationContext,
 	const UCustomizableObjectNode* CurrentNode, USkeletalMesh* TableReferenceSkeletalMesh, bool bIsPassthrough)
 {
@@ -2918,7 +2920,7 @@ TSharedPtr<UE::Mutable::Private::FMesh> GenerateMutableStaticMesh(TSoftObjectPtr
 	EMutableMeshConversionFlags CurrentFlags = GenerationContext.MeshGenerationFlags.Last();
 
 	FMutableGraphGenerationContext::FGeneratedMeshData::FKey Key = { Mesh, LODIndex, GenerationContext.CurrentLOD, SectionIndex, CurrentFlags, UniqueTags, CurrentNode };
-	TSharedPtr<UE::Mutable::Private::FMesh> MutableMesh = GenerationContext.FindGeneratedMesh(Key);
+	UE::Mutable::Private::TManagedPtr<UE::Mutable::Private::FMesh> MutableMesh = GenerationContext.FindGeneratedMesh(Key);
 	if (MutableMesh)
 	{
 		return MutableMesh;
@@ -2945,7 +2947,7 @@ TSharedPtr<UE::Mutable::Private::FMesh> GenerateMutableStaticMesh(TSoftObjectPtr
 }
 
 
-TSharedPtr<UE::Mutable::Private::FMesh> BuildMorphedMutableMesh(const UEdGraphPin* BaseSourcePin, const FString& MorphTargetName, FMutableGraphGenerationContext& GenerationContext, const bool bOnlyConnectedLOD, const FName& RowName)
+UE::Mutable::Private::TManagedPtr<UE::Mutable::Private::FMesh> BuildMorphedMutableMesh(const UEdGraphPin* BaseSourcePin, const FString& MorphTargetName, FMutableGraphGenerationContext& GenerationContext, const bool bOnlyConnectedLOD, const FName& RowName)
 {
 	check(BaseSourcePin);
 
@@ -2988,14 +2990,14 @@ TSharedPtr<UE::Mutable::Private::FMesh> BuildMorphedMutableMesh(const UEdGraphPi
 		unimplemented();	
 	}
 
-	TSharedPtr<UE::Mutable::Private::FMesh> MorphedSourceMesh;
+	UE::Mutable::Private::TManagedPtr<UE::Mutable::Private::FMesh> MorphedSourceMesh;
 
 	if (SkeletalMesh)
 	{
 		// Get the base mesh
 		FMutableSourceMeshData Source;
 		Source.Mesh = SkeletalMesh;
-		TSharedPtr<UE::Mutable::Private::FMesh> BaseSourceMesh = GenerateMutableSkeletalMesh(Source, BaseLODIndex, BaseSectionIndex, FString(), GenerationContext, Node, bOnlyConnectedLOD);
+		UE::Mutable::Private::TManagedPtr<UE::Mutable::Private::FMesh> BaseSourceMesh = GenerateMutableSkeletalMesh(Source, BaseLODIndex, BaseSectionIndex, FString(), GenerationContext, Node, bOnlyConnectedLOD);
 
 		if (BaseSourceMesh)
 		{
@@ -3551,7 +3553,7 @@ UE::Mutable::Private::NodeMeshPtr GenerateMorphMesh(const UEdGraphPin* Pin,
 	// Target
 	UE::Mutable::Private::NodeMeshPtr BaseSourceMesh = SourceNode;
 
-	TSharedPtr<UE::Mutable::Private::FMesh> MorphedSourceMesh;
+	UE::Mutable::Private::TManagedPtr<UE::Mutable::Private::FMesh> MorphedSourceMesh;
 
 	bool bSuccess = false;
 	
@@ -3589,7 +3591,7 @@ UE::Mutable::Private::NodeMeshPtr GenerateMorphMesh(const UEdGraphPin* Pin,
 				ColumnIndex = GeneratedTable->AddColumn(ColumnName, UE::Mutable::Private::ETableColumnType::Mesh);
 			}
 
-			TSharedPtr<UE::Mutable::Private::FMesh> MorphedSourceTableMesh = BuildMorphedMutableMesh(Pin, TypedNodeMorphs[MorphIndex].MorphTargetName, GenerationContext, bOnlyConnectedLOD, RowName);
+			UE::Mutable::Private::TManagedPtr<UE::Mutable::Private::FMesh> MorphedSourceTableMesh = BuildMorphedMutableMesh(Pin, TypedNodeMorphs[MorphIndex].MorphTargetName, GenerationContext, bOnlyConnectedLOD, RowName);
 			GeneratedTable->SetCell(ColumnIndex, RowId, MorphedSourceTableMesh);
 		}
 
@@ -3853,7 +3855,7 @@ UE::Mutable::Private::Ptr<UE::Mutable::Private::NodeMesh> GenerateMutableSourceM
 			FMutableSourceMeshData Source = BaseMeshData;
 			Source.Mesh = TypedNodeSkel->SkeletalMesh;
 			Source.AnimInstance = TypedNodeSkel->AnimInstance;
-			TSharedPtr<UE::Mutable::Private::FMesh> MutableMesh = GenerateMutableSkeletalMesh(Source, BaseLODIndex, BaseSectionIndex, MeshUniqueTags, GenerationContext, TypedNodeSkel, bOnlyConnectedLOD);
+			UE::Mutable::Private::TManagedPtr<UE::Mutable::Private::FMesh> MutableMesh = GenerateMutableSkeletalMesh(Source, BaseLODIndex, BaseSectionIndex, MeshUniqueTags, GenerationContext, TypedNodeSkel, bOnlyConnectedLOD);
 
 			if (MutableMesh)
 			{
@@ -3974,7 +3976,7 @@ UE::Mutable::Private::Ptr<UE::Mutable::Private::NodeMesh> GenerateMutableSourceM
 		check(TypedNodeStatic->LODs.IsValidIndex(LODIndex) && SectionIndex < TypedNodeStatic->LODs[LODIndex].Materials.Num());
 
 		constexpr bool bIsPassthrough = false;
-		TSharedPtr<UE::Mutable::Private::FMesh> MutableMesh = GenerateMutableStaticMesh(TypedNodeStatic->StaticMesh, TSoftClassPtr<UAnimInstance>(), LODIndex, SectionIndex,
+		UE::Mutable::Private::TManagedPtr<UE::Mutable::Private::FMesh> MutableMesh = GenerateMutableStaticMesh(TypedNodeStatic->StaticMesh, TSoftClassPtr<UAnimInstance>(), LODIndex, SectionIndex,
 													  FString(), GenerationContext, TypedNodeStatic, nullptr, bIsPassthrough);
 		if (MutableMesh)
 		{
diff --git a/Engine/Plugins/Mutable/Source/CustomizableObjectEditor/Private/MuCOE/GenerateMutableSource/GenerateMutableSourceMesh.h b/Engine/Plugins/Mutable/Source/CustomizableObjectEditor/Private/MuCOE/GenerateMutableSource/GenerateMutableSourceMesh.h
--- a/Engine/Plugins/Mutable/Source/CustomizableObjectEditor/Private/MuCOE/GenerateMutableSource/GenerateMutableSourceMesh.h
+++ b/Engine/Plugins/Mutable/Source/CustomizableObjectEditor/Private/MuCOE/GenerateMutableSource/GenerateMutableSourceMesh.h
@@ -47,14 +47,14 @@ void GetLODAndSectionForAutomaticLODs(const FMutableCompilationContext& Context,
  * @param GenerationContext 
  * @param CurrentNode The node being compiled, can be a skeletal mesh or a table node.
  * @return Mutable Mesh. Nullptr if there has been an error. Empty mesh if the Skeletal Mesh does not contain the requested LOD + Section. */
-TSharedPtr<UE::Mutable::Private::FMesh> ConvertSkeletalMeshToMutable(FMutableSourceMeshData& Source,
+UE::Mutable::Private::TManagedPtr<UE::Mutable::Private::FMesh> ConvertSkeletalMeshToMutable(FMutableSourceMeshData& Source,
                                          int32 LODIndex, int32 SectionIndex,
                                          FMutableGraphGenerationContext& GenerationContext,
                                          const UCustomizableObjectNode* CurrentNode, 
 										 bool bForceImmediateConversion = false );
 
 
-TSharedPtr<UE::Mutable::Private::FMesh> ConvertStaticMeshToMutable(const UStaticMesh* StaticMesh, int32 LODIndex, int32 SectionIndex,
+UE::Mutable::Private::TManagedPtr<UE::Mutable::Private::FMesh> ConvertStaticMeshToMutable(const UStaticMesh* StaticMesh, int32 LODIndex, int32 SectionIndex,
 	FMutableGraphGenerationContext& GenerationContext, const UCustomizableObjectNode* CurrentNode);
 
 
@@ -67,7 +67,7 @@ TSharedPtr<UE::Mutable::Private::FMesh> ConvertStaticMeshToMutable(const UStatic
  * @param CurrentNode The node being compiled, can be a skeletal mesh or a table node.
  * @param bOnlyConnectedLOD
  * @return Mutable Mesh. Nullptr if there has been an error. Empty mesh if the Skeletal Mesh does not contain the requested LOD + Section. */
-TSharedPtr<UE::Mutable::Private::FMesh> GenerateMutableSkeletalMesh(FMutableSourceMeshData& Source,
+UE::Mutable::Private::TManagedPtr<UE::Mutable::Private::FMesh> GenerateMutableSkeletalMesh(FMutableSourceMeshData& Source,
                                 int32 LODIndexConnected, int32 SectionIndexConnected,
                                 const FString& MeshUniqueTags, 
                                 FMutableGraphGenerationContext& GenerationContext, 
@@ -75,14 +75,14 @@ TSharedPtr<UE::Mutable::Private::FMesh> GenerateMutableSkeletalMesh(FMutableSour
 								const bool bOnlyConnectedLOD);
 
 /** */
-TSharedPtr<UE::Mutable::Private::FMesh> GenerateMutableStaticMesh(TSoftObjectPtr<UStreamableRenderAsset> Mesh, const TSoftClassPtr<UAnimInstance>& AnimBp,
+UE::Mutable::Private::TManagedPtr<UE::Mutable::Private::FMesh> GenerateMutableStaticMesh(TSoftObjectPtr<UStreamableRenderAsset> Mesh, const TSoftClassPtr<UAnimInstance>& AnimBp,
 	int32 LODIndex, int32 SectionIndex, const FString& MeshUniqueTags,
 	FMutableGraphGenerationContext& GenerationContext,
 	const UCustomizableObjectNode* CurrentNode, USkeletalMesh* TableReferenceSkeletalMesh,
 	bool bIsPassthrough);
 
 
-TSharedPtr<UE::Mutable::Private::FMesh> BuildMorphedMutableMesh(const UEdGraphPin* BaseSourcePin, const FString& MorphTargetName,
+UE::Mutable::Private::TManagedPtr<UE::Mutable::Private::FMesh> BuildMorphedMutableMesh(const UEdGraphPin* BaseSourcePin, const FString& MorphTargetName,
 	FMutableGraphGenerationContext& GenerationContext, bool bOnlyConnectedLOD, const FName& RowName = "");
 
 
diff --git a/Engine/Plugins/Mutable/Source/CustomizableObjectEditor/Private/MuCOE/GenerateMutableSource/GenerateMutableSourceTable.cpp b/Engine/Plugins/Mutable/Source/CustomizableObjectEditor/Private/MuCOE/GenerateMutableSource/GenerateMutableSourceTable.cpp
--- a/Engine/Plugins/Mutable/Source/CustomizableObjectEditor/Private/MuCOE/GenerateMutableSource/GenerateMutableSourceTable.cpp
+++ b/Engine/Plugins/Mutable/Source/CustomizableObjectEditor/Private/MuCOE/GenerateMutableSource/GenerateMutableSourceTable.cpp
@@ -77,7 +77,7 @@ bool FillTableColumn(const UCustomizableObjectNodeTable* TableNode, UE::Mutable:
 
 			if(SoftObject.IsNull())
 			{
-				TSharedPtr<UE::Mutable::Private::FMesh> EmptySkeletalMesh = nullptr;
+				UE::Mutable::Private::TManagedPtr<UE::Mutable::Private::FMesh> EmptySkeletalMesh = nullptr;
 				MutableTable->SetCell(CurrentColumn, RowId, EmptySkeletalMesh);
 
 				return true;
@@ -213,7 +213,7 @@ bool FillTableColumn(const UCustomizableObjectNodeTable* TableNode, UE::Mutable:
 			Source.Mesh = SoftObject.ToSoftObjectPath();
 			Source.TableReferenceSkeletalMesh = ReferenceSkeletalMesh;
 			Source.OptionalMessageInfo = FString::Printf(TEXT(" SkeletalMesh from column [%s] row [%s]"), *ColumnName, *RowName);
-			TSharedPtr<UE::Mutable::Private::FMesh> MutableMesh = GenerateMutableSkeletalMesh(Source, Source.BaseLODIndex, Source.BaseSectionIndex, MeshUniqueTags, GenerationContext, TableNode, Source.bOnlyConnectedLOD);
+			UE::Mutable::Private::TManagedPtr<UE::Mutable::Private::FMesh> MutableMesh = GenerateMutableSkeletalMesh(Source, Source.BaseLODIndex, Source.BaseSectionIndex, MeshUniqueTags, GenerationContext, TableNode, Source.bOnlyConnectedLOD);
 
 			if (MutableMesh)
 			{
@@ -288,7 +288,7 @@ bool FillTableColumn(const UCustomizableObjectNodeTable* TableNode, UE::Mutable:
 			}
 
 			constexpr bool bIsPassthrough = false;
-			TSharedPtr<UE::Mutable::Private::FMesh> MutableMesh = GenerateMutableStaticMesh(StaticMesh, TSoftClassPtr<UAnimInstance>(), CurrentLOD, BaseMeshData.BaseSectionIndex,
+			UE::Mutable::Private::TManagedPtr<UE::Mutable::Private::FMesh> MutableMesh = GenerateMutableStaticMesh(StaticMesh, TSoftClassPtr<UAnimInstance>(), CurrentLOD, BaseMeshData.BaseSectionIndex,
 														  FString(), GenerationContext, TableNode, nullptr, bIsPassthrough);
 
 			if (MutableMesh)
@@ -486,8 +486,11 @@ bool FillTableColumn(const UCustomizableObjectNodeTable* TableNode, UE::Mutable:
 				TArray<FTransform> ArrayTransform;
 				UCustomizableObjectNodeAnimationPose::StaticRetrievePoseInformation(PoseAsset, GenerationContext.GetCurrentComponentInfo()->RefSkeletalMesh.Get(), ArrayBoneName, ArrayTransform);
 
-				TSharedPtr<UE::Mutable::Private::FMesh> MutableMesh = MakeShared<UE::Mutable::Private::FMesh>();
-				TSharedPtr<UE::Mutable::Private::FSkeleton> MutableSkeleton = MakeShared<UE::Mutable::Private::FSkeleton>();
+				UE::Mutable::Private::TManagedPtr<UE::Mutable::Private::FMesh> MutableMesh = 
+						UE::Mutable::Private::MakeManaged<UE::Mutable::Private::FMesh>();
+
+				UE::Mutable::Private::TManagedPtr<UE::Mutable::Private::FSkeleton> MutableSkeleton = 
+						UE::Mutable::Private::MakeManaged<UE::Mutable::Private::FSkeleton>();
 
 				MutableMesh->SetSkeleton(MutableSkeleton);
 				MutableMesh->SetBonePoseCount(ArrayBoneName.Num());
diff --git a/Engine/Plugins/Mutable/Source/CustomizableObjectEditor/Private/MuCOE/SMutableCodeViewer.cpp b/Engine/Plugins/Mutable/Source/CustomizableObjectEditor/Private/MuCOE/SMutableCodeViewer.cpp
--- a/Engine/Plugins/Mutable/Source/CustomizableObjectEditor/Private/MuCOE/SMutableCodeViewer.cpp
+++ b/Engine/Plugins/Mutable/Source/CustomizableObjectEditor/Private/MuCOE/SMutableCodeViewer.cpp
@@ -2733,7 +2733,7 @@ bool SMutableCodeViewer::IsConstantResourceUsedByOperation(const int32 IndexOnCo
 					else
 					{
 						int32 DebuggerIndexOffset = InProgram.ConstantMeshesPermanent.Num();
-						for (const TPair<uint32, TSharedPtr<const UE::Mutable::Private::FMesh>>& Entry : InProgram.ConstantMeshesStreamed)
+						for (const TPair<uint32, UE::Mutable::Private::TManagedPtr<const UE::Mutable::Private::FMesh>>& Entry : InProgram.ConstantMeshesStreamed)
 						{
 							if (Entry.Key == ResourceIndex.Index)
 							{
@@ -2887,7 +2887,7 @@ void SMutableCodeViewer::Tick(const FGeometry& AllottedGeometry, const double In
 	case UE::Mutable::Private::EDataType::Layout:
 	{
 		check(PreviewLayoutViewer);
-		TSharedPtr<const UE::Mutable::Private::FLayout> MutableLayout = System->GetPrivate()->BuildLayout(ExternalResourceProvider, MutableModel, PreviewParameters.Get(), SelectedOperationAddress);
+		UE::Mutable::Private::TManagedPtr<const UE::Mutable::Private::FLayout> MutableLayout = System->GetPrivate()->BuildLayout(ExternalResourceProvider, MutableModel, PreviewParameters.Get(), SelectedOperationAddress);
 		PreviewLayoutViewer->SetLayout(MutableLayout);
 		break;
 	}
@@ -2895,7 +2895,7 @@ void SMutableCodeViewer::Tick(const FGeometry& AllottedGeometry, const double In
 	case UE::Mutable::Private::EDataType::Image:
 	{
 		check(PreviewImageViewer);
-		TSharedPtr<const UE::Mutable::Private::FImage> MutableImage = System->GetPrivate()->BuildImage(ExternalResourceProvider, MutableModel, PreviewParameters.Get(), SelectedOperationAddress, MipsToSkip);
+		UE::Mutable::Private::TManagedPtr<const UE::Mutable::Private::FImage> MutableImage = System->GetPrivate()->BuildImage(ExternalResourceProvider, MutableModel, PreviewParameters.Get(), SelectedOperationAddress, MipsToSkip);
 		PreviewImageViewer->SetImage(MutableImage, 0);
 		break;
 	}
@@ -2903,7 +2903,7 @@ void SMutableCodeViewer::Tick(const FGeometry& AllottedGeometry, const double In
 	case UE::Mutable::Private::EDataType::Mesh:
 	{
 		check(PreviewMeshViewer);
-		TSharedPtr<const UE::Mutable::Private::FMesh> MutableMesh = System->GetPrivate()->BuildMesh(ExternalResourceProvider, MutableModel, PreviewParameters.Get(), SelectedOperationAddress, UE::Mutable::Private::EMeshContentFlags::AllFlags);
+		UE::Mutable::Private::TManagedPtr<const UE::Mutable::Private::FMesh> MutableMesh = System->GetPrivate()->BuildMesh(ExternalResourceProvider, MutableModel, PreviewParameters.Get(), SelectedOperationAddress, UE::Mutable::Private::EMeshContentFlags::AllFlags);
 		PreviewMeshViewer->SetMesh(MutableMesh);
 		break;
 	}
@@ -2911,7 +2911,7 @@ void SMutableCodeViewer::Tick(const FGeometry& AllottedGeometry, const double In
 	case UE::Mutable::Private::EDataType::Instance:
 	{
 		check(PreviewInstanceViewer);
-		TSharedPtr<const UE::Mutable::Private::FInstance> MutableInstance = System->GetPrivate()->BuildInstance(ExternalResourceProvider, MutableModel, PreviewParameters.Get(), SelectedOperationAddress);
+		UE::Mutable::Private::TManagedPtr<const UE::Mutable::Private::FInstance> MutableInstance = System->GetPrivate()->BuildInstance(ExternalResourceProvider, MutableModel, PreviewParameters.Get(), SelectedOperationAddress);
 		PreviewInstanceViewer->SetInstance(MutableInstance, MutableModel, PreviewParameters, *System);
 		break;
 	}
@@ -2943,7 +2943,7 @@ void SMutableCodeViewer::Tick(const FGeometry& AllottedGeometry, const double In
 	case UE::Mutable::Private::EDataType::String:
 	{
 		check(PreviewStringViewer);
-		TSharedPtr<const UE::Mutable::Private::String> MutableString = System->GetPrivate()->BuildString(ExternalResourceProvider, MutableModel, PreviewParameters.Get(), SelectedOperationAddress);
+		UE::Mutable::Private::TManagedPtr<const UE::Mutable::Private::String> MutableString = System->GetPrivate()->BuildString(ExternalResourceProvider, MutableModel, PreviewParameters.Get(), SelectedOperationAddress);
 		const FText MutableText = FText::FromString(FString(MutableString->GetValue()));
 		PreviewStringViewer->SetString(MutableText);
 		break;
@@ -3058,21 +3058,21 @@ void SMutableCodeViewer::PreviewMutableString(const FString& InString)
 }
 
 
-void SMutableCodeViewer::PreviewMutableImage(TSharedPtr<const UE::Mutable::Private::FImage> InImagePtr)
+void SMutableCodeViewer::PreviewMutableImage(UE::Mutable::Private::TManagedPtr<const UE::Mutable::Private::FImage> InImagePtr)
 {
 	PrepareImageViewer();
 	PreviewImageViewer->SetImage(InImagePtr,0);
 }
 
 
-void SMutableCodeViewer::PreviewMutableMesh(TSharedPtr<const UE::Mutable::Private::FMesh> InMeshPtr)
+void SMutableCodeViewer::PreviewMutableMesh(UE::Mutable::Private::TManagedPtr<const UE::Mutable::Private::FMesh> InMeshPtr)
 {
 	PrepareMeshViewer();
 	PreviewMeshViewer->SetMesh(InMeshPtr);
 }
 
 
-void SMutableCodeViewer::PreviewMutableLayout(TSharedPtr<const UE::Mutable::Private::FLayout> Layout)
+void SMutableCodeViewer::PreviewMutableLayout(UE::Mutable::Private::TManagedPtr<const UE::Mutable::Private::FLayout> Layout)
 {
 	PrepareLayoutViewer();
 	PreviewLayoutViewer->SetLayout(Layout);
@@ -3092,7 +3092,7 @@ void SMutableCodeViewer::PreviewMutableProjector(const UE::Mutable::Private::FPr
 }
 
 
-void SMutableCodeViewer::PreviewMutableSkeleton(TSharedPtr<const UE::Mutable::Private::FSkeleton> Skeleton)
+void SMutableCodeViewer::PreviewMutableSkeleton(UE::Mutable::Private::TManagedPtr<const UE::Mutable::Private::FSkeleton> Skeleton)
 {
 	if (!PreviewSkeletonViewer)
 	{
@@ -3118,7 +3118,7 @@ void SMutableCodeViewer::PreviewMutableCurve(const FRichCurve& Curve)
 }
 
 // TODO: Implement physics viewer
-void SMutableCodeViewer::PreviewMutablePhysics(TSharedPtr<const UE::Mutable::Private::FPhysicsBody> Physics)
+void SMutableCodeViewer::PreviewMutablePhysics(UE::Mutable::Private::TManagedPtr<const UE::Mutable::Private::FPhysicsBody> Physics)
 {
 	UE_LOG(LogMutable, Warning, TEXT("Previewer for Mutable Physics not yet implemented"))
 }
diff --git a/Engine/Plugins/Mutable/Source/CustomizableObjectEditor/Private/MuCOE/SMutableCodeViewer.h b/Engine/Plugins/Mutable/Source/CustomizableObjectEditor/Private/MuCOE/SMutableCodeViewer.h
--- a/Engine/Plugins/Mutable/Source/CustomizableObjectEditor/Private/MuCOE/SMutableCodeViewer.h
+++ b/Engine/Plugins/Mutable/Source/CustomizableObjectEditor/Private/MuCOE/SMutableCodeViewer.h
@@ -791,11 +791,11 @@ private:
 	 */
 
 public:
-	void PreviewMutableImage (TSharedPtr<const UE::Mutable::Private::FImage> InImagePtr);
-	void PreviewMutableMesh (TSharedPtr<const UE::Mutable::Private::FMesh> InMeshPtr);
-	void PreviewMutableLayout(TSharedPtr<const UE::Mutable::Private::FLayout> Layout);
-	void PreviewMutableSkeleton(TSharedPtr<const UE::Mutable::Private::FSkeleton> Skeleton);
-	void PreviewMutablePhysics(TSharedPtr<const UE::Mutable::Private::FPhysicsBody> Physics);
+	void PreviewMutableImage (UE::Mutable::Private::TManagedPtr<const UE::Mutable::Private::FImage> InImagePtr);
+	void PreviewMutableMesh (UE::Mutable::Private::TManagedPtr<const UE::Mutable::Private::FMesh> InMeshPtr);
+	void PreviewMutableLayout(UE::Mutable::Private::TManagedPtr<const UE::Mutable::Private::FLayout> Layout);
+	void PreviewMutableSkeleton(UE::Mutable::Private::TManagedPtr<const UE::Mutable::Private::FSkeleton> Skeleton);
+	void PreviewMutablePhysics(UE::Mutable::Private::TManagedPtr<const UE::Mutable::Private::FPhysicsBody> Physics);
 	void PreviewMutableString(const FString& InString);
 	void PreviewMutableProjector(const UE::Mutable::Private::FProjector* Projector);
 	void PreviewMutableMatrix(const FMatrix44f& Mat);
