[Mutable] Memory managment rework.

diff --git a/Engine/Plugins/Mutable/Source/CustomizableObjectEditor/Private/MuCOE/SMutableConstantsWidget.cpp b/Engine/Plugins/Mutable/Source/CustomizableObjectEditor/Private/MuCOE/SMutableConstantsWidget.cpp
--- a/Engine/Plugins/Mutable/Source/CustomizableObjectEditor/Private/MuCOE/SMutableConstantsWidget.cpp
+++ b/Engine/Plugins/Mutable/Source/CustomizableObjectEditor/Private/MuCOE/SMutableConstantsWidget.cpp
@@ -81,7 +81,7 @@ namespace
 	 * @param MeshPtr Pointer to the mesh whose channels we want to count
 	 * @return : the amount of channels found.
 	 */
-	int32 GetMeshChannelCount (TSharedPtr<const UE::Mutable::Private::FMesh> MeshPtr)
+	int32 GetMeshChannelCount (UE::Mutable::Private::TManagedPtr<const UE::Mutable::Private::FMesh> MeshPtr)
 	{
 		check(MeshPtr)
 
@@ -135,7 +135,7 @@ public:
 		// The amount of mesh vertex buffer and image buffer channels
 		if (ColumnName == MeshConstantTitles::MeshChannels)
 		{
-			TSharedPtr<const UE::Mutable::Private::FMesh> Mesh = RowElement->MeshPtr;
+			UE::Mutable::Private::TManagedPtr<const UE::Mutable::Private::FMesh> Mesh = RowElement->MeshPtr;
 			const int32 ChannelCount =  GetMeshChannelCount(Mesh);
 			const FText ChannelCountText = FText::AsNumber(ChannelCount);
 			return SNew(STextBlock).Text(ChannelCountText);
@@ -144,7 +144,7 @@ public:
 		// The amount of indices of the mesh
 		if (ColumnName == MeshConstantTitles::MeshIndices)
 		{
-			TSharedPtr<const UE::Mutable::Private::FMesh> Mesh = RowElement->MeshPtr;
+			UE::Mutable::Private::TManagedPtr<const UE::Mutable::Private::FMesh> Mesh = RowElement->MeshPtr;
 			const FText IndexCount = FText::AsNumber(Mesh->GetIndexCount());
 			return SNew(STextBlock).Text(IndexCount);
 		}
@@ -152,7 +152,7 @@ public:
 		// The amount of vertices of the mesh
 		if (ColumnName == MeshConstantTitles::MeshVertices)
 		{
-			TSharedPtr<const UE::Mutable::Private::FMesh> Mesh = RowElement->MeshPtr;
+			UE::Mutable::Private::TManagedPtr<const UE::Mutable::Private::FMesh> Mesh = RowElement->MeshPtr;
 			const FText VertexCount = FText::AsNumber(Mesh->GetVertexCount());
 			return SNew(STextBlock).Text(VertexCount);
 		}
@@ -160,7 +160,7 @@ public:
 		// Memory used by the mesh 
 		if (ColumnName == MeshConstantTitles::MeshMemory)
 		{
-			TSharedPtr<const UE::Mutable::Private::FMesh> Mesh = RowElement->MeshPtr;
+			UE::Mutable::Private::TManagedPtr<const UE::Mutable::Private::FMesh> Mesh = RowElement->MeshPtr;
 			const FText SizeAsText = GenerateTextForSize(Mesh->GetDataSize());
 			return SNew(STextBlock).Text(SizeAsText);
 		}
@@ -1259,7 +1259,11 @@ void SMutableConstantsWidget::LoadConstantImages()
 		TSharedPtr<FMutableConstantImageElement> ConstantImageElement = MakeShared<FMutableConstantImageElement>();
 
 		MutableProgramPtr->GetConstant(ImageIndex, ConstantImageElement->ImagePtr, 0,
-			[this](int32 x, int32 y, int32 m, UE::Mutable::Private::EImageFormat f, UE::Mutable::Private::EInitializationType i) { return MakeShared<UE::Mutable::Private::FImage>(x, y, m, f, i); });
+			[this](int32 x, int32 y, int32 m, UE::Mutable::Private::EImageFormat f, UE::Mutable::Private::EInitializationType i) 
+			{ 
+				return UE::Mutable::Private::MakeManaged<UE::Mutable::Private::FImage>(x, y, m, f, i); 
+			}
+		);
 
 		ConstantImageElement->IndexOnSourceVector = ImageIndex;
 		
@@ -1276,10 +1280,10 @@ void SMutableConstantsWidget::LoadConstantMeshes()
 {
 	check (MutableProgramPtr);
 	
-	TArray<TSharedPtr<const UE::Mutable::Private::FMesh>> AllMeshes;
+	TArray<UE::Mutable::Private::TManagedPtr<const UE::Mutable::Private::FMesh>> AllMeshes;
 	AllMeshes.Reserve(MutableProgramPtr->ConstantMeshesPermanent.Num() + MutableProgramPtr->ConstantMeshesStreamed.Num() );
 	AllMeshes.Append(MutableProgramPtr->ConstantMeshesPermanent);
-	for (const TPair<uint32, TSharedPtr<const UE::Mutable::Private::FMesh>>& Entry : MutableProgramPtr->ConstantMeshesStreamed)
+	for (const TPair<uint32, UE::Mutable::Private::TManagedPtr<const UE::Mutable::Private::FMesh>>& Entry : MutableProgramPtr->ConstantMeshesStreamed)
 	{
 		AllMeshes.Add(Entry.Value);
 	}
diff --git a/Engine/Plugins/Mutable/Source/CustomizableObjectEditor/Private/MuCOE/SMutableConstantsWidget.h b/Engine/Plugins/Mutable/Source/CustomizableObjectEditor/Private/MuCOE/SMutableConstantsWidget.h
--- a/Engine/Plugins/Mutable/Source/CustomizableObjectEditor/Private/MuCOE/SMutableConstantsWidget.h
+++ b/Engine/Plugins/Mutable/Source/CustomizableObjectEditor/Private/MuCOE/SMutableConstantsWidget.h
@@ -35,7 +35,7 @@ struct FMutableConstantElement
 */
 struct FMutableConstantImageElement : public FMutableConstantElement
 {
-	TSharedPtr<const UE::Mutable::Private::FImage> ImagePtr;
+	UE::Mutable::Private::TManagedPtr<const UE::Mutable::Private::FImage> ImagePtr;
 };
 
 /**
@@ -43,7 +43,7 @@ struct FMutableConstantImageElement : public FMutableConstantElement
 */
 struct FMutableConstantMeshElement : public FMutableConstantElement
 {
-	TSharedPtr<const UE::Mutable::Private::FMesh> MeshPtr;
+	UE::Mutable::Private::TManagedPtr<const UE::Mutable::Private::FMesh> MeshPtr;
 };
 
 /**
@@ -59,7 +59,7 @@ struct FMutableConstantStringElement : public FMutableConstantElement
 */
 struct FMutableConstantLayoutElement : public FMutableConstantElement
 {
-	TSharedPtr<const UE::Mutable::Private::FLayout> Layout;
+	UE::Mutable::Private::TManagedPtr<const UE::Mutable::Private::FLayout> Layout;
 };
 
 /**
@@ -67,7 +67,7 @@ struct FMutableConstantLayoutElement : public FMutableConstantElement
 */
 struct FMutableConstantSkeletonElement : public FMutableConstantElement
 {
-	TSharedPtr<const UE::Mutable::Private::FSkeleton> Skeleton;
+	UE::Mutable::Private::TManagedPtr<const UE::Mutable::Private::FSkeleton> Skeleton;
 };
 
 /**
@@ -75,7 +75,7 @@ struct FMutableConstantSkeletonElement : public FMutableConstantElement
 */
 struct FMutableConstantPhysicsElement : public FMutableConstantElement
 {
-	TSharedPtr<const UE::Mutable::Private::FPhysicsBody> Physics;
+	UE::Mutable::Private::TManagedPtr<const UE::Mutable::Private::FPhysicsBody> Physics;
 };
 
 /**
diff --git a/Engine/Plugins/Mutable/Source/CustomizableObjectEditor/Private/MuCOE/SMutableImageViewer.cpp b/Engine/Plugins/Mutable/Source/CustomizableObjectEditor/Private/MuCOE/SMutableImageViewer.cpp
--- a/Engine/Plugins/Mutable/Source/CustomizableObjectEditor/Private/MuCOE/SMutableImageViewer.cpp
+++ b/Engine/Plugins/Mutable/Source/CustomizableObjectEditor/Private/MuCOE/SMutableImageViewer.cpp
@@ -267,7 +267,7 @@ void SMutableImageViewer::Tick(const FGeometry& AllottedGeometry, const double I
 		{
 			// If the image format is not supported in the editor platform (for instance, it is a mobile-specific format)
 			// convert the mutable image to a generic format so that it can be previewed.
-			TSharedPtr<const UE::Mutable::Private::FImage> ImageToConvert = MutableImage;
+			UE::Mutable::Private::TManagedPtr<const UE::Mutable::Private::FImage> ImageToConvert = MutableImage;
 
 			// \TODO: This test is a bit weak
 			if (MutableImage->GetFormat() >= UE::Mutable::Private::EImageFormat::ASTC_4x4_RGB_LDR
@@ -305,13 +305,14 @@ void SMutableImageViewer::Tick(const FGeometry& AllottedGeometry, const double I
 }
 
 
-void SMutableImageViewer::SetImage(const TSharedPtr<const UE::Mutable::Private::FImage>& InMutableImage, int32 CurrentLOD)
+void SMutableImageViewer::SetImage(const UE::Mutable::Private::TManagedPtr<const UE::Mutable::Private::FImage>& InMutableImage, int32 CurrentLOD)
 {
 	MutableImage = InMutableImage;
 	CurrentVisibleLOD = CurrentLOD;
 	bIsPendingUpdate = true;
 
 	constexpr bool bDebugDump = false;
+
 	if (bDebugDump && InMutableImage)
 	{
 		static int32 SImageCount = 0;
diff --git a/Engine/Plugins/Mutable/Source/CustomizableObjectEditor/Private/MuCOE/SMutableImageViewer.h b/Engine/Plugins/Mutable/Source/CustomizableObjectEditor/Private/MuCOE/SMutableImageViewer.h
--- a/Engine/Plugins/Mutable/Source/CustomizableObjectEditor/Private/MuCOE/SMutableImageViewer.h
+++ b/Engine/Plugins/Mutable/Source/CustomizableObjectEditor/Private/MuCOE/SMutableImageViewer.h
@@ -78,7 +78,7 @@ public:
 
 	SLATE_BEGIN_ARGS(SMutableImageViewer) {}
 		SLATE_ATTRIBUTE(FIntPoint, GridSize)
-		SLATE_ARGUMENT(TSharedPtr<const UE::Mutable::Private::FImage>, Image)
+		SLATE_ARGUMENT(UE::Mutable::Private::TManagedPtr<const UE::Mutable::Private::FImage>, Image)
 	SLATE_END_ARGS()
 
 	void Construct(const FArguments& InArgs);
@@ -89,7 +89,7 @@ public:
 	// Own interface
 
 	/** Set the image to show in the widget. */
-	void SetImage(const TSharedPtr<const UE::Mutable::Private::FImage>& Image, int32 VisibleLOD );
+	void SetImage(const UE::Mutable::Private::TManagedPtr<const UE::Mutable::Private::FImage>& Image, int32 VisibleLOD );
 
 private:
 
@@ -100,7 +100,7 @@ private:
 	TAttribute<FIntPoint> GridSize;
 
 	/** */
-	TSharedPtr<const UE::Mutable::Private::FImage> MutableImage;
+	UE::Mutable::Private::TManagedPtr<const UE::Mutable::Private::FImage> MutableImage;
 
 	/** Mipmap shown in the widget. 0 is the biggest one. */
 	int32 CurrentVisibleLOD=0;
diff --git a/Engine/Plugins/Mutable/Source/CustomizableObjectEditor/Private/MuCOE/SMutableInstanceViewer.cpp b/Engine/Plugins/Mutable/Source/CustomizableObjectEditor/Private/MuCOE/SMutableInstanceViewer.cpp
--- a/Engine/Plugins/Mutable/Source/CustomizableObjectEditor/Private/MuCOE/SMutableInstanceViewer.cpp
+++ b/Engine/Plugins/Mutable/Source/CustomizableObjectEditor/Private/MuCOE/SMutableInstanceViewer.cpp
@@ -285,7 +285,7 @@ FString ToString(const UE::Mutable::Private::TPassthroughObjectPtr<UObject>& Pas
 
 
 
-void SMutableInstanceViewer::SetInstance(const TSharedPtr<const UE::Mutable::Private::FInstance>& InInstance, const TSharedPtr<const UE::Mutable::Private::FModel>& Model, const TSharedPtr<const UE::Mutable::Private::FParameters>& Parameters, const UE::Mutable::Private::FSystem& System)
+void SMutableInstanceViewer::SetInstance(const UE::Mutable::Private::TManagedPtr<const UE::Mutable::Private::FInstance>& InInstance, const TSharedPtr<const UE::Mutable::Private::FModel>& Model, const TSharedPtr<const UE::Mutable::Private::FParameters>& Parameters, const UE::Mutable::Private::FSystem& System)
 {
 	if (InInstance != MutableInstance)
 	{
@@ -321,7 +321,7 @@ void SMutableInstanceViewer::SetInstance(const TSharedPtr<const UE::Mutable::Pri
 					if (MeshId)
 					{
 						UE::Mutable::Private::OP::ADDRESS MeshRootAddress = MeshId.GetKey()->Address;
-						TSharedPtr<const UE::Mutable::Private::FMesh> MutableMesh = System.GetPrivate()->BuildMesh(
+						UE::Mutable::Private::TManagedPtr<const UE::Mutable::Private::FMesh> MutableMesh = System.GetPrivate()->BuildMesh(
 								ExternalResourceProvider,
 								Model, 
 								Parameters.Get(), 
diff --git a/Engine/Plugins/Mutable/Source/CustomizableObjectEditor/Private/MuCOE/SMutableInstanceViewer.h b/Engine/Plugins/Mutable/Source/CustomizableObjectEditor/Private/MuCOE/SMutableInstanceViewer.h
--- a/Engine/Plugins/Mutable/Source/CustomizableObjectEditor/Private/MuCOE/SMutableInstanceViewer.h
+++ b/Engine/Plugins/Mutable/Source/CustomizableObjectEditor/Private/MuCOE/SMutableInstanceViewer.h
@@ -2,7 +2,10 @@
 
 #pragma once
 
+
+
 #include "MuR/Instance.h"
+#include "MuR/ManagedPointer.h"
 #include "MuR/System.h"
 #include "UObject/GCObject.h"
 #include "Widgets/SCompoundWidget.h"
@@ -59,7 +62,7 @@ public:
 	void Construct(const FArguments& InArgs, const TSharedPtr<FUnrealMutableResourceProvider>& InExternalResourceProvider);
 
 	/** Set the Mutable Instance to be used for this widget */
-	void SetInstance(const TSharedPtr<const UE::Mutable::Private::FInstance>&, const TSharedPtr<const UE::Mutable::Private::FModel>&, const TSharedPtr<const UE::Mutable::Private::FParameters>&, const UE::Mutable::Private::FSystem&);
+	void SetInstance(const UE::Mutable::Private::TManagedPtr<const UE::Mutable::Private::FInstance>&, const TSharedPtr<const UE::Mutable::Private::FModel>&, const TSharedPtr<const UE::Mutable::Private::FParameters>&, const UE::Mutable::Private::FSystem&);
 	
 	/** */
 	TSharedRef<SWidget> GenerateLODsListView(const TSharedPtr<TArray<TSharedPtr<FMutableInstanceViewerLODElement>>>& InBufferChannelElements);
@@ -69,7 +72,7 @@ private:
 	TSharedPtr<FUnrealMutableResourceProvider> ExternalResourceProvider;
 	
 	/** Data backend for the widget. It represents the Instance that is being "displayed" */
-	TSharedPtr<const UE::Mutable::Private::FInstance> MutableInstance = nullptr;
+	UE::Mutable::Private::TManagedPtr<const UE::Mutable::Private::FInstance> MutableInstance = nullptr;
 
 	/** Splitter used to separate the two sides of the slate (tables and viewport) */
 	TSharedPtr<SSplitter> SpaceSplitter;
diff --git a/Engine/Plugins/Mutable/Source/CustomizableObjectEditor/Private/MuCOE/SMutableLayoutViewer.cpp b/Engine/Plugins/Mutable/Source/CustomizableObjectEditor/Private/MuCOE/SMutableLayoutViewer.cpp
--- a/Engine/Plugins/Mutable/Source/CustomizableObjectEditor/Private/MuCOE/SMutableLayoutViewer.cpp
+++ b/Engine/Plugins/Mutable/Source/CustomizableObjectEditor/Private/MuCOE/SMutableLayoutViewer.cpp
@@ -120,7 +120,7 @@ void SMutableLayoutViewer::Tick(const FGeometry& AllottedGeometry, const double
 }
 
 
-void SMutableLayoutViewer::SetLayout(const TSharedPtr<const UE::Mutable::Private::FLayout>& InMutableLayout)
+void SMutableLayoutViewer::SetLayout(const UE::Mutable::Private::TManagedPtr<const UE::Mutable::Private::FLayout>& InMutableLayout)
 {
 	MutableLayout = InMutableLayout;
 	bIsPendingUpdate = true;
diff --git a/Engine/Plugins/Mutable/Source/CustomizableObjectEditor/Private/MuCOE/SMutableLayoutViewer.h b/Engine/Plugins/Mutable/Source/CustomizableObjectEditor/Private/MuCOE/SMutableLayoutViewer.h
--- a/Engine/Plugins/Mutable/Source/CustomizableObjectEditor/Private/MuCOE/SMutableLayoutViewer.h
+++ b/Engine/Plugins/Mutable/Source/CustomizableObjectEditor/Private/MuCOE/SMutableLayoutViewer.h
@@ -17,7 +17,7 @@ public:
 
 	SLATE_BEGIN_ARGS(SMutableLayoutViewer) {}
 		SLATE_ATTRIBUTE(FIntPoint, GridSize)
-		SLATE_ARGUMENT(TSharedPtr<const UE::Mutable::Private::FLayout>, Layout)
+		SLATE_ARGUMENT(UE::Mutable::Private::TManagedPtr<const UE::Mutable::Private::FLayout>, Layout)
 	SLATE_END_ARGS()
 
 	void Construct(const FArguments& InArgs);
@@ -28,7 +28,7 @@ public:
 	// Own interface
 
 	/** Set the image to show in the widget. */
-	void SetLayout(const TSharedPtr<const UE::Mutable::Private::FLayout>& Image );
+	void SetLayout(const UE::Mutable::Private::TManagedPtr<const UE::Mutable::Private::FLayout>& Image );
 
 private:
 
@@ -36,7 +36,7 @@ private:
 	TSharedPtr<SCustomizableObjectLayoutGrid> LayoutViewer;
 
 	/** Mutable Layout object being previewed */
-	TSharedPtr<const UE::Mutable::Private::FLayout> MutableLayout;
+	UE::Mutable::Private::TManagedPtr<const UE::Mutable::Private::FLayout> MutableLayout;
 
 	/** Is true, the image or the visible LOD have changed and we need to update. */
 	bool bIsPendingUpdate = false;
diff --git a/Engine/Plugins/Mutable/Source/CustomizableObjectEditor/Private/MuCOE/SMutableMeshViewer.cpp b/Engine/Plugins/Mutable/Source/CustomizableObjectEditor/Private/MuCOE/SMutableMeshViewer.cpp
--- a/Engine/Plugins/Mutable/Source/CustomizableObjectEditor/Private/MuCOE/SMutableMeshViewer.cpp
+++ b/Engine/Plugins/Mutable/Source/CustomizableObjectEditor/Private/MuCOE/SMutableMeshViewer.cpp
@@ -200,7 +200,7 @@ void SMutableMeshViewer::Construct(const FArguments& InArgs)
 }
 
 
-void SMutableMeshViewer::SetMesh(const TSharedPtr<const UE::Mutable::Private::FMesh>& InMesh)
+void SMutableMeshViewer::SetMesh(const UE::Mutable::Private::TManagedPtr<const UE::Mutable::Private::FMesh>& InMesh)
 {
 	if (InMesh != MutableMesh)
 	{
@@ -219,7 +219,7 @@ void SMutableMeshViewer::SetMesh(const TSharedPtr<const UE::Mutable::Private::FM
 			OnMeshChanged();
 
 			// Skeleton slate viewer update
-			if (TSharedPtr<const UE::Mutable::Private::FSkeleton> MutableSkeleton = MutableMesh->GetSkeleton())
+			if (UE::Mutable::Private::TManagedPtr<const UE::Mutable::Private::FSkeleton> MutableSkeleton = MutableMesh->GetSkeleton())
 			{
 				MutableSkeletonViewer->SetSkeleton(MutableSkeleton);
 
diff --git a/Engine/Plugins/Mutable/Source/CustomizableObjectEditor/Private/MuCOE/SMutableMeshViewer.h b/Engine/Plugins/Mutable/Source/CustomizableObjectEditor/Private/MuCOE/SMutableMeshViewer.h
--- a/Engine/Plugins/Mutable/Source/CustomizableObjectEditor/Private/MuCOE/SMutableMeshViewer.h
+++ b/Engine/Plugins/Mutable/Source/CustomizableObjectEditor/Private/MuCOE/SMutableMeshViewer.h
@@ -47,14 +47,14 @@ class SMutableMeshViewer : public SCompoundWidget
 {
 public:
 	SLATE_BEGIN_ARGS(SMutableMeshViewer) {}
-		SLATE_ARGUMENT_DEFAULT(TSharedPtr<const UE::Mutable::Private::FMesh>, Mesh){nullptr};
+		SLATE_ARGUMENT_DEFAULT(UE::Mutable::Private::TManagedPtr<const UE::Mutable::Private::FMesh>, Mesh){nullptr};
 	SLATE_END_ARGS()
 
 	/** Builds the widget */
 	void Construct(const FArguments& InArgs);
 
 	/** Set the Mutable Mesh to be used for this widget */
-	void SetMesh(const TSharedPtr<const UE::Mutable::Private::FMesh>& InMesh);
+	void SetMesh(const UE::Mutable::Private::TManagedPtr<const UE::Mutable::Private::FMesh>& InMesh);
 
 	
 	/** Generate a new SListView for all the ChannelElements provided.
@@ -74,7 +74,7 @@ private:
 	TSharedRef<ITableRow> GenerateTagRow(TSharedPtr<FString> InItem, const TSharedRef<STableViewBase>& OwnerTable);
 
 	/** Data backend for the widget. It represents the mesh that is being "displayed" */
-	TSharedPtr<const UE::Mutable::Private::FMesh> MutableMesh = nullptr;
+	UE::Mutable::Private::TManagedPtr<const UE::Mutable::Private::FMesh> MutableMesh = nullptr;
 
 	/** Splitter used to separate the two sides of the slate (tables and viewport) */
 	TSharedPtr<SSplitter> SpaceSplitter;
diff --git a/Engine/Plugins/Mutable/Source/CustomizableObjectEditor/Private/MuCOE/SMutableMeshViewport.cpp b/Engine/Plugins/Mutable/Source/CustomizableObjectEditor/Private/MuCOE/SMutableMeshViewport.cpp
--- a/Engine/Plugins/Mutable/Source/CustomizableObjectEditor/Private/MuCOE/SMutableMeshViewport.cpp
+++ b/Engine/Plugins/Mutable/Source/CustomizableObjectEditor/Private/MuCOE/SMutableMeshViewport.cpp
@@ -33,7 +33,7 @@ void SMutableMeshViewport::Tick(const FGeometry& AllottedGeometry, const double
 	SEditorViewport::Tick(AllottedGeometry, InCurrentTime, InDeltaTime);
 }
 
-void SMutableMeshViewport::SetMesh(const TSharedPtr<const UE::Mutable::Private::FMesh>& InMesh)
+void SMutableMeshViewport::SetMesh(const UE::Mutable::Private::TManagedPtr<const UE::Mutable::Private::FMesh>& InMesh)
 {
 	if (InMesh == MutableMesh)
 	{
diff --git a/Engine/Plugins/Mutable/Source/CustomizableObjectEditor/Private/MuCOE/SMutableMeshViewport.h b/Engine/Plugins/Mutable/Source/CustomizableObjectEditor/Private/MuCOE/SMutableMeshViewport.h
--- a/Engine/Plugins/Mutable/Source/CustomizableObjectEditor/Private/MuCOE/SMutableMeshViewport.h
+++ b/Engine/Plugins/Mutable/Source/CustomizableObjectEditor/Private/MuCOE/SMutableMeshViewport.h
@@ -22,7 +22,7 @@ struct FGeometry;
 class SMutableMeshViewport : public SEditorViewport , public FGCObject
 {
 	SLATE_BEGIN_ARGS(SMutableMeshViewport) {}
-    	SLATE_ARGUMENT_DEFAULT(TSharedPtr<const UE::Mutable::Private::FMesh>, Mesh) { nullptr };
+    	SLATE_ARGUMENT_DEFAULT(UE::Mutable::Private::TManagedPtr<const UE::Mutable::Private::FMesh>, Mesh) { nullptr };
     SLATE_END_ARGS()
     
 public:
@@ -32,7 +32,7 @@ public:
     /** Set the Mutable Mesh to be used for this widget
      * @param InMesh - The mutable mesh to be displayed
      */
-    void SetMesh(const TSharedPtr<const UE::Mutable::Private::FMesh>& InMesh);
+    void SetMesh(const UE::Mutable::Private::TManagedPtr<const UE::Mutable::Private::FMesh>& InMesh);
 
 	// FGCObject interface
 	virtual void AddReferencedObjects(FReferenceCollector& Collector) override;
@@ -52,7 +52,7 @@ protected:
 private:
 
 	/** The mutable mesh being displayed */
-	TSharedPtr<const UE::Mutable::Private::FMesh> MutableMesh = nullptr;
+	UE::Mutable::Private::TManagedPtr<const UE::Mutable::Private::FMesh> MutableMesh = nullptr;
 
 	/*
 	 *	UI  
diff --git a/Engine/Plugins/Mutable/Source/CustomizableObjectEditor/Private/MuCOE/SMutableSkeletonViewer.cpp b/Engine/Plugins/Mutable/Source/CustomizableObjectEditor/Private/MuCOE/SMutableSkeletonViewer.cpp
--- a/Engine/Plugins/Mutable/Source/CustomizableObjectEditor/Private/MuCOE/SMutableSkeletonViewer.cpp
+++ b/Engine/Plugins/Mutable/Source/CustomizableObjectEditor/Private/MuCOE/SMutableSkeletonViewer.cpp
@@ -91,7 +91,7 @@ void SMutableSkeletonViewer::Construct(const FArguments& InArgs)
 	SetSkeleton(InArgs._Skeleton);
 }
 
-void SMutableSkeletonViewer::SetSkeleton(TSharedPtr<const UE::Mutable::Private::FSkeleton> InSkeleton)
+void SMutableSkeletonViewer::SetSkeleton(UE::Mutable::Private::TManagedPtr<const UE::Mutable::Private::FSkeleton> InSkeleton)
 {
 	if (InSkeleton != Skeleton)
 	{
diff --git a/Engine/Plugins/Mutable/Source/CustomizableObjectEditor/Private/MuCOE/SMutableSkeletonViewer.h b/Engine/Plugins/Mutable/Source/CustomizableObjectEditor/Private/MuCOE/SMutableSkeletonViewer.h
--- a/Engine/Plugins/Mutable/Source/CustomizableObjectEditor/Private/MuCOE/SMutableSkeletonViewer.h
+++ b/Engine/Plugins/Mutable/Source/CustomizableObjectEditor/Private/MuCOE/SMutableSkeletonViewer.h
@@ -24,7 +24,7 @@ struct FBoneDefinition
 class SMutableSkeletonViewer : public SCompoundWidget
 {
 	SLATE_BEGIN_ARGS(SMutableSkeletonViewer) {}
-		SLATE_ARGUMENT_DEFAULT(TSharedPtr<const UE::Mutable::Private::FSkeleton>, Skeleton) { nullptr };
+		SLATE_ARGUMENT_DEFAULT(UE::Mutable::Private::TManagedPtr<const UE::Mutable::Private::FSkeleton>, Skeleton) { nullptr };
 	SLATE_END_ARGS()
 
 public:
@@ -32,12 +32,13 @@ public:
 	void Construct(const FArguments& InArgs);
 
 	/** Sets the mutable skeleton object to be inspected */
-	void SetSkeleton(TSharedPtr<const UE::Mutable::Private::FSkeleton> InSkeleton);
+	void SetSkeleton(UE::Mutable::Private::TManagedPtr<const UE::Mutable::Private::FSkeleton> InSkeleton);
 	
 private:
 
+
 	/** The mutable skeleton whose data is being inspected */
-	TSharedPtr<const UE::Mutable::Private::FSkeleton> Skeleton;
+	UE::Mutable::Private::TManagedPtr<const UE::Mutable::Private::FSkeleton> Skeleton;
 	
 	/** Bones data: Bone definitions used to generate the bone tree structure */
 	TArray<TSharedPtr<FBoneDefinition>> RootBones;
diff --git a/Engine/Plugins/Mutable/Source/CustomizableObjectEditor/Private/MuCOE/UnrealBakeHelpers.cpp b/Engine/Plugins/Mutable/Source/CustomizableObjectEditor/Private/MuCOE/UnrealBakeHelpers.cpp
--- a/Engine/Plugins/Mutable/Source/CustomizableObjectEditor/Private/MuCOE/UnrealBakeHelpers.cpp
+++ b/Engine/Plugins/Mutable/Source/CustomizableObjectEditor/Private/MuCOE/UnrealBakeHelpers.cpp
@@ -197,7 +197,7 @@ UTexture2D* FUnrealBakeHelpers::BakeHelper_CreateAssetTexture(UTexture2D* Source
 
 	UE::Mutable::Private::FImageOperator ImOp = UE::Mutable::Private::FImageOperator::GetDefault(UE::Mutable::Private::FImageOperator::FImagePixelFormatFunc());
 	int32 Quality = 4; // Doesn't matter for decompression.
-	TSharedPtr<UE::Mutable::Private::FImage> UncompressedImage = ImOp.ImagePixelFormat(Quality, PlatformImage.Get(), UncompressedMutableFormat);
+	UE::Mutable::Private::TManagedPtr<UE::Mutable::Private::FImage> UncompressedImage = ImOp.ImagePixelFormat(Quality, PlatformImage.Get(), UncompressedMutableFormat);
 
 	// Copy the decompressed data to the texture source data
 	int32 SourceDataSize = DupTex->Source.CalcMipSize(MipIndex);
diff --git a/Engine/Plugins/Mutable/Source/MutableRuntime/Internal/MuR/ExtensionData.h b/Engine/Plugins/Mutable/Source/MutableRuntime/Internal/MuR/ExtensionData.h
--- a/Engine/Plugins/Mutable/Source/MutableRuntime/Internal/MuR/ExtensionData.h
+++ b/Engine/Plugins/Mutable/Source/MutableRuntime/Internal/MuR/ExtensionData.h
@@ -19,7 +19,7 @@ namespace UE::Mutable::Private
 	{
 	public:
 		static UE_API void Serialise(const FExtensionData* Data, FOutputArchive& Archive);
-		static UE_API TSharedPtr<FExtensionData> StaticUnserialise(FInputArchive& Archive);
+		static UE_API TManagedPtr<FExtensionData> StaticUnserialise(FInputArchive& Archive);
 
 		// Resource interface
 		UE_API int32 GetDataSize() const override;
diff --git a/Engine/Plugins/Mutable/Source/MutableRuntime/Internal/MuR/Image.h b/Engine/Plugins/Mutable/Source/MutableRuntime/Internal/MuR/Image.h
--- a/Engine/Plugins/Mutable/Source/MutableRuntime/Internal/MuR/Image.h
+++ b/Engine/Plugins/Mutable/Source/MutableRuntime/Internal/MuR/Image.h
@@ -19,6 +19,8 @@
 
 #include "MuR/ImageDataStorage.h"
 
+#include "MuR/ManagedPointer.h"
+
 #define UE_API MUTABLERUNTIME_API
 
 namespace UE::Mutable::Private
@@ -46,11 +48,11 @@ namespace UE::Mutable::Private
 		UE_API FImage(uint32 SizeX, uint32 SizeY, uint32 Lods, EImageFormat Format, EInitializationType InitType);
 
 		/** Create a new empty image that repreents an external resource image. */
-		static UE_API TSharedPtr<FImage> CreateAsReference(const TPassthroughObjectPtr<UTexture>& ExternalObject, const FImageDesc& Desc, bool bForceLoad);
+		static UE_API TManagedRef<FImage> CreateAsReference(const TPassthroughObjectPtr<UTexture>& ExternalObject, const FImageDesc& Desc, bool bForceLoad);
 
 		//! Serialisation
 		static UE_API void Serialise( const FImage*, FOutputArchive& );
-		static UE_API TSharedPtr<FImage> StaticUnserialise( FInputArchive& );
+		static UE_API TManagedRef<FImage> StaticUnserialise( FInputArchive& );
 
 		// Resource interface
 		UE_API int32 GetDataSize() const override;
@@ -134,12 +136,12 @@ namespace UE::Mutable::Private
 		//-----------------------------------------------------------------------------------------
 
 		//! Deep clone.
-		TSharedPtr<FImage> Clone() const
+		TManagedPtr<FImage> Clone() const
 		{
 			LLM_SCOPE_BYNAME(TEXT("MutableRuntime"));
 			MUTABLE_CPUPROFILER_SCOPE(ImageClone)
 
-			TSharedPtr<FImage> Result = MakeShared<FImage>();
+			TManagedPtr<FImage> Result = MakeManaged<FImage>();
 
 			Result->Flags = Flags;
 			Result->RelevancyMinY = RelevancyMinY;
@@ -249,12 +251,10 @@ namespace UE::Mutable::Private
 	struct FImageOperator
 	{
 		/** Common callback used for functions that can create temporary images. */
-		typedef TFunction<TSharedPtr<FImage>(int32, int32, int32, EImageFormat, EInitializationType)> FImageCreateFunc;
-		typedef TFunction<void(TSharedPtr<FImage>&)> FImageReleaseFunc;
-		typedef TFunction<TSharedPtr<FImage>(const FImage*)> FImageCloneFunc;
+		typedef TFunction<TManagedPtr<FImage>(int32, int32, int32, EImageFormat, EInitializationType)> FImageCreateFunc;
+		typedef TFunction<TManagedPtr<FImage>(const FImage*)> FImageCloneFunc;
 
 		FImageCreateFunc CreateImage;
-		FImageReleaseFunc ReleaseImage;
 		FImageCloneFunc CloneImage;
 
 		/** Interface to override the internal mutable image pixel format conversion functions.
@@ -263,14 +263,18 @@ namespace UE::Mutable::Private
 		typedef TFunction<void(bool&, int32, FImage*, const FImage*, int32)> FImagePixelFormatFunc;
 		FImagePixelFormatFunc FormatImageOverride;
 
-		FImageOperator(FImageCreateFunc Create, FImageReleaseFunc Release, FImageCloneFunc Clone, FImagePixelFormatFunc FormatOverride) : CreateImage(Create), ReleaseImage(Release), CloneImage(Clone), FormatImageOverride(FormatOverride) {};
+		FImageOperator(FImageCreateFunc Create, FImageCloneFunc Clone, FImagePixelFormatFunc FormatOverride) 
+			: CreateImage(Create)
+			, CloneImage(Clone)
+			, FormatImageOverride(FormatOverride) 
+		{
+		};
 
 		/** Create an default version for untracked resources. */
 		static FImageOperator GetDefault( const FImagePixelFormatFunc& InFormatOverride )
 		{
 			return FImageOperator(
-				[](int32 x, int32 y, int32 m, EImageFormat f, EInitializationType i) { return MakeShared<FImage>(x, y, m, f, i); },
-				[](TSharedPtr<FImage>& In) {In = nullptr; },
+				[](int32 SizeX, int32 SizeY, int32 NumLODs, EImageFormat Format, EInitializationType InitType) { return MakeManaged<FImage>(SizeX, SizeY, NumLODs, Format, InitType); },
 				[](const FImage* In) { return In->Clone(); },
 				InFormatOverride
 			);
@@ -282,7 +286,7 @@ namespace UE::Mutable::Private
 		* \param onlyLOD If different than -1, only the specified lod level will be converted in the returned image.
 		* \return nullptr if the conversion failed, usually because not enough memory was allocated in the result. This is only checked for RLE compression.
 		*/
-		MUTABLERUNTIME_API TSharedPtr<FImage> ImagePixelFormat(int32 Quality, const FImage* Base, EImageFormat TargetFormat, int32 OnlyLOD = -1);
+		MUTABLERUNTIME_API TManagedPtr<FImage> ImagePixelFormat(int32 Quality, const FImage* Base, EImageFormat TargetFormat, int32 OnlyLOD = -1);
 
 		/** Convert an image to another pixel format.
 		* \warning Not all format conversions are implemented.
@@ -292,10 +296,10 @@ namespace UE::Mutable::Private
 		MUTABLERUNTIME_API void ImagePixelFormat(bool& bOutSuccess, int32 Quality, FImage* Result, const FImage* Base, int32 OnlyLOD = -1);
 		MUTABLERUNTIME_API void ImagePixelFormat(bool& bOutSuccess, int32 Quality, FImage* Result, const FImage* Base, int32 BeginResultLOD, int32 BeginBaseLOD, int32 NumLODs);
 
-		MUTABLERUNTIME_API TSharedPtr<FImage> ImageSwizzle(EImageFormat Format, const TSharedPtr<const FImage> Sources[], const uint8 Channels[]);
+		MUTABLERUNTIME_API TManagedPtr<FImage> ImageSwizzle(EImageFormat Format, const TManagedPtr<const FImage> Sources[], const uint8 Channels[]);
 
 		/** */
-		MUTABLERUNTIME_API TSharedPtr<FImage> ExtractMip(const FImage* From, int32 Mip);
+		MUTABLERUNTIME_API TManagedPtr<FImage> ExtractMip(const FImage* From, int32 Mip);
 
 		/** Bilinear filter image resize. */
 		MUTABLERUNTIME_API void ImageResizeLinear(FImage* Dest, int32 ImageCompressionQuality, const FImage* Base);
@@ -306,9 +310,9 @@ namespace UE::Mutable::Private
 		/** Support struct to keep preallocated data required for some mipmap operations. */
 		struct FScratchImageMipmap
 		{
-			TSharedPtr<FImage> Uncompressed;
-			TSharedPtr<FImage> UncompressedMips;
-			TSharedPtr<FImage> CompressedMips;
+			TManagedPtr<FImage> Uncompressed;
+			TManagedPtr<FImage> UncompressedMips;
+			TManagedPtr<FImage> CompressedMips;
 		};
 
 		/** Generate the mipmaps for images.
@@ -325,10 +329,9 @@ namespace UE::Mutable::Private
 		void ImageMipmap(FImageOperator::FScratchImageMipmap&, int32 CompressionQuality, FImage* Dest, const FImage* Base,
 			int32 LevelStart, int32 LevelCount,
 			const FMipmapGenerationSettings&, bool bGenerateOnlyTail = false);
-		void ImageMipmap_ReleaseScratch(FScratchImageMipmap&);
 
 		/** */
-		MUTABLERUNTIME_API bool ImageCrop(TSharedPtr<FImage>& OutCropped, int32 CompressionQuality, const FImage* Base, const box<FIntVector2>& Rect);
+		MUTABLERUNTIME_API bool ImageCrop(TManagedPtr<FImage>& OutCropped, int32 CompressionQuality, const FImage* Base, const box<FIntVector2>& Rect);
 
 		/** */
 		void ImageCompose(FImage* Base, const FImage* Block, const box<FIntVector2>& Rect);
@@ -341,7 +344,7 @@ namespace UE::Mutable::Private
 	MUTABLERUNTIME_API void ImageMipmapInPlace(int32 CompressionQuality, FImage* Base, const FMipmapGenerationSettings&);
 
 	/** */
-	MUTABLERUNTIME_API void ImageSwizzle(FImage* Result, const TSharedPtr<const FImage> Sources[], const uint8 Channels[]);
+	MUTABLERUNTIME_API void ImageSwizzle(FImage* Result, const TManagedPtr<const FImage> Sources[], const uint8 Channels[]);
 
 	MUTABLERUNTIME_API inline EImageFormat GetUncompressedFormat(EImageFormat Format)
 	{
@@ -383,35 +386,37 @@ namespace UE::Mutable::Private
 
 	/** */
 	template<class T>
-	TSharedPtr<T> CloneOrTakeOver(const TSharedPtr<T>& Source)
+	TManagedPtr<T> CloneOrTakeOver(TManagedPtr<T>&& Source)
 	{
-		TSharedPtr<T> Result;
-		if (Source.IsUnique())
+		TManagedPtr<T> Result;
+		if (Source.IsUniqueReference())
 		{
-			Result = ConstCastSharedPtr<T>(Source);
+			Result = ConstCastManagedPtr<T>(Source);
 		}
 		else
 		{
 			Result = Source->Clone();
 		}
 
+		Source = nullptr;
 		return Result;
 	}
 
 	/** */
 	template<class T>
-	TSharedPtr<T> CloneOrTakeOver(const TSharedPtr<const T>& Source)
+	TManagedPtr<T> CloneOrTakeOver(TManagedPtr<const T>&& Source)
 	{
-		TSharedPtr<T> Result;
-		if (Source.IsUnique())
+		TManagedPtr<T> Result;
+		if (Source.IsUniqueReference())
 		{
-			Result = ConstCastSharedPtr<T>(Source);
+			Result = ConstCastManagedPtr<T>(Source);
 		}
 		else
 		{
 			Result = Source->Clone();
 		}
 
+		Source = nullptr;
 		return Result;
 	}
 
diff --git a/Engine/Plugins/Mutable/Source/MutableRuntime/Internal/MuR/Instance.h b/Engine/Plugins/Mutable/Source/MutableRuntime/Internal/MuR/Instance.h
--- a/Engine/Plugins/Mutable/Source/MutableRuntime/Internal/MuR/Instance.h
+++ b/Engine/Plugins/Mutable/Source/MutableRuntime/Internal/MuR/Instance.h
@@ -2,6 +2,7 @@
 
 #pragma once
 
+#include "MuR/ManagedPointer.h"
 #include "MuR/ResourceID.h"
 #include "HAL/PlatformMath.h"
 #include "MuR/ExtensionData.h"
@@ -67,7 +68,7 @@ namespace UE::Mutable::Private
 	
 	struct NamedExtensionData
 	{
-		TSharedPtr<const FExtensionData> Data;
+		UE::Mutable::Private::TManagedPtr<const FExtensionData> Data;
 		FName Name;
 	};
 
@@ -85,7 +86,7 @@ namespace UE::Mutable::Private
 		//-----------------------------------------------------------------------------------------
 		
         //! Clone this instance
-		UE_API TSharedPtr<FInstance> Clone() const;
+		UE_API UE::Mutable::Private::TManagedPtr<FInstance> Clone() const;
 
 		// Resource interface
 		UE_API int32 GetDataSize() const override;
@@ -180,7 +181,7 @@ namespace UE::Mutable::Private
 		//! \param Index Index of the ExtensionData to fetch
 		//! \param OutExtensionData Receives the ExtensionData
 		//! \param OutName Receives the name associated with the ExtensionData. Guaranteed to be a valid string of non-zero length.
-		UE_API void GetExtensionData(int32 Index, TSharedPtr<const FExtensionData>& OutExtensionData, FName& OutName) const;
+		UE_API void GetExtensionData(int32 Index, TManagedPtr<const FExtensionData>& OutExtensionData, FName& OutName) const;
 		
 		int32 AddComponent();
 		
@@ -202,7 +203,7 @@ namespace UE::Mutable::Private
 		
 		void AddOverrideMaterial(int32 ComponentIndex, const FName& SlotName, const FMaterialId& MaterialId);
 
-		void SetExtensionData(const TSharedRef<const FExtensionData>& Data, FName Name);
+		void SetExtensionData(const TManagedRef<const FExtensionData>& Data, FName Name);
 		
 		TArray<FInstanceComponent, TInlineAllocator<1>> Components;
 
diff --git a/Engine/Plugins/Mutable/Source/MutableRuntime/Internal/MuR/Layout.h b/Engine/Plugins/Mutable/Source/MutableRuntime/Internal/MuR/Layout.h
--- a/Engine/Plugins/Mutable/Source/MutableRuntime/Internal/MuR/Layout.h
+++ b/Engine/Plugins/Mutable/Source/MutableRuntime/Internal/MuR/Layout.h
@@ -129,7 +129,7 @@ namespace UE::Mutable::Private
 		/** List of masks that may be used by the blocks for finer-grade selection of vertices. 
 		* The masks may be shared by multiple blocks and also across layouts. 
 		*/
-		TArray<TSharedPtr<const FImage>> Masks;
+		TArray<TManagedPtr<const FImage>> Masks;
 
 	public:
 
@@ -138,11 +138,11 @@ namespace UE::Mutable::Private
 		//-----------------------------------------------------------------------------------------
 
 		//! Deep clone this layout.
-		UE_API TSharedPtr<FLayout> Clone() const;
+		UE_API TManagedPtr<FLayout> Clone() const;
 
 		//! Serialisation
 		static UE_API void Serialise( const FLayout* p, FOutputArchive& arch );
-		static UE_API TSharedPtr<FLayout> StaticUnserialise( FInputArchive& arch );
+		static UE_API TManagedRef<FLayout> StaticUnserialise( FInputArchive& arch );
 
 		//! Full compare
 		UE_API bool operator==( const FLayout& other ) const;
diff --git a/Engine/Plugins/Mutable/Source/MutableRuntime/Internal/MuR/Material.h b/Engine/Plugins/Mutable/Source/MutableRuntime/Internal/MuR/Material.h
--- a/Engine/Plugins/Mutable/Source/MutableRuntime/Internal/MuR/Material.h
+++ b/Engine/Plugins/Mutable/Source/MutableRuntime/Internal/MuR/Material.h
@@ -51,10 +51,10 @@ namespace UE::Mutable::Private
 
 		//! Serialisation
 		static UE_API void Serialise(const FMaterial* MaterialPtr, FOutputArchive& Arch);
-		static UE_API TSharedPtr<FMaterial> StaticUnserialise(FInputArchive& Arch);
+		static UE_API TManagedPtr<FMaterial> StaticUnserialise(FInputArchive& Arch);
 
         //! Clone this material
-		UE_API TSharedPtr<FMaterial> Clone() const;
+		UE_API TManagedPtr<FMaterial> Clone() const;
 
 		// Resource interface
 		virtual int32 GetDataSize() const override { return 0; };
@@ -71,7 +71,7 @@ namespace UE::Mutable::Private
 
 		struct FImageParameterData 
 		{
-			TVariant<OP::ADDRESS, TSharedPtr<const FImage>> ImageParameter;
+			TVariant<OP::ADDRESS, TManagedPtr<const FImage>> ImageParameter;
 			int32 ImagePropertyIndex;
 
 			inline void Serialise(FOutputArchive& arch) const;
diff --git a/Engine/Plugins/Mutable/Source/MutableRuntime/Internal/MuR/Mesh.h b/Engine/Plugins/Mutable/Source/MutableRuntime/Internal/MuR/Mesh.h
--- a/Engine/Plugins/Mutable/Source/MutableRuntime/Internal/MuR/Mesh.h
+++ b/Engine/Plugins/Mutable/Source/MutableRuntime/Internal/MuR/Mesh.h
@@ -19,6 +19,8 @@
 #include "MuR/PassthroughObject.h"
 #include "Templates/Tuple.h"
 
+#include "MuR/ManagedPointer.h"
+
 #include "UObject/StrongObjectPtr.h"
 
 #include <type_traits>
@@ -235,20 +237,20 @@ namespace UE::Mutable::Private
         //-----------------------------------------------------------------------------------------
 				
 		/** Create a new empty mesh that repreents an external resource mesh. */
-		static UE_API TSharedPtr<FMesh> CreateAsReference(const TPassthroughObjectPtr<USkeletalMesh>& PasstroughObjectPtr);
+		static UE_API TManagedRef<FMesh> CreateAsReference(const TPassthroughObjectPtr<USkeletalMesh>& PasstroughObjectPtr);
 
         /** Deep clone this mesh. */
-        UE_API TSharedPtr<FMesh> Clone() const;
+        UE_API TManagedPtr<FMesh> Clone() const;
 		
 		/** Clone with flags allowing to not include some parts in the cloned mesh */
-		UE_API TSharedPtr<FMesh> Clone(EMeshCopyFlags Flags) const;
+		UE_API TManagedPtr<FMesh> Clone(EMeshCopyFlags Flags) const;
 
 		/** Copy form another mesh. */
 		UE_API void CopyFrom(const FMesh& From, EMeshCopyFlags Flags = EMeshCopyFlags::AllFlags);
 
         /** Serialisation */
         static UE_API void Serialise(const FMesh* InMesh, FOutputArchive& Arch );
-        static UE_API TSharedPtr<FMesh> StaticUnserialise(FInputArchive& Arch);
+        static UE_API TManagedPtr<FMesh> StaticUnserialise(FInputArchive& Arch);
 
 		// Resource interface
 		UE_API int32 GetDataSize() const override;
@@ -317,34 +319,34 @@ namespace UE::Mutable::Private
         //! \{
 
         //!
-        UE_API void AddLayout( TSharedPtr<const FLayout> );
+        UE_API void AddLayout( TManagedPtr<const FLayout> );
 
         //!
         UE_API int32 GetLayoutCount() const;
 
         //!
-		UE_API TSharedPtr<const FLayout> GetLayout( int32 Index) const;
+		UE_API TManagedPtr<const FLayout> GetLayout( int32 Index) const;
 
         //!
-        UE_API void SetLayout( int32 Index, TSharedPtr<const FLayout> );
+        UE_API void SetLayout( int32 Index, TManagedPtr<const FLayout> );
         //! \}
 
         //! \name Skeleton information
         //! \{
 
-        UE_API void SetSkeleton(TSharedPtr<const FSkeleton> );
-		UE_API TSharedPtr<const FSkeleton> GetSkeleton() const;
+        UE_API void SetSkeleton(const TManagedPtr<const FSkeleton>&);
+		UE_API TManagedPtr<const FSkeleton> GetSkeleton() const;
 
         //! \}
 
         //! \name PhysicsBody information
         //! \{
 
-        UE_API void SetPhysicsBody(TSharedPtr<const FPhysicsBody> );
-		UE_API TSharedPtr<const FPhysicsBody> GetPhysicsBody() const;
+        UE_API void SetPhysicsBody(const TManagedPtr<const FPhysicsBody>&);
+		UE_API TManagedPtr<const FPhysicsBody> GetPhysicsBody() const;
 
-		UE_API int32 AddAdditionalPhysicsBody(TSharedPtr<const FPhysicsBody>);
-		UE_API TSharedPtr<const FPhysicsBody> GetAdditionalPhysicsBody(int32 I) const;
+		UE_API int32 AddAdditionalPhysicsBody(const TManagedPtr<const FPhysicsBody>&);
+		UE_API TManagedPtr<const FPhysicsBody> GetAdditionalPhysicsBody(int32 I) const;
 		//int32 GetAdditionalPhysicsBodyExternalId(int32 I) const;
 
         //! \}
@@ -439,6 +441,7 @@ namespace UE::Mutable::Private
 		/** If the mesh is a reference the referenced morph name is stored here. Otherwise it is an empty string. */
 		FString ReferencedMorph;
 
+		TManagedPtr<int32> A = nullptr;
 		/** 
 		 * Prefix for the unique IDs related to this mesh (vertices and layout blocks). Useful if the mesh stores them in an implicit, or relative way. 
 		 * See MeshVertexIdIterator for details.
@@ -465,17 +468,17 @@ namespace UE::Mutable::Private
 		 * This skeleton and physicsbody are not owned and may be used by other meshes, so it cannot be modified
 		 * once the mesh has been fully created.
 		 */
-		TSharedPtr<const FSkeleton> Skeleton;
-		TSharedPtr<const FPhysicsBody> PhysicsBody;
+		TManagedPtr<const FSkeleton> Skeleton;
+		TManagedPtr<const FPhysicsBody> PhysicsBody;
     	
 		/** Additional physics bodies referenced by the mesh that don't merge. */
-		TArray<TSharedPtr<const FPhysicsBody>> AdditionalPhysicsBodies;
+		TArray<TManagedPtr<const FPhysicsBody>> AdditionalPhysicsBodies;
 
 		/** 
 		 * Texture Layout blocks attached to this mesh. They are const because they could be shared with
 		 * other meshes, so they need to be cloned and replaced if a modification is needed.
 		 */
-		TArray<TSharedPtr<const FLayout>> Layouts;		
+		TArray<TManagedPtr<const FLayout>> Layouts;		
 
 		TArray<FString> Tags;
 
diff --git a/Engine/Plugins/Mutable/Source/MutableRuntime/Internal/MuR/ModelPrivate.h b/Engine/Plugins/Mutable/Source/MutableRuntime/Internal/MuR/ModelPrivate.h
--- a/Engine/Plugins/Mutable/Source/MutableRuntime/Internal/MuR/ModelPrivate.h
+++ b/Engine/Plugins/Mutable/Source/MutableRuntime/Internal/MuR/ModelPrivate.h
@@ -150,7 +150,7 @@ namespace UE::Mutable::Private
 		// This should always be valid, but if the state is Unloaded it won't be usable.
 		//
 		// Avoid storing references to this Data in Memory while the state is Unloaded.
-		TSharedPtr<const FExtensionData> Data;
+		TManagedPtr<const FExtensionData> Data;
 
 		enum class ELoadState : uint8
 		{
@@ -297,13 +297,13 @@ namespace UE::Mutable::Private
 		/** Constant image mip data is split in 2 sets: ConstantImageLODsPermanent constains data that is always loaded. 
 		* Index with FConstantResourceIndex::Index, when Streamable is 0.
 		*/
-		TArray<TSharedPtr<const FImage>> ConstantImageLODsPermanent;
+		TArray<TManagedPtr<const FImage>> ConstantImageLODsPermanent;
 
 		/** Constant image mip data is split in 2 sets: ConstantImageLODsStreamed constains data that is streamed in and out. 
 		* Index with FConstantResourceIndex::Index, when Streamable is 1.
 		* This part is empty for an unused Model, and shouldn't be serialised.
 		*/
-		TMap<uint32, TSharedPtr<const FImage>> ConstantImageLODsStreamed;
+		TMap<uint32, TManagedPtr<const FImage>> ConstantImageLODsStreamed;
 
 		/** Constant image mip chain indices: ranges in this array are defined in FImageLODRange and the indices here refer to ConstantImageLODs. */
 		TArray<FConstantResourceIndex> ConstantImageLODIndices;
@@ -320,13 +320,13 @@ namespace UE::Mutable::Private
 		/** Constant mesh data is split in 2 sets: ConstantMeshesPermanent constains data that is always loaded.
 		* Index with FConstantResourceIndex::Index, when Streamable is 0.
 		*/
-		TArray<TSharedPtr<const FMesh>> ConstantMeshesPermanent;
+		TArray<TManagedPtr<const FMesh>> ConstantMeshesPermanent;
 
 		/** Constant mesh data is split in 2 sets: ConstantMeshesStreamed constains data that is streamed in and out.
 		* Index with FConstantResourceIndex::Index, when Streamable is 1.
 		* This part is empty for an unused Model, and shouldn't be serialised.
 		*/
-		TMap<uint32, TSharedPtr<const FMesh>> ConstantMeshesStreamed;
+		TMap<uint32, TManagedPtr<const FMesh>> ConstantMeshesStreamed;
 
 		/** Constant FExtensionData */
 		TArray<FExtensionDataConstant> ConstantExtensionData;
@@ -341,7 +341,7 @@ namespace UE::Mutable::Private
 		TArray<TArray<uint64>> ConstantUInt64Lists;
 
         /** Constant layout data */
-		TArray<TSharedPtr<const FLayout>> ConstantLayouts;
+		TArray<TManagedPtr<const FLayout>> ConstantLayouts;
 
         /** Constant projectors */
 		TArray<FProjector> ConstantProjectors;
@@ -356,10 +356,10 @@ namespace UE::Mutable::Private
 		TArray<FRichCurve> ConstantCurves;
 
         /** Constant skeletons */
-		TArray<TSharedPtr<const FSkeleton>> ConstantSkeletons;
+		TArray<TManagedPtr<const FSkeleton>> ConstantSkeletons;
 
 		/** Constant Physics Bodies */
-		TArray<TSharedPtr<const FPhysicsBody>> ConstantPhysicsBodies;
+		TArray<TManagedPtr<const FPhysicsBody>> ConstantPhysicsBodies;
 
         /** FParameters of the model. */
         /** The value stored here is the default value. */
@@ -380,7 +380,7 @@ namespace UE::Mutable::Private
 
 
 		/** Constant Material Data */
-		TArray<TSharedPtr<const FMaterial>> ConstantMaterials;
+		TArray<TManagedPtr<const FMaterial>> ConstantMaterials;
 
     	TSharedPtr<const IExternalOperationProvider> ExternalOperationProvider;
     	
@@ -494,7 +494,7 @@ namespace UE::Mutable::Private
 			{
 			case uint32(ERomDataType::Image):
 			{
-				TSharedPtr<const FImage> Data;
+				TManagedPtr<const FImage> Data;
 				ConstantImageLODsStreamed.RemoveAndCopyValue(RomIndex, Data);
 				if (OutDataSize && Data)
 				{
@@ -504,7 +504,7 @@ namespace UE::Mutable::Private
 			}
 			case uint32(ERomDataType::Mesh):
 			{
-				TSharedPtr<const FMesh> Data;
+				TManagedPtr<const FMesh> Data;
 				ConstantMeshesStreamed.RemoveAndCopyValue(RomIndex, Data);
 				if (OutDataSize && Data)
 				{
@@ -519,7 +519,39 @@ namespace UE::Mutable::Private
 			}
 		}
 
-		FORCEINLINE void SetMeshRomValue(uint32 RomIndex, const TSharedPtr<FMesh>& Value)
+		FORCEINLINE SSIZE_T GetRomBytes(int32 RomIndex)
+		{
+			SSIZE_T Result = 0;
+			switch (Roms[RomIndex].ResourceType)
+			{
+			case uint32(ERomDataType::Image):
+			{
+				TManagedPtr<const FImage>* Data = ConstantImageLODsStreamed.Find(RomIndex);
+				if (Data)
+				{
+					Result = (*Data)->GetDataSize();
+				}
+				break;
+			}
+			case uint32(ERomDataType::Mesh):
+			{
+				TManagedPtr<const FMesh>* Data = ConstantMeshesStreamed.Find(RomIndex);
+				if (Data)
+				{
+					Result = (*Data)->GetDataSize();
+				}
+				break;
+			}
+			//TODO(Max): Make materials loadable and unloadable from a rom
+			default:
+				check(false);
+				break;
+			}
+
+			return Result;
+		}
+
+		FORCEINLINE void SetMeshRomValue(uint32 RomIndex, const TManagedPtr<FMesh>& Value)
 		{
 			check(Roms[RomIndex].ResourceType == uint32(ERomDataType::Mesh));
 			check(!ConstantMeshesStreamed.Contains(RomIndex));
@@ -528,7 +560,7 @@ namespace UE::Mutable::Private
 			ConstantMeshesStreamed.Add(RomIndex, Value);
 		}
 
-		FORCEINLINE void SetImageRomValue(uint32 RomIndex, const TSharedPtr<FImage>& Value)
+		FORCEINLINE void SetImageRomValue(uint32 RomIndex, const TManagedPtr<FImage>& Value)
 		{
 			check(Roms[RomIndex].ResourceType == uint32(ERomDataType::Image));
 			check(!ConstantImageLODsStreamed.Contains(RomIndex));
@@ -537,7 +569,7 @@ namespace UE::Mutable::Private
 			ConstantImageLODsStreamed.Add(RomIndex, Value);
 		}
 
-		OP::ADDRESS AddConstant(TSharedPtr<const FExtensionData> Data)
+		OP::ADDRESS AddConstant(TManagedPtr<const FExtensionData> Data)
 		{
 			// Ensure unique
 			for (int32 Index = 0; Index < ConstantExtensionData.Num(); Index++)
@@ -555,7 +587,7 @@ namespace UE::Mutable::Private
 			return ConstantExtensionData.Num() - 1;
 		}
 
-		OP::ADDRESS AddConstant( TSharedPtr<const FLayout> pLayout )
+		OP::ADDRESS AddConstant( TManagedPtr<const FLayout> pLayout )
         {
             // Ensure unique
             for (int32 i=0; i<ConstantLayouts.Num(); ++i)
@@ -571,7 +603,7 @@ namespace UE::Mutable::Private
             return index;
         }
 
-        OP::ADDRESS AddConstant( TSharedPtr<const FSkeleton> Skeleton )
+        OP::ADDRESS AddConstant(const TManagedPtr<const FSkeleton>& Skeleton )
         {
             // Ensure unique
             for ( int32 i=0; i<ConstantSkeletons.Num(); ++i)
@@ -590,22 +622,19 @@ namespace UE::Mutable::Private
             return index;
         }
 
-        OP::ADDRESS AddConstant(TSharedPtr<const FPhysicsBody> pPhysicsBody )
+        OP::ADDRESS AddConstant(const TManagedPtr<const FPhysicsBody>& PhysicsBody)
         {
             // Ensure unique
-            for (int32 i=0; i<ConstantPhysicsBodies.Num(); ++i)
+            for (int32 i = 0; i < ConstantPhysicsBodies.Num(); ++i)
             {
-                if ( ConstantPhysicsBodies[i]==pPhysicsBody
-                     ||
-                     *ConstantPhysicsBodies[i]==*pPhysicsBody
-                     )
+                if (ConstantPhysicsBodies[i] == PhysicsBody || *ConstantPhysicsBodies[i] == *PhysicsBody )
                 {
                     return (OP::ADDRESS)i;
                 }
             }
 
-            OP::ADDRESS index = OP::ADDRESS( ConstantPhysicsBodies.Num() );
-            ConstantPhysicsBodies.Add( pPhysicsBody );
+            OP::ADDRESS index = OP::ADDRESS(ConstantPhysicsBodies.Num());
+            ConstantPhysicsBodies.Add(PhysicsBody);
             return index;
         }
 
@@ -693,11 +722,11 @@ namespace UE::Mutable::Private
             return index;
         }
 
-		OP::ADDRESS AddConstant(const TSharedPtr<const FMaterial>& Material)
+		OP::ADDRESS AddConstant(const TManagedPtr<const FMaterial>& Material)
 		{
 			for (int32 Index=0; Index<ConstantMaterials.Num(); ++Index)
 			{
-				if (ConstantMaterials[Index]== Material || *ConstantMaterials[Index] == *Material)
+				if (ConstantMaterials[Index] == Material || *ConstantMaterials[Index] == *Material)
 				{
 					return static_cast<OP::ADDRESS>(Index);
 				}
@@ -708,9 +737,9 @@ namespace UE::Mutable::Private
 			return Index;
 		}
 
-		inline TSharedPtr<const FImage> GetImageLOD(FConstantResourceIndex Index) const
+		inline TManagedPtr<const FImage> GetImageLOD(FConstantResourceIndex Index) const
 		{
-			TSharedPtr<const FImage> Mip;
+			TManagedPtr<const FImage> Mip;
 			if (!Index.Streamable)
 			{
 				if (ConstantImageLODsPermanent.IsValidIndex(Index.Index))
@@ -720,7 +749,7 @@ namespace UE::Mutable::Private
 			}
 			else
 			{
-				const TSharedPtr<const FImage>* Found = ConstantImageLODsStreamed.Find(Index.Index);
+				const TManagedPtr<const FImage>* Found = ConstantImageLODsStreamed.Find(Index.Index);
 				if (Found)
 				{
 					Mip = *Found;
@@ -732,7 +761,7 @@ namespace UE::Mutable::Private
 
         /** Get a constant image, assuming at least some mips are loaded. The image constant will be composed with lodaded mips if necessary. */
 		template <typename CreateImageFunc>
-        void GetConstant(uint32 ConstantIndex, TSharedPtr<const FImage>& Result, int32 MipsToSkip, const CreateImageFunc& CreateImage) const
+        void GetConstant(uint32 ConstantIndex, TManagedPtr<const FImage>& Result, int32 MipsToSkip, const CreateImageFunc& CreateImage) const
         {
 			const int32 NumLODs = ConstantImages[ConstantIndex].LODCount;
 
@@ -745,7 +774,7 @@ namespace UE::Mutable::Private
 			{
 				MUTABLE_CPUPROFILER_SCOPE(ComposeConstantImage);
 
-				TSharedPtr<FImage> ResultImage = nullptr;
+				TManagedPtr<FImage> ResultImage = nullptr;
 
 				int32 LOD = ReallySkippedLODs;
 
@@ -753,7 +782,7 @@ namespace UE::Mutable::Private
 				for (; LOD <= FirstTailLOD; ++LOD)
 				{
 					MUTABLE_CPUPROFILER_SCOPE(Search_FisrtValid);
-					 TSharedPtr<const FImage> FirstValidLODImage = GetImageLOD(ConstantImageLODIndices[FirstLODIndexIndex + LOD]);
+					TManagedPtr<const FImage> FirstValidLODImage = GetImageLOD(ConstantImageLODIndices[FirstLODIndexIndex + LOD]);
 					
 					if (FirstValidLODImage)
 					{
@@ -788,7 +817,7 @@ namespace UE::Mutable::Private
 				for (; LOD < FirstTailLOD; ++LOD, ++DestLOD)
 				{
 					MUTABLE_CPUPROFILER_SCOPE(NonTailLODs);
-					TSharedPtr<const FImage> Image = GetImageLOD(ConstantImageLODIndices[FirstLODIndexIndex + LOD]);
+					TManagedPtr<const FImage> Image = GetImageLOD(ConstantImageLODIndices[FirstLODIndexIndex + LOD]);
 
 					if (bFormatNeedsDataResize)
 					{
@@ -807,7 +836,7 @@ namespace UE::Mutable::Private
 				if (LOD < NumLODs)
 				{
 					MUTABLE_CPUPROFILER_SCOPE(TailLODs);
-					TSharedPtr<const FImage> TailLODsImage = GetImageLOD(ConstantImageLODIndices[FirstLODIndexIndex + FirstTailLOD]);
+					TManagedPtr<const FImage> TailLODsImage = GetImageLOD(ConstantImageLODIndices[FirstLODIndexIndex + FirstTailLOD]);
 
 					if (!ensureAlwaysMsgf(TailLODsImage, TEXT("Missing tail image rom, setting image to black.")))
 					{
@@ -880,7 +909,7 @@ namespace UE::Mutable::Private
         void GetConstant(
 				uint32 MeshConstantIndex, 
 				int32 SkeletonConstantIndex, 
-				TSharedPtr<const FMesh>& OutMesh, 
+				TManagedPtr<const FMesh>& OutMesh, 
 				EMeshContentFlags FilterContentFlags, 
 				CreateMeshFunc&& CreateMesh) const
         {
@@ -888,8 +917,8 @@ namespace UE::Mutable::Private
 
 			FMeshContentRange MeshContentRange = ConstantMeshes[MeshConstantIndex];
 
-			TSharedPtr<const FMesh> EmptyMesh = nullptr;
-			auto GetMeshAtResourceIndex = [&](FConstantResourceIndex ResourceIndex) -> TSharedPtr<const FMesh>
+			TManagedPtr<const FMesh> EmptyMesh = nullptr;
+			auto GetMeshAtResourceIndex = [&](FConstantResourceIndex ResourceIndex) -> TManagedPtr<const FMesh>
 			{
 				MUTABLE_CPUPROFILER_SCOPE(GetConstant_Mesh_GetMesh)
 
@@ -902,14 +931,14 @@ namespace UE::Mutable::Private
 
 					if (!EmptyMesh)
 					{
-						EmptyMesh = MakeShared<FMesh>();
+						EmptyMesh = MakeManaged<FMesh>();
 					}
 					
 					return EmptyMesh;
 				}
 				else
 				{
-					const TSharedPtr<const FMesh>* Found = ConstantMeshesStreamed.Find(ResourceIndex.Index);
+					const TManagedPtr<const FMesh>* Found = ConstantMeshesStreamed.Find(ResourceIndex.Index);
 					if (Found)
 					{
 						return *Found;
@@ -917,17 +946,17 @@ namespace UE::Mutable::Private
 
 					if (!EmptyMesh)
 					{
-						EmptyMesh = MakeShared<FMesh>();
+						EmptyMesh = MakeManaged<FMesh>();
 					}
 
 					return EmptyMesh;
 				}
 			};
 
-			TSharedPtr<const FMesh> GeometryMesh = nullptr;
-			TSharedPtr<const FMesh> PoseMesh     = nullptr;
-			TSharedPtr<const FMesh> PhysicsMesh  = nullptr;
-			TSharedPtr<const FMesh> MetaDataMesh = nullptr;
+			TManagedPtr<const FMesh> GeometryMesh = nullptr;
+			TManagedPtr<const FMesh> PoseMesh     = nullptr;
+			TManagedPtr<const FMesh> PhysicsMesh  = nullptr;
+			TManagedPtr<const FMesh> MetaDataMesh = nullptr;
 
 			int32 MeshRomCurrentIndex = MeshContentRange.GetFirstIndex();
 			if (EnumHasAnyFlags(FilterContentFlags, EMeshContentFlags::GeometryData) &&
@@ -970,7 +999,7 @@ namespace UE::Mutable::Private
 			MeshBudgetReserve += PhysicsMesh  ? PhysicsMesh->GetDataSize()  : 0;
 			MeshBudgetReserve += MetaDataMesh ? MetaDataMesh->GetDataSize() : 0;
 
-			TSharedPtr<FMesh> Result = nullptr;
+			TManagedPtr<FMesh> Result = nullptr;
 			if (GeometryMesh)
 			{
 				MUTABLE_CPUPROFILER_SCOPE(GetConstant_Mesh_Geoemtry)
@@ -1074,7 +1103,7 @@ namespace UE::Mutable::Private
 			OutMesh = Result;
 		}
 
-		void GetExtensionDataConstant(int32 ConstantIndex, TSharedPtr<const FExtensionData>& Result) const
+		void GetExtensionDataConstant(int32 ConstantIndex, TManagedPtr<const FExtensionData>& Result) const
 		{
 			const FExtensionDataConstant& Constant = ConstantExtensionData[ConstantIndex];
 
diff --git a/Engine/Plugins/Mutable/Source/MutableRuntime/Internal/MuR/OpImageGrow.h b/Engine/Plugins/Mutable/Source/MutableRuntime/Internal/MuR/OpImageGrow.h
--- a/Engine/Plugins/Mutable/Source/MutableRuntime/Internal/MuR/OpImageGrow.h
+++ b/Engine/Plugins/Mutable/Source/MutableRuntime/Internal/MuR/OpImageGrow.h
@@ -23,7 +23,7 @@ namespace UE::Mutable::Private
 		int sizeY = pImage->GetSizeY();
 		check( sizeX>0 && sizeY>0 );
 
-		TSharedPtr<const FImage> pSource = pImage->Clone();
+		TManagedPtr<const FImage> pSource = pImage->Clone();
 
 		int rowSize = sizeX*GetImageFormatData( pImage->GetFormat() ).BytesPerBlock;
 
diff --git a/Engine/Plugins/Mutable/Source/MutableRuntime/Internal/MuR/OpMeshMerge.h b/Engine/Plugins/Mutable/Source/MutableRuntime/Internal/MuR/OpMeshMerge.h
--- a/Engine/Plugins/Mutable/Source/MutableRuntime/Internal/MuR/OpMeshMerge.h
+++ b/Engine/Plugins/Mutable/Source/MutableRuntime/Internal/MuR/OpMeshMerge.h
@@ -18,8 +18,8 @@ namespace UE::Mutable::Private
 
 	struct FMeshMergeScratchMeshes
 	{
-		TSharedPtr<FMesh> FirstReformat;
-		TSharedPtr<FMesh> SecondReformat;
+		TManagedPtr<FMesh> FirstReformat;
+		TManagedPtr<FMesh> SecondReformat;
 	};
 
 	inline void FixMorphIndices(FMesh& Result, int32 VertexIndexOffset)
@@ -586,7 +586,7 @@ namespace UE::Mutable::Private
 			Result->Layouts.SetNum(ResultLayoutCount);
 			for (int32 LayoutIndex = 0; LayoutIndex < ResultLayoutCount; ++LayoutIndex)
 			{
-				TSharedPtr<FLayout> pR;
+				TManagedPtr<FLayout> pR;
 				
 				if (LayoutIndex < pFirst->Layouts.Num())
 				{
@@ -636,15 +636,15 @@ namespace UE::Mutable::Private
 		if (bMergeSkeletons)
 		{
 			MUTABLE_CPUPROFILER_SCOPE(MergeSkeleton);
-			TSharedPtr<FSkeleton> ResultSkeleton;
+			TManagedPtr<FSkeleton> ResultSkeleton;
 
-			TSharedPtr<const FSkeleton> FirstSkeleton = pFirst->GetSkeleton();
-			TSharedPtr<const FSkeleton> SecondSkeleton = pSecond->GetSkeleton();
+			TManagedPtr<const FSkeleton> FirstSkeleton = pFirst->GetSkeleton();
+			TManagedPtr<const FSkeleton> SecondSkeleton = pSecond->GetSkeleton();
 
 			const int32 NumBonesFirst = FirstSkeleton ? FirstSkeleton->GetBoneCount() : 0;
 			const int32 NumBonesSecond = SecondSkeleton ? SecondSkeleton->GetBoneCount() : 0;
 
-			ResultSkeleton = FirstSkeleton ? FirstSkeleton->Clone() : MakeShared<FSkeleton>();
+			ResultSkeleton = FirstSkeleton ? FirstSkeleton->Clone() : MakeManaged<FSkeleton>();
 			Result->SetSkeleton(ResultSkeleton);
 
 			TArray<uint16> SecondToResultBoneIndices;
@@ -918,8 +918,8 @@ namespace UE::Mutable::Private
 				return bModified;
 			};
 
-			TTuple<TSharedPtr<const FPhysicsBody>, bool> SharedResultPhysicsBody = Invoke([&]()
-				-> TTuple<TSharedPtr<const FPhysicsBody>, bool>
+			TTuple<TManagedPtr<const FPhysicsBody>, bool> SharedResultPhysicsBody = Invoke([&]()
+				-> TTuple<TManagedPtr<const FPhysicsBody>, bool>
 			{
 				if (pFirst->GetPhysicsBody() == pSecond->GetPhysicsBody())
 				{
@@ -948,7 +948,7 @@ namespace UE::Mutable::Private
 			{
 				check(pFirst->GetPhysicsBody() && pSecond->GetPhysicsBody());
 
-				TSharedPtr<FPhysicsBody> MergedResultPhysicsBody = pFirst->GetPhysicsBody()->Clone();
+				TManagedPtr<FPhysicsBody> MergedResultPhysicsBody = pFirst->GetPhysicsBody()->Clone();
 
 				MergedResultPhysicsBody->bBodiesModified =
 					AppendPhysicsBodiesUnique(*MergedResultPhysicsBody, *pSecond->GetPhysicsBody()) ||
@@ -969,8 +969,8 @@ namespace UE::Mutable::Private
 			{
 				const int32 CustomIdToMerge = pSecond->AdditionalPhysicsBodies[Index]->CustomId;
 
-				const TSharedPtr<const FPhysicsBody>* Found = pFirst->AdditionalPhysicsBodies.FindByPredicate(
-					[CustomIdToMerge](const TSharedPtr<const FPhysicsBody>& Body) { return CustomIdToMerge == Body->CustomId; });
+				const TManagedPtr<const FPhysicsBody>* Found = pFirst->AdditionalPhysicsBodies.FindByPredicate(
+					[CustomIdToMerge](const TManagedPtr<const FPhysicsBody>& Body) { return CustomIdToMerge == Body->CustomId; });
 
 				// TODO: current usages do not expect collisions, but same Id collision with bodies modified in differnet ways
 				// may need to be contemplated at some point.
@@ -1433,7 +1433,7 @@ namespace UE::Mutable::Private
     //---------------------------------------------------------------------------------------------
     //! Return null if there is no need to remap the mesh
     //---------------------------------------------------------------------------------------------
-    inline void MeshRemapSkeleton(FMesh* Result, const FMesh* SourceMesh, TSharedPtr<const FSkeleton> Skeleton, bool& bOutSuccess)
+    inline void MeshRemapSkeleton(FMesh* Result, const FMesh* SourceMesh, TManagedPtr<const FSkeleton> Skeleton, bool& bOutSuccess)
     {
 		bOutSuccess = true;
 
diff --git a/Engine/Plugins/Mutable/Source/MutableRuntime/Internal/MuR/Parameters.h b/Engine/Plugins/Mutable/Source/MutableRuntime/Internal/MuR/Parameters.h
--- a/Engine/Plugins/Mutable/Source/MutableRuntime/Internal/MuR/Parameters.h
+++ b/Engine/Plugins/Mutable/Source/MutableRuntime/Internal/MuR/Parameters.h
@@ -3,6 +3,7 @@
 #pragma once
 
 #include "HAL/Platform.h"
+#include "MuR/ManagedPointer.h"
 #include "MuR/RefCounted.h"
 #include "MuR/Types.h"
 #include "Templates/SharedPointer.h"
@@ -376,13 +377,13 @@ namespace UE::Mutable::Private
 		* @param Index Index of the parameter from 0 to GetCount()-1
 		* @param RangeIndex Only for multidimensional parameters: relevant position to set in the ranges
 		* @return The Instanced Struct interface specified when setting the Instanced Struct value (\see SetInstancedStructValue) */
-		UE_API TSharedPtr<const FInstancedStruct> GetInstancedStructValue(int32 Index, const FRangeIndex* RangeIndex = nullptr) const;
+		UE_API TManagedPtr<const FInstancedStruct> GetInstancedStructValue(int32 Index, const FRangeIndex* RangeIndex = nullptr) const;
 
 		//! If the parameter is of the Instanced Struct type, set its value.
 		//! \param Index Index of the parameter from 0 to GetCount()-1
 		//! \param Value Instanced Struct interface to set as value.
 		//! \param RangeIndex Only for multidimensional parameters: relevant position to set in the ranges
-		UE_API void SetInstancedStructValue(int32 Index, const TSharedPtr<const FInstancedStruct>& Value, const FRangeIndex* RangeIndex = nullptr);
+		UE_API void SetInstancedStructValue(int32 Index, const TManagedPtr<const FInstancedStruct>& Value, const FRangeIndex* RangeIndex = nullptr);
 
         //! Utility method to compare the values of a specific parameter with the values of another
         //! FParameters object. It returns false if type or values are different.
diff --git a/Engine/Plugins/Mutable/Source/MutableRuntime/Internal/MuR/ParametersPrivate.h b/Engine/Plugins/Mutable/Source/MutableRuntime/Internal/MuR/ParametersPrivate.h
--- a/Engine/Plugins/Mutable/Source/MutableRuntime/Internal/MuR/ParametersPrivate.h
+++ b/Engine/Plugins/Mutable/Source/MutableRuntime/Internal/MuR/ParametersPrivate.h
@@ -263,7 +263,7 @@ namespace Private
 	using FParamStringType = Private::TIndirectObject<FString>;
 	using FParamMatrixType = Private::TIndirectObject<FMatrix44f>;
 	using FParamMaterialType = TStrongObjectPtr<UMaterialInterface>;
-	using FParamInstancedStructType = TSharedPtr<const FInstancedStruct>; // TSharedPtr always valid and FInstancedStruct always initialized to its type.
+	using FParamInstancedStructType = TManagedPtr<const FInstancedStruct>; // TSharedPtr always valid and FInstancedStruct always initialized to its type.
 	
 	using FParameterValue = TVariant<
 		FParamBoolType,
diff --git a/Engine/Plugins/Mutable/Source/MutableRuntime/Internal/MuR/PhysicsBody.h b/Engine/Plugins/Mutable/Source/MutableRuntime/Internal/MuR/PhysicsBody.h
--- a/Engine/Plugins/Mutable/Source/MutableRuntime/Internal/MuR/PhysicsBody.h
+++ b/Engine/Plugins/Mutable/Source/MutableRuntime/Internal/MuR/PhysicsBody.h
@@ -151,11 +151,11 @@ namespace UE::Mutable::Private
 	{
 	public:
 		
-		UE_API TSharedPtr<FPhysicsBody> Clone() const;
+		UE_API TManagedPtr<FPhysicsBody> Clone() const;
 
 		//! Serialisation
         static UE_API void Serialise( const FPhysicsBody*, FOutputArchive& );
-		static UE_API TSharedPtr<FPhysicsBody> StaticUnserialise( FInputArchive& );
+		static UE_API TManagedPtr<FPhysicsBody> StaticUnserialise( FInputArchive& );
 
         //bool operator==( const PhysicsBody& Other ) const;
 
diff --git a/Engine/Plugins/Mutable/Source/MutableRuntime/Internal/MuR/Serialisation.h b/Engine/Plugins/Mutable/Source/MutableRuntime/Internal/MuR/Serialisation.h
--- a/Engine/Plugins/Mutable/Source/MutableRuntime/Internal/MuR/Serialisation.h
+++ b/Engine/Plugins/Mutable/Source/MutableRuntime/Internal/MuR/Serialisation.h
@@ -3,6 +3,7 @@
 #pragma once
 
 #include "HAL/Platform.h"
+#include "MuR/ManagedPointer.h"
 #include "MuR/Ptr.h"
 #include "MuR/RefCounted.h"
 #include "MuR/MutableMath.h"
@@ -175,7 +176,7 @@ namespace UE::Mutable::Private
 		FInputStream* Stream = nullptr;
 
 		/** Already read pointers. */
-		TArray< TSharedPtr<void> > History;
+		TArray< TManagedPtr<void> > History;
 
     };
 
@@ -388,8 +389,8 @@ namespace UE::Mutable::Private
     	// Do Nothing. UObjects values are saved in the CO.	
 	}
 
-	MUTABLERUNTIME_API void operator<<(FOutputArchive& Arch, const TSharedPtr<const FInstancedStruct>& V);
-	MUTABLERUNTIME_API void operator>>(FInputArchive& Arch, TSharedPtr<const FInstancedStruct>& V);
+	MUTABLERUNTIME_API void operator<<(FOutputArchive& Arch, const TManagedPtr<const FInstancedStruct>& V);
+	MUTABLERUNTIME_API void operator>>(FInputArchive& Arch, TManagedPtr<const FInstancedStruct>& V);
 }
 
 
diff --git a/Engine/Plugins/Mutable/Source/MutableRuntime/Internal/MuR/SerialisationPrivate.h b/Engine/Plugins/Mutable/Source/MutableRuntime/Internal/MuR/SerialisationPrivate.h
--- a/Engine/Plugins/Mutable/Source/MutableRuntime/Internal/MuR/SerialisationPrivate.h
+++ b/Engine/Plugins/Mutable/Source/MutableRuntime/Internal/MuR/SerialisationPrivate.h
@@ -240,9 +240,10 @@ namespace UE::Mutable::Private
 	MUTABLE_IMPLEMENT_POD_VECTOR_SERIALISABLE(FRichCurveKey);
 
 
-	//---------------------------------------------------------------------------------------------
-	template< typename T >
-	inline void operator<< (FOutputArchive& Arch, const TSharedPtr<const T>& Ptr)
+	template < 
+		typename T 
+	>
+	inline void operator<< (FOutputArchive& Arch, const TManagedPtr<const T>& Ptr)
 	{
 		if (!Ptr)
 		{
@@ -266,8 +267,10 @@ namespace UE::Mutable::Private
 		}
 	}
 
-	template< typename T >
-	inline void operator>> (FInputArchive& Arch, TSharedPtr<const T>& Ptr)
+	template< 
+		typename T
+	>
+	inline void operator>> (FInputArchive& Arch, TManagedPtr<const T>& Ptr)
 	{
 		int32 Id;
 		Arch >> Id;
@@ -280,7 +283,7 @@ namespace UE::Mutable::Private
 		{
 			if (Id < Arch.History.Num())
 			{
-				Ptr = StaticCastSharedPtr<T>(Arch.History[Id]);
+				Ptr = StaticCastManagedPtr<T>(Arch.History[Id]);
 
 				// If the pointer was null it means the position in history is used, but not set yet
 				// option 1: we have a smart pointer loop which is very bad.
@@ -295,14 +298,13 @@ namespace UE::Mutable::Private
 				//check( Id == Arch.>History.Num() );
 				Arch.History.SetNum(Id + 1, EAllowShrinking::No);
 
-				TSharedPtr<T> Temp = T::StaticUnserialise(Arch);
+				TManagedPtr<T> Temp = T::StaticUnserialise(Arch);
 				Ptr = Temp;
 				Arch.History[Id] = Temp;
 			}
 		}
 	}
 
-
 	//---------------------------------------------------------------------------------------------
 	template< typename T >
 	inline void operator<< ( FOutputArchive& arch, const Ptr<T>& p )
diff --git a/Engine/Plugins/Mutable/Source/MutableRuntime/Internal/MuR/Skeleton.h b/Engine/Plugins/Mutable/Source/MutableRuntime/Internal/MuR/Skeleton.h
--- a/Engine/Plugins/Mutable/Source/MutableRuntime/Internal/MuR/Skeleton.h
+++ b/Engine/Plugins/Mutable/Source/MutableRuntime/Internal/MuR/Skeleton.h
@@ -58,11 +58,11 @@ namespace UE::Mutable::Private
 		//-----------------------------------------------------------------------------------------
 
         //! Deep clone this skeleton.
-        UE_API TSharedPtr<FSkeleton> Clone() const;
+        UE_API TManagedPtr<FSkeleton> Clone() const;
 
 		//! Serialisation
         static UE_API void Serialise( const FSkeleton*, FOutputArchive& );
-        static UE_API TSharedPtr<FSkeleton> StaticUnserialise( FInputArchive& );
+        static UE_API TManagedPtr<FSkeleton> StaticUnserialise( FInputArchive& );
 
 
 		//-----------------------------------------------------------------------------------------
diff --git a/Engine/Plugins/Mutable/Source/MutableRuntime/Internal/MuR/System.h b/Engine/Plugins/Mutable/Source/MutableRuntime/Internal/MuR/System.h
--- a/Engine/Plugins/Mutable/Source/MutableRuntime/Internal/MuR/System.h
+++ b/Engine/Plugins/Mutable/Source/MutableRuntime/Internal/MuR/System.h
@@ -82,13 +82,13 @@ namespace UE::Mutable::Private
         virtual ~FExternalResourceProvider() = default;
 
         /** Returns the completion event and a cleanup function that must be called once event is completed. */
-		virtual TTuple<UE::Tasks::FTask, TFunction<void()>> GetImageAsync(UTexture* Texture, uint8 MipmapsToSkip, TFunction<void(TSharedPtr<FImage>)>& ResultCallback) = 0;
-		virtual TTuple<UE::Tasks::FTask, TFunction<void()>> GetReferencedImageAsync(int32 Id, uint8 MipmapsToSkip, TFunction<void(TSharedPtr<FImage>)>& ResultCallback) { check(false); return {}; }
+		virtual TTuple<UE::Tasks::FTask, TFunction<void()>> GetImageAsync(UTexture* Texture, uint8 MipmapsToSkip, TFunction<void(UE::Mutable::Private::TManagedPtr<FImage>)>& ResultCallback) = 0;
+		virtual TTuple<UE::Tasks::FTask, TFunction<void()>> GetReferencedImageAsync(int32 Id, uint8 MipmapsToSkip, TFunction<void(UE::Mutable::Private::TManagedPtr<FImage>)>& ResultCallback) { check(false); return {}; }
 
 		virtual FExtendedImageDesc GetImageDesc(UTexture* Texture) = 0;
 
 		/** Returns the completion event and a cleanup function that must be called once event is completed. */
-		virtual TTuple<UE::Tasks::FTask, TFunction<void()>> GetMeshAsync(USkeletalMesh* SkeleltalMesh, int32 InLODIndex, int32 InSectionIndex, TFunction<void(TSharedPtr<FMesh>)>& ResultCallback) = 0;
+		virtual TTuple<UE::Tasks::FTask, TFunction<void()>> GetMeshAsync(USkeletalMesh* SkeleltalMesh, int32 InLODIndex, int32 InSectionIndex, TFunction<void(UE::Mutable::Private::TManagedPtr<FMesh>)>& ResultCallback) = 0;
     };
 
 
@@ -154,7 +154,7 @@ namespace UE::Mutable::Private
         //! the same instanceID parameter.
 		//!
 		//! Very fast. Can be called in the Game Thread.
-		UE_API TSharedPtr<const FInstance> BeginUpdate_GameThread(FInstance::FID InstanceID,
+		UE_API TManagedPtr<const FInstance> BeginUpdate_GameThread(FInstance::FID InstanceID,
 			const TSharedPtr<const FParameters>& Params,
 			const TSharedPtr<FMeshIdRegistry>& InMeshIdRegistry,
 			const TSharedPtr<FImageIdRegistry>& InImageIdRegistry,
@@ -164,7 +164,7 @@ namespace UE::Mutable::Private
 			uint32 LodMask);
 
 		/** Potentially expensive. Must be called in the Mutable Thread. */
-		UE_API TSharedPtr<const FInstance> BeginUpdate_MutableThread(FInstance::FID InInstanceID,
+		UE_API TManagedPtr<const FInstance> BeginUpdate_MutableThread(FInstance::FID InInstanceID,
 			const TSharedPtr<const FParameters>& InParams,
 			const TSharedPtr<FMeshIdRegistry>& InMeshIdRegistry,
 			const TSharedPtr<FImageIdRegistry>& InImageIdRegistry,
@@ -181,12 +181,12 @@ namespace UE::Mutable::Private
 		//! \param MipsToSkip Number of mips to skip compared from the full image.
 		//! If 0, all mip levels will be generated. If more levels than possible to discard are specified, 
 		//! the image will still contain a minimum number of mips specified at model compile time.
-		UE_API UE::Tasks::TTask<TSharedPtr<const FImage>> GetImage(UE::Tasks::FTask Prerequisite, FInstance::FID InstanceID, const FImageId& ImageId, int32 MipsToSkip = 0);
+		UE_API UE::Tasks::TTask<TManagedPtr<const FImage>> GetImage(UE::Tasks::FTask Prerequisite, FInstance::FID InstanceID, const FImageId& ImageId, int32 MipsToSkip = 0);
 
 		//! Only valid between BeginUpdate and EndUpdate
-		UE_API UE::Tasks::TTask<TSharedPtr<const FMesh>> GetMesh(UE::Tasks::FTask Prerequisite, FInstance::FID InstanceID, const FMeshId& MeshId, EMeshContentFlags MeshContentFilter = EMeshContentFlags::AllFlags);
+		UE_API UE::Tasks::TTask<TManagedPtr<const FMesh>> GetMesh(UE::Tasks::FTask Prerequisite, FInstance::FID InstanceID, const FMeshId& MeshId, EMeshContentFlags MeshContentFilter = EMeshContentFlags::AllFlags);
 		
-		UE_API UE::Tasks::TTask<TSharedPtr<const FSkeletalMesh>> GetSkeletalMesh(UE::Tasks::FTask Prerequisite, FInstance::FID InstanceID, const FSkeletalMeshId& SkeletalMeshId, EMeshContentFlags MeshContentFilter);
+		UE_API UE::Tasks::TTask<TManagedPtr<const FSkeletalMesh>> GetSkeletalMesh(UE::Tasks::FTask Prerequisite, FInstance::FID InstanceID, const FSkeletalMeshId& SkeletalMeshId, EMeshContentFlags MeshContentFilter);
 
 		//! Only valid between BeginUpdate and EndUpdate
 		//! Calculate the description of an image, without generating it.
@@ -196,12 +196,12 @@ namespace UE::Mutable::Private
 		//! \param MipsToSkip Number of mips to skip compared from the full image.
 		//! If 0, all mip levels will be generated. If more levels than possible to discard are specified, 
 		//! the image will still contain a minimum number of mips specified at model compile time.
-		UE_API TSharedPtr<const FImage> GetImageInline(FInstance::FID InstanceID, const FImageId& ImageId, int32 MipsToSkip = 0);
+		UE_API TManagedPtr<const FImage> GetImageInline(FInstance::FID InstanceID, const FImageId& ImageId, int32 MipsToSkip = 0);
 
         //! Only valid between BeginUpdate and EndUpdate
-		UE_API TSharedPtr<const FMesh> GetMeshInline(FInstance::FID InstanceID, const FMeshId& MeshId, EMeshContentFlags MeshContentFilter = EMeshContentFlags::AllFlags);
+		UE_API TManagedPtr<const FMesh> GetMeshInline(FInstance::FID InstanceID, const FMeshId& MeshId, EMeshContentFlags MeshContentFilter = EMeshContentFlags::AllFlags);
 		
-		UE_API UE::Tasks::TTask<TSharedPtr<const FMaterial>> GetMaterial(UE::Tasks::FTask Prerequisite, FInstance::FID InstanceID, const FMaterialId& MaterialId);
+		UE_API UE::Tasks::TTask<TManagedPtr<const FMaterial>> GetMaterial(UE::Tasks::FTask Prerequisite, FInstance::FID InstanceID, const FMaterialId& MaterialId);
 		
         //! Invalidate and free the last Instance data returned by a call to BeginUpdate with
         //! the same instance index. After a call to this method, that Instance cannot be used any
