[Mutable] Memory managment rework.

diff --git a/Engine/Plugins/Mutable/Source/MutableRuntime/Internal/MuR/SystemPrivate.h b/Engine/Plugins/Mutable/Source/MutableRuntime/Internal/MuR/SystemPrivate.h
--- a/Engine/Plugins/Mutable/Source/MutableRuntime/Internal/MuR/SystemPrivate.h
+++ b/Engine/Plugins/Mutable/Source/MutableRuntime/Internal/MuR/SystemPrivate.h
@@ -3,6 +3,8 @@
 #pragma once
 
 #include "SkeletalMesh.h"
+#include "MuR/SystemPrivateTypes.h"
+
 #include "MuR/MemoryCounters.h"
 #include "MuR/System.h"
 
@@ -19,15 +21,13 @@
 
 #include "Templates/UnrealTypeTraits.h"
 
+#include "MuR/ManagedPointer.h"
+
 #if !(UE_BUILD_SHIPPING || UE_BUILD_TEST)
 #include "HAL/Thread.h"
 #include "HAL/PlatformTLS.h"
 #endif
 
-#ifndef MUTABLE_DEBUG_CODERUNNER_TASK_SCHEDULE_CALLSTACK
-	#define MUTABLE_DEBUG_CODERUNNER_TASK_SCHEDULE_CALLSTACK 0	
-#endif
-
 
 namespace UE::Mutable::Private
 {
@@ -87,137 +87,6 @@ namespace UE::Mutable::Private
 		}
 	};
 
-
-	/** This structure stores the data about an ongoing mutable operation that needs to be executed. */
-	struct FScheduledOp
-	{
-		inline FScheduledOp()
-		{
-			Stage = 0;
-			bForceLazy = false;
-			Type = EType::Full;
-		}
-
-		inline FScheduledOp(OP::ADDRESS InAt, const FScheduledOp& InOpTemplate, uint8 InStage = 0, uint32 InCustomState = 0)
-		{
-			check(InStage <= 63);
-			At = InAt;
-			ExecutionOptions = InOpTemplate.ExecutionOptions;
-			ExecutionIndex = InOpTemplate.ExecutionIndex;
-			Stage = InStage;
-			CustomState = InCustomState;
-			bForceLazy = InOpTemplate.bForceLazy;
-			Type = InOpTemplate.Type;
-		}
-
-		static inline FScheduledOp FromOpAndOptions(OP::ADDRESS InAt, const FScheduledOp& InOpTemplate, uint8 InExecutionOptions)
-		{
-			FScheduledOp r;
-			r.At = InAt;
-			r.ExecutionOptions = InExecutionOptions;
-			r.ExecutionIndex = InOpTemplate.ExecutionIndex;
-			r.Stage = 0;
-			r.CustomState = InOpTemplate.CustomState;
-			r.bForceLazy = InOpTemplate.bForceLazy;
-			r.Type = InOpTemplate.Type;
-			return r;
-		}
-
-		//! Address of the operation
-		OP::ADDRESS At = 0;
-
-		//! Additional custom state data that the operation can store. This is usually used to pass information
-		//! between execution stages of an operation.
-		uint32 CustomState = 0;
-
-		//! Index of the operation execution: This is used for iteration of different ranges.
-		//! It is an index into the CodeRunner::GetMemory()::m_rangeIndex vector.
-		//! executionIndex 0 is always used for empty ExecutionIndex, which is the most common
-		//! one.
-		uint16 ExecutionIndex = 0;
-
-		//! Additional execution options. Set externally to this op, it usually alters the result.
-		//! For example, this is used to keep track of the mipmaps to skip in image operations.
-		uint8 ExecutionOptions = 0;
-
-		//! Internal stage of the operation.
-		//! Stage 0 is usually scheduling of children, and 1 is execution. Some instructions
-		//! may have more steges to schedule children that are optional for execution, etc.
-		uint8 Stage : 6;
-
-		uint8 bForceLazy : 1;
-
-		//! Type of calculation we are requesting for this operation.
-		enum class EType : uint8
-		{
-			//! Execute the operation to calculate the full result
-			Full,
-
-			//! Execute the operation to obtain the descriptor of an image or mesh.
-			ImageDesc,
-		};
-		EType Type : 1;
-
-#if MUTABLE_DEBUG_CODERUNNER_TASK_SCHEDULE_CALLSTACK
-		static constexpr uint64 CallstackMaxDepth = 16;
-		uint64 StackDepth = 0; 
-		uint64 ScheduleCallstack[CallstackMaxDepth];
-#endif
-	};
-	
-
-	inline uint32 GetTypeHash(const FScheduledOp& Op)
-	{
-		return HashCombineFast(::GetTypeHash(Op.At), HashCombineFast((uint32)Op.Stage), (uint32)Op.ExecutionIndex);
-	}
-	
-	
-	/** A cache address is the operation plus the context of execution(iteration indices, etc...). */
-	struct FCacheAddress
-	{
-		/** The meaning of all these fields is the same than the FScheduledOp struct. */
-		OP::ADDRESS At = 0;
-		uint16 ExecutionIndex = 0;
-		uint8 ExecutionOptions = 0;
-		FScheduledOp::EType Type = FScheduledOp::EType::Full;
-
-		FCacheAddress() 
-		{
-		}
-
-		FCacheAddress(OP::ADDRESS InAt, uint16 InExecutionIndex, uint8 InExecutionOptions, FScheduledOp::EType InType = FScheduledOp::EType::Full)
-			: At(InAt), ExecutionIndex(InExecutionIndex), ExecutionOptions(InExecutionOptions), Type(InType)
-		{
-		}
-
-		FCacheAddress(OP::ADDRESS InAt, const FScheduledOp& Item)
-			: At(InAt), ExecutionIndex(Item.ExecutionIndex), ExecutionOptions(Item.ExecutionOptions), Type(Item.Type)
-		{
-		}
-		
-		FCacheAddress(const FScheduledOp& Item)
-			: At(Item.At), ExecutionIndex(Item.ExecutionIndex), ExecutionOptions(Item.ExecutionOptions), Type(Item.Type)
-		{
-		}
-	};
-
-	inline uint32 GetTypeHash(const FCacheAddress& Address)
-	{
-		uint32 Hash = Address.At;
-
-		Hash = HashCombineFast(Hash, static_cast<uint32>(Address.ExecutionIndex));
-		Hash = HashCombineFast(Hash, static_cast<uint32>(Address.ExecutionOptions));
-		Hash = HashCombineFast(Hash, static_cast<uint32>(Address.Type));
-
-		return Hash;
-	}
-	
-	inline uint32 GetTypeHash(const Ptr<const FResource>& V)
-	{
-		return ::GetTypeHash(V.get());
-	}
-
-
 	//! Container that stores data per executable code operation (indexed by address and execution
 	//! index).
 	template<class DATA>
@@ -488,31 +357,30 @@ namespace UE::Mutable::Private
 		*/
 		CodeContainer<FOpExecutionData> OpExecutionData;
 
-		template<class ResourceType>
-		struct TResourceResult
-		{
-			static_assert(TIsDerivedFrom<ResourceType, FResource>::Value);
-
-			FCacheAddress OpAddress;
-			TSharedPtr<const ResourceType> Value = nullptr;
-		};
-		static_assert(sizeof(TResourceResult<FImage>) == 24);
-		static_assert(sizeof(TResourceResult<FMesh>) == 24);
 
 		/** */
 
+		TMemoryTrackedArray<TManagedWeakPtr<const FLayout>> LayoutResults;
+		TMemoryTrackedArray<TManagedWeakPtr<const FInstance>> InstanceResults;
+		TMemoryTrackedArray<TManagedWeakPtr<const String>> StringResults;
+		TMemoryTrackedArray<TManagedWeakPtr<const FExtensionData>> ExtensionDataResults;
+		TMemoryTrackedArray<TManagedWeakPtr<const FInstancedStruct>> InstancedStructResults;
+		TMemoryTrackedArray<TManagedWeakPtr<const FMaterial>> MaterialResults;
+		TMemoryTrackedArray<TManagedWeakPtr<const FSkeletalMesh>> SkeletalMeshResults;
+
+		TMemoryTrackedArray<TManagedWeakPtr<const FImage>> ImageResults;
+		TMemoryTrackedArray<TManagedWeakPtr<const FMesh>> MeshResults;
+		
+		/** Storage for strong references. Size of Strong Results is the same as the weak. */
+		TMemoryTrackedArray<TManagedPtr<const FImage>> StrongImageResults;
+		TMemoryTrackedArray<TManagedPtr<const FMesh>> StrongMeshResults;
+
+		//TMemoryTrackedArray<TManagedPtr<const FResource>> ResourceResults;
+		//TMemoryTrackedArray<TManagedPtr<const FResource>> WeakResourceResults;
+
 		TMemoryTrackedArray<FVector4f> ColorResults;
-		TMemoryTrackedArray<TResourceResult<FImage>> ImageResults;
-		TMemoryTrackedArray<TResourceResult<FMesh>> MeshResults;
-		TMemoryTrackedArray<TSharedPtr<const FSkeletalMesh>> SkeletalMeshResults;
-		TMemoryTrackedArray<TSharedPtr<const FLayout>> LayoutResults;
-		TMemoryTrackedArray<TSharedPtr<const FInstance>> InstanceResults;
 		TMemoryTrackedArray<FProjector> ProjectorResults;
-		TMemoryTrackedArray<TSharedPtr<const String>> StringResults;
-		TMemoryTrackedArray<TSharedPtr<const FExtensionData>> ExtensionDataResults;
-		TMemoryTrackedArray<TSharedPtr<const FInstancedStruct>> InstancedStructResults;
 		TMemoryTrackedArray<FMatrix44f> MatrixResults;
-		TMemoryTrackedArray<TSharedPtr<const FMaterial>> MaterialResults;
 
 		/** */
 		inline const ExecutionIndex& GetRangeIndex(uint32_t i)
@@ -562,92 +430,74 @@ namespace UE::Mutable::Private
 			if (ColorResults.IsEmpty())
 			{
 				// Insert dafault/null values
-				ColorResults.Add(FVector4f());
-				ImageResults.Emplace();
-				LayoutResults.Add(nullptr);
-				MeshResults.Emplace();
+
+				LayoutResults.Emplace();
+				InstanceResults.Emplace();
+				StringResults.Emplace();
+				ExtensionDataResults.Emplace();
+				MaterialResults.Emplace();
 				SkeletalMeshResults.Emplace();
-				InstanceResults.Add(nullptr);
-				ProjectorResults.Add(FProjector());
-				StringResults.Add(nullptr);
-				MatrixResults.Add(FMatrix44f());
-				ExtensionDataResults.Add(nullptr);
-				MaterialResults.Add(nullptr);
 				InstancedStructResults.Add({});
-			}
-		}
 
-		void SetUnused(FOpExecutionData& Data)
-		{
-			// Only clear datatypes that have results that use relevant amounts of memory			
-			uint32 DataTypeIndex = Data.DataTypeIndex;
-			Data.IsValueValid = false;
+				ImageResults.Emplace();
+				MeshResults.Emplace();
+				
+				StrongImageResults.Emplace();
+				StrongMeshResults.Emplace();
+	
+				ColorResults.Emplace();
+				MatrixResults.Emplace();
+				ProjectorResults.Emplace();
+			}
 
-			if (DataTypeIndex)
+			for (TManagedPtr<const FImage>& StrongImagePtr : StrongImageResults)
 			{
-				check(uint8(Data.DataType) < uint8(EDataType::Count));
-				switch (Data.DataType)
-				{
-				case EDataType::Image:					
-					ImageResults[DataTypeIndex].Value = nullptr;
-					break;
-
-				case EDataType::Mesh:
-					MeshResults[DataTypeIndex].Value = nullptr;
-					break;
-
-				case EDataType::SkeletalMesh:
-					SkeletalMeshResults[DataTypeIndex] = nullptr;
-					break;
-
-				case EDataType::Instance:
-					InstanceResults[DataTypeIndex] = nullptr;
-					break;
+				StrongImagePtr = nullptr;
+			}
 
-				case EDataType::ExtensionData:
-					ExtensionDataResults[DataTypeIndex] = nullptr;
-					break;
+			for (TManagedPtr<const FMesh>& StrongMeshPtr : StrongMeshResults)
+			{
+				StrongMeshPtr = nullptr;
+			}
 
-				case EDataType::InstancedStruct:
-					InstancedStructResults[DataTypeIndex] = {};
-					break;
+			StrongImageResults[0] = MakeManagedNull<const FImage>();
+			StrongMeshResults[0] = MakeManagedNull<const FMesh>();
 
-				default:
-					break;
-				}
-			}
 		}
 
-
-		bool IsValid(FCacheAddress at) const
+		bool IsValid(const FCacheAddress& Address) const
 		{
-			if (at.At == 0 || at.At>=OpExecutionData.size_code() )
+			if (Address.At == 0 || Address.At>=OpExecutionData.size_code())
 			{
 				return false;
 			}
 
-			const FOpExecutionData* Data = OpExecutionData.get_ptr(at);
+			const FOpExecutionData* Data = OpExecutionData.get_ptr(Address);
 			if (!Data)
 			{
 				return false;
 			}
 
 			// Is it a desc data query?
-			switch (at.Type)
+			switch (Address.Type)
 			{
 			case FScheduledOp::EType::Full:
+			{
+				// Resource validity cannot be checked with this method. See 
+				// AquireCachedResult and AquireCachedResultSet.
+				//check(Data->DataType != EDataType::Mesh && Data->DataType != EDataType::Image);
 				return Data->IsValueValid;
-
+			}
 			case FScheduledOp::EType::ImageDesc:
+			{
 				return Data->IsDescCacheValid;
-			
+			}
 			default:
 				unimplemented();
 				return false;
 			}
 		}
 
-
 		/** */
 		void CheckHitCountsCleared()
 		{
@@ -702,75 +552,111 @@ namespace UE::Mutable::Private
 		}
 
 
-		bool GetBool(FCacheAddress at)
+		bool GetBool(FCacheAddress Address)
 		{
-			if (!at.At) return false;
-			const FOpExecutionData* Data = OpExecutionData.get_ptr(at);
-			if (!Data) return false;
+			if (!Address.At) 
+			{
+				return false;
+			}
+
+			const FOpExecutionData* Data = OpExecutionData.get_ptr(Address);
+			if (!Data)
+			{
+				return false;
+			}
 
-			check(Data->DataType==EDataType::Bool);
-			return Data->DataTypeIndex!=0;
+			check(Data->DataType == EDataType::Bool);
+			return Data->DataTypeIndex != 0;
 		}
 
-		float GetScalar(FCacheAddress at)
+		float GetScalar(FCacheAddress Address)
 		{
-			if (!at.At) return 0.0f;
-			const FOpExecutionData* Data = OpExecutionData.get_ptr(at);
-			if (!Data) return 0.0f;
+			if (!Address.At)
+			{
+				return 0.0f;
+			}
+
+			const FOpExecutionData* Data = OpExecutionData.get_ptr(Address);
+			if (!Data)
+			{
+				return 0.0f;
+			}
 
 			check(Data->DataType == EDataType::Scalar);
 			return Data->ScalarResult;
 		}
 
-		int32 GetInt(FCacheAddress at)
+		int32 GetInt(FCacheAddress Address)
 		{
-			if (!at.At) return 0;
-			const FOpExecutionData* Data = OpExecutionData.get_ptr(at);
-			if (!Data) return 0;
+			if (!Address.At)
+			{
+				return 0;
+			}
+
+
+			const FOpExecutionData* Data = OpExecutionData.get_ptr(Address);
+			if (!Data) 
+			{
+				return 0;
+			}
 
 			check(Data->DataType == EDataType::Int);
 			return Data->DataTypeIndex;
 		}
 
-		FVector4f GetColour(FCacheAddress at)
+		FVector4f GetColour(FCacheAddress Address)
 		{
-			if (!at.At) return FVector4f();
-			const FOpExecutionData* Data = OpExecutionData.get_ptr(at);
-			if (!Data) return FVector4f();
+			if (!Address.At) 
+			{
+				return FVector4f();
+			}
+
+			const FOpExecutionData* Data = OpExecutionData.get_ptr(Address);
+			if (!Data)
+			{
+				return FVector4f();
+			}
 
 			check(Data->DataType == EDataType::Color);
 			return ColorResults[Data->DataTypeIndex];
 		}
 
-		FMatrix44f GetMatrix(FCacheAddress at)
+		FMatrix44f GetMatrix(FCacheAddress Address)
 		{
-			if (!at.At) return FMatrix44f::Identity;
-			const FOpExecutionData* Data = OpExecutionData.get_ptr(at);
-			if (!Data) return FMatrix44f::Identity;
+			if (!Address.At)
+			{
+				return FMatrix44f::Identity;
+			}
+
+			const FOpExecutionData* Data = OpExecutionData.get_ptr(Address);
+			if (!Data) 
+			{
+				return FMatrix44f::Identity;
+			}
 
 			check(Data->DataType == EDataType::Matrix);
 			return MatrixResults[Data->DataTypeIndex];
 		}
 
-		TSharedPtr<const FMaterial> GetMaterial(FCacheAddress at)
+		TManagedPtr<const FMaterial> GetMaterial(FCacheAddress at)
 		{
 			if (!at.At)
 			{
-				return MakeShared<FMaterial>();
+				return MakeManaged<FMaterial>();
 			}
 			const FOpExecutionData* Data = OpExecutionData.get_ptr(at);
 			if (!Data)
 			{
-				return MakeShared<FMaterial>();
+				return MakeManaged<FMaterial>();
 			}
 
 			check(Data->DataType == EDataType::Material);
-			TSharedPtr<const FMaterial> FoundMaterial = MaterialResults[Data->DataTypeIndex];
+			TManagedPtr<const FMaterial> FoundMaterial = MaterialResults[Data->DataTypeIndex].Pin();
 			check(FoundMaterial);
 
 			if (!FoundMaterial)
 			{
-				return MakeShared<FMaterial>();
+				return MakeManaged<FMaterial>();
 			}
 			
 			return FoundMaterial;
@@ -786,14 +672,22 @@ namespace UE::Mutable::Private
 			return ProjectorResults[Data->DataTypeIndex];
 		}
 
-		TSharedPtr<const FInstance> GetInstance(FCacheAddress at)
+
+		TManagedPtr<const FInstance> GetInstance(FCacheAddress Address)
 		{
-			if (!at.At) return nullptr;
-			FOpExecutionData* Data = OpExecutionData.get_ptr(at);
-			if (!Data) return nullptr;
+			if (!Address.At)
+			{
+				return nullptr;
+			}
+
+			FOpExecutionData* Data = OpExecutionData.get_ptr(Address);
+			if (!Data)
+			{
+				return nullptr;
+			}
 
 			check(Data->DataType == EDataType::Instance);
-			TSharedPtr<const FInstance> Result = InstanceResults[Data->DataTypeIndex];
+			TManagedPtr<const FInstance> Result = InstanceResults[Data->DataTypeIndex].Pin();
 
 			// We need to decrease the hit-count even if the result is null.
 			check(Data->OpHitCount > 0);
@@ -801,50 +695,64 @@ namespace UE::Mutable::Private
 			--Data->OpHitCount;
 			if (Data->OpHitCount == 0 && !Data->IsCacheLocked)
 			{
-				SetUnused(*Data);
+				Data->IsValueValid = false;
+				InstanceResults[Data->DataTypeIndex] = nullptr;
 			}
 
 			return Result;
 		}
 
 
-		TSharedPtr<const FImage> GetImage(FCacheAddress at, bool& bIsLastReference)
+		TManagedPtr<const FImage> GetImage(FCacheAddress Address)
 		{
-			bIsLastReference = false;
-
-			if (!at.At) return nullptr;
-			if (at.At >= OpExecutionData.size_code()) return nullptr;
-			FOpExecutionData* Data = OpExecutionData.get_ptr(at);
-			if (!Data) return nullptr;
+			if (!Address.At || Address.At >= OpExecutionData.size_code())
+			{
+				return nullptr;
+			}
+			
+			FOpExecutionData* Data = OpExecutionData.get_ptr(Address);
+			if (!Data) 
+			{
+				return nullptr;
+			}
 
 			check(Data->DataType == EDataType::Image);
-			TSharedPtr<const FImage> Result = ImageResults[Data->DataTypeIndex].Value;
+			TManagedPtr<const FImage> Result = StrongImageResults[Data->DataTypeIndex];
 
 			// We need to decrease the hit-count even if the result is null.
 			check(Data->OpHitCount > 0);
 
 			--Data->OpHitCount;
-			if (Data->OpHitCount == 0 && !Data->IsCacheLocked)
+			if (Data->OpHitCount == 0)
 			{
-				SetUnused(*Data);
-				bIsLastReference = true;
+				// If this is the last cached reference, now Result is the only reference, which is given 
+				// to the calling code. 
+				if (!Data->IsCacheLocked)
+				{
+					ImageResults[Data->DataTypeIndex] = nullptr;
+				}
+
+				StrongImageResults[Data->DataTypeIndex] = nullptr;
 			}
 
 			return Result;
 		}
 
-
-		TSharedPtr<const FMesh> GetMesh(FCacheAddress at, bool& bIsLastReference)
+		TManagedPtr<const FMesh> GetMesh(FCacheAddress Address)
 		{
-			bIsLastReference = false;
+			if (!Address.At || Address.At >= OpExecutionData.size_code()) 
+			{
+				return nullptr;
+			}
 
-			if (!at.At) return nullptr;
-			if (at.At >= OpExecutionData.size_code()) return nullptr;
-			FOpExecutionData* Data = OpExecutionData.get_ptr(at);
-			if (!Data) return nullptr;
+			FOpExecutionData* Data = OpExecutionData.get_ptr(Address);
+			if (!Data) 
+			{
+				return nullptr;
+			}
 
 			check(Data->DataType == EDataType::Mesh);
-			TSharedPtr<const FMesh> Result = MeshResults[Data->DataTypeIndex].Value;
+			TManagedPtr<const FMesh> Result = StrongMeshResults[Data->DataTypeIndex];
 
 			// We need to decrease the hit-count even if the result is null.
 			check(Data->OpHitCount > 0);
@@ -852,23 +760,39 @@ namespace UE::Mutable::Private
 			--Data->OpHitCount;
 			if (Data->OpHitCount == 0 && !Data->IsCacheLocked)
 			{
-				SetUnused(*Data);
-				bIsLastReference = true;
+				// If this is the last cached reference, now Result is the only reference, which is given 
+				// to the calling code. 
+				if (!Data->IsCacheLocked)
+				{
+					MeshResults[Data->DataTypeIndex] = nullptr;
+				}
+
+				StrongMeshResults[Data->DataTypeIndex] = nullptr;
 			}
 
 			return Result;
 		}
 
-
-		TSharedPtr<const FSkeletalMesh> GetSkeletalMesh(FCacheAddress at)
+		TManagedPtr<const FSkeletalMesh> GetSkeletalMesh(FCacheAddress Address)
 		{
-			if (!at.At) return nullptr;
-			if (at.At >= OpExecutionData.size_code()) return nullptr;
-			FOpExecutionData* Data = OpExecutionData.get_ptr(at);
-			if (!Data) return nullptr;
+			if (!Address.At)
+			{
+				return nullptr;
+			}
+
+			if (Address.At >= OpExecutionData.size_code())
+			{
+				return nullptr;
+			}
+
+			FOpExecutionData* Data = OpExecutionData.get_ptr(Address);
+			if (!Data) 
+			{
+				return nullptr;
+			}
 
 			check(Data->DataType == EDataType::SkeletalMesh);
-			TSharedPtr<const FSkeletalMesh> Result = SkeletalMeshResults[Data->DataTypeIndex];
+			TManagedPtr<const FSkeletalMesh> Result = SkeletalMeshResults[Data->DataTypeIndex].Pin();
 
 			// We need to decrease the hit-count even if the result is null.
 			check(Data->OpHitCount > 0);
@@ -876,78 +800,80 @@ namespace UE::Mutable::Private
 			--Data->OpHitCount;
 			if (Data->OpHitCount == 0 && !Data->IsCacheLocked)
 			{
-				SetUnused(*Data);
+				Data->IsValueValid = false;
+				SkeletalMeshResults[Data->DataTypeIndex] = nullptr;
 			}
 
 			return Result;
 		}
 		
-
-		TSharedPtr<const FMaterial> GetMaterial(FCacheAddress at, bool& bIsLastReference)
+		TManagedPtr<const FLayout> GetLayout(FCacheAddress Address)
 		{
-			bIsLastReference = false;
-
-			if (!at.At) return nullptr;
-			if (at.At >= OpExecutionData.size_code()) return nullptr;
-			FOpExecutionData* Data = OpExecutionData.get_ptr(at);
-			if (!Data) return nullptr;
-
-			check(Data->DataType == EDataType::Material);
-			TSharedPtr<const FMaterial> Result = MaterialResults[Data->DataTypeIndex];
-
-			// We need to decrease the hit-count even if the result is null.
-			check(Data->OpHitCount > 0);
-
-			--Data->OpHitCount;
-			if (Data->OpHitCount == 0 && !Data->IsCacheLocked)
+			if (!Address.At)
 			{
-				SetUnused(*Data);
-				bIsLastReference = true;
+				return nullptr;
 			}
 
-			return Result;
-		}
-
-		TSharedPtr<const FLayout> GetLayout(FCacheAddress at)
-		{
-			if (!at.At) return nullptr;
-			FOpExecutionData* Data = OpExecutionData.get_ptr(at);
-			if (!Data) return nullptr;
+			FOpExecutionData* Data = OpExecutionData.get_ptr(Address);
+			if (!Data)
+			{
+				return nullptr;
+			}
 
 			check(Data->DataType == EDataType::Layout);
-			return LayoutResults[Data->DataTypeIndex];
+			return LayoutResults[Data->DataTypeIndex].Pin();
 		}
 
 
-		TSharedPtr<const String> GetString(FCacheAddress at)
+		TManagedPtr<const String> GetString(FCacheAddress Address)
 		{
-			if (!at.At) return nullptr;
-			FOpExecutionData* Data = OpExecutionData.get_ptr(at);
-			if (!Data) return nullptr;
+			if (!Address.At)
+			{
+				return nullptr;
+			}
+			FOpExecutionData* Data = OpExecutionData.get_ptr(Address);
+			if (!Data) 
+			{
+				return nullptr;
+			}
 
 			check(Data->DataType == EDataType::String);
-			return StringResults[Data->DataTypeIndex];
+			return StringResults[Data->DataTypeIndex].Pin();
 		}
 
 
-		TSharedPtr<const FExtensionData> GetExtensionData(FCacheAddress at)
+		TManagedPtr<const FExtensionData> GetExtensionData(FCacheAddress Address)
 		{
-			if (!at.At) return nullptr;
-			FOpExecutionData* Data = OpExecutionData.get_ptr(at);
-			if (!Data) return nullptr;
+			if (!Address.At) 
+			{
+				return nullptr;
+			}
+
+			FOpExecutionData* Data = OpExecutionData.get_ptr(Address);
+			if (!Data)
+			{
+				return nullptr;
+			}
 
 			check(Data->DataType == EDataType::ExtensionData);
-			return ExtensionDataResults[Data->DataTypeIndex];
+			return ExtensionDataResults[Data->DataTypeIndex].Pin();
 		}
 		
-		TSharedPtr<const FInstancedStruct> GetInstancedStruct(FCacheAddress At)
+		TManagedPtr<const FInstancedStruct> GetInstancedStruct(FCacheAddress Address)
 		{
-			check(At.At);
-			FOpExecutionData* Data = OpExecutionData.get_ptr(At);
-			check(Data);
+			if (!Address.At) 
+			{
+				return nullptr;
+			}
+
+			FOpExecutionData* Data = OpExecutionData.get_ptr(Address);
+			if (!Data) 
+			{
+				return nullptr;
+			}
 
 			check(Data->DataType == EDataType::InstancedStruct);
-			TSharedPtr<const FInstancedStruct> Result = InstancedStructResults[Data->DataTypeIndex];
+			TManagedPtr<const FInstancedStruct> Result = InstancedStructResults[Data->DataTypeIndex].Pin();
 
 			check(Result);
 			check(Result->IsValid());
@@ -984,20 +910,20 @@ namespace UE::Mutable::Private
 			Data->IsValueValid = true;
 		}
 
-		void SetScalar(FCacheAddress at, float v)
+		void SetScalar(FCacheAddress Address, float v)
 		{
-			check(at.At < OpExecutionData.size_code());
-			FOpExecutionData* Data = OpExecutionData.get_ptr(at);
+			check(Address.At < OpExecutionData.size_code());
+			FOpExecutionData* Data = OpExecutionData.get_ptr(Address);
 			check(Data->DataType == EDataType::Scalar || Data->DataType == EDataType::None);
 			Data->DataType = EDataType::Scalar;
 			Data->ScalarResult = v;
 			Data->IsValueValid = true;
 		}
 
-		void SetColour(FCacheAddress at, const FVector4f& v)
+		void SetColour(FCacheAddress Address, const FVector4f& v)
 		{
-			check(at.At < OpExecutionData.size_code());
-			FOpExecutionData* Data = OpExecutionData.get_ptr(at);
+			check(Address.At < OpExecutionData.size_code());
+			FOpExecutionData* Data = OpExecutionData.get_ptr(Address);
 			check(Data->DataType == EDataType::Color || Data->DataType == EDataType::None);
 			Data->DataType = EDataType::Color;
 			Data->IsValueValid = true;
@@ -1014,10 +940,10 @@ namespace UE::Mutable::Private
 			check(Data->DataTypeIndex != 0);
 		}
 
-		void SetMatrix(FCacheAddress at, const FMatrix44f& v)
+		void SetMatrix(FCacheAddress Address, const FMatrix44f& v)
 		{
-			check(at.At < OpExecutionData.size_code());
-			FOpExecutionData* Data = OpExecutionData.get_ptr(at);
+			check(Address.At < OpExecutionData.size_code());
+			FOpExecutionData* Data = OpExecutionData.get_ptr(Address);
 			check(Data->DataType == EDataType::Matrix || Data->DataType == EDataType::None);
 			Data->DataType = EDataType::Matrix;
 			Data->IsValueValid = true;
@@ -1034,11 +960,11 @@ namespace UE::Mutable::Private
 			check(Data->DataTypeIndex != 0);
 		}
 
-		void SetMaterial(FCacheAddress at, TSharedPtr<const FMaterial> Value)
+		void SetMaterial(FCacheAddress Address, const TManagedPtr<const FMaterial>& Value)
 		{
 			check(Value);		// All material pointers provided must be valid
-			check(at.At < OpExecutionData.size_code());
-			FOpExecutionData* Data = OpExecutionData.get_ptr(at);
+			check(Address.At < OpExecutionData.size_code());
+			FOpExecutionData* Data = OpExecutionData.get_ptr(Address);
 			check(Data->DataType == EDataType::Material || Data->DataType == EDataType::None);
 			Data->DataType = EDataType::Material;
 			Data->IsValueValid = true;
@@ -1046,19 +972,19 @@ namespace UE::Mutable::Private
 			if (!Data->DataTypeIndex)
 			{
 				Data->DataTypeIndex = MaterialResults.Num();
-				MaterialResults.Add(Value);
+				MaterialResults.Add(Value.ToWeakPtr());
 			}
 			else
 			{
-				MaterialResults[Data->DataTypeIndex] = Value;
+				MaterialResults[Data->DataTypeIndex] = Value.ToWeakPtr();
 			}
 			check(Data->DataTypeIndex != 0);
 		}
 		
-		void SetProjector(FCacheAddress at, const FProjector& v)
+		void SetProjector(FCacheAddress Address, const FProjector& v)
 		{
-			check(at.At < OpExecutionData.size_code());
-			FOpExecutionData* Data = OpExecutionData.get_ptr(at);
+			check(Address.At < OpExecutionData.size_code());
+			FOpExecutionData* Data = OpExecutionData.get_ptr(Address);
 			check(Data->DataType == EDataType::Projector || Data->DataType == EDataType::None);
 			Data->DataType = EDataType::Projector;
 			Data->IsValueValid = true;
@@ -1075,10 +1001,10 @@ namespace UE::Mutable::Private
 			check(Data->DataTypeIndex != 0);
 		}
 
-		void SetInstance(FCacheAddress at, TSharedPtr<const FInstance> v)
+		void SetInstance(FCacheAddress Address, const TManagedPtr<const FInstance>& Value)
 		{
-			check(at.At < OpExecutionData.size_code());
-			FOpExecutionData* Data = OpExecutionData.get_ptr(at);
+			check(Address.At < OpExecutionData.size_code());
+			FOpExecutionData* Data = OpExecutionData.get_ptr(Address);
 			check(Data->DataType == EDataType::Instance || Data->DataType == EDataType::None);
 			Data->DataType = EDataType::Instance;
 			Data->IsValueValid = true;
@@ -1086,19 +1012,20 @@ namespace UE::Mutable::Private
 			if (!Data->DataTypeIndex)
 			{
 				Data->DataTypeIndex = InstanceResults.Num();
-				InstanceResults.Add(v);
+				InstanceResults.Add(Value.ToWeakPtr());
 			}
 			else
 			{
-				InstanceResults[Data->DataTypeIndex] = v;
+				InstanceResults[Data->DataTypeIndex] = Value.ToWeakPtr();
 			}
+
 			check(Data->DataTypeIndex != 0);
 		}
 
-		void SetExtensionData(FCacheAddress at, TSharedPtr<const FExtensionData> v)
+		void SetExtensionData(FCacheAddress Address, const TManagedPtr<const FExtensionData>& Value)
 		{
-			check(at.At < OpExecutionData.size_code());
-			FOpExecutionData* Data = OpExecutionData.get_ptr(at);
+			check(Address.At < OpExecutionData.size_code());
+			FOpExecutionData* Data = OpExecutionData.get_ptr(Address);
 			check(Data->DataType == EDataType::ExtensionData || Data->DataType == EDataType::None);
 			Data->DataType = EDataType::ExtensionData;
 			Data->IsValueValid = true;
@@ -1106,19 +1033,19 @@ namespace UE::Mutable::Private
 			if (!Data->DataTypeIndex)
 			{
 				Data->DataTypeIndex = ExtensionDataResults.Num();
-				ExtensionDataResults.Add(v);
+				ExtensionDataResults.Add(Value.ToWeakPtr());
 			}
 			else
 			{
-				ExtensionDataResults[Data->DataTypeIndex] = v;
+				ExtensionDataResults[Data->DataTypeIndex] = Value.ToWeakPtr();
 			}
 			check(Data->DataTypeIndex != 0);
 		}
 
-		void SetInstancedStruct(FCacheAddress At, const TSharedPtr<const FInstancedStruct>& Value)
+		void SetInstancedStruct(FCacheAddress Address, const TManagedPtr<const FInstancedStruct>& Value)
 		{
-			check(At.At < OpExecutionData.size_code());
-			FOpExecutionData* Data = OpExecutionData.get_ptr(At);
+			check(Address.At < OpExecutionData.size_code());
+			FOpExecutionData* Data = OpExecutionData.get_ptr(Address);
 			check(Data->DataType == EDataType::InstancedStruct || Data->DataType == EDataType::None);
 			Data->DataType = EDataType::InstancedStruct;
 			Data->IsValueValid = true;
@@ -1126,65 +1053,95 @@ namespace UE::Mutable::Private
 			if (!Data->DataTypeIndex)
 			{
 				Data->DataTypeIndex = InstancedStructResults.Num();
-				InstancedStructResults.Add(Value);
+				InstancedStructResults.Add(Value.ToWeakPtr());
 			}
 			else
 			{
-				InstancedStructResults[Data->DataTypeIndex] = Value;
+				InstancedStructResults[Data->DataTypeIndex] = Value.ToWeakPtr();
 			}
 			check(Data->DataTypeIndex != 0);
 		}
 		
-		void SetImage(FCacheAddress At, TSharedPtr<const FImage> Value)
+		void SetImage(FCacheAddress Address, const TManagedPtr<const FImage>& Value)
 		{
-			check(At.At < OpExecutionData.size_code());
-			FOpExecutionData* Data = OpExecutionData.get_ptr(At);
+			check(Address.At < OpExecutionData.size_code());
+			FOpExecutionData* Data = OpExecutionData.get_ptr(Address);
 			check(Data->DataType == EDataType::Image || Data->DataType == EDataType::None);
 			Data->DataType = EDataType::Image;
 			Data->IsValueValid = true;
 
+			const TManagedPtr<const FImage>* ValueToStorePtr;
+			if (Value.IsValid())
+			{
+				ValueToStorePtr = &Value;
+			}
+			else
+			{
+				// Set to special null reference that is not valid but can be checked if has been set.
+				ValueToStorePtr = &StrongImageResults[0];  
+			}
+
+			check(ValueToStorePtr->IsSet());
+
 			if (!Data->DataTypeIndex)
 			{
 				Data->DataTypeIndex = ImageResults.Num();
-				ImageResults.Add(TResourceResult<FImage>{At, Value});
+				ImageResults.Emplace(ValueToStorePtr->ToWeakPtr());
+				StrongImageResults.Emplace(*ValueToStorePtr);
 			}
 			else
 			{
-				ImageResults[Data->DataTypeIndex].Value = Value;
+				ImageResults[Data->DataTypeIndex] = ValueToStorePtr->ToWeakPtr();
+				StrongImageResults[Data->DataTypeIndex] = *ValueToStorePtr;
 			}
 			check(Data->DataTypeIndex != 0);
 
 			UE::Mutable::Private::UpdateLLMStats();
 		}
 
-		void SetMesh(FCacheAddress At, TSharedPtr<const FMesh> Value)
+		void SetMesh(FCacheAddress Address, const TManagedPtr<const FMesh>& Value)
 		{
-			check(At.At < OpExecutionData.size_code());
+			check(Address.At < OpExecutionData.size_code());
 
-			FOpExecutionData* Data = OpExecutionData.get_ptr(At);
+			FOpExecutionData* Data = OpExecutionData.get_ptr(Address);
 
 			check(Data->DataType == EDataType::Mesh || Data->DataType == EDataType::None);
 			Data->DataType = EDataType::Mesh;
 			Data->IsValueValid = true;
 
-			if (!Data->DataTypeIndex)
+			const TManagedPtr<const FMesh>* ValueToStorePtr;
+			if (Value.IsValid())
+			{
+				ValueToStorePtr = &Value;
+			}
+			else
+			{
+				// Set to special null reference that is not valid but can be check if is Set.
+				ValueToStorePtr = &StrongMeshResults[0];   
+			}
+
+			check(ValueToStorePtr->IsSet());
+
+			if (!Data->DataTypeIndex)
 			{
 				Data->DataTypeIndex = MeshResults.Num();
-				MeshResults.Add(TResourceResult<FMesh>{At, Value});
+				MeshResults.Emplace(ValueToStorePtr->ToWeakPtr());
+				StrongMeshResults.Emplace(*ValueToStorePtr);
 			}
 			else
 			{
-				MeshResults[Data->DataTypeIndex].Value = Value;
+				MeshResults[Data->DataTypeIndex] = ValueToStorePtr->ToWeakPtr();
+				StrongMeshResults[Data->DataTypeIndex] = *ValueToStorePtr;
 			}
 			check(Data->DataTypeIndex != 0);
 
 			UE::Mutable::Private::UpdateLLMStats();
 		}
 
-		void SetSkeletalMesh(FCacheAddress At, TSharedPtr<const FSkeletalMesh> V)
+		void SetSkeletalMesh(FCacheAddress Address, const TManagedPtr<const FSkeletalMesh>& Value)
 		{
-			check(At.At < OpExecutionData.size_code());
-			FOpExecutionData* Data = OpExecutionData.get_ptr(At);
+			check(Address.At < OpExecutionData.size_code());
+			FOpExecutionData* Data = OpExecutionData.get_ptr(Address);
 			check(Data->DataType == EDataType::SkeletalMesh || Data->DataType == EDataType::None);
 			Data->DataType = EDataType::SkeletalMesh;
 			Data->IsValueValid = true;
@@ -1192,21 +1149,21 @@ namespace UE::Mutable::Private
 			if (!Data->DataTypeIndex)
 			{
 				Data->DataTypeIndex = SkeletalMeshResults.Num();
-				SkeletalMeshResults.Add(V);
+				SkeletalMeshResults.Add(Value);
 			}
 			else
 			{
-				SkeletalMeshResults[Data->DataTypeIndex] = V;
+				SkeletalMeshResults[Data->DataTypeIndex] = Value;
 			}
 			check(Data->DataTypeIndex != 0);
 
 			UE::Mutable::Private::UpdateLLMStats();
 		}
 
-		void SetLayout(FCacheAddress at, TSharedPtr<const FLayout> v)
+		void SetLayout(FCacheAddress Address, const TManagedPtr<const FLayout>& Value)
 		{
-			check(at.At < OpExecutionData.size_code());
-			FOpExecutionData* Data = OpExecutionData.get_ptr(at);
+			check(Address.At < OpExecutionData.size_code());
+			FOpExecutionData* Data = OpExecutionData.get_ptr(Address);
 			check(Data->DataType == EDataType::Layout || Data->DataType == EDataType::None);
 			Data->DataType = EDataType::Layout;
 			Data->IsValueValid = true;
@@ -1214,21 +1171,21 @@ namespace UE::Mutable::Private
 			if (!Data->DataTypeIndex)
 			{
 				Data->DataTypeIndex = LayoutResults.Num();
-				LayoutResults.Add(v);
+				LayoutResults.Add(Value);
 			}
 			else
 			{
-				LayoutResults[Data->DataTypeIndex] = v;
+				LayoutResults[Data->DataTypeIndex] = Value;
 			}
 			check(Data->DataTypeIndex != 0);
 
 			UE::Mutable::Private::UpdateLLMStats();
 		}
 
-		void SetString(FCacheAddress at, TSharedPtr<const String> v)
+		void SetString(FCacheAddress Address, const TManagedPtr<const String>& Value)
 		{
-			check(at.At < OpExecutionData.size_code());
-			FOpExecutionData* Data = OpExecutionData.get_ptr(at);
+			check(Address.At < OpExecutionData.size_code());
+			FOpExecutionData* Data = OpExecutionData.get_ptr(Address);
 			check(Data->DataType == EDataType::String || Data->DataType == EDataType::None);
 			Data->DataType = EDataType::String;
 			Data->IsValueValid = true;
@@ -1236,39 +1193,193 @@ namespace UE::Mutable::Private
 			if (!Data->DataTypeIndex)
 			{
 				Data->DataTypeIndex = StringResults.Num();
-				StringResults.Add(v);
+				StringResults.Add(Value.ToWeakPtr());
 			}
 			else
 			{
-				StringResults[Data->DataTypeIndex] = v;
+				StringResults[Data->DataTypeIndex] = Value.ToWeakPtr();
 			}
 			check(Data->DataTypeIndex!=0);
 		}
 
 
-		inline void IncreaseHitCount(FCacheAddress at)
+		inline void IncreaseHitCount(FCacheAddress Address)
 		{
 			// Don't count hits for instruction 0, which is always null. It is usually already
 			// check that At is not 0, and then it is not requested, generating a stray non-zero count
 			// at its position.
-			if (at.At)
+			if (Address.At)
 			{
-				check(at.At < OpExecutionData.size_code());
-				FOpExecutionData& Data = OpExecutionData[at];
+				check(Address.At < OpExecutionData.size_code());
+				FOpExecutionData& Data = OpExecutionData[Address];
 				Data.OpHitCount++;
 			}
 		}
 
-		inline void SetForceCached(OP::ADDRESS at)
+		inline void SetForceCached(OP::ADDRESS Address)
 		{
-			// \TODO: It only locks at,0,0
-			if (at)
+			// \TODO: It only locks Address,0,0
+			if (Address)
 			{
-				check(at < OpExecutionData.size_code());
-				FOpExecutionData& Data = OpExecutionData[FCacheAddress(at, 0, 0)];
+				check(Address < OpExecutionData.size_code());
+				FOpExecutionData& Data = OpExecutionData[FCacheAddress(Address, 0, 0)];
 				Data.IsCacheLocked = true;
 			}
 		}
+
+		bool TryAquireCachedResultSet(TConstArrayView<FCacheAddress> AddressSet)
+		{
+			TArray<TManagedPtr<const FImage>, TInlineAllocator<16>> LocalImages;
+			TArray<TManagedPtr<const FMesh> , TInlineAllocator<16>> LocalMeshes;
+			TArray<FOpExecutionData*        , TInlineAllocator<16>> LocalData;
+
+			if (AddressSet.IsEmpty())
+			{
+				return true;
+			}
+
+			for (const FCacheAddress& Address : AddressSet)
+			{
+				if (!Address.At)
+				{
+					continue;
+				}
+
+				check(Address.At < OpExecutionData.size_code());
+
+				FOpExecutionData* Data = OpExecutionData.get_ptr(Address);
+			
+				if (!Data)
+				{
+					return false;
+				}
+			
+				if (Address.Type == FScheduledOp::EType::ImageDesc)
+				{
+					if (!Data->IsDescCacheValid)
+					{
+						return false;
+					}
+				}
+				else
+				{
+					switch (Data->DataType)
+					{
+						case EDataType::Image:
+						{
+							if (!LocalImages.Add_GetRef(ImageResults[Data->DataTypeIndex].Pin()).IsSet())
+							{
+								return false;
+							}
+							LocalData.Add(Data);
+							break;
+						}
+						case EDataType::Mesh:
+						{
+							if (!LocalMeshes.Add_GetRef(MeshResults[Data->DataTypeIndex].Pin()).IsSet())
+							{
+								return false;
+							}
+							LocalData.Add(Data);
+							break;
+						}
+
+						default:
+						{
+							if (!Data->IsValueValid)
+							{
+								return false;
+							}
+							break;
+						}
+					}
+				}
+			}
+
+			// We have been successful acquiring all addresses in the set. 
+			// Commit resource acquisition.
+			for (int32 Index = LocalData.Num() - 1; Index >= 0; --Index)
+			{
+				FOpExecutionData* Data = LocalData[Index];
+
+				switch (Data->DataType)
+				{
+					case EDataType::Image:
+					{
+						StrongImageResults[Data->DataTypeIndex] = LocalImages[Index];
+						break;
+					}
+					case EDataType::Mesh:
+					{
+						StrongMeshResults[Data->DataTypeIndex] = LocalMeshes[Index];
+						break;
+					}
+
+					default:
+					{
+						check(false);
+						break;
+					}
+				}	
+			}
+
+			return true;
+		}
+
+		bool AquireCachedResult(FCacheAddress Address)
+		{
+			check(Address.At)
+
+			check(Address.At < OpExecutionData.size_code());
+
+			FOpExecutionData* Data = OpExecutionData.get_ptr(Address);
+		
+			if (!Data)
+			{
+				return false;
+			}
+
+			if (Address.Type == FScheduledOp::EType::ImageDesc)
+			{
+				return Data->IsDescCacheValid;
+			}
+	
+			switch (Data->DataType)
+			{
+				case EDataType::Image:
+				{
+					if (StrongImageResults[Data->DataTypeIndex].IsSet())
+					{
+						return true;
+					}
+					else
+					{
+						TManagedPtr<const FImage> PinnedPtr = ImageResults[Data->DataTypeIndex].Pin();
+						StrongImageResults[Data->DataTypeIndex] = PinnedPtr;
+						return PinnedPtr.IsSet();
+					}
+				}
+				case EDataType::Mesh:
+				{
+					if (StrongMeshResults[Data->DataTypeIndex].IsSet())
+					{
+						return true;
+					}
+					else
+					{
+
+						TManagedPtr<const FMesh> PinnedPtr = MeshResults[Data->DataTypeIndex].Pin();
+						StrongMeshResults[Data->DataTypeIndex] = PinnedPtr;
+						return PinnedPtr.IsSet();
+					}
+				}
+
+				default:
+				{
+					return Data->IsValueValid;
+				}
+			}
+		}
 	};
 
 
@@ -1302,7 +1413,7 @@ namespace UE::Mutable::Private
 		FInstance::FID InstanceID;
 		int32 State = 0;
 		uint32 LODMask = 0;
-		TSharedPtr<const FInstance> Instance;
+		TManagedPtr<const FInstance> Instance;
 		TSharedPtr<const FModel> Model;
 
 		TSharedPtr<const FParameters> OldParameters;
@@ -1322,10 +1433,23 @@ namespace UE::Mutable::Private
 		{
 			// Manually done to trace mem deallocations
 			MUTABLE_CPUPROFILER_SCOPE(LiveInstanceDestructor);
-			Cache = nullptr;
-			OldParameters = nullptr;
-			Instance = nullptr;
-			Model = nullptr;
+			{
+				MUTABLE_CPUPROFILER_SCOPE(LiveInstanceDestructor_DeleteCache);
+				Cache = nullptr;
+			}
+
+			{	
+				MUTABLE_CPUPROFILER_SCOPE(LiveInstanceDestructor_DeleteOldParameters);
+				OldParameters = nullptr;
+			}
+			{
+				MUTABLE_CPUPROFILER_SCOPE(LiveInstanceDestructor_DeleteInstance);
+				Instance = nullptr;
+			}
+			{
+				MUTABLE_CPUPROFILER_SCOPE(LiveInstanceDestructor_DeleteModel);
+				Model = nullptr;
+			}
 		}
 	};
 
@@ -1415,387 +1539,123 @@ namespace UE::Mutable::Private
         void MarkRomUsed( int32 RomIndex, const TSharedPtr<const FModel>& );
 
 		/** */
-		[[nodiscard]] TSharedPtr<FImage> CreateImage(uint32 SizeX, uint32 SizeY, uint32 Lods, EImageFormat Format, EInitializationType Init)
+		[[nodiscard]] TManagedPtr<FImage> CreateImage(uint32 SizeX, uint32 SizeY, uint32 Lods, EImageFormat Format, EInitializationType Init)
 		{
 			CheckRunnerThread();
 
 			uint32 DataSize = FImage::CalculateDataSize(SizeX, SizeY, Lods, Format);
 
-			// Look for an unused image in the pool that can be reused
-			int32 PooledImageCount = PooledImages.Num();
-			for (int Index = 0; DataSize>0 && Index<PooledImageCount; ++Index)
-			{
-				TSharedPtr<FImage>& Candidate = PooledImages[Index];
-				if (Candidate->GetFormat() == Format
-					&& Candidate->GetSizeX() == SizeX
-					&& Candidate->GetSizeY() == SizeY
-					&& Candidate->GetLODCount() == Lods
-					)
-				{
-					TSharedPtr<FImage> Result = Candidate;
-					PooledImages.RemoveAtSwap(Index);
-					
-					if (Init == EInitializationType::Black)
-					{
-						Result->InitToBlack();
-					}
-					else
-					{
-						Result->Flags = 0;
-						Result->RelevancyMinY = 0;
-						Result->RelevancyMaxY = 0;
-					}
-					return Result;
-				}
-			}
+			//TODO: 
 
 			// Make room in the budget
 			EnsureBudgetBelow(DataSize);
 
 			// Create it
-			TSharedPtr<FImage> Result = MakeShared<FImage>(SizeX, SizeY, Lods, Format, Init);
+			TManagedPtr<FImage> Result = MakeManaged<FImage>(SizeX, SizeY, Lods, Format, Init);
 
-			TempImages.Add(Result);
 			return Result;
 		}
 
 		/** Ref will be nulled and relesed in any case. */
-		[[nodiscard]] TSharedPtr<FImage> CloneOrTakeOver(TSharedPtr<const FImage>& Resource)
+		[[nodiscard]] TManagedPtr<FImage> CloneOrTakeOver(TManagedPtr<const FImage>&& Resource)
 		{
 			CheckRunnerThread();
-
-			TempImages.RemoveSingle(Resource);
-			
-			check(!TempImages.Contains(Resource));
-			check(!PooledImages.Contains(Resource));
-
-			TSharedPtr<FImage> Result;
-			if (!Resource.IsUnique())
+	
+			TManagedPtr<FImage> Result;
+			if (!Resource.IsUniqueReference())
 			{
-				// TODO: try to grab from the pool
-
 				uint32 DataSize = Resource->GetDataSize();
 				EnsureBudgetBelow(DataSize);
 
 				Result = Resource->Clone();
-				Release(Resource);
-			}
-			else
-			{
-				Result = ConstCastSharedPtr<FImage>(Resource);
-				Resource = nullptr;
-			}
-
-			return Result;
-		}
-
-		/** */
-		void Release(TSharedPtr<const FImage>& Resource)
-		{
-			CheckRunnerThread();
-
-			if (!Resource)
-			{
-				return;
-			}
-
-			const int32 ResourceDataSize = Resource->GetDataSize();
-			TempImages.RemoveSingle(Resource);
-
-			check(!TempImages.Contains(Resource));
-			check(!PooledImages.Contains(Resource));
-
-			if (IsBudgetTemp(Resource))
-			{
-				// Check if we are exceeding the budget
-				bool bInBudget = EnsureBudgetBelow(ResourceDataSize);
-				if (bInBudget)
-				{
-					PooledImages.Add(ConstCastSharedPtr<FImage>(Resource));
-				}
 			}
 			else
 			{
-				// Check if we are exceeding the budget
-				EnsureBudgetBelow(0);
+				Result = ConstCastManagedPtr<FImage>(Resource);
 			}
 
 			Resource = nullptr;
+			return Result;
 		}
 
-		/** */
-		void Release(TSharedPtr<FImage>& Resource)
-		{
-			CheckRunnerThread();
-
-			if (!Resource)
-			{
-				return;
-			}
-
-			const int32 ResourceDataSize = Resource->GetDataSize();
-			TempImages.RemoveSingle(Resource);
-
-			check(!TempImages.Contains(Resource));
-			check(!PooledImages.Contains(Resource));
-
-			if (IsBudgetTemp(Resource))
-			{
-				// Check if we are exceeding the budget
-				bool bInBudget = EnsureBudgetBelow(ResourceDataSize);
-				if (bInBudget)
-				{
-					PooledImages.Add(Resource);
-				}
-			}
-			else
-			{
-				// Check if we are exceeding the budget
-				EnsureBudgetBelow(0);
-			}
-
-			Resource = nullptr;
-		}
 
-		[[nodiscard]] TSharedPtr<FMesh> CreateMesh(int32 BudgetReserveSize)
+		[[nodiscard]] TManagedPtr<FMesh> CreateMesh(int32 BudgetReserveSize)
 		{
 			CheckRunnerThread();
 
 			EnsureBudgetBelow(BudgetReserveSize);
 
-			TSharedPtr<FMesh> Result = MakeShared<FMesh>();
-
-			TempMeshes.Add(Result);
+			TManagedPtr<FMesh> Result = MakeManaged<FMesh>();
 
 			return Result;
 		}
 
-		[[nodiscard]] TSharedPtr<FMesh> CloneOrTakeOver(TSharedPtr<const FMesh>& Resource)
+		[[nodiscard]] TManagedPtr<FMesh> CloneOrTakeOver(TManagedPtr<const FMesh>&& Resource)
 		{
 			CheckRunnerThread();
 
-			const int32 ResourceDataSize = Resource->GetDataSize();
-			TempMeshes.RemoveSingle(Resource);
-
-			TSharedPtr<FMesh> Result;
-			if (!Resource.IsUnique())
+			TManagedPtr<FMesh> Result;
+			if (!Resource.IsUniqueReference())
 			{
+				const int32 ResourceDataSize = Resource->GetDataSize();
 				Result = CreateMesh(ResourceDataSize); 
 				Result->CopyFrom(*Resource);
-				Release(Resource);
 			}
 			else
 			{
-				Result = TSharedPtr<FMesh>(ConstCastSharedPtr<FMesh>(Resource));
-				Resource = nullptr;
+				Result = ConstCastManagedPtr<FMesh>(Resource);
 			}
 
-			return Result;
-		}
-
-		void Release(TSharedPtr<const FMesh>& Resource)
-		{
-			CheckRunnerThread();
-
-			if (!Resource)
-			{
-				return;
-			}
-
-			TempMeshes.RemoveSingle(Resource);
-			check(!TempMeshes.Contains(Resource));
-
-			EnsureBudgetBelow(0);
-
-			Resource = nullptr;
-		}
-
-		void Release(TSharedPtr<FMesh>& Resource)
-		{
-			TSharedPtr<const FMesh> ConstPtr = Resource;
-
-			Release(ConstPtr);
-
 			Resource = nullptr;
+			return Result;
 		}
 
 		/** */
-		[[nodiscard]] TSharedPtr<const FMesh> LoadMesh(const FCacheAddress& From, bool bTakeOwnership = false)
+		[[nodiscard]] TManagedPtr<const FMesh> LoadMesh(const FCacheAddress& From)
 		{
-			bool bIsLastReference = false;
-			TSharedPtr<const FMesh> Result = CurrentInstanceCache->GetMesh(From, bIsLastReference);
-			if (!Result)
-			{
-				return nullptr;
-			}
-
-			// If we retrieved the last reference to this resource in "From" cache position (it could still be in other cache positions as well)
-			if (bIsLastReference)
-			{
-				int32* CountPtr = CacheResources.Find(Result);
-				check(CountPtr);
-				*CountPtr = (*CountPtr) - 1;
-				if (!*CountPtr)
-				{
-					CacheResources.FindAndRemoveChecked(Result);
-				}
-			}
-
-			if (!bTakeOwnership && Result.IsUnique())
-			{
-				TempMeshes.Add(Result);
-			}
-
+			TManagedPtr<const FMesh> Result = CurrentInstanceCache->GetMesh(From);
 			return Result;
 		}
 
-
     	/** */
-    	[[nodiscard]] TSharedPtr<const FSkeletalMesh> LoadSkeletalMesh(const FCacheAddress& From)
+    	[[nodiscard]] TManagedPtr<const FSkeletalMesh> LoadSkeletalMesh(const FCacheAddress& From)
 		{
-			return CurrentInstanceCache->GetSkeletalMesh(From);
+			TManagedPtr<const FSkeletalMesh> Result = CurrentInstanceCache->GetSkeletalMesh(From);
+			return Result;
 		}
     	
 
 		/** */
-		[[nodiscard]] TSharedPtr<const FImage> LoadImage(const FCacheAddress& From, bool bTakeOwnership = false)
+		[[nodiscard]] TManagedPtr<const FImage> LoadImage(const FCacheAddress& From)
 		{
-			bool bIsLastReference = false;
-			TSharedPtr<const FImage> Result = CurrentInstanceCache->GetImage(From, bIsLastReference);
-			if (!Result)
-			{
-				return nullptr;
-			}
-
-			// If we retrieved the last reference to this resource in "From" cache position (it could still be in other cache positions as well)
-			if (bIsLastReference)
-			{
-				int32* CountPtr = CacheResources.Find(Result);
-				check(CountPtr);
-				*CountPtr = (*CountPtr) - 1;
-				if (!*CountPtr)
-				{
-					CacheResources.FindAndRemoveChecked(Result);
-				}
-			}
-
-
-			if (!bTakeOwnership && Result.IsUnique())
-			{
-				TempImages.Add(Result);
-			}
-
-			return Result;
+			return CurrentInstanceCache->GetImage(From);
 		}
 
+
     	/** */
-    	[[nodiscard]] TSharedPtr<const FMaterial> LoadMaterial(const FCacheAddress& From)
+    	[[nodiscard]] TManagedPtr<const FMaterial> LoadMaterial(const FCacheAddress& From)
 		{
-			bool bIsLastReference = false;
-			TSharedPtr<const FMaterial> Result = CurrentInstanceCache->GetMaterial(From, bIsLastReference);
-			if (!Result)
-			{
-				return nullptr;
-			}
-
-			// If we retrieved the last reference to this resource in "From" cache position (it could still be in other cache positions as well)
-			if (bIsLastReference)
-			{
-				int32* CountPtr = CacheResources.Find(Result);
-				check(CountPtr);
-				*CountPtr = (*CountPtr) - 1;
-				if (!*CountPtr)
-				{
-					CacheResources.FindAndRemoveChecked(Result);
-				}
-			}
-
+			TManagedPtr<const FMaterial> Result = CurrentInstanceCache->GetMaterial(From);
 			return Result;
 		}
 
 		/** */
-		void StoreImage(const FCacheAddress& To, TSharedPtr<const FImage> Resource)
+		void StoreImage(const FCacheAddress& To, const TManagedPtr<const FImage>& Resource)
 		{
-			if (Resource)
-			{
-				int32 ResourceDataSize = Resource->GetDataSize();
-				TempImages.RemoveSingle(Resource);
-
-				check(!TempImages.Contains(Resource));
-
-				int32& Count = CacheResources.FindOrAdd(Resource, 0);
-				++Count;
-			}
-
 			CurrentInstanceCache->SetImage(To, Resource);
 		}
 
-		void StoreMesh(const FCacheAddress& To, TSharedPtr<const FMesh> Resource)
+		void StoreMesh(const FCacheAddress& To, const TManagedPtr<const FMesh>& Resource)
 		{
-			if (Resource)
-			{
-				//Resource->CheckIntegrity();
-
-				TempMeshes.RemoveSingle(Resource);
-
-				int32& Count = CacheResources.FindOrAdd(Resource, 0);
-				++Count;
-			}
-
+			//Resource->CheckIntegrity();
 			CurrentInstanceCache->SetMesh(To, Resource);
 		}
 
-    	void StoreMaterial(const FCacheAddress& To, TSharedPtr<const FMaterial> Resource)
+    	void StoreMaterial(const FCacheAddress& To, const TManagedPtr<const FMaterial>& Resource)
 		{
-			if (Resource)
-			{
-				int32& Count = CacheResources.FindOrAdd(Resource, 0);
-				++Count;
-			}
-
 			CurrentInstanceCache->SetMaterial(To, Resource);
 		}
 
-		/** Return true if the resource is not in any cache (0,1,rom). */
-		bool IsBudgetTemp(const TSharedPtr<const FResource>& Resource)
-		{
-			if (!Resource)
-			{
-				return false;
-			}
-
-			bool bIsTemp = Resource.IsUnique();
-			return bIsTemp;
-		}
-
-		int32 GetPooledBytes() const
-		{
-			int32 Result = 0;
-			for (const TSharedPtr<FImage>& Value : PooledImages)
-			{
-				Result += Value->GetDataSize();
-			}
-
-			return Result;
-		}
-
-		int32 GetTempBytes() const
-		{
-			int32 Result = 0;
-			for (const TSharedPtr<const FImage>& Value : TempImages)
-			{
-				Result += Value->GetDataSize();
-			}
-
-			for (const TSharedPtr<const FMesh>& Value : TempMeshes)
-			{
-				Result += Value->GetDataSize();
-			}
-			
-			return Result;
-		}
-
-
 		int32 GetRomBytes() const
 		{
 			int32 Result = 0;
@@ -1811,14 +1671,14 @@ namespace UE::Mutable::Private
 			{
 				// Count streamable and currently-loaded resources
 				const FProgram& Program = Model->GetPrivate()->Program;
-				for (const TPair<uint32,TSharedPtr<const FImage>>& ImageLOD : Program.ConstantImageLODsStreamed)
+				for (const TPair<uint32,TManagedPtr<const FImage>>& ImageLOD : Program.ConstantImageLODsStreamed)
 				{
 					if (ImageLOD.Value)
 					{
 						Result += ImageLOD.Value->GetDataSize();
 					}
 				}
-				for (const TPair<uint32, TSharedPtr<const FMesh>>& Rom : Program.ConstantMeshesStreamed)
+				for (const TPair<uint32, TManagedPtr<const FMesh>>& Rom : Program.ConstantMeshesStreamed)
 				{
 					if (Rom.Value)
 					{
@@ -1830,71 +1690,60 @@ namespace UE::Mutable::Private
 			return Result;
 		}
 
-
-		int32 GetTrackedCacheBytes() const
-		{
-			int32 Result = 0;
-			for (const TTuple<TSharedPtr<const FResource>,int32>& It : CacheResources)
-			{
-				Result += It.Key->GetDataSize();
-			}
-
-			return Result;
-		}
-
 		/** Calculate the amount of bytes in data cached in the level 0 and 1 cache in all live instances. */
 		int32 GetCacheBytes() const
 		{
 			int32 Result = 0;
-			TSet<const FResource*> Cache0Unique;
-			TSet<const FResource*> Cache1Unique;
+			TSet<TManagedPtr<const FResource>> Cache0Unique;
+			TSet<TManagedPtr<const FResource>> Cache1Unique;
 
 			for (const FLiveInstance& Instance : LiveInstances)
 			{
-				CodeContainer<FProgramCache::FOpExecutionData>::iterator it = Instance.Cache->OpExecutionData.begin();
-				for (; it.IsValid(); ++it)
+			CodeContainer<FProgramCache::FOpExecutionData>::iterator it = Instance.Cache->OpExecutionData.begin();
+			for (; it.IsValid(); ++it)
+			{
+				if (!(*it).DataTypeIndex)
 				{
-					if (!(*it).DataTypeIndex)
-					{
-						continue;
-					}
+					continue;
+				}
+
+				TSet<TManagedPtr<const FResource>>* TargetSet = &Cache0Unique;
+				if ((*it).IsCacheLocked)
+				{
+					TargetSet = &Cache1Unique;
+				}
 
-					TSet<const FResource*>* TargetSet = &Cache0Unique;
-					if ((*it).IsCacheLocked)
+				switch ((*it).DataType)
+				{
+				case EDataType::Image:
+				{
+					TManagedPtr<const FImage> Value = Instance.Cache->ImageResults[(*it).DataTypeIndex].Pin();
+					if (Value)
 					{
-						TargetSet = &Cache1Unique;
+						TargetSet->Add(MoveTemp(Value));
 					}
-
-					const FResource* Value = nullptr;
-					switch ((*it).DataType)
+					break;
+				}
+				case EDataType::Mesh:
+				{
+					TManagedPtr<const FMesh> Value = Instance.Cache->MeshResults[(*it).DataTypeIndex].Pin();
+					if (Value)
 					{
-					case EDataType::Image:
-						Value = Instance.Cache->ImageResults[(*it).DataTypeIndex].Value.Get();
-						if (Value)
-						{
-							TargetSet->Add(Value);
-						}
-						break;
-
-					case EDataType::Mesh:
-						Value = Instance.Cache->MeshResults[(*it).DataTypeIndex].Value.Get();
-						if (Value)
-						{
-							TargetSet->Add(Value);
-						}
-						break;
-
-					default:
-						break;
+						TargetSet->Add(MoveTemp(Value));
 					}
+					break;
+				}
+				default:
+					break;
 				}
 			}
+		}
 
 			// Merge
 			Cache0Unique.Append(Cache1Unique);
 
 			// Count
-			for (const FResource* Value : Cache0Unique)
+			for (const TManagedPtr<const FResource>& Value : Cache0Unique)
 			{
 				Result += Value->GetDataSize();
 			}
@@ -1902,9 +1751,10 @@ namespace UE::Mutable::Private
 			return Result;
 		}
 
-		/** Remove all intermediate data (big and small) from the memory except for the one that has been explicitely
-		* marked as state cache.
-		*/
+		/** 
+		 * Remove all intermediate data (big and small) from the memory except for the one that has been explicitely
+		 * marked as state cache.
+		 */
 		void ClearCacheLayer0()
 		{
 			check(CurrentInstanceCache);
@@ -1915,9 +1765,7 @@ namespace UE::Mutable::Private
 			for (; it.IsValid(); ++it)
 			{
 				FProgramCache::FOpExecutionData& Data = *it;
-				if (!Data.DataTypeIndex
-					||
-					Data.IsCacheLocked)
+				if (!Data.DataTypeIndex || Data.IsCacheLocked)
 				{
 					continue;
 				}
@@ -1926,23 +1774,15 @@ namespace UE::Mutable::Private
 				{
 				case EDataType::Image:
 				{
-					TSharedPtr<const FImage> Value = CurrentInstanceCache->ImageResults[Data.DataTypeIndex].Value;
-					if (Value)
-					{
-						CacheResources.Remove(Value);
-						CurrentInstanceCache->ImageResults[Data.DataTypeIndex].Value = nullptr;
-					}
+					CurrentInstanceCache->ImageResults[Data.DataTypeIndex] = nullptr;
+					CurrentInstanceCache->StrongImageResults[Data.DataTypeIndex] = nullptr;
 					break;
 				}
 
 				case EDataType::Mesh:
 				{
-					TSharedPtr<const FMesh> Value = CurrentInstanceCache->MeshResults[Data.DataTypeIndex].Value;
-					if (Value)
-					{
-						CacheResources.Remove(Value);
-						CurrentInstanceCache->MeshResults[Data.DataTypeIndex].Value = nullptr;
-					}
+					CurrentInstanceCache->MeshResults[Data.DataTypeIndex] = nullptr;
+					CurrentInstanceCache->StrongMeshResults[Data.DataTypeIndex] = nullptr;
 					break;
 				}
 
@@ -1985,53 +1825,41 @@ namespace UE::Mutable::Private
 			for (; it.IsValid(); ++it)
 			{
 				FProgramCache::FOpExecutionData& Data = *it;
-				if (!Data.DataTypeIndex)
-				{
-					continue;
-				}
+				
+				Data.OpHitCount = 0;
+				Data.IsValueValid = false;
+			}
 
-				switch (Data.DataType)
-				{
-				case EDataType::Image:
-				{
-					TSharedPtr<const FImage> Value = CurrentInstanceCache->ImageResults[Data.DataTypeIndex].Value;
-					CacheResources.Remove(Value);
-					CurrentInstanceCache->ImageResults[Data.DataTypeIndex].Value = nullptr;
-					break;
-				}
+			for (TManagedWeakPtr<const FImage>& Value : CurrentInstanceCache->ImageResults)
+			{
+				Value = nullptr;
+			}
 
-				case EDataType::Mesh:
-				{
-					TSharedPtr<const FMesh> Value = CurrentInstanceCache->MeshResults[Data.DataTypeIndex].Value;
-					CacheResources.Remove(Value);
-					CurrentInstanceCache->MeshResults[Data.DataTypeIndex].Value = nullptr;
-					break;
-				}
+			for (TManagedPtr<const FImage>& Value : CurrentInstanceCache->StrongImageResults)
+			{
+				Value = nullptr;
+			}
 
-				case EDataType::SkeletalMesh:
-				{
-					CurrentInstanceCache->SkeletalMeshResults[Data.DataTypeIndex] = nullptr;
-					break;
-				}
 
-				case EDataType::Layout:
-				{
-					CurrentInstanceCache->LayoutResults[Data.DataTypeIndex] = nullptr;
-					break;
-				}
+			for (TManagedWeakPtr<const FSkeletalMesh>& Value : CurrentInstanceCache->SkeletalMeshResults)
+			{
+				Value = nullptr;
+			}
 
-				case EDataType::Instance:
-				{
-					CurrentInstanceCache->InstanceResults[Data.DataTypeIndex] = nullptr;
-					break;
-				}
+			for (TManagedPtr<const FMesh>& Value : CurrentInstanceCache->StrongMeshResults)
+			{
+				Value = nullptr;
+			}
 
-				default:
-					break;
-				}
 
-				Data.OpHitCount = 0;
-				Data.IsValueValid = false;
+			for (TManagedWeakPtr<const FLayout>& Value : CurrentInstanceCache->LayoutResults)
+			{
+				Value = nullptr;
+			}
+
+			for (TManagedWeakPtr<const FInstance>& Value : CurrentInstanceCache->InstanceResults)
+			{
+				Value = nullptr;
 			}
 		}
 
@@ -2151,11 +1979,11 @@ namespace UE::Mutable::Private
 		MUTABLERUNTIME_API int32 BuildInt(const TSharedPtr<FExternalResourceProvider>& ExternalResourceProvider, const TSharedPtr<const FModel>&, const FParameters*, OP::ADDRESS) ;
 		MUTABLERUNTIME_API float BuildScalar(const TSharedPtr<FExternalResourceProvider>& ExternalResourceProvider, const TSharedPtr<const FModel>&, const FParameters*, OP::ADDRESS) ;
 		MUTABLERUNTIME_API FVector4f BuildColour(const TSharedPtr<FExternalResourceProvider>& ExternalResourceProvider, const TSharedPtr<const FModel>&, const FParameters*, OP::ADDRESS) ;
-		MUTABLERUNTIME_API TSharedPtr<const String> BuildString(const TSharedPtr<FExternalResourceProvider>& ExternalResourceProvider, const TSharedPtr<const FModel>&, const FParameters*, OP::ADDRESS) ;
-		MUTABLERUNTIME_API TSharedPtr<const FImage> BuildImage(const TSharedPtr<FExternalResourceProvider>& ExternalResourceProvider, const TSharedPtr<const FModel>&, const FParameters*, OP::ADDRESS, int32 MipsToSkip);
-		MUTABLERUNTIME_API TSharedPtr<const FMesh> BuildMesh(const TSharedPtr<FExternalResourceProvider>& ExternalResourceProvider, const TSharedPtr<const FModel>&, const FParameters*, OP::ADDRESS, EMeshContentFlags MeshContentFilter);
-		MUTABLERUNTIME_API TSharedPtr<const FInstance> BuildInstance(const TSharedPtr<FExternalResourceProvider>& ExternalResourceProvider, const TSharedPtr<const FModel>&, const FParameters*, OP::ADDRESS);
-		MUTABLERUNTIME_API TSharedPtr<const FLayout> BuildLayout(const TSharedPtr<FExternalResourceProvider>& ExternalResourceProvider, const TSharedPtr<const FModel>&, const FParameters*, OP::ADDRESS) ;
+		MUTABLERUNTIME_API TManagedPtr<const String> BuildString(const TSharedPtr<FExternalResourceProvider>& ExternalResourceProvider, const TSharedPtr<const FModel>&, const FParameters*, OP::ADDRESS) ;
+		MUTABLERUNTIME_API TManagedPtr<const FImage> BuildImage(const TSharedPtr<FExternalResourceProvider>& ExternalResourceProvider, const TSharedPtr<const FModel>&, const FParameters*, OP::ADDRESS, int32 MipsToSkip);
+		MUTABLERUNTIME_API TManagedPtr<const FMesh> BuildMesh(const TSharedPtr<FExternalResourceProvider>& ExternalResourceProvider, const TSharedPtr<const FModel>&, const FParameters*, OP::ADDRESS, EMeshContentFlags MeshContentFilter);
+		MUTABLERUNTIME_API TManagedPtr<const FInstance> BuildInstance(const TSharedPtr<FExternalResourceProvider>& ExternalResourceProvider, const TSharedPtr<const FModel>&, const FParameters*, OP::ADDRESS);
+		MUTABLERUNTIME_API TManagedPtr<const FLayout> BuildLayout(const TSharedPtr<FExternalResourceProvider>& ExternalResourceProvider, const TSharedPtr<const FModel>&, const FParameters*, OP::ADDRESS) ;
     	MUTABLERUNTIME_API FProjector BuildProjector(const TSharedPtr<FExternalResourceProvider>& ExternalResourceProvider, const TSharedPtr<const FModel>&, const FParameters*, OP::ADDRESS) ;
     	
         //!
diff --git a/Engine/Plugins/Mutable/Source/MutableRuntime/Internal/MuR/SystemPrivateTypes.h b/Engine/Plugins/Mutable/Source/MutableRuntime/Internal/MuR/SystemPrivateTypes.h
new file mode 100644
--- /dev/null
+++ b/Engine/Plugins/Mutable/Source/MutableRuntime/Internal/MuR/SystemPrivateTypes.h
@@ -0,0 +1,130 @@
+// Copyright Epic Games, Inc. All Rights Reserved.
+
+#pragma once
+
+#include "MuR/Operations.h"
+
+#ifndef MUTABLE_DEBUG_CODERUNNER_TASK_SCHEDULE_CALLSTACK
+	#define MUTABLE_DEBUG_CODERUNNER_TASK_SCHEDULE_CALLSTACK 0	
+#endif
+
+namespace UE::Mutable::Private
+{
+	/** This structure stores the data about an ongoing mutable operation that needs to be executed. */
+	struct FScheduledOp
+	{
+		inline FScheduledOp() = default;
+
+		inline FScheduledOp(OP::ADDRESS InAt, const FScheduledOp& InOpTemplate, uint8 InStage = 0, uint32 InCustomState = 0)
+			: At(InAt)
+			, CustomState(InCustomState)
+			, ExecutionIndex(InOpTemplate.ExecutionIndex)
+			, ExecutionOptions(InOpTemplate.ExecutionOptions)
+			, Stage(InStage)
+			, Type(InOpTemplate.Type)
+		{
+			check(InStage < ((1 << 6) - 1));
+		}
+
+		static inline FScheduledOp FromOpAndOptions(OP::ADDRESS InAt, const FScheduledOp& InOpTemplate, uint8 InExecutionOptions)
+		{
+			FScheduledOp Result;
+
+			Result.At = InAt;
+			Result.ExecutionOptions = InExecutionOptions;
+			Result.ExecutionIndex = InOpTemplate.ExecutionIndex;
+			Result.Stage = 0;
+			Result.CustomState = InOpTemplate.CustomState;
+			Result.Type = InOpTemplate.Type;
+
+			return Result;
+		}
+
+		//! Address of the operation
+		OP::ADDRESS At = 0;
+
+		//! Additional custom state data that the operation can store. This is usually used to pass information
+		//! between execution stages of an operation.
+		uint32 CustomState = 0;
+
+		//! Index of the operation execution: This is used for iteration of different ranges.
+		//! It is an index into the CodeRunner::GetMemory()::m_rangeIndex vector.
+		//! executionIndex 0 is always used for empty ExecutionIndex, which is the most common
+		//! one.
+		uint16 ExecutionIndex = 0;
+
+		//! Additional execution options. Set externally to this op, it usually alters the result.
+		//! For example, this is used to keep track of the mipmaps to skip in image operations.
+		uint8 ExecutionOptions = 0;
+
+		//! Internal stage of the operation.
+		//! Stage 0 is usually scheduling of children, and 1 is execution. Some instructions
+		//! may have more steges to schedule children that are optional for execution, etc.
+		uint8 Stage      : 6 = 0;
+
+		uint8 bForceLazy : 1 = 0;
+
+		//! Type of calculation we are requesting for this operation.
+		enum class EType : uint8
+		{
+			//! Execute the operation to calculate the full result
+			Full,
+
+			//! Execute the operation to obtain the descriptor of an image or mesh.
+			ImageDesc,
+		};
+		EType Type       : 1 = EType::Full;
+
+#if MUTABLE_DEBUG_CODERUNNER_TASK_SCHEDULE_CALLSTACK
+		static constexpr uint64 CallstackMaxDepth = 16;
+		uint64 StackDepth = 0; 
+		uint64 ScheduleCallstack[CallstackMaxDepth];
+#endif
+	};
+	
+	inline uint32 GetTypeHash(const FScheduledOp& Op)
+	{
+		return HashCombineFast(::GetTypeHash(Op.At), HashCombineFast((uint32)Op.Stage), (uint32)Op.ExecutionIndex);
+	}
+
+	/** A cache address is the operation plus the context of execution(iteration indices, etc...). */
+	struct FCacheAddress
+	{
+		/** The meaning of all these fields is the same than the FScheduledOp struct. */
+		OP::ADDRESS At = 0;
+		uint16 ExecutionIndex = 0;
+		uint8 ExecutionOptions = 0;
+		FScheduledOp::EType Type = FScheduledOp::EType::Full;
+
+		FCacheAddress() 
+		{
+		}
+
+		FCacheAddress(OP::ADDRESS InAt, uint16 InExecutionIndex, uint8 InExecutionOptions, FScheduledOp::EType InType = FScheduledOp::EType::Full)
+			: At(InAt), ExecutionIndex(InExecutionIndex), ExecutionOptions(InExecutionOptions), Type(InType)
+		{
+		}
+
+		FCacheAddress(OP::ADDRESS InAt, const FScheduledOp& Item)
+			: At(InAt), ExecutionIndex(Item.ExecutionIndex), ExecutionOptions(Item.ExecutionOptions), Type(Item.Type)
+		{
+		}
+		
+		FCacheAddress(const FScheduledOp& Item)
+			: At(Item.At), ExecutionIndex(Item.ExecutionIndex), ExecutionOptions(Item.ExecutionOptions), Type(Item.Type)
+		{
+		}
+	};
+
+	inline uint32 GetTypeHash(const FCacheAddress& Address)
+	{
+		uint32 Hash = Address.At;
+
+		Hash = HashCombineFast(Hash, static_cast<uint32>(Address.ExecutionIndex));
+		Hash = HashCombineFast(Hash, static_cast<uint32>(Address.ExecutionOptions));
+		Hash = HashCombineFast(Hash, static_cast<uint32>(Address.Type));
+
+		return Hash;
+	}
+
+} // namespace UE::Mutable::Private
