[Mutable] Memory managment rework.

diff --git a/Engine/Plugins/Mutable/Source/MutableRuntime/Private/MuR/CodeRunner.cpp b/Engine/Plugins/Mutable/Source/MutableRuntime/Private/MuR/CodeRunner.cpp
--- a/Engine/Plugins/Mutable/Source/MutableRuntime/Private/MuR/CodeRunner.cpp
+++ b/Engine/Plugins/Mutable/Source/MutableRuntime/Private/MuR/CodeRunner.cpp
@@ -222,7 +222,7 @@ namespace UE::Mutable::Private
 	}
 
 
-	TTuple<UE::Tasks::FTask, TFunction<void()>> CodeRunner::LoadExternalImageAsync(FExternalResourceId Id, uint8 MipmapsToSkip, TFunction<void(TSharedPtr<FImage>)>& ResultCallback)
+	TTuple<UE::Tasks::FTask, TFunction<void()>> CodeRunner::LoadExternalImageAsync(FExternalResourceId Id, uint8 MipmapsToSkip, TFunction<void(TManagedPtr<FImage>)>& ResultCallback)
     {
 		MUTABLE_CPUPROFILER_SCOPE(LoadExternalImageAsync);
 		
@@ -249,7 +249,7 @@ namespace UE::Mutable::Private
 	}
 
 
-	TTuple<UE::Tasks::FTask, TFunction<void()>> CodeRunner::LoadExternalMeshAsync(FExternalResourceId Id, int32 LODIndex, int32 SectionIndex, TFunction<void(TSharedPtr<FMesh>)>& ResultCallback)
+	TTuple<UE::Tasks::FTask, TFunction<void()>> CodeRunner::LoadExternalMeshAsync(FExternalResourceId Id, int32 LODIndex, int32 SectionIndex, TFunction<void(TManagedPtr<FMesh>)>& ResultCallback)
 	{
 		MUTABLE_CPUPROFILER_SCOPE(LoadExternalImageAsync);
 
@@ -414,8 +414,8 @@ namespace UE::Mutable::Private
 				case EDataType::Image:			StoreImage(Item, nullptr); break;
 				case EDataType::Layout:			StoreLayout(Item, nullptr); break;
 				case EDataType::Instance:		StoreInstance(Item, nullptr); break;
-				case EDataType::ExtensionData:	StoreExtensionData(Item, MakeShared<FExtensionData>()); break;
-				case EDataType::Material:		StoreMaterial(Item, MakeShared<FMaterial>()); break;
+				case EDataType::ExtensionData:	StoreExtensionData(Item, MakeManaged<FExtensionData>()); break;
+				case EDataType::Material:		StoreMaterial(Item, MakeManaged<FMaterial>()); break;
 				default:
 					// Not implemented
 					check(false);
@@ -539,21 +539,21 @@ namespace UE::Mutable::Private
 
             case 1:
             {
-                TSharedPtr<const FInstance> Base =
-                    LoadInstance( FCacheAddress( Args.instance, Item ) );
-                TSharedPtr<FInstance> Result;
-                if ( !Base )
-                {
-                    Result = MakeShared<FInstance>();
-                }
-                else
-                {
-					Result = UE::Mutable::Private::CloneOrTakeOver<FInstance>(Base);
+                TManagedPtr<const FInstance> Base = LoadInstance( FCacheAddress( Args.instance, Item ) );
+				TManagedPtr<FInstance> Result = nullptr;
+
+				if (!Base)
+				{
+					Result = MakeManaged<FInstance>();
+				}
+				else
+				{
+					Result = UE::Mutable::Private::CloneOrTakeOver<FInstance>(MoveTemp(Base));
 				}
 
                 if ( Args.value )
                 {
-                    TSharedPtr<const String> value = LoadString( FCacheAddress( Args.value, Item ) );
+                    TManagedPtr<const String> value = LoadString( FCacheAddress( Args.value, Item ) );
 
                     OP::ADDRESS nameAd = Args.name;
                     check( nameAd < (uint32)Program.ConstantStrings.Num() );
@@ -585,20 +585,21 @@ namespace UE::Mutable::Private
 
             case 1:
             {
-				TSharedPtr<const FInstance> Base = LoadInstance( FCacheAddress(Args.instance,Item) );
-				TSharedPtr<FInstance> Result;
-                if (!Base)
-                {
-                    Result = MakeShared<FInstance>();
-                }
-                else
-                {
-					Result = UE::Mutable::Private::CloneOrTakeOver<FInstance>(Base);
+				TManagedPtr<const FInstance> Base = LoadInstance( FCacheAddress(Args.instance,Item) );
+				TManagedPtr<FInstance> Result = nullptr;
+
+				if (!Base)
+				{
+					Result = MakeManaged<FInstance>();
+				}
+				else
+				{
+					Result = UE::Mutable::Private::CloneOrTakeOver<FInstance>(MoveTemp(Base));
 				}
 
                 if ( Args.value )
                 {
-                    TSharedPtr<const FInstance> pComp = LoadInstance( FCacheAddress(Args.value,Item) );
+                    TManagedPtr<const FInstance> pComp = LoadInstance( FCacheAddress(Args.value,Item) );
 
                     int32 NewComponentIndex = Result->AddComponent();
 
@@ -634,23 +635,23 @@ namespace UE::Mutable::Private
 
             case 1:
             {
-				TSharedPtr<const FInstance> Base = LoadInstance( FCacheAddress(Args.instance,Item) );
+				TManagedPtr<const FInstance> Base = LoadInstance( FCacheAddress(Args.instance,Item) );
+				TManagedPtr<FInstance> Result = nullptr;
 
-				TSharedPtr<FInstance> Result;
 				if (Base)
 				{
-					Result = UE::Mutable::Private::CloneOrTakeOver<FInstance>(Base);
+					Result = UE::Mutable::Private::CloneOrTakeOver<FInstance>(MoveTemp(Base));
 				}
 				else
 				{
-					Result = MakeShared<FInstance>();
+					Result = MakeManaged<FInstance>();
 				}
 
                 // Empty surfaces are ok, they still need to be created, because they may contain
                 // additional information like internal or external IDs
                 //if ( Args.value )
                 {
-					TSharedPtr<const FInstance> pSurf = LoadInstance( FCacheAddress(Args.value,Item) );
+					TManagedPtr<const FInstance> pSurf = LoadInstance( FCacheAddress(Args.value,Item) );
 
                     int32 sindex = Result->AddSurface( 0, 0 );
 
@@ -726,7 +727,7 @@ namespace UE::Mutable::Private
             case 1:
             {
                 // Assemble result
-				TSharedPtr<FInstance> Result = MakeShared<FInstance>();
+				TManagedPtr<FInstance> Result = MakeManaged<FInstance>();
 				int32 ComponentIndex = Result->AddComponent();
 
                 for (uint8 LODIndex = 0; LODIndex < LODCount; ++LODIndex)
@@ -744,12 +745,10 @@ namespace UE::Mutable::Private
 						
 						if (bIsSelectedLod)
                         {
-							TSharedPtr<const FInstance> pLOD = LoadInstance( FCacheAddress(LODAddress,Item) );
+							TManagedPtr<const FInstance> pLOD = LoadInstance(FCacheAddress(LODAddress, Item));
 
                             // In a degenerated case, the returned pLOD may not have an LOD inside
- 							if (!pLOD->Components.IsEmpty()
-								&&
-								!pLOD->Components[0].LODs.IsEmpty())
+ 							if (!pLOD->Components.IsEmpty() && !pLOD->Components[0].LODs.IsEmpty())
 							{
 								Result->Components[ComponentIndex].LODs[InstanceLODIndex] = pLOD->Components[0].LODs[0];
 							}
@@ -789,23 +788,23 @@ namespace UE::Mutable::Private
 				case 1:
 				{
 					// Assemble result
-					TSharedPtr<FInstance> Instance;
+					TManagedPtr<FInstance> Instance;
 					if (Args.Instance)
 					{
 						Instance = UE::Mutable::Private::CloneOrTakeOver<FInstance>(LoadInstance(FCacheAddress(Args.Instance, Item)));
 					}
 					else
 					{
-						Instance = MakeShared<FInstance>();
+						Instance = MakeManaged<FInstance>();
 					}
 
-					if (TSharedPtr<const FExtensionData> ExtensionData = LoadExtensionData(FCacheAddress(Args.ExtensionData, Item)))
+					if (TManagedPtr<const FExtensionData> ExtensionData = LoadExtensionData(FCacheAddress(Args.ExtensionData, Item)))
 					{
 						const OP::ADDRESS NameAddress = Args.ExtensionDataName;
 						check(NameAddress < (uint32)Program.ConstantStrings.Num());
 						const FString& NameString = Program.ConstantStrings[NameAddress];
 
-						Instance->SetExtensionData(ExtensionData.ToSharedRef(), FName(NameString));
+						Instance->SetExtensionData(ExtensionData.ToStrongRef(), FName(NameString));
 					}
 
 					StoreInstance(Item, Instance);
@@ -852,15 +851,15 @@ namespace UE::Mutable::Private
 
             case 1:
             {
-                TSharedPtr<const FInstance> Base = LoadInstance( FCacheAddress(Args.instance,Item) );
-                TSharedPtr<FInstance> Result;
-                if (!Base)
-                {
-                    Result = MakeShared<FInstance>();
-                }
-                else
-                {
-					Result = UE::Mutable::Private::CloneOrTakeOver<FInstance>(Base);
+                TManagedPtr<const FInstance> Base = LoadInstance( FCacheAddress(Args.instance,Item) );
+				TManagedPtr<FInstance> Result = nullptr;
+				if (!Base)
+				{
+					Result = MakeManaged<FInstance>();
+				}
+				else
+				{
+					Result = UE::Mutable::Private::CloneOrTakeOver<FInstance>(MoveTemp(Base));
 				}
 
                 if ( Args.value )
@@ -897,16 +896,16 @@ namespace UE::Mutable::Private
 
         		case 1:
         			{
-        				TSharedPtr<const FInstance> Base = LoadInstance(FCacheAddress(Args.Instance, Item));
-        				TSharedPtr<FInstance> Result;
-        				if (!Base)
-        				{
-        					Result = MakeShared<FInstance>();
-        				}
-        				else
-        				{
-        					Result = UE::Mutable::Private::CloneOrTakeOver<FInstance>(Base);
-        				}
+        				TManagedPtr<const FInstance> Base = LoadInstance(FCacheAddress(Args.Instance, Item));
+						TManagedPtr<FInstance> Result = nullptr;
+						if (!Base)
+						{
+							Result = MakeManaged<FInstance>();
+						}
+						else
+						{
+							Result = UE::Mutable::Private::CloneOrTakeOver<FInstance>(MoveTemp(Base));
+						}
 
         				if (Args.Material)
         				{
@@ -941,15 +940,15 @@ namespace UE::Mutable::Private
 
         		case 1:
         			{
-        				TSharedPtr<const FInstance> Base = LoadInstance(FCacheAddress(Args.Instance, Item));
-        				TSharedPtr<FInstance> Result;
+        				TManagedPtr<const FInstance> Base = LoadInstance(FCacheAddress(Args.Instance, Item));
+        				TManagedPtr<FInstance> Result;
         				if (!Base)
         				{
-        					Result = MakeShared<FInstance>();
+        					Result = MakeManaged<FInstance>();
         				}
         				else
         				{
-        					Result = UE::Mutable::Private::CloneOrTakeOver<FInstance>(Base);
+        					Result = UE::Mutable::Private::CloneOrTakeOver<FInstance>(MoveTemp(Base));
         				}
 
         				if (Args.Material)
@@ -983,17 +982,17 @@ namespace UE::Mutable::Private
 
         		case 1:
         			{
-        				TSharedPtr<const FInstance> Base = LoadInstance(FCacheAddress(Args.Instance, Item));
-						TSharedPtr<const FMaterial> Material = LoadMaterial(FCacheAddress(Args.Material, Item));
+        				TManagedPtr<const FInstance> Base = LoadInstance(FCacheAddress(Args.Instance, Item));
+						TManagedPtr<const FMaterial> Material = LoadMaterial(FCacheAddress(Args.Material, Item));
 
-        				TSharedPtr<FInstance> Result;
+        				TManagedPtr<FInstance> Result;
         				if (Base)
         				{
-        					Result = UE::Mutable::Private::CloneOrTakeOver<FInstance>(Base);
+        					Result = UE::Mutable::Private::CloneOrTakeOver<FInstance>(MoveTemp(Base));
         				}
         				else
         				{
-        					Result = MakeShared<FInstance>();
+        					Result = MakeManaged<FInstance>();
         				}
 
         				if (Args.Material)
@@ -1029,20 +1028,20 @@ namespace UE::Mutable::Private
 
         		case 1:
         			{
-        				TSharedPtr<const FInstance> InInstance;
+        				TManagedPtr<const FInstance> InInstance;
         				if (Args.instance)
         				{
         					InInstance = LoadInstance(FCacheAddress(Args.instance, Item));
         				}
 
-        				TSharedPtr<FInstance> Instance;
+        				TManagedPtr<FInstance> Instance;
         				if (InInstance)
         				{
-        					Instance = UE::Mutable::Private::CloneOrTakeOver<FInstance>(InInstance);
+        					Instance = UE::Mutable::Private::CloneOrTakeOver<FInstance>(MoveTemp(InInstance));
         				}
         				else
         				{
-        					Instance = MakeShared<FInstance>();
+        					Instance = MakeManaged<FInstance>();
         				}
 
 
@@ -1085,7 +1084,7 @@ namespace UE::Mutable::Private
 			OP::MeshConstantArgs Args = Program.GetOpArgs<OP::MeshConstantArgs>(Item.At);
 			EMeshContentFlags MeshContentFlags = (EMeshContentFlags)Item.ExecutionOptions;
 
-			TSharedPtr<const FMesh> Source;
+			TManagedPtr<const FMesh> Source;
 			Program.GetConstant(Args.Value, Args.Skeleton, Source, MeshContentFlags, 
 					[this](int32 BudgetReserve){ return CreateMesh(BudgetReserve); });
 			
@@ -1105,7 +1104,7 @@ namespace UE::Mutable::Private
 			OP::ResourceConstantArgs Args = Program.GetOpArgs<OP::ResourceConstantArgs>(Item.At);
 
 			int32 MipsToSkip = Item.ExecutionOptions;
-            TSharedPtr<const FImage> Source;
+            TManagedPtr<const FImage> Source;
 			Program.GetConstant(Args.value, Source, MipsToSkip, 
 			[this](int32 SizeX, int32 SizeY, int32 NumLODs, EImageFormat Format, EInitializationType InitPolicy)
 			{
@@ -1118,7 +1117,7 @@ namespace UE::Mutable::Private
 				return false;
 			}
 
-            StoreImage( Item, Source );
+            StoreImage(Item, Source);
 			//UE_LOG(LogMutableCore, Log, TEXT("Set image constant %d."), Item.At);
 			break;
         }
@@ -1128,7 +1127,7 @@ namespace UE::Mutable::Private
 			OP::ResourceConstantArgs Args = Program.GetOpArgs<OP::ResourceConstantArgs>(Item.At);
 
 			// Assume the ROM has been loaded previously
-			TSharedPtr<const UE::Mutable::Private::FExtensionData> SourceConst;
+			TManagedPtr<const UE::Mutable::Private::FExtensionData> SourceConst;
 			Program.GetExtensionDataConstant(Args.value, SourceConst);
 
 			check(SourceConst);
@@ -1141,10 +1140,10 @@ namespace UE::Mutable::Private
 		{
 			OP::MaterialConstantArgs Args = Program.GetOpArgs<OP::MaterialConstantArgs>(Item.At);
 
-        	TSharedPtr<const FMaterial> SourceConst = Program.ConstantMaterials[Args.Value];
+        	TManagedPtr<const FMaterial> SourceConst = Program.ConstantMaterials[Args.Value];
        		check(SourceConst);
         		
-			TSharedPtr<FMaterial> Source = UE::Mutable::Private::CloneOrTakeOver<FMaterial>(SourceConst);
+			TManagedPtr<FMaterial> Source = UE::Mutable::Private::CloneOrTakeOver<FMaterial>(MoveTemp(SourceConst));
         	Source->PassthroughObject = FindPassthrough<UMaterialInterface>(Args.ID);
         		
 			StoreMaterial(Item, Source);
@@ -1203,22 +1202,22 @@ namespace UE::Mutable::Private
             		
 				if (!EnumHasAnyFlags(EMeshContentFlags(Item.ExecutionOptions), EMeshContentFlags::GeometryData))
 				{
-					TSharedPtr<const FMesh> Base = LoadMesh(FCacheAddress(Args.Mesh, Item));
+					TManagedPtr<const FMesh> Base = LoadMesh(FCacheAddress(Args.Mesh, Item));
 					StoreMesh(Item, Base);
 				}
 				else
 				{
-					TSharedPtr<const FMesh> Base = LoadMesh(FCacheAddress(Args.Mesh, Item));
-					TSharedPtr<const FLayout> Layout = LoadLayout(FCacheAddress(Args.Layout, Item));
+					TManagedPtr<const FMesh> Base = LoadMesh(FCacheAddress(Args.Mesh, Item));
+					TManagedPtr<const FLayout> Layout = LoadLayout(FCacheAddress(Args.Layout, Item));
 
 					if (Base)
 					{
-						TSharedPtr<FMesh> Result = CloneOrTakeOver(Base);
+						TManagedPtr<FMesh> Result = CloneOrTakeOver(MoveTemp(Base));
 
 						int32 texCoordsSet = Args.Channel;
 
 						MeshApplyLayout(Result.Get(), Layout.Get(), texCoordsSet);
-						
+					
 						StoreMesh(Item, Result);
 					}
 					else
@@ -1264,20 +1263,20 @@ namespace UE::Mutable::Private
 
 				if (!EnumHasAnyFlags(EMeshContentFlags(Item.ExecutionOptions), EMeshContentFlags::GeometryData))
 				{
-					TSharedPtr<const FMesh> Base = LoadMesh(FCacheAddress(Args.Mesh, Item));
+					TManagedPtr<const FMesh> Base = LoadMesh(FCacheAddress(Args.Mesh, Item));
 					StoreMesh(Item, Base);
 				}
 				else
 				{
-					TSharedPtr<const FMesh> Base = LoadMesh(FCacheAddress(Args.Mesh, Item));
-					TSharedPtr<const FLayout> Layout = LoadLayout(FCacheAddress(Args.Layout, Item));
+					TManagedPtr<const FMesh> Base = LoadMesh(FCacheAddress(Args.Mesh, Item));
+					TManagedPtr<const FLayout> Layout = LoadLayout(FCacheAddress(Args.Layout, Item));
 
 					if (Base && Layout)
 					{
-						TSharedPtr<FMesh> Result = CloneOrTakeOver(Base);
+						TManagedPtr<FMesh> Result = CloneOrTakeOver(MoveTemp(Base));
 
 						MeshPrepareLayout(*Result, *Layout, Args.LayoutChannel, Args.bNormalizeUVs, Args.bClampUVIslands, Args.bEnsureAllVerticesHaveLayoutBlock, Args.bUseAbsoluteBlockIds);
-
+					
 						StoreMesh(Item, Result);
 					}
 					else
@@ -1314,7 +1313,7 @@ namespace UE::Mutable::Private
 			{
 				if (!EnumHasAnyFlags(EMeshContentFlags(Item.ExecutionOptions), EMeshContentFlags::GeometryData))
 				{
-                	TSharedPtr<const FMesh> Result = CreateMesh();
+                	TManagedPtr<const FMesh> Result = CreateMesh();
 					StoreMesh(Item, Result);
 				}
 				else
@@ -1322,8 +1321,8 @@ namespace UE::Mutable::Private
 					if (BaseAt && TargetAt)
 					{
 						AddOp(FScheduledOp(Item.At, Item, 1),
-							FScheduledOp(BaseAt, Item),
-							FScheduledOp(TargetAt, Item));
+								FScheduledOp(BaseAt, Item),
+								FScheduledOp(TargetAt, Item));
 					}
 					else
 					{
@@ -1338,8 +1337,8 @@ namespace UE::Mutable::Private
 
 				check(EnumHasAnyFlags(EMeshContentFlags(Item.ExecutionOptions), EMeshContentFlags::GeometryData))
 
-				TSharedPtr<const FMesh> Base = LoadMesh(FCacheAddress(BaseAt,Item));
-				TSharedPtr<const FMesh> pTarget = LoadMesh(FCacheAddress(TargetAt,Item));
+				TManagedPtr<const FMesh> Base = LoadMesh(FCacheAddress(BaseAt, Item));
+				TManagedPtr<const FMesh> pTarget = LoadMesh(FCacheAddress(TargetAt, Item));
 
 				TArray<EMeshBufferSemantic, TInlineAllocator<8>> Semantics;
 				TArray<int32, TInlineAllocator<8>> SemanticIndices;
@@ -1366,13 +1365,11 @@ namespace UE::Mutable::Private
 					SemanticIndices.Add(SemanticIndex);
 				}
 
-				TSharedPtr<FMesh> Result = CreateMesh();
+				TManagedPtr<FMesh> Result = CreateMesh();
 				bool bOutSuccess = false;
 				MeshDifference(Result.Get(), Base.Get(), pTarget.Get(),
 							   NumChannels, Semantics.GetData(), SemanticIndices.GetData(),
 							   bIgnoreTextureCoords != 0, bOutSuccess);
-				Release(Base);
-				Release(pTarget);
 
 				StoreMesh(Item, Result);
 
@@ -1427,7 +1424,6 @@ namespace UE::Mutable::Private
 				{
 					StoreMesh(Item, nullptr);
 				}
-				
                 break;
 			}
             case 1:
@@ -1436,7 +1432,7 @@ namespace UE::Mutable::Private
 
 				if (!EnumHasAnyFlags(EMeshContentFlags(Item.ExecutionOptions), EMeshContentFlags::GeometryData))
 				{
-                	TSharedPtr<const FMesh> Base = LoadMesh(FCacheAddress(BaseAt, Item));
+                	TManagedPtr<const FMesh> Base = LoadMesh(FCacheAddress(BaseAt, Item));
 					StoreMesh(Item, Base);
 				}
 				else	
@@ -1472,7 +1468,7 @@ namespace UE::Mutable::Private
        		    MUTABLE_CPUPROFILER_SCOPE(ME_MORPH_2)
 				check(EnumHasAnyFlags(EMeshContentFlags(Item.ExecutionOptions), EMeshContentFlags::GeometryData))
 				
-				TSharedPtr<const FMesh> BaseMesh = LoadMesh(FCacheAddress(BaseAt, Item));
+				TManagedPtr<const FMesh> BaseMesh = LoadMesh(FCacheAddress(BaseAt, Item));
 
 				// Factor from 0 to 1 between the two targets
 				const FScheduledOpData& OpHeapData = HeapData[Item.CustomState];
@@ -1486,9 +1482,9 @@ namespace UE::Mutable::Private
 					}
 					else 
 					{
-						TSharedPtr<const FMesh> MorphMesh = LoadMesh(FCacheAddress(TargetAt, Item));
-						
-						TSharedPtr<FMesh> Result = CloneOrTakeOver(BaseMesh);
+						TManagedPtr<const FMesh> MorphMesh = LoadMesh(FCacheAddress(TargetAt, Item));
+
+						TManagedPtr<FMesh> Result = CloneOrTakeOver(MoveTemp(BaseMesh));
 						if (Result->HasMorphs())
 						{
 							// Morph data from the Base Mesh
@@ -1502,7 +1498,6 @@ namespace UE::Mutable::Private
 						}
 
 						StoreMesh(Item, Result);
-						Release(MorphMesh);
 					}
 				}
 				else
@@ -1537,8 +1532,8 @@ namespace UE::Mutable::Private
             {
 				MUTABLE_CPUPROFILER_SCOPE(ME_MERGE_1)
 
-                TSharedPtr<const FMesh> pA = LoadMesh(FCacheAddress(Args.Base, Item));
-                TSharedPtr<const FMesh> pB = LoadMesh(FCacheAddress(Args.Added, Item));
+                TManagedPtr<const FMesh> pA = LoadMesh(FCacheAddress(Args.Base, Item));
+                TManagedPtr<const FMesh> pB = LoadMesh(FCacheAddress(Args.Added, Item));
 
                 if (pA && pB && pA->GetVertexCount() && pB->GetVertexCount())
                 {
@@ -1548,31 +1543,27 @@ namespace UE::Mutable::Private
 					Scratch.FirstReformat = CreateMesh();
 					Scratch.SecondReformat = CreateMesh();
 
-					TSharedPtr<FMesh> Result = CreateMesh(pA->GetDataSize() + pB->GetDataSize());
+					TManagedPtr<FMesh> Result = CreateMesh(pA->GetDataSize() + pB->GetDataSize());
 
 					MeshMerge(Result.Get(), pA.Get(), pB.Get(), !Args.NewSurfaceID, Scratch);
 
-					Release(Scratch.FirstReformat);
-					Release(Scratch.SecondReformat);
-
+					Scratch.FirstReformat =  nullptr;
+					Scratch.SecondReformat =  nullptr;
                     if (Args.NewSurfaceID)
                     {
 						check(pB->GetSurfaceCount() == 1);
 						Result->Surfaces.Last().Id = Args.NewSurfaceID;
                     }
 
-					Release(pA);
-					Release(pB);
 					StoreMesh(Item, Result);
                 }
                 else if (pA && (pA->GetVertexCount() || pA->IsReference()))
                 {
-					Release(pB);
 					StoreMesh(Item, pA);
                 }
                 else if (pB && (pB->GetVertexCount() || pB->IsReference()))
                 {
-					TSharedPtr<FMesh> Result = CloneOrTakeOver(pB);
+					TManagedPtr<FMesh> Result = CloneOrTakeOver(MoveTemp(pB));
 
                     check(Result->IsReference() || (Result->GetSurfaceCount() == 1) );
 
@@ -1581,13 +1572,10 @@ namespace UE::Mutable::Private
                         Result->Surfaces.Last().Id = Args.NewSurfaceID;
                     }
 
-					Release(pA);
 					StoreMesh(Item, Result);
                 }
                 else
                 {
-					Release(pA);
-					Release(pB);
 					StoreMesh(Item, CreateMesh());
                 }
 
@@ -1626,27 +1614,24 @@ namespace UE::Mutable::Private
             	MUTABLE_CPUPROFILER_SCOPE(ME_MASKCLIPMESH_1)
 				if (!EnumHasAnyFlags(EMeshContentFlags(Item.ExecutionOptions), EMeshContentFlags::GeometryData))
 				{
-					TSharedPtr<const FMesh> Source = LoadMesh(FCacheAddress(Args.source, Item));
+					TManagedPtr<const FMesh> Source = LoadMesh(FCacheAddress(Args.source, Item));
 					StoreMesh(Item, Source);
 				}
 				else
 				{
-					TSharedPtr<const FMesh> Source = LoadMesh(FCacheAddress(Args.source, Item));
-					TSharedPtr<const FMesh> pClip = LoadMesh(FCacheAddress(Args.clip, Item));
+					TManagedPtr<const FMesh> Source = LoadMesh(FCacheAddress(Args.source, Item));
+					TManagedPtr<const FMesh> pClip = LoadMesh(FCacheAddress(Args.clip, Item));
 
 					// Only if both are valid.
-					if (Source.Get() && pClip.Get())
+					if (Source && pClip)
 					{
-						TSharedPtr<FMesh> Result = CreateMesh();
-
+						TManagedPtr<FMesh> Result = CreateMesh();
+					
 						bool bOutSuccess = false;
 						MeshMaskClipMesh(Result.Get(), Source.Get(), pClip.Get(), bOutSuccess);
-						
-						Release(Source);
-						Release(pClip);
+					
 						if (!bOutSuccess)
 						{
-							Release(Result);
 							StoreMesh(Item, nullptr);
 						}
 						else
@@ -1656,10 +1641,9 @@ namespace UE::Mutable::Private
 					}
 					else
 					{
-						Release(Source);
-						Release(pClip);
 						StoreMesh(Item, nullptr);
 					}
+					
 				}
                 break;
             }
@@ -1699,30 +1683,26 @@ namespace UE::Mutable::Private
 
 				if (!EnumHasAnyFlags(EMeshContentFlags(Item.ExecutionOptions), EMeshContentFlags::GeometryData))
 				{
-					TSharedPtr<const FMesh> Source = LoadMesh(FCacheAddress(Args.Source, Item));
+					TManagedPtr<const FMesh> Source = LoadMesh(FCacheAddress(Args.Source, Item));
 					StoreMesh(Item, Source);
 				}
 				else
 				{
-					TSharedPtr<const FMesh> Source = LoadMesh(FCacheAddress(Args.Source, Item));
-					TSharedPtr<const FMesh> UVSource = LoadMesh(FCacheAddress(Args.UVSource, Item));
-					TSharedPtr<const FImage> MaskImage = LoadImage(FCacheAddress(Args.MaskImage, Item));
-					TSharedPtr<const FLayout> MaskLayout = LoadLayout(FCacheAddress(Args.MaskLayout, Item));
+					TManagedPtr<const FMesh> Source = LoadMesh(FCacheAddress(Args.Source, Item));
+					TManagedPtr<const FMesh> UVSource = LoadMesh(FCacheAddress(Args.UVSource, Item));
+					TManagedPtr<const FImage> MaskImage = LoadImage(FCacheAddress(Args.MaskImage, Item));
+					TManagedPtr<const FLayout> MaskLayout = LoadLayout(FCacheAddress(Args.MaskLayout, Item));
 
 					// Only if both are valid.
-					if (Source.Get() && MaskImage.Get())
+					if (Source && MaskImage)
 					{
-						TSharedPtr<FMesh> Result = CreateMesh();
+						TManagedPtr<FMesh> Result = CreateMesh();
 
 						bool bOutSuccess = false;
 						MakeMeshMaskFromUVMask(Result.Get(), Source.Get(), UVSource.Get(), MaskImage.Get(), Args.LayoutIndex, bOutSuccess);
-
-						Release(Source);
-						Release(UVSource);
-						Release(MaskImage);
+						
 						if (!bOutSuccess)
 						{
-							Release(Result);
 							StoreMesh(Item, nullptr);
 						}
 						else
@@ -1730,19 +1710,15 @@ namespace UE::Mutable::Private
 							StoreMesh(Item, Result);
 						}
 					}
-					else if (Source.Get() && MaskLayout.Get())
+					else if (Source && MaskLayout)
 					{					
-						TSharedPtr<FMesh> Result = CreateMesh();
+						TManagedPtr<FMesh> Result = CreateMesh();
 
 						bool bOutSuccess = false;
 						MakeMeshMaskFromLayout(Result.Get(), Source.Get(), UVSource.Get(), MaskLayout.Get(), Args.LayoutIndex, bOutSuccess);
 
-						Release(Source);
-						Release(UVSource);
-						Release(MaskImage);
 						if (!bOutSuccess)
 						{
-							Release(Result);
 							StoreMesh(Item, nullptr);
 						}
 						else
@@ -1752,9 +1728,6 @@ namespace UE::Mutable::Private
 					}
 					else
 					{
-						Release(Source);
-						Release(UVSource);
-						Release(MaskImage);
 						StoreMesh(Item, nullptr);
 					}
 				}
@@ -1794,28 +1767,24 @@ namespace UE::Mutable::Private
 				MUTABLE_CPUPROFILER_SCOPE(ME_MASKDIFF_1)
 				if (!EnumHasAnyFlags(EMeshContentFlags(Item.ExecutionOptions), EMeshContentFlags::GeometryData))
 				{
-					TSharedPtr<const FMesh> Source = LoadMesh(FCacheAddress(Args.Source, Item));
+					TManagedPtr<const FMesh> Source = LoadMesh(FCacheAddress(Args.Source, Item));
 					StoreMesh(Item, Source);
 				}
 				else
 				{
-					TSharedPtr<const FMesh> Source = LoadMesh(FCacheAddress(Args.Source, Item));
-					TSharedPtr<const FMesh> pClip = LoadMesh(FCacheAddress(Args.Fragment, Item));
+					TManagedPtr<const FMesh> Source = LoadMesh(FCacheAddress(Args.Source, Item));
+					TManagedPtr<const FMesh> pClip = LoadMesh(FCacheAddress(Args.Fragment, Item));
 
 					// Only if both are valid.
-					if (Source.Get() && pClip.Get())
+					if (Source && pClip)
 					{
-						TSharedPtr<FMesh> Result = CreateMesh();
+						TManagedPtr<FMesh> Result = CreateMesh();
 
 						bool bOutSuccess = false;
 						MeshMaskDiff(Result.Get(), Source.Get(), pClip.Get(), bOutSuccess);
-
-						Release(Source);
-						Release(pClip);
-
+						
 						if (!bOutSuccess)
 						{
-							Release(Result);
 							StoreMesh(Item, nullptr);
 						}
 						else
@@ -1825,8 +1794,6 @@ namespace UE::Mutable::Private
 					}
 					else
 					{
-						Release(Source);
-						Release(pClip);
 						StoreMesh(Item, nullptr);
 					}
 				}
@@ -1850,8 +1817,8 @@ namespace UE::Mutable::Private
 				if (Args.source && Args.format)
 				{
 					AddOp(FScheduledOp(Item.At, Item, 1),
-						FScheduledOp(Args.source, Item),
-						FScheduledOp(Args.format, Item));
+							FScheduledOp(Args.source, Item),
+							FScheduledOp(Args.format, Item));
 				}
 				else
 				{
@@ -1862,12 +1829,11 @@ namespace UE::Mutable::Private
             case 1:
             {
 				MUTABLE_CPUPROFILER_SCOPE(ME_FORMAT_1)
-				TSharedPtr<const FMesh> Source = LoadMesh(FCacheAddress(Args.source,Item));
-				TSharedPtr<const FMesh> Format = LoadMesh(FCacheAddress(Args.format,Item));
+				TManagedPtr<const FMesh> Source = LoadMesh(FCacheAddress(Args.source,Item));
+				TManagedPtr<const FMesh> Format = LoadMesh(FCacheAddress(Args.format,Item));
 
 				if (Source && Source->IsReference())
 				{
-					Release(Format);
 					StoreMesh(Item, Source);
 				}
 				else if (Source)
@@ -1879,7 +1845,7 @@ namespace UE::Mutable::Private
 					}
 					else if (!Format)
 					{
-						TSharedPtr<FMesh> Result = CloneOrTakeOver(Source);
+						TManagedPtr<FMesh> Result = CloneOrTakeOver(MoveTemp(Source));
 
 						if (Flags & OP::MeshFormatArgs::ResetBufferIndices)
 						{
@@ -1890,7 +1856,7 @@ namespace UE::Mutable::Private
 					}
 					else
 					{
-						TSharedPtr<FMesh> Result = CreateMesh();
+						TManagedPtr<FMesh> Result = CreateMesh();
 
 						bool bOutSuccess = false;
 						MeshFormat(Result.Get(), Source.Get(), Format.Get(),
@@ -1899,7 +1865,6 @@ namespace UE::Mutable::Private
 							(Flags & OP::MeshFormatArgs::Index) != 0,
 							(Flags & OP::MeshFormatArgs::IgnoreMissing) != 0,
 							bOutSuccess);
-
 						check(bOutSuccess);
 
 						if (Flags & OP::MeshFormatArgs::ResetBufferIndices)
@@ -1913,14 +1878,11 @@ namespace UE::Mutable::Private
 							MeshOptimizeBuffers(Result.Get());
 						}
 
-						Release(Source);
-						Release(Format);
 						StoreMesh(Item, Result);
 					}
 				}
 				else
 				{
-					Release(Format);
 					StoreMesh(Item, nullptr);
 				}
                 break;
@@ -1956,7 +1918,7 @@ namespace UE::Mutable::Private
 				if (source)
 				{
 					AddOp(FScheduledOp(Item.At, Item, 1),
-						FScheduledOp(source, Item));
+							FScheduledOp(source, Item));
 				}
 				else
 				{
@@ -1970,12 +1932,12 @@ namespace UE::Mutable::Private
 
 				if (!EnumHasAnyFlags(EMeshContentFlags(Item.ExecutionOptions), EMeshContentFlags::GeometryData))
 				{
-					TSharedPtr<const FMesh> SourceMesh = LoadMesh(FCacheAddress(source, Item));
+					TManagedPtr<const FMesh> SourceMesh = LoadMesh(FCacheAddress(source, Item));
 					StoreMesh(Item, SourceMesh);
 				}
 				else
 				{
-					TSharedPtr<const FMesh> Source = LoadMesh(FCacheAddress(source, Item));
+					TManagedPtr<const FMesh> Source = LoadMesh(FCacheAddress(source, Item));
 
 					// Access with memcpy necessary for unaligned memory access issues.
 					uint64 blocks[512];
@@ -1984,7 +1946,7 @@ namespace UE::Mutable::Private
 
 					if (Source)
 					{
-						TSharedPtr<FMesh> Result = CreateMesh();
+						TManagedPtr<FMesh> Result = CreateMesh();
 						bool bOutSuccess;
 
 						if (blockCount > 0)
@@ -1995,17 +1957,16 @@ namespace UE::Mutable::Private
 						{
 							MeshExtractLayoutBlock(Result.Get(), Source.Get(), LayoutIndex, bOutSuccess);
 						}
-
+						
 						if (!bOutSuccess)
 						{
-							Release(Result);
 							StoreMesh(Item, Source);
 						}
 						else
 						{
-							Release(Source);
 							StoreMesh(Item, Result);
 						}
+
 					}
 					else
 					{
@@ -2032,7 +1993,7 @@ namespace UE::Mutable::Private
 				if (Args.source)
 				{
 					AddOp(FScheduledOp(Item.At, Item, 1),
-						FScheduledOp(Args.source, Item));
+							FScheduledOp(Args.source, Item));
 				}
 				else
 				{
@@ -2046,28 +2007,26 @@ namespace UE::Mutable::Private
 
 				if (!EnumHasAnyFlags(EMeshContentFlags(Item.ExecutionOptions), EMeshContentFlags::GeometryData))
 				{
-                	TSharedPtr<const FMesh> Source = LoadMesh(FCacheAddress(Args.source, Item));
+                	TManagedPtr<const FMesh> Source = LoadMesh(FCacheAddress(Args.source, Item));
 					StoreMesh(Item, Source);
 				}
 				else
 				{
-					TSharedPtr<const FMesh> Source = LoadMesh(FCacheAddress(Args.source,Item));
+					TManagedPtr<const FMesh> Source = LoadMesh(FCacheAddress(Args.source,Item));
 
 					const FMatrix44f& mat = Program.ConstantMatrices[Args.matrix];
 
-					TSharedPtr<FMesh> Result = CreateMesh(Source ? Source->GetDataSize() : 0);
+					TManagedPtr<FMesh> Result = CreateMesh(Source ? Source->GetDataSize() : 0);
 
 					bool bOutSuccess = false;
 					MeshTransform(Result.Get(), Source.Get(), mat, bOutSuccess);
 
 					if (!bOutSuccess)
 					{
-						Release(Result);
 						StoreMesh(Item, Source);
 					}
 					else
 					{
-						Release(Source);
 						StoreMesh(Item, Result);
 					}
 				}
@@ -2097,7 +2056,6 @@ namespace UE::Mutable::Private
 				{
 					StoreMesh(Item, nullptr);
 				}
-
 				break;
 			}
             case 1:
@@ -2106,14 +2064,13 @@ namespace UE::Mutable::Private
 
 				if (!EnumHasAnyFlags(EMeshContentFlags(Item.ExecutionOptions), EMeshContentFlags::GeometryData))
 				{
-                	TSharedPtr<const FMesh> Source = LoadMesh(FCacheAddress(Args.Source, Item));
+                	TManagedPtr<const FMesh> Source = LoadMesh(FCacheAddress(Args.Source, Item));
 					StoreMesh(Item, Source);
 				}
 				else
 				{
-
-					TSharedPtr<const FMesh> Source = LoadMesh(FCacheAddress(Args.Source, Item));
-
+					TManagedPtr<const FMesh> Source = LoadMesh(FCacheAddress(Args.Source, Item));
+					
 					check(Args.MorphShape < (uint32)InModel->GetPrivate()->Program.ConstantShapes.Num());
 
 					// Should be an ellipse
@@ -2131,19 +2088,16 @@ namespace UE::Mutable::Private
 						// Should be None or an axis aligned box
 						const FShape& SelectionShape = Program.ConstantShapes[Args.VertexSelectionShapeOrBone];
 
-						TSharedPtr<FMesh> Result = CreateMesh(Source ? Source->GetDataSize() : 0);
-
+						TManagedPtr<FMesh> Result = CreateMesh(Source ? Source->GetDataSize() : 0);
+						
 						bool bOutSuccess = false;
 						MeshClipMorphPlane(Result.Get(), Source.Get(), Origin, Normal, Args.Dist, Args.Factor, MorphShape.size[0], MorphShape.size[1], MorphShape.size[2], SelectionShape, bRemoveFaceIfAllVerticesCulled, bOutSuccess, nullptr, -1);
-						
 						if (!bOutSuccess)
 						{
-							Release(Result);
 							StoreMesh(Item, Source);
 						}
 						else
 						{
-							Release(Source);
 							StoreMesh(Item, Result);
 						}
 					}
@@ -2153,22 +2107,20 @@ namespace UE::Mutable::Private
 						FShape SelectionShape;
 						SelectionShape.type = (uint8)FShape::Type::None;
 
-						TSharedPtr<FMesh> Result = CreateMesh(Source->GetDataSize());
-
+						TManagedPtr<FMesh> Result = CreateMesh(Source ? Source->GetDataSize() : 0);
+						
 						check(Args.VertexSelectionShapeOrBone <= MAX_uint32);
 						const FBoneName Bone(Args.VertexSelectionShapeOrBone);
-
+						
 						bool bOutSuccess = false;
 						MeshClipMorphPlane(Result.Get(), Source.Get(), Origin, Normal, Args.Dist, Args.Factor, MorphShape.size[0], MorphShape.size[1], MorphShape.size[2], SelectionShape, bRemoveFaceIfAllVerticesCulled, bOutSuccess, &Bone, Args.MaxBoneRadius);
 
 						if (!bOutSuccess)
 						{
-							Release(Result);
 							StoreMesh(Item, Source);
 						}
 						else
 						{
-							Release(Source);
 							StoreMesh(Item, Result);
 						}
 					}
@@ -2178,19 +2130,17 @@ namespace UE::Mutable::Private
 						FShape SelectionShape;
 						SelectionShape.type = (uint8)FShape::Type::None;
 
-						TSharedPtr<FMesh> Result = CreateMesh(Source ? Source->GetDataSize() : 0);
-
+						TManagedPtr<FMesh> Result = CreateMesh(Source ? Source->GetDataSize() : 0);
+						
 						bool bOutSuccess = false;
 						MeshClipMorphPlane(Result.Get(), Source.Get(), Origin, Normal, Args.Dist, Args.Factor, MorphShape.size[0], MorphShape.size[1], MorphShape.size[2], SelectionShape, bRemoveFaceIfAllVerticesCulled, bOutSuccess, nullptr, -1.0f);
-
+						
 						if (!bOutSuccess)
 						{
-							Release(Result);
 							StoreMesh(Item, Source);
 						}
 						else
 						{
-							Release(Source);
 							StoreMesh(Item, Result);
 						}
 					}
@@ -2231,7 +2181,6 @@ namespace UE::Mutable::Private
 				{
 					StoreMesh(Item, nullptr);
 				}
-
 				break;
 			}
             case 1:
@@ -2240,38 +2189,33 @@ namespace UE::Mutable::Private
 				
 				if (!EnumHasAnyFlags(EMeshContentFlags(Item.ExecutionOptions), EMeshContentFlags::GeometryData))
 				{
-					TSharedPtr<const FMesh> Source = LoadMesh(FCacheAddress(Args.Source, Item));
+					TManagedPtr<const FMesh> Source = LoadMesh(FCacheAddress(Args.Source, Item));
 					StoreMesh(Item, Source);
 				}
 				else
 				{
 
-					TSharedPtr<const FMesh> Source = LoadMesh(FCacheAddress(Args.Source, Item));
-					TSharedPtr<const FMesh> Clip = LoadMesh(FCacheAddress(Args.ClipMesh, Item));
+					TManagedPtr<const FMesh> Source = LoadMesh(FCacheAddress(Args.Source, Item));
+					TManagedPtr<const FMesh> Clip = LoadMesh(FCacheAddress(Args.ClipMesh, Item));
 
-					// Only if both are valid.
 					if (Source && Clip)
 					{
-						TSharedPtr<FMesh> Result = CreateMesh(Source->GetDataSize());
+						TManagedPtr<FMesh> Result = CreateMesh(Source->GetDataSize());
 
 						bool bOutSuccess = false;
 						MeshClipWithMesh(Result.Get(), Source.Get(), Clip.Get(), bOutSuccess);
 
-						Release(Clip);
 						if (!bOutSuccess)
 						{
-							Release(Result);
 							StoreMesh(Item, Source);
 						}
 						else
 						{
-							Release(Source);
 							StoreMesh(Item, Result);
 						}
 					}
 					else
 					{
-						Release(Clip);
 						StoreMesh(Item, Source);
 					}
 				}
@@ -2317,39 +2261,33 @@ namespace UE::Mutable::Private
 
 				if (!EnumHasAnyFlags(EMeshContentFlags(Item.ExecutionOptions), EMeshContentFlags::GeometryData))
 				{
-					TSharedPtr<const FMesh> BaseMesh = LoadMesh(FCacheAddress(Args.mesh, Item));
+					TManagedPtr<const FMesh> BaseMesh = LoadMesh(FCacheAddress(Args.mesh, Item));
 					StoreMesh(Item, BaseMesh);
 				}
 				else
 				{
-					TSharedPtr<const FMesh> BaseMesh = LoadMesh(FCacheAddress(Args.mesh, Item));
-					TSharedPtr<const FMesh> ClipShape = LoadMesh(FCacheAddress(Args.clipShape, Item));
+					TManagedPtr<const FMesh> BaseMesh = LoadMesh(FCacheAddress(Args.mesh, Item));
+					TManagedPtr<const FMesh> ClipShape = LoadMesh(FCacheAddress(Args.clipShape, Item));
 
 					if (BaseMesh && ClipShape)
 					{
-						TSharedPtr<FMesh> Result = CreateMesh(BaseMesh->GetDataSize());
+						TManagedPtr<FMesh> Result = CreateMesh(BaseMesh->GetDataSize());
 
 						bool bRemoveIfAllVerticesCulled = Args.FaceCullStrategy == EFaceCullStrategy::AllVerticesCulled;
 
 						bool bOutSuccess = false;
 						MeshClipDeform(Result.Get(), BaseMesh.Get(), ClipShape.Get(), Args.clipWeightThreshold, bRemoveIfAllVerticesCulled, bOutSuccess);
-
-						Release(ClipShape);
-
 						if (!bOutSuccess)
 						{
-							Release(Result);
 							StoreMesh(Item, BaseMesh);
 						}
 						else
 						{
-							Release(BaseMesh);
 							StoreMesh(Item, Result);
 						}
 					}
 					else
 					{
-						Release(ClipShape);
 						StoreMesh(Item, BaseMesh);
 					}
 				}
@@ -2386,34 +2324,30 @@ namespace UE::Mutable::Private
             {
           		MUTABLE_CPUPROFILER_SCOPE(ME_APPLYPOSE_1)
 
-                TSharedPtr<const FMesh> Base = LoadMesh(FCacheAddress(Args.base, Item));
-                TSharedPtr<const FMesh> pPose = LoadMesh(FCacheAddress(Args.pose, Item));
+                TManagedPtr<const FMesh> Base = LoadMesh(FCacheAddress(Args.base, Item));
+                TManagedPtr<const FMesh> Pose = LoadMesh(FCacheAddress(Args.pose, Item));
 
                 // Only if both are valid.
-                if (Base && pPose)
-                {
-					TSharedPtr<FMesh> Result = CreateMesh(Base->GetSkeleton() ? Base->GetDataSize() : 0);
+                if (Base && Pose)
+				{
+					TManagedPtr<FMesh> Result = CreateMesh(Base->GetSkeleton() ? Base->GetDataSize() : 0);
 
 					bool bOutSuccess = false;
-					MeshApplyPose(Result.Get(), Base.Get(), pPose.Get(), bOutSuccess);
-
-					Release(pPose);
+					MeshApplyPose(Result.Get(), Base.Get(), Pose.Get(), bOutSuccess);
+					
 					if (!bOutSuccess)
 					{
-						Release(Result);
 						StoreMesh(Item, Base);
 					}
 					else
 					{
-						Release(Base);
 						StoreMesh(Item, Result);
 					}
                 }
-                else
-                {
-					Release(pPose);
+				else
+				{
 					StoreMesh(Item, Base);
-                }
+				}
 
                 break;
             }
@@ -2472,9 +2406,8 @@ namespace UE::Mutable::Private
 				
 				if (BindingMethod == EShapeBindingMethod::ReshapeClosestProject)
 				{
-					TSharedPtr<const FMesh> BaseMesh = LoadMesh(FCacheAddress(Args.mesh, Item));
-					TSharedPtr<const FMesh> Shape = 
-					LoadMesh(FScheduledOp::FromOpAndOptions(Args.shape, Item, ShapeContentFilter));
+					TManagedPtr<const FMesh> BaseMesh = LoadMesh(FCacheAddress(Args.mesh, Item));
+					TManagedPtr<const FMesh> Shape = LoadMesh(FScheduledOp::FromOpAndOptions(Args.shape, Item, ShapeContentFilter));
 					// Bones are stored after the Args
 					Data += sizeof(Args);
 
@@ -2523,31 +2456,26 @@ namespace UE::Mutable::Private
 					FMemory::Memcpy(&ColorChannelUsages, &Args.ColorUsage, sizeof(ColorChannelUsages));
 					static_assert(sizeof(ColorChannelUsages) == sizeof(Args.ColorUsage));
 
-					TSharedPtr<FMesh> BindMeshResult = CreateMesh();
+					TManagedPtr<FMesh> BindMeshResult = CreateMesh();
 
 					bool bOutSuccess = false;
 					MeshBindShapeReshape(BindMeshResult.Get(), BaseMesh.Get(), Shape.Get(), BonesToDeform, PhysicsToDeform, BindFlags, ColorChannelUsages, bOutSuccess);
-				
-					Release(Shape);
+
 					// not success indicates nothing has bond so the base mesh can be reused.
 					if (!bOutSuccess)
 					{
-						Release(BindMeshResult);
 						StoreMesh(Item, BaseMesh);
 					}
 					else
 					{
 						if (!EnumHasAnyFlags(BindFlags, EMeshBindShapeFlags::ReshapeVertices))
 						{
-							TSharedPtr<FMesh> BindMeshNoVertsResult = CloneOrTakeOver(BaseMesh);
+							TManagedPtr<FMesh> BindMeshNoVertsResult = CloneOrTakeOver(MoveTemp(BaseMesh));
 							BindMeshNoVertsResult->AdditionalBuffers = MoveTemp(BindMeshResult->AdditionalBuffers);
-							Release(BaseMesh);
-							Release(BindMeshResult);
 							StoreMesh(Item, BindMeshNoVertsResult);
 						}
 						else
 						{
-							Release(BaseMesh);
 							StoreMesh(Item, BindMeshResult);
 						}
 					}
@@ -2556,29 +2484,26 @@ namespace UE::Mutable::Private
 				{
 					if (!EnumHasAnyFlags(EMeshContentFlags(Item.ExecutionOptions), EMeshContentFlags::GeometryData))
 					{
-						TSharedPtr<const FMesh> BaseMesh = LoadMesh(FCacheAddress(Args.mesh, Item));
+						TManagedPtr<const FMesh> BaseMesh = LoadMesh(FCacheAddress(Args.mesh, Item));
 						StoreMesh(Item, BaseMesh);
 					}
 					else
 					{
-						TSharedPtr<const FMesh> BaseMesh = LoadMesh(FCacheAddress(Args.mesh, Item));
-						TSharedPtr<const FMesh> Shape = 
+						TManagedPtr<const FMesh> BaseMesh = LoadMesh(FCacheAddress(Args.mesh, Item));
+						TManagedPtr<const FMesh> Shape = 
 								LoadMesh(FScheduledOp::FromOpAndOptions(Args.shape, Item, ShapeContentFilter));
 
-						TSharedPtr<FMesh> Result = CreateMesh(BaseMesh ? BaseMesh->GetDataSize() : 0);
+						TManagedPtr<FMesh> Result = CreateMesh(BaseMesh ? BaseMesh->GetDataSize() : 0);
 
 						bool bOutSuccess = false;
 						MeshBindShapeClipDeform(Result.Get(), BaseMesh.Get(), Shape.Get(), BindingMethod, bOutSuccess);
 
-						Release(Shape);
 						if (!bOutSuccess)
 						{
-							Release(Result);
 							StoreMesh(Item, BaseMesh);
 						}
 						else
 						{
-							Release(BaseMesh);
 							StoreMesh(Item, Result);
 						}
 					}
@@ -2618,8 +2543,8 @@ namespace UE::Mutable::Private
 			{
 				MUTABLE_CPUPROFILER_SCOPE(ME_APPLYSHAPE_1)
 					
-				TSharedPtr<const FMesh> BaseMesh = LoadMesh(FCacheAddress(Args.mesh, Item));
-				TSharedPtr<const FMesh> Shape = LoadMesh(FCacheAddress(Args.shape, Item));
+				TManagedPtr<const FMesh> BaseMesh = LoadMesh(FCacheAddress(Args.mesh, Item));
+				TManagedPtr<const FMesh> Shape = LoadMesh(FCacheAddress(Args.shape, Item));
 
 				EMeshBindShapeFlags ReshapeFlags = static_cast<EMeshBindShapeFlags>(Args.flags);
 				const EMeshContentFlags MeshContentFilter = (EMeshContentFlags)Item.ExecutionOptions;
@@ -2645,16 +2570,15 @@ namespace UE::Mutable::Private
 
 				const bool bReshapeVertices = EnumHasAnyFlags(ReshapeFlags, EMeshBindShapeFlags::ReshapeVertices);
 
-				TSharedPtr<FMesh> ReshapedMeshResult = CreateMesh(BaseMesh ? BaseMesh->GetDataSize() : 0);
+				TManagedPtr<FMesh> ReshapedMeshResult = CreateMesh(BaseMesh->GetDataSize());
 
 				bool bOutSuccess = false;
 				MeshApplyShape(ReshapedMeshResult.Get(), BaseMesh.Get(), Shape.Get(), ReshapeFlags, bOutSuccess);
 
-				Release(Shape);
-				
+				Shape = nullptr;
+
 				if (!bOutSuccess)
 				{
-					Release(ReshapedMeshResult);
 					StoreMesh(Item, BaseMesh);
 				}
 				else
@@ -2668,20 +2592,17 @@ namespace UE::Mutable::Private
 							EMeshCopyFlags::WithAdditionalPhysics |
 							EMeshCopyFlags::WithPoses);
 
-						TSharedPtr<FMesh> NoVerticesReshpedMesh = CloneOrTakeOver(BaseMesh);
+						TManagedPtr<FMesh> NoVerticesReshpedMesh = CloneOrTakeOver(MoveTemp(BaseMesh));
 
 						NoVerticesReshpedMesh->SetSkeleton(ReshapedMeshResult->GetSkeleton());
 						NoVerticesReshpedMesh->SetPhysicsBody(ReshapedMeshResult->GetPhysicsBody());
 						NoVerticesReshpedMesh->AdditionalPhysicsBodies = ReshapedMeshResult->AdditionalPhysicsBodies;
 						NoVerticesReshpedMesh->BonePoses = ReshapedMeshResult->BonePoses;
 
-						Release(BaseMesh);
-						Release(ReshapedMeshResult);
 						StoreMesh(Item, NoVerticesReshpedMesh);
 					}
 					else
 					{
-						Release(BaseMesh);
 						StoreMesh(Item, ReshapedMeshResult);
 					}
 				}
@@ -2718,8 +2639,8 @@ namespace UE::Mutable::Private
 			{
 				MUTABLE_CPUPROFILER_SCOPE(ME_MORPHRESHAPE_1)
 					
-				TSharedPtr<const FMesh> MorphedMesh = LoadMesh(FCacheAddress(Args.Morph, Item));
-				TSharedPtr<const FMesh> ReshapeMesh = LoadMesh(FCacheAddress(Args.Reshape, Item));
+				TManagedPtr<const FMesh> MorphedMesh = LoadMesh(FCacheAddress(Args.Morph, Item));
+				TManagedPtr<const FMesh> ReshapeMesh = LoadMesh(FCacheAddress(Args.Reshape, Item));
 
 				if (ReshapeMesh && MorphedMesh)
 				{
@@ -2729,15 +2650,13 @@ namespace UE::Mutable::Private
 							EMeshCopyFlags::WithPhysicsBody | 
 							EMeshCopyFlags::WithPoses);
 
-					TSharedPtr<FMesh> Result = CreateMesh(MorphedMesh->GetDataSize());
+					TManagedPtr<FMesh> Result = CreateMesh(MorphedMesh->GetDataSize());
 					Result->CopyFrom(*MorphedMesh, CopyFlags);
 
 					Result->SetSkeleton(ReshapeMesh->GetSkeleton());
 					Result->SetPhysicsBody(ReshapeMesh->GetPhysicsBody());
 					Result->BonePoses = ReshapeMesh->BonePoses;
 
-					Release(MorphedMesh);
-					Release(ReshapeMesh);
 					StoreMesh(Item, Result);
 				}
 				else
@@ -2778,15 +2697,13 @@ namespace UE::Mutable::Private
             {
             	MUTABLE_CPUPROFILER_SCOPE(ME_SETSKELETON_1)
             		
-                TSharedPtr<const FMesh> Source = LoadMesh(FCacheAddress(Args.Source, Item));
-                TSharedPtr<const FMesh> Skeleton = LoadMesh(FCacheAddress(Args.Skeleton, Item));
+                TManagedPtr<const FMesh> Source = LoadMesh(FCacheAddress(Args.Source, Item));
+                TManagedPtr<const FMesh> Skeleton = LoadMesh(FCacheAddress(Args.Skeleton, Item));
 
                 // Only if both are valid.
                 if (Source && Skeleton)
                 {
-                    if ( Source->GetSkeleton()
-                         &&
-                         Source->GetSkeleton()->GetBoneCount() > 0 )
+                    if (Source->GetSkeleton() && Source->GetSkeleton()->GetBoneCount() > 0)
                     {
                         // For some reason we already have bone data, so we can't just overwrite it
                         // or the skinning may break. This may happen because of a problem in the
@@ -2794,41 +2711,35 @@ namespace UE::Mutable::Private
                         // \TODO Be defensive, for now.
                         UE_LOG(LogMutableCore, Warning, TEXT("Performing a MeshRemapSkeleton, instead of MeshSetSkeletonData because source mesh already has some skeleton."));
 
-						TSharedPtr<FMesh> Result = CreateMesh(Source->GetDataSize());
+						TManagedPtr<FMesh> Result = CreateMesh(Source->GetDataSize());
 
 						bool bOutSuccess = false;
                         MeshRemapSkeleton(Result.Get(), Source.Get(), Skeleton->GetSkeleton(), bOutSuccess);
 
-						Release(Skeleton);
-
-                        if (!bOutSuccess)
-                        {
-							Release(Result);
+						if (!bOutSuccess)
+						{
 							StoreMesh(Item, Source);
-                        }
+						}
 						else
 						{
 							//Result->GetPrivate()->CheckIntegrity();
-							Release(Source);
 							StoreMesh(Item, Result);
 						}
                     }
                     else
                     {
-						TSharedPtr<FMesh> Result = CloneOrTakeOver(Source);
+						TManagedPtr<FMesh> Result = CloneOrTakeOver(MoveTemp(Source));
 
                         Result->SetSkeleton(Skeleton->GetSkeleton());
 
 						//Result->GetPrivate()->CheckIntegrity();
-						Release(Skeleton);
 						StoreMesh(Item, Result);
                     }
                 }
-                else
-                {
-					Release(Skeleton);
+				else
+				{
 					StoreMesh(Item, Source);
-                }
+				}
 
                 break;
             }
@@ -2846,36 +2757,36 @@ namespace UE::Mutable::Private
         		
             // Decode op
             // TODO: Partial decode for each stage
-            const uint8* data = Program.GetOpArgsPointer(Item.At);
+            const uint8* Data = Program.GetOpArgsPointer(Item.At);
 
-            OP::ADDRESS source;
-            FMemory::Memcpy(&source,data,sizeof(OP::ADDRESS)); 
-			data += sizeof(OP::ADDRESS);
+            OP::ADDRESS SourceAddress;
+            FMemory::Memcpy(&SourceAddress, Data, sizeof(OP::ADDRESS)); 
+			Data += sizeof(OP::ADDRESS);
 
 			EFaceCullStrategy FaceCullStrategy;
-			FMemory::Memcpy(&FaceCullStrategy, data, sizeof(EFaceCullStrategy));
-			data += sizeof(EFaceCullStrategy);
+			FMemory::Memcpy(&FaceCullStrategy, Data, sizeof(EFaceCullStrategy));
+			Data += sizeof(EFaceCullStrategy);
 
-            TArray<FScheduledOp> conditions;
-			TArray<OP::ADDRESS> masks;
+            TArray<FScheduledOp, TInlineAllocator<4>> Conditions;
+			TArray<OP::ADDRESS, TInlineAllocator<4>> Masks;
 
-            uint16 removes;
-			FMemory::Memcpy(&removes,data,sizeof(uint16)); 
-			data += sizeof(uint16);
+            uint16 NumRemoves;
+			FMemory::Memcpy(&NumRemoves, Data, sizeof(uint16)); 
+			Data += sizeof(uint16);
 
-            for( uint16 r=0; r<removes; ++r)
+            for(uint16 RemoveIndex = 0; RemoveIndex < NumRemoves; ++RemoveIndex)
             {
-                OP::ADDRESS condition;
-				FMemory::Memcpy(&condition,data,sizeof(OP::ADDRESS)); 
-				data += sizeof(OP::ADDRESS);
+                OP::ADDRESS Condition;
+				FMemory::Memcpy(&Condition, Data, sizeof(OP::ADDRESS)); 
+				Data += sizeof(OP::ADDRESS);
                 
-				conditions.Emplace(condition, Item);
+				Conditions.Emplace(Condition, Item);
 
-                OP::ADDRESS mask;
-				FMemory::Memcpy(&mask,data,sizeof(OP::ADDRESS)); 
-				data += sizeof(OP::ADDRESS);
+                OP::ADDRESS Mask;
+				FMemory::Memcpy(&Mask, Data, sizeof(OP::ADDRESS)); 
+				Data += sizeof(OP::ADDRESS);
 
-                masks.Add(mask);
+                Masks.Add(Mask);
             }
 
 
@@ -2884,17 +2795,17 @@ namespace UE::Mutable::Private
             {
             case 0:
 			{
-				if (source)
+				if (SourceAddress)
 				{
 					if (!EnumHasAnyFlags(EMeshContentFlags(Item.ExecutionOptions), EMeshContentFlags::GeometryData))
 					{
 						AddOp(FScheduledOp(Item.At, Item, 1), 
-								FScheduledOp(source, Item));
+								FScheduledOp(SourceAddress, Item));
 					}
 					else
 					{
 						// Request the conditions
-						AddOp(FScheduledOp(Item.At, Item, 1), conditions);
+						AddOp(FScheduledOp(Item.At, Item, 1), Conditions);
 					}
 				}
 				else
@@ -2909,33 +2820,35 @@ namespace UE::Mutable::Private
 
 				if (!EnumHasAnyFlags(EMeshContentFlags(Item.ExecutionOptions), EMeshContentFlags::GeometryData))
                 {
-                	TSharedPtr<const FMesh> Source = LoadMesh(FCacheAddress(source, Item));
+                	TManagedPtr<const FMesh> Source = LoadMesh(FCacheAddress(SourceAddress, Item));
 					StoreMesh(Item, Source);
 				}
 				else
 				{
 					// Request the source and the necessary masks
 					// \todo: store condition values in heap?
-					TArray<FScheduledOp> deps;
-					deps.Emplace( source, Item );
-					for( int32 r=0; source && r<conditions.Num(); ++r )
+					TArray<FScheduledOp, TInlineAllocator<4>> Deps;
+					Deps.Reserve(Conditions.Num());
+
+					Deps.Emplace(SourceAddress, Item);
+					for( int32 RemoveIndex = 0; RemoveIndex < Conditions.Num(); ++RemoveIndex)
 					{
 						// If there is no expression, we'll assume true.
-						bool value = true;
-						if (conditions[r].At)
+						bool bValue = true;
+						if (Conditions[RemoveIndex].At)
 						{
-							value = LoadBool(FCacheAddress(conditions[r].At, Item));
+							bValue = LoadBool(FCacheAddress(Conditions[RemoveIndex].At, Item));
 						}
 
-						if (value)
+						if (bValue)
 						{
-							deps.Emplace(masks[r], Item);
+							Deps.Emplace(Masks[RemoveIndex], Item);
 						}
 					}
 
-					if (source)
+					if (SourceAddress)
 					{
-						AddOp(FScheduledOp(Item.At, Item, 2), deps);
+						AddOp(FScheduledOp(Item.At, Item, 2), Deps);
 					}
 				}
                 break;
@@ -2946,30 +2859,28 @@ namespace UE::Mutable::Private
             	MUTABLE_CPUPROFILER_SCOPE(ME_REMOVEMASK_2)
             	
                 // \todo: single remove operation with all masks?
-                TSharedPtr<const FMesh> Source = LoadMesh(FCacheAddress(source, Item));
+                TManagedPtr<const FMesh> Source = LoadMesh(FCacheAddress(SourceAddress, Item));
 
 				if (Source)
 				{
-					TSharedPtr<FMesh> Result = CloneOrTakeOver(Source);
+					TManagedPtr<FMesh> Result = CloneOrTakeOver(MoveTemp(Source));
 
-					for (int32 r = 0; r < conditions.Num(); ++r)
+					for (int32 RemoveIndex = 0; RemoveIndex < Conditions.Num(); ++RemoveIndex)
 					{
 						// If there is no expression, we'll assume true.
-						bool value = true;
-						if (conditions[r].At)
+						bool bValue = true;
+						if (Conditions[RemoveIndex].At)
 						{
-							value = LoadBool(FCacheAddress(conditions[r].At, Item));
+							bValue = LoadBool(FCacheAddress(Conditions[RemoveIndex].At, Item));
 						}
 
-						if (value)
+						if (bValue)
 						{
-							TSharedPtr<const FMesh> Mask = LoadMesh(FCacheAddress(masks[r], Item));
+							TManagedPtr<const FMesh> Mask = LoadMesh(FCacheAddress(Masks[RemoveIndex], Item));
 							if (Mask)
 							{
 								bool bRemoveIfAllVerticesCulled = FaceCullStrategy == EFaceCullStrategy::AllVerticesCulled;
 								MeshRemoveMaskInline(Result.Get(), Mask.Get(), bRemoveIfAllVerticesCulled);
-
-								Release(Mask);
 							}
 						}
 					}
@@ -2980,7 +2891,6 @@ namespace UE::Mutable::Private
 				{
 					StoreMesh(Item, nullptr);
 				}
-
                 break;
             }
 
@@ -3004,7 +2914,6 @@ namespace UE::Mutable::Private
 			{
 				if (OpArgs.Source)
 				{
-					// Request the source
 					AddOp(FScheduledOp(Item.At, Item, 1), FScheduledOp(OpArgs.Source, Item));
 				}
 				else
@@ -3018,12 +2927,12 @@ namespace UE::Mutable::Private
 			{
 				MUTABLE_CPUPROFILER_SCOPE(ME_ADDMETADATA_2)
 
-				TSharedPtr<const FMesh> SourceMesh = LoadMesh(FCacheAddress(OpArgs.Source, Item));
+				TManagedPtr<const FMesh> SourceMesh = LoadMesh(FCacheAddress(OpArgs.Source, Item));
+				TManagedPtr<FMesh> Result = nullptr;
 
-				TSharedPtr<FMesh> Result = nullptr;
 				if (SourceMesh)
 				{
-					Result = CloneOrTakeOver(SourceMesh);
+					Result = CloneOrTakeOver(MoveTemp(SourceMesh));
 
 					using OpEnumFlags = OP::MeshAddMetadataArgs::EnumFlags;
 					bool bIsTagList      = EnumHasAnyFlags(OpArgs.Flags, OpEnumFlags::IsTagList);
@@ -3112,7 +3021,6 @@ namespace UE::Mutable::Private
 				StoreMesh(Item, Result);
 				break;
 			}
-
 			default:
 				check(false);
 			}
@@ -3153,38 +3061,33 @@ namespace UE::Mutable::Private
 
 				if (!EnumHasAnyFlags(EMeshContentFlags(Item.ExecutionOptions), EMeshContentFlags::GeometryData))
 				{
-					TSharedPtr<const FMesh> Mesh = LoadMesh(FCacheAddress(Args.Mesh, Item));
+					TManagedPtr<const FMesh> Mesh = LoadMesh(FCacheAddress(Args.Mesh, Item));
 					StoreMesh(Item, Mesh);
 				}
 				else
 				{
-					TSharedPtr<const FMesh> Mesh = LoadMesh(FCacheAddress(Args.Mesh, Item));
+					TManagedPtr<const FMesh> Mesh = LoadMesh(FCacheAddress(Args.Mesh, Item));
 					const FProjector Projector = LoadProjector(FCacheAddress(Args.Projector, Item));
 
-					// Only if both are valid.
 					if (Mesh && Mesh->GetVertexBuffers().GetBufferCount() > 0)
 					{
-						TSharedPtr<FMesh> Result = CreateMesh();
+						TManagedPtr<FMesh> Result = CreateMesh();
 
 						bool bOutSuccess = false;
 						MeshProject(Result.Get(), Mesh.Get(), Projector, bOutSuccess);
-
 						if (!bOutSuccess)
 						{
-							Release(Result);
 							StoreMesh(Item, Mesh);
 						}
 						else
 						{
 							// Result->GetPrivate()->CheckIntegrity();
-							Release(Mesh);
 							StoreMesh(Item, Result);
 						}
 					}
 					else
 					{
-						Release(Mesh);
-						StoreMesh(Item, nullptr);
+						StoreMesh(Item, Mesh);
 					}
 				}
                 break;
@@ -3212,32 +3115,31 @@ namespace UE::Mutable::Private
 				{
 					StoreMesh(Item, nullptr);
 				}
+
 				break;
 			}
 			case 1:
 			{
 				MUTABLE_CPUPROFILER_SCOPE(ME_OPTIMIZESKINNING_1)
 
-				TSharedPtr<const FMesh> Source = LoadMesh(FCacheAddress(Args.source, Item));
+				TManagedPtr<const FMesh> Source = LoadMesh(FCacheAddress(Args.source, Item));
 
-				if (Source && Source->IsReference())
+				if (Source->IsReference())
 				{
 					StoreMesh(Item, Source);
 				}
 
-				TSharedPtr<FMesh> Result = CreateMesh();
+				TManagedPtr<FMesh> Result = CreateMesh();
 
 				bool bOutSuccess = false;
 				MeshOptimizeSkinning(Result.Get(), Source.Get(), bOutSuccess);
 
 				if (!bOutSuccess)
 				{
-					Release(Result);
 					StoreMesh(Item, Source);
 				}
 				else
 				{
-					Release(Source);
 					StoreMesh(Item, Result);
 				}
 
@@ -3285,37 +3187,32 @@ namespace UE::Mutable::Private
 
 				if (!EnumHasAnyFlags(EMeshContentFlags(Item.ExecutionOptions), EMeshContentFlags::GeometryData))
 				{
-					TSharedPtr<const FMesh> SourceMesh = LoadMesh(FCacheAddress(Args.sourceMesh, Item));
+					TManagedPtr<const FMesh> SourceMesh = LoadMesh(FCacheAddress(Args.sourceMesh, Item));
 					StoreMesh(Item, SourceMesh);
 				}
 				else
 				{
-					TSharedPtr<const FMesh> SourceMesh = LoadMesh(FCacheAddress(Args.sourceMesh,Item));
-					TSharedPtr<const FMesh> BoundingMesh = LoadMesh(FCacheAddress(Args.boundingMesh, Item));
+					TManagedPtr<const FMesh> SourceMesh = LoadMesh(FCacheAddress(Args.sourceMesh,Item));
+					TManagedPtr<const FMesh> BoundingMesh = LoadMesh(FCacheAddress(Args.boundingMesh, Item));
 					const FMatrix44f& Transform = LoadMatrix(FCacheAddress(Args.matrix, Item));
 
 					if (SourceMesh)
 					{
-						TSharedPtr<FMesh> Result = CreateMesh(SourceMesh->GetDataSize());
+						TManagedPtr<FMesh> Result = CreateMesh(SourceMesh->GetDataSize());
 
 						bool bOutSuccess = false;
 						MeshTransformWithMesh(Result.Get(), SourceMesh.Get(), BoundingMesh.Get(), Transform, bOutSuccess);
-						Release(BoundingMesh);
-
 						if (!bOutSuccess)
 						{
-							Release(Result);
 							StoreMesh(Item, SourceMesh);
 						}
 						else
 						{
-							Release(SourceMesh);
 							StoreMesh(Item, Result);
 						}
 					}
 					else
 					{
-						Release(BoundingMesh);
 						StoreMesh(Item, SourceMesh);
 					}
 				}
@@ -3359,29 +3256,29 @@ namespace UE::Mutable::Private
 			{
 				MUTABLE_CPUPROFILER_SCOPE(ME_TRANSFORMWITHBONE_1)
 
-					if (!EnumHasAnyFlags(EMeshContentFlags(Item.ExecutionOptions), EMeshContentFlags::GeometryData))
+				if (!EnumHasAnyFlags(EMeshContentFlags(Item.ExecutionOptions), EMeshContentFlags::GeometryData))
+				{
+					TManagedPtr<const FMesh> SourceMesh = LoadMesh(FCacheAddress(Args.SourceMesh, Item));
+					StoreMesh(Item, SourceMesh);
+				}
+				else
+				{
+					TManagedPtr<const FMesh> SourceMesh = LoadMesh(FCacheAddress(Args.SourceMesh, Item));
+					const FMatrix44f& Transform = LoadMatrix(FCacheAddress(Args.Matrix, Item));
+					const FBoneName BoneName = FBoneName(Args.BoneId);
+
+					if (SourceMesh)
 					{
-						TSharedPtr<const FMesh> SourceMesh = LoadMesh(FCacheAddress(Args.SourceMesh, Item));
-						StoreMesh(Item, SourceMesh);
+						TManagedPtr<FMesh> Result = CloneOrTakeOver(MoveTemp(SourceMesh));
+
+						MeshTransformWithBoneInline(Result.Get(), Transform, FBoneName(Args.BoneId), Args.ThresholdFactor);
+						StoreMesh(Item, Result);
 					}
 					else
 					{
-						TSharedPtr<const FMesh> SourceMesh = LoadMesh(FCacheAddress(Args.SourceMesh, Item));
-						const FMatrix44f& Transform = LoadMatrix(FCacheAddress(Args.Matrix, Item));
-						const FBoneName BoneName = FBoneName(Args.BoneId);
-
-						if (SourceMesh)
-						{
-							TSharedPtr<FMesh> Result = CloneOrTakeOver(SourceMesh);
-
-							MeshTransformWithBoneInline(Result.Get(), Transform, FBoneName(Args.BoneId), Args.ThresholdFactor);
-							StoreMesh(Item, Result);
-						}
-						else
-						{
-							StoreMesh(Item, SourceMesh);
-						}
+						StoreMesh(Item, SourceMesh);
 					}
+				}
 				break;
 			}
 
@@ -3423,7 +3320,6 @@ namespace UE::Mutable::Private
 		}
 	}
 
-    //---------------------------------------------------------------------------------------------
     void CodeRunner::RunCode_Image(const FScheduledOp& Item, const FParameters* pParams, const FModel* InModel )
     {
 		MUTABLE_CPUPROFILER_SCOPE(RunCode_Image);
@@ -3547,7 +3443,7 @@ namespace UE::Mutable::Private
                     }
                 }
 
-				TSharedPtr<const FImage> Base = LoadImage(FCacheAddress(Args.base, Item));
+				TManagedPtr<const FImage> Base = LoadImage(FCacheAddress(Args.base, Item));
 
 				if (Iterations <= 0)
 				{
@@ -3557,7 +3453,7 @@ namespace UE::Mutable::Private
 				else
 				{
 					// Store the base
-					TSharedPtr<FImage> New = CloneOrTakeOver(Base);
+					TManagedPtr<FImage> New = CloneOrTakeOver(MoveTemp(Base));
 					EImageFormat InitialBaseFormat = New->GetFormat();
 
 					// Reset relevancy map.
@@ -3568,13 +3464,12 @@ namespace UE::Mutable::Private
 					EImageFormat BaseFormat = GetUncompressedFormat(New->GetFormat());
 					if (New->GetFormat() != BaseFormat)
 					{
-						TSharedPtr<FImage> Formatted = CreateImage( New->GetSizeX(), New->GetSizeY(), New->GetLODCount(), BaseFormat, EInitializationType::NotInitialized );
+						TManagedPtr<FImage> Formatted = CreateImage( New->GetSizeX(), New->GetSizeY(), New->GetLODCount(), BaseFormat, EInitializationType::NotInitialized );
 
 						bool bSuccess = false;
 						ImOp.ImagePixelFormat(bSuccess, Settings.ImageCompressionQuality, Formatted.Get(), New.Get());
 						check(bSuccess); // Decompression cannot fail
 
-						Release(New);
 						New = Formatted;
 					}
 
@@ -3611,8 +3506,8 @@ namespace UE::Mutable::Private
 
 				// Process the current layer
 
-				TSharedPtr<FImage> Base = StaticCastSharedPtr<FImage>(ConstCastSharedPtr<FResource>(Data.Resource));
- 
+				TManagedPtr<FImage> Base = StaticCastManagedPtr<FImage>(ConstCastManagedPtr<FResource>(Data.Resource));
+				 
                 FScheduledOp itemCopy = Item;
                 ExecutionIndex index = GetMemory().GetRangeIndex( Item.ExecutionIndex );
 				
@@ -3621,21 +3516,20 @@ namespace UE::Mutable::Private
                     itemCopy.ExecutionIndex = GetMemory().GetRangeIndexIndex(index);
 					itemCopy.CustomState = 0;
 
-                    TSharedPtr<const FImage> Blended = LoadImage( FCacheAddress(Args.blended,itemCopy) );
+                    TManagedPtr<const FImage> Blended = LoadImage( FCacheAddress(Args.blended,itemCopy) );
 
                     // This shouldn't happen in optimised models, but it could happen in editors, etc.
                     // \todo: raise a performance warning?
-                    if (Blended && Blended->GetFormat()!=Base->GetFormat() )
+                    if (Blended && Blended->GetFormat() != Base->GetFormat())
                     {
 						MUTABLE_CPUPROFILER_SCOPE(ImageResize_BlendedReformat);
 
-						TSharedPtr<FImage> Formatted = CreateImage(Blended->GetSizeX(), Blended->GetSizeY(), Blended->GetLODCount(), Base->GetFormat(), EInitializationType::NotInitialized);
+						TManagedPtr<FImage> Formatted = CreateImage(Blended->GetSizeX(), Blended->GetSizeY(), Blended->GetLODCount(), Base->GetFormat(), EInitializationType::NotInitialized);
 
 						bool bSuccess = false;
 						ImOp.ImagePixelFormat(bSuccess, Settings.ImageCompressionQuality, Formatted.Get(), Blended.Get());
 						check(bSuccess);
 
-						Release(Blended);
 						Blended = Formatted;
                     }
 
@@ -3645,9 +3539,9 @@ namespace UE::Mutable::Private
 					{
 						MUTABLE_CPUPROFILER_SCOPE(ImageResize_BlendedFixForMultilayer);
 
-						TSharedPtr<FImage> Resized = CreateImage(ResultSize[0], ResultSize[1], Blended->GetLODCount(), Blended->GetFormat(), EInitializationType::NotInitialized);
+						TManagedPtr<FImage> Resized = CreateImage(ResultSize[0], ResultSize[1], Blended->GetLODCount(), Blended->GetFormat(), EInitializationType::NotInitialized);
 						ImOp.ImageResizeLinear(Resized.Get(), 0, Blended.Get());
-						Release(Blended);
+
 						Blended = Resized;
 					}
 
@@ -3687,16 +3581,15 @@ namespace UE::Mutable::Private
 
                     if (!bDone && Args.mask)
                     {
-                        TSharedPtr<const FImage> Mask = LoadImage( FCacheAddress(Args.mask,itemCopy) );
+                        TManagedPtr<const FImage> Mask = LoadImage( FCacheAddress(Args.mask,itemCopy) );
 
 						// TODO: This shouldn't happen, but be defensive.
 						if (Mask && Mask->GetSize() != ResultSize)
 						{
 							MUTABLE_CPUPROFILER_SCOPE(ImageResize_MaskFixForMultilayer);
 
-							TSharedPtr<FImage> Resized = CreateImage(ResultSize[0], ResultSize[1], Mask->GetLODCount(), Mask->GetFormat(), EInitializationType::NotInitialized);
+							TManagedPtr<FImage> Resized = CreateImage(ResultSize[0], ResultSize[1], Mask->GetLODCount(), Mask->GetFormat(), EInitializationType::NotInitialized);
 							ImOp.ImageResizeLinear(Resized.Get(), 0, Mask.Get());
-							Release(Mask);
 							Mask = Resized;
 						}
 
@@ -3716,9 +3609,9 @@ namespace UE::Mutable::Private
                         case EBlendType::BT_MULTIPLY: BufferLayer<MultiplyChannelMasked, MultiplyChannel, false>(Base.Get(), Base.Get(), Mask.Get(), Blended.Get(), bApplyColorBlendToAlpha, bBlendOnlyOneMip); break;
                         case EBlendType::BT_BLEND: BufferLayer<BlendChannelMasked, BlendChannel, false>(Base.Get(), Base.Get(), Mask.Get(), Blended.Get(), bApplyColorBlendToAlpha, bBlendOnlyOneMip); break;
                         default: check(false);
-                        }
+						}
 
-						Release(Mask);
+						Mask = nullptr;
                     }
 					else if (!bDone && Args.bUseMaskFromBlended)
 					{
@@ -3777,7 +3670,7 @@ namespace UE::Mutable::Private
 						}
 					}
 
-					Release(Blended);
+					Blended = nullptr;
 				}
 
 				// Are we done?
@@ -3825,8 +3718,8 @@ namespace UE::Mutable::Private
 				if (Args.base && Args.normal)
 				{
 					AddOp(FScheduledOp(Item.At, Item, 1),
-							FScheduledOp(Args.base, Item),
-							FScheduledOp(Args.normal, Item));
+						FScheduledOp(Args.base, Item),
+						FScheduledOp(Args.normal, Item));
 				}
 				else
 				{
@@ -3838,8 +3731,8 @@ namespace UE::Mutable::Private
 			{
 				MUTABLE_CPUPROFILER_SCOPE(IM_NORMALCOMPOSITE_1)
 
-				TSharedPtr<const FImage> Base = LoadImage(FCacheAddress(Args.base, Item));
-				TSharedPtr<const FImage> Normal = LoadImage(FCacheAddress(Args.normal, Item));
+				TManagedPtr<const FImage> Base = LoadImage(FCacheAddress(Args.base, Item));
+				TManagedPtr<const FImage> Normal = LoadImage(FCacheAddress(Args.normal, Item));
 
 				if (Normal->GetLODCount() < Base->GetLODCount())
 				{
@@ -3848,7 +3741,7 @@ namespace UE::Mutable::Private
 					int32 StartLevel = Normal->GetLODCount() - 1;
 					int32 LevelCount = Base->GetLODCount();
 					
-					TSharedPtr<FImage> NormalFix = CloneOrTakeOver(Normal);
+					TManagedPtr<FImage> NormalFix = CloneOrTakeOver(MoveTemp(Normal));
 					NormalFix->DataStorage.SetNumLODs(LevelCount);
 
 					FMipmapGenerationSettings MipSettings{};
@@ -3857,12 +3750,9 @@ namespace UE::Mutable::Private
 					Normal = NormalFix;
 				}
 
-
-                TSharedPtr<FImage> Result = CreateImage(Base->GetSizeX(), Base->GetSizeY(), Base->GetLODCount(), Base->GetFormat(), EInitializationType::NotInitialized);
+                TManagedPtr<FImage> Result = CreateImage(Base->GetSizeX(), Base->GetSizeY(), Base->GetLODCount(), Base->GetFormat(), EInitializationType::NotInitialized);
 				ImageNormalComposite(Result.Get(), Base.Get(), Normal.Get(), Args.mode, Args.power);
 
-				Release(Base);
-				Release(Normal);
 				StoreImage(Item, Result);
 				break;
 			}
@@ -3955,17 +3845,17 @@ namespace UE::Mutable::Private
             {
             	MUTABLE_CPUPROFILER_SCOPE(IM_RESIZELIKE_1)
             	
-                TSharedPtr<const FImage> Base = LoadImage( FCacheAddress(Args.Source,Item) );
-                TSharedPtr<const FImage> SizeBase = LoadImage( FCacheAddress(Args.SizeSource,Item) );
+                TManagedPtr<const FImage> Base = LoadImage( FCacheAddress(Args.Source,Item) );
+                TManagedPtr<const FImage> SizeBase = LoadImage( FCacheAddress(Args.SizeSource,Item) );
 				FImageSize DestSize = SizeBase->GetSize();
-				Release(SizeBase);
+				SizeBase = nullptr;
 
                 if (Base->GetSize() != DestSize)
                 {
 					int32 BaseLODCount = Base->GetLODCount();
-					TSharedPtr<FImage> Result = CreateImage(DestSize[0], DestSize[1], BaseLODCount, Base->GetFormat(), EInitializationType::NotInitialized);
+					TManagedPtr<FImage> Result = CreateImage(DestSize[0], DestSize[1], BaseLODCount, Base->GetFormat(), EInitializationType::NotInitialized);
 					ImOp.ImageResizeLinear(Result.Get(), Settings.ImageCompressionQuality, Base.Get());
-					Release(Base);
+					Base = nullptr;
 
                     // If the source image had mips, generate them as well for the resized image.
                     // This shouldn't happen often since "ResizeLike" should be usually optimised out
@@ -4034,7 +3924,7 @@ namespace UE::Mutable::Private
             {
             	MUTABLE_CPUPROFILER_SCOPE(IM_BLANKLAYOUT_1)
             		
-                TSharedPtr<const FLayout> pLayout = LoadLayout(FScheduledOp::FromOpAndOptions(Args.Layout, Item, 0));
+                TManagedPtr<const FLayout> pLayout = LoadLayout(FScheduledOp::FromOpAndOptions(Args.Layout, Item, 0));
 
                 FIntPoint SizeInBlocks = pLayout->GetGridSize();
 
@@ -4084,7 +3974,7 @@ namespace UE::Mutable::Private
                 }
 
 				// It needs to be initialized in case it has gaps.
-                TSharedPtr<FImage> New = CreateImage(ImageSizeInPixels.X, ImageSizeInPixels.Y, MipsToGenerate, EImageFormat(Args.Format), EInitializationType::Black );
+                TManagedPtr<FImage> New = CreateImage(ImageSizeInPixels.X, ImageSizeInPixels.Y, MipsToGenerate, EImageFormat(Args.Format), EInitializationType::Black );
                 StoreImage(Item, New);
                 break;
             }
@@ -4110,7 +4000,7 @@ namespace UE::Mutable::Private
 					break;
 				case 1:
 				{
-					TSharedPtr<const FLayout> ComposeLayout = 
+					TManagedPtr<const FLayout> ComposeLayout = 
 							LoadLayout(FCacheAddress(Args.layout, FScheduledOp::FromOpAndOptions(Args.layout, Item, 0)));
 
 					FScheduledOpData Data;
@@ -4158,7 +4048,7 @@ namespace UE::Mutable::Private
 
 				case 1:
 				{	
-					TSharedPtr<const FLayout> ComposeLayout = 
+					TManagedPtr<const FLayout> ComposeLayout = 
 							LoadLayout(FCacheAddress(Args.layout, FScheduledOp::FromOpAndOptions(Args.layout, Item, 0)));
 
 					FScheduledOpData Data;
@@ -4207,11 +4097,11 @@ namespace UE::Mutable::Private
 			}
 			case 1:
 			{
-				TSharedPtr<const FLayout> ComposeLayout =
+				TManagedPtr<const FLayout> ComposeLayout =
 					LoadLayout(FCacheAddress(Args.Layout, FScheduledOp::FromOpAndOptions(Args.Layout, Item, 0)));
 
 
-				TSharedPtr<const FLayout> SourceLayout =
+				TManagedPtr<const FLayout> SourceLayout =
 					LoadLayout(FCacheAddress(Args.SourceLayout, FScheduledOp::FromOpAndOptions(Args.SourceLayout, Item, 0)));
 
 				bool bHasBlocksToCompose = false;
@@ -4248,10 +4138,10 @@ namespace UE::Mutable::Private
 			case 2:
 			{
 				FScheduledOpData& Data = HeapData[Item.CustomState];
-				TSharedPtr<const FLayout> ComposeLayout = StaticCastSharedPtr<const FLayout>(Data.Resource);
+				TManagedPtr<const FLayout> ComposeLayout = StaticCastManagedPtr<const FLayout>(Data.Resource);
 
 				FScheduledOpData& SourceData = HeapData[Data.MultiCompose.NextHeapIndex];
-				TSharedPtr<const FLayout> SourceLayout = StaticCastSharedPtr<const FLayout>(SourceData.Resource);
+				TManagedPtr<const FLayout> SourceLayout = StaticCastManagedPtr<const FLayout>(SourceData.Resource);
 
 				FIntVector2 SourceLayoutGridSize = SourceLayout->GetGridSize();
 				FIntVector2 ComposeLayoutGridSize = ComposeLayout->GetGridSize();
@@ -4319,19 +4209,19 @@ namespace UE::Mutable::Private
 			{
 				MUTABLE_CPUPROFILER_SCOPE(RunCode_ImageMultiCompose);
 
-				TSharedPtr<const FImage> BaseImage = LoadImage(FCacheAddress(Args.Base, Item));
+				TManagedPtr<const FImage> BaseImage = LoadImage(FCacheAddress(Args.Base, Item));
 
 				FScheduledOpData& Data = HeapData[Item.CustomState];
 				if (BaseImage && Data.MultiCompose.bHasBlocksToCompose)
 				{
-					TSharedPtr<const FLayout> ComposeLayout = StaticCastSharedPtr<const FLayout>(Data.Resource);
+					TManagedPtr<const FLayout> ComposeLayout = StaticCastManagedPtr<const FLayout>(Data.Resource);
 
 					FScheduledOpData& SourceData = HeapData[Data.MultiCompose.NextHeapIndex];
-					TSharedPtr<const FLayout> SourceLayout = StaticCastSharedPtr<const FLayout>(SourceData.Resource);
+					TManagedPtr<const FLayout> SourceLayout = StaticCastManagedPtr<const FLayout>(SourceData.Resource);
 
 					int32 SourceImageLOD = Data.MultiCompose.SourceLOD;
 
-					TSharedPtr<const FImage> SourceImage = 
+					TManagedPtr<const FImage> SourceImage = 
 							LoadImage(FCacheAddress(Args.SourceImage, Item.ExecutionIndex, SourceImageLOD));
 
 					if (SourceImage)
@@ -4344,14 +4234,13 @@ namespace UE::Mutable::Private
 						{
 							MUTABLE_CPUPROFILER_SCOPE(RunCode_ImageMultiCompose_SizeFixup);
 							
-							TSharedPtr<FImage> Resized = CreateImage(ExpectedSourceSize.X, ExpectedSourceSize.Y, 1, SourceImage->GetFormat(), EInitializationType::NotInitialized);
+							TManagedPtr<FImage> Resized = CreateImage(ExpectedSourceSize.X, ExpectedSourceSize.Y, 1, SourceImage->GetFormat(), EInitializationType::NotInitialized);
 							ImOp.ImageResizeLinear(Resized.Get(), Settings.ImageCompressionQuality, SourceImage.Get());
 							
-							Release(SourceImage);
 							SourceImage = Resized;	
 						}
 
-						TSharedPtr<FImage> Result = CloneOrTakeOver(BaseImage);
+						TManagedPtr<FImage> Result = CloneOrTakeOver(MoveTemp(BaseImage));
 
 						FIntVector2 SourceLayoutGridSize = SourceLayout->GetGridSize();
 						FIntVector2 ComposeLayoutGridSize = ComposeLayout->GetGridSize();
@@ -4413,7 +4302,7 @@ namespace UE::Mutable::Private
 						{
 							MUTABLE_CPUPROFILER_SCOPE(RunCode_ImageTransform_BilinearMipGen);
 
-							TSharedPtr<FImage> OwnedSource = CloneOrTakeOver(SourceImage);
+							TManagedPtr<FImage> OwnedSource = CloneOrTakeOver(MoveTemp(SourceImage));
 							OwnedSource->DataStorage.SetNumLODs(NumSrcLODsNeeded);
 
 							FMipmapGenerationSettings MipSettings{};
@@ -4426,7 +4315,6 @@ namespace UE::Mutable::Private
 						ImageMultiCompose(*Result, *SourceImage, DstRects, SrcRects, SelectedLODForRect, Settings.ImageCompressionQuality);
 						Result->DataStorage.SetNumLODs(1);
 
-						Release(SourceImage);
 						StoreImage(Item, Result);
 					}
 					else
@@ -4526,30 +4414,29 @@ namespace UE::Mutable::Private
 
                 if (Bifactor < UE_SMALL_NUMBER)
                 {
-                    TSharedPtr<const FImage> Source = LoadImage(FCacheAddress(Args.Targets[Min], Item));
+                    TManagedPtr<const FImage> Source = LoadImage(FCacheAddress(Args.Targets[Min], Item));
 					StoreImage(Item, Source);
 				}
                 else if (Bifactor > 1.0f - UE_SMALL_NUMBER)
                 {
-                    TSharedPtr<const FImage> Source = LoadImage(FCacheAddress(Args.Targets[Max], Item));
+                    TManagedPtr<const FImage> Source = LoadImage(FCacheAddress(Args.Targets[Max], Item));
 					StoreImage(Item, Source);
 				}
                 else
                 {
-					TSharedPtr<const FImage> pMin = LoadImage(FCacheAddress(Args.Targets[Min], Item));
-                    TSharedPtr<const FImage> pMax = LoadImage(FCacheAddress(Args.Targets[Max], Item));
+					TManagedPtr<const FImage> pMin = LoadImage(FCacheAddress(Args.Targets[Min], Item));
+                    TManagedPtr<const FImage> pMax = LoadImage(FCacheAddress(Args.Targets[Max], Item));
 
                     if (pMin && pMax)
                     {						
-						TSharedPtr<FImage> pNew = CloneOrTakeOver(pMin);
+						TManagedPtr<FImage> pNew = CloneOrTakeOver(MoveTemp(pMin));
 
 						// Be defensive: ensure image sizes match.
 						if (pNew->GetSize() != pMax->GetSize())
 						{
 							MUTABLE_CPUPROFILER_SCOPE(ImageResize_ForInterpolate);
-							TSharedPtr<FImage> Resized = CreateImage(pNew->GetSizeX(), pNew->GetSizeY(), pMax->GetLODCount(), pMax->GetFormat(), EInitializationType::NotInitialized);
+							TManagedPtr<FImage> Resized = CreateImage(pNew->GetSizeX(), pNew->GetSizeY(), pMax->GetLODCount(), pMax->GetFormat(), EInitializationType::NotInitialized);
 							ImOp.ImageResizeLinear(Resized.Get(), 0, pMax.Get());
-							Release(pMax);
 							pMax = Resized;
 						}
 
@@ -4558,13 +4445,12 @@ namespace UE::Mutable::Private
 						{
 							MUTABLE_CPUPROFILER_SCOPE(Format_ForInterpolate);
 
-							TSharedPtr<FImage> Formatted = CreateImage(pMax->GetSizeX(), pMax->GetSizeY(), pMax->GetLODCount(), pNew->GetFormat(), EInitializationType::NotInitialized);
+							TManagedPtr<FImage> Formatted = CreateImage(pMax->GetSizeX(), pMax->GetSizeY(), pMax->GetLODCount(), pNew->GetFormat(), EInitializationType::NotInitialized);
 							
 							bool bSuccess = false;
 							ImOp.ImagePixelFormat(bSuccess, Settings.ImageCompressionQuality, Formatted.Get(), pMax.Get());
 							check(bSuccess);
 							
-							Release(pMax);
 							pMax = Formatted;
 						}
 
@@ -4589,7 +4475,8 @@ namespace UE::Mutable::Private
 
 							int32 StartLevel = pMax->GetLODCount() - 1;
 
-							TSharedPtr<FImage> MaxFix = CloneOrTakeOver(pMax);
+							TManagedPtr<FImage> MaxFix = CloneOrTakeOver(MoveTemp(pMax));
+
 							MaxFix->DataStorage.SetNumLODs(LevelCount);
 							
 							FMipmapGenerationSettings MipSettings{};
@@ -4600,7 +4487,6 @@ namespace UE::Mutable::Private
 
                         ImageInterpolate(pNew.Get(), pMax.Get(), Bifactor);
 
-						Release(pMax);
 						StoreImage(Item, pNew);
 					}
                     else if (pMin)
@@ -4661,13 +4547,12 @@ namespace UE::Mutable::Private
             {
            		MUTABLE_CPUPROFILER_SCOPE(IM_LUMINANCE_1)
             		
-                TSharedPtr<const FImage> Base = LoadImage( FCacheAddress(Args.Base,Item) );
+                TManagedPtr<const FImage> Base = LoadImage( FCacheAddress(Args.Base,Item) );
 
-				TSharedPtr<FImage> Result = CreateImage(Base->GetSizeX(), Base->GetSizeY(), Base->GetLODCount(), EImageFormat::L_UByte, EInitializationType::NotInitialized);
-                ImageLuminance( Result.Get(),Base.Get() );
+				TManagedPtr<FImage> Result = CreateImage(Base->GetSizeX(), Base->GetSizeY(), Base->GetLODCount(), EImageFormat::L_UByte, EInitializationType::NotInitialized);
+                ImageLuminance(Result.Get(), Base.Get());
 
-				Release(Base);
-				StoreImage( Item, Result );
+				StoreImage(Item, Result);
                 break;
             }
 
@@ -4719,9 +4604,9 @@ namespace UE::Mutable::Private
             {
             	MUTABLE_CPUPROFILER_SCOPE(IM_COLOURMAP_1)
             		
-                TSharedPtr<const FImage> Source = LoadImage( FCacheAddress(Args.Base,Item) );
-                TSharedPtr<const FImage> Mask = LoadImage( FCacheAddress(Args.Mask,Item) );
-                TSharedPtr<const FImage> Map = LoadImage( FCacheAddress(Args.Map,Item) );
+                TManagedPtr<const FImage> Source = LoadImage( FCacheAddress(Args.Base,Item) );
+                TManagedPtr<const FImage> Mask = LoadImage( FCacheAddress(Args.Mask,Item) );
+                TManagedPtr<const FImage> Map = LoadImage( FCacheAddress(Args.Map,Item) );
 
 				bool bOnlyOneMip = (Mask->GetLODCount() < Source->GetLODCount());
 
@@ -4729,13 +4614,12 @@ namespace UE::Mutable::Private
 				if (Mask->GetSize() != Source->GetSize())
 				{
 					MUTABLE_CPUPROFILER_SCOPE(ImageResize_ForColourmap);
-					TSharedPtr<FImage> Resized = CreateImage(Source->GetSizeX(), Source->GetSizeY(), 1, Mask->GetFormat(), EInitializationType::NotInitialized);
+					TManagedPtr<FImage> Resized = CreateImage(Source->GetSizeX(), Source->GetSizeY(), 1, Mask->GetFormat(), EInitializationType::NotInitialized);
 					ImOp.ImageResizeLinear(Resized.Get(), 0, Mask.Get());
-					Release(Mask);
 					Mask = Resized;
 				}
 
-				TSharedPtr<FImage> Result = CreateImage(Source->GetSizeX(), Source->GetSizeY(), Source->GetLODCount(), Source->GetFormat(), EInitializationType::NotInitialized);
+				TManagedPtr<FImage> Result = CreateImage(Source->GetSizeX(), Source->GetSizeY(), Source->GetLODCount(), Source->GetFormat(), EInitializationType::NotInitialized);
 				ImageColourMap( Result.Get(), Source.Get(), Mask.Get(), Map.Get(), bOnlyOneMip);
 
 				if (bOnlyOneMip)
@@ -4745,10 +4629,7 @@ namespace UE::Mutable::Private
 					ImageMipmapInPlace(Settings.ImageCompressionQuality, Result.Get(), DummyMipSettings);
 				}
 
-				Release(Source);
-				Release(Mask);
-				Release(Map);
-				StoreImage( Item, Result );
+				StoreImage(Item, Result);
                 break;
             }
 
@@ -4774,15 +4655,14 @@ namespace UE::Mutable::Private
             {
             	MUTABLE_CPUPROFILER_SCOPE(IM_BINARISE_1)
             		
-                TSharedPtr<const FImage> pA = LoadImage( FCacheAddress(Args.Base,Item) );
+                TManagedPtr<const FImage> pA = LoadImage( FCacheAddress(Args.Base,Item) );
 
                 float c = LoadScalar(FScheduledOp::FromOpAndOptions(Args.Threshold, Item, 0));
 
-                TSharedPtr<FImage> Result = CreateImage(pA->GetSizeX(), pA->GetSizeY(), pA->GetLODCount(), EImageFormat::L_UByte, EInitializationType::NotInitialized);
-				ImageBinarise( Result.Get(), pA.Get(), c );
+                TManagedPtr<FImage> Result = CreateImage(pA->GetSizeX(), pA->GetSizeY(), pA->GetLODCount(), EImageFormat::L_UByte, EInitializationType::NotInitialized);
+				ImageBinarise(Result.Get(), pA.Get(), c);
 
-				Release(pA);
-				StoreImage( Item, Result );
+				StoreImage(Item, Result);
                 break;
             }
 
@@ -4847,11 +4727,11 @@ namespace UE::Mutable::Private
 					--LODs;
 				}
 
-                TSharedPtr<FImage> pA = CreateImage( SizeX, SizeY, FMath::Max(LODs,1), EImageFormat(Args.Format), EInitializationType::NotInitialized );
+                TManagedPtr<FImage> pA = CreateImage( SizeX, SizeY, FMath::Max(LODs,1), EImageFormat(Args.Format), EInitializationType::NotInitialized );
 
 				ImOp.FillColor(pA.Get(), c);
 
-				StoreImage( Item, pA );
+				StoreImage(Item, pA);
                 break;
             }
 
@@ -4869,7 +4749,7 @@ namespace UE::Mutable::Private
 			{
 			case 0:
 			{
-				TSharedPtr<FImage> Result;
+				TManagedPtr<FImage> Result;
 				if (Args.ForceLoad)
 				{
 					// This should never be reached because it should have been caught as a Task in IssueOp
@@ -4898,7 +4778,7 @@ namespace UE::Mutable::Private
 			UTexture* Image = pParams->GetTextureValue(Args.variable, Index.Get());
 
        		TPassthroughObjectPtr<UTexture> ExternalObject(Image);
-			TSharedPtr<FImage> Result = FImage::CreateAsReference(ExternalObject, FImageDesc(), false);
+			TManagedPtr<FImage> Result = FImage::CreateAsReference(ExternalObject, FImageDesc(), false);
         		
 			StoreImage(Item, Result);
 			break;
@@ -4917,7 +4797,7 @@ namespace UE::Mutable::Private
             {
             	MUTABLE_CPUPROFILER_SCOPE(IM_CROP_1)
             		
-                TSharedPtr<const FImage> pA = LoadImage( FCacheAddress(Args.source,Item) );
+                TManagedPtr<const FImage> pA = LoadImage( FCacheAddress(Args.source,Item) );
 
                 box< UE::Math::TIntVector2<int32> > rect;
                 rect.min[0] = Args.minX;
@@ -4936,15 +4816,14 @@ namespace UE::Mutable::Private
 					MipsToSkip--;
 				}
 
-				TSharedPtr<FImage> Result;
+				TManagedPtr<FImage> Result;
 				if (!rect.IsEmpty())
 				{
 					Result = CreateImage( rect.size[0], rect.size[1], 1, pA->GetFormat(), EInitializationType::NotInitialized);
 					ImOp.ImageCrop(Result, Settings.ImageCompressionQuality, pA.Get(), rect);
 				}
-
-				Release(pA);
-				StoreImage( Item, Result );
+				
+				StoreImage(Item, Result);
                 break;
             }
 
@@ -4973,13 +4852,12 @@ namespace UE::Mutable::Private
             {
                 MUTABLE_CPUPROFILER_SCOPE(IM_PATCH_1)
 
-                TSharedPtr<const FImage> pA = LoadImage( FCacheAddress(Args.base,Item) );
-                TSharedPtr<const FImage> pB = LoadImage( FCacheAddress(Args.patch,Item) );
+                TManagedPtr<const FImage> pA = LoadImage( FCacheAddress(Args.base,Item) );
+                TManagedPtr<const FImage> pB = LoadImage( FCacheAddress(Args.patch,Item) );
 
 				// Failsafe
 				if (!pA || !pB)
 				{
-					Release(pB);
 					StoreImage(Item, pA);
 					break;
 				}
@@ -4992,7 +4870,7 @@ namespace UE::Mutable::Private
 				rect.size[0] = pB->GetSizeX();
 				rect.size[1] = pB->GetSizeY();
 
-                TSharedPtr<FImage> Result = CloneOrTakeOver(pA);
+                TManagedPtr<FImage> Result = CloneOrTakeOver(MoveTemp(pA));
 
 				bool bApplyPatch = !rect.IsEmpty();
 				if (bApplyPatch)
@@ -5013,20 +4891,18 @@ namespace UE::Mutable::Private
 
 						if (Result->GetFormat() != format)
 						{
-							TSharedPtr<FImage> Formatted = CreateImage(Result->GetSizeX(), Result->GetSizeY(), Result->GetLODCount(), format, EInitializationType::NotInitialized);
+							TManagedPtr<FImage> Formatted = CreateImage(Result->GetSizeX(), Result->GetSizeY(), Result->GetLODCount(), format, EInitializationType::NotInitialized);
 							bool bSuccess = false;
 							ImOp.ImagePixelFormat(bSuccess, Settings.ImageCompressionQuality, Formatted.Get(), Result.Get());
 							check(bSuccess);
-							Release(Result);
 							Result = Formatted;
 						}
 						if (pB->GetFormat() != format)
 						{
-							TSharedPtr<FImage> Formatted = CreateImage(pB->GetSizeX(), pB->GetSizeY(), pB->GetLODCount(), format, EInitializationType::NotInitialized);
+							TManagedPtr<FImage> Formatted = CreateImage(pB->GetSizeX(), pB->GetSizeY(), pB->GetLODCount(), format, EInitializationType::NotInitialized);
 							bool bSuccess = false;
 							ImOp.ImagePixelFormat(bSuccess, Settings.ImageCompressionQuality, Formatted.Get(), pB.Get());
 							check(bSuccess);
-							Release(pB);
 							pB = Formatted;
 						}
 					}
@@ -5055,7 +4931,6 @@ namespace UE::Mutable::Private
 					//	rect.min[0], rect.min[1], rect.size[0], rect.size[1], Result->GetSizeX(), Result->GetSizeY());
 				}
 
-				Release(pB);
 				StoreImage( Item, Result );
                 break;
             }
@@ -5093,7 +4968,7 @@ namespace UE::Mutable::Private
 				MUTABLE_CPUPROFILER_SCOPE(IM_RASTERMESH_1)
 
 				
-				TSharedPtr<const FMesh> Mesh = 
+				TManagedPtr<const FMesh> Mesh = 
 						LoadMesh(FScheduledOp::FromOpAndOptions(Args.mesh, Item, MeshContentFilter));
 
 				// If no image, we are generating a flat mesh UV raster. This is the final stage in this case.
@@ -5122,11 +4997,10 @@ namespace UE::Mutable::Private
 					}
 
                     // Flat mesh UV raster
-					TSharedPtr<FImage> ResultImage = CreateImage(SizeX, SizeY, 1, EImageFormat::L_UByte, EInitializationType::Black);
+					TManagedPtr<FImage> ResultImage = CreateImage(SizeX, SizeY, 1, EImageFormat::L_UByte, EInitializationType::Black);
 					if (Mesh)
 					{
 						ImageRasterMesh(Mesh.Get(), ResultImage.Get(), Args.LayoutIndex, Args.BlockId, CropMin, UncroppedSize);
-						Release(Mesh);
 					}
 
 					// Stop execution.
@@ -5201,7 +5075,7 @@ namespace UE::Mutable::Private
 				FScheduledOpData& Data = HeapData[Item.CustomState];
 
 				// Unsafe downcast, should be fine as it is known to be a Mesh.
-				TSharedPtr<const FMesh> Mesh = StaticCastSharedPtr<FMesh>(ConstCastSharedPtr<FResource>(Data.Resource));
+				TManagedPtr<const FMesh> Mesh = StaticCastManagedPtr<FMesh>(ConstCastManagedPtr<FResource>(Data.Resource));
 				Data.Resource = nullptr;
 
 				if (!Mesh)
@@ -5234,22 +5108,21 @@ namespace UE::Mutable::Private
 				}
 
 				// Raster with projection
-				TSharedPtr<const FImage> Source = LoadImage(FCacheAddress(Args.image, Item.ExecutionIndex, Data.RasterMesh.Mip));
+				TManagedPtr<const FImage> Source = LoadImage(FCacheAddress(Args.image, Item.ExecutionIndex, Data.RasterMesh.Mip));
 
-				TSharedPtr<const FImage> Mask = nullptr;
+				TManagedPtr<const FImage> Mask = nullptr;
 				if (Args.mask)
 				{
 					Mask = LoadImage(FCacheAddress(Args.mask, Item));
 
 					// TODO: This shouldn't happen, but be defensive.
 					FImageSize ResultSize(SizeX, SizeY);
-					if (Mask && Mask->GetSize()!= ResultSize)
+					if (Mask && Mask->GetSize() != ResultSize)
 					{
 						MUTABLE_CPUPROFILER_SCOPE(ImageResize_MaskFixForProjection);
 
-						TSharedPtr<FImage> Resized = CreateImage(SizeX, SizeY, Mask->GetLODCount(), Mask->GetFormat(), EInitializationType::NotInitialized);
+						TManagedPtr<FImage> Resized = CreateImage(SizeX, SizeY, Mask->GetLODCount(), Mask->GetFormat(), EInitializationType::NotInitialized);
 						ImOp.ImageResizeLinear(Resized.Get(), 0, Mask.Get());
-						Release(Mask);
 						Mask = Resized;
 					}
 				}
@@ -5267,14 +5140,13 @@ namespace UE::Mutable::Private
 
 				EImageFormat Format = Source ? GetUncompressedFormat(Source->GetFormat()) : EImageFormat::L_UByte;
 
-				if (Source && Source->GetFormat()!=Format)
+				if (Source && Source->GetFormat() != Format)
 				{
 					MUTABLE_CPUPROFILER_SCOPE(RunCode_RasterMesh_ReformatSource);
-					TSharedPtr<FImage> Formatted = CreateImage(Source->GetSizeX(), Source->GetSizeY(), Source->GetLODCount(), Format, EInitializationType::NotInitialized);
+					TManagedPtr<FImage> Formatted = CreateImage(Source->GetSizeX(), Source->GetSizeY(), Source->GetLODCount(), Format, EInitializationType::NotInitialized);
 					bool bSuccess = false;
 					ImOp.ImagePixelFormat(bSuccess, Settings.ImageCompressionQuality, Formatted.Get(), Source.Get());
 					check(bSuccess); 
-					Release(Source);
 					Source = Formatted;
 				}
 			
@@ -5303,10 +5175,9 @@ namespace UE::Mutable::Private
 					{
 						MUTABLE_CPUPROFILER_SCOPE(RunCode_ImageRasterMesh_SizeFixup);	
 						
-						TSharedPtr<FImage> Resized = CreateImage(ExpectedSourceSize.X, ExpectedSourceSize.Y, 1, Format, EInitializationType::NotInitialized);
+						TManagedPtr<FImage> Resized = CreateImage(ExpectedSourceSize.X, ExpectedSourceSize.Y, 1, Format, EInitializationType::NotInitialized);
 						ImOp.ImageResizeLinear(Resized.Get(), Settings.ImageCompressionQuality, Source.Get());
 						
-						Release(Source);
 						Source = Resized;	
 					}
 				}
@@ -5336,7 +5207,7 @@ namespace UE::Mutable::Private
 					{
 						MUTABLE_CPUPROFILER_SCOPE(RunCode_RasterMesh_BilinearMipGen);
 
-						TSharedPtr<FImage> OwnedSource = CloneOrTakeOver(Source);
+						TManagedPtr<FImage> OwnedSource = CloneOrTakeOver(MoveTemp(Source));
 
 						OwnedSource->DataStorage.SetNumLODs(2);
 						ImageMipmapInPlace(0, OwnedSource.Get(), FMipmapGenerationSettings{});
@@ -5346,7 +5217,7 @@ namespace UE::Mutable::Private
 				}
 
 				// Allocate new image after bilinear mip generation to reduce operation memory peak.
-				TSharedPtr<FImage> New = CreateImage(SizeX, SizeY, 1, Format, EInitializationType::Black);
+				TManagedPtr<FImage> New = CreateImage(SizeX, SizeY, 1, Format, EInitializationType::Black);
 
 				if (Args.projector && Source && Source->GetSizeX() > 0 && Source->GetSizeY() > 0)
 				{
@@ -5394,9 +5265,6 @@ namespace UE::Mutable::Private
 					}
 				}
 
-				Release(Mesh);
-				Release(Source);
-				Release(Mask);
 				StoreImage(Item, New);
 
                 break;
@@ -5422,17 +5290,15 @@ namespace UE::Mutable::Private
             {
                 MUTABLE_CPUPROFILER_SCOPE(IM_MAKEGROWMAP_1)
 
-                TSharedPtr<const FImage> Mask = LoadImage( FCacheAddress(Args.mask,Item) );
+                TManagedPtr<const FImage> Mask = LoadImage( FCacheAddress(Args.mask,Item) );
 
-                TSharedPtr<FImage> Result = CreateImage( Mask->GetSizeX(), Mask->GetSizeY(), Mask->GetLODCount(), EImageFormat::L_UByte, EInitializationType::NotInitialized);
+                TManagedPtr<FImage> Result = CreateImage( Mask->GetSizeX(), Mask->GetSizeY(), Mask->GetLODCount(), EImageFormat::L_UByte, EInitializationType::NotInitialized);
 
-				TSharedPtr<FImage> OwnedMask = CloneOrTakeOver(Mask);
+				TManagedPtr<FImage> OwnedMask = CloneOrTakeOver(MoveTemp(Mask));
 
                 ImageMakeGrowMap(Result.Get(), OwnedMask.Get(), Args.border);
 				Result->Flags |= FImage::IF_CANNOT_BE_SCALED;
 
-				Release(Mask);
-				Release(OwnedMask);
                 StoreImage( Item, Result);
                 break;
             }
@@ -5459,12 +5325,11 @@ namespace UE::Mutable::Private
             {
                 MUTABLE_CPUPROFILER_SCOPE(IM_DISPLACE_1)
 
-                TSharedPtr<const FImage> Source = LoadImage( FCacheAddress(Args.Source,Item) );
-                TSharedPtr<const FImage> pMap = LoadImage( FCacheAddress(Args.DisplacementMap,Item) );
+                TManagedPtr<const FImage> Source = LoadImage( FCacheAddress(Args.Source,Item) );
+                TManagedPtr<const FImage> pMap = LoadImage( FCacheAddress(Args.DisplacementMap,Item) );
 
 				if (!Source)
 				{
-					Release(pMap);
 					StoreImage(Item, nullptr);
 					break;
 				}
@@ -5477,15 +5342,14 @@ namespace UE::Mutable::Private
 				{
 					MUTABLE_CPUPROFILER_SCOPE(ImageResize_EmergencyHackForDisplacementStep1);
 
-					TSharedPtr<FImage> Resized = CreateImage(pMap->GetSizeX(), pMap->GetSizeY(), Source->GetLODCount(), Source->GetFormat(), EInitializationType::NotInitialized);
+					TManagedPtr<FImage> Resized = CreateImage(pMap->GetSizeX(), pMap->GetSizeY(), Source->GetLODCount(), Source->GetFormat(), EInitializationType::NotInitialized);
 					ImOp.ImageResizeLinear(Resized.Get(), 0, Source.Get());
-					Release(Source);
 					Source = Resized;
 				}
 
 				// This works based on the assumption that displacement maps never read from a position they actually write to.
 				// Since they are used for UV border expansion, this should always be the case.
-				TSharedPtr<FImage> Result = CloneOrTakeOver(Source);
+				TManagedPtr<FImage> Result = CloneOrTakeOver(MoveTemp(Source));
 
 				if (OriginalSourceScale[0] > 0 && OriginalSourceScale[1] > 0)
 				{
@@ -5494,14 +5358,12 @@ namespace UE::Mutable::Private
 					if (OriginalSourceScale != Result->GetSize())
 					{
 						MUTABLE_CPUPROFILER_SCOPE(ImageResize_EmergencyHackForDisplacementStep2);
-						TSharedPtr<FImage> Resized = CreateImage(OriginalSourceScale[0], OriginalSourceScale[1], Result->GetLODCount(), Result->GetFormat(), EInitializationType::NotInitialized);
+						TManagedPtr<FImage> Resized = CreateImage(OriginalSourceScale[0], OriginalSourceScale[1], Result->GetLODCount(), Result->GetFormat(), EInitializationType::NotInitialized);
 						ImOp.ImageResizeLinear(Resized.Get(), 0, Result.Get());
-						Release(Result);
 						Result = Resized;
 					}
 				}
 
-				Release(pMap);
                 StoreImage( Item, Result);
                 break;
             }
@@ -5617,7 +5479,7 @@ namespace UE::Mutable::Private
 				const FScheduledOpData OpHeapData = HeapData[Item.CustomState];
 
 				const uint8 Mip = static_cast<uint8>(FMath::Max(0, FMath::FloorToInt(OpHeapData.ImageTransform.MipValue)));
-				TSharedPtr<const FImage> Source = LoadImage(FCacheAddress(Args.Base, Item.ExecutionIndex, Mip));
+				TManagedPtr<const FImage> Source = LoadImage(FCacheAddress(Args.Base, Item.ExecutionIndex, Mip));
 
 				const FVector2f Offset = FVector2f(
                         Args.OffsetX ? LoadScalar(FCacheAddress(Args.OffsetX, Item)) : 0.0f,
@@ -5649,12 +5511,11 @@ namespace UE::Mutable::Private
 				if (Format != SourceFormat)
 				{
 					MUTABLE_CPUPROFILER_SCOPE(RunCode_ImageTransform_FormatFixup);	
-					TSharedPtr<FImage> Formatted = CreateImage(Source->GetSizeX(), Source->GetSizeY(), Source->GetLODCount(), Format, EInitializationType::NotInitialized);
+					TManagedPtr<FImage> Formatted = CreateImage(Source->GetSizeX(), Source->GetSizeY(), Source->GetLODCount(), Format, EInitializationType::NotInitialized);
 					bool bSuccess = false;
 					ImOp.ImagePixelFormat(bSuccess, Settings.ImageCompressionQuality, Formatted.Get(), Source.Get());
 					check(bSuccess); 
 
-					Release(Source);
 					Source = Formatted;
 				}
 
@@ -5665,10 +5526,9 @@ namespace UE::Mutable::Private
 				{
 					MUTABLE_CPUPROFILER_SCOPE(RunCode_ImageTransform_SizeFixup);	
 					
-					TSharedPtr<FImage> Resized = CreateImage(ExpectedSourceSize.X, ExpectedSourceSize.Y, 1, Format, EInitializationType::NotInitialized);
+					TManagedPtr<FImage> Resized = CreateImage(ExpectedSourceSize.X, ExpectedSourceSize.Y, 1, Format, EInitializationType::NotInitialized);
 					ImOp.ImageResizeLinear(Resized.Get(), Settings.ImageCompressionQuality, Source.Get());
 					
-					Release(Source);
 					Source = Resized;	
 				}
 
@@ -5676,7 +5536,7 @@ namespace UE::Mutable::Private
 				{
 					MUTABLE_CPUPROFILER_SCOPE(RunCode_ImageTransform_BilinearMipGen);
 
-					TSharedPtr<FImage> OwnedSource = CloneOrTakeOver(Source);
+					TManagedPtr<FImage> OwnedSource = CloneOrTakeOver(MoveTemp(Source));
 					OwnedSource->DataStorage.SetNumLODs(2);
 
 					ImageMipmapInPlace(0, OwnedSource.Get(), FMipmapGenerationSettings{});
@@ -5707,13 +5567,12 @@ namespace UE::Mutable::Private
 											       ? EInitializationType::Black
 											       : EInitializationType::NotInitialized;
 
-				TSharedPtr<FImage> Result = CreateImage(
+				TManagedPtr<FImage> Result = CreateImage(
 						OpHeapData.ImageTransform.SizeX, OpHeapData.ImageTransform.SizeY, 1, Format, InitType);
 
 				const float MipFactor = FMath::Frac(FMath::Max(0.0f, OpHeapData.ImageTransform.MipValue));
 				ImageTransform(Result.Get(), Source.Get(), Transform, MipFactor, AddressMode, bUseImageTransformVectorImpl);
 
-				Release(Source);
 				StoreImage(Item, Result);
 
                 break;
@@ -5740,7 +5599,7 @@ namespace UE::Mutable::Private
 			}
 			case 1:
 			{
-				TSharedPtr<const FMaterial> Material = LoadMaterial(FCacheAddress(Args.Material, Item));
+				TManagedPtr<const FMaterial> Material = LoadMaterial(FCacheAddress(Args.Material, Item));
 
 				check(Args.ParameterName < (uint32)Program.ConstantStrings.Num());
 				const FString& ParameterName = Program.ConstantStrings[Args.ParameterName];
@@ -5753,7 +5612,7 @@ namespace UE::Mutable::Private
 				}
 				else
 				{
-					const TVariant<OP::ADDRESS, TSharedPtr<const FImage>> ImageParameter = ImageParameterData->ImageParameter;
+					const TVariant<OP::ADDRESS, TManagedPtr<const FImage>> ImageParameter = ImageParameterData->ImageParameter;
 
 					if (ImageParameter.IsType<OP::ADDRESS>())
 					{
@@ -5765,14 +5624,14 @@ namespace UE::Mutable::Private
 					}
 					else
 					{
-						StoreImage(Item, ImageParameter.Get<TSharedPtr<const FImage>>());
+						StoreImage(Item, ImageParameter.Get<TManagedPtr<const FImage>>());
 					}
 				}
 				break;
 			}
 			case 2:
 			{
-				TSharedPtr<const FMaterial> Material = LoadMaterial(FCacheAddress(Args.Material, Item));
+				TManagedPtr<const FMaterial> Material = LoadMaterial(FCacheAddress(Args.Material, Item));
 
 				check(Args.ParameterName < (uint32)Program.ConstantStrings.Num());
 				const FString& ParameterName = Program.ConstantStrings[Args.ParameterName];
@@ -5785,8 +5644,8 @@ namespace UE::Mutable::Private
 				}
 				else
 				{
-					const TVariant<OP::ADDRESS, TSharedPtr<const FImage>>& ImageParameter = ImageParameterData->ImageParameter;
-					TSharedPtr<const FImage> pImage = LoadImage(FCacheAddress(ImageParameter.Get<OP::ADDRESS>(), Item));
+					const TVariant<OP::ADDRESS, TManagedPtr<const FImage>>& ImageParameter = ImageParameterData->ImageParameter;
+					TManagedPtr<const FImage> pImage = LoadImage(FCacheAddress(ImageParameter.Get<OP::ADDRESS>(), Item));
 
 					StoreImage(Item, pImage);
 				}
@@ -6256,8 +6115,8 @@ namespace UE::Mutable::Private
 			}
 			case 1:
 			{
-				TSharedPtr<const FMaterial> Material = LoadMaterial(FCacheAddress(Args.Material, Item));
-					
+				TManagedPtr<const FMaterial> Material = LoadMaterial(FCacheAddress(Args.Material, Item));
+				
 				check(Args.ParameterName < (uint32)Program.ConstantStrings.Num());
 				const FString& ParameterName = Program.ConstantStrings[Args.ParameterName];
 				const FParameterKey ParameterKey = { FName(ParameterName), Args.LayerIndex };
@@ -6297,8 +6156,8 @@ namespace UE::Mutable::Private
 			OP::ResourceConstantArgs Args = Program.GetOpArgs<OP::ResourceConstantArgs>( Item.At );
             check( Args.value < (uint32)InModel->GetPrivate()->Program.ConstantStrings.Num() );
 
-            const FString& result = Program.ConstantStrings[Args.value];
-            StoreString( Item, MakeShared<String>(result) );
+            const FString& Result = Program.ConstantStrings[Args.value];
+            StoreString( Item, MakeManaged<String>(Result) );
 
             break;
         }
@@ -6307,9 +6166,9 @@ namespace UE::Mutable::Private
         {
 			OP::ParameterArgs Args = Program.GetOpArgs<OP::ParameterArgs>( Item.At );
 			TSharedPtr<FRangeIndex> Index = BuildCurrentOpRangeIndex( Item, Args.variable );
-			FString result;
-			pParams->GetStringValue(Args.variable, result, Index.Get());
-            StoreString( Item, MakeShared<String>(result) );
+			FString Result;
+			pParams->GetStringValue(Args.variable, Result, Index.Get());
+            StoreString( Item, MakeManaged<String>(Result) );
             break;
         }
 
@@ -6367,7 +6226,7 @@ namespace UE::Mutable::Private
                 float x = Args.X ? LoadScalar( FCacheAddress(Args.X,Item) ) : 0.5f;
                 float y = Args.Y ? LoadScalar( FCacheAddress(Args.Y,Item) ) : 0.5f;
 
-                TSharedPtr<const FImage> pImage = LoadImage(FScheduledOp::FromOpAndOptions(Args.Image, Item, 0));
+                TManagedPtr<const FImage> pImage = LoadImage(FScheduledOp::FromOpAndOptions(Args.Image, Item, 0));
 
 				FVector4f result;
                 if (pImage)
@@ -6387,7 +6246,6 @@ namespace UE::Mutable::Private
                     result = FVector4f();
                 }
 
-				Release(pImage);
                 StoreColor( Item, result );
                 break;
             }
@@ -6572,7 +6430,7 @@ namespace UE::Mutable::Private
 			}
 			case 1:
 			{
-				TSharedPtr<const FMaterial> Material = LoadMaterial(FCacheAddress(Args.Material, Item));
+				TManagedPtr<const FMaterial> Material = LoadMaterial(FCacheAddress(Args.Material, Item));
 
 				check(Args.ParameterName < (uint32)Program.ConstantStrings.Num());
 				const FString& ParameterName = Program.ConstantStrings[Args.ParameterName];
@@ -6687,8 +6545,8 @@ namespace UE::Mutable::Private
 			TSharedPtr<FRangeIndex> Index = BuildCurrentOpRangeIndex(Item, Args.Variable);
 
 			UMaterialInterface* MaterialInterface = InParams->GetMaterialValue(Args.Variable, Index.Get());
-
-			TSharedPtr<FMaterial> Material = MakeShared<FMaterial>();
+			TManagedPtr<FMaterial> Material = MakeManaged<FMaterial>();
+			
 			Material->PassthroughObject = TPassthroughObjectPtr(MaterialInterface);
  
 			if (MaterialInterface)
@@ -6785,7 +6643,7 @@ namespace UE::Mutable::Private
 					// Store Texture into the FMaterial
 					if (bParameterFound)
 					{
-						TVariant<OP::ADDRESS, TSharedPtr<const FImage>> NewParameterImage;
+						TVariant<OP::ADDRESS, TManagedPtr<const FImage>> NewParameterImage;
 						NewParameterImage.Set<OP::ADDRESS>(ImageParameterOperations[ParameterIndex]);
 
 						FMaterial::FImageParameterData NewImageParameter;
@@ -6883,12 +6741,12 @@ namespace UE::Mutable::Private
 					}
 				}
 
-				TSharedPtr<const FMaterial> Material = LoadMaterial(FCacheAddress(Args.Material, Item));
-				TSharedPtr<FMaterial> Result;
+				TManagedPtr<const FMaterial> Material = LoadMaterial(FCacheAddress(Args.Material, Item));
+				TManagedPtr<FMaterial> Result;
 
 				if (Material.IsValid())
 				{
-					Result = UE::Mutable::Private::CloneOrTakeOver<FMaterial>(Material);
+					Result = UE::Mutable::Private::CloneOrTakeOver<FMaterial>(MoveTemp(Material));
 					int32 ImageIndex = 0;
 				
 					// ParameterIndex 0 is reserver for the material operation
@@ -6906,7 +6764,7 @@ namespace UE::Mutable::Private
 						case EDataType::Image:
 						{
 							//TODO: Evaluate the non-lazy textures
-							TVariant<OP::ADDRESS, TSharedPtr<const FImage>> NewParameterImage;
+							TVariant<OP::ADDRESS, TManagedPtr<const FImage>> NewParameterImage;
 							NewParameterImage.Set<OP::ADDRESS>(ParameterOperations[ParameterIndex]);
 
 							FMaterial::FImageParameterData ImageParameterData;
@@ -6993,7 +6851,7 @@ namespace UE::Mutable::Private
 				}
 			case 1:
 				{
-					TSharedPtr<const FSkeletalMesh> LoadedSkeletalMesh = LoadSkeletalMesh(FCacheAddress(Args.SkeletalMesh, Item));
+					TManagedPtr<const FSkeletalMesh> LoadedSkeletalMesh = LoadSkeletalMesh(FCacheAddress(Args.SkeletalMesh, Item));
 					
 					TPassthroughObjectPtr<USkeletalMesh> PassthroughMesh = LoadedSkeletalMesh->PassthroughObject;
 					check(PassthroughMesh);
@@ -7015,7 +6873,7 @@ namespace UE::Mutable::Private
 						}
 					}
 
-					TSharedPtr<FMaterial> MutableMaterial = MakeShared<FMaterial>();
+					TManagedPtr<FMaterial> MutableMaterial = MakeManaged<FMaterial>();
 					MutableMaterial->PassthroughObject = TPassthroughObjectPtr(FoundMaterialInterface.Get());
 
 					StoreMaterial(Item, MutableMaterial);
@@ -7051,9 +6909,8 @@ namespace UE::Mutable::Private
 			OP::ResourceConstantArgs Args = Program.GetOpArgs<OP::ResourceConstantArgs>(Item.At);
             check( Args.value < (uint32)InModel->GetPrivate()->Program.ConstantLayouts.Num() );
 
-            TSharedPtr<const FLayout> Result = Program.ConstantLayouts
-                    [ Args.value ];
-            StoreLayout( Item, Result );
+            TManagedPtr<const FLayout> Result = Program.ConstantLayouts[ Args.value ];
+            StoreLayout(Item, Result);
             break;
         }
 
@@ -7070,14 +6927,14 @@ namespace UE::Mutable::Private
 
             case 1:
             {
-                TSharedPtr<const FLayout> pA = LoadLayout( FCacheAddress(Args.Base,Item) );
-                TSharedPtr<const FLayout> pB = LoadLayout( FCacheAddress(Args.Added,Item) );
+                TManagedPtr<const FLayout> pA = LoadLayout( FCacheAddress(Args.Base,Item) );
+                TManagedPtr<const FLayout> pB = LoadLayout( FCacheAddress(Args.Added,Item) );
 
-                TSharedPtr<const FLayout> Result;
+                TManagedPtr<const FLayout> Result;
 
                 if (pA && pB)
                 {
-					Result = LayoutMerge(pA.Get(),pB.Get());
+					Result = LayoutMerge(pA.Get(), pB.Get());
                 }
                 else if (pA)
                 {
@@ -7088,7 +6945,8 @@ namespace UE::Mutable::Private
                     Result = pB->Clone();
                 }
 
-                StoreLayout( Item, Result );
+				StoreLayout(Item, Result);
+
                 break;
             }
 
@@ -7111,9 +6969,9 @@ namespace UE::Mutable::Private
 
             case 1:
             {
-                TSharedPtr<const FLayout> Source = LoadLayout( FCacheAddress(Args.Source,Item) );
+                TManagedPtr<const FLayout> Source = LoadLayout( FCacheAddress(Args.Source,Item) );
 
-				TSharedPtr<FLayout> Result;
+				TManagedPtr<FLayout> Result;
 
 				if (Source)
 				{
@@ -7121,7 +6979,6 @@ namespace UE::Mutable::Private
 
 					LayoutPack3(Result.Get(), Source.Get() );
 				}
-
                 StoreLayout( Item, Result );
                 break;
             }
@@ -7147,12 +7004,10 @@ namespace UE::Mutable::Private
 
 			case 1:
 			{
-				TSharedPtr<const FMesh> Mesh = LoadMesh(
+				TManagedPtr<const FMesh> Mesh = LoadMesh(
 						FScheduledOp::FromOpAndOptions(Args.Mesh, Item, MeshContentFilter));
 
-				TSharedPtr<const FLayout> Result = LayoutFromMesh_RemoveBlocks(Mesh.Get(), Args.LayoutIndex);
-
-				Release(Mesh);
+				TManagedPtr<const FLayout> Result = LayoutFromMesh_RemoveBlocks(Mesh.Get(), Args.LayoutIndex);
 				StoreLayout(Item, Result);
 				break;
 			}
@@ -7177,10 +7032,10 @@ namespace UE::Mutable::Private
 
 			case 1:
 			{
-				TSharedPtr<const FLayout> Source = LoadLayout(FCacheAddress(Args.Source, Item));
-				TSharedPtr<const FLayout> ReferenceLayout = LoadLayout(FCacheAddress(Args.ReferenceLayout, Item));
+				TManagedPtr<const FLayout> Source = LoadLayout(FCacheAddress(Args.Source, Item));
+				TManagedPtr<const FLayout> ReferenceLayout = LoadLayout(FCacheAddress(Args.ReferenceLayout, Item));
 
-				TSharedPtr<const FLayout> Result;
+				TManagedPtr<const FLayout> Result;
 
 				if (Source && ReferenceLayout)
 				{
@@ -7343,7 +7198,7 @@ namespace UE::Mutable::Private
 									const FVector4f Vector = LoadColor(FCacheAddress(OperantAt, Item));
 									
 									FValue Value;
-									Value.Ptr = MakeShared<FInstancedStruct>();
+									Value.Ptr = MakeManaged<FInstancedStruct>();
 									Value.Ptr->InitializeAs<FVectorAdapter>().Value = Vector;
 
 									Context.Inputs.Add(Input.Key, MoveTemp(Value));
@@ -7352,11 +7207,11 @@ namespace UE::Mutable::Private
 								
 							case EDataType::Image:
 								{
-									const TSharedPtr<const FImage> Image = LoadImage(FCacheAddress(OperantAt, Item));
+									const TManagedPtr<const FImage> Image = LoadImage(FCacheAddress(OperantAt, Item));
 									
 									FValue Value;
-									Value.Ptr = MakeShared<FInstancedStruct>();
-									Value.Ptr->InitializeAs<FTextureAdapter>().Image = ConstCastSharedPtr<FImage>(Image);
+									Value.Ptr = MakeManaged<FInstancedStruct>();
+									Value.Ptr->InitializeAs<FTextureAdapter>().Image = ConstCastManagedPtr<FImage>(Image);
 									
 									Context.Inputs.Add(Input.Key, MoveTemp(Value));
 									break;
@@ -7364,11 +7219,11 @@ namespace UE::Mutable::Private
 
 							case EDataType::Mesh:
 								{
-									const TSharedPtr<const FMesh> Mesh = LoadMesh(FCacheAddress(OperantAt, Item));
+									const TManagedPtr<const FMesh> Mesh = LoadMesh(FCacheAddress(OperantAt, Item));
 
 									FValue Value;
-									Value.Ptr = MakeShared<FInstancedStruct>();
-									Value.Ptr->InitializeAs<FMeshAdapter>().Mesh = ConstCastSharedPtr<FMesh>(Mesh);
+									Value.Ptr = MakeManaged<FInstancedStruct>();
+									Value.Ptr->InitializeAs<FMeshAdapter>().Mesh = ConstCastManagedPtr<FMesh>(Mesh);
 
 									Context.Inputs.Add(Input.Key, MoveTemp(Value));
 									break;
@@ -7376,11 +7231,11 @@ namespace UE::Mutable::Private
 								
 							case EDataType::Material:
 								{
-									const TSharedPtr<const FMaterial> Material = LoadMaterial(FCacheAddress(OperantAt, Item));
+									const TManagedPtr<const FMaterial> Material = LoadMaterial(FCacheAddress(OperantAt, Item));
 
 									FValue Value;
-									Value.Ptr = MakeShared<FInstancedStruct>();
-									Value.Ptr->InitializeAs<FMaterialAdapter>().Material = ConstCastSharedPtr<FMaterial>(Material);
+									Value.Ptr = MakeManaged<FInstancedStruct>();
+									Value.Ptr->InitializeAs<FMaterialAdapter>().Material = ConstCastManagedPtr<FMaterial>(Material);
 
 									Context.Inputs.Add(Input.Key, MoveTemp(Value));
 									break;
@@ -7388,11 +7243,11 @@ namespace UE::Mutable::Private
 								
 							case EDataType::InstancedStruct:
 								{
-									const TSharedPtr<const FInstancedStruct> InstancedStruct = LoadInstancedStruct(FCacheAddress(OperantAt, Item));
+									const TManagedPtr<const FInstancedStruct> InstancedStruct = LoadInstancedStruct(FCacheAddress(OperantAt, Item));
 									
 									FValue Value;
-									Value.Ptr = MakeShared<FInstancedStruct>();
-									Value.Ptr = ConstCastSharedPtr<FInstancedStruct>(InstancedStruct);
+									Value.Ptr = MakeManaged<FInstancedStruct>();
+									Value.Ptr = ConstCastManagedPtr<FInstancedStruct>(InstancedStruct);
 									
 									Context.Inputs.Add(Input.Key, MoveTemp(Value));
 									break;
@@ -7403,7 +7258,7 @@ namespace UE::Mutable::Private
 									const float Float = LoadScalar(FCacheAddress(OperantAt, Item));
 
 									FValue Value;
-									Value.Ptr = MakeShared<FInstancedStruct>();
+									Value.Ptr = MakeManaged<FInstancedStruct>();
 									Value.Ptr->InitializeAs<FFloatAdapter>().Value = Float;
 
 									Context.Inputs.Add(Input.Key, MoveTemp(Value));
@@ -7443,8 +7298,8 @@ namespace UE::Mutable::Private
 							
 						case EDataType::Material:
 						{
-							const TSharedPtr<FMaterial> OutputMaterial = Context.Output.Get<FMaterialAdapter>().Material;
-							StoreMaterial(FCacheAddress(Item.At, Item), OutputMaterial ? OutputMaterial :  MakeShared<FMaterial>());
+							const TManagedPtr<FMaterial> OutputMaterial = Context.Output.Get<FMaterialAdapter>().Material;
+							StoreMaterial(FCacheAddress(Item.At, Item), OutputMaterial ? OutputMaterial :  MakeManaged<FMaterial>());
 							break;
 						}
 						
@@ -7474,7 +7329,7 @@ namespace UE::Mutable::Private
 			{
 				OP::ParameterArgs Args = Program.GetOpArgs<OP::ParameterArgs>(Item.At);
 				TSharedPtr<FRangeIndex> Index = BuildCurrentOpRangeIndex(Item, Args.variable);
-				TSharedPtr<const FInstancedStruct> Result = pParams->GetInstancedStructValue(Args.variable, Index.Get());
+				TManagedPtr<const FInstancedStruct> Result = pParams->GetInstancedStructValue(Args.variable, Index.Get());
 				StoreInstancedStruct(Item, Result);
 
 				break;
@@ -7486,7 +7341,7 @@ namespace UE::Mutable::Private
 				TSharedPtr<FRangeIndex> Index = BuildCurrentOpRangeIndex(Item, Args.variable);
 				USkeletalMesh* SkeletalMesh = pParams->GetSkeletalMeshValue(Args.variable, Index.Get());
 
-				TSharedPtr<FSkeletalMesh> Result = MakeShared<FSkeletalMesh>();
+				TManagedPtr<FSkeletalMesh> Result = MakeManaged<FSkeletalMesh>();
 				Result->PassthroughObject = TPassthroughObjectPtr<USkeletalMesh>(SkeletalMesh);
         		
 				StoreSkeletalMesh(Item, Result);
@@ -7504,10 +7359,12 @@ namespace UE::Mutable::Private
 
 			case EDataType::Mesh:
 				RunCode_Mesh(Item, InModel.Get());
+				System->WorkingMemoryManager.EnsureBudgetBelow(0);
 				break;
 
 			case EDataType::Image:
 				RunCode_Image(Item, pParams, InModel.Get());
+				System->WorkingMemoryManager.EnsureBudgetBelow(0);
 				break;
 
 			case EDataType::Layout:
@@ -7554,6 +7411,7 @@ namespace UE::Mutable::Private
 		}
 
         }
+
     }
 
 	
@@ -8161,7 +8019,7 @@ namespace UE::Mutable::Private
 			{
 				FExtendedImageDesc& Result = ImageDescResults.FindOrAdd(Item.At);
 				
-				TSharedPtr<const FLayout> Layout = LoadLayout(FCacheAddress(Args.Layout, Item.ExecutionIndex, Item.ExecutionOptions, FScheduledOp::EType::Full));
+				TManagedPtr<const FLayout> Layout = LoadLayout(FCacheAddress(Args.Layout, Item.ExecutionIndex, Item.ExecutionOptions, FScheduledOp::EType::Full));
 
 				FIntPoint SizeInBlocks = Layout->GetGridSize();
 				FIntPoint BlockSizeInPixels(Args.BlockSize[0], Args.BlockSize[1]);
@@ -8695,7 +8553,7 @@ namespace UE::Mutable::Private
 				FCacheAddress CacheMaterialOp(Args.Material, Item);
 				CacheMaterialOp.Type = FScheduledOp::EType::Full;
 
-				TSharedPtr<const FMaterial> Material = LoadMaterial(CacheMaterialOp);
+				TManagedPtr<const FMaterial> Material = LoadMaterial(CacheMaterialOp);
 			
 				check(Args.ParameterName < (uint32)Program.ConstantStrings.Num());
 				const FString& ParameterName = Program.ConstantStrings[Args.ParameterName];
@@ -8704,7 +8562,7 @@ namespace UE::Mutable::Private
 
 				if (ImageParameterData)
 				{
-					const TVariant<OP::ADDRESS, TSharedPtr<const FImage>> ImageParameter = ImageParameterData->ImageParameter;
+					const TVariant<OP::ADDRESS, TManagedPtr<const FImage>> ImageParameter = ImageParameterData->ImageParameter;
 					if (ImageParameter.IsType<OP::ADDRESS>())
 					{
 						FScheduledOp ItemNextStage(Item.At, Item, 2);
@@ -8732,7 +8590,7 @@ namespace UE::Mutable::Private
 				FCacheAddress CacheMaterialOp(Args.Material, Item);
 				CacheMaterialOp.Type = FScheduledOp::EType::Full;
 
-				TSharedPtr<const FMaterial> Material = LoadMaterial(CacheMaterialOp);
+				TManagedPtr<const FMaterial> Material = LoadMaterial(CacheMaterialOp);
 
 				check(Args.ParameterName < (uint32)Program.ConstantStrings.Num());
 				const FString& ParameterName = Program.ConstantStrings[Args.ParameterName];
@@ -8741,7 +8599,7 @@ namespace UE::Mutable::Private
 					
 				if (ImageParameterData)
 				{
-					const TVariant<OP::ADDRESS, TSharedPtr<const FImage>> ImageParameter = ImageParameterData->ImageParameter;
+					const TVariant<OP::ADDRESS, TManagedPtr<const FImage>> ImageParameter = ImageParameterData->ImageParameter;
 
 					FExtendedImageDesc& Result = ImageDescResults.FindOrAdd(Item.At);
 					Result = ImageDescResults.FindChecked(ImageParameter.Get<OP::ADDRESS>());
@@ -8810,16 +8668,10 @@ namespace UE::Mutable::Private
 				return Runner->CreateImage(x, y, m, f, i);
 			},
 
-			// Release
-			[Runner](TSharedPtr<FImage>& Image)
-			{
-				Runner->Release(Image);
-			},
-
 			// Clone
 			[Runner](const FImage* Image)
 			{
-				TSharedPtr<FImage> New = Runner->CreateImage(Image->GetSizeX(), Image->GetSizeY(), Image->GetLODCount(), Image->GetFormat(), EInitializationType::NotInitialized);
+				TManagedPtr<FImage> New = Runner->CreateImage(Image->GetSizeX(), Image->GetSizeY(), Image->GetLODCount(), Image->GetFormat(), EInitializationType::NotInitialized);
 				New->Copy(Image);
 				return New;
 			},
