[Mutable] Memory managment rework.

diff --git a/Engine/Plugins/Mutable/Source/MutableRuntime/Private/MuR/CodeRunner.h b/Engine/Plugins/Mutable/Source/MutableRuntime/Private/MuR/CodeRunner.h
--- a/Engine/Plugins/Mutable/Source/MutableRuntime/Private/MuR/CodeRunner.h
+++ b/Engine/Plugins/Mutable/Source/MutableRuntime/Private/MuR/CodeRunner.h
@@ -20,6 +20,8 @@
 #include "Templates/RefCounting.h"
 #include "Templates/SharedPointer.h"
 #include "Templates/Tuple.h"
+#include "Containers/StaticArray.h"
+#include "Concepts/SameAs.h"
 
 namespace UE::Mutable::Private::MemoryCounters
 {
@@ -77,6 +79,7 @@ namespace  UE::Mutable::Private
 			uint32 RunOpsPerType[int32(EOpType::COUNT)] = {};
 		};
 
+	public:
         /** Type of data sometimes stored in the code runner heap to pass info between operation stages. */
         struct FScheduledOpData
         {
@@ -120,11 +123,11 @@ namespace  UE::Mutable::Private
 
 			};
 
-			TSharedPtr<const FResource> Resource;
+			TManagedPtr<const FResource> Resource;
         };
-
+	protected:
 		// Assertion to know when FScheduledOpData size changes. It is ok to modifiy if needed. 
-		static_assert(sizeof(FScheduledOpData) == 4*4+sizeof(TSharedPtr<FResource>), "FScheduledOpData size changed.");
+		static_assert(sizeof(FScheduledOpData) == 4*4+sizeof(TManagedPtr<FResource>), "FScheduledOpData size changed.");
     	
         TSharedPtr<FRangeIndex> BuildCurrentOpRangeIndex( const FScheduledOp&, int32 ParameterIndex);
 
@@ -164,21 +167,21 @@ namespace  UE::Mutable::Private
 		};
 
 		/** Load an external image asynchronously, returns an event to wait for complition and a cleanup function that must be called once the event has completed. */
-		TTuple<UE::Tasks::FTask, TFunction<void()>> LoadExternalImageAsync(FExternalResourceId Id, uint8 MipmapsToSkip, TFunction<void(TSharedPtr<FImage>)>& ResultCallback);
+		TTuple<UE::Tasks::FTask, TFunction<void()>> LoadExternalImageAsync(FExternalResourceId Id, uint8 MipmapsToSkip, TFunction<void(TManagedPtr<FImage>)>& ResultCallback);
  	    UE::Mutable::Private::FExtendedImageDesc GetExternalImageDesc(UTexture* Id);
 
 		/** Load an external mesh asynchronously, returns an event to wait for complition and a cleanup function that must be called once the event has completed. */
-		TTuple<UE::Tasks::FTask, TFunction<void()>> LoadExternalMeshAsync(FExternalResourceId Id, int32 LODIndex, int32 SectionIndex, TFunction<void(TSharedPtr<FMesh>)>& ResultCallback);
+		TTuple<UE::Tasks::FTask, TFunction<void()>> LoadExternalMeshAsync(FExternalResourceId Id, int32 LODIndex, int32 SectionIndex, TFunction<void(TManagedPtr<FMesh>)>& ResultCallback);
 
 		/** Settings that may affect the execution of some operations, like image conversion quality. */
 		FSettings Settings;
-
-    protected:
+	public:
         /** 
 		 * Heap of intermediate data pushed by some instructions and referred by others.
          * It is not released until no operations are pending.
 		 */
 		TArray<FScheduledOpData> HeapData;
+    protected:
 		
 		/** Image descriptor intermediate results. */
 		TMap<OP::ADDRESS, FExtendedImageDesc> ImageDescResults;
@@ -200,36 +203,18 @@ namespace  UE::Mutable::Private
 		struct FTask
 		{
 			FTask() {}
-			FTask(const FScheduledOp& InOp) : Op(InOp) {}
-			FTask(const FScheduledOp& InOp, const FScheduledOp& InDep0) : Op(InOp) 
-			{ 
-				if (InDep0.At) Deps.Add(InDep0); 
-			}
-			FTask(const FScheduledOp& InOp, const FScheduledOp& InDep0, const FScheduledOp& InDep1) : Op(InOp) 
-			{ 
-				if (InDep0.At) Deps.Add(InDep0); 
-				if (InDep1.At) Deps.Add(InDep1); 
-			}
-			FTask(const FScheduledOp& InOp, const FScheduledOp& InDep0, const FScheduledOp& InDep1, const FScheduledOp& InDep2) : Op(InOp) 
-			{ 
-				if (InDep0.At) Deps.Add(InDep0); 
-				if (InDep1.At) Deps.Add(InDep1); 
-				if (InDep2.At) Deps.Add(InDep2); 
-			}
-			FTask(const FScheduledOp& InOp, const FScheduledOp& InDep0, const FScheduledOp& InDep1, const FScheduledOp& InDep2, const FScheduledOp& InDep3) : Op(InOp) 
-			{ 
-				if (InDep0.At) Deps.Add(InDep0); 
-				if (InDep1.At) Deps.Add(InDep1); 
-				if (InDep2.At) Deps.Add(InDep2); 
-				if (InDep3.At) Deps.Add(InDep3); 
-			}
-			FTask(const FScheduledOp& InOp, const FScheduledOp& InDep0, const FScheduledOp& InDep1, const FScheduledOp& InDep2, const FScheduledOp& InDep3, const FScheduledOp& InDep4) : Op(InOp) 
-			{ 
-				if (InDep0.At) Deps.Add(InDep0); 
-				if (InDep1.At) Deps.Add(InDep1); 
-				if (InDep2.At) Deps.Add(InDep2); 
-				if (InDep3.At) Deps.Add(InDep3); 
-				if (InDep4.At) Deps.Add(InDep4); 
+			FTask(const FScheduledOp& InOp, TConstArrayView<FScheduledOp> InDeps) 
+				: Op(InOp) 
+			{
+				Deps.Reserve(InDeps.Num());
+
+				for (const FScheduledOp& Dep : InDeps)
+				{
+					if (Dep.At)
+					{
+						Deps.Add(Dep);
+					}
+				}
 			}
 
 			FScheduledOp Op;
@@ -319,91 +304,44 @@ namespace  UE::Mutable::Private
 			TArray<int32> RomIndices;
 		};
 
-		void AddOp(const FScheduledOp& op)
-		{
-			// It has no dependencies, so add it directly to the open tasks list.
-			OpenTasks.Add(op);
-			ScheduledStagePerOp[op] = op.Stage + 1;
-		}
-
-		void AddOp(const FScheduledOp& op,
-			const FScheduledOp& dep0)
-		{
-			ClosedTasks.Add(FTask(op, dep0));
-			ScheduledStagePerOp[op] = op.Stage + 1;
-			AddChildren(dep0);
-		}
-
-		void AddOp(const FScheduledOp& op,
-			const FScheduledOp& dep0,
-			const FScheduledOp& dep1)
-		{
-			ClosedTasks.Add(FTask(op, dep0, dep1));
-			ScheduledStagePerOp[op] = op.Stage + 1;
-			AddChildren(dep0);
-			AddChildren(dep1);
-		}
-
-		void AddOp(const FScheduledOp& op,
-			const FScheduledOp& dep0,
-			const FScheduledOp& dep1,
-			const FScheduledOp& dep2)
-		{
-			ClosedTasks.Add(FTask(op, dep0, dep1, dep2));
-			ScheduledStagePerOp[op] = op.Stage + 1;
-			AddChildren(dep0);
-			AddChildren(dep1);
-			AddChildren(dep2);
-		}
-
-		void AddOp(const FScheduledOp& op,
-			const FScheduledOp& dep0,
-			const FScheduledOp& dep1,
-			const FScheduledOp& dep2,
-			const FScheduledOp& dep3)
+		template <
+			typename... DepsTypes
+			UE_REQUIRES(UE::CSameAs<typename TDecay<DepsTypes>::Type, FScheduledOp> && ...)
+		>
+		void AddOp(const FScheduledOp& Op, DepsTypes&&... InDeps)
 		{
-			ClosedTasks.Add(FTask(op, dep0, dep1, dep2, dep3));
-			ScheduledStagePerOp[op] = op.Stage + 1;
-			AddChildren(dep0);
-			AddChildren(dep1);
-			AddChildren(dep2);
-			AddChildren(dep3);
-		}
+			constexpr int32 NumDeps = sizeof...(DepsTypes);
+			
+			if constexpr (NumDeps == 0)
+			{
+				OpenTasks.Add(Op);
+				ScheduledStagePerOp[Op] = Op.Stage + 1;
+			}
+			else
+			{
+				FScheduledOp DepsStorage[NumDeps] = {InDeps...};
 
-		void AddOp(const FScheduledOp& op,
-			const FScheduledOp& dep0,
-			const FScheduledOp& dep1,
-			const FScheduledOp& dep2,
-			const FScheduledOp& dep3,
-			const FScheduledOp& dep4)
-		{
-			ClosedTasks.Add(FTask(op, dep0, dep1, dep2, dep3, dep4));
-			ScheduledStagePerOp[op] = op.Stage + 1;
-			AddChildren(dep0);
-			AddChildren(dep1);
-			AddChildren(dep2);
-			AddChildren(dep3);
-			AddChildren(dep4);
+				ClosedTasks.Add(FTask(Op, TConstArrayView<FScheduledOp>(DepsStorage, NumDeps)));
+				ScheduledStagePerOp[Op] = Op.Stage + 1;
+				AddChildren(DepsStorage);
+			}
+			
 		}
 
-		void AddOp(const FScheduledOp& Op, TArrayView<const FScheduledOp> Deps)
+		void AddOp(const FScheduledOp& Op, TArrayView<const FScheduledOp> InDeps)
 		{
-			FTask Task(Op);
-
-			Task.Deps.Reserve(Deps.Num());
-			for (const FScheduledOp& D : Deps)
+			if (!InDeps.Num())
 			{
-				Task.Deps.Add(D);
+				OpenTasks.Add(Op);
+				ScheduledStagePerOp[Op] = Op.Stage + 1;
 			}
-
-			ClosedTasks.Add(MoveTemp(Task));
-
-			ScheduledStagePerOp[Op] = Op.Stage + 1;
-
-			for (const FScheduledOp& D : Deps)
+			else
 			{
-				AddChildren(D);
+				ClosedTasks.Emplace(Op, InDeps);
+				ScheduledStagePerOp[Op] = Op.Stage + 1;
+				AddChildren(InDeps);
 			}
+
 		}
 
     	/** Calculate an approximation of memory used by streaming buffers in this class. */
@@ -472,12 +410,12 @@ namespace  UE::Mutable::Private
 			return System->WorkingMemoryManager.CurrentInstanceCache->GetMatrix(From);
 		}
 
-		inline TSharedPtr<const FMaterial> LoadMaterial(const FCacheAddress& From)
+		inline TManagedPtr<const FMaterial> LoadMaterial(const FCacheAddress& From)
 		{
 			return System->WorkingMemoryManager.CurrentInstanceCache->GetMaterial(From);
 		}
 
-		inline TSharedPtr<const String> LoadString(const FCacheAddress& From)
+		inline TManagedPtr<const String> LoadString(const FCacheAddress& From)
 		{
 			return System->WorkingMemoryManager.CurrentInstanceCache->GetString(From);
 		}
@@ -487,37 +425,37 @@ namespace  UE::Mutable::Private
 			return System->WorkingMemoryManager.CurrentInstanceCache->GetProjector(From);
 		}
 
-		inline TSharedPtr<const FMesh> LoadMesh(const FCacheAddress& From)
+		inline TManagedPtr<const FMesh> LoadMesh(const FCacheAddress& From)
 		{
 			return System->WorkingMemoryManager.LoadMesh(From);
 		}
 
-    	inline TSharedPtr<const FSkeletalMesh> LoadSkeletalMesh(const FCacheAddress& From)
+    	inline TManagedPtr<const FSkeletalMesh> LoadSkeletalMesh(const FCacheAddress& From)
 		{
 			return System->WorkingMemoryManager.CurrentInstanceCache->GetSkeletalMesh(From);
 		}
     	
-		inline TSharedPtr<const FImage> LoadImage(const FCacheAddress& From)
+		inline TManagedPtr<const FImage> LoadImage(const FCacheAddress& From)
 		{
 			return System->WorkingMemoryManager.LoadImage(From);
 		}
 
-		inline TSharedPtr<const FLayout> LoadLayout(const FCacheAddress& From)
+		inline TManagedPtr<const FLayout> LoadLayout(const FCacheAddress& From)
 		{
 			return System->WorkingMemoryManager.CurrentInstanceCache->GetLayout(From);
 		}
 
-		inline TSharedPtr<const FInstance> LoadInstance(const FCacheAddress& From)
+		inline TManagedPtr<const FInstance> LoadInstance(const FCacheAddress& From)
 		{
 			return System->WorkingMemoryManager.CurrentInstanceCache->GetInstance(From);
 		}
 
-		inline TSharedPtr<const FExtensionData> LoadExtensionData(const FCacheAddress& From)
+		inline TManagedPtr<const FExtensionData> LoadExtensionData(const FCacheAddress& From)
 		{
 			return System->WorkingMemoryManager.CurrentInstanceCache->GetExtensionData(From);
 		}
     	
-    	inline TSharedPtr<const FInstancedStruct> LoadInstancedStruct(const FCacheAddress& From)
+    	inline TManagedPtr<const FInstancedStruct> LoadInstancedStruct(const FCacheAddress& From)
 		{
 			return System->WorkingMemoryManager.CurrentInstanceCache->GetInstancedStruct(From);
 		}
@@ -542,7 +480,7 @@ namespace  UE::Mutable::Private
 			System->WorkingMemoryManager.CurrentInstanceCache->SetScalar(To, Value);
 		}
 
-		inline void StoreString(const FCacheAddress& To, TSharedPtr<const String> Value)
+		inline void StoreString(const FCacheAddress& To, TManagedPtr<const String> Value)
 		{
 			System->WorkingMemoryManager.CurrentInstanceCache->SetString(To, Value);
 		}
@@ -557,7 +495,7 @@ namespace  UE::Mutable::Private
 			System->WorkingMemoryManager.CurrentInstanceCache->SetMatrix(To, Value);
 		}
 
-		inline void StoreMaterial(const FCacheAddress& To, TSharedPtr<const FMaterial> Value)
+		inline void StoreMaterial(const FCacheAddress& To, TManagedPtr<const FMaterial> Value)
 		{
 			System->WorkingMemoryManager.CurrentInstanceCache->SetMaterial(To, Value);
 		}
@@ -567,88 +505,67 @@ namespace  UE::Mutable::Private
 			System->WorkingMemoryManager.CurrentInstanceCache->SetProjector(To, Value);
 		}
 
-		inline void StoreMesh(const FCacheAddress& To, TSharedPtr<const FMesh> Resource)
+		inline void StoreMesh(const FCacheAddress& To, const TManagedPtr<const FMesh>& Resource)
 		{
 			System->WorkingMemoryManager.StoreMesh(To, Resource);
 		}
 
-    	inline void StoreSkeletalMesh(const FCacheAddress& To, TSharedPtr<const FSkeletalMesh> Resource)
+    	inline void StoreSkeletalMesh(const FCacheAddress& To, TManagedPtr<const FSkeletalMesh> Resource)
 		{
 			System->WorkingMemoryManager.CurrentInstanceCache->SetSkeletalMesh(To, Resource);
 		}
     	
-		inline void StoreImage(const FCacheAddress& To, TSharedPtr<const FImage> Resource)
+		inline void StoreImage(const FCacheAddress& To, const TManagedPtr<const FImage>& Resource)
 		{
 			System->WorkingMemoryManager.StoreImage(To, Resource);
 		}
 
-		inline void StoreLayout(const FCacheAddress& To, TSharedPtr<const FLayout> Resource)
+		inline void StoreLayout(const FCacheAddress& To, const TManagedPtr<const FLayout>& Resource)
 		{
 			System->WorkingMemoryManager.CurrentInstanceCache->SetLayout(To, Resource);
 		}
 
-		inline void StoreInstance(const FCacheAddress& To, TSharedPtr<const FInstance> Resource)
+		inline void StoreInstance(const FCacheAddress& To, const TManagedPtr<const FInstance>& Resource)
 		{
 			System->WorkingMemoryManager.CurrentInstanceCache->SetInstance(To, Resource);
 		}
 
-		inline void StoreExtensionData(const FCacheAddress& To, TSharedPtr<const FExtensionData> Resource)
+		inline void StoreExtensionData(const FCacheAddress& To, const TManagedPtr<const FExtensionData>& Resource)
 		{
-			System->WorkingMemoryManager.CurrentInstanceCache->SetExtensionData(To, Resource);
+			System->WorkingMemoryManager.CurrentInstanceCache->SetExtensionData(To, Resource.ToStrongRef());
 		}
 
-		inline void StoreInstancedStruct(const FCacheAddress& To, TSharedPtr<const FInstancedStruct> Resource)
+		inline void StoreInstancedStruct(const FCacheAddress& To, const TManagedPtr<const FInstancedStruct>& Resource)
 		{
-			System->WorkingMemoryManager.CurrentInstanceCache->SetInstancedStruct(To, Resource);
+			System->WorkingMemoryManager.CurrentInstanceCache->SetInstancedStruct(To, Resource.ToStrongRef());
 		}
     	
-		inline TSharedPtr<FImage> CreateImage(uint32 SizeX, uint32 SizeY, uint32 Lods, EImageFormat Format, EInitializationType Init)
-		{
-			TSharedPtr<FImage> Result = System->WorkingMemoryManager.CreateImage(SizeX, SizeY, Lods, Format, Init);
-			return MoveTemp(Result);
-		}
-
-		TSharedPtr<FImage> CreateImageLike(const FImage* Ref, EInitializationType Init)
-		{
-			TSharedPtr<FImage> Result = System->WorkingMemoryManager.CreateImage(Ref->GetSizeX(), Ref->GetSizeY(), Ref->GetLODCount(), Ref->GetFormat(), Init);
-			return MoveTemp(Result);
-		}
-
-		/** Ref will be nulled and relesed in any case. */
-		inline TSharedPtr<FImage> CloneOrTakeOver(TSharedPtr<const FImage>& Ref)
+		inline TManagedPtr<FImage> CreateImage(uint32 SizeX, uint32 SizeY, uint32 Lods, EImageFormat Format, EInitializationType Init)
 		{
-			TSharedPtr<FImage> Result = System->WorkingMemoryManager.CloneOrTakeOver(Ref);
-			return MoveTemp(Result);
+			TManagedPtr<FImage> Result = System->WorkingMemoryManager.CreateImage(SizeX, SizeY, Lods, Format, Init);
+			return Result;
 		}
 
-		inline void Release(TSharedPtr<const FImage>& Resource)
+		TManagedPtr<FImage> CreateImageLike(const FImage* Ref, EInitializationType Init)
 		{
-			return System->WorkingMemoryManager.Release(Resource);
+			TManagedPtr<FImage> Result = System->WorkingMemoryManager.CreateImage(Ref->GetSizeX(), Ref->GetSizeY(), Ref->GetLODCount(), Ref->GetFormat(), Init);
+			return Result;
 		}
 
-		inline void Release(TSharedPtr<FImage>& Resource)
+		inline TManagedPtr<FImage> CloneOrTakeOver(TManagedPtr<const FImage>&& Ref)
 		{
-			return System->WorkingMemoryManager.Release(Resource);
+			TManagedPtr<FImage> Result = System->WorkingMemoryManager.CloneOrTakeOver(MoveTemp(Ref));
+			return Result;
 		}
 
-		[[nodiscard]] inline TSharedPtr<FMesh> CreateMesh(int32 BudgetReserveSize = 0)
+		[[nodiscard]] inline TManagedPtr<FMesh> CreateMesh(int32 BudgetReserveSize = 0)
 		{
 			return System->WorkingMemoryManager.CreateMesh(BudgetReserveSize);
 		}
 
-		[[nodiscard]] inline TSharedPtr<FMesh> CloneOrTakeOver(TSharedPtr<const FMesh>& Ref)
-		{
-			return System->WorkingMemoryManager.CloneOrTakeOver(Ref);
-		}
-
-		inline void Release(TSharedPtr<const FMesh>& Resource)
-		{
-			System->WorkingMemoryManager.Release(Resource);
-		}
-
-		inline void Release(TSharedPtr<FMesh>& Resource)
+		[[nodiscard]] inline TManagedPtr<FMesh> CloneOrTakeOver(TManagedPtr<const FMesh>&& Ref)
 		{
-			System->WorkingMemoryManager.Release(Resource);
+			return System->WorkingMemoryManager.CloneOrTakeOver(MoveTemp(Ref));
 		}
 
     	template<typename T>
@@ -747,7 +664,8 @@ namespace  UE::Mutable::Private
     	
     	FRomLoadOps RomLoadOps = FRomLoadOps(*this);
     	    	
-		void AddChildren(const FScheduledOp& dep);
+		void AddChildren(TArrayView<const FScheduledOp> Deps);
+		//void AddChildren(const FScheduledOp& Deps);
 
 		/** Try to create a concurrent task for the given op. Return null if not possible. */
 		TSharedPtr<FIssuedTask> IssueOp(FScheduledOp item);
diff --git a/Engine/Plugins/Mutable/Source/MutableRuntime/Private/MuR/CodeRunnerBegin.cpp b/Engine/Plugins/Mutable/Source/MutableRuntime/Private/MuR/CodeRunnerBegin.cpp
--- a/Engine/Plugins/Mutable/Source/MutableRuntime/Private/MuR/CodeRunnerBegin.cpp
+++ b/Engine/Plugins/Mutable/Source/MutableRuntime/Private/MuR/CodeRunnerBegin.cpp
@@ -159,7 +159,7 @@ namespace UE::Mutable::Private
 	}
 
 
-	void CodeRunnerBegin::StoreInstance(const FCacheAddressInline& To, const TSharedPtr<const FInstance>& Value)
+	void CodeRunnerBegin::StoreInstance(const FCacheAddressInline& To, const TManagedPtr<const FInstance>& Value)
 	{
 		if (!To.At)
 		{
@@ -169,14 +169,14 @@ namespace UE::Mutable::Private
 		Executed.Add(To);
 
 		FStackValue StackValue;
-		StackValue.Set<TSharedPtr<const FInstance>>(Value);
+		StackValue.Set<TManagedPtr<const FInstance>>(Value);
 		Stack.Push(StackValue);
 
 		ResultsInstance.Add(To, Value);
 	}
 
 
-	void CodeRunnerBegin::StoreString(const FCacheAddressInline& To, const TSharedPtr<const String>& Value)
+	void CodeRunnerBegin::StoreString(const FCacheAddressInline& To, const TManagedPtr<const String>& Value)
 	{
 		if (!To.At)
 		{
@@ -186,14 +186,14 @@ namespace UE::Mutable::Private
 		Executed.Add(To);
 
 		FStackValue StackValue;
-		StackValue.Set<TSharedPtr<const String>>(Value);
+		StackValue.Set<TManagedPtr<const String>>(Value);
 		Stack.Push(StackValue);
 
 		ResultsString.Add(To, Value);
 	}
 
 
-	void CodeRunnerBegin::StoreExtensionData(const FCacheAddressInline& To, const TSharedPtr<const FExtensionData>& Value)
+	void CodeRunnerBegin::StoreExtensionData(const FCacheAddressInline& To, const TManagedPtr<const FExtensionData>& Value)
 	{
 		if (!To.At)
 		{
@@ -203,7 +203,7 @@ namespace UE::Mutable::Private
 		Executed.Add(To);
 
 		FStackValue StackValue;
-		StackValue.Set<TSharedPtr<const FExtensionData>>(Value);
+		StackValue.Set<TManagedPtr<const FExtensionData>>(Value);
 		Stack.Push(StackValue);
 
 		ResultsExtensionData.Add(To, Value);
@@ -254,40 +254,40 @@ namespace UE::Mutable::Private
 	}
 
 
-	TSharedPtr<const FInstance> CodeRunnerBegin::LoadInstance(const FCacheAddressInline& To)
+	TManagedPtr<const FInstance> CodeRunnerBegin::LoadInstance(const FCacheAddressInline& To)
 	{
 		if (!To.At)
 		{
 			return nullptr;	
 		}
 			
-		return Stack.Pop().Get<TSharedPtr<const FInstance>>();
+		return Stack.Pop().Get<TManagedPtr<const FInstance>>();
 	}
 
 
-	TSharedPtr<const String> CodeRunnerBegin::LoadString(const FCacheAddressInline& To)
+	TManagedPtr<const String> CodeRunnerBegin::LoadString(const FCacheAddressInline& To)
 	{
 		if (!To.At)
 		{
 			return nullptr;	
 		}
 		
-		return Stack.Pop().Get<TSharedPtr<const String>>();
+		return Stack.Pop().Get<TManagedPtr<const String>>();
 	}
 
 
-	TSharedPtr<const FExtensionData> CodeRunnerBegin::LoadExtensionData(const FCacheAddressInline& To)
+	TManagedPtr<const FExtensionData> CodeRunnerBegin::LoadExtensionData(const FCacheAddressInline& To)
 	{
 		if (!To.At)
 		{
 			return nullptr;	
 		}
 		
-		return Stack.Pop().Get<TSharedPtr<const FExtensionData>>();
+		return Stack.Pop().Get<TManagedPtr<const FExtensionData>>();
 	}
 
 
-	TSharedPtr<const FInstance> CodeRunnerBegin::RunCode(const FScheduledOpInline& Root)
+	TManagedPtr<const FInstance> CodeRunnerBegin::RunCode(const FScheduledOpInline& Root)
 	{
 		MUTABLE_CPUPROFILER_SCOPE(CodeRunnerBegin::RunCode);
 		
@@ -324,7 +324,7 @@ namespace UE::Mutable::Private
 				case EDataType::Instance:
 					{
 						FStackValue Value;
-						Value.Set<TSharedPtr<const FInstance>>(ResultsInstance[Item]);
+						Value.Set<TManagedPtr<const FInstance>>(ResultsInstance[Item]);
 					
 						Stack.Push(Value);
 						break;
@@ -333,7 +333,7 @@ namespace UE::Mutable::Private
 				case EDataType::String:
 					{
 						FStackValue Value;
-						Value.Set<TSharedPtr<const String>>(ResultsString[Item]);
+						Value.Set<TManagedPtr<const String>>(ResultsString[Item]);
 					
 						Stack.Push(Value);
 						break;
@@ -342,7 +342,7 @@ namespace UE::Mutable::Private
 				case EDataType::ExtensionData:
 					{
 						FStackValue Value;
-						Value.Set<TSharedPtr<const FExtensionData>>(ResultsExtensionData[Item]);
+						Value.Set<TManagedPtr<const FExtensionData>>(ResultsExtensionData[Item]);
 					
 						Stack.Push(Value);
 						break;
@@ -568,7 +568,7 @@ namespace UE::Mutable::Private
 									break;
 								
 								case EDataType::ExtensionData:
-									StoreExtensionData(Item, MakeShared<FExtensionData>());
+									StoreExtensionData(Item, MakeManaged<FExtensionData>());
 									break;
 
 								case EDataType::Scalar:
@@ -718,20 +718,20 @@ namespace UE::Mutable::Private
 
 					case 1:
 						{
-							TSharedPtr<const FInstance> InInstance;
+							TManagedPtr<const FInstance> InInstance;
 							if (Args.instance)
 							{
 								InInstance = LoadInstance(Args.instance);
 							}
 
-							TSharedPtr<FInstance> Instance;
+							TManagedPtr<FInstance> Instance;
 							if (InInstance)
 							{
-								Instance = UE::Mutable::Private::CloneOrTakeOver<FInstance>(InInstance);
+								Instance = UE::Mutable::Private::CloneOrTakeOver<FInstance>(MoveTemp(InInstance));
 							}
 							else
 							{
-								Instance = MakeShared<FInstance>();
+								Instance = MakeManaged<FInstance>();
 							}
 							
 							if (Args.value)
@@ -770,20 +770,20 @@ namespace UE::Mutable::Private
 
 					case 1:
 						{
-							TSharedPtr<const FInstance> InInstance;
+							TManagedPtr<const FInstance> InInstance;
 							if (Args.Instance)
 							{
 								InInstance = LoadInstance(Args.Instance);
 							}
 
-							TSharedPtr<FInstance> Instance;
+							TManagedPtr<FInstance> Instance;
 							if (InInstance)
 							{
-								Instance = UE::Mutable::Private::CloneOrTakeOver<FInstance>(InInstance);
+								Instance = UE::Mutable::Private::CloneOrTakeOver<FInstance>(MoveTemp(InInstance));
 							}
 							else
 							{
-								Instance = MakeShared<FInstance>();
+								Instance = MakeManaged<FInstance>();
 							}
 
 
@@ -820,20 +820,20 @@ namespace UE::Mutable::Private
 
 					case 1:
 						{
-							TSharedPtr<const FInstance> InInstance;
+							TManagedPtr<const FInstance> InInstance;
 							if (Args.Instance)
 							{
 								InInstance = LoadInstance(Args.Instance);
 							}
 
-							TSharedPtr<FInstance> Instance;
+							TManagedPtr<FInstance> Instance;
 							if (InInstance)
 							{
-								Instance = UE::Mutable::Private::CloneOrTakeOver<FInstance>(InInstance);
+								Instance = UE::Mutable::Private::CloneOrTakeOver<FInstance>(MoveTemp(InInstance));
 							}
 							else
 							{
-								Instance = MakeShared<FInstance>();
+								Instance = MakeManaged<FInstance>();
 							}
 
 							if (Args.Material)
@@ -869,20 +869,20 @@ namespace UE::Mutable::Private
 
 					case 1:
 						{
-							TSharedPtr<const FInstance> InInstance;
+							TManagedPtr<const FInstance> InInstance;
 							if (Args.Instance)
 							{
 								InInstance = LoadInstance(Args.Instance);
 							}
 
-							TSharedPtr<FInstance> Instance;
+							TManagedPtr<FInstance> Instance;
 							if (InInstance)
 							{
-								Instance = UE::Mutable::Private::CloneOrTakeOver<FInstance>(InInstance);
+								Instance = UE::Mutable::Private::CloneOrTakeOver<FInstance>(MoveTemp(InInstance));
 							}
 							else
 							{
-								Instance = MakeShared<FInstance>();
+								Instance = MakeManaged<FInstance>();
 							}
 
 							if (Args.Material)
@@ -919,25 +919,25 @@ namespace UE::Mutable::Private
 
 					case 1:
 						{
-							TSharedPtr<const FInstance> InInstance;
+							TManagedPtr<const FInstance> InInstance;
 							if (Args.instance)
 							{
 								InInstance = LoadInstance(Args.instance);
 							}
 
-							TSharedPtr<FInstance> Instance;
+							TManagedPtr<FInstance> Instance;
 							if (InInstance)
 							{
-								Instance = UE::Mutable::Private::CloneOrTakeOver<FInstance>(InInstance);
+								Instance = UE::Mutable::Private::CloneOrTakeOver<FInstance>(MoveTemp(InInstance));
 							}
 							else
 							{
-								Instance = MakeShared<FInstance>();
+								Instance = MakeManaged<FInstance>();
 							}
 
 							if (Args.value)
 							{
-								TSharedPtr<const FInstance> pComp = LoadInstance(Args.value);
+								TManagedPtr<const FInstance> pComp = LoadInstance(Args.value);
 
 								int32 NewComponentIndex = Instance->AddComponent();
 
@@ -977,20 +977,20 @@ namespace UE::Mutable::Private
 
 					case 1:
 						{
-							TSharedPtr<const FInstance> InInstance;
+							TManagedPtr<const FInstance> InInstance;
 							if (Args.instance)
 							{
 								InInstance = LoadInstance(Args.instance);
 							}
 
-							TSharedPtr<FInstance> Instance;
+							TManagedPtr<FInstance> Instance;
 							if (InInstance)
 							{
-								Instance = UE::Mutable::Private::CloneOrTakeOver<FInstance>(InInstance);
+								Instance = UE::Mutable::Private::CloneOrTakeOver<FInstance>(MoveTemp(InInstance));
 							}
 							else
 							{
-								Instance = MakeShared<FInstance>();
+								Instance = MakeManaged<FInstance>();
 							}
 							
 							StoreInstance( Item, Instance );
@@ -1020,25 +1020,25 @@ namespace UE::Mutable::Private
 		            
 		            case 1:
 		            {
-		            	TSharedPtr<const FInstance> InInstance;
+		            	TManagedPtr<const FInstance> InInstance;
 		            	if (Args.instance)
 		            	{
 		            		InInstance = LoadInstance(Args.instance);
 		            	}
 
-		            	TSharedPtr<FInstance> Instance;
+		            	TManagedPtr<FInstance> Instance;
 		            	if (InInstance)
 		            	{
-		            		Instance = UE::Mutable::Private::CloneOrTakeOver<FInstance>(InInstance);
+		            		Instance = UE::Mutable::Private::CloneOrTakeOver<FInstance>(MoveTemp(InInstance));
 		            	}
 		            	else
 		            	{
-		            		Instance = MakeShared<FInstance>();
+		            		Instance = MakeManaged<FInstance>();
 		            	}
 
 		                // Empty surfaces are ok, they still need to be created, because they may contain
 		                // additional information like internal or external IDs
-		            	TSharedPtr<const FInstance> Value;
+		            	TManagedPtr<const FInstance> Value;
 		            	if (Args.value)
 		            	{
 		            		Value = LoadInstance(Args.value);
@@ -1113,7 +1113,7 @@ namespace UE::Mutable::Private
 		            case 1:
 		            {
 		                // Assemble result
-						TSharedPtr<FInstance> Result = MakeShared<FInstance>();
+						TManagedPtr<FInstance> Result = MakeManaged<FInstance>();
 						int32 ComponentIndex = Result->AddComponent();
 
 		                for (uint8 LODIndex = 0; LODIndex < LODCount; ++LODIndex)
@@ -1131,7 +1131,7 @@ namespace UE::Mutable::Private
 								
 								if (bIsSelectedLod)
 		                        {
-									TSharedPtr<const FInstance> LOD = LoadInstance(LODAddress);
+									TManagedPtr<const FInstance> LOD = LoadInstance(LODAddress);
 
 		                            // In a degenerated case, the returned pLOD may not have an LOD inside
  									if (!LOD->Components.IsEmpty() &&
@@ -1170,20 +1170,20 @@ namespace UE::Mutable::Private
 
 					case 1:
 						{
-							TSharedPtr<const FInstance> InInstance;
+							TManagedPtr<const FInstance> InInstance;
 							if (Args.instance)
 							{
 								InInstance = LoadInstance(Args.instance);
 							}
 
-							TSharedPtr<FInstance> Instance;
+							TManagedPtr<FInstance> Instance;
 							if (InInstance)
 							{
-								Instance = UE::Mutable::Private::CloneOrTakeOver<FInstance>(InInstance);
+								Instance = UE::Mutable::Private::CloneOrTakeOver<FInstance>(MoveTemp(InInstance));
 							}
 							else
 							{
-								Instance = MakeShared<FInstance>();
+								Instance = MakeManaged<FInstance>();
 							}
 							
 							if (Args.value)
@@ -1220,29 +1220,29 @@ namespace UE::Mutable::Private
 
 					case 1:
 						{
-							TSharedPtr<const FInstance> InInstance;
+							TManagedPtr<const FInstance> InInstance;
 							if (Args.Instance)
 							{
 								InInstance = LoadInstance(Args.Instance);
 							}
 
-							TSharedPtr<FInstance> Instance;
+							TManagedPtr<FInstance> Instance;
 							if (InInstance)
 							{
-								Instance = UE::Mutable::Private::CloneOrTakeOver<FInstance>(InInstance);
+								Instance = UE::Mutable::Private::CloneOrTakeOver<FInstance>(MoveTemp(InInstance));
 							}
 							else
 							{
-								Instance = MakeShared<FInstance>();
+								Instance = MakeManaged<FInstance>();
 							}
 							
-							if (TSharedPtr<const FExtensionData> ExtensionData = LoadExtensionData(Args.ExtensionData))
+							if (TManagedPtr<const FExtensionData> ExtensionData = LoadExtensionData(Args.ExtensionData))
 							{
 								const OP::ADDRESS NameAddress = Args.ExtensionDataName;
 								check(NameAddress < (uint32)Program.ConstantStrings.Num());
 								const FString& NameString = Program.ConstantStrings[NameAddress];
 
-								Instance->SetExtensionData(ExtensionData.ToSharedRef(), FName(NameString));
+								Instance->SetExtensionData(ExtensionData.ToStrongRef(), FName(NameString));
 							}
 
 							StoreInstance(Item, Instance);
@@ -1520,7 +1520,7 @@ namespace UE::Mutable::Private
 					OP::ResourceConstantArgs Args = Program.GetOpArgs<OP::ResourceConstantArgs>(Item.At);
 
 					// Assume the ROM has been loaded previously
-					TSharedPtr<const FExtensionData> SourceConst;
+					TManagedPtr<const FExtensionData> SourceConst;
 					Program.GetExtensionDataConstant(Args.value, SourceConst);
 
 					check(SourceConst);
@@ -1537,7 +1537,7 @@ namespace UE::Mutable::Private
 
 					const FString& Result = Program.ConstantStrings[Args.value];
 
-					TSharedPtr<const String> Value = MakeShared<String>(Result);
+					TManagedPtr<const String> Value = MakeManaged<String>(Result);
 
 					StoreString(Item, Value);
 					break;
@@ -1556,7 +1556,7 @@ namespace UE::Mutable::Private
 					FString Result;
 					Params->GetStringValue(Args.variable, Result, Index.Get());
 
-					TSharedPtr<const String> Value = MakeShared<String>(Result);
+					TManagedPtr<const String> Value = MakeManaged<String>(Result);
 					
 					StoreString(Item, Value);
 					break;
@@ -2581,6 +2581,6 @@ namespace UE::Mutable::Private
 
 		check(Stack.Num() <= 1);
 
-		return Stack.Num() ? LoadInstance(Root) : MakeShared<FInstance>();
+		return Stack.Num() ? LoadInstance(Root) : MakeManaged<FInstance>();
 	}
 }
diff --git a/Engine/Plugins/Mutable/Source/MutableRuntime/Private/MuR/CodeRunnerBegin.h b/Engine/Plugins/Mutable/Source/MutableRuntime/Private/MuR/CodeRunnerBegin.h
--- a/Engine/Plugins/Mutable/Source/MutableRuntime/Private/MuR/CodeRunnerBegin.h
+++ b/Engine/Plugins/Mutable/Source/MutableRuntime/Private/MuR/CodeRunnerBegin.h
@@ -87,7 +87,7 @@ namespace UE::Mutable::Private
 		TSet<FCacheAddressInline> IndexOther;
 	};
 	
-	using FStackValue = TVariant<bool, int32, float, TSharedPtr<const String>, TSharedPtr<const FExtensionData>, TSharedPtr<const FInstance>>;
+	using FStackValue = TVariant<bool, int32, float, TManagedPtr<const String>, TManagedPtr<const FExtensionData>, TManagedPtr<const FInstance>>;
 
 	class FStack : public TArray<FStackValue, TInlineAllocator<32>>
 	{
@@ -100,7 +100,7 @@ namespace UE::Mutable::Private
 	public:
 		CodeRunnerBegin(FSystem& InSystem, const TSharedPtr<const FModel>& InModel, const TSharedPtr<const FParameters>& InParams, uint32 InLodMask);
 
-		TSharedPtr<const FInstance> RunCode(const FScheduledOpInline& Root);
+		TManagedPtr<const FInstance> RunCode(const FScheduledOpInline& Root);
 
 		FWorkingMemoryManager& GetMemory() const;
 		
@@ -114,11 +114,11 @@ namespace UE::Mutable::Private
 		
 		void StoreBool(const FCacheAddressInline& To, bool Value);
 		
-		void StoreInstance(const FCacheAddressInline& To, const TSharedPtr<const FInstance>& Value);
+		void StoreInstance(const FCacheAddressInline& To, const TManagedPtr<const FInstance>& Value);
 		
-		void StoreString(const FCacheAddressInline& To, const TSharedPtr<const String>& Value);
+		void StoreString(const FCacheAddressInline& To, const TManagedPtr<const String>& Value);
 		
-		void StoreExtensionData(const FCacheAddressInline& To, const TSharedPtr<const FExtensionData>& Value);
+		void StoreExtensionData(const FCacheAddressInline& To, const TManagedPtr<const FExtensionData>& Value);
 
 		void StoreNone(const FCacheAddressInline& To);
 
@@ -128,11 +128,11 @@ namespace UE::Mutable::Private
 
 		bool LoadBool(const FCacheAddressInline& To);
 		
-		TSharedPtr<const FInstance> LoadInstance(const FCacheAddressInline& To);
+		TManagedPtr<const FInstance> LoadInstance(const FCacheAddressInline& To);
 
-		TSharedPtr<const String> LoadString(const FCacheAddressInline& To);
+		TManagedPtr<const String> LoadString(const FCacheAddressInline& To);
 
-		TSharedPtr<const FExtensionData> LoadExtensionData(const FCacheAddressInline& To);
+		TManagedPtr<const FExtensionData> LoadExtensionData(const FCacheAddressInline& To);
 		
 		TSharedPtr<const FModel> Model;
 		TSharedPtr<const FParameters> Params;
@@ -151,8 +151,8 @@ namespace UE::Mutable::Private
 		TMap<FCacheAddressInline, float> ResultsFloat;
 		TMap<FCacheAddressInline, bool> ResultsBool;
 		TMap<FCacheAddressInline, int32> ResultsInt;
-		TMap<FCacheAddressInline, TSharedPtr<const FInstance>> ResultsInstance;
-		TMap<FCacheAddressInline, TSharedPtr<const String>> ResultsString;
-		TMap<FCacheAddressInline, TSharedPtr<const FExtensionData>> ResultsExtensionData;
+		TMap<FCacheAddressInline, TManagedPtr<const FInstance>> ResultsInstance;
+		TMap<FCacheAddressInline, TManagedPtr<const String>> ResultsString;
+		TMap<FCacheAddressInline, TManagedPtr<const FExtensionData>> ResultsExtensionData;
 	};
 }
diff --git a/Engine/Plugins/Mutable/Source/MutableRuntime/Private/MuR/CodeRunnerTasked.cpp b/Engine/Plugins/Mutable/Source/MutableRuntime/Private/MuR/CodeRunnerTasked.cpp
--- a/Engine/Plugins/Mutable/Source/MutableRuntime/Private/MuR/CodeRunnerTasked.cpp
+++ b/Engine/Plugins/Mutable/Source/MutableRuntime/Private/MuR/CodeRunnerTasked.cpp
@@ -64,27 +64,39 @@ namespace UE::Mutable::Private
 	TRACE_DECLARE_INT_COUNTER(MutableRuntime_IssuedTasks, TEXT("MutableRuntime/IssuedTasks"));
 	TRACE_DECLARE_INT_COUNTER(MutableRuntime_IssuedTasksOnHold, TEXT("MutableRuntime/IssuedHoldTasks"));
 
-	void CodeRunner::AddChildren(const FScheduledOp& dep)
+	void CodeRunner::AddChildren(TConstArrayView<FScheduledOp> Deps)
 	{
-		FCacheAddress at(dep);
-		if (dep.At && !GetMemory().IsValid(at))
+		for (const FScheduledOp& Dep : Deps)
 		{
-			if (ScheduledStagePerOp.get(at) <= dep.Stage)
+			FCacheAddress At(Dep);
+	
+			// We need to increase hit count unconditionally for Full execution. 
+			if (Dep.Type == FScheduledOp::EType::Full)
+			{
+				System->WorkingMemoryManager.CurrentInstanceCache->IncreaseHitCount(At);
+			}
+
+			if (!Dep.At)
 			{
+				continue;
+			}
+
+			bool bSuccess = GetMemory().AquireCachedResult(At);
+
+			if (!bSuccess)
+			{
+				if (ScheduledStagePerOp.get(At) <= Dep.Stage)
+				{
 #if MUTABLE_DEBUG_CODERUNNER_TASK_SCHEDULE_CALLSTACK
-				FScheduledOp& AddedOp = OpenTasks.Add_GetRef(dep);
-				AddedOp.StackDepth = FPlatformStackWalk::CaptureStackBackTrace(&AddedOp.ScheduleCallstack[0], FScheduledOp::CallstackMaxDepth);
+					FScheduledOp& AddedOp = OpenTasks.Add_GetRef(Dep);
+					AddedOp.StackDepth = FPlatformStackWalk::CaptureStackBackTrace(&AddedOp.ScheduleCallstack[0], FScheduledOp::CallstackMaxDepth);
 #else
-				OpenTasks.Add(dep);
+					OpenTasks.Add(Dep);
 #endif
-				ScheduledStagePerOp[at] = dep.Stage + 1;
+					ScheduledStagePerOp[At] = Dep.Stage + 1;
+				}
 			}
 		}
-
-		if (dep.Type == FScheduledOp::EType::Full)
-		{
-			System->WorkingMemoryManager.CurrentInstanceCache->IncreaseHitCount(at);
-		}
 	}
 
 
@@ -239,7 +251,7 @@ namespace UE::Mutable::Private
 			while (!OpenTasks.IsEmpty())
 			{
 				// Get a new task to run
-				FScheduledOp item;
+				FScheduledOp Item;
 				switch (ExecutionStrategy)
 				{
 				//case EExecutionStrategy::MinimizeMemory:
@@ -253,27 +265,26 @@ namespace UE::Mutable::Private
 				case EExecutionStrategy::None:
 				default:
 					// Just get one.
-					item = OpenTasks.Pop(EAllowShrinking::No);
+					Item = OpenTasks.Pop(EAllowShrinking::No);
 					break;
 
 				}
 
 				// Special processing in case it is an ImageDesc operation
-				if (item.Type == FScheduledOp::EType::ImageDesc)
+				if (Item.Type == FScheduledOp::EType::ImageDesc)
 				{
-					RunCodeImageDesc(item, Params);
+					RunCodeImageDesc(Item, Params);
 					continue;
 				}
 
-				// Don't run it if we already have the result.
-				FCacheAddress cat(item);
-				if (GetMemory().IsValid(cat))
+				FCacheAddress CachedAddress(Item);
+				if (GetMemory().AquireCachedResult(CachedAddress))
 				{
 					continue;
 				}
 
-				// See if we can schedule this item concurrently
-				TSharedPtr<FIssuedTask> IssuedTask = IssueOp(item);
+				// See if we can schedule this Item concurrently
+				TSharedPtr<FIssuedTask> IssuedTask = IssueOp(Item);
 				if (IssuedTask)
 				{
 					if (ShouldIssueTask())
@@ -294,24 +305,24 @@ namespace UE::Mutable::Private
 				{
 					// Run immediately
 					bool bFailed = false;
-					RunCode(item, Params, Model, LODMask, bFailed);
+					RunCode(Item, Params, Model, LODMask, bFailed);
 					if (bFailed)
 					{
 						return AbortRun();
 					}
 
-					if (ScheduledStagePerOp[item] == item.Stage + 1)
+					if (ScheduledStagePerOp[Item] == Item.Stage + 1)
 					{
 						// We completed everything that was requested, clear it. Otherwise if needed again it is not going to be rebuilt.
 						// \TODO: track operations that are run more than once?.
-						ScheduledStagePerOp[item] = 0;
+						ScheduledStagePerOp[Item] = 0;
 					}
 				}
 
 				if (ProfileContext)
 				{
 					++ProfileContext->NumRunOps;
-					++ProfileContext->RunOpsPerType[int32(Model->GetPrivate()->Program.GetOpType(item.At))];
+					++ProfileContext->RunOpsPerType[int32(Model->GetPrivate()->Program.GetOpType(Item.At))];
 				}
 			}
 
@@ -332,33 +343,19 @@ namespace UE::Mutable::Private
 
 			// Look for a closed task with dependencies satisfied and move them to the open task list.
 			bool bSomeWasReady = false;
-			for (int32 Index = 0; Index < ClosedTasks.Num(); )
+			
+			for (int32 Index = 0; Index < ClosedTasks.Num();)
 			{
-				bool Ready = true;
-				for ( const FCacheAddress& Dep: ClosedTasks[Index].Deps )
-				{
-					bool bDependencyFailed = false;
-					bDependencyFailed = Dep.At && !GetMemory().IsValid(Dep);
-					
-					if (bDependencyFailed)
-					{
-						Ready = false;
-						break;
-					}
-				}
-
-				if (Ready)
+				if (GetMemory().TryAquireCachedResultSet(ClosedTasks[Index].Deps))
 				{
-					bSomeWasReady = true;
-					FTask Task = ClosedTasks[Index];
+					OpenTasks.Push(ClosedTasks[Index].Op);
 					ClosedTasks.RemoveAt(Index, EAllowShrinking::No); // with swap? would change order of execution.
-					OpenTasks.Push(Task.Op);
+					bSomeWasReady = true;
 				}
 				else
 				{
 					++Index;
 				}
-
 			}
 
 			UpdateTraces();
@@ -369,7 +366,7 @@ namespace UE::Mutable::Private
 			{
 				// Log the task graph
 				for (int32 Index = 0; Index < ClosedTasks.Num(); ++Index)
-				{
+				{	
 					FString TaskDesc = FString::Printf(TEXT("Closed task %d-%d-%d depends on : "), ClosedTasks[Index].Op.At, ClosedTasks[Index].Op.ExecutionIndex, ClosedTasks[Index].Op.Stage );
 					for (const FCacheAddress& Dep : ClosedTasks[Index].Deps)
 					{
@@ -515,9 +512,9 @@ namespace UE::Mutable::Private
 	private:
 		int32 ImageCompressionQuality = 0;
 		OP::ImageLayerArgs Args;
-		TSharedPtr<const FImage> Blended;
-		TSharedPtr<const FImage> Mask;
-		TSharedPtr<FImage> Result;
+		TManagedPtr<const FImage> Blended;
+		TManagedPtr<const FImage> Mask;
+		TManagedPtr<FImage> Result;
 		EImageFormat InitialFormat = EImageFormat::None;
 	};
 
@@ -530,7 +527,7 @@ namespace UE::Mutable::Private
 
 		ImageCompressionQuality = Runner->Settings.ImageCompressionQuality;
 
-		TSharedPtr<const FImage> Base = Runner->LoadImage({ Args.base, Op.ExecutionIndex, Op.ExecutionOptions });
+		TManagedPtr<const FImage> Base = Runner->LoadImage({ Args.base, Op.ExecutionIndex, Op.ExecutionOptions });
 		check(!Base || Base->GetFormat() < EImageFormat::Count);
 
 		Blended = Runner->LoadImage({ Args.blended, Op.ExecutionIndex, Op.ExecutionOptions });
@@ -543,8 +540,9 @@ namespace UE::Mutable::Private
 		// Shortcuts
 		if (!Base)
 		{
-			Runner->Release(Blended);
-			Runner->Release(Mask);
+			Blended = nullptr;
+			Mask = nullptr;
+
 			Runner->StoreImage(Op, Base);
 			return false;
 		}
@@ -552,8 +550,9 @@ namespace UE::Mutable::Private
 		bool bValid = Base->GetSizeX() > 0 && Base->GetSizeY() > 0;
 		if (!bValid || !Blended)
 		{
-			Runner->Release(Blended);
-			Runner->Release(Mask);
+			Blended = nullptr;
+			Mask = nullptr;
+
 			Runner->StoreImage(Op, Base);
 			return false;
 		}
@@ -566,31 +565,28 @@ namespace UE::Mutable::Private
 		if (IsCompressedFormat(InitialFormat))
 		{
 			EImageFormat UncompressedFormat = GetUncompressedFormat(InitialFormat);
-			TSharedPtr<FImage> Formatted = Runner->CreateImage( Base->GetSizeX(), Base->GetSizeY(), Base->GetLODCount(), UncompressedFormat, EInitializationType::NotInitialized );
+			TManagedPtr<FImage> Formatted = Runner->CreateImage( Base->GetSizeX(), Base->GetSizeY(), Base->GetLODCount(), UncompressedFormat, EInitializationType::NotInitialized );
 			bool bSuccess = false;
 			ImOp.ImagePixelFormat(bSuccess, ImageCompressionQuality, Formatted.Get(), Base.Get());
 			check(bSuccess); // Decompression cannot fail
-			Runner->Release(Base);
 			Base = Formatted;
 		}
 
 		bool bMustHaveSameFormat = !(Args.flags & (OP::ImageLayerArgs::F_BASE_RGB_FROM_ALPHA | OP::ImageLayerArgs::F_BLENDED_RGB_FROM_ALPHA));
 		if (Blended && InitialFormat != Blended->GetFormat() && bMustHaveSameFormat)
 		{
-			TSharedPtr<FImage> Formatted = Runner->CreateImage(Blended->GetSizeX(), Blended->GetSizeY(), Blended->GetLODCount(), Base->GetFormat(), EInitializationType::NotInitialized);
+			TManagedPtr<FImage> Formatted = Runner->CreateImage(Blended->GetSizeX(), Blended->GetSizeY(), Blended->GetLODCount(), Base->GetFormat(), EInitializationType::NotInitialized);
 			bool bSuccess = false;
 			ImOp.ImagePixelFormat(bSuccess, ImageCompressionQuality, Formatted.Get(), Blended.Get());
 			check(bSuccess); // Decompression cannot fail
-			Runner->Release(Blended);
 			Blended = Formatted;
 		}
 
 		if (Base->GetSize() != Blended->GetSize())
 		{
 			MUTABLE_CPUPROFILER_SCOPE(ImageResize_EmergencyFix);
-			TSharedPtr<FImage> Resized = Runner->CreateImage(Base->GetSizeX(), Base->GetSizeY(), 1, Blended->GetFormat(), EInitializationType::NotInitialized);
+			TManagedPtr<FImage> Resized = Runner->CreateImage(Base->GetSizeX(), Base->GetSizeY(), 1, Blended->GetFormat(), EInitializationType::NotInitialized);
 			ImOp.ImageResizeLinear(Resized.Get(), ImageCompressionQuality, Blended.Get());
-			Runner->Release(Blended);
 			Blended = Resized;
 
 		}
@@ -600,9 +596,8 @@ namespace UE::Mutable::Private
 			if (Base->GetSize() != Mask->GetSize())
 			{
 				MUTABLE_CPUPROFILER_SCOPE(ImageResize_EmergencyFix);
-				TSharedPtr<FImage> Resized = Runner->CreateImage(Base->GetSizeX(), Base->GetSizeY(), 1, Mask->GetFormat(), EInitializationType::NotInitialized);
+				TManagedPtr<FImage> Resized = Runner->CreateImage(Base->GetSizeX(), Base->GetSizeY(), 1, Mask->GetFormat(), EInitializationType::NotInitialized);
 				ImOp.ImageResizeLinear(Resized.Get(), ImageCompressionQuality, Mask.Get());
-				Runner->Release(Mask);
 				Mask = Resized;
 
 			}
@@ -615,23 +610,22 @@ namespace UE::Mutable::Private
 				int32 LevelCount = Base->GetLODCount();
 
 				// Uncompress mask to avoid excessive RLE-compression and decompression in the following ops.
-				TSharedPtr<FImage> UncompressedMask = Runner->CreateImage(Mask->GetSizeX(), Mask->GetSizeY(), Mask->GetLODCount(), GetUncompressedFormat(Mask->GetFormat()), EInitializationType::NotInitialized);
+				TManagedPtr<FImage> UncompressedMask = Runner->CreateImage(Mask->GetSizeX(), Mask->GetSizeY(), Mask->GetLODCount(), GetUncompressedFormat(Mask->GetFormat()), EInitializationType::NotInitialized);
 				bool bSuccess = false;
 				ImOp.ImagePixelFormat(bSuccess, ImageCompressionQuality, UncompressedMask.Get(), Mask.Get());
 
-				TSharedPtr<FImage> MaskFix = UncompressedMask;
+				TManagedPtr<FImage> MaskFix = UncompressedMask;
 				MaskFix->DataStorage.SetNumLODs(LevelCount);
 
 				FMipmapGenerationSettings Settings{};
 				ImOp.ImageMipmap(ImageCompressionQuality, MaskFix.Get(), MaskFix.Get(), StartLevel, LevelCount, Settings);
 
-				Runner->Release(Mask);
 				Mask = MaskFix;
 			}
 		}
 
 		// Create destination data
-		Result = Runner->CloneOrTakeOver(Base);
+		Result = Runner->CloneOrTakeOver(MoveTemp(Base));
 		return true;
 	}
 
@@ -780,22 +774,21 @@ namespace UE::Mutable::Private
 	bool FImageLayerTask::Complete( CodeRunner* Runner )
 	{
 		// This runs in the Runner thread
-		Runner->Release(Blended);
-		Runner->Release(Mask);
 
+		Blended = nullptr;
+		Mask    = nullptr;
 		// If no shortcut was taken
 		if (Result)
 		{
 			if (InitialFormat != Result->GetFormat())
 			{
-				TSharedPtr<FImage> Formatted = Runner->CreateImage(Result->GetSizeX(), Result->GetSizeY(), Result->GetLODCount(), InitialFormat, EInitializationType::NotInitialized);
+				TManagedPtr<FImage> Formatted = Runner->CreateImage(Result->GetSizeX(), Result->GetSizeY(), Result->GetLODCount(), InitialFormat, EInitializationType::NotInitialized);
 				bool bSuccess = false;
 
 				FImageOperator ImOp = MakeImageOperator(Runner);
 				ImOp.ImagePixelFormat(bSuccess, ImageCompressionQuality, Formatted.Get(), Result.Get());
 				check(bSuccess);
 
-				Runner->Release(Result);
 				Result = Formatted;
 			}
 
@@ -826,8 +819,8 @@ namespace UE::Mutable::Private
 		int32 ImageCompressionQuality = 0;
 		OP::ImageLayerColourArgs Args;
 		FVector4f Color;
-		TSharedPtr<const FImage> Mask;
-		TSharedPtr<FImage> Result;
+		TManagedPtr<const FImage> Mask;
+		TManagedPtr<FImage> Result;
 		EImageFormat InitialFormat = EImageFormat::None;
 	};
 
@@ -840,7 +833,7 @@ namespace UE::Mutable::Private
 
 		ImageCompressionQuality = Runner->Settings.ImageCompressionQuality;
 
-		TSharedPtr<const FImage> Base = Runner->LoadImage({ Args.base, Op.ExecutionIndex, Op.ExecutionOptions });
+		TManagedPtr<const FImage> Base = Runner->LoadImage({ Args.base, Op.ExecutionIndex, Op.ExecutionOptions });
 		check(!Base || Base->GetFormat() < EImageFormat::Count);
 
 		Color = Runner->LoadColor({ Args.colour, Op.ExecutionIndex, 0 });
@@ -853,7 +846,7 @@ namespace UE::Mutable::Private
 		// Shortcuts
 		if (!Base)
 		{
-			Runner->Release(Mask);
+			Mask = nullptr;
 			Runner->StoreImage(Op, Base);
 			return false;
 		}
@@ -861,7 +854,7 @@ namespace UE::Mutable::Private
 		bool bValid = Base->GetSizeX() > 0 && Base->GetSizeY() > 0;
 		if (!bValid)
 		{
-			Runner->Release(Mask);
+			Mask = nullptr;
 			Runner->StoreImage(Op, Base);
 			return false;
 		}
@@ -878,22 +871,20 @@ namespace UE::Mutable::Private
 				static_cast<EBlendType>(Args.blendType) == EBlendType::BT_NORMAL_COMBINE) // TODO Optimize. BT_NORMAL_COMBINE (ImageLayerCombineColour with mask) does not support formats such as RLE. Hence why we are changing the format here.
 			{
 				MUTABLE_CPUPROFILER_SCOPE(EmergencyFix_Format);
-				TSharedPtr<FImage> Formatted = Runner->CreateImage(Mask->GetSizeX(), Mask->GetSizeY(), Mask->GetLODCount(), EImageFormat::L_UByte, EInitializationType::NotInitialized);
+				TManagedPtr<FImage> Formatted = Runner->CreateImage(Mask->GetSizeX(), Mask->GetSizeY(), Mask->GetLODCount(), EImageFormat::L_UByte, EInitializationType::NotInitialized);
 
 				bool bSuccess = false;
 				ImOp.ImagePixelFormat(bSuccess, ImageCompressionQuality, Formatted.Get(), Mask.Get());
 				check(bSuccess); // Decompression cannot fail
 
-				Runner->Release(Mask);
 				Mask = Formatted;
 			}
 			
 			if (Base->GetSize() != Mask->GetSize())
 			{
 				MUTABLE_CPUPROFILER_SCOPE(EmergencyFix_Size);
-				TSharedPtr<FImage> Resized = Runner->CreateImage(Base->GetSizeX(), Base->GetSizeY(), 1, Mask->GetFormat(), EInitializationType::NotInitialized);
+				TManagedPtr<FImage> Resized = Runner->CreateImage(Base->GetSizeX(), Base->GetSizeY(), 1, Mask->GetFormat(), EInitializationType::NotInitialized);
 				ImOp.ImageResizeLinear(Resized.Get(), ImageCompressionQuality, Mask.Get() );
-				Runner->Release(Mask);
 				Mask = Resized;
 			}
 
@@ -903,7 +894,7 @@ namespace UE::Mutable::Private
 				int32 StartLevel = Mask->GetLODCount() - 1;
 				int32 LevelCount = Base->GetLODCount();
 
-				TSharedPtr<FImage> MaskFix = Runner->CloneOrTakeOver(Mask);
+				TManagedPtr<FImage> MaskFix = Runner->CloneOrTakeOver(MoveTemp(Mask));
 				MaskFix->DataStorage.SetNumLODs(LevelCount);
 
 				FMipmapGenerationSettings Settings{};
@@ -914,7 +905,7 @@ namespace UE::Mutable::Private
 		}
 
 		// Create destination data
-		Result = Runner->CloneOrTakeOver(Base);
+		Result = Runner->CloneOrTakeOver(MoveTemp(Base));
 		return true;
 	}
 
@@ -1046,7 +1037,7 @@ namespace UE::Mutable::Private
 	bool FImageLayerColourTask::Complete(CodeRunner* Runner)
 	{
 		// This runs in the Runner thread
-		Runner->Release(Mask);
+		Mask = nullptr;
 
 		// If no shortcut was taken
 		if (Result)
@@ -1078,8 +1069,8 @@ namespace UE::Mutable::Private
 		int ImageCompressionQuality = 0;
 		OP::ImagePixelFormatArgs Args;
 		EImageFormat TargetFormat = EImageFormat::None;
-		TSharedPtr<const FImage> Base;
-		TSharedPtr<FImage> Result;
+		TManagedPtr<const FImage> Base;
+		TManagedPtr<FImage> Result;
 		FImageOperator::FImagePixelFormatFunc ImagePixelFormatFunc;
 	};
 
@@ -1099,6 +1090,7 @@ namespace UE::Mutable::Private
 		if (!Base)
 		{
 			Runner->StoreImage(Op, Base);
+			Base = nullptr;
 			return false;
 		}
 
@@ -1106,6 +1098,7 @@ namespace UE::Mutable::Private
 		if (!bValid)
 		{
 			Runner->StoreImage(Op, Base);
+			Base = nullptr;
 			return false;
 		}
 
@@ -1120,6 +1113,7 @@ namespace UE::Mutable::Private
 		if (TargetFormat == EImageFormat::None || TargetFormat == Base->GetFormat())
 		{
 			Runner->StoreImage(Op, Base);
+			Base = nullptr;
 			return false;
 		}
 
@@ -1146,7 +1140,7 @@ namespace UE::Mutable::Private
 	bool FImagePixelFormatTask::Complete(CodeRunner* Runner)
 	{
 		// This runs in the Runner thread
-		Runner->Release(Base);
+		Base = nullptr;
 
 		// If no shortcut was taken
 		if (Result)
@@ -1174,8 +1168,8 @@ namespace UE::Mutable::Private
 		int32 ImageCompressionQuality = 0;
 		int32 StartLevel = -1;
 		OP::ImageMipmapArgs Args;
-		TSharedPtr<const FImage> Base;
-		TSharedPtr<FImage> Result;
+		TManagedPtr<const FImage> Base;
+		TManagedPtr<FImage> Result;
 		FImageOperator::FScratchImageMipmap Scratch;
 		FImageOperator::FImagePixelFormatFunc ImagePixelFormatFunc;
 	};
@@ -1196,6 +1190,7 @@ namespace UE::Mutable::Private
 		if (!Base)
 		{
 			Runner->StoreImage(Op, Base);
+			Base = nullptr;
 			return false;
 		}
 
@@ -1203,6 +1198,7 @@ namespace UE::Mutable::Private
 		if (!bValid)
 		{
 			Runner->StoreImage(Op, Base);
+			Base = nullptr;
 			return false;
 		}
 
@@ -1225,13 +1221,13 @@ namespace UE::Mutable::Private
 		if (LevelCount == Base->GetLODCount())
 		{
 			Runner->StoreImage(Op, Base);
+			Base = nullptr;
 			return false;
 		}
 
 		StartLevel = Base->GetLODCount() - 1;
 		
-		Result = Runner->CloneOrTakeOver(Base);
-		Base = nullptr;
+		Result = Runner->CloneOrTakeOver(MoveTemp(Base));
 
 		Result->DataStorage.SetNumLODs(LevelCount);
 
@@ -1260,13 +1256,11 @@ namespace UE::Mutable::Private
 	bool FImageMipmapTask::Complete(CodeRunner* Runner)
 	{
 		FImageOperator ImOp = MakeImageOperator(Runner);
-		ImOp.ImageMipmap_ReleaseScratch(Scratch);
 
-		// This runs in the Runner thread
-		if (Base)
-		{
-			Runner->Release(Base);
-		}
+		Scratch.Uncompressed = nullptr;
+		Scratch.UncompressedMips = nullptr;
+		Scratch.CompressedMips = nullptr;
+		Base = nullptr;
 
 		// If no shortcut was taken
 		if (Result)
@@ -1279,9 +1273,6 @@ namespace UE::Mutable::Private
 	}
 
 
-	//---------------------------------------------------------------------------------------------
-	//---------------------------------------------------------------------------------------------
-	//---------------------------------------------------------------------------------------------
 	class FImageSwizzleTask : public CodeRunner::FIssuedTask
 	{
 	public:
@@ -1296,8 +1287,8 @@ namespace UE::Mutable::Private
 
 	private:
 		OP::ImageSwizzleArgs Args;
-		TSharedPtr<const FImage> Sources[MUTABLE_OP_MAX_SWIZZLE_CHANNELS];
-		TSharedPtr<FImage> Result;
+		TManagedPtr<const FImage> Sources[MUTABLE_OP_MAX_SWIZZLE_CHANNELS];
+		TManagedPtr<FImage> Result;
 	};
 
 
@@ -1330,8 +1321,9 @@ namespace UE::Mutable::Private
 		{
 			for (int32 SourceIndex = 0; SourceIndex < MUTABLE_OP_MAX_SWIZZLE_CHANNELS; ++SourceIndex)
 			{
-				Runner->Release(Sources[SourceIndex]);
+				Sources[SourceIndex] = nullptr;
 			}
+
 			Runner->StoreImage(Op, nullptr);
 
 			return false;
@@ -1352,12 +1344,11 @@ namespace UE::Mutable::Private
 				MUTABLE_CPUPROFILER_SCOPE(ImageFormat_ForSwizzle);
 
 				EImageFormat UncompressedFormat = GetUncompressedFormat(Sources[SourceIndex]->GetFormat());
-				TSharedPtr<FImage> Formatted = Runner->CreateImage(Sources[SourceIndex]->GetSizeX(), Sources[SourceIndex]->GetSizeY(), 1, UncompressedFormat, EInitializationType::NotInitialized);
+				TManagedPtr<FImage> Formatted = Runner->CreateImage(Sources[SourceIndex]->GetSizeX(), Sources[SourceIndex]->GetSizeY(), 1, UncompressedFormat, EInitializationType::NotInitialized);
 				bool bSuccess = false;
 				int32 ImageCompressionQuality = 4; // TODO
 				ImOp.ImagePixelFormat(bSuccess, ImageCompressionQuality, Formatted.Get(), Sources[SourceIndex].Get());
 				check(bSuccess); // Decompression cannot fail
-				Runner->Release(Sources[SourceIndex]);
 				Sources[SourceIndex] = Formatted;
 				ResultLODs = 1;
 			}
@@ -1371,9 +1362,8 @@ namespace UE::Mutable::Private
 			if (Sources[SourceIndex] && ResultSize != Sources[SourceIndex]->GetSize())
 			{
 				MUTABLE_CPUPROFILER_SCOPE(ImageResize_ForSwizzle);
-				TSharedPtr<FImage> Resized = Runner->CreateImage(ResultSize.X, ResultSize.Y, 1, Sources[SourceIndex]->GetFormat(), EInitializationType::NotInitialized);
+				TManagedPtr<FImage> Resized = Runner->CreateImage(ResultSize.X, ResultSize.Y, 1, Sources[SourceIndex]->GetFormat(), EInitializationType::NotInitialized);
 				ImOp.ImageResizeLinear(Resized.Get(), 0, Sources[SourceIndex].Get());
-				Runner->Release(Sources[SourceIndex]);
 				Sources[SourceIndex] = Resized;
 				ResultLODs = 1;
 			}
@@ -1398,16 +1388,17 @@ namespace UE::Mutable::Private
 		// This runs in a worker thread
 		MUTABLE_CPUPROFILER_SCOPE(FImageSwizzleTask);
 
-		ImageSwizzle(Result.Get(), Sources, Args.sourceChannels);
+		ImageSwizzle(Result.Get(), Sources, Args.sourceChannels);	
 	}
 
 
 	bool FImageSwizzleTask::Complete(CodeRunner* Runner)
 	{
 		// This runs in the Runner thread
+
 		for (int32 SourceIndex = 0; SourceIndex < MUTABLE_OP_MAX_SWIZZLE_CHANNELS; ++SourceIndex)
 		{
-			Runner->Release(Sources[SourceIndex]);
+			Sources[SourceIndex] = nullptr;
 		}
 
 		// \TODO: If Result LODs differ from Source[0]'s, rebuild mips?
@@ -1438,7 +1429,7 @@ namespace UE::Mutable::Private
 
 	private:
 		const OP::ImageSaturateArgs Args;
-		TSharedPtr<FImage> Result;
+		TManagedPtr<FImage> Result;
 		float Factor;
 	};
 
@@ -1457,7 +1448,7 @@ namespace UE::Mutable::Private
 		MUTABLE_CPUPROFILER_SCOPE(FImageSaturateTask_Prepare);
 		bOutFailed = false;
 
-		TSharedPtr<const FImage> Source = Runner->LoadImage(FCacheAddress(Args.Base, Op));
+		TManagedPtr<const FImage> Source = Runner->LoadImage(FCacheAddress(Args.Base, Op));
 		Factor = Runner->LoadScalar(FScheduledOp::FromOpAndOptions(Args.Factor, Op, 0));
 
 		if (!Source)
@@ -1473,7 +1464,7 @@ namespace UE::Mutable::Private
 			return false;
 		}
 
-		Result = Runner->CloneOrTakeOver(Source);
+		Result = Runner->CloneOrTakeOver(MoveTemp(Source));
 		return true;
 	}
 
@@ -1521,8 +1512,8 @@ namespace UE::Mutable::Private
 	private:
 		int32 ImageCompressionQuality = 0;
 		OP::ImageResizeArgs Args;
-		TSharedPtr<const FImage> Base;
-		TSharedPtr<FImage> Result;
+		TManagedPtr<const FImage> Base;
+		TManagedPtr<FImage> Result;
 		FImageOperator::FImagePixelFormatFunc ImagePixelFormatFunc;
 	};
 
@@ -1543,10 +1534,7 @@ namespace UE::Mutable::Private
 		destSize[1] = FMath::Max(destSize[1] >> MipsToSkip, 1);
 
 		Base = Runner->LoadImage(FCacheAddress(Args.Source, Op));
-		if (!Base 
-			|| 
-			( Base->GetSizeX()==destSize[0] && Base->GetSizeY()==destSize[1] )
-			)
+		if (!Base || (Base->GetSizeX() == destSize[0] && Base->GetSizeY() == destSize[1]))
 		{
 			Runner->StoreImage(Op, Base);
 			return false;
@@ -1610,7 +1598,7 @@ namespace UE::Mutable::Private
 	bool FImageResizeTask::Complete(CodeRunner* Runner)
 	{
 		// This runs in the Runner thread
-		Runner->Release(Base);
+		Base = nullptr;
 
 		// If didn't take a shortcut and set it already
 		if (Result)
@@ -1641,8 +1629,8 @@ namespace UE::Mutable::Private
 	private:
 		OP::ImageResizeRelArgs Args;
 		int32 ImageCompressionQuality=0;
-		TSharedPtr<const FImage> Base;
-		TSharedPtr<FImage> Result;
+		TManagedPtr<const FImage> Base;
+		TManagedPtr<FImage> Result;
 		FImageSize DestSize;
 		FImageOperator::FImagePixelFormatFunc ImagePixelFormatFunc;
 	};
@@ -1657,9 +1645,9 @@ namespace UE::Mutable::Private
 		ImageCompressionQuality = Runner->Settings.ImageCompressionQuality;
 
 		Base = Runner->LoadImage(FCacheAddress(Args.Source, Op));
-		if (!Base)
+		if (!Base->DataStorage.NumLODs)
 		{
-			Runner->StoreImage(Op, Base);
+			Runner->StoreImage(Op, Base.ToStrongRef());
 			return false;
 		}
 
@@ -1688,7 +1676,7 @@ namespace UE::Mutable::Private
 		if (Base->IsReference())
 		{
 			// We are trying to resize an external reference. This shouldn't happen, but be deffensive.
-			Runner->StoreImage(Op, Base);
+			Runner->StoreImage(Op, Base.ToStrongRef());
 			return false;
 		}
 
@@ -1721,7 +1709,7 @@ namespace UE::Mutable::Private
 	bool FImageResizeRelTask::Complete(CodeRunner* Runner)
 	{
 		// This runs in the Runner thread
-		Runner->Release(Base);
+		Base = nullptr;
 
 		// If didn't take a shortcut and set it already
 		if (Result)
@@ -1750,7 +1738,7 @@ namespace UE::Mutable::Private
 		virtual bool Complete(CodeRunner* Runner) override;
 
 	private:
-		TSharedPtr<FImage> Result;
+		TManagedPtr<FImage> Result;
 		OP::ImageInvertArgs Args;
 	};
 
@@ -1761,12 +1749,12 @@ namespace UE::Mutable::Private
 		MUTABLE_CPUPROFILER_SCOPE(FImageInvertTask_Prepare);
 		bOutFailed = false;
 
-		TSharedPtr<const FImage> Source = Runner->LoadImage({ Args.Base, Op.ExecutionIndex, Op.ExecutionOptions });
+		TManagedPtr<const FImage> Source = Runner->LoadImage({ Args.Base, Op.ExecutionIndex, Op.ExecutionOptions });
 
 		if (Source)
 		{
 			// Create destination data
-			Result = Runner->CloneOrTakeOver(Source);
+			Result = Runner->CloneOrTakeOver(MoveTemp(Source));
 		}
 
 		return true;
@@ -1801,7 +1789,7 @@ namespace UE::Mutable::Private
 	class FImageComposeTask : public CodeRunner::FIssuedTask
 	{
 	public:
-		FImageComposeTask(const FScheduledOp& InOp, const OP::ImageComposeArgs& InArgs, const TSharedPtr<const FLayout>& InLayout)
+		FImageComposeTask(const FScheduledOp& InOp, const OP::ImageComposeArgs& InArgs, const TManagedPtr<const FLayout>& InLayout)
 			: FIssuedTask(InOp), Args(InArgs), Layout(InLayout)
 		{}
 
@@ -1813,10 +1801,10 @@ namespace UE::Mutable::Private
 	private:
 		int32 ImageCompressionQuality = 0;
 		OP::ImageComposeArgs Args;
-		TSharedPtr<const FLayout> Layout;
-		TSharedPtr<const FImage> Block;
-		TSharedPtr<const FImage> Mask;
-		TSharedPtr<FImage> Result;
+		TManagedPtr<const FLayout> Layout;
+		TManagedPtr<const FImage> Block;
+		TManagedPtr<const FImage> Mask;
+		TManagedPtr<FImage> Result;
 		box<FIntVector2> Rect;
 		FImageOperator::FImagePixelFormatFunc ImagePixelFormatFunc;
 	};
@@ -1830,7 +1818,7 @@ namespace UE::Mutable::Private
 
 		ImageCompressionQuality = Runner->Settings.ImageCompressionQuality;
 
-		TSharedPtr<const FImage> Base = Runner->LoadImage({ Args.base, Op.ExecutionIndex, Op.ExecutionOptions });
+		TManagedPtr<const FImage> Base = Runner->LoadImage({ Args.base, Op.ExecutionIndex, Op.ExecutionOptions });
 
 		int32 RelBlockIndex = Layout->FindBlock(Args.BlockId);
 
@@ -1864,14 +1852,14 @@ namespace UE::Mutable::Private
 
 		if (!(Block && Rect.size[0] && Rect.size[1] && Block->GetSizeX() && Block->GetSizeY()))
 		{
-			Runner->Release(Block);
-			Runner->Release(Mask);
+			Block = nullptr;
+			Mask = nullptr;
 			Runner->StoreImage(Op, Base);
 			return false;
 		}
 
 		// Create destination data
-		Result = Runner->CloneOrTakeOver(Base);
+		Result = Runner->CloneOrTakeOver(MoveTemp(Base));
 		Result->Flags = 0;
 
 		bool useMask = Args.mask != 0;
@@ -1890,9 +1878,8 @@ namespace UE::Mutable::Private
 
 				// This now happens more often since the generation of specific mips on request. For this reason
 				// this warning is usually acceptable.
-				TSharedPtr<FImage> Resized = Runner->CreateImage(Rect.size[0], Rect.size[1], 1, Block->GetFormat(), EInitializationType::NotInitialized );
+				TManagedPtr<FImage> Resized = Runner->CreateImage(Rect.size[0], Rect.size[1], 1, Block->GetFormat(), EInitializationType::NotInitialized );
 				ImOp.ImageResizeLinear(Resized.Get(), ImageCompressionQuality, Block.Get());
-				Runner->Release(Block);
 				Block = Resized;
 			}
 
@@ -1904,20 +1891,18 @@ namespace UE::Mutable::Private
 
 				if (Result->GetFormat() != Format)
 				{
-					TSharedPtr<FImage> Formatted = Runner->CreateImage(Result->GetSizeX(), Result->GetSizeY(), Result->GetLODCount(), Format, EInitializationType::NotInitialized);
+					TManagedPtr<FImage> Formatted = Runner->CreateImage(Result->GetSizeX(), Result->GetSizeY(), Result->GetLODCount(), Format, EInitializationType::NotInitialized);
 					bool bSuccess = false;
 					ImOp.ImagePixelFormat(bSuccess, ImageCompressionQuality, Formatted.Get(), Result.Get());
 					check(bSuccess); // Decompression cannot fail
-					Runner->Release(Result);
 					Result = Formatted;
 				}
 				if (Block->GetFormat() != Format)
 				{
-					TSharedPtr<FImage> Formatted = Runner->CreateImage(Block->GetSizeX(), Block->GetSizeY(), Block->GetLODCount(), Format, EInitializationType::NotInitialized);
+					TManagedPtr<FImage> Formatted = Runner->CreateImage(Block->GetSizeX(), Block->GetSizeY(), Block->GetLODCount(), Format, EInitializationType::NotInitialized);
 					bool bSuccess = false;
 					ImOp.ImagePixelFormat(bSuccess, ImageCompressionQuality, Formatted.Get(), Block.Get());
 					check(bSuccess); // Decompression cannot fail
-					Runner->Release(Block);
 					Block = Formatted;
 				}
 			}
@@ -1959,8 +1944,8 @@ namespace UE::Mutable::Private
 	bool FImageComposeTask::Complete(CodeRunner* Runner)
 	{
 		// This runs in the mutable Runner thread
-		Runner->Release(Block);
-		Runner->Release(Mask);
+		Block = nullptr;
+		Mask = nullptr;
 
 		// If we didn't take a shortcut
 		if (Result)
@@ -2141,7 +2126,7 @@ namespace UE::Mutable::Private
 					FInputMemoryStream Stream(RomLoadOp->m_streamBuffer.GetData(), RomLoadOp->m_streamBuffer.Num());
 					FInputArchive Arch(&Stream);
 
-					TSharedPtr<FMesh> Value = FMesh::StaticUnserialise(Arch);
+					TManagedPtr<FMesh> Value = FMesh::StaticUnserialise(Arch);
 
 					Program.SetMeshRomValue(RomIndex, Value);
 				}
@@ -2313,7 +2298,7 @@ namespace UE::Mutable::Private
 					FInputArchive Arch(&Stream);
 
 					// TODO: Try to reuse buffer from PooledImages.
-					TSharedPtr<FImage> Value = FImage::StaticUnserialise(Arch);
+					TManagedPtr<FImage> Value = FImage::StaticUnserialise(Arch);
 
 					Program.SetImageRomValue(RomIndex, Value);
 				}
@@ -2380,7 +2365,7 @@ namespace UE::Mutable::Private
 		uint8 MipmapsToSkip;
 		CodeRunner::FExternalResourceId Id;
 
-		TSharedPtr<FImage> Result;
+		TManagedPtr<FImage> Result;
 		
 		TFunction<void()> ExternalCleanUpFunc;
 	};
@@ -2406,7 +2391,7 @@ namespace UE::Mutable::Private
 		// the task and the later will always outlive the former. 
 		
 		// Probably we could simply pass a reference to the result image. 
-		TFunction<void (TSharedPtr<FImage>)> ResultCallback = [this](TSharedPtr<FImage> InResult)
+		TFunction<void (TManagedPtr<FImage>)> ResultCallback = [this](TManagedPtr<FImage> InResult)
 		{
 			Result = InResult;
 		};
@@ -2451,7 +2436,7 @@ namespace UE::Mutable::Private
 		int32 SectionIndex = 0;
 		uint32 MeshID = 0;
 
-		TSharedPtr<FMesh> Result;
+		TManagedPtr<FMesh> Result;
 
 		TFunction<void()> ExternalCleanUpFunc;
 	};
@@ -2479,7 +2464,7 @@ namespace UE::Mutable::Private
 		// the task and the later will always outlive the former. 
 
 		// Probably we could simply pass a reference to the result mesh. 
-		TFunction<void(TSharedPtr<FMesh>)> ResultCallback = [this](TSharedPtr<FMesh> InResult)
+		TFunction<void(TManagedPtr<FMesh>)> ResultCallback = [this](TManagedPtr<FMesh> InResult)
 			{
 				Result = InResult;
 			};
@@ -2500,7 +2485,7 @@ namespace UE::Mutable::Private
 
 		Result->MeshIDPrefix = MeshID;
 		Runner->StoreMesh(Op, Result);
-
+		
 		bool bSuccess = true;
 		return bSuccess;
 	}
@@ -2585,15 +2570,13 @@ namespace UE::Mutable::Private
 			{
 				OP::ImageParameterConvertArgs Args = Program.GetOpArgs<OP::ImageParameterConvertArgs>(Item.At);
 			
-				TSharedPtr<const FImage> Image = LoadImage(FCacheAddress(Args.ImageParameter, Item));
+				TManagedPtr<const FImage> Image = LoadImage(FCacheAddress(Args.ImageParameter, Item));
 
 				const uint8 MipmapsToSkip = Item.ExecutionOptions;
 
 				FExternalResourceId FullId;
 				FullId.ImageParameter = Image->PassthroughObject.Get();
 				Issued = MakeShared<FImageExternalLoadTask>(Item, MipmapsToSkip, FullId);
-
-				Release(Image);
 			}
 
 			break;
@@ -2655,7 +2638,7 @@ namespace UE::Mutable::Private
 			{
 				OP::MeshSkeletalMeshBreakArgs Args = Program.GetOpArgs<OP::MeshSkeletalMeshBreakArgs>(Item.At);
 
-				TSharedPtr<const FSkeletalMesh> SkeletalMesh = LoadSkeletalMesh(FCacheAddress(Args.SkeletalMeshParameter, Item));
+				TManagedPtr<const FSkeletalMesh> SkeletalMesh = LoadSkeletalMesh(FCacheAddress(Args.SkeletalMeshParameter, Item));
 
 				int32 LODIndex = Args.LOD;
 				int32 SectionIndex = Args.Section;
@@ -2767,7 +2750,7 @@ namespace UE::Mutable::Private
 				(ExecutionStrategy != EExecutionStrategy::MinimizeMemory && Item.Stage == 2))
 			{
 				OP::ImageComposeArgs Args = Program.GetOpArgs<OP::ImageComposeArgs>(Item.At);
-				TSharedPtr<const FLayout> ComposeLayout = StaticCastSharedPtr<const FLayout>( HeapData[Item.CustomState].Resource);
+				TManagedPtr<const FLayout> ComposeLayout = StaticCastManagedPtr<const FLayout>( HeapData[Item.CustomState].Resource);
 				Issued = MakeShared<FImageComposeTask>(Item, Args, ComposeLayout);
 			}
 			break;
diff --git a/Engine/Plugins/Mutable/Source/MutableRuntime/Private/MuR/ExtensionData.cpp b/Engine/Plugins/Mutable/Source/MutableRuntime/Private/MuR/ExtensionData.cpp
--- a/Engine/Plugins/Mutable/Source/MutableRuntime/Private/MuR/ExtensionData.cpp
+++ b/Engine/Plugins/Mutable/Source/MutableRuntime/Private/MuR/ExtensionData.cpp
@@ -17,12 +17,12 @@ void FExtensionData::Serialise(const FExtensionData* Data, FOutputArchive& Archi
 
 
 //---------------------------------------------------------------------------------------------
-TSharedPtr<FExtensionData> FExtensionData::StaticUnserialise(FInputArchive& Archive)
+TManagedPtr<FExtensionData> FExtensionData::StaticUnserialise(FInputArchive& Archive)
 {
 	MUTABLE_CPUPROFILER_SCOPE(ExtensionDataUnserialise);
 	LLM_SCOPE_BYNAME(TEXT("MutableRuntime"));
 
-	TSharedPtr<FExtensionData> Result = MakeShared<FExtensionData>();
+	TManagedPtr<FExtensionData> Result = MakeManaged<FExtensionData>();
 	Archive >> *Result;
 
 	return Result;
diff --git a/Engine/Plugins/Mutable/Source/MutableRuntime/Private/MuR/External/MaterialAdapter.cpp b/Engine/Plugins/Mutable/Source/MutableRuntime/Private/MuR/External/MaterialAdapter.cpp
--- a/Engine/Plugins/Mutable/Source/MutableRuntime/Private/MuR/External/MaterialAdapter.cpp
+++ b/Engine/Plugins/Mutable/Source/MutableRuntime/Private/MuR/External/MaterialAdapter.cpp
@@ -13,13 +13,13 @@ namespace UE::Mutable
 {
 	FMaterialAdapter::FMaterialAdapter()
 	{
-		Material = MakeShared<Private::FMaterial>();
+		Material = Private::MakeManaged<Private::FMaterial>();
 	}
 	
 
 	FMaterialAdapter::FMaterialAdapter(const FMaterialAdapter& Other)
 	{
-		Material = MakeShared<Private::FMaterial>();
+		Material = Private::MakeManaged<Private::FMaterial>();
 		*Material = *Other.Material;
 	}
 
@@ -40,7 +40,7 @@ namespace UE::Mutable
 		}
 
 		TValue<FFloatAdapter> Value;
-		Value.Ptr = MakeShared<FInstancedStruct>();
+		Value.Ptr = Private::MakeManaged<FInstancedStruct>();
 		Value.Ptr->InitializeAs(FFloatAdapter::StaticStruct());
 		Value.Ptr->GetMutable<FFloatAdapter>().Value = *Result;
 		
@@ -77,7 +77,7 @@ namespace UE::Mutable
 		}
 
 		TValue<FVectorAdapter> Value;
-		Value.Ptr = MakeShared<FInstancedStruct>();
+		Value.Ptr = Private::MakeManaged<FInstancedStruct>();
 		Value.Ptr->InitializeAs(FVectorAdapter::StaticStruct());
 		Value.Ptr->GetMutable<FVectorAdapter>().Value = *Result;
 		
@@ -113,14 +113,14 @@ namespace UE::Mutable
 			return {};
 		}
 
-		TVariant<Private::OP::ADDRESS, TSharedPtr<const Private::FImage>> Variant = Result->ImageParameter;
+		TVariant<Private::OP::ADDRESS, Private::TManagedPtr<const Private::FImage>> Variant = Result->ImageParameter;
 
-		if (ensure(Variant.IsType<TSharedPtr<const Private::FImage>>()))
+		if (ensure(Variant.IsType<Private::TManagedPtr<const Private::FImage>>()))
 		{
 			TValue<FTextureAdapter> Value;
-			Value.Ptr = MakeShared<FInstancedStruct>();
+			Value.Ptr = Private::MakeManaged<FInstancedStruct>();
 			Value.Ptr->InitializeAs(FTextureAdapter::StaticStruct());
-			Value.Ptr->GetMutable<FTextureAdapter>().Image = ConstCastSharedPtr<Private::FImage>(Variant.Get<TSharedPtr<const Private::FImage>>()); // TODO GMT If lazy wait
+			Value.Ptr->GetMutable<FTextureAdapter>().Image = ConstCastManagedPtr<Private::FImage>(Variant.Get<Private::TManagedPtr<const Private::FImage>>()); // TODO GMT If lazy wait
 			
 			return TOptional<TValueConst<FTextureAdapter>>(MoveTemp(Value));
 		}
@@ -141,8 +141,8 @@ namespace UE::Mutable
 
 		if (Value)
 		{
-			TVariant<Private::OP::ADDRESS, TSharedPtr<const Private::FImage>> Image;
-			Image.Set<TSharedPtr<const Private::FImage>>(Value->Get().Image);
+			TVariant<Private::OP::ADDRESS, Private::TManagedPtr<const Private::FImage>> Image;
+			Image.Set<Private::TManagedPtr<const Private::FImage>>(Value->Get().Image);
 			
 			Material->ImageParameters.Add(ParameterKey, { Image, ImageData->ImagePropertyIndex });
 		}
diff --git a/Engine/Plugins/Mutable/Source/MutableRuntime/Private/MuR/External/MeshAdapter.cpp b/Engine/Plugins/Mutable/Source/MutableRuntime/Private/MuR/External/MeshAdapter.cpp
--- a/Engine/Plugins/Mutable/Source/MutableRuntime/Private/MuR/External/MeshAdapter.cpp
+++ b/Engine/Plugins/Mutable/Source/MutableRuntime/Private/MuR/External/MeshAdapter.cpp
@@ -11,13 +11,13 @@ namespace UE::Mutable
 {
 	FMeshAdapter::FMeshAdapter()
 	{
-		Mesh = MakeShared<Private::FMesh>();
+		Mesh = Private::MakeManaged<Private::FMesh>();
 	}
 	
 
 	FMeshAdapter::FMeshAdapter(const FMeshAdapter& Other)
 	{
-		Mesh = MakeShared<Private::FMesh>();
+		Mesh = Private::MakeManaged<Private::FMesh>();
 		Mesh->CopyFrom(*Other.Mesh);
 	}
 
diff --git a/Engine/Plugins/Mutable/Source/MutableRuntime/Private/MuR/External/TextureAdapter.cpp b/Engine/Plugins/Mutable/Source/MutableRuntime/Private/MuR/External/TextureAdapter.cpp
--- a/Engine/Plugins/Mutable/Source/MutableRuntime/Private/MuR/External/TextureAdapter.cpp
+++ b/Engine/Plugins/Mutable/Source/MutableRuntime/Private/MuR/External/TextureAdapter.cpp
@@ -4,18 +4,17 @@
 
 #include "MuR/Image.h"
 
-
 namespace UE::Mutable
 {
 	FTextureAdapter::FTextureAdapter()
 	{
-		Image = MakeShared<Private::FImage>();		
+		Image = Private::MakeManaged<Private::FImage>();		
 	}
 
 
 	FTextureAdapter::FTextureAdapter(const FTextureAdapter& Other)
 	{
-		Image = MakeShared<Private::FImage>();
+		Image = Private::MakeManaged<Private::FImage>();
 		Image->Copy(Other.Image.Get());
 	}
 
diff --git a/Engine/Plugins/Mutable/Source/MutableRuntime/Private/MuR/External/Value.cpp b/Engine/Plugins/Mutable/Source/MutableRuntime/Private/MuR/External/Value.cpp
--- a/Engine/Plugins/Mutable/Source/MutableRuntime/Private/MuR/External/Value.cpp
+++ b/Engine/Plugins/Mutable/Source/MutableRuntime/Private/MuR/External/Value.cpp
@@ -9,17 +9,17 @@ namespace UE::Mutable
 {
 	FValue CopyOrMove(FValueConst&& ValueConst)
 	{
-		if (ValueConst.Ptr.IsUnique())
+		if (ValueConst.Ptr.IsUniqueReference())
 		{
 			FValue Value;
-			Value.Ptr = ConstCastSharedPtr<FInstancedStruct>(ValueConst.Ptr); 
+			Value.Ptr = Private::ConstCastManagedPtr<FInstancedStruct>(ValueConst.Ptr); 
 
 			return MoveTemp(Value);
 		}
 		else
 		{
 			FValue Value;
-			Value.Ptr = MakeShared<FInstancedStruct>(*ValueConst.Ptr); 
+			Value.Ptr = Private::MakeManaged<FInstancedStruct>(*ValueConst.Ptr); 
 
 			return MoveTemp(Value);
 		}
diff --git a/Engine/Plugins/Mutable/Source/MutableRuntime/Private/MuR/Image.cpp b/Engine/Plugins/Mutable/Source/MutableRuntime/Private/MuR/Image.cpp
--- a/Engine/Plugins/Mutable/Source/MutableRuntime/Private/MuR/Image.cpp
+++ b/Engine/Plugins/Mutable/Source/MutableRuntime/Private/MuR/Image.cpp
@@ -43,9 +43,9 @@ namespace UE::Mutable::Private
 	}
 
 
-	TSharedPtr<FImage> FImage::CreateAsReference(const TPassthroughObjectPtr<UTexture>& ExternalObject, const FImageDesc& Desc, bool bForceLoad)
+	TManagedRef<FImage> FImage::CreateAsReference(const TPassthroughObjectPtr<UTexture>& ExternalObject, const FImageDesc& Desc, bool bForceLoad)
 	{
-		TSharedPtr<FImage> Result = MakeShared<FImage>();
+		TManagedRef<FImage> Result = MakeManaged<FImage>();
 		Result->PassthroughObject = ExternalObject;
 	
 		Result->DataStorage.InitVoid(Desc);
@@ -80,11 +80,11 @@ namespace UE::Mutable::Private
 	}
 
 
-	TSharedPtr<FImage> FImage::StaticUnserialise(FInputArchive& Arch)
+	TManagedRef<FImage> FImage::StaticUnserialise(FInputArchive& Arch)
 	{
 		LLM_SCOPE_BYNAME(TEXT("MutableRuntime"));
 
-		TSharedPtr<FImage> ResultPtr = MakeShared<FImage>();
+		TManagedRef<FImage> ResultPtr = MakeManaged<FImage>();
 		Arch >> *ResultPtr;
 		
 		return ResultPtr;
@@ -95,7 +95,7 @@ namespace UE::Mutable::Private
 		return DataStorage.GetAllocatedSize();
 	}
 
-	TSharedPtr<FImage> FImageOperator::ExtractMip(const FImage* This, int32 Mip)
+	TManagedPtr<FImage> FImageOperator::ExtractMip(const FImage* This, int32 Mip)
 	{
 		MUTABLE_CPUPROFILER_SCOPE(ExtractMip);
 		
@@ -108,7 +108,7 @@ namespace UE::Mutable::Private
 		
 		if (Mip < This->GetLODCount())
 		{
-			TSharedPtr<FImage> Result = CreateImage(MipSize.X, MipSize.Y, 1, This->GetFormat(), EInitializationType::NotInitialized);
+			TManagedPtr<FImage> Result = CreateImage(MipSize.X, MipSize.Y, 1, This->GetFormat(), EInitializationType::NotInitialized);
 			Result->Flags = This->Flags;
 
 			// RLE formats will not allocate memory when created. Resize manually.
@@ -126,7 +126,7 @@ namespace UE::Mutable::Private
 			return Result;	
 		}
 
-		TSharedPtr<FImage> Resized = CreateImage(MipSize.X, MipSize.Y, 1, This->GetFormat(), EInitializationType::NotInitialized);
+		TManagedPtr<FImage> Resized = CreateImage(MipSize.X, MipSize.Y, 1, This->GetFormat(), EInitializationType::NotInitialized);
 		
 		::FImageView SourceView;
 		::FImageView DestView;
@@ -912,7 +912,7 @@ namespace UE::Mutable::Private
 		{
 			// Generic case that supports compressed formats.
 			const FImageFormatData& FormatData = GetImageFormatData(Format);
-			TSharedPtr<FImage> BlockImage = CreateImage(
+			TManagedPtr<FImage> BlockImage = CreateImage(
 						FormatData.PixelsPerBlockX, FormatData.PixelsPerBlockY, 1, EImageFormat::RGBA_UByte, EInitializationType::NotInitialized);
 
 			const uint32 PixelData = 
@@ -928,8 +928,7 @@ namespace UE::Mutable::Private
 				FMemory::Memcpy(UncompressedBlockData + I*4, &PixelData, 4);
 			}
 
-			TSharedPtr<FImage> Converted = ImagePixelFormat(0, BlockImage.Get(), Format);
-			ReleaseImage(BlockImage);
+			TManagedPtr<FImage> Converted = ImagePixelFormat(0, BlockImage.Get(), Format);
 
 			const TArrayView<const uint8> BlockView = Converted->DataStorage.GetLOD(0); 
 
@@ -955,7 +954,6 @@ namespace UE::Mutable::Private
 				}
 			}
 
-			ReleaseImage(Converted);
 			break;	
 		}
 		}
diff --git a/Engine/Plugins/Mutable/Source/MutableRuntime/Private/MuR/Instance.cpp b/Engine/Plugins/Mutable/Source/MutableRuntime/Private/MuR/Instance.cpp
--- a/Engine/Plugins/Mutable/Source/MutableRuntime/Private/MuR/Instance.cpp
+++ b/Engine/Plugins/Mutable/Source/MutableRuntime/Private/MuR/Instance.cpp
@@ -9,10 +9,10 @@
 
 namespace UE::Mutable::Private
 {
-    TSharedPtr<FInstance> FInstance::Clone() const
+    TManagedPtr<FInstance> FInstance::Clone() const
     {
 		LLM_SCOPE_BYNAME(TEXT("MutableRuntime"));
-		TSharedPtr<FInstance> Result = MakeShared<FInstance>();
+		TManagedPtr<FInstance> Result = MakeManaged<FInstance>();
 
         *Result = *this;
 
@@ -207,7 +207,7 @@ int32 FInstance::GetStringCount(int32 ComponentIndex, int32 LODIndex, int32 Surf
 	}
 
 	
-	void FInstance::GetExtensionData(int32 Index, TSharedPtr<const FExtensionData>& OutExtensionData, FName& OutName) const
+	void FInstance::GetExtensionData(int32 Index, TManagedPtr<const FExtensionData>& OutExtensionData, FName& OutName) const
 	{
 		check(ExtensionData.IsValidIndex(Index));
 
@@ -365,7 +365,7 @@ int32 FInstance::GetStringCount(int32 ComponentIndex, int32 LODIndex, int32 Surf
 	}
 
 
-	void FInstance::SetExtensionData(const TSharedRef<const FExtensionData>& Data, FName Name)
+	void FInstance::SetExtensionData(const TManagedRef<const FExtensionData>& Data, FName Name)
 	{
 		NamedExtensionData& Entry = ExtensionData.AddDefaulted_GetRef();
 		Entry.Data = Data;
diff --git a/Engine/Plugins/Mutable/Source/MutableRuntime/Private/MuR/Layout.cpp b/Engine/Plugins/Mutable/Source/MutableRuntime/Private/MuR/Layout.cpp
--- a/Engine/Plugins/Mutable/Source/MutableRuntime/Private/MuR/Layout.cpp
+++ b/Engine/Plugins/Mutable/Source/MutableRuntime/Private/MuR/Layout.cpp
@@ -22,19 +22,19 @@ namespace UE::Mutable::Private {
 	}
 
 
-	TSharedPtr<FLayout> FLayout::StaticUnserialise( FInputArchive& Arch )
+	TManagedRef<FLayout> FLayout::StaticUnserialise( FInputArchive& Arch )
 	{
 		LLM_SCOPE_BYNAME(TEXT("MutableRuntime"));
-		TSharedPtr<FLayout> Result = MakeShared<FLayout>();
+		TManagedRef<FLayout> Result = MakeManaged<FLayout>();
 		Arch >> *Result;
 		return Result;
 	}
 
 
-	TSharedPtr<FLayout> FLayout::Clone() const
+	TManagedPtr<FLayout> FLayout::Clone() const
 	{
 		LLM_SCOPE_BYNAME(TEXT("MutableRuntime"));
-		TSharedPtr<FLayout> Result = MakeShared<FLayout>();
+		TManagedPtr<FLayout> Result = MakeManaged<FLayout>();
 		Result->Size = Size;
 		Result->MaxSize = MaxSize;
 		Result->Blocks = Blocks;
diff --git a/Engine/Plugins/Mutable/Source/MutableRuntime/Private/MuR/Material.cpp b/Engine/Plugins/Mutable/Source/MutableRuntime/Private/MuR/Material.cpp
--- a/Engine/Plugins/Mutable/Source/MutableRuntime/Private/MuR/Material.cpp
+++ b/Engine/Plugins/Mutable/Source/MutableRuntime/Private/MuR/Material.cpp
@@ -18,12 +18,12 @@ namespace UE::Mutable::Private
 	}
 
 
-	TSharedPtr<FMaterial> FMaterial::StaticUnserialise(FInputArchive& Arch)
+	TManagedPtr<FMaterial> FMaterial::StaticUnserialise(FInputArchive& Arch)
 	{
 		MUTABLE_CPUPROFILER_SCOPE(MaterialUnserialise)
 		LLM_SCOPE_BYNAME(TEXT("MutableRuntime"));
 
-		TSharedPtr<FMaterial> Result = MakeShared<FMaterial>();
+		TManagedPtr<FMaterial> Result = MakeManaged<FMaterial>();
 		Arch >> *Result;
 
 		return Result;
@@ -46,9 +46,9 @@ namespace UE::Mutable::Private
 	}
 
 
-	TSharedPtr<FMaterial> FMaterial::Clone() const
+	TManagedPtr<FMaterial> FMaterial::Clone() const
 	{
-		TSharedPtr<FMaterial> Result = MakeShared<FMaterial>();
+		TManagedPtr<FMaterial> Result = MakeManaged<FMaterial>();
 
 		Result->PassthroughObject = PassthroughObject;
 		Result->ImageParameters = ImageParameters;
diff --git a/Engine/Plugins/Mutable/Source/MutableRuntime/Private/MuR/MemoryTrackingUtils.cpp b/Engine/Plugins/Mutable/Source/MutableRuntime/Private/MuR/MemoryTrackingUtils.cpp
--- a/Engine/Plugins/Mutable/Source/MutableRuntime/Private/MuR/MemoryTrackingUtils.cpp
+++ b/Engine/Plugins/Mutable/Source/MutableRuntime/Private/MuR/MemoryTrackingUtils.cpp
@@ -16,7 +16,7 @@ namespace UE::Mutable::Private::Private
 	{
 		// This code will only be compiled if the global definition to enable LLM tracking is set to 1 for the host program
 		// E.g.: GlobalDefinitions.Add("LLM_ENABLED_IN_CONFIG=1");
-#if ENABLE_LOW_LEVEL_MEM_TRACKER && IS_PROGRAM
+#if ENABLE_LOW_LEVEL_MEM_TRACKER /* && IS_PROGRAM*/
 		FLowLevelMemTracker& MemTracker = FLowLevelMemTracker::Get();
 		if (MemTracker.IsEnabled())
 		{
@@ -38,8 +38,8 @@ namespace UE::Mutable::Private
 			TEXT("Enables Mutable plugin global memory tracking, peak and total."),
 			ECVF_Default);
 
-#if ENABLE_LOW_LEVEL_MEM_TRACKER && IS_PROGRAM
-		static inline bool bEnableUpdateLLMStats = false;
+#if ENABLE_LOW_LEVEL_MEM_TRACKER /*&& IS_PROGRAM*/
+		static inline bool bEnableUpdateLLMStats = true;
 		static inline FAutoConsoleVariableRef CVarEnableUpdateLLMStats = FAutoConsoleVariableRef(
 			TEXT("mutable.EnableUpdateLLMStats"),
 			bEnableUpdateLLMStats,
@@ -64,7 +64,7 @@ namespace UE::Mutable::Private
 		AbsoluteCounter += Differential;
 		AbsolutePeakValue = FMath::Max(AbsolutePeakValue, AbsoluteCounter);
 
-#if ENABLE_LOW_LEVEL_MEM_TRACKER && IS_PROGRAM
+#if ENABLE_LOW_LEVEL_MEM_TRACKER /*&& IS_PROGRAM*/
 		if (UNLIKELY(FMemoryTrackingConsoleFlags::bEnableUpdateLLMStats))
 		{
 			Private::UpdateLLMStats();
diff --git a/Engine/Plugins/Mutable/Source/MutableRuntime/Private/MuR/Mesh.cpp b/Engine/Plugins/Mutable/Source/MutableRuntime/Private/MuR/Mesh.cpp
--- a/Engine/Plugins/Mutable/Source/MutableRuntime/Private/MuR/Mesh.cpp
+++ b/Engine/Plugins/Mutable/Source/MutableRuntime/Private/MuR/Mesh.cpp
@@ -29,12 +29,12 @@ void FMesh::Serialise(const FMesh* MeshPtr, FOutputArchive& Arch)
 }
 
 
-TSharedPtr<FMesh> FMesh::StaticUnserialise(FInputArchive& Arch)
+TManagedPtr<FMesh> FMesh::StaticUnserialise(FInputArchive& Arch)
 {
     MUTABLE_CPUPROFILER_SCOPE(MeshUnserialise)
 	LLM_SCOPE_BYNAME(TEXT("MutableRuntime"));
 
-    TSharedPtr<FMesh> Result = MakeShared<FMesh>();
+    TManagedPtr<FMesh> Result = MakeManaged<FMesh>();
     Arch >> *Result;
 
     //Result->m_pD->CheckIntegrity();
@@ -43,9 +43,9 @@ TSharedPtr<FMesh> FMesh::StaticUnserialise(FInputArchive& Arch)
 }
 
 
-TSharedPtr<FMesh> FMesh::CreateAsReference(const TPassthroughObjectPtr<USkeletalMesh>& PasstroughObjectPtr)
+TManagedRef<FMesh> FMesh::CreateAsReference(const TPassthroughObjectPtr<USkeletalMesh>& PasstroughObjectPtr)
 {
-	TSharedPtr<FMesh> Result = MakeShared<FMesh>();
+	TManagedRef<FMesh> Result = MakeManaged<FMesh>();
 	Result->PassthroughObject = PasstroughObjectPtr;
 	
 	return Result;
@@ -70,19 +70,19 @@ const FString& FMesh::GetReferencedMorph() const
 }
 
 
-TSharedPtr<FMesh> FMesh::Clone() const
+TManagedPtr<FMesh> FMesh::Clone() const
 {
     //MUTABLE_CPUPROFILER_SCOPE(MeshClone);
 	return Clone(EMeshCopyFlags::AllFlags);
 }
 
 
-TSharedPtr<FMesh> FMesh::Clone(EMeshCopyFlags InFlags) const
+TManagedPtr<FMesh> FMesh::Clone(EMeshCopyFlags InFlags) const
 {
     //MUTABLE_CPUPROFILER_SCOPE(MeshClone);
 	LLM_SCOPE_BYNAME(TEXT("MutableRuntime"));
 
-	TSharedPtr<FMesh> Result = MakeShared<FMesh>();
+	TManagedPtr<FMesh> Result = MakeManaged<FMesh>();
 
 	Result->CopyFrom(*this, InFlags);
 
@@ -332,37 +332,37 @@ void FMesh::MakeIdsExplicit()
 }
 
 
-TSharedPtr<const FSkeleton> FMesh::GetSkeleton() const
+TManagedPtr<const FSkeleton> FMesh::GetSkeleton() const
 {
     return Skeleton;
 }
 
 
-void FMesh::SetSkeleton(TSharedPtr<const FSkeleton> InSkeleton)
+void FMesh::SetSkeleton(const TManagedPtr<const FSkeleton>& InSkeleton)
 {
     Skeleton = InSkeleton;
 }
 
 
-TSharedPtr<const FPhysicsBody> FMesh::GetPhysicsBody() const
+TManagedPtr<const FPhysicsBody> FMesh::GetPhysicsBody() const
 {
     return PhysicsBody;
 }
 
 
-void FMesh::SetPhysicsBody(TSharedPtr<const FPhysicsBody> InPhysicsBody)
+void FMesh::SetPhysicsBody(const TManagedPtr<const FPhysicsBody>& InPhysicsBody)
 {
     PhysicsBody = InPhysicsBody;
 }
 
 
-int32 FMesh::AddAdditionalPhysicsBody(TSharedPtr<const FPhysicsBody> Body)
+int32 FMesh::AddAdditionalPhysicsBody(const TManagedPtr<const FPhysicsBody>& Body)
 {
 	return AdditionalPhysicsBodies.Add(Body);
 }
 
 
-TSharedPtr<const FPhysicsBody> FMesh::GetAdditionalPhysicsBody(int32 Index) const
+TManagedPtr<const FPhysicsBody> FMesh::GetAdditionalPhysicsBody(int32 Index) const
 {
 	check(AdditionalPhysicsBodies.IsValidIndex(Index));
 
@@ -464,7 +464,7 @@ uint32 FMesh::GetSurfaceId(int32 SurfaceIndex) const
 }
 
 
-void FMesh::AddLayout(TSharedPtr<const FLayout> InLayout)
+void FMesh::AddLayout(TManagedPtr<const FLayout> InLayout)
 {
 	LLM_SCOPE_BYNAME(TEXT("MutableRuntime"));
     Layouts.Add(InLayout);
@@ -477,7 +477,7 @@ int32 FMesh::GetLayoutCount() const
 }
 
 
-TSharedPtr<const FLayout> FMesh::GetLayout(int32 LayoutIndex) const
+TManagedPtr<const FLayout> FMesh::GetLayout(int32 LayoutIndex) const
 {
     check(Layouts.IsValidIndex(LayoutIndex));
 
@@ -485,7 +485,7 @@ TSharedPtr<const FLayout> FMesh::GetLayout(int32 LayoutIndex) const
 }
 
 
-void FMesh::SetLayout(int32 LayoutIndex, TSharedPtr<const FLayout> InLayout)
+void FMesh::SetLayout(int32 LayoutIndex, TManagedPtr<const FLayout> InLayout)
 {
     check(Layouts.IsValidIndex(LayoutIndex));
 
diff --git a/Engine/Plugins/Mutable/Source/MutableRuntime/Private/MuR/Model.cpp b/Engine/Plugins/Mutable/Source/MutableRuntime/Private/MuR/Model.cpp
--- a/Engine/Plugins/Mutable/Source/MutableRuntime/Private/MuR/Model.cpp
+++ b/Engine/Plugins/Mutable/Source/MutableRuntime/Private/MuR/Model.cpp
@@ -151,7 +151,7 @@ namespace UE::Mutable::Private
 		FOutputMemoryStream MemStream(16 * 1024 * 1024);
 
 		// Save images and unload from memory
-		for (TPair<uint32, TSharedPtr<const FImage>>& Entry: Program.ConstantImageLODsStreamed )
+		for (TPair<uint32, TManagedPtr<const FImage>>& Entry: Program.ConstantImageLODsStreamed )
 		{
 			int32 RomIndex = Entry.Key;
 			const FRomDataRuntime& RomData = Program.Roms[RomIndex];
@@ -180,7 +180,7 @@ namespace UE::Mutable::Private
 		}
 
 		// Save meshes and unload from memory
-		for (TPair<uint32, TSharedPtr<const FMesh>>& ResData : Program.ConstantMeshesStreamed)
+		for (TPair<uint32, TManagedPtr<const FMesh>>& ResData : Program.ConstantMeshesStreamed)
 		{
 			int32 RomIndex = ResData.Key;
 			const FRomDataRuntime& RomData = Program.Roms[RomIndex];
diff --git a/Engine/Plugins/Mutable/Source/MutableRuntime/Private/MuR/OpImageCompose.cpp b/Engine/Plugins/Mutable/Source/MutableRuntime/Private/MuR/OpImageCompose.cpp
--- a/Engine/Plugins/Mutable/Source/MutableRuntime/Private/MuR/OpImageCompose.cpp
+++ b/Engine/Plugins/Mutable/Source/MutableRuntime/Private/MuR/OpImageCompose.cpp
@@ -40,15 +40,12 @@ namespace UE::Mutable::Private
 			MUTABLE_CPUPROFILER_SCOPE(ImageCompose_SubBlock_Fix);
 			EImageFormat UncompressedFormat = GetUncompressedFormat(OriginalFormat);
 			constexpr int32 Quality = 0;
-			TSharedPtr<FImage> TempBase = ImagePixelFormat(Quality, Base, UncompressedFormat);
-			TSharedPtr<FImage> TempBlock = ImagePixelFormat(Quality, Block, UncompressedFormat);
+			TManagedPtr<FImage> TempBase = ImagePixelFormat(Quality, Base, UncompressedFormat);
+			TManagedPtr<FImage> TempBlock = ImagePixelFormat(Quality, Block, UncompressedFormat);
 			ImageCompose(TempBase.Get(), TempBlock.Get(), Rect);
-			ReleaseImage(TempBlock);
-			TSharedPtr<FImage> NewBase = ImagePixelFormat(Quality, TempBase.Get(), OriginalFormat);
-			ReleaseImage(TempBase);
+			TManagedPtr<FImage> NewBase = ImagePixelFormat(Quality, TempBase.Get(), OriginalFormat);
 			// ImageCompose above may have reduced the LODs because of the block lods.
 			Base->CopyMove(NewBase.Get());
-			ReleaseImage(NewBase);
 			return;
 		}
 		else
diff --git a/Engine/Plugins/Mutable/Source/MutableRuntime/Private/MuR/OpImageCrop.cpp b/Engine/Plugins/Mutable/Source/MutableRuntime/Private/MuR/OpImageCrop.cpp
--- a/Engine/Plugins/Mutable/Source/MutableRuntime/Private/MuR/OpImageCrop.cpp
+++ b/Engine/Plugins/Mutable/Source/MutableRuntime/Private/MuR/OpImageCrop.cpp
@@ -8,12 +8,12 @@
 
 namespace UE::Mutable::Private
 {
-	bool FImageOperator::ImageCrop(TSharedPtr<FImage>& OutCropped, int32 CompressionQuality, const FImage* InBase, const box< FIntVector2 >& Rect)
+	bool FImageOperator::ImageCrop(TManagedPtr<FImage>& OutCropped, int32 CompressionQuality, const FImage* InBase, const box< FIntVector2 >& Rect)
 	{
 		check(InBase);
 
-		TSharedPtr<FImage> BaseReformat;
-		TSharedPtr<FImage> Cropped = OutCropped;
+		TManagedPtr<FImage> BaseReformat;
+		TManagedPtr<FImage> Cropped = OutCropped;
 
 		EImageFormat BaseFormat = InBase->GetFormat();
 		EImageFormat UncompressedFormat = GetUncompressedFormat(BaseFormat);
@@ -86,13 +86,9 @@ namespace UE::Mutable::Private
 
 		if (BaseFormat != UncompressedFormat)
 		{
-			ReleaseImage(BaseReformat);
-
 			bool bSuccess = false;
 			ImagePixelFormat(bSuccess, CompressionQuality, OutCropped.Get(), Cropped.Get());
 			check(bSuccess);
-
-			ReleaseImage(Cropped);
 		}
 
 		return true;
diff --git a/Engine/Plugins/Mutable/Source/MutableRuntime/Private/MuR/OpImageDisplace.h b/Engine/Plugins/Mutable/Source/MutableRuntime/Private/MuR/OpImageDisplace.h
--- a/Engine/Plugins/Mutable/Source/MutableRuntime/Private/MuR/OpImageDisplace.h
+++ b/Engine/Plugins/Mutable/Source/MutableRuntime/Private/MuR/OpImageDisplace.h
@@ -57,7 +57,7 @@ namespace UE::Mutable::Private
 			FMemory::Memset(Buffer.GetData(), ZeroOffset, Buffer.Num());
 		}
 
-		TSharedPtr<FImage> NextMask = MaskImage->Clone();
+		TManagedPtr<FImage> NextMask = MaskImage->Clone();
 
 		for (int32 B = 0; B < InBorder; ++B)
 		{
diff --git a/Engine/Plugins/Mutable/Source/MutableRuntime/Private/MuR/OpImageLayer.cpp b/Engine/Plugins/Mutable/Source/MutableRuntime/Private/MuR/OpImageLayer.cpp
--- a/Engine/Plugins/Mutable/Source/MutableRuntime/Private/MuR/OpImageLayer.cpp
+++ b/Engine/Plugins/Mutable/Source/MutableRuntime/Private/MuR/OpImageLayer.cpp
@@ -2451,7 +2451,7 @@ namespace OpImageLayerInternal
 
 		const EImageFormat BaseFormat = BaseImage->GetFormat();
 
-		TSharedPtr<FImage> TempMaskImage;
+		TManagedPtr<FImage> TempMaskImage;
 		if (MaskImage->GetFormat() != EImageFormat::L_UByte)
 		{
 			UE_LOG(LogMutableCore, Log, TEXT("Image layer format not supported. A generic one will be used. "));
@@ -2481,8 +2481,8 @@ namespace OpImageLayerInternal
 
 			FImageOperator ImOp = FImageOperator::GetDefault(nullptr);
 			constexpr int32 Quality = 4;
-			TSharedPtr<FImage> TempBaseImage = ImOp.ImagePixelFormat(Quality, BaseImage, EImageFormat::RGBA_UByte);
-			TSharedPtr<FImage> TempBlededImage = ImOp.ImagePixelFormat(Quality, BlendedImage, EImageFormat::RGBA_UByte);
+			TManagedPtr<FImage> TempBaseImage = ImOp.ImagePixelFormat(Quality, BaseImage, EImageFormat::RGBA_UByte);
+			TManagedPtr<FImage> TempBlededImage = ImOp.ImagePixelFormat(Quality, BlendedImage, EImageFormat::RGBA_UByte);
 			BufferLayerCombine<BLEND_FUNC_MASKED, 4>(ResultImage, TempBaseImage.Get(), MaskImage, TempBlededImage.Get(), bOnlyFirstLOD);
 		}
 	}
diff --git a/Engine/Plugins/Mutable/Source/MutableRuntime/Private/MuR/OpImageMipmap.cpp b/Engine/Plugins/Mutable/Source/MutableRuntime/Private/MuR/OpImageMipmap.cpp
--- a/Engine/Plugins/Mutable/Source/MutableRuntime/Private/MuR/OpImageMipmap.cpp
+++ b/Engine/Plugins/Mutable/Source/MutableRuntime/Private/MuR/OpImageMipmap.cpp
@@ -602,13 +602,6 @@ namespace OpImageMipmapInternal
     }
 
 
-	void FImageOperator::ImageMipmap_ReleaseScratch(FScratchImageMipmap& Scratch)
-	{
-		ReleaseImage(Scratch.Uncompressed);
-		ReleaseImage(Scratch.UncompressedMips);
-		ReleaseImage(Scratch.CompressedMips);
-	}
-
 	void FImageOperator::ImageMipmap(FScratchImageMipmap& Scratch, int32 CompressionQuality, FImage* DestImage, const FImage* BaseImage,
 		int32 StartLevel, int32 NumLODs, const FMipmapGenerationSettings& Settings, bool bGenerateOnlyTail)
 	{
@@ -751,7 +744,6 @@ namespace OpImageMipmapInternal
 
 		ImageMipmap_PrepareScratch(Dest, StartLevel, LevelCount, Scratch);
 		ImageMipmap(Scratch, CompressionQuality, Dest, Base, StartLevel, LevelCount, Settings, bGenerateOnlyTail);
-		ImageMipmap_ReleaseScratch(Scratch);
 	}
 
 	/** 
