[Mutable] Memory managment rework.

diff --git a/Engine/Plugins/Mutable/Source/MutableRuntime/Private/MuR/OpImagePixelFormat.cpp b/Engine/Plugins/Mutable/Source/MutableRuntime/Private/MuR/OpImagePixelFormat.cpp
--- a/Engine/Plugins/Mutable/Source/MutableRuntime/Private/MuR/OpImagePixelFormat.cpp
+++ b/Engine/Plugins/Mutable/Source/MutableRuntime/Private/MuR/OpImagePixelFormat.cpp
@@ -19,7 +19,7 @@ namespace UE::Mutable::Private
 {
 
 
-TSharedPtr<FImage> FImageOperator::ImagePixelFormat(int32 CompressionQuality, const FImage* Base, EImageFormat TargetFormat, int32 OnlyLOD)
+TManagedPtr<FImage> FImageOperator::ImagePixelFormat(int32 CompressionQuality, const FImage* Base, EImageFormat TargetFormat, int32 OnlyLOD)
 {
 	MUTABLE_CPUPROFILER_SCOPE(ImagePixelFormat);
 	TRACE_CPUPROFILER_EVENT_SCOPE_TEXT(*FString::Printf(TEXT("From %d to %d"), int32(Base->GetFormat()), int32(TargetFormat)));
@@ -39,7 +39,7 @@ TSharedPtr<FImage> FImageOperator::ImagePixelFormat(int32 CompressionQuality, co
 		ResultLODCount = 1;
 	}
 
-	TSharedPtr<FImage> Result = CreateImage(ResultSize[0], ResultSize[1], ResultLODCount, TargetFormat, EInitializationType::NotInitialized);
+	TManagedPtr<FImage> Result = CreateImage(ResultSize[0], ResultSize[1], ResultLODCount, TargetFormat, EInitializationType::NotInitialized);
 	Result->Flags = Base->Flags;
 
 	if (Base->GetSizeX() <= 0 || Base->GetSizeY() <= 0)
@@ -263,7 +263,7 @@ void FImageOperator::ImagePixelFormat(bool& bOutSuccess, int32 CompressionQualit
 			}
 			case EImageFormat::RGB_UByteRLE:
 			{
-				TSharedPtr<FImage> TempBase =
+				TManagedPtr<FImage> TempBase =
 						ImagePixelFormat(CompressionQuality, Base, EImageFormat::RGB_UByte);
 
 				ImagePixelFormatInternal::ProcessBatchedHelper(Result, TempBase.Get(), 1, 3, ResultLODBegin, BaseLODBegin, NumLODs,
@@ -276,7 +276,6 @@ void FImageOperator::ImagePixelFormat(bool& bOutSuccess, int32 CompressionQualit
 					}
 				});
 
-				ReleaseImage(TempBase);
 
 				break;
 			}
@@ -296,7 +295,7 @@ void FImageOperator::ImagePixelFormat(bool& bOutSuccess, int32 CompressionQualit
 			}
 			case EImageFormat::RGBA_UByteRLE:
 			{
-				TSharedPtr<FImage> TempBase = ImagePixelFormat(CompressionQuality, Base, EImageFormat::RGBA_UByte);
+				TManagedPtr<FImage> TempBase = ImagePixelFormat(CompressionQuality, Base, EImageFormat::RGBA_UByte);
 				
 				ImagePixelFormatInternal::ProcessBatchedHelper(Result, TempBase.Get(), 1, 4, ResultLODBegin, BaseLODBegin, NumLODs,
 				[](uint8* DestBuf, const uint8* BaseBuf, int32 NumElems)
@@ -308,7 +307,6 @@ void FImageOperator::ImagePixelFormat(bool& bOutSuccess, int32 CompressionQualit
 					}
 				});
 
-				ReleaseImage(TempBase);
 
 				break;
 			}
@@ -330,7 +328,7 @@ void FImageOperator::ImagePixelFormat(bool& bOutSuccess, int32 CompressionQualit
 			case EImageFormat::BC2:
 			case EImageFormat::BC3:
 			{
-				TSharedPtr<FImage> TempBase = ImagePixelFormat(CompressionQuality, Base, EImageFormat::RGB_UByte);
+				TManagedPtr<FImage> TempBase = ImagePixelFormat(CompressionQuality, Base, EImageFormat::RGB_UByte);
 				
 				ImagePixelFormatInternal::ProcessBatchedHelper(Result, TempBase.Get(), 1, 3, ResultLODBegin, BaseLODBegin, NumLODs,
 				[](uint8* DestBuf, const uint8* BaseBuf, int32 NumElems)
@@ -342,7 +340,6 @@ void FImageOperator::ImagePixelFormat(bool& bOutSuccess, int32 CompressionQualit
 					}
 				});
 
-				ReleaseImage(TempBase);
 
 				break;
 			}
@@ -411,7 +408,7 @@ void FImageOperator::ImagePixelFormat(bool& bOutSuccess, int32 CompressionQualit
 //				if (bOutSuccess)
 //				{
 //					// Verify
-//					TSharedPtr<FImage> Test = ImagePixelFormat(CompressionQuality, Result, EImageFormat::L_UByte);
+//					TManagedStrongPtr<FImage> Test = ImagePixelFormat(CompressionQuality, Result, EImageFormat::L_UByte);
 //					check(Test->DataStorage == Base->DataStorage);
 //				}
 //#endif
@@ -423,7 +420,7 @@ void FImageOperator::ImagePixelFormat(bool& bOutSuccess, int32 CompressionQualit
 			case EImageFormat::BGRA_UByte:
 			{
 				check(bFormatAllLODs);
-				TSharedPtr<FImage> TempBase = ImagePixelFormat(CompressionQuality, Base, EImageFormat::L_UByte, -1);
+				TManagedPtr<FImage> TempBase = ImagePixelFormat(CompressionQuality, Base, EImageFormat::L_UByte, -1);
  
 				// Try to compress
 				if (Result->DataStorage.IsEmpty())
@@ -445,11 +442,10 @@ void FImageOperator::ImagePixelFormat(bool& bOutSuccess, int32 CompressionQualit
 //				if (bOutSuccess)
 //				{
 //					// Verify
-//					TSharedPtr<FImage> Test = ImagePixelFormat(CompressionQuality, Result, EImageFormat::L_UByte);
+//					TManagedStrongPtr<FImage> Test = ImagePixelFormat(CompressionQuality, Result, EImageFormat::L_UByte);
 //					check(Test->DataStorage == TempBase->DataStorage);
 //				}
 //#endif
-				ReleaseImage(TempBase);
 
 				break;
 			}
@@ -492,7 +488,7 @@ void FImageOperator::ImagePixelFormat(bool& bOutSuccess, int32 CompressionQualit
 //				if (bOutSuccess)
 //				{
 //					// Verify
-//					TSharedPtr<FImage> Test = ImagePixelFormat(CompressionQuality, Result, EImageFormat::L_UByte);
+//					TManagedStrongPtr<FImage> Test = ImagePixelFormat(CompressionQuality, Result, EImageFormat::L_UByte);
 //					check(Test->DataStorage == Base->DataStorage);
 //				}
 //#endif
@@ -505,18 +501,17 @@ void FImageOperator::ImagePixelFormat(bool& bOutSuccess, int32 CompressionQualit
 			case EImageFormat::BGRA_UByte:
 			{
 				check(bFormatAllLODs);
-				TSharedPtr<FImage> TempBase = ImagePixelFormat(CompressionQuality, Base, EImageFormat::L_UByte, -1);
+				TManagedPtr<FImage> TempBase = ImagePixelFormat(CompressionQuality, Base, EImageFormat::L_UByte, -1);
 				CompressRLE_L1(TempBase.Get(), Result);
 
 //#ifdef MUTABLE_DEBUG_RLE
 //				if (bOutSuccess)
 //				{
 //					// Verify
-//					TSharedPtr<FImage> Test = ImagePixelFormat(CompressionQuality, Result, EImageFormat::L_UByte);
+//					TManagedStrongPtr<FImage> Test = ImagePixelFormat(CompressionQuality, Result, EImageFormat::L_UByte);
 //					check(Test->DataStorage == TempBase->DataStorage);
 //				}
 //#endif
-				ReleaseImage(TempBase);
 				break;
 			}
 
@@ -598,7 +593,7 @@ void FImageOperator::ImagePixelFormat(bool& bOutSuccess, int32 CompressionQualit
 			{
 				check(bFormatAllLODs);
 
-				TSharedPtr<FImage> TempBase = ImagePixelFormat(CompressionQuality, Base, EImageFormat::RGBA_UByte, -1);
+				TManagedPtr<FImage> TempBase = ImagePixelFormat(CompressionQuality, Base, EImageFormat::RGBA_UByte, -1);
 
 				ImagePixelFormatInternal::ProcessBatchedHelper(Result, TempBase.Get(), 3, 4, ResultLODBegin, BaseLODBegin, NumLODs,
 				[](uint8* DestBuf, const uint8* BaseBuf, int32 NumElems)
@@ -611,7 +606,6 @@ void FImageOperator::ImagePixelFormat(bool& bOutSuccess, int32 CompressionQualit
 					}
 				});
 
-				ReleaseImage(TempBase);
 				break;
 			}
 			case EImageFormat::BC1:
@@ -957,7 +951,7 @@ void FImageOperator::ImagePixelFormat(bool& bOutSuccess, int32 CompressionQualit
 			case EImageFormat::RGB_UByteRLE:
 			{
 				check(bFormatAllLODs);
-				TSharedPtr<FImage> TempBase = ImagePixelFormat(CompressionQuality, Base, EImageFormat::RGB_UByte);
+				TManagedPtr<FImage> TempBase = ImagePixelFormat(CompressionQuality, Base, EImageFormat::RGB_UByte);
 
 				ImagePixelFormatInternal::ProcessBatchedHelper(Result, TempBase.Get(), 4, 3, ResultLODBegin, BaseLODBegin, NumLODs,
 				[](uint8* DestBuf, const uint8* BaseBuf, int32 NumElems)
@@ -971,13 +965,12 @@ void FImageOperator::ImagePixelFormat(bool& bOutSuccess, int32 CompressionQualit
 					}
 				});
 
-				ReleaseImage(TempBase);
 				break;
 			}
 			case EImageFormat::L_UByteRLE:
 			{
 				check(bFormatAllLODs);
-				TSharedPtr<FImage> TempBase = ImagePixelFormat(CompressionQuality, Base, EImageFormat::L_UByte);
+				TManagedPtr<FImage> TempBase = ImagePixelFormat(CompressionQuality, Base, EImageFormat::L_UByte);
 
 				ImagePixelFormatInternal::ProcessBatchedHelper(Result, TempBase.Get(), 4, 1, ResultLODBegin, BaseLODBegin, NumLODs,
 				[](uint8* DestBuf, const uint8* BaseBuf, int32 NumElems)
@@ -991,7 +984,6 @@ void FImageOperator::ImagePixelFormat(bool& bOutSuccess, int32 CompressionQualit
 					}
 				});
 
-				ReleaseImage(TempBase);
 				break;
 			}
 
@@ -1020,14 +1012,13 @@ void FImageOperator::ImagePixelFormat(bool& bOutSuccess, int32 CompressionQualit
 				check(bFormatAllLODs);
 
 				// \todo: optimise
-				TSharedPtr<FImage> TempBase = ImagePixelFormat(CompressionQuality, Base, EImageFormat::RGBA_UByte);
+				TManagedPtr<FImage> TempBase = ImagePixelFormat(CompressionQuality, Base, EImageFormat::RGBA_UByte);
 				CompressRLE_RGBA(TempBase.Get(), Result);
 
 				// Test
-				//TSharedPtr<FImage> Test = ImagePixelFormat(CompressionQuality, Result, EImageFormat::RGBA_UByte); 
+				//TManagedStrongPtr<FImage> Test = ImagePixelFormat(CompressionQuality, Result, EImageFormat::RGBA_UByte); 
 				//check(Test->DataStorage == Base->DataStorage);
 
-				ReleaseImage(TempBase);
 
 				break;
 			}
@@ -1036,11 +1027,9 @@ void FImageOperator::ImagePixelFormat(bool& bOutSuccess, int32 CompressionQualit
 				check(bFormatAllLODs);
 
 				// \todo: optimise
-				TSharedPtr<FImage> Temp1 = ImagePixelFormat(CompressionQuality, Base, EImageFormat::RGB_UByte);
-				TSharedPtr<FImage> Temp2 = ImagePixelFormat(CompressionQuality, Temp1.Get(), EImageFormat::RGBA_UByte);
-				ReleaseImage(Temp1);
+				TManagedPtr<FImage> Temp1 = ImagePixelFormat(CompressionQuality, Base, EImageFormat::RGB_UByte);
+				TManagedPtr<FImage> Temp2 = ImagePixelFormat(CompressionQuality, Temp1.Get(), EImageFormat::RGBA_UByte);
 				CompressRLE_RGBA(Temp2.Get(), Result);
-				ReleaseImage(Temp2);
 
 				break;
 			}
@@ -1064,7 +1053,7 @@ void FImageOperator::ImagePixelFormat(bool& bOutSuccess, int32 CompressionQualit
 				CompressRLE_RGB(Base, Result);
 
 				// Test
-				// TSharedPtr<FImage> Test = ImagePixelFormat(CompressionQuality, Result, EImageFormat::RGB_UByte); 
+				// TManagedStrongPtr<FImage> Test = ImagePixelFormat(CompressionQuality, Result, EImageFormat::RGB_UByte); 
 				// check(Test->DataStorage == Base->DataStorage);
 				break;
 			}
@@ -1074,11 +1063,9 @@ void FImageOperator::ImagePixelFormat(bool& bOutSuccess, int32 CompressionQualit
 				check(bFormatAllLODs);
 
 				// \todo: optimise
-				TSharedPtr<FImage> Temp1 = ImagePixelFormat(CompressionQuality, Base, EImageFormat::RGBA_UByte);
-				TSharedPtr<FImage> Temp2 = ImagePixelFormat(CompressionQuality, Temp1.Get(), EImageFormat::RGB_UByte);
-				ReleaseImage(Temp1);
+				TManagedPtr<FImage> Temp1 = ImagePixelFormat(CompressionQuality, Base, EImageFormat::RGBA_UByte);
+				TManagedPtr<FImage> Temp2 = ImagePixelFormat(CompressionQuality, Temp1.Get(), EImageFormat::RGB_UByte);
 				CompressRLE_RGB(Temp2.Get(), Result);
-				ReleaseImage(Temp2);
 
 				break;
 			}
@@ -1109,19 +1096,17 @@ void FImageOperator::ImagePixelFormat(bool& bOutSuccess, int32 CompressionQualit
 			}
 			case EImageFormat::L_UByte:
 			{
-				TSharedPtr<FImage> TempBase = ImagePixelFormat(CompressionQuality, Base, EImageFormat::RGB_UByte);
+				TManagedPtr<FImage> TempBase = ImagePixelFormat(CompressionQuality, Base, EImageFormat::RGB_UByte);
 				ImagePixelFormatInternal::BlockCompressionLODRangeHelper(
 						Result, TempBase.Get(), ResultLODBegin, BaseLODBegin, NumLODs, CompressionQuality, miro::RGB_to_BC1);
-				ReleaseImage(TempBase);
 
 				break;
 			}
 			case EImageFormat::BC3:
 			{
-				TSharedPtr<FImage> TempBase = ImagePixelFormat(CompressionQuality, Base, EImageFormat::RGBA_UByte);
+				TManagedPtr<FImage> TempBase = ImagePixelFormat(CompressionQuality, Base, EImageFormat::RGBA_UByte);
 				ImagePixelFormatInternal::BlockCompressionLODRangeHelper(
 						Result, TempBase.Get(), ResultLODBegin, BaseLODBegin, NumLODs, CompressionQuality, miro::RGBA_to_BC1);
-				ReleaseImage(TempBase);
 
 				break;
 			}
@@ -1130,10 +1115,9 @@ void FImageOperator::ImagePixelFormat(bool& bOutSuccess, int32 CompressionQualit
 				UE_LOG(LogMutableCore, Log, TEXT("Image format conversion not implemented. Expensive generic one used. "));
 
 				// Case not implemented, try a generic more expensive operation.
-				TSharedPtr<FImage> TempBase = ImagePixelFormat(CompressionQuality, Base, EImageFormat::RGBA_UByte);
+				TManagedPtr<FImage> TempBase = ImagePixelFormat(CompressionQuality, Base, EImageFormat::RGBA_UByte);
 				ImagePixelFormatInternal::BlockCompressionLODRangeHelper(
 						Result, TempBase.Get(), ResultLODBegin, BaseLODBegin, NumLODs, CompressionQuality, miro::RGBA_to_BC1);
-				ReleaseImage(TempBase);
 			}
 
 			}
@@ -1156,10 +1140,9 @@ void FImageOperator::ImagePixelFormat(bool& bOutSuccess, int32 CompressionQualit
 				UE_LOG(LogMutableCore, Log, TEXT("Image format conversion not implemented. Expensive generic one used."));
 
 				// Case not implemented, try a generic more expensive operatino.
-				TSharedPtr<FImage> TempBase = ImagePixelFormat(CompressionQuality, Base, EImageFormat::RGBA_UByte);
+				TManagedPtr<FImage> TempBase = ImagePixelFormat(CompressionQuality, Base, EImageFormat::RGBA_UByte);
 				ImagePixelFormatInternal::BlockCompressionLODRangeHelper(
 						Result, TempBase.Get(), ResultLODBegin, BaseLODBegin, NumLODs, CompressionQuality, miro::RGBA_to_BC2);
-				ReleaseImage(TempBase);
 			}
 
 			}
@@ -1195,10 +1178,9 @@ void FImageOperator::ImagePixelFormat(bool& bOutSuccess, int32 CompressionQualit
 				UE_LOG(LogMutableCore, Log, TEXT("Image format conversion not implemented. Expensive generic one used."));
 
 				// Case not implemented, try a generic more expensive operation.
-				TSharedPtr<FImage> TempBase = ImagePixelFormat(CompressionQuality, Base, EImageFormat::RGBA_UByte);
+				TManagedPtr<FImage> TempBase = ImagePixelFormat(CompressionQuality, Base, EImageFormat::RGBA_UByte);
 				ImagePixelFormatInternal::BlockCompressionLODRangeHelper(
 						Result, TempBase.Get(), ResultLODBegin, BaseLODBegin, NumLODs, CompressionQuality, miro::RGBA_to_BC3);
-				ReleaseImage(TempBase);
 			}
 
 			}
@@ -1221,10 +1203,9 @@ void FImageOperator::ImagePixelFormat(bool& bOutSuccess, int32 CompressionQualit
 			case EImageFormat::BGRA_UByte:
 			default:
 			{	
-				TSharedPtr<FImage> TempBase = ImagePixelFormat(CompressionQuality, Base, EImageFormat::L_UByte, -1);
+				TManagedPtr<FImage> TempBase = ImagePixelFormat(CompressionQuality, Base, EImageFormat::L_UByte, -1);
 				ImagePixelFormatInternal::BlockCompressionLODRangeHelper(
 						Result, TempBase.Get(), ResultLODBegin, BaseLODBegin, NumLODs, CompressionQuality, miro::L_to_BC4);
-				ReleaseImage(TempBase);
 				break;
 			}
 			}
@@ -1252,10 +1233,9 @@ void FImageOperator::ImagePixelFormat(bool& bOutSuccess, int32 CompressionQualit
 			{
 				UE_LOG(LogMutableCore, Log, TEXT("Image format conversion not implemented. Expensive generic one used. "));
 
-				TSharedPtr<FImage> TempBase = ImagePixelFormat(CompressionQuality, Base, EImageFormat::RGBA_UByte, -1);
+				TManagedPtr<FImage> TempBase = ImagePixelFormat(CompressionQuality, Base, EImageFormat::RGBA_UByte, -1);
 				ImagePixelFormatInternal::BlockCompressionLODRangeHelper(
 						Result, TempBase.Get(), ResultLODBegin, BaseLODBegin, NumLODs, CompressionQuality, miro::RGBA_to_BC5);
-				ReleaseImage(TempBase);
 				break;
 			}
 			}
@@ -1295,10 +1275,9 @@ void FImageOperator::ImagePixelFormat(bool& bOutSuccess, int32 CompressionQualit
 			{
 				UE_LOG(LogMutableCore, Log, TEXT("Image format conversion not implemented. Expensive generic one used. "));
 
-				TSharedPtr<FImage> TempBase = ImagePixelFormat(CompressionQuality, Base, EImageFormat::RGBA_UByte, -1);
+				TManagedPtr<FImage> TempBase = ImagePixelFormat(CompressionQuality, Base, EImageFormat::RGBA_UByte, -1);
 				ImagePixelFormatInternal::BlockCompressionLODRangeHelper(
 						Result, TempBase.Get(), ResultLODBegin, BaseLODBegin, NumLODs, CompressionQuality, miro::RGBA_to_ASTC4x4RGBL);
-				ReleaseImage(TempBase);
 				break;
 			}
 			}
@@ -1332,10 +1311,9 @@ void FImageOperator::ImagePixelFormat(bool& bOutSuccess, int32 CompressionQualit
 			{
 				UE_LOG(LogMutableCore, Log, TEXT("Image format conversion not implemented. Expensive generic one used. "));
 
-				TSharedPtr<FImage> TempBase = ImagePixelFormat(CompressionQuality, Base, EImageFormat::RGBA_UByte, -1);
+				TManagedPtr<FImage> TempBase = ImagePixelFormat(CompressionQuality, Base, EImageFormat::RGBA_UByte, -1);
 				ImagePixelFormatInternal::BlockCompressionLODRangeHelper(
 						Result, TempBase.Get(), ResultLODBegin, BaseLODBegin, NumLODs, CompressionQuality, miro::RGBA_to_ASTC4x4RGBAL);
-				ReleaseImage(TempBase);
 				break;
 			}
 			}
@@ -1370,10 +1348,9 @@ void FImageOperator::ImagePixelFormat(bool& bOutSuccess, int32 CompressionQualit
 			{
 				UE_LOG(LogMutableCore, Log, TEXT("Image format conversion not implemented. Expensive generic one used. "));
 
-				TSharedPtr<FImage> TempBase = ImagePixelFormat(CompressionQuality, Base, EImageFormat::RGBA_UByte, -1);
+				TManagedPtr<FImage> TempBase = ImagePixelFormat(CompressionQuality, Base, EImageFormat::RGBA_UByte, -1);
 				ImagePixelFormatInternal::BlockCompressionLODRangeHelper(
 						Result, TempBase.Get(), ResultLODBegin, BaseLODBegin, NumLODs, CompressionQuality, miro::RGBA_to_ASTC4x4RGL);
-				ReleaseImage(TempBase);
 				break;
 			}
 			}
@@ -1414,10 +1391,9 @@ void FImageOperator::ImagePixelFormat(bool& bOutSuccess, int32 CompressionQualit
 			{
 				UE_LOG(LogMutableCore, Log, TEXT("Image format conversion not implemented. Expensive generic one used. "));
 
-				TSharedPtr<FImage> TempBase = ImagePixelFormat(CompressionQuality, Base, EImageFormat::RGBA_UByte, -1);
+				TManagedPtr<FImage> TempBase = ImagePixelFormat(CompressionQuality, Base, EImageFormat::RGBA_UByte, -1);
 				ImagePixelFormatInternal::BlockCompressionLODRangeHelper(
 						Result, TempBase.Get(), ResultLODBegin, BaseLODBegin, NumLODs, CompressionQuality, miro::RGBA_to_ASTC6x6RGBL);
-				ReleaseImage(TempBase);
 				break;
 			}
 			}
@@ -1451,10 +1427,9 @@ void FImageOperator::ImagePixelFormat(bool& bOutSuccess, int32 CompressionQualit
 			{
 				UE_LOG(LogMutableCore, Log, TEXT("Image format conversion not implemented. Expensive generic one used. "));
 
-				TSharedPtr<FImage> TempBase = ImagePixelFormat(CompressionQuality, Base, EImageFormat::RGBA_UByte, -1);
+				TManagedPtr<FImage> TempBase = ImagePixelFormat(CompressionQuality, Base, EImageFormat::RGBA_UByte, -1);
 				ImagePixelFormatInternal::BlockCompressionLODRangeHelper(
 						Result, TempBase.Get(), ResultLODBegin, BaseLODBegin, NumLODs, CompressionQuality, miro::RGBA_to_ASTC6x6RGBAL);
-				ReleaseImage(TempBase);
 				break;
 			}
 			}
@@ -1490,10 +1465,9 @@ void FImageOperator::ImagePixelFormat(bool& bOutSuccess, int32 CompressionQualit
 			{
 				UE_LOG(LogMutableCore, Log, TEXT("Image format conversion not implemented. Expensive generic one used. "));
 
-				TSharedPtr<FImage> TempBase = ImagePixelFormat(CompressionQuality, Base, EImageFormat::RGBA_UByte, -1);
+				TManagedPtr<FImage> TempBase = ImagePixelFormat(CompressionQuality, Base, EImageFormat::RGBA_UByte, -1);
 				ImagePixelFormatInternal::BlockCompressionLODRangeHelper(
 						Result, TempBase.Get(), ResultLODBegin, BaseLODBegin, NumLODs, CompressionQuality, miro::RGBA_to_ASTC6x6RGL);
-				ReleaseImage(TempBase);
 				break;
 			}
 			}
@@ -1533,10 +1507,9 @@ void FImageOperator::ImagePixelFormat(bool& bOutSuccess, int32 CompressionQualit
 			{
 				UE_LOG(LogMutableCore, Log, TEXT("Image format conversion not implemented. Expensive generic one used. "));
 
-				TSharedPtr<FImage> TempBase = ImagePixelFormat(CompressionQuality, Base, EImageFormat::RGBA_UByte, -1);
+				TManagedPtr<FImage> TempBase = ImagePixelFormat(CompressionQuality, Base, EImageFormat::RGBA_UByte, -1);
 				ImagePixelFormatInternal::BlockCompressionLODRangeHelper(
 						Result, TempBase.Get(), ResultLODBegin, BaseLODBegin, NumLODs, CompressionQuality, miro::RGBA_to_ASTC8x8RGBL);
-				ReleaseImage(TempBase);
 				break;
 			}
 			}
@@ -1570,10 +1543,9 @@ void FImageOperator::ImagePixelFormat(bool& bOutSuccess, int32 CompressionQualit
 			{
 				UE_LOG(LogMutableCore, Log, TEXT("Image format conversion not implemented. Expensive generic one used. "));
 
-				TSharedPtr<FImage> TempBase = ImagePixelFormat(CompressionQuality, Base, EImageFormat::RGBA_UByte, -1);
+				TManagedPtr<FImage> TempBase = ImagePixelFormat(CompressionQuality, Base, EImageFormat::RGBA_UByte, -1);
 				ImagePixelFormatInternal::BlockCompressionLODRangeHelper(
 						Result, TempBase.Get(), ResultLODBegin, BaseLODBegin, NumLODs, CompressionQuality, miro::RGBA_to_ASTC8x8RGBAL);
-				ReleaseImage(TempBase);
 				break;
 			}
 			}
@@ -1601,10 +1573,9 @@ void FImageOperator::ImagePixelFormat(bool& bOutSuccess, int32 CompressionQualit
 			{
 				UE_LOG(LogMutableCore, Log, TEXT("Image format conversion not implemented. Expensive generic one used. "));
 
-				TSharedPtr<FImage> TempBase = ImagePixelFormat(CompressionQuality, Base, EImageFormat::RGBA_UByte, -1);
+				TManagedPtr<FImage> TempBase = ImagePixelFormat(CompressionQuality, Base, EImageFormat::RGBA_UByte, -1);
 				ImagePixelFormatInternal::BlockCompressionLODRangeHelper(
 						Result, TempBase.Get(), ResultLODBegin, BaseLODBegin, NumLODs, CompressionQuality, miro::RGBA_to_ASTC8x8RGL);
-				ReleaseImage(TempBase);
 				break;
 			}
 			}
diff --git a/Engine/Plugins/Mutable/Source/MutableRuntime/Private/MuR/OpImageResize.cpp b/Engine/Plugins/Mutable/Source/MutableRuntime/Private/MuR/OpImageResize.cpp
--- a/Engine/Plugins/Mutable/Source/MutableRuntime/Private/MuR/OpImageResize.cpp
+++ b/Engine/Plugins/Mutable/Source/MutableRuntime/Private/MuR/OpImageResize.cpp
@@ -907,7 +907,7 @@ namespace
 			const EImageFormat BaseUncompressedFormat = GetUncompressedFormat(InBase->GetFormat());
 			const EImageFormat DestFormat = Dest->GetFormat();
 			
-			TSharedPtr<FImage> TempImage;
+			TManagedPtr<FImage> TempImage;
 			FImage* DecompressedDest = Dest;
 			if (DestFormat != BaseUncompressedFormat)
 			{
@@ -931,14 +931,12 @@ namespace
 				bool bSuccess = false;
 				ImagePixelFormat(bSuccess, ImageCompressionQuality, Dest, DecompressedDest);
 				check(bSuccess);
-
-				ReleaseImage(TempImage);
 			}
 		}
 		else
 		{
 			const FImage* Base = InBase;
-			TSharedPtr<FImage> BaseReformat;
+			TManagedPtr<FImage> BaseReformat;
 			
 			// Shouldn't happen! But if it does...
 			const EImageFormat SourceFormat = Base->GetFormat();
@@ -952,7 +950,7 @@ namespace
 			const EImageFormat BaseFormat = Base->GetFormat();
 
 			// First resize X
-			TSharedPtr<FImage> Temp;
+			TManagedPtr<FImage> Temp;
 			if (DestSize[0] > BaseSize[0])
 			{
 				Temp = CreateImage(DestSize[0], BaseSize[1], 1, BaseFormat, EInitializationType::NotInitialized);
@@ -973,24 +971,21 @@ namespace
 			if (BaseReformat)
 			{
 				check(Base == BaseReformat.Get());
-				ReleaseImage(BaseReformat);
 				BaseReformat = nullptr;
 				Base = nullptr;
 			}
 
 			// Now resize Y
-			TSharedPtr<FImage> Temp2;
+			TManagedPtr<FImage> Temp2;
 			if (DestSize[1] > BaseSize[1])
 			{
 				Temp2 = CreateImage(DestSize[0], DestSize[1], 1, BaseFormat, EInitializationType::NotInitialized);
 				ImageMagnifyY(Temp2.Get(), Temp.Get());
-				ReleaseImage(Temp);
 			}
 			else if (DestSize[1] < BaseSize[1])
 			{
 				Temp2 = CreateImage(DestSize[0], DestSize[1], 1, BaseFormat, EInitializationType::NotInitialized);
 				ImageMinifyY(Temp2.Get(), Temp.Get());
-				ReleaseImage(Temp);
 			}
 			else
 			{
@@ -1002,13 +997,11 @@ namespace
 			// Reset format if it was changed to scale
 			if (SourceFormat != UncompressedFormat)
 			{
-				TSharedPtr<FImage> OldTemp = Temp2;
+				TManagedPtr<FImage> OldTemp = Temp2;
 				Temp2 = ImagePixelFormat(ImageCompressionQuality, Temp2.Get(), SourceFormat);
-				ReleaseImage(OldTemp);
 			}
 
 			Dest->CopyMove(Temp2.Get());
-			ReleaseImage(Temp2);
 			Temp2 = nullptr;
 		}
 		
diff --git a/Engine/Plugins/Mutable/Source/MutableRuntime/Private/MuR/OpImageSwizzle.cpp b/Engine/Plugins/Mutable/Source/MutableRuntime/Private/MuR/OpImageSwizzle.cpp
--- a/Engine/Plugins/Mutable/Source/MutableRuntime/Private/MuR/OpImageSwizzle.cpp
+++ b/Engine/Plugins/Mutable/Source/MutableRuntime/Private/MuR/OpImageSwizzle.cpp
@@ -6,7 +6,7 @@
 
 namespace UE::Mutable::Private
 {
-	void ImageSwizzle(FImage* Result, const TSharedPtr<const FImage> Sources[], const uint8 Channels[])
+	void ImageSwizzle(FImage* Result, const TManagedPtr<const FImage> Sources[], const uint8 Channels[])
 	{
 		MUTABLE_CPUPROFILER_SCOPE(ImageSwizzle);
 
@@ -194,7 +194,7 @@ namespace UE::Mutable::Private
 		}
 	}
 
-	TSharedPtr<FImage> FImageOperator::ImageSwizzle(EImageFormat Format, const TSharedPtr<const FImage> Sources[], const uint8 Channels[])
+	TManagedPtr<FImage> FImageOperator::ImageSwizzle(EImageFormat Format, const TManagedPtr<const FImage> Sources[], const uint8 Channels[])
 	{
 		MUTABLE_CPUPROFILER_SCOPE(ImageSwizzle);
 
@@ -216,7 +216,7 @@ namespace UE::Mutable::Private
 		const FImageSize ResultSize = Sources[FirstValidSourceIndex]->GetSize();
 		const int32 ResultNumLODs = Sources[FirstValidSourceIndex]->GetLODCount();
 
-		TSharedPtr<FImage> Dest = CreateImage(ResultSize.X, ResultSize.Y, ResultNumLODs, Format, EInitializationType::Black);
+		TManagedPtr<FImage> Dest = CreateImage(ResultSize.X, ResultSize.Y, ResultNumLODs, Format, EInitializationType::Black);
 
 		UE::Mutable::Private::ImageSwizzle(Dest.Get(), Sources, Channels);
 
diff --git a/Engine/Plugins/Mutable/Source/MutableRuntime/Private/MuR/OpLayoutRemoveBlocks.h b/Engine/Plugins/Mutable/Source/MutableRuntime/Private/MuR/OpLayoutRemoveBlocks.h
--- a/Engine/Plugins/Mutable/Source/MutableRuntime/Private/MuR/OpLayoutRemoveBlocks.h
+++ b/Engine/Plugins/Mutable/Source/MutableRuntime/Private/MuR/OpLayoutRemoveBlocks.h
@@ -11,15 +11,15 @@ namespace UE::Mutable::Private
 {
 
 	//---------------------------------------------------------------------------------------------
-	inline TSharedPtr<FLayout> LayoutFromMesh_RemoveBlocks(const FMesh* InMesh, int32 InLayoutIndex)
+	inline TManagedPtr<FLayout> LayoutFromMesh_RemoveBlocks(const FMesh* InMesh, int32 InLayoutIndex)
 	{
 		if (!InMesh || InMesh->GetLayoutCount() <= InLayoutIndex)
 		{
 			return nullptr;
 		}
 
-		TSharedPtr<const FLayout> Source = InMesh->GetLayout(InLayoutIndex);
-		TSharedPtr<FLayout> pResult;
+		TManagedPtr<const FLayout> Source = InMesh->GetLayout(InLayoutIndex);
+		TManagedPtr<FLayout> pResult;
 
 		UntypedMeshBufferIteratorConst itBlocks(InMesh->GetVertexBuffers(), EMeshBufferSemantic::LayoutBlock, InLayoutIndex);
 		if (itBlocks.GetFormat() == EMeshBufferFormat::UInt16)
@@ -117,11 +117,10 @@ namespace UE::Mutable::Private
 	}
 
 
-	//---------------------------------------------------------------------------------------------
-	inline TSharedPtr<FLayout> LayoutRemoveBlocks(const FLayout* Source, const FLayout* ReferenceLayout)
+	inline TManagedPtr<FLayout> LayoutRemoveBlocks(const FLayout* Source, const FLayout* ReferenceLayout)
 	{
 		// Remove blocks that are not in the mesh
-		TSharedPtr<FLayout> pResult = Source->Clone();
+		TManagedPtr<FLayout> pResult = Source->Clone();
 		int32 DestBlockIndex = 0;
 		for (int32 BlockIndexInSource = 0; BlockIndexInSource < pResult->Blocks.Num(); ++BlockIndexInSource)
 		{
@@ -140,10 +139,9 @@ namespace UE::Mutable::Private
 	}
 
 
-	//---------------------------------------------------------------------------------------------
-	inline TSharedPtr<FLayout> LayoutMerge(const FLayout* pA, const FLayout* pB )
+	inline TManagedPtr<FLayout> LayoutMerge(const FLayout* pA, const FLayout* pB )
 	{
-		TSharedPtr<FLayout> pResult = pA->Clone();
+		TManagedPtr<FLayout> pResult = pA->Clone();
 
 		for ( const FLayoutBlock& Block: pB->Blocks )
 		{
diff --git a/Engine/Plugins/Mutable/Source/MutableRuntime/Private/MuR/OpMeshApplyPose.h b/Engine/Plugins/Mutable/Source/MutableRuntime/Private/MuR/OpMeshApplyPose.h
--- a/Engine/Plugins/Mutable/Source/MutableRuntime/Private/MuR/OpMeshApplyPose.h
+++ b/Engine/Plugins/Mutable/Source/MutableRuntime/Private/MuR/OpMeshApplyPose.h
@@ -53,7 +53,7 @@ namespace UE::Mutable::Private
 
 		bOutSuccess = true;
 
-		TSharedPtr<const FSkeleton> Skeleton = BaseMesh->GetSkeleton();
+		TManagedPtr<const FSkeleton> Skeleton = BaseMesh->GetSkeleton();
 		if (!Skeleton)
 		{
 			bOutSuccess = false;
diff --git a/Engine/Plugins/Mutable/Source/MutableRuntime/Private/MuR/OpMeshBind.h b/Engine/Plugins/Mutable/Source/MutableRuntime/Private/MuR/OpMeshBind.h
--- a/Engine/Plugins/Mutable/Source/MutableRuntime/Private/MuR/OpMeshBind.h
+++ b/Engine/Plugins/Mutable/Source/MutableRuntime/Private/MuR/OpMeshBind.h
@@ -1290,7 +1290,7 @@ namespace UE::Mutable::Private
 				PhysicsBodiesToBind.Reserve(Result->AdditionalPhysicsBodies.Num() + 1);
 
 				PhysicsBodiesToBind.Add(ResultPhysicsBody);
-				for (const TSharedPtr<const FPhysicsBody>& Body : Result->AdditionalPhysicsBodies)
+				for (const TManagedPtr<const FPhysicsBody>& Body : Result->AdditionalPhysicsBodies)
 				{
 					PhysicsBodiesToBind.Add(Body.Get());
 				}
diff --git a/Engine/Plugins/Mutable/Source/MutableRuntime/Private/MuR/OpMeshClipMorphPlane.h b/Engine/Plugins/Mutable/Source/MutableRuntime/Private/MuR/OpMeshClipMorphPlane.h
--- a/Engine/Plugins/Mutable/Source/MutableRuntime/Private/MuR/OpMeshClipMorphPlane.h
+++ b/Engine/Plugins/Mutable/Source/MutableRuntime/Private/MuR/OpMeshClipMorphPlane.h
@@ -101,7 +101,7 @@ namespace UE::Mutable::Private
 
 		TArray<FVertexBoneInfo> VertexInfo;
 
-        TSharedPtr<const FSkeleton> BaseSkeleton = pBase->GetSkeleton();
+        TManagedPtr<const FSkeleton> BaseSkeleton = pBase->GetSkeleton();
 
 		TBitArray<> AffectedBoneMapIndices;
 		const int32 BaseBoneIndex = BaseSkeleton && BoneId ? BaseSkeleton->FindBone(*BoneId) : INDEX_NONE;
diff --git a/Engine/Plugins/Mutable/Source/MutableRuntime/Private/MuR/OpMeshFormat.cpp b/Engine/Plugins/Mutable/Source/MutableRuntime/Private/MuR/OpMeshFormat.cpp
--- a/Engine/Plugins/Mutable/Source/MutableRuntime/Private/MuR/OpMeshFormat.cpp
+++ b/Engine/Plugins/Mutable/Source/MutableRuntime/Private/MuR/OpMeshFormat.cpp
@@ -698,7 +698,7 @@ namespace UE::Mutable::Private
 		Result->SetPhysicsBody(Source->GetPhysicsBody());
 
 		Result->Layouts.Empty();
-		for (const TSharedPtr<const FLayout>& Layout : Source->Layouts)
+		for (const TManagedPtr<const FLayout>& Layout : Source->Layouts)
 		{
 			Result->Layouts.Add(Layout->Clone());
 		}
diff --git a/Engine/Plugins/Mutable/Source/MutableRuntime/Private/MuR/OpMeshPrepareLayout.cpp b/Engine/Plugins/Mutable/Source/MutableRuntime/Private/MuR/OpMeshPrepareLayout.cpp
--- a/Engine/Plugins/Mutable/Source/MutableRuntime/Private/MuR/OpMeshPrepareLayout.cpp
+++ b/Engine/Plugins/Mutable/Source/MutableRuntime/Private/MuR/OpMeshPrepareLayout.cpp
@@ -510,7 +510,7 @@ namespace UE::Mutable::Private
 			return;
 		}
 
-		TSharedPtr<FLayout> Layout = InLayout.Clone();
+		TManagedPtr<FLayout> Layout = InLayout.Clone();
 
 		// The layout must have block ids.
 		check(Layout->Blocks.IsEmpty() || Layout->Blocks[0].Id != FLayoutBlock::InvalidBlockId);
diff --git a/Engine/Plugins/Mutable/Source/MutableRuntime/Private/MuR/OpMeshReshape.h b/Engine/Plugins/Mutable/Source/MutableRuntime/Private/MuR/OpMeshReshape.h
--- a/Engine/Plugins/Mutable/Source/MutableRuntime/Private/MuR/OpMeshReshape.h
+++ b/Engine/Plugins/Mutable/Source/MutableRuntime/Private/MuR/OpMeshReshape.h
@@ -732,7 +732,7 @@ namespace UE::Mutable::Private
 		// Apply main physics body
 		if (BaseMesh.PhysicsBody)
 		{
-			TSharedPtr<FPhysicsBody> NewBody = BaseMesh.PhysicsBody->Clone();
+			TManagedPtr<FPhysicsBody> NewBody = BaseMesh.PhysicsBody->Clone();
 			ApplyPhysicsBody(*NewBody, 0, BodyOffsets[1]);
 			OutNewMesh.PhysicsBody = NewBody;
 		}
@@ -741,7 +741,7 @@ namespace UE::Mutable::Private
 		// Apply additional physics bodies
 		for (int32 I = 1; I < PhysicsBodiesNum; ++I)
 		{	
-			TSharedPtr<FPhysicsBody> NewBody = BaseMesh.AdditionalPhysicsBodies[I - 1]->Clone();
+			TManagedPtr<FPhysicsBody> NewBody = BaseMesh.AdditionalPhysicsBodies[I - 1]->Clone();
 			ApplyPhysicsBody(*NewBody, BodyOffsets[I], BodyOffsets[I + 1]);
 			OutNewMesh.AdditionalPhysicsBodies[I - 1] = NewBody;
 		}
@@ -790,12 +790,12 @@ namespace UE::Mutable::Private
 		const FMeshBufferSet& VB = BaseMesh->GetVertexBuffers();
 		VB.FindChannel(EMeshBufferSemantic::BarycentricCoords, BindingDataIndex, &BarycentricDataBuffer, &BarycentricDataChannel);
 		
-		TSharedPtr<FMesh> TempMesh = nullptr;
+		TManagedPtr<FMesh> TempMesh = nullptr;
 		FMesh* VerticesReshapeMesh = Result;
 		
 		if (bApplyLaplacian)
 		{
-			TempMesh = MakeShared<FMesh>();
+			TempMesh = MakeManaged<FMesh>();
 			VerticesReshapeMesh = TempMesh.Get();
 		}
 
diff --git a/Engine/Plugins/Mutable/Source/MutableRuntime/Private/MuR/OpMeshTransformWithBone.cpp b/Engine/Plugins/Mutable/Source/MutableRuntime/Private/MuR/OpMeshTransformWithBone.cpp
--- a/Engine/Plugins/Mutable/Source/MutableRuntime/Private/MuR/OpMeshTransformWithBone.cpp
+++ b/Engine/Plugins/Mutable/Source/MutableRuntime/Private/MuR/OpMeshTransformWithBone.cpp
@@ -24,7 +24,7 @@ void MeshTransformWithBoneInline(FMesh* Mesh, const FMatrix44f& Transform, const
 	}
 
 
-	TSharedPtr<const FSkeleton> BaseSkeleton = Mesh->GetSkeleton();
+	TManagedPtr<const FSkeleton> BaseSkeleton = Mesh->GetSkeleton();
 
 	const int32 BaseBoneIndex = BaseSkeleton ? BaseSkeleton->FindBone(InBoneName.Id) : INDEX_NONE;
 	if (BaseBoneIndex == INDEX_NONE)
diff --git a/Engine/Plugins/Mutable/Source/MutableRuntime/Private/MuR/Parameters.cpp b/Engine/Plugins/Mutable/Source/MutableRuntime/Private/MuR/Parameters.cpp
--- a/Engine/Plugins/Mutable/Source/MutableRuntime/Private/MuR/Parameters.cpp
+++ b/Engine/Plugins/Mutable/Source/MutableRuntime/Private/MuR/Parameters.cpp
@@ -972,7 +972,7 @@ namespace UE::Mutable::Private
 	}
 
 
-	TSharedPtr<const FInstancedStruct> FParameters::GetInstancedStructValue(int32 Index, const FRangeIndex* RangeIndex) const
+	TManagedPtr<const FInstancedStruct> FParameters::GetInstancedStructValue(int32 Index, const FRangeIndex* RangeIndex) const
 	{
     	check(Index >= 0 && Index < (int32)m_pD->Values.Num());
     	check(GetType(Index) == EParameterType::InstancedStruct);
@@ -1009,7 +1009,7 @@ namespace UE::Mutable::Private
 	}
 
 
-	void FParameters::SetInstancedStructValue(int32 Index, const TSharedPtr<const FInstancedStruct>& Value, const FRangeIndex* RangeIndex)
+	void FParameters::SetInstancedStructValue(int32 Index, const TManagedPtr<const FInstancedStruct>& Value, const FRangeIndex* RangeIndex)
 	{
     	LLM_SCOPE_BYNAME(TEXT("MutableRuntime"));
 
diff --git a/Engine/Plugins/Mutable/Source/MutableRuntime/Private/MuR/PhysicsBody.cpp b/Engine/Plugins/Mutable/Source/MutableRuntime/Private/MuR/PhysicsBody.cpp
--- a/Engine/Plugins/Mutable/Source/MutableRuntime/Private/MuR/PhysicsBody.cpp
+++ b/Engine/Plugins/Mutable/Source/MutableRuntime/Private/MuR/PhysicsBody.cpp
@@ -14,11 +14,11 @@ namespace UE::Mutable::Private
 	}
 
 
-	TSharedPtr<FPhysicsBody> FPhysicsBody::StaticUnserialise(FInputArchive& Arch)
+	TManagedPtr<FPhysicsBody> FPhysicsBody::StaticUnserialise(FInputArchive& Arch)
 	{
-		TSharedPtr<FPhysicsBody> pResult = MakeShared<FPhysicsBody>();
-		Arch >> *pResult;
-		return pResult;
+		TManagedPtr<FPhysicsBody> Result = MakeManaged<FPhysicsBody>();
+		Arch >> *Result;
+		return Result;
 	}
 
 
@@ -160,20 +160,20 @@ namespace UE::Mutable::Private
 	}
 
 
-	TSharedPtr<FPhysicsBody> FPhysicsBody::Clone() const
+	TManagedPtr<FPhysicsBody> FPhysicsBody::Clone() const
 	{
-		TSharedPtr<FPhysicsBody> pResult = MakeShared<FPhysicsBody>();
+		TManagedPtr<FPhysicsBody> Result = MakeManaged<FPhysicsBody>();
 
-		pResult->CustomId = CustomId;
+		Result->CustomId = CustomId;
 
-		pResult->Bodies = Bodies;
-		pResult->BoneIds = BoneIds;
-		pResult->BodiesCustomIds = BodiesCustomIds;
+		Result->Bodies = Bodies;
+		Result->BoneIds = BoneIds;
+		Result->BodiesCustomIds = BodiesCustomIds;
 
-		pResult->bBodiesModified = bBodiesModified;
+		Result->bBodiesModified = bBodiesModified;
 
 
-		return pResult;
+		return Result;
 	}
 
 	void FPhysicsBody::SetCustomId(int32 InCustomId)
diff --git a/Engine/Plugins/Mutable/Source/MutableRuntime/Private/MuR/Serialisation.cpp b/Engine/Plugins/Mutable/Source/MutableRuntime/Private/MuR/Serialisation.cpp
--- a/Engine/Plugins/Mutable/Source/MutableRuntime/Private/MuR/Serialisation.cpp
+++ b/Engine/Plugins/Mutable/Source/MutableRuntime/Private/MuR/Serialisation.cpp
@@ -96,13 +96,13 @@ namespace UE::Mutable::Private
     }
 
 
-	void operator<<(FOutputArchive& Arch, const TSharedPtr<const FInstancedStruct>& V)
+	void operator<<(FOutputArchive& Arch, const TManagedPtr<const FInstancedStruct>& V)
 	{
 		// Do Nothing. UStructs values are saved in the CO.	
 	}
 
 
-	void operator>>(FInputArchive& Arch, TSharedPtr<const FInstancedStruct>& V)
+	void operator>>(FInputArchive& Arch, TManagedPtr<const FInstancedStruct>& V)
 	{
 		// Do Nothing. UStructs values are saved in the CO.	
 	}
diff --git a/Engine/Plugins/Mutable/Source/MutableRuntime/Private/MuR/Skeleton.cpp b/Engine/Plugins/Mutable/Source/MutableRuntime/Private/MuR/Skeleton.cpp
--- a/Engine/Plugins/Mutable/Source/MutableRuntime/Private/MuR/Skeleton.cpp
+++ b/Engine/Plugins/Mutable/Source/MutableRuntime/Private/MuR/Skeleton.cpp
@@ -15,17 +15,17 @@ namespace UE::Mutable::Private
 	}
 
 
-	TSharedPtr<FSkeleton> FSkeleton::StaticUnserialise(FInputArchive& Arch)
+	TManagedPtr<FSkeleton> FSkeleton::StaticUnserialise(FInputArchive& Arch)
 	{
-		TSharedPtr<FSkeleton> Result = MakeShared<FSkeleton>();
+		TManagedPtr<FSkeleton> Result = MakeManaged<FSkeleton>();
 		Arch >> *Result;
 		return Result;
 	}
 
 
-	TSharedPtr<FSkeleton> FSkeleton::Clone() const
+	TManagedPtr<FSkeleton> FSkeleton::Clone() const
 	{
-		TSharedPtr<FSkeleton> Result = MakeShared<FSkeleton>();
+		TManagedPtr<FSkeleton> Result = MakeManaged<FSkeleton>();
 
 		Result->BoneIds = BoneIds;
 		Result->BoneParents = BoneParents;
diff --git a/Engine/Plugins/Mutable/Source/MutableRuntime/Private/MuR/System.cpp b/Engine/Plugins/Mutable/Source/MutableRuntime/Private/MuR/System.cpp
--- a/Engine/Plugins/Mutable/Source/MutableRuntime/Private/MuR/System.cpp
+++ b/Engine/Plugins/Mutable/Source/MutableRuntime/Private/MuR/System.cpp
@@ -234,7 +234,7 @@ namespace UE::Mutable::Private
 
 
     //---------------------------------------------------------------------------------------------
-    TSharedPtr<const FInstance> FSystem::BeginUpdate_GameThread(FInstance::FID InInstanceID,
+    TManagedPtr<const FInstance> FSystem::BeginUpdate_GameThread(FInstance::FID InInstanceID,
                                      const TSharedPtr<const FParameters>& InParams,
                                      const TSharedPtr<FMeshIdRegistry>& InMeshIdRegistry,
 									 const TSharedPtr<FImageIdRegistry>& InImageIdRegistry,
@@ -272,11 +272,11 @@ namespace UE::Mutable::Private
 
     	CodeRunnerBegin Runner(*this, LiveInstance->Model, InParams, InLodMask);
     	
-    	TSharedPtr<const FInstance> Result = Runner.RunCode(FScheduledOpInline(RootAt, 0));
+    	TManagedPtr<const FInstance> Result = Runner.RunCode(FScheduledOpInline(RootAt, 0));
     	if (!Result)
 		{
 			// In case of failure return an empty instance, to prevent following code to have to check it every time
-			Result = MakeShared<FInstance>();
+			Result = MakeManaged<FInstance>();
 		}
 
     	m_pD->WorkingMemoryManager.CurrentInstanceCache = nullptr;
@@ -291,7 +291,7 @@ namespace UE::Mutable::Private
 	}
 
 	
-    TSharedPtr<const FInstance> FSystem::BeginUpdate_MutableThread(FInstance::FID InInstanceID,
+    TManagedPtr<const FInstance> FSystem::BeginUpdate_MutableThread(FInstance::FID InInstanceID,
 		const TSharedPtr<const FParameters>& InParams,
 		const TSharedPtr<FMeshIdRegistry>& InMeshIdRegistry,
 		const TSharedPtr<FImageIdRegistry>& InImageIdRegistry,
@@ -367,7 +367,7 @@ namespace UE::Mutable::Private
 
 		m_pD->RunCode(LiveInstance->ExternalResourceProvider, LiveInstance->Model, InParams.Get(), RootAt, InLodMask);
 
-		TSharedPtr<const FInstance> Result = LiveInstance->Cache->GetInstance(FCacheAddress(RootAt, 0, 0));
+		TManagedPtr<const FInstance> Result = LiveInstance->Cache->GetInstance(FCacheAddress(RootAt, 0, 0));
 
 		// Debug check to see if we managed the op-hit-counts correctly
 		LiveInstance->Cache->CheckHitCountsCleared();
@@ -377,7 +377,7 @@ namespace UE::Mutable::Private
     	if (!Result)
 		{
 			// In case of failure return an empty instance, to prevent following code to have to check it every time
-			Result = MakeShared<FInstance>();
+			Result = MakeManaged<FInstance>();
 		}
 
 		m_pD->WorkingMemoryManager.EndRunnerThread();
@@ -394,13 +394,13 @@ namespace UE::Mutable::Private
 	}
 
 
-    TSharedRef<FImage> BlackImage()
+    TManagedRef<FImage> BlackImage()
     {
-    	return MakeShared<UE::Mutable::Private::FImage>(16, 16, 1, EImageFormat::RGBA_UByte, EInitializationType::Black);
+    	return MakeManaged<UE::Mutable::Private::FImage>(16, 16, 1, EImageFormat::RGBA_UByte, EInitializationType::Black);
     }
 	
 
-	TSharedPtr<const FImage> FSystem::GetImageInline(FInstance::FID instanceID, const FImageId& ImageId, int32 MipsToSkip)
+	TManagedPtr<const FImage> FSystem::GetImageInline(FInstance::FID instanceID, const FImageId& ImageId, int32 MipsToSkip)
 	{
 		LLM_SCOPE_BYNAME(TEXT("MutableRuntime"));
 		MUTABLE_CPUPROFILER_SCOPE(SystemGetImage);
@@ -410,8 +410,6 @@ namespace UE::Mutable::Private
 			return BlackImage();
     	}
     	
-		TSharedPtr<const FImage> pResult;
-
 		// Find the live instance
 		FLiveInstance* LiveInstance = m_pD->FindLiveInstance(instanceID);
 		check(LiveInstance);
@@ -419,21 +417,20 @@ namespace UE::Mutable::Private
 		m_pD->WorkingMemoryManager.CurrentPassthroughObjectLoader = LiveInstance->PassthroughObjectLoader;
 
 		OP::ADDRESS RootAddress = ImageId.GetKey()->Address;
-		pResult = m_pD->BuildImage(LiveInstance->ExternalResourceProvider, LiveInstance->Model, LiveInstance->OldParameters.Get(), RootAddress, MipsToSkip);
+		TManagedPtr<const FImage> Result = m_pD->BuildImage(LiveInstance->ExternalResourceProvider, LiveInstance->Model, LiveInstance->OldParameters.Get(), RootAddress, MipsToSkip);
 
-		// We always need to return something valid.
-		if (!pResult)
+		if (!Result)
 		{
-			pResult = BlackImage();
+			Result = BlackImage();
 		}
 
 		m_pD->WorkingMemoryManager.CurrentInstanceCache = nullptr;
     	m_pD->WorkingMemoryManager.CurrentPassthroughObjectLoader = nullptr;
     	
-		return pResult;
+		return Result;
 	}
 
-	UE::Tasks::TTask<TSharedPtr<const FImage>> FSystem::GetImage(UE::Tasks::FTask Prerequisite, FInstance::FID instanceID, const FImageId& ImageId, int32 MipsToSkip)
+	UE::Tasks::TTask<TManagedPtr<const FImage>> FSystem::GetImage(UE::Tasks::FTask Prerequisite, FInstance::FID instanceID, const FImageId& ImageId, int32 MipsToSkip)
 	{
 		LLM_SCOPE_BYNAME(TEXT("MutableRuntime"));
 		MUTABLE_CPUPROFILER_SCOPE(SystemGetImage);
@@ -441,9 +438,9 @@ namespace UE::Mutable::Private
     	if (!ImageId)
     	{
     		return UE::Tasks::Launch(TEXT("FSystem::GetImageInvalid"),
-			[]() -> TSharedPtr<const FImage>
+			[]() -> TManagedPtr<const FImage>
 			{
-				return MakeShared<FImage>();
+				return MakeManaged<FImage>();
 			},
 			UE::Tasks::Prerequisites(Prerequisite),
 			UE::Tasks::ETaskPriority::Inherit,
@@ -460,9 +457,9 @@ namespace UE::Mutable::Private
     	if (GetOpDataType(OpType) != EDataType::Image)
     	{
     		return UE::Tasks::Launch(TEXT("FSystem::GetImageInvalid"),
-			[]() -> TSharedPtr<const FImage>
+			[]() -> TManagedPtr<const FImage>
 			{
-				return MakeShared<FImage>();
+				return MakeManaged<FImage>();
 			},
 			UE::Tasks::Prerequisites(Prerequisite),
 			UE::Tasks::ETaskPriority::Inherit,
@@ -501,19 +498,19 @@ namespace UE::Mutable::Private
 		UE::Tasks::ETaskPriority::Inherit);
     	
 		return UE::Tasks::Launch(TEXT("FSystem::GetImageResultTask"),
-		[Context, SystemPrivate = m_pD, RootAddress, MipsToSkip]() -> TSharedPtr<const FImage>
+		[Context, SystemPrivate = m_pD, RootAddress, MipsToSkip]() -> TManagedPtr<const FImage>
 		{
-			TSharedPtr<const FImage> Result;
+			TManagedPtr<const FImage> Result;
 
 			SystemPrivate->bUnrecoverableError = Context->Runner->bUnrecoverableError;
 			if (!Context->Runner->bUnrecoverableError)
 			{
-				Result = SystemPrivate->WorkingMemoryManager.LoadImage(FCacheAddress(RootAddress, 0, MipsToSkip), true);
+				Result = SystemPrivate->WorkingMemoryManager.LoadImage(FCacheAddress(RootAddress, 0, MipsToSkip));
 			}
 
 			if (!Result)
 			{
-				Result = MakeShared<FImage>(16, 16, 1, EImageFormat::RGBA_UByte, EInitializationType::Black);
+				Result = MakeManaged<FImage>(16, 16, 1, EImageFormat::RGBA_UByte, EInitializationType::Black);
 			}
 
 			SystemPrivate->WorkingMemoryManager.EndRunnerThread();
@@ -696,14 +693,14 @@ namespace UE::Mutable::Private
 	}	
 
 	
-    TSharedPtr<const FMesh> FSystem::GetMeshInline(FInstance::FID instanceID, const FMeshId& MeshId, EMeshContentFlags MeshContentFilter)
+    TManagedPtr<const FMesh> FSystem::GetMeshInline(FInstance::FID instanceID, const FMeshId& MeshId, EMeshContentFlags MeshContentFilter)
     {
 		LLM_SCOPE_BYNAME(TEXT("MutableRuntime"));
 		MUTABLE_CPUPROFILER_SCOPE(SystemGetMesh);
 
 		if (!MeshId)
 		{
-			return MakeShared<FMesh>();
+			return MakeManaged<FMesh>();
 		}
 
 		// Find the live instance
@@ -713,7 +710,7 @@ namespace UE::Mutable::Private
 		m_pD->WorkingMemoryManager.CurrentPassthroughObjectLoader = LiveInstance->PassthroughObjectLoader;
 
 		OP::ADDRESS RootAddress = MeshId.GetKey()->Address;
-		TSharedPtr<const FMesh> Result = m_pD->BuildMesh(
+		TManagedPtr<const FMesh> Result = m_pD->BuildMesh(
 			LiveInstance->ExternalResourceProvider,
 			LiveInstance->Model,
 			LiveInstance->OldParameters.Get(),
@@ -723,7 +720,7 @@ namespace UE::Mutable::Private
 		// If the mesh is null it means empty, but we still need to return a valid one
 		if (!Result)
 		{
-			Result = MakeShared<FMesh>();
+			Result = MakeManaged<FMesh>();
 		}
 
 		m_pD->WorkingMemoryManager.CurrentInstanceCache = nullptr;
@@ -733,7 +730,7 @@ namespace UE::Mutable::Private
 	}
 	
 
-    UE::Tasks::TTask<TSharedPtr<const FMesh>> FSystem::GetMesh(UE::Tasks::FTask Prerequisite, FInstance::FID instanceID, const FMeshId& MeshId, EMeshContentFlags MeshContentFilter)
+    UE::Tasks::TTask<TManagedPtr<const FMesh>> FSystem::GetMesh(UE::Tasks::FTask Prerequisite, FInstance::FID instanceID, const FMeshId& MeshId, EMeshContentFlags MeshContentFilter)
 	{
 		LLM_SCOPE_BYNAME(TEXT("MutableRuntime"));
 		MUTABLE_CPUPROFILER_SCOPE(SystemGetMesh);
@@ -741,9 +738,9 @@ namespace UE::Mutable::Private
 		if (!MeshId)
 		{
 			return UE::Tasks::Launch(TEXT("FSystem::GetMeshInvalid"),
-			[]() -> TSharedPtr<const FMesh>
+			[]() -> TManagedPtr<const FMesh>
 			{
-				return MakeShared<FMesh>();
+				return MakeManaged<FMesh>();
 			},
 			UE::Tasks::Prerequisites(Prerequisite),
 			UE::Tasks::ETaskPriority::Inherit,
@@ -760,9 +757,9 @@ namespace UE::Mutable::Private
 		if (GetOpDataType(OpType) != EDataType::Mesh)
 		{
 			return UE::Tasks::Launch(TEXT("FSystem::GetMeshInvalid"),
-			[]() -> TSharedPtr<const FMesh>
+			[]() -> TManagedPtr<const FMesh>
 			{
-				return MakeShared<FMesh>();
+				return MakeManaged<FMesh>();
 			},
 			UE::Tasks::Prerequisites(Prerequisite),
 			UE::Tasks::ETaskPriority::Inherit,
@@ -803,14 +800,14 @@ namespace UE::Mutable::Private
 		UE::Tasks::ETaskPriority::Inherit);
 		
 		return UE::Tasks::Launch(TEXT("FSystem::GetMeshResultTask"),
-		[Context, SystemPrivate = m_pD, RootAddress, MeshContentFilter]() -> TSharedPtr<const FMesh>
+		[Context, SystemPrivate = m_pD, RootAddress, MeshContentFilter]() -> TManagedPtr<const FMesh>
 		{
-			TSharedPtr<const FMesh> Result;
+			TManagedPtr<const FMesh> Result;
 
 			SystemPrivate->bUnrecoverableError = Context->Runner->bUnrecoverableError;
 			if (!Context->Runner->bUnrecoverableError)
 			{
-				Result = SystemPrivate->WorkingMemoryManager.LoadMesh(FCacheAddress(RootAddress, 0, static_cast<uint8>(MeshContentFilter)), true);
+				Result = SystemPrivate->WorkingMemoryManager.LoadMesh(FCacheAddress(RootAddress, 0, static_cast<uint8>(MeshContentFilter)));
 			}
 			
 			SystemPrivate->WorkingMemoryManager.EndRunnerThread();
@@ -825,7 +822,7 @@ namespace UE::Mutable::Private
 	}
 
 
-	UE::Tasks::TTask<TSharedPtr<const FSkeletalMesh>> FSystem::GetSkeletalMesh(UE::Tasks::FTask Prerequisite, FInstance::FID InstanceID, const FSkeletalMeshId& SkeletalMeshId, EMeshContentFlags SkeletalMeshContentFilter)
+	UE::Tasks::TTask<TManagedPtr<const FSkeletalMesh>> FSystem::GetSkeletalMesh(UE::Tasks::FTask Prerequisite, FInstance::FID InstanceID, const FSkeletalMeshId& SkeletalMeshId, EMeshContentFlags SkeletalMeshContentFilter)
     {
 		LLM_SCOPE_BYNAME(TEXT("MutableRuntime"));
 		MUTABLE_CPUPROFILER_SCOPE(SystemGetSkeletalMesh);
@@ -833,9 +830,9 @@ namespace UE::Mutable::Private
 		if (!SkeletalMeshId)
 		{
 			return UE::Tasks::Launch(TEXT("FSystem::GetSkeletalMeshInvalid"),
-			[]() -> TSharedPtr<const FSkeletalMesh>
+			[]() -> TManagedPtr<const FSkeletalMesh>
 			{
-				return MakeShared<FSkeletalMesh>();
+				return MakeManaged<FSkeletalMesh>();
 			},
 			UE::Tasks::Prerequisites(Prerequisite),
 			UE::Tasks::ETaskPriority::Inherit,
@@ -852,9 +849,9 @@ namespace UE::Mutable::Private
 		if (GetOpDataType(OpType) != EDataType::SkeletalMesh)
 		{
 			return UE::Tasks::Launch(TEXT("FSystem::GetSkeletalMeshInvalid"),
-				[]() -> TSharedPtr<const FSkeletalMesh>
+				[]() -> TManagedPtr<const FSkeletalMesh>
 				{
-					return MakeShared<FSkeletalMesh>();
+					return MakeManaged<FSkeletalMesh>();
 				},
 				UE::Tasks::Prerequisites(Prerequisite),
 				UE::Tasks::ETaskPriority::Inherit,
@@ -894,9 +891,9 @@ namespace UE::Mutable::Private
 		UE::Tasks::ETaskPriority::Inherit);
 		
 		return UE::Tasks::Launch(TEXT("FSystem::GetSkeletalMeshResultTask"),
-		[Context, SystemPrivate = m_pD, RootAddress, SkeletalMeshContentFilter]() -> TSharedPtr<const FSkeletalMesh>
+		[Context, SystemPrivate = m_pD, RootAddress, SkeletalMeshContentFilter]() -> TManagedPtr<const FSkeletalMesh>
 		{
-			TSharedPtr<const FSkeletalMesh> Result;
+			TManagedPtr<const FSkeletalMesh> Result;
 
 			SystemPrivate->bUnrecoverableError = Context->Runner->bUnrecoverableError;
 			if (!Context->Runner->bUnrecoverableError)
@@ -916,7 +913,7 @@ namespace UE::Mutable::Private
 	}
 
 	
-	UE::Tasks::TTask<TSharedPtr<const FMaterial>> FSystem::GetMaterial(UE::Tasks::FTask Prerequisite, FInstance::FID InstanceID, const FMaterialId& MaterialId)
+	UE::Tasks::TTask<TManagedPtr<const FMaterial>> FSystem::GetMaterial(UE::Tasks::FTask Prerequisite, FInstance::FID InstanceID, const FMaterialId& MaterialId)
 	{
 		LLM_SCOPE_BYNAME(TEXT("MutableRuntime"));
 		MUTABLE_CPUPROFILER_SCOPE(SystemGetMaterial);
@@ -924,9 +921,9 @@ namespace UE::Mutable::Private
 		if (!MaterialId)
 		{
 			return UE::Tasks::Launch(TEXT("FSystem::GetMaterialInvalid"),
-			[]() -> TSharedPtr<const FMaterial>
+			[]() -> TManagedPtr<const FMaterial>
 			{
-				return MakeShared<FMaterial>();
+				return MakeManaged<FMaterial>();
 			},
 			UE::Tasks::Prerequisites(Prerequisite),
 			UE::Tasks::ETaskPriority::Inherit,
@@ -943,9 +940,9 @@ namespace UE::Mutable::Private
 		if (GetOpDataType(OpType) != EDataType::Material)
 		{
 			return UE::Tasks::Launch(TEXT("FSystem::GetMaterialInvalid"),
-			[]() -> TSharedPtr<const FMaterial>
+			[]() -> TManagedPtr<const FMaterial>
 			{
-				return MakeShared<FMaterial>();
+				return MakeManaged<FMaterial>();
 			},
 			UE::Tasks::Prerequisites(Prerequisite),
 			UE::Tasks::ETaskPriority::Inherit,
@@ -985,9 +982,9 @@ namespace UE::Mutable::Private
 		UE::Tasks::ETaskPriority::Inherit);
 		
 		return UE::Tasks::Launch(TEXT("FSystem::GetMaterialResultTask"),
-		[Context, SystemPrivate = m_pD, RootAddress]() -> TSharedPtr<const FMaterial>
+		[Context, SystemPrivate = m_pD, RootAddress]() -> TManagedPtr<const FMaterial>
 		{
-			TSharedPtr<const FMaterial> Result;
+			TManagedPtr<const FMaterial> Result;
 
 			SystemPrivate->bUnrecoverableError = Context->Runner->bUnrecoverableError;
 			if (!Context->Runner->bUnrecoverableError)
@@ -1059,23 +1056,6 @@ namespace UE::Mutable::Private
 			FLiveInstance& Instance = m_pD->WorkingMemoryManager.LiveInstances[Index];
 			if (Instance.InstanceID == instanceID)
 			{
-				// Make sure all the resources cached in the instance are removed from the tracking list
-				for (const FProgramCache::TResourceResult<FImage>& Data : Instance.Cache->ImageResults)
-				{
-					if (Data.Value)
-					{
-						m_pD->WorkingMemoryManager.CacheResources.Remove(Data.Value);
-					}
-				}
-
-				for (const FProgramCache::TResourceResult<FMesh>& Data : Instance.Cache->MeshResults)
-				{
-					if (Data.Value)
-					{
-						m_pD->WorkingMemoryManager.CacheResources.Remove(Data.Value);
-					}
-				}
-
 				m_pD->WorkingMemoryManager.LiveInstances.RemoveAtSwap(Index);
 				break;
 			}
@@ -1396,11 +1376,11 @@ namespace UE::Mutable::Private
 	}
 
 	
-	TSharedPtr<const FImage> FSystem::Private::BuildImage(const TSharedPtr<FExternalResourceProvider>& ExternalResourceProvider, const TSharedPtr<const FModel>& pModel, const FParameters* Params, OP::ADDRESS at, int32 MipsToSkip)
+	TManagedPtr<const FImage> FSystem::Private::BuildImage(const TSharedPtr<FExternalResourceProvider>& ExternalResourceProvider, const TSharedPtr<const FModel>& pModel, const FParameters* Params, OP::ADDRESS at, int32 MipsToSkip)
 	{
 		WorkingMemoryManager.BeginRunnerThread();
 
-		TSharedPtr<const FImage> Result;
+		TManagedPtr<const FImage> Result;
 
 		UE::Mutable::Private::EOpType OpType = pModel->GetPrivate()->Program.GetOpType(at);
 		if (GetOpDataType(OpType) == EDataType::Image)
@@ -1408,7 +1388,7 @@ namespace UE::Mutable::Private
 			RunCode(ExternalResourceProvider, pModel, Params, at, FSystem::AllLODs, uint8(MipsToSkip));
 			if (!bUnrecoverableError)
 			{
-				Result = WorkingMemoryManager.LoadImage(FCacheAddress(at, 0, MipsToSkip), true);
+				Result = WorkingMemoryManager.LoadImage(FCacheAddress(at, 0, MipsToSkip));
 			}			
 		}
 
@@ -1418,11 +1398,11 @@ namespace UE::Mutable::Private
 	}
 
 
-	TSharedPtr<const FMesh> FSystem::Private::BuildMesh(const TSharedPtr<FExternalResourceProvider>& ExternalResourceProvider, const TSharedPtr<const FModel>& Model, const FParameters* Params, OP::ADDRESS RootAddress, EMeshContentFlags MeshContentFilter)
+	TManagedPtr<const FMesh> FSystem::Private::BuildMesh(const TSharedPtr<FExternalResourceProvider>& ExternalResourceProvider, const TSharedPtr<const FModel>& Model, const FParameters* Params, OP::ADDRESS RootAddress, EMeshContentFlags MeshContentFilter)
 	{
 		WorkingMemoryManager.BeginRunnerThread();
 
-		TSharedPtr<const FMesh> Result;
+		TManagedPtr<const FMesh> Result;
 
 		UE::Mutable::Private::EOpType OpType = Model->GetPrivate()->Program.GetOpType(RootAddress);
 
@@ -1432,8 +1412,8 @@ namespace UE::Mutable::Private
 			RunCode(ExternalResourceProvider, Model, Params, RootAddress, FSystem::AllLODs, ExecutionOptions);
 			if (!bUnrecoverableError)
 			{
-				Result = WorkingMemoryManager.LoadMesh(FCacheAddress(RootAddress, 0, ExecutionOptions), true);
-			}	
+				Result = WorkingMemoryManager.LoadMesh(FCacheAddress(RootAddress, 0, ExecutionOptions));
+			}
 		}
 
 		WorkingMemoryManager.EndRunnerThread();
@@ -1442,11 +1422,11 @@ namespace UE::Mutable::Private
 	}
 
 
-	TSharedPtr<const FInstance> FSystem::Private::BuildInstance(const TSharedPtr<FExternalResourceProvider>& ExternalResourceProvider, const TSharedPtr<const FModel>& pModel, const FParameters* Params, OP::ADDRESS at)
+	TManagedPtr<const FInstance> FSystem::Private::BuildInstance(const TSharedPtr<FExternalResourceProvider>& ExternalResourceProvider, const TSharedPtr<const FModel>& pModel, const FParameters* Params, OP::ADDRESS at)
 	{
 		WorkingMemoryManager.BeginRunnerThread();
 
-		TSharedPtr<const FInstance> Result;
+		TManagedPtr<const FInstance> Result;
 
 		UE::Mutable::Private::EOpType OpType = pModel->GetPrivate()->Program.GetOpType(at);
 		if (GetOpDataType(OpType) == EDataType::Instance)
@@ -1464,12 +1444,11 @@ namespace UE::Mutable::Private
 	}
 
 
-	//---------------------------------------------------------------------------------------------
-	TSharedPtr<const FLayout> FSystem::Private::BuildLayout(const TSharedPtr<FExternalResourceProvider>& ExternalResourceProvider, const TSharedPtr<const FModel>& pModel, const FParameters* Params, OP::ADDRESS at)
+	TManagedPtr<const FLayout> FSystem::Private::BuildLayout(const TSharedPtr<FExternalResourceProvider>& ExternalResourceProvider, const TSharedPtr<const FModel>& pModel, const FParameters* Params, OP::ADDRESS at)
 	{
 		WorkingMemoryManager.BeginRunnerThread();
 
-		TSharedPtr<const FLayout> Result;
+		TManagedPtr<const FLayout> Result;
 
 		if (pModel->GetPrivate()->Program.States[0].Root)
 		{
@@ -1491,11 +1470,11 @@ namespace UE::Mutable::Private
 
 
 	//---------------------------------------------------------------------------------------------
-	TSharedPtr<const String> FSystem::Private::BuildString(const TSharedPtr<FExternalResourceProvider>& ExternalResourceProvider, const TSharedPtr<const FModel>& pModel, const FParameters* Params, OP::ADDRESS at)
+	TManagedPtr<const String> FSystem::Private::BuildString(const TSharedPtr<FExternalResourceProvider>& ExternalResourceProvider, const TSharedPtr<const FModel>& pModel, const FParameters* Params, OP::ADDRESS at)
 	{
 		WorkingMemoryManager.BeginRunnerThread();
 
-		TSharedPtr<const String> Result;
+		TManagedPtr<const String> Result;
 
 		if (pModel->GetPrivate()->Program.States[0].Root)
 		{
@@ -1572,9 +1551,16 @@ namespace UE::Mutable::Private
 		// located. Keep the localized memory computation for now.   
 		const uint32 RomBytes = GetRomBytes();
 		const uint32 CacheBytes = GetCacheBytes();
-		const uint32 TrackedCacheBytes = GetTrackedCacheBytes();
-		const uint32 PoolBytes = GetPooledBytes();
-		const uint32 TempBytes = GetTempBytes();
+		uint32 HeapBytes = 0;
+	
+		for (const CodeRunner::FScheduledOpData& Data : CurrentRunner->HeapData)
+		{
+			if (Data.Resource)
+			{
+				HeapBytes += Data.Resource->GetDataSize();
+			}
+		}
+
 
 		// Get allocator counters.
 		const SSIZE_T ImageAllocBytes	 = MemoryCounters::FImageMemoryCounter::Get().load(std::memory_order_relaxed);
@@ -1585,14 +1571,10 @@ namespace UE::Mutable::Private
 		SSIZE_T TotalBytes = ImageAllocBytes + MeshAllocBytes + StreamAllocBytes + InternalAllocBytes;
 
 		UE_LOG(LogMutableCore, Log, 
-			TEXT("Mem KB: ImageAlloc %7" SIZE_T_FMT ", MeshAlloc %7" SIZE_T_FMT ", StreamAlloc %7" SIZE_T_FMT ", InternalAlloc %7" SIZE_T_FMT ",  AllocTotal %7" SIZE_T_FMT " / %7" INT64_FMT ". \
-				  Resources MemLoc: Temp %7d, Pool %7d, Cache0+1 %7d (%7d), Rom %7d."),
-			ImageAllocBytes/1024, MeshAllocBytes/1024, StreamAllocBytes/1024, InternalAllocBytes/1024, TotalBytes/1024, BudgetBytes/1024,  
-			TempBytes/1024, PoolBytes/1024, CacheBytes/1024, TrackedCacheBytes/1024, RomBytes/1024);
-
-		TRACE_COUNTER_SET(MutableRuntime_MemTemp,	  TempBytes);
-		TRACE_COUNTER_SET(MutableRuntime_MemPool,     PoolBytes);
-		TRACE_COUNTER_SET(MutableRuntime_MemCache,    TrackedCacheBytes);
+			TEXT("Mem KB: ImageAlloc %7d, MeshAlloc %7d, StreamAlloc %7d, InternalAlloc %7d, AllocTotal %7d / %7d")
+			TEXT("Resources MemLoc: Cache0+1 %7d, Rom %7d, Heap %7d"),
+			ImageAllocBytes/1024, MeshAllocBytes/1024, StreamAllocBytes/1024, InternalAllocBytes/1024, TotalBytes/1024, BudgetBytes/1024, CacheBytes/1024, RomBytes/1024, HeapBytes/1024);
+
 		TRACE_COUNTER_SET(MutableRuntime_MemRom,      RomBytes);
 		TRACE_COUNTER_SET(MutableRuntime_MemInternal, InternalAllocBytes);
 		TRACE_COUNTER_SET(MutableRuntime_MemMesh,     MeshAllocBytes);
@@ -1660,7 +1642,6 @@ namespace UE::Mutable::Private
     }
 
 
-	//---------------------------------------------------------------------------------------------
 	int64 FWorkingMemoryManager::GetCurrentMemoryBytes() const
 	{
 		MUTABLE_CPUPROFILER_SCOPE(GetCurrentMemoryBytes);
@@ -1673,8 +1654,6 @@ namespace UE::Mutable::Private
 		return TotalBytes;
 	}
 
-
-	//---------------------------------------------------------------------------------------------
 	bool FWorkingMemoryManager::IsMemoryBudgetFull() const
 	{
 		// If we have 0 budget it means we have unlimited budget
@@ -1689,20 +1668,6 @@ namespace UE::Mutable::Private
 		return CurrentBytes > BudgetThresholdBytes;
 	}
 
-	namespace Private
-	{
-		enum class EBudgetBelowSearchFlags : uint8
-		{
-			None = 0,
-			Keep = 1 << 1,
-			Visited = 1 << 2,
-			FirstOccurance = 1 << 3
-		};
-
-		ENUM_CLASS_FLAGS(EBudgetBelowSearchFlags);
-	}
-
-    //---------------------------------------------------------------------------------------------
 	bool FWorkingMemoryManager::EnsureBudgetBelow( uint64 AdditionalMemory )
     {
 		MUTABLE_CPUPROFILER_SCOPE(EnsureBudgetBelow);
@@ -1718,23 +1683,8 @@ namespace UE::Mutable::Private
 		// Add the extra memory that we are trying to allocate when we return.
 		TotalBytes += AdditionalMemory;
 
-
         bool bFinished = TotalBytes <= BudgetBytes;
-
-		// Try to free pooled resources first
-		if (!bFinished)
-		{
-			MUTABLE_CPUPROFILER_SCOPE(EnsureBudgetBelow_FreePooled);
-			while (PooledImages.Num() && !bFinished)
-			{
-				// TODO: Actually advancing index if possible after swap may be better to remove the oldest in the pool first.
-				int32 PooledResourceSize = PooledImages[0]->GetDataSize();
-				TotalBytes -= PooledResourceSize;				
-				PooledImages.RemoveAtSwap(0);
-				bFinished = TotalBytes <= BudgetBytes;
-			}
-		}
-		
+	
 		// Try to free loaded roms
 		if (!bFinished)
 		{
@@ -1749,6 +1699,7 @@ namespace UE::Mutable::Private
 			TArray<TPair<float, FRomRef>> Candidates;
 			Candidates.Reserve(512);
 
+			SSIZE_T RomBytes = 0;
 			for (FModelCacheEntry& ModelCache : CachePerModel)
 			{
 				TSharedPtr<const FModel> CacheModel = ModelCache.Model.Pin();
@@ -1767,6 +1718,8 @@ namespace UE::Mutable::Private
 						check(Program.IsRomLoaded(RomIndex));
 						check(Rom.ResourceType == (uint32)*Iter);
 
+						RomBytes += CacheModel->GetPrivate()->Program.GetRomBytes(RomIndex);
+
 						// We cannot unload a rom if some operation is expecting it.
 						const bool bIsRomLocked = ModelCache.PendingOpsPerRom.IsValidIndex(RomIndex) && 
 												  ModelCache.PendingOpsPerRom[RomIndex] > 0;
@@ -1787,123 +1740,110 @@ namespace UE::Mutable::Private
 			auto CompareCandidates = [](const TPair<float, FRomRef>& A, const TPair<float, FRomRef>& B) { return A.Key < B.Key; };
 			Candidates.Heapify(CompareCandidates);
 
+			SSIZE_T CandidatesRomBytes = 0;
+			for (const TPair<float, FRomRef>& Entry : Candidates)
+			{
+				CandidatesRomBytes += Entry.Value.Model->GetPrivate()->Program.GetRomBytes(Entry.Value.RomIndex);
+			}
+		
 			while (!bFinished && Candidates.Num())
 			{
 				MUTABLE_CPUPROFILER_SCOPE(EnsureBudgetBelow_UnloadRom);
 
-				TPair<float, FRomRef> Candidate;
-				Candidates.HeapPop(Candidate, CompareCandidates, EAllowShrinking::No);
+				TPair<float, FRomRef> Candidate = Candidates.HeapTop();
 
 				// UE_LOG(LogMutableCore,Log, "Unloading rom because of memory budget: %d.", lowestPriorityRom);
 				int32 UnloadedSize = 0;
 				Candidate.Value.Model->GetPrivate()->Program.UnloadRom(Candidate.Value.RomIndex, &UnloadedSize);
 				TotalBytes -= UnloadedSize;
 				bFinished = TotalBytes <= BudgetBytes;
+
+				if (!bFinished)
+				{
+					Candidates.HeapPopDiscard(CompareCandidates, EAllowShrinking::No);
+				}
 			}
 		}
-	
-		// Try to free cache 1 memory
-		if (!bFinished)
+
+		auto TryRemoveFromCache = [&bFinished, &TotalBytes, Budget = BudgetBytes](FProgramCache& Cache)
 		{
-			MUTABLE_CPUPROFILER_SCOPE(EnsureBudgetBelow_FreeCached);
+			int32 ImageRemovedBytes = 0;
 
-			TSet<const FResource*> RemovedResources;
-			RemovedResources.Reserve(1024);
-			
-			// From other live instances first
-			for (const FLiveInstance& Instance : LiveInstances)
+			for (const TManagedWeakPtr<const FImage>& ImageResult : Cache.ImageResults)
 			{
-				if (Instance.Cache == CurrentInstanceCache)
+				TManagedPtr<const FImage> StrongImagePtr = ImageResult.Pin();
+				if (!StrongImagePtr)
 				{
-					// Ignore the current live instance.
 					continue;
 				}
 
-				// Gather all data in the cache for this instance
+				int32 ImageDataSize = StrongImagePtr->GetDataSize();
+				if (StrongImagePtr.TryDeleteObject())
 				{
-					MUTABLE_CPUPROFILER_SCOPE(EnsureBudgetBelow_FreeCached_GatherAndRemove_Other);
+					ImageRemovedBytes += ImageDataSize;
+				}
 
-					RemovedResources.Reset();
-					int32 ImageRemovedBytes = 0;
-					for (const FProgramCache::TResourceResult<FImage>& ImageResult : Instance.Cache->ImageResults)
-					{
-						if (!ImageResult.Value)
-						{
-							continue;
-						}
+				bFinished = TotalBytes - ImageRemovedBytes <= Budget;
 
-						if (!RemovedResources.Find(ImageResult.Value.Get()))
-						{
-							ImageRemovedBytes += ImageResult.Value->GetDataSize();
-							bFinished = TotalBytes - ImageRemovedBytes <= BudgetBytes;
+				if (bFinished)
+				{
+					break;
+				}
+			}
 
-							RemovedResources.Add(ImageResult.Value.Get());
-							CacheResources.Remove(ImageResult.Value);
-						}
+			TotalBytes -= ImageRemovedBytes;
+			if (bFinished)
+			{
+				return;
+			}
 
-						if (bFinished)
-						{
-							break;
-						}
-					}
+			int32 MeshRemovedBytes = 0;
+			for (const TManagedWeakPtr<const FMesh>& MeshResult : Cache.MeshResults)
+			{
+				TManagedPtr<const FMesh> StrongMeshPtr = MeshResult.Pin();
+				if (!StrongMeshPtr)
+				{
+					continue;
+				}
 
-					if (ImageRemovedBytes > 0)
-					{
-						for (const FProgramCache::TResourceResult<FImage>& ImageResult : Instance.Cache->ImageResults)
-						{
-							if (RemovedResources.Find(ImageResult.Value.Get()))
-							{
-								Instance.Cache->SetUnused(Instance.Cache->OpExecutionData[ImageResult.OpAddress]);
-							}
-						}
-					}
+				int32 MeshDataSize = StrongMeshPtr->GetDataSize();
+				if (StrongMeshPtr.TryDeleteObject())
+				{
+					MeshRemovedBytes += MeshDataSize;
+				}
 
-					TotalBytes -= ImageRemovedBytes;
-					if (bFinished)
-					{
-						break;
-					}
+				bFinished = TotalBytes - MeshRemovedBytes <= Budget;
 
-					int32 MeshRemovedBytes = 0;
-					RemovedResources.Reset();
-					for (const FProgramCache::TResourceResult<FMesh>& MeshResult : Instance.Cache->MeshResults)
-					{
-						if (!MeshResult.Value)
-						{
-							continue;
-						}
+				if (bFinished)
+				{
+					break;
+				}
+			}
 
-						if (!RemovedResources.Find(MeshResult.Value.Get()))
-						{
-							MeshRemovedBytes = MeshResult.Value->GetDataSize();
-							bFinished = TotalBytes - MeshRemovedBytes <= BudgetBytes;
+			TotalBytes -= MeshRemovedBytes;
+		};
 
-							RemovedResources.Add(MeshResult.Value.Get());
-							CacheResources.Remove(MeshResult.Value);
-						}
+		// Try to free cache 1 memory
+		if (!bFinished)
+		{
+			MUTABLE_CPUPROFILER_SCOPE(EnsureBudgetBelow_FreeCached);	
+			// From other live instances first
+			for (const FLiveInstance& Instance : LiveInstances)
+			{
+				if (Instance.Cache == CurrentInstanceCache)
+				{
+					// Ignore the current live instance.
+					continue;
+				}
 
-						if (bFinished)
-						{
-							break;
-						}
-					}
+				MUTABLE_CPUPROFILER_SCOPE(EnsureBudgetBelow_FreeCached_GatherAndRemove_Other);
+				// Gather all data in the cache for this instance
 
-					if (MeshRemovedBytes > 0)
-					{
-						for (const FProgramCache::TResourceResult<FMesh>& MeshResult : Instance.Cache->MeshResults)
-						{
-							if (RemovedResources.Find(MeshResult.Value.Get()))
-							{
-								Instance.Cache->SetUnused(Instance.Cache->OpExecutionData[MeshResult.OpAddress]);
-							}
-						}
-					}
+				TryRemoveFromCache(*Instance.Cache);
 
-					TotalBytes -= MeshRemovedBytes;
-					if (bFinished)
-					{
-						break;
-					}
+				if (bFinished)
+				{
+					break;
 				}
 			}
 		}
@@ -1913,163 +1853,11 @@ namespace UE::Mutable::Private
 		if (!bFinished && CurrentInstanceCache)
 		{
 			MUTABLE_CPUPROFILER_SCOPE(EnsureBudgetBelow_FreeCached_Current);
-			// This removes images first, not sure if this can become a problem.
-			
-			using namespace Private;
-			using ESearchFlags = EBudgetBelowSearchFlags;
-			
-			auto SearchResourcesToRemove = 
-			[
-				TotalBytes = TotalBytes, 
-				BudgetBytes = BudgetBytes, 
-				CurrentInstanceCache = CurrentInstanceCache.Get()
-			](const auto& ResourceRange, TArray<ESearchFlags>& SearchFlags) -> int32
-			{
-				MUTABLE_CPUPROFILER_SCOPE(EnsureBudgetBelow_FreeCached_Current_SearchResources);
-
-				int32 RemovedBytes = 0;
-
-				const int32 NumResources = ResourceRange.Num();
-				for (int32 ResourceIndex = 0; ResourceIndex < NumResources; ++ResourceIndex)
-				{
-					// Check if visited.
-					if (EnumHasAnyFlags(SearchFlags[ResourceIndex], ESearchFlags::Visited))
-					{
-						continue;
-					}
-	
-					// Null values will not have any flag set.
-					if (!ResourceRange[ResourceIndex].Value)
-					{
-						continue;
-					}
-
-					SearchFlags[ResourceIndex] = (ESearchFlags::Visited | ESearchFlags::FirstOccurance);
-					
-					if (CurrentInstanceCache->OpExecutionData[ResourceRange[ResourceIndex].OpAddress].OpHitCount > 0)
-					{
-						// Mark all occurences as visited to keep. 
-						for (int32 I = ResourceIndex; I < NumResources; ++I)
-						{
-							if (ResourceRange[I].Value == ResourceRange[ResourceIndex].Value)
-							{
-								EnumAddFlags(SearchFlags[I], (ESearchFlags::Visited | ESearchFlags::Keep));
-							}
-						}
-					}
-					else
-					{
-						int32 I = ResourceIndex + 1;
-						for (; I < NumResources; ++I)
-						{
-							// Mark as visted.
-							if (ResourceRange[I].Value == ResourceRange[ResourceIndex].Value)
-							{
-								EnumAddFlags(SearchFlags[I], ESearchFlags::Visited);	
-								if (CurrentInstanceCache->OpExecutionData[ResourceRange[I].OpAddress].OpHitCount > 0)
-								{
-									// Still used, next step will mark to keep.
-									break;
-								}
-							}
-						}
-
-						if (I < NumResources)
-						{
-							// The image is still used, mark all occurences as visited and keep.
-							for (int32 J = ResourceIndex; J < NumResources; ++J)
-							{
-								if (ResourceRange[J].Value == ResourceRange[ResourceIndex].Value)
-								{
-									EnumAddFlags(SearchFlags[J], (ESearchFlags::Visited | ESearchFlags::Keep));	
-								}
-							}
-						}
-						else
-						{
-							// The resource is not used, see if we can stop searching. 
-							// In that case all occurences have been marked as visited.
-							RemovedBytes += ResourceRange[ResourceIndex].Value->GetDataSize();
-
-							if (TotalBytes - RemovedBytes <= BudgetBytes)
-							{
-								return RemovedBytes;
-							}
-						}
-					}
-				}
-
-				return RemovedBytes;
-			};
-
-			const int32 MaxNumResources = FMath::Max(CurrentInstanceCache->ImageResults.Num(), CurrentInstanceCache->MeshResults.Num());
-			TArray<ESearchFlags> SearchFlags;
-			SearchFlags.SetNumUninitialized(MaxNumResources);
-			
-			if (!bFinished)
-			{
-				FMemory::Memzero(SearchFlags.GetData(), MaxNumResources*sizeof(ESearchFlags));
-				
-				const TArrayView<FProgramCache::TResourceResult<FImage>>& Images = MakeArrayView(CurrentInstanceCache->ImageResults);
-				const int32 RemovedBytes = SearchResourcesToRemove(Images, SearchFlags); 
-				
-				if (RemovedBytes > 0)
-				{	
-					MUTABLE_CPUPROFILER_SCOPE(EnsureBudgetBelow_FreeCached_Current_FreeResources);
-			
-					const int32 NumImages = Images.Num();
-					for (int32 I = 0; I < NumImages; ++I)
-					{
-						// Remove the first occurence of any visited not to keep. 
-						if (SearchFlags[I] == (ESearchFlags::FirstOccurance | ESearchFlags::Visited))
-						{
-							CacheResources.Remove(Images[I].Value);
-						}
-					
-						// Set unused all references visited not marked to keep. 
-						if ((SearchFlags[I] & (ESearchFlags::Visited | ESearchFlags::Keep)) == ESearchFlags::Visited)
-						{
-							CurrentInstanceCache->SetUnused(CurrentInstanceCache->OpExecutionData[Images[I].OpAddress]);
-						}
-					}	
-				}
-
-				TotalBytes -= RemovedBytes;
-				bFinished = TotalBytes <= BudgetBytes;
-			}
-			
-			if (!bFinished)
-			{
-				FMemory::Memzero(SearchFlags.GetData(), MaxNumResources*sizeof(ESearchFlags));
-				
-				TArrayView<FProgramCache::TResourceResult<FMesh>> Meshes = MakeArrayView(CurrentInstanceCache->MeshResults);
-				const int32 RemovedBytes = SearchResourcesToRemove(Meshes, SearchFlags); 
-
-				if (RemovedBytes > 0)
-				{	
-					MUTABLE_CPUPROFILER_SCOPE(EnsureBudgetBelow_FreeCached_Current_FreeResources);
-					
-					const int32 NumMeshes = Meshes.Num();
-					for (int32 I = 0; I < NumMeshes; ++I)
-					{
-						// Remove the first occurence of any visited not to keep. 
-						if (SearchFlags[I] == (ESearchFlags::FirstOccurance | ESearchFlags::Visited))
-						{
-							CacheResources.Remove(Meshes[I].Value);
-						}
-						
-						if ((SearchFlags[I] & (ESearchFlags::Visited | ESearchFlags::Keep)) == ESearchFlags::Visited)
-						{
-							CurrentInstanceCache->SetUnused(CurrentInstanceCache->OpExecutionData[Meshes[I].OpAddress]);
-						}
-					}
-				}
 
-				TotalBytes -= RemovedBytes;
-				bFinished = TotalBytes <= BudgetBytes;
-			}
+			TryRemoveFromCache(*CurrentInstanceCache);
 		}
 
+
 		if (!bFinished)
 		{
 			int64 ExcessBytes = TotalBytes - BudgetBytes;
diff --git a/Engine/Plugins/Mutable/Source/MutableRuntime/Public/MuR/External/MaterialAdapter.h b/Engine/Plugins/Mutable/Source/MutableRuntime/Public/MuR/External/MaterialAdapter.h
--- a/Engine/Plugins/Mutable/Source/MutableRuntime/Public/MuR/External/MaterialAdapter.h
+++ b/Engine/Plugins/Mutable/Source/MutableRuntime/Public/MuR/External/MaterialAdapter.h
@@ -52,6 +52,6 @@ namespace UE::Mutable
 		MUTABLERUNTIME_API bool SetTexture(const Private::FParameterKey& ParameterKey, const TValueConst<FTextureAdapter>* Value);
 	
 	private:
-		TSharedPtr<Private::FMaterial> Material;
+		Private::TManagedPtr<Private::FMaterial> Material;
 	};
 }
diff --git a/Engine/Plugins/Mutable/Source/MutableRuntime/Public/MuR/External/MeshAdapter.h b/Engine/Plugins/Mutable/Source/MutableRuntime/Public/MuR/External/MeshAdapter.h
--- a/Engine/Plugins/Mutable/Source/MutableRuntime/Public/MuR/External/MeshAdapter.h
+++ b/Engine/Plugins/Mutable/Source/MutableRuntime/Public/MuR/External/MeshAdapter.h
@@ -42,6 +42,6 @@ namespace UE::Mutable
 		MUTABLERUNTIME_API void RemoveVertices(const TBitArray<>& Vertices);
 		
 	private:
-		TSharedPtr<Private::FMesh> Mesh;
+		Private::TManagedPtr<Private::FMesh> Mesh;
 	};
 }
diff --git a/Engine/Plugins/Mutable/Source/MutableRuntime/Public/MuR/External/Operation.h b/Engine/Plugins/Mutable/Source/MutableRuntime/Public/MuR/External/Operation.h
--- a/Engine/Plugins/Mutable/Source/MutableRuntime/Public/MuR/External/Operation.h
+++ b/Engine/Plugins/Mutable/Source/MutableRuntime/Public/MuR/External/Operation.h
@@ -2,7 +2,10 @@
 
 #pragma once
 
+#include "StructUtils/InstancedStruct.h"
 #include "Value.h"
+#include "MuR/ManagedPointer.h"
+#include "StructUtils/InstancedStruct.h"
 
 #include "Operation.generated.h"
 
@@ -68,4 +71,4 @@ namespace UE::Mutable
 		/** This function must be pure functional (same inputs, same output). */
 		virtual void Evaluate(FContext& Context) const {}
 	};
-}
\ No newline at end of file
+}
diff --git a/Engine/Plugins/Mutable/Source/MutableRuntime/Public/MuR/External/TextureAdapter.h b/Engine/Plugins/Mutable/Source/MutableRuntime/Public/MuR/External/TextureAdapter.h
--- a/Engine/Plugins/Mutable/Source/MutableRuntime/Public/MuR/External/TextureAdapter.h
+++ b/Engine/Plugins/Mutable/Source/MutableRuntime/Public/MuR/External/TextureAdapter.h
@@ -3,6 +3,7 @@
 #pragma once
 
 
+#include "MuR/ManagedPointer.h"
 #include "TextureAdapter.generated.h"
 
 
@@ -36,6 +37,6 @@ namespace UE::Mutable
 		FTextureAdapter& operator=(FTextureAdapter&& Other) = delete;
 		
 	private:
-		TSharedPtr<Private::FImage> Image;
+		Private::TManagedPtr<Private::FImage> Image;
 	};
 }
diff --git a/Engine/Plugins/Mutable/Source/MutableRuntime/Public/MuR/External/Value.h b/Engine/Plugins/Mutable/Source/MutableRuntime/Public/MuR/External/Value.h
--- a/Engine/Plugins/Mutable/Source/MutableRuntime/Public/MuR/External/Value.h
+++ b/Engine/Plugins/Mutable/Source/MutableRuntime/Public/MuR/External/Value.h
@@ -3,7 +3,7 @@
 #pragma once
 
 #include "StructUtils/InstancedStruct.h"
-
+#include "MuR/ManagedPointer.h"
 
 struct FInstancedStruct;
 
@@ -31,7 +31,7 @@ namespace UE::Mutable
 		static FValue New()
 		{
 			FValue Value;
-			Value.Ptr = MakeShared<FInstancedStruct>();
+			Value.Ptr = Private::MakeManaged<FInstancedStruct>();
 			Value.Ptr->InitializeAs(T::StaticStruct());
 
 			return Value;
@@ -54,7 +54,7 @@ namespace UE::Mutable
 
 		FValue(const UScriptStruct* ScriptStruct)
 		{
-			Ptr = MakeShared<FInstancedStruct>();
+			Ptr = Private::MakeManaged<FInstancedStruct>();
 			Ptr->InitializeAs(ScriptStruct);
 		}
 		
@@ -73,7 +73,7 @@ namespace UE::Mutable
 	private:
 		FValue() = default;
 		
-		TSharedPtr<FInstancedStruct> Ptr;
+		Private::TManagedPtr<FInstancedStruct> Ptr;
 	};
 
 
@@ -88,7 +88,7 @@ namespace UE::Mutable
 		static FValueConst New()
 		{
 			FValueConst Value;
-			Value.Ptr = MakeShared<const FInstancedStruct>();
+			Value.Ptr = Private::MakeManaged<const FInstancedStruct>();
 			Value.Ptr->InitializeAs(T::StaticStruct());
 
 			return Value;
@@ -126,7 +126,7 @@ namespace UE::Mutable
 	private:
 		FValueConst() = default;
 		
-		TSharedPtr<const FInstancedStruct> Ptr;
+		Private::TManagedPtr<const FInstancedStruct> Ptr;
 	};
 
 	
@@ -143,7 +143,7 @@ namespace UE::Mutable
 		static TValue New()
 		{
 			TValue Value;
-			Value.Ptr = MakeShared<FInstancedStruct>();
+			Value.Ptr = Private::MakeManaged<FInstancedStruct>();
 			Value.Ptr->InitializeAs(T::StaticStruct());
 
 			return Value;
@@ -166,7 +166,7 @@ namespace UE::Mutable
 
 		TValue(const UScriptStruct* ScriptStruct)
 		{
-			Ptr = MakeShared<FInstancedStruct>();
+			Ptr = Private::MakeManaged<FInstancedStruct>();
 			Ptr->InitializeAs(ScriptStruct);
 		}
 		
@@ -183,7 +183,7 @@ namespace UE::Mutable
 	private:
 		TValue() = default;
 		
-		TSharedPtr<FInstancedStruct> Ptr; // TODO GMT Extension TVariant Lazy	
+		Private::TManagedPtr<FInstancedStruct> Ptr; // TODO GMT Extension TVariant Lazy	
 	};
 	
 
@@ -199,7 +199,7 @@ namespace UE::Mutable
 		static TValueConst New()
 		{
 			TValueConst Value;
-			Value.Ptr = MakeShared<const FInstancedStruct>();
+			Value.Ptr = Private::MakeManaged<const FInstancedStruct>();
 			Value.Ptr->InitializeAs(T::StaticStruct());
 
 			return Value;
@@ -238,7 +238,7 @@ namespace UE::Mutable
 	private:
 		TValueConst() = default;
 		
-		TSharedPtr<const FInstancedStruct> Ptr;
+		Private::TManagedPtr<const FInstancedStruct> Ptr;
 	};
 	
 	
@@ -255,14 +255,14 @@ namespace UE::Mutable
 		if (ValueConst.Ptr.IsUnique())
 		{
 			TValue<T> Value;
-			Value.Ptr = ConstCastSharedPtr<FInstancedStruct>(ValueConst.Ptr); 
+			Value.Ptr = Private::ConstCastManagedPtr<FInstancedStruct>(ValueConst.Ptr); 
 
 			return MoveTemp(Value);
 		}
 		else
 		{
 			TValue<T> Value;
-			Value.Ptr = MakeShared<FInstancedStruct>(*ValueConst.Ptr); 
+			Value.Ptr = Private::MakeManaged<FInstancedStruct>(*ValueConst.Ptr); 
 
 			return MoveTemp(Value);
 		}
diff --git a/Engine/Plugins/Mutable/Source/MutableTools/Internal/MuT/AST.h b/Engine/Plugins/Mutable/Source/MutableTools/Internal/MuT/AST.h
--- a/Engine/Plugins/Mutable/Source/MutableTools/Internal/MuT/AST.h
+++ b/Engine/Plugins/Mutable/Source/MutableTools/Internal/MuT/AST.h
@@ -245,7 +245,7 @@ namespace UE::Mutable::Private
 		bool bSeparateImageMips = true;
 
 		/** Structure used to speedup mesh constant comparison. */
-		struct FDeduplicationMeshFuncs : TDefaultMapHashableKeyFuncs<TSharedPtr<const FMesh>, int32, false>
+		struct FDeduplicationMeshFuncs : TDefaultMapHashableKeyFuncs<TManagedPtr<const FMesh>, int32, false>
 		{
 			static FORCEINLINE bool Matches(KeyInitType A, KeyInitType B)
 			{
@@ -262,10 +262,10 @@ namespace UE::Mutable::Private
 			}
 		};
 
-		TMap<TSharedPtr<const FMesh>, int32, FDefaultSetAllocator, FDeduplicationMeshFuncs> MeshConstantMap;
+		TMap<TManagedPtr<const FMesh>, int32, FDefaultSetAllocator, FDeduplicationMeshFuncs> MeshConstantMap;
 
 		/** Structure used to speedup image mip comparison. */
-		struct FDeduplicationImageFuncs : TDefaultMapHashableKeyFuncs<TSharedPtr<const FImage>, int32, false>
+		struct FDeduplicationImageFuncs : TDefaultMapHashableKeyFuncs<TManagedPtr<const FImage>, int32, false>
 		{
 			static FORCEINLINE bool Matches(KeyInitType A, KeyInitType B)
 			{
@@ -285,7 +285,7 @@ namespace UE::Mutable::Private
 			}
 		};
 
-		TMap<TSharedPtr<const FImage>, int32, FDefaultSetAllocator, FDeduplicationImageFuncs> ImageConstantMipMap;
+		TMap<TManagedPtr<const FImage>, int32, FDefaultSetAllocator, FDeduplicationImageFuncs> ImageConstantMipMap;
 
 		/** Image operation functions, so that they can be overriden. */
 		FImageOperator& ImageOperator;
diff --git a/Engine/Plugins/Mutable/Source/MutableTools/Internal/MuT/Compiler.h b/Engine/Plugins/Mutable/Source/MutableTools/Internal/MuT/Compiler.h
--- a/Engine/Plugins/Mutable/Source/MutableTools/Internal/MuT/Compiler.h
+++ b/Engine/Plugins/Mutable/Source/MutableTools/Internal/MuT/Compiler.h
@@ -28,7 +28,7 @@ namespace UE::Mutable::Private
     class NodeTransformedObject;
 
 	/** Arguments are Texture ID, output generated texture and a "run immediately" flag. */
-	typedef TFunction<UE::Tasks::FTask(PASSTHROUGH_ID, TSharedPtr<TSharedPtr<FImage>>, bool)> FReferencedImageResourceFunc;
+	typedef TFunction<UE::Tasks::FTask(PASSTHROUGH_ID, TSharedPtr<TManagedPtr<FImage>>, bool)> FReferencedImageResourceFunc;
 
 	/** Arguments are 
 	* - Mesh ID, 
@@ -36,7 +36,7 @@ namespace UE::Mutable::Private
 	* - output generated mesh 
 	* - "run immediately" flag. 
 	*/
-	typedef TFunction<UE::Tasks::FTask(PASSTHROUGH_ID, const FString&, TSharedPtr<TSharedPtr<FMesh>>, bool)> FReferencedMeshResourceFunc;
+	typedef TFunction<UE::Tasks::FTask(PASSTHROUGH_ID, const FString&, TSharedPtr<TManagedPtr<FMesh>>, bool)> FReferencedMeshResourceFunc;
 
     /** Options used to compile the models with a compiler. */
     class CompilerOptions : public RefCounted
diff --git a/Engine/Plugins/Mutable/Source/MutableTools/Internal/MuT/NodeExtensionDataConstant.h b/Engine/Plugins/Mutable/Source/MutableTools/Internal/MuT/NodeExtensionDataConstant.h
--- a/Engine/Plugins/Mutable/Source/MutableTools/Internal/MuT/NodeExtensionDataConstant.h
+++ b/Engine/Plugins/Mutable/Source/MutableTools/Internal/MuT/NodeExtensionDataConstant.h
@@ -19,7 +19,7 @@ namespace UE::Mutable::Private
 	{
 	public:
 
-		TSharedPtr<const FExtensionData> Value;
+		TManagedPtr<const FExtensionData> Value;
 
 	public:
 
@@ -34,7 +34,10 @@ namespace UE::Mutable::Private
 		//-----------------------------------------------------------------------------------------
 
 		/** Deprecated. Access Value attribute directly. */
-		void SetValue(const TSharedPtr<const FExtensionData>& In) { Value = In; }
+		void SetValue(const TSharedPtr<const FExtensionData>& In) 
+		{
+			Value = MakeManaged<FExtensionData>(*In);
+		}
 
 	protected:
 
diff --git a/Engine/Plugins/Mutable/Source/MutableTools/Internal/MuT/NodeExternalParameter.h b/Engine/Plugins/Mutable/Source/MutableTools/Internal/MuT/NodeExternalParameter.h
--- a/Engine/Plugins/Mutable/Source/MutableTools/Internal/MuT/NodeExternalParameter.h
+++ b/Engine/Plugins/Mutable/Source/MutableTools/Internal/MuT/NodeExternalParameter.h
@@ -26,7 +26,7 @@ namespace UE::Mutable::Private
 		static UE_API FNodeType StaticType;
 
 	public:
-		TSharedPtr<const FInstancedStruct> DefaultValue;
+		TManagedPtr<const FInstancedStruct> DefaultValue;
 		
 		FString Name;
 		FString UID;
diff --git a/Engine/Plugins/Mutable/Source/MutableTools/Internal/MuT/NodeImageConstant.h b/Engine/Plugins/Mutable/Source/MutableTools/Internal/MuT/NodeImageConstant.h
--- a/Engine/Plugins/Mutable/Source/MutableTools/Internal/MuT/NodeImageConstant.h
+++ b/Engine/Plugins/Mutable/Source/MutableTools/Internal/MuT/NodeImageConstant.h
@@ -19,7 +19,7 @@ namespace UE::Mutable::Private
 	public:
 
 		/** */
-		TSharedPtr<FImage> Image;
+		TManagedPtr<FImage> Image;
 		
 		/** */
 		FSourceDataDescriptor SourceDataDescriptor;
diff --git a/Engine/Plugins/Mutable/Source/MutableTools/Internal/MuT/NodeLayout.h b/Engine/Plugins/Mutable/Source/MutableTools/Internal/MuT/NodeLayout.h
--- a/Engine/Plugins/Mutable/Source/MutableTools/Internal/MuT/NodeLayout.h
+++ b/Engine/Plugins/Mutable/Source/MutableTools/Internal/MuT/NodeLayout.h
@@ -20,7 +20,7 @@ namespace UE::Mutable::Private
 	struct FSourceLayoutBlock
 	{
 		/** Optional mask image that selects the vertices to include in the block. */
-		TSharedPtr<FImage> Mask;
+		TManagedPtr<FImage> Mask;
 
 		FIntVector2 Min = { 0, 0 };
 		FIntVector2 Size = { 0, 0 };
@@ -75,13 +75,13 @@ namespace UE::Mutable::Private
 		// Own interface
 
 		/** */
-		TSharedPtr<FLayout> BuildRuntimeLayout(uint32 BlockIDPrefix) const;
+		TManagedPtr<FLayout> BuildRuntimeLayout(uint32 BlockIDPrefix) const;
 
 		/** Generate the blocks of a layout using the UV of the meshes. 
 		* A list of existing blocks may be provided to specify starting blocks that shouldn't be modified unless it is mandatory because they contain partial face islands. 
 		*/
-		UE_API void GenerateLayoutBlocks(const TSharedPtr<FMesh>&, int32 LayoutIndex );
-		UE_API void GenerateLayoutBlocksFromUVIslands(const TSharedPtr<FMesh>&, int32 LayoutIndex);
+		UE_API void GenerateLayoutBlocks(const TManagedPtr<FMesh>&, int32 LayoutIndex );
+		UE_API void GenerateLayoutBlocksFromUVIslands(const TManagedPtr<FMesh>&, int32 LayoutIndex);
 
 	protected:
 
diff --git a/Engine/Plugins/Mutable/Source/MutableTools/Internal/MuT/NodeMeshConstant.h b/Engine/Plugins/Mutable/Source/MutableTools/Internal/MuT/NodeMeshConstant.h
--- a/Engine/Plugins/Mutable/Source/MutableTools/Internal/MuT/NodeMeshConstant.h
+++ b/Engine/Plugins/Mutable/Source/MutableTools/Internal/MuT/NodeMeshConstant.h
@@ -23,7 +23,7 @@ namespace UE::Mutable::Private
 
 		FSourceDataDescriptor SourceDataDescriptor;
 
-		TSharedPtr<FMesh> Value;
+		UE::Mutable::Private::TManagedPtr<FMesh> Value;
 
 		TArray<Ptr<NodeLayout>> Layouts;
 
diff --git a/Engine/Plugins/Mutable/Source/MutableTools/Internal/MuT/NodeMeshSkeletalMeshBreak.h b/Engine/Plugins/Mutable/Source/MutableTools/Internal/MuT/NodeMeshSkeletalMeshBreak.h
--- a/Engine/Plugins/Mutable/Source/MutableTools/Internal/MuT/NodeMeshSkeletalMeshBreak.h
+++ b/Engine/Plugins/Mutable/Source/MutableTools/Internal/MuT/NodeMeshSkeletalMeshBreak.h
@@ -15,7 +15,7 @@ namespace UE::Mutable::Private
 	class NodeMeshSkeletalMeshBreak : public NodeMesh
 	{
 	public:
-		TSharedPtr<FMesh> ReferenceMesh;
+		TManagedPtr<FMesh> ReferenceMesh;
 
 		TArray<Ptr<NodeLayout>> Layouts;
 
diff --git a/Engine/Plugins/Mutable/Source/MutableTools/Internal/MuT/Table.h b/Engine/Plugins/Mutable/Source/MutableTools/Internal/MuT/Table.h
--- a/Engine/Plugins/Mutable/Source/MutableTools/Internal/MuT/Table.h
+++ b/Engine/Plugins/Mutable/Source/MutableTools/Internal/MuT/Table.h
@@ -60,8 +60,8 @@ namespace UE::Mutable::Private
         UE_API void SetCell( int32 Column, uint32 RowId, int32 IntValue, const void* ErrorContext = nullptr);
         UE_API void SetCell( int32 Column, uint32 RowId, float Value, const void* ErrorContext = nullptr);
         UE_API void SetCell( int32 Column, uint32 RowId, const FVector4f& Value, const void* ErrorContext = nullptr);
-		UE_API void SetCell( int32 Column, uint32 RowId, const TSharedPtr<FImage>& Value, const void* ErrorContext = nullptr);
-		UE_API void SetCell( int32 Column, uint32 RowId, const TSharedPtr<FMesh>& Value, const void* ErrorContext = nullptr);
+		UE_API void SetCell( int32 Column, uint32 RowId, const UE::Mutable::Private::TManagedPtr<FImage>& Value, const void* ErrorContext = nullptr);
+		UE_API void SetCell( int32 Column, uint32 RowId, const UE::Mutable::Private::TManagedPtr<FMesh>& Value, const void* ErrorContext = nullptr);
         UE_API void SetCell( int32 Column, uint32 RowId, const FString& Value, const void* ErrorContext = nullptr);
 
 
diff --git a/Engine/Plugins/Mutable/Source/MutableTools/Internal/MuT/TablePrivate.h b/Engine/Plugins/Mutable/Source/MutableTools/Internal/MuT/TablePrivate.h
--- a/Engine/Plugins/Mutable/Source/MutableTools/Internal/MuT/TablePrivate.h
+++ b/Engine/Plugins/Mutable/Source/MutableTools/Internal/MuT/TablePrivate.h
@@ -25,8 +25,8 @@ namespace UE::Mutable::Private
 		int32 Int;
 		float Scalar;
 		FVector4f Color;
-		TSharedPtr<FImage> Image;
-		TSharedPtr<FMesh> Mesh;
+		TManagedPtr<FImage> Image;
+		TManagedPtr<FMesh> Mesh;
 		FString String;
 
 		const void* ErrorContext;
diff --git a/Engine/Plugins/Mutable/Source/MutableTools/Private/MuT/ASTOpConstantExtensionData.h b/Engine/Plugins/Mutable/Source/MutableTools/Private/MuT/ASTOpConstantExtensionData.h
--- a/Engine/Plugins/Mutable/Source/MutableTools/Private/MuT/ASTOpConstantExtensionData.h
+++ b/Engine/Plugins/Mutable/Source/MutableTools/Private/MuT/ASTOpConstantExtensionData.h
@@ -19,8 +19,7 @@ namespace UE::Mutable::Private
 	{
 	public:
 
-		//!
-		TSharedPtr<const FExtensionData> Value;
+		TManagedPtr<const FExtensionData> Value;
 
 	public:
 
