[Mutable] Memory managment rework.

diff --git a/Engine/Plugins/Mutable/Source/MutableTools/Private/MuT/ASTOpConstantResource.cpp b/Engine/Plugins/Mutable/Source/MutableTools/Private/MuT/ASTOpConstantResource.cpp
--- a/Engine/Plugins/Mutable/Source/MutableTools/Private/MuT/ASTOpConstantResource.cpp
+++ b/Engine/Plugins/Mutable/Source/MutableTools/Private/MuT/ASTOpConstantResource.cpp
@@ -76,13 +76,13 @@ namespace UE::Mutable::Private
 	namespace
 	{
  		/** Adds a constant mesh data to a program and returns its constant index. */
-		int32 AddConstantMesh(FProgram& Program, const TSharedPtr<const FMesh>& MeshData, FLinkerOptions& Options)
+		int32 AddConstantMesh(FProgram& Program, const TManagedPtr<const FMesh>& MeshData, FLinkerOptions& Options)
 		{
-			auto AddMeshToProgram = [&Options, &Program](const TSharedPtr<FMesh>& Mesh)
+			auto AddMeshToProgram = [&Options, &Program](const TManagedPtr<FMesh>& Mesh)
 			{
 				// Use a map-based deduplication
 				int32 MeshIndex = -1;
-				TSharedPtr<const FMesh> MeshKey = Mesh;
+				TManagedPtr<const FMesh> MeshKey = Mesh;
 				const int32* MeshIndexPtr = Options.MeshConstantMap.Find(MeshKey);
 				if (!MeshIndexPtr)
 				{
@@ -114,7 +114,7 @@ namespace UE::Mutable::Private
 						EMeshCopyFlags::WithIndexBuffers  |
 						EMeshCopyFlags::WithLayouts;
 
-				TSharedPtr<FMesh> MeshGeometryData = MeshData->Clone(GeometryDataCopyFlags);
+				TManagedPtr<FMesh> MeshGeometryData = MeshData->Clone(GeometryDataCopyFlags);
 
 				// Copy geometry related additional buffers.
 				for (const TPair<EMeshBufferType, FMeshBufferSet>& AdditionalBuffer : MeshData->AdditionalBuffers)
@@ -141,7 +141,7 @@ namespace UE::Mutable::Private
 						EMeshCopyFlags::WithPoses  |
 						EMeshCopyFlags::WithBoneMap;
 
-				TSharedPtr<FMesh> MeshPoseData = MeshData->Clone(PoseDataCopyFlags);
+				TManagedPtr<FMesh> MeshPoseData = MeshData->Clone(PoseDataCopyFlags);
 
 				for (const TPair<EMeshBufferType, FMeshBufferSet>& AdditionalBuffer : MeshData->AdditionalBuffers)
 				{
@@ -165,7 +165,7 @@ namespace UE::Mutable::Private
 						EMeshCopyFlags::WithAdditionalPhysics |
 						EMeshCopyFlags::WithPhysicsBody;
 
-				TSharedPtr<FMesh> MeshPhysicsData = MeshData->Clone(PhysicsDataCopyFlags);
+				TManagedPtr<FMesh> MeshPhysicsData = MeshData->Clone(PhysicsDataCopyFlags);
 			
 				// Copy components related additional buffers.
 				for (const TPair<EMeshBufferType, FMeshBufferSet>& AdditionalBuffer : MeshData->AdditionalBuffers)
@@ -194,7 +194,7 @@ namespace UE::Mutable::Private
 						EMeshCopyFlags::WithSkeletonIDs      |
 						EMeshCopyFlags::WithStreamedResources;
 
-				TSharedPtr<FMesh> MeshMetadataData = MeshData->Clone(MetadataDataCopyFlags);
+				TManagedPtr<FMesh> MeshMetadataData = MeshData->Clone(MetadataDataCopyFlags);
 			
 				// Add a descriptor MeshBufferSet to the metadata part to have formating info.
 				{
@@ -255,7 +255,7 @@ namespace UE::Mutable::Private
 		}   
 
 		/** Adds a constant image data to a program and returns its constant index. */
-		uint32 AddConstantImage(FProgram& Program, const TSharedPtr<const FImage>& Image, FLinkerOptions& Options)
+		uint32 AddConstantImage(FProgram& Program, const TManagedPtr<const FImage>& Image, FLinkerOptions& Options)
 		{
 			MUTABLE_CPUPROFILER_SCOPE(AddConstantImage);
 
@@ -283,7 +283,7 @@ namespace UE::Mutable::Private
 				return Program.ConstantImages.Add(LODRange);
 			}
 
-			auto AddImageToProgram = [&Program, &Options](const TSharedPtr<FImage>& Image) -> int32
+			auto AddImageToProgram = [&Program, &Options](const TManagedPtr<FImage>& Image) -> int32
 			{
 				MUTABLE_CPUPROFILER_SCOPE(Deduplicate);
 				const int32* FoundStoreIndex = Options.ImageConstantMipMap.Find(Image);
@@ -338,7 +338,7 @@ namespace UE::Mutable::Private
 					? EImageFormat::RGBA_UByte
 					: ImageUnCompressedFormat;
 
-			TSharedPtr<const FImage> RawImage = Image;
+			TManagedPtr<const FImage> RawImage = Image;
 			if (ImageUnCompressedFormat != Image->GetFormat())
 			{
 				int32 FirstLODToDecompress = FMath::Min(Image->GetLODCount() - 1, LOD);
@@ -364,13 +364,13 @@ namespace UE::Mutable::Private
 			const int32 RawImageLOD = LOD;
 			if (LOD < LODsToStore)
 			{
-				TSharedPtr<FImage> ImageToStore = ImOp.ExtractMip(RawImage.Get(), LOD - RawImageLOD);
+				TManagedPtr<FImage> ImageToStore = ImOp.ExtractMip(RawImage.Get(), LOD - RawImageLOD);
 				ImageToStore->DataStorage.SetNumLODs(LODRange.NumLODsInTail);
 				++LOD;
 				
 				for (; LOD < LODsToStore; ++LOD)
 				{
-					TSharedPtr<FImage> SrcLODImage = ImOp.ExtractMip(RawImage.Get(), LOD - RawImageLOD);
+					TManagedPtr<FImage> SrcLODImage = ImOp.ExtractMip(RawImage.Get(), LOD - RawImageLOD);
 
 					const int32 DestLOD = LOD - FirstLODToStoreInTail;
 					
@@ -409,11 +409,11 @@ namespace UE::Mutable::Private
 					OP::MeshConstantArgs Args;
 					FMemory::Memzero(Args);
 
-					TSharedPtr<FMesh> MeshValue = StaticCastSharedPtr<const FMesh>(Value)->Clone();
+					TManagedPtr<FMesh> MeshValue = StaticCastManagedPtr<const FMesh>(Value)->Clone();
 					check(MeshValue);
 
 					Args.Skeleton = -1;
-					if (TSharedPtr<const FSkeleton> MeshSkeleton = MeshValue->GetSkeleton())
+					if (TManagedPtr<const FSkeleton> MeshSkeleton = MeshValue->GetSkeleton())
 					{
 						Args.Skeleton = Program.AddConstant(MeshSkeleton);
 						MeshValue->SetSkeleton(nullptr);
@@ -438,7 +438,7 @@ namespace UE::Mutable::Private
 					OP::ResourceConstantArgs Args;
 					FMemory::Memzero(Args);
 					
-					TSharedPtr<const FImage> pTyped = StaticCastSharedPtr<const FImage>(Value);
+					TManagedPtr<const FImage> pTyped = StaticCastManagedPtr<const FImage>(Value);
 					check(pTyped);
 
 					if (pTyped->GetSizeX() * pTyped->GetSizeY() == 0)
@@ -471,7 +471,7 @@ namespace UE::Mutable::Private
 					OP::ResourceConstantArgs Args;
 					FMemory::Memzero(Args);
 					
-					TSharedPtr<const FLayout> pTyped = StaticCastSharedPtr<const FLayout>(Value);
+					TManagedPtr<const FLayout> pTyped = StaticCastManagedPtr<const FLayout>(Value);
 					check(pTyped);
 					Args.value = Program.AddConstant(pTyped);
 
@@ -488,13 +488,13 @@ namespace UE::Mutable::Private
 					OP::MaterialConstantArgs Args;
 					FMemory::Memzero(Args);
 
-					TSharedPtr<FMaterial> MaterialValue = StaticCastSharedPtr<const FMaterial>(Value)->Clone();
+					TManagedPtr<FMaterial> MaterialValue = StaticCastManagedPtr<const FMaterial>(Value)->Clone();
 					check(MaterialValue);
 
 					// Store linked child images into the material constant
 					for (const TPair<FParameterKey, ASTChild>& ImageOperation : ImageOperations)
 					{
-						TVariant<OP::ADDRESS, TSharedPtr<const FImage>> NewImage;
+						TVariant<OP::ADDRESS, TManagedPtr<const FImage>> NewImage;
 						NewImage.Set<OP::ADDRESS>(ImageOperation.Value->linkedAddress);
 
 						FMaterial::FImageParameterData ImageParameterData;
@@ -533,7 +533,7 @@ namespace UE::Mutable::Private
 		if (Type == EOpType::IM_CONSTANT)
 		{
 			// TODO: cache to avoid disk loading
-			TSharedPtr<const FImage> ConstImage = StaticCastSharedPtr<const FImage>(Value);
+			TManagedPtr<const FImage> ConstImage = StaticCastManagedPtr<const FImage>(Value);
 			if (ensure(ConstImage))
 			{
 				Result.m_format = ConstImage->GetFormat();
@@ -563,7 +563,7 @@ namespace UE::Mutable::Private
 		{
 		case EOpType::LA_CONSTANT:
 		{
-			TSharedPtr<const FLayout> pLayout = StaticCastSharedPtr<const FLayout>(Value);
+			TManagedPtr<const FLayout> pLayout = StaticCastManagedPtr<const FLayout>(Value);
 			check(pLayout);
 
 			if (pLayout)
@@ -615,7 +615,7 @@ namespace UE::Mutable::Private
 		if (Type == EOpType::IM_CONSTANT)
 		{
 			// TODO: cache
-			TSharedPtr<const FImage> pMask = StaticCastSharedPtr<const FImage>(Value);
+			TManagedPtr<const FImage> pMask = StaticCastManagedPtr<const FImage>(Value);
 			pMask->GetNonBlackRect(maskUsage);
 			return true;
 		}
@@ -633,7 +633,7 @@ namespace UE::Mutable::Private
 
 		case EOpType::IM_CONSTANT:
 		{
-			TSharedPtr<const FImage> pImage = StaticCastSharedPtr<const FImage>(Value);
+			TManagedPtr<const FImage> pImage = StaticCastManagedPtr<const FImage>(Value);
 			if (pImage->GetSizeX() <= 0 || pImage->GetSizeY() <= 0)
 			{
 				res = true;
@@ -686,14 +686,14 @@ namespace UE::Mutable::Private
 
 
 	//-------------------------------------------------------------------------------------------------
-	TSharedPtr<const FResource> ASTOpConstantResource::GetValue() const
+	TManagedPtr<const FResource> ASTOpConstantResource::GetValue() const
 	{
 		return Value;
 	}
 
 
 	//-------------------------------------------------------------------------------------------------
-	void ASTOpConstantResource::SetValue(const TSharedPtr<const FResource>& InValue)
+	void ASTOpConstantResource::SetValue(const TManagedPtr<const FResource>& InValue)
 	{
 		MUTABLE_CPUPROFILER_SCOPE(ASTOpConstantResource_SetValue);
 
@@ -701,8 +701,8 @@ namespace UE::Mutable::Private
 		{
 		case EOpType::IM_CONSTANT:
 		{
-			TSharedPtr<const FImage> Resource = StaticCastSharedPtr<const FImage>(InValue);
-
+			TManagedPtr<const FImage> Resource = StaticCastManagedPtr<const FImage>(InValue);
+			
 			FOutputHashStream stream;
 			{
 				MUTABLE_CPUPROFILER_SCOPE(Serialize);
@@ -718,7 +718,7 @@ namespace UE::Mutable::Private
 
 		case EOpType::ME_CONSTANT:
 		{
-			TSharedPtr<const FMesh> Resource = StaticCastSharedPtr<const FMesh>(InValue);
+			TManagedPtr<const FMesh> Resource = StaticCastManagedPtr<const FMesh>(InValue);
 
 			FOutputHashStream stream;
 			{
@@ -735,7 +735,7 @@ namespace UE::Mutable::Private
 
 		case EOpType::LA_CONSTANT:
 		{
-			TSharedPtr<const FLayout> Resource = StaticCastSharedPtr<const FLayout>(InValue);
+			TManagedPtr<const FLayout> Resource = StaticCastManagedPtr<const FLayout>(InValue);
 
 			FOutputHashStream stream;
 			{
@@ -752,7 +752,7 @@ namespace UE::Mutable::Private
 		
 		case EOpType::MI_CONSTANT:
 		{
-			TSharedPtr<const FMaterial> Resource = StaticCastSharedPtr<const FMaterial>(InValue);
+			TManagedPtr<const FMaterial> Resource = StaticCastManagedPtr<const FMaterial>(InValue);
 
 			FOutputHashStream stream;
 			{
@@ -780,7 +780,7 @@ namespace UE::Mutable::Private
 		{
 			Ptr<ImageSizeExpression> Res = new ImageSizeExpression;
 			Res->type = ImageSizeExpression::ISET_CONSTANT;
-			TSharedPtr<const FImage> Const = StaticCastSharedPtr<const FImage>(Value);
+			TManagedPtr<const FImage> Const = StaticCastManagedPtr<const FImage>(Value);
 			Res->size = Const->GetSize();
 			return Res;
 		}
@@ -809,7 +809,7 @@ namespace UE::Mutable::Private
 		ASTOp::EClosedMeshTest Result = EClosedMeshTest::Unknown;
 		if (Type == EOpType::ME_CONSTANT)
 		{
-			TSharedPtr<const FMesh> Mesh = StaticCastSharedPtr<const FMesh>(Value);
+			TManagedPtr<const FMesh> Mesh = StaticCastManagedPtr<const FMesh>(Value);
 			if (Mesh)
 			{
 				if (Mesh->IsClosed())
diff --git a/Engine/Plugins/Mutable/Source/MutableTools/Private/MuT/ASTOpConstantResource.h b/Engine/Plugins/Mutable/Source/MutableTools/Private/MuT/ASTOpConstantResource.h
--- a/Engine/Plugins/Mutable/Source/MutableTools/Private/MuT/ASTOpConstantResource.h
+++ b/Engine/Plugins/Mutable/Source/MutableTools/Private/MuT/ASTOpConstantResource.h
@@ -20,7 +20,7 @@ namespace UE::Mutable::Private
 	private:
 
 		//!
-		TSharedPtr<const FResource> Value;
+		TManagedPtr<const FResource> Value;
 
 		//! Value hash
 		uint64 ValueHash;
@@ -49,10 +49,10 @@ namespace UE::Mutable::Private
 		uint64 GetValueHash() const;
 
 		//! Get a pointer to the stored value.
-		TSharedPtr<const FResource> GetValue() const;
+		TManagedPtr<const FResource> GetValue() const;
 
 		/** Set the value to store in this op.	*/
-		void SetValue(const TSharedPtr<const FResource>& Value);
+		void SetValue(const TManagedPtr<const FResource>& Value);
 
 
 		// ASTOp interface
diff --git a/Engine/Plugins/Mutable/Source/MutableTools/Private/MuT/ASTOpMeshFormat.cpp b/Engine/Plugins/Mutable/Source/MutableTools/Private/MuT/ASTOpMeshFormat.cpp
--- a/Engine/Plugins/Mutable/Source/MutableTools/Private/MuT/ASTOpMeshFormat.cpp
+++ b/Engine/Plugins/Mutable/Source/MutableTools/Private/MuT/ASTOpMeshFormat.cpp
@@ -127,9 +127,6 @@ FSourceDataDescriptor ASTOpMeshFormat::GetSourceDataDescriptor(FGetSourceDataDes
 }
 
 
-//---------------------------------------------------------------------------------------------
-//---------------------------------------------------------------------------------------------
-//---------------------------------------------------------------------------------------------
 UE::Mutable::Private::Ptr<ASTOp> Sink_MeshFormatAST::Apply(const ASTOpMeshFormat* root)
 {
 	Root = root;
@@ -151,16 +148,16 @@ UE::Mutable::Private::Ptr<ASTOp> Sink_MeshFormatAST::Apply(const ASTOpMeshFormat
 
 namespace
 {
-	TSharedPtr<const FMesh> FindBaseMeshConstant(UE::Mutable::Private::Ptr<ASTOp> at)
+	TManagedPtr<const FMesh> FindBaseMeshConstant(UE::Mutable::Private::Ptr<ASTOp> at)
 	{
-		TSharedPtr<const FMesh> res;
+		TManagedPtr<const FMesh> Result;
 
 		switch (at->GetOpType())
 		{
 		case EOpType::ME_CONSTANT:
 		{
 			const ASTOpConstantResource* typed = static_cast<const ASTOpConstantResource*>(at.get());
-			res = StaticCastSharedPtr<const FMesh>(typed->GetValue());
+			Result = StaticCastManagedPtr<const FMesh>(typed->GetValue());
 			break;
 		}
 
@@ -168,14 +165,14 @@ namespace
 			check(false);
 		}
 
-		check(res);
+		check(Result);
 
-		return res;
+		return Result;
 	}
 
 
 	// Make a mesh format suitable to morph a particular other format.
-	TSharedPtr<FMesh> MakeMorphTargetFormat(TSharedPtr<const FMesh> pTargetFormat)
+	TManagedPtr<FMesh> MakeMorphTargetFormat(TManagedPtr<const FMesh> pTargetFormat)
 	{
 		MUTABLE_CPUPROFILER_SCOPE(MakeMorphTargetFormat);
 
@@ -215,7 +212,7 @@ namespace
 		}
 
 
-		TSharedPtr<FMesh> pTargetMorphFormat = MakeShared<FMesh>();
+		TManagedPtr<FMesh> pTargetMorphFormat = MakeManaged<FMesh>();
 		pTargetMorphFormat->GetVertexBuffers().SetBufferCount(1);
 
 		pTargetMorphFormat->GetVertexBuffers().SetBuffer(0, offset,
@@ -229,7 +226,7 @@ namespace
 		return pTargetMorphFormat;
 	}
 
-    TSharedPtr<const FMesh> EnsureFormatHasSkinningBuffers(TSharedPtr<const FMesh>& FormatMesh)
+    TManagedPtr<const FMesh> EnsureFormatHasSkinningBuffers(TManagedPtr<const FMesh>& FormatMesh)
     {
         const FMeshBufferSet& FormatMeshVertexBuffers = FormatMesh->GetVertexBuffers();
     
@@ -246,7 +243,7 @@ namespace
             return FormatMesh;
         }
 
-		TSharedPtr<FMesh> NewMesh = FormatMesh->Clone();
+		TManagedPtr<FMesh> NewMesh = FormatMesh->Clone();
         FMeshBufferSet& MeshBuffers = NewMesh->GetVertexBuffers();
         
         FMeshBuffer& Buffer = MeshBuffers.Buffers.AddDefaulted_GetRef();
@@ -353,8 +350,8 @@ UE::Mutable::Private::Ptr<ASTOp> Sink_MeshFormatAST::Visit(const UE::Mutable::Pr
 		NewOp->Base = Visit(NewOp->Base.child(), currentFormatOp);
 
 		// Reformat the morph targets to match the new format.
-		TSharedPtr<const FMesh> pTargetFormat = FindBaseMeshConstant(currentFormatOp->Format.child());
-		TSharedPtr<const FMesh> pTargetMorphFormat = MakeMorphTargetFormat(pTargetFormat);
+		TManagedPtr<const FMesh> pTargetFormat = FindBaseMeshConstant(currentFormatOp->Format.child());
+		TManagedPtr<const FMesh> pTargetMorphFormat = MakeMorphTargetFormat(pTargetFormat);
 
 		UE::Mutable::Private::Ptr<ASTOpConstantResource> NewFormatConstant = new ASTOpConstantResource();
 		NewFormatConstant->Type = EOpType::ME_CONSTANT;
@@ -385,7 +382,7 @@ UE::Mutable::Private::Ptr<ASTOp> Sink_MeshFormatAST::Visit(const UE::Mutable::Pr
 
 	case EOpType::ME_APPLYPOSE:
 	{
-		TSharedPtr<const FMesh> TargetFormat = FindBaseMeshConstant(currentFormatOp->Format.child());
+		TManagedPtr<const FMesh> TargetFormat = FindBaseMeshConstant(currentFormatOp->Format.child());
         TargetFormat = EnsureFormatHasSkinningBuffers(TargetFormat);
 
 		Ptr<ASTOpMeshApplyPose> NewOp = UE::Mutable::Private::Clone<ASTOpMeshApplyPose>(at);
diff --git a/Engine/Plugins/Mutable/Source/MutableTools/Private/MuT/CodeGenerator.cpp b/Engine/Plugins/Mutable/Source/MutableTools/Private/MuT/CodeGenerator.cpp
--- a/Engine/Plugins/Mutable/Source/MutableTools/Private/MuT/CodeGenerator.cpp
+++ b/Engine/Plugins/Mutable/Source/MutableTools/Private/MuT/CodeGenerator.cpp
@@ -415,19 +415,19 @@ namespace UE::Mutable::Private
 	}
 
 
-	TSharedPtr<const FLayout> CodeGenerator::GenerateLayout(Ptr<const NodeLayout> SourceLayout, uint32 MeshIDPrefix)
+	TManagedPtr<const FLayout> CodeGenerator::GenerateLayout(Ptr<const NodeLayout> SourceLayout, uint32 MeshIDPrefix)
 	{
 		// This can run in any thread.
 		UE::TUniqueLock Lock(GenerateLayoutConstantState.Mutex);
 		
-		TSharedPtr<const FLayout>* CachedLayout = GenerateLayoutConstantState.GeneratedLayouts.Find({ SourceLayout,MeshIDPrefix });
+		TManagedPtr<const FLayout>* CachedLayout = GenerateLayoutConstantState.GeneratedLayouts.Find({ SourceLayout,MeshIDPrefix });
 
 		if (CachedLayout)
 		{
 			return *CachedLayout;
 		}
 
-		TSharedPtr<FLayout> GeneratedLayout = SourceLayout->BuildRuntimeLayout(MeshIDPrefix);
+		TManagedPtr<FLayout> GeneratedLayout = SourceLayout->BuildRuntimeLayout(MeshIDPrefix);
 
 		check(GeneratedLayout->Blocks.IsEmpty() || GeneratedLayout->Blocks[0].Id != FLayoutBlock::InvalidBlockId);
 		GenerateLayoutConstantState.GeneratedLayouts.Add({ SourceLayout,MeshIDPrefix }, GeneratedLayout);
@@ -438,7 +438,7 @@ namespace UE::Mutable::Private
 
 	Ptr<ASTOp> CodeGenerator::GenerateImageBlockPatch(Ptr<ASTOp> InBlockOp,
 		const NodeModifierSurfaceEdit::FTexture& Patch,
-		TSharedPtr<FImage> PatchMask,
+		TManagedPtr<FImage> PatchMask,
 		Ptr<ASTOp> conditionAd,
 		const FImageGenerationOptions& ImageOptions )
 	{
@@ -791,10 +791,10 @@ namespace UE::Mutable::Private
 	}
 
 
-	TSharedPtr<FImage> CodeGenerator::GenerateImageBlockPatchMask(const NodeModifierSurfaceEdit::FTexture& Patch, FIntPoint GridSize, int32 BlockPixelsX, int32 BlockPixelsY, box<FIntVector2> RectInCells )
+	TManagedPtr<FImage> CodeGenerator::GenerateImageBlockPatchMask(const NodeModifierSurfaceEdit::FTexture& Patch, FIntPoint GridSize, int32 BlockPixelsX, int32 BlockPixelsY, box<FIntVector2> RectInCells )
 	{
 		// Create a patching mask for the block
-		TSharedPtr<FImage> PatchMask;
+		TManagedPtr<FImage> PatchMask;
 
 		FIntVector2 SourceTextureSize = { GridSize[0] * BlockPixelsX, GridSize[1] * BlockPixelsY };
 
@@ -817,7 +817,7 @@ namespace UE::Mutable::Private
 				FInt32Point BlockSize = BlockRectInPixels.Size();
 				if (!PatchMask)
 				{
-					PatchMask = MakeShared<FImage>(BlockSize[0], BlockSize[1], 1, UE::Mutable::Private::EImageFormat::L_UByte, UE::Mutable::Private::EInitializationType::Black);
+					PatchMask = MakeManaged<FImage>(BlockSize[0], BlockSize[1], 1, UE::Mutable::Private::EImageFormat::L_UByte, UE::Mutable::Private::EInitializationType::Black);
 				}
 
 				uint8* Pixels = PatchMask->GetMipData(0);
@@ -1971,7 +1971,7 @@ namespace UE::Mutable::Private
 				check(OriginalMeshNode);
 
 				// Request morphed skeletal mesh
-				TSharedPtr< TSharedPtr<FMesh> > ResolveMorphedMesh = MakeShared<TSharedPtr<FMesh>>();
+				TSharedPtr< TManagedPtr<FMesh> > ResolveMorphedMesh = MakeShared<TManagedPtr<FMesh>>();
 
 				UE::Tasks::FTask TargetMeshTask;
 				if (OriginalMeshNode->GetType() == NodeMeshConstant::GetStaticType())
@@ -2025,7 +2025,7 @@ namespace UE::Mutable::Private
 							if (OriginalMeshCopy->GetType() == NodeMeshConstant::GetStaticType())
 							{
 								// Target mesh
-								TSharedPtr<FMesh> TargetMesh = *ResolveMorphedMesh;
+								TManagedPtr<FMesh> TargetMesh = *ResolveMorphedMesh;
 								if (!TargetMesh)
 								{
 									return LastMeshResults;
@@ -2162,7 +2162,7 @@ namespace UE::Mutable::Private
 				// TODO: Optimize by formatting and keeping only UVs
 				check(OriginalMeshNode);
 
-				TSharedPtr< TSharedPtr<FMesh> > ResolveOriginalMesh = MakeShared<TSharedPtr<FMesh>>();
+				TSharedPtr< TManagedPtr<FMesh> > ResolveOriginalMesh = MakeShared<TManagedPtr<FMesh>>();
 
 				UE::Tasks::FTask TargetMeshTask;
 				if (OriginalMeshNode->GetType() == NodeMeshConstant::GetStaticType())
@@ -2190,7 +2190,7 @@ namespace UE::Mutable::Private
 				const NodeModifierMeshClipWithUVMask* TypedClipNode = static_cast<const NodeModifierMeshClipWithUVMask*>(Modifier.Node);
 
 				Ptr<ASTOp> MaskImage;
-				TSharedPtr<const FLayout> Layout;
+				TManagedPtr<const FLayout> Layout;
 				if (TypedClipNode->ClipMask)
 				{
 					FImageGenerationOptions ClipOptions(StaticOptions.ComponentId, StaticOptions.LODIndex);
@@ -2219,7 +2219,7 @@ namespace UE::Mutable::Private
 						
 						if (OriginalMeshCopy->GetType() == NodeMeshConstant::GetStaticType())
 						{
-							TSharedPtr<FMesh> TargetMesh = *ResolveOriginalMesh;
+							TManagedPtr<FMesh> TargetMesh = *ResolveOriginalMesh;
 							if (!TargetMesh)
 							{
 								return FMeshGenerationResult();
@@ -2594,7 +2594,7 @@ namespace UE::Mutable::Private
 				if (MatchingEdit->PatchImage.get())
 				{
 					// Does the current block need to be patched? Find out by building a mask.
-					TSharedPtr<FImage> PatchMask = GenerateImageBlockPatchMask(*MatchingEdit, GridSize, LayoutBlockDesc.BlockPixelsX, LayoutBlockDesc.BlockPixelsY, RectInCells);
+					TManagedPtr<FImage> PatchMask = GenerateImageBlockPatchMask(*MatchingEdit, GridSize, LayoutBlockDesc.BlockPixelsX, LayoutBlockDesc.BlockPixelsY, RectInCells);
 
 					if (PatchMask)
 					{
@@ -2702,7 +2702,7 @@ namespace UE::Mutable::Private
 					continue;
 				}
 
-				TSharedPtr<const FLayout> ExtendLayout = ExtraLayouts[LayoutIndex].Layout;
+				TManagedPtr<const FLayout> ExtendLayout = ExtraLayouts[LayoutIndex].Layout;
 
 				Ptr<ASTOp> LastBase = LastImageOp;
 
diff --git a/Engine/Plugins/Mutable/Source/MutableTools/Private/MuT/CodeGenerator.h b/Engine/Plugins/Mutable/Source/MutableTools/Private/MuT/CodeGenerator.h
--- a/Engine/Plugins/Mutable/Source/MutableTools/Private/MuT/CodeGenerator.h
+++ b/Engine/Plugins/Mutable/Source/MutableTools/Private/MuT/CodeGenerator.h
@@ -360,7 +360,7 @@ namespace UE::Mutable::Private
 			*/
 			struct FGeneratedConstantMesh
 			{
-				TSharedPtr<FMesh> Mesh;
+				TManagedPtr<FMesh> Mesh;
 				Ptr<ASTOp> LastMeshOp;
 			};
 			TMap<uint64, TArray<FGeneratedConstantMesh>> GeneratedConstantMeshes;
@@ -487,7 +487,7 @@ namespace UE::Mutable::Private
 
 			/** Layout block that we are trying to generate if any. */
 			uint64 LayoutBlockId = FLayoutBlock::InvalidBlockId;
-			TSharedPtr<const FLayout> LayoutToApply;
+			TManagedPtr<const FLayout> LayoutToApply;
 
 			/**(Optional) Parameter Descriptor of the material where this image belongs. */
 			FParameterDesc MaterialParameter;
@@ -578,7 +578,7 @@ namespace UE::Mutable::Private
     	void GenerateCrop(const FImageGenerationOptions& Options, FImageGenerationResult& Result, const NodeImage&);
 
 		/** */
-		TSharedPtr<FImage> GenerateMissingImage(EImageFormat);
+		TManagedPtr<FImage> GenerateMissingImage(EImageFormat);
 
 		/** */
 		Ptr<ASTOp> GenerateMissingImageCode(const TCHAR* strWhere, EImageFormat, const void* errorContext, const FImageGenerationOptions& Options);
@@ -593,13 +593,13 @@ namespace UE::Mutable::Private
 		Ptr<ASTOp> GenerateImageSize(Ptr<ASTOp>, UE::Math::TIntVector2<int32>);
 
 		/** */
-		Ptr<ASTOp> GenerateImageMulticompose(Ptr<ASTOp> BaseOp, Ptr<ASTOp> LayoutOp, const FImageGenerationOptions& Options, const NodeImagePtrConst& NodeImage, TSharedPtr<const FLayout> SourceLayout, const CodeGenerator::FLayoutBlockDesc& InLayoutBlockDesc );
+		Ptr<ASTOp> GenerateImageMulticompose(Ptr<ASTOp> BaseOp, Ptr<ASTOp> LayoutOp, const FImageGenerationOptions& Options, const NodeImagePtrConst& NodeImage, TManagedPtr<const FLayout> SourceLayout, const CodeGenerator::FLayoutBlockDesc& InLayoutBlockDesc );
 
 		/** Generate a layout block-sized image with a mask including all pixels in the blocks defined in the patch node. */
-		TSharedPtr<FImage> GenerateImageBlockPatchMask(const NodeModifierSurfaceEdit::FTexture&, FIntPoint GridSize, int32 BlockPixelsX, int32 BlockPixelsY, box<FIntVector2> RectInCells);
+		TManagedPtr<FImage> GenerateImageBlockPatchMask(const NodeModifierSurfaceEdit::FTexture&, FIntPoint GridSize, int32 BlockPixelsX, int32 BlockPixelsY, box<FIntVector2> RectInCells);
 
 		/** Generate all the operations to apply the block patching on top of the BlockOp, and masking with PatchMask. */
-		Ptr<ASTOp> GenerateImageBlockPatch(Ptr<ASTOp> BlockOp, const NodeModifierSurfaceEdit::FTexture&, TSharedPtr<FImage> PatchMask, Ptr<ASTOp> ConditionOp, const FImageGenerationOptions&);
+		Ptr<ASTOp> GenerateImageBlockPatch(Ptr<ASTOp> BlockOp, const NodeModifierSurfaceEdit::FTexture&, TManagedPtr<FImage> PatchMask, Ptr<ASTOp> ConditionOp, const FImageGenerationOptions&);
 
 		//-----------------------------------------------------------------------------------------
 		// Materials	
@@ -631,7 +631,7 @@ namespace UE::Mutable::Private
 
 			/** Layout block that we are trying to generate if any. */
 			uint64 LayoutBlockId = FLayoutBlock::InvalidBlockId;
-			TSharedPtr<const FLayout> LayoutToApply;
+			TManagedPtr<const FLayout> LayoutToApply;
 		};
 
 		/** */
@@ -842,7 +842,7 @@ namespace UE::Mutable::Private
 
 				FORCEINLINE bool operator==(const FGeneratedLayoutKey& Other) const = default;
 			};
-			TMap<FGeneratedLayoutKey, TSharedPtr<const FLayout>> GeneratedLayouts;
+			TMap<FGeneratedLayoutKey, TManagedPtr<const FLayout>> GeneratedLayouts;
 		};
 		FGenerateLayoutConstantState GenerateLayoutConstantState;
 
@@ -867,7 +867,7 @@ namespace UE::Mutable::Private
 		Ptr<ASTOp> GenerateLayoutOpsAndResult(const FMeshGenerationDynamicOptions&, Ptr<ASTOp> LastMeshOp, const TArray<Ptr<NodeLayout>>& OriginalLayouts, uint32 MeshId, FMeshGenerationResult& OutResult);
 
     	//!
-    	TSharedPtr<const FLayout> GenerateLayout(Ptr<const NodeLayout> SourceLayout, uint32 MeshIDPrefix);
+    	TManagedPtr<const FLayout> GenerateLayout(Ptr<const NodeLayout> SourceLayout, uint32 MeshIDPrefix);
 
     	//-----------------------------------------------------------------------------------------
     	// Skeletal Mesh
diff --git a/Engine/Plugins/Mutable/Source/MutableTools/Private/MuT/CodeGenerator_ExtensionData.cpp b/Engine/Plugins/Mutable/Source/MutableTools/Private/MuT/CodeGenerator_ExtensionData.cpp
--- a/Engine/Plugins/Mutable/Source/MutableTools/Private/MuT/CodeGenerator_ExtensionData.cpp
+++ b/Engine/Plugins/Mutable/Source/MutableTools/Private/MuT/CodeGenerator_ExtensionData.cpp
@@ -67,11 +67,11 @@ class Node;
 		Ptr<ASTOpConstantExtensionData> Op = new ASTOpConstantExtensionData();
 		OutResult.Op = Op;
 
-		TSharedPtr<const FExtensionData> Data = Constant->Value;
+		TManagedPtr<const FExtensionData> Data = Constant->Value;
 		if (!Data)
 		{
 			// Data can't be null, so make an empty one
-			Data = MakeShared<FExtensionData>();
+			Data = MakeManaged<FExtensionData>();
 			
 			// Log an error message
 			ErrorLog->Add("Constant extension data not set", ELMT_WARNING, Constant->GetMessageContext());
diff --git a/Engine/Plugins/Mutable/Source/MutableTools/Private/MuT/CodeGenerator_FirstPass.h b/Engine/Plugins/Mutable/Source/MutableTools/Private/MuT/CodeGenerator_FirstPass.h
--- a/Engine/Plugins/Mutable/Source/MutableTools/Private/MuT/CodeGenerator_FirstPass.h
+++ b/Engine/Plugins/Mutable/Source/MutableTools/Private/MuT/CodeGenerator_FirstPass.h
@@ -29,6 +29,8 @@
 #include "MuT/ErrorLog.h"
 #include "Async/Mutex.h"
 
+#include "MuR/ManagedPointer.h"
+
 namespace UE::Mutable::Private
 {
 	class FLayout;
@@ -37,7 +39,7 @@ namespace UE::Mutable::Private
 	/** */
 	struct FGeneratedLayout
 	{
-		TSharedPtr<const FLayout> Layout;
+		TManagedPtr<const FLayout> Layout;
 		Ptr<const NodeLayout> Source;
 
 		FORCEINLINE bool operator==(const FGeneratedLayout& Other) const
diff --git a/Engine/Plugins/Mutable/Source/MutableTools/Private/MuT/CodeGenerator_Image.cpp b/Engine/Plugins/Mutable/Source/MutableTools/Private/MuT/CodeGenerator_Image.cpp
--- a/Engine/Plugins/Mutable/Source/MutableTools/Private/MuT/CodeGenerator_Image.cpp
+++ b/Engine/Plugins/Mutable/Source/MutableTools/Private/MuT/CodeGenerator_Image.cpp
@@ -227,7 +227,7 @@ namespace UE::Mutable::Private
 		Result.bCanMulticompose = false;
 
         // TODO: check duplicates
-        TSharedPtr<const FImage> pImage = InNode->Image;
+        TManagedPtr<const FImage> pImage = InNode->Image;
 
         if (!pImage)
         {
@@ -1439,7 +1439,7 @@ namespace UE::Mutable::Private
 			FMeshGenerationResult MeshResult = MeshTask.GetResult();
 
 			// Match the block id of the block we are generating with the id that resulted in the generated mesh			
-			TSharedPtr<const FLayout> Layout = MeshResult.GeneratedLayouts.IsValidIndex(Node->Layout) ? MeshResult.GeneratedLayouts[Node->Layout].Layout : nullptr;
+			TManagedPtr<const FLayout> Layout = MeshResult.GeneratedLayouts.IsValidIndex(Node->Layout) ? MeshResult.GeneratedLayouts[Node->Layout].Layout : nullptr;
 			if (Layout && Layout->Blocks.IsValidIndex(LayoutBlockIndex))
 			{
 				GeneratedLayoutBlockId = Layout->Blocks[LayoutBlockIndex].Id;
@@ -1474,7 +1474,7 @@ namespace UE::Mutable::Private
                 Ptr<ASTOpConstantResource> cop = new ASTOpConstantResource();
                 cop->Type = EOpType::ME_CONSTANT;
 				cop->SourceDataDescriptor = CurrentMeshToProjectOp->GetSourceDataDescriptor();
-				TSharedPtr<FMesh> FormatMeshResult = MakeShared<FMesh>();
+				TManagedPtr<FMesh> FormatMeshResult = MakeManaged<FMesh>();
 				CreateMeshOptimisedForWrappingProjection(FormatMeshResult.Get(), Node->Layout);
 
                 cop->SetValue(FormatMeshResult);
@@ -1506,7 +1506,7 @@ namespace UE::Mutable::Private
                 cop->Type = EOpType::ME_CONSTANT;
 				cop->SourceDataDescriptor = CurrentMeshToProjectOp->GetSourceDataDescriptor();
 
-				TSharedPtr<FMesh> FormatMeshResult = MakeShared<FMesh>();
+				TManagedPtr<FMesh> FormatMeshResult = MakeManaged<FMesh>();
                 CreateMeshOptimisedForProjection(FormatMeshResult.Get(), Node->Layout);
 
                 cop->SetValue(FormatMeshResult);
@@ -1526,7 +1526,7 @@ namespace UE::Mutable::Private
         else
         {
             // This argument is required
-            TSharedPtr<const FMesh> TempMesh = MakeShared<FMesh>();
+            TManagedPtr<const FMesh> TempMesh = MakeManaged<FMesh>();
             Ptr<ASTOpConstantResource> cop = new ASTOpConstantResource();
             cop->Type = EOpType::ME_CONSTANT;
             cop->SetValue(TempMesh);
@@ -1912,7 +1912,7 @@ namespace UE::Mutable::Private
 			[this, InNode, Options, &bCanMulticompose](const NodeImageTable& node, int32 colIndex, int32 row, FErrorLog* LocalErrorLog)
 			{
 				const FTableValue& CellData = node.Table->GetPrivate()->Rows[row].Values[colIndex];
-				TSharedPtr<FImage> pImage = CellData.Image;
+				TManagedPtr<FImage> pImage = CellData.Image;
 
 				Ptr<ASTOp> ImageOp;
 
@@ -2028,14 +2028,14 @@ namespace UE::Mutable::Private
 	}
 
 
-	TSharedPtr<FImage> CodeGenerator::GenerateMissingImage(EImageFormat Format)
+	TManagedPtr<FImage> CodeGenerator::GenerateMissingImage(EImageFormat Format)
 	{
 		MUTABLE_CPUPROFILER_SCOPE(GenerateMissingImage);
 
 		// Make a checkered debug image
 		const FImageSize Size(16, 16);
 
-		TSharedPtr<FImage> GeneratedImage = MakeShared<FImage>(Size[0], Size[1], 1, Format, EInitializationType::NotInitialized);
+		TManagedPtr<FImage> GeneratedImage = MakeManaged<FImage>(Size[0], Size[1], 1, Format, EInitializationType::NotInitialized);
 
 		switch (Format)
 		{
@@ -2122,7 +2122,7 @@ namespace UE::Mutable::Private
 		ErrorLog->Add( Msg, ELMT_ERROR, ErrorContext);
 
 		// Make a checkered debug image
-		TSharedPtr<FImage> GeneratedImage = GenerateMissingImage( format );
+		TManagedPtr<FImage> GeneratedImage = GenerateMissingImage( format );
 
 		Ptr<NodeImageConstant> pNode = new NodeImageConstant();
 		pNode->Image = GeneratedImage;
@@ -2192,13 +2192,13 @@ namespace UE::Mutable::Private
 			Ptr<ASTOp> InLayoutOp, 
 			const FImageGenerationOptions& InBaseImageOptions,
 			const NodeImagePtrConst& NodeImage, 
-			TSharedPtr<const FLayout> SourceLayout, 
+			TManagedPtr<const FLayout> SourceLayout, 
 			const CodeGenerator::FLayoutBlockDesc& InLayoutBlockDesc)
 	{
 		Ptr<ASTOp> Result = InBaseImageOp;
 
 		// Single block layout to generate the entire SourceImage
-		TSharedPtr<FLayout> AuxLayout = SourceLayout->Clone();
+		TManagedPtr<FLayout> AuxLayout = SourceLayout->Clone();
 		AuxLayout->Blocks.SetNum(1);
 		AuxLayout->Blocks[0].Min = FIntVector2::ZeroValue;
 		AuxLayout->Blocks[0].Size = AuxLayout->Size;
diff --git a/Engine/Plugins/Mutable/Source/MutableTools/Private/MuT/CodeGenerator_Material.cpp b/Engine/Plugins/Mutable/Source/MutableTools/Private/MuT/CodeGenerator_Material.cpp
--- a/Engine/Plugins/Mutable/Source/MutableTools/Private/MuT/CodeGenerator_Material.cpp
+++ b/Engine/Plugins/Mutable/Source/MutableTools/Private/MuT/CodeGenerator_Material.cpp
@@ -112,7 +112,7 @@ namespace UE::Mutable::Private
 	void CodeGenerator::GenerateMaterial_Constant(const FMaterialGenerationOptions& Options, FMaterialGenerationResult& Result, const NodeMaterialConstant* Node)
 	{
 		// Constant Material Resource
-		TSharedPtr<FMaterial> MaterialResource = MakeShared<FMaterial>();
+		TManagedPtr<FMaterial> MaterialResource = MakeManaged<FMaterial>();
 		MaterialResource->PassthroughObject = TPassthroughObjectPtr<UMaterialInterface>(Node->MaterialId);
 		MaterialResource->ColorParameters = Node->ColorValues;
 		MaterialResource->ScalarParameters = Node->ScalarValues;
@@ -267,7 +267,7 @@ namespace UE::Mutable::Private
 				{
 					// TODO(Max): UE-314401
 					// Constant Material Resource
-					TSharedPtr<FMaterial> MaterialResource = MakeShared<FMaterial>();
+					TManagedPtr<FMaterial> MaterialResource = MakeManaged<FMaterial>();
 					MaterialResource->PassthroughObject = TPassthroughObjectPtr<UMaterialInterface>(MaterialId);
 
 					// Constant Material Operation
@@ -434,7 +434,7 @@ namespace UE::Mutable::Private
 		Ptr<ASTOpConstantResource> ConstantMaterial = new ASTOpConstantResource();
 		ConstantMaterial->Type = EOpType::MI_CONSTANT;
 
-		TSharedPtr<FMaterial> GenerateMaterial = MakeShared<FMaterial>();
+		TManagedPtr<FMaterial> GenerateMaterial = MakeManaged<FMaterial>();
 		GenerateMaterial->PassthroughObject = {};
 
 		// Store the constant material
@@ -606,7 +606,7 @@ namespace UE::Mutable::Private
 				}
 				else
 				{
-					TSharedPtr<const FLayout> pLayout = MeshResults.GeneratedLayouts[LayoutIndex].Layout;
+					TManagedPtr<const FLayout> pLayout = MeshResults.GeneratedLayouts[LayoutIndex].Layout;
 					check(pLayout);
 
 					// Image
diff --git a/Engine/Plugins/Mutable/Source/MutableTools/Private/MuT/CodeGenerator_Mesh.cpp b/Engine/Plugins/Mutable/Source/MutableTools/Private/MuT/CodeGenerator_Mesh.cpp
--- a/Engine/Plugins/Mutable/Source/MutableTools/Private/MuT/CodeGenerator_Mesh.cpp
+++ b/Engine/Plugins/Mutable/Source/MutableTools/Private/MuT/CodeGenerator_Mesh.cpp
@@ -328,7 +328,7 @@ namespace UE::Mutable::Private
 			return UE::Tasks::MakeCompletedTask<FMeshGenerationResult>();
 		}
 
-		TSharedPtr<const FLayout> Layout;
+		TManagedPtr<const FLayout> Layout;
 		if (Node->Layout)
 		{
 			// Generate the layout with blocks to extract
@@ -473,7 +473,7 @@ namespace UE::Mutable::Private
 			[this] (const NodeMeshTable& node, int32 colIndex, int32 row, FErrorLog* pErrorLog) -> Ptr<NodeMesh>
 			{
 				const FTableValue& Cell = node.Table->GetPrivate()->Rows[row].Values[colIndex];
-				TSharedPtr<UE::Mutable::Private::FMesh> pMesh = Cell.Mesh;
+				TManagedPtr<UE::Mutable::Private::FMesh> pMesh = Cell.Mesh;
 				
 				if (!pMesh)
 				{
@@ -630,7 +630,7 @@ namespace UE::Mutable::Private
 	}
 
 
-	void GenerateLayoutsForMesh(const TSharedPtr<FMesh>& Mesh, const TArray<Ptr<NodeLayout>>& Layouts)
+	void GenerateLayoutsForMesh(const TManagedPtr<FMesh>& Mesh, const TArray<Ptr<NodeLayout>>& Layouts)
 	{
 		if (!Mesh)
 		{
@@ -723,7 +723,7 @@ namespace UE::Mutable::Private
 		// Generate the chain of ops preparing the mesh for the layouts
 		for (int32 LayoutIndex = 0; LayoutIndex < OutResult.GeneratedLayouts.Num(); ++LayoutIndex)
 		{
-			TSharedPtr<const UE::Mutable::Private::FLayout> Layout = OutResult.GeneratedLayouts[LayoutIndex].Layout;
+			TManagedPtr<const UE::Mutable::Private::FLayout> Layout = OutResult.GeneratedLayouts[LayoutIndex].Layout;
 
 			Ptr<ASTOpConstantResource> LayoutOp = new ASTOpConstantResource;
 			LayoutOp->Type = EOpType::LA_CONSTANT;
@@ -749,7 +749,7 @@ namespace UE::Mutable::Private
     {
 		MUTABLE_CPUPROFILER_SCOPE(GenerateMesh_Constant);
 
-		TSharedPtr<FMesh> Mesh = Node->Value;
+		TManagedPtr<FMesh> Mesh = Node->Value;
 		
 		if (!Mesh)
 		{
@@ -762,7 +762,7 @@ namespace UE::Mutable::Private
 			Result.MeshOp = ConstantOp;
 
 			// This data is required
-			TSharedPtr<FMesh> EmptyMesh = MakeShared<FMesh>();
+			TManagedPtr<FMesh> EmptyMesh = MakeManaged<FMesh>();
 			ConstantOp->SetValue(EmptyMesh);
 			EmptyMesh->MeshIDPrefix = UniqueMeshIds.EnsureUnique(ConstantOp->GetValueHash());
 
@@ -775,7 +775,7 @@ namespace UE::Mutable::Private
 		TArray<UE::Tasks::FTask, TInlineAllocator<4>> Requisites;
 		Requisites.Emplace(InOptions);
 
-		TSharedPtr< TSharedPtr<FMesh> > ResolveMesh = MakeShared<TSharedPtr<FMesh>>();
+		TSharedPtr< TManagedPtr<FMesh> > ResolveMesh = MakeShared<TManagedPtr<FMesh>>();
 
 		// Compile-time reference? Add task to resolve as requisite.
 		if (Mesh->IsReference())
@@ -803,7 +803,7 @@ namespace UE::Mutable::Private
 			//&GenerateMeshConstantState = this->GenerateMeshConstantState
 			]() mutable
 			{
-				TSharedPtr<FMesh> Mesh = *ResolveMesh;
+				TManagedPtr<FMesh> Mesh = *ResolveMesh;
 
 				if (!Mesh)
 				{
@@ -835,7 +835,7 @@ namespace UE::Mutable::Private
 
 					BonePosePriority = Mesh->BonePosePriority;
 
-					TSharedPtr<FMesh> MetadatalessMesh = CloneOrTakeOver(Mesh);
+					TManagedPtr<FMesh> MetadatalessMesh = CloneOrTakeOver(MoveTemp(Mesh));
 
 					MetadatalessMesh->Tags.SetNum(0, EAllowShrinking::No);
 					MetadatalessMesh->StreamedResources.SetNum(0, EAllowShrinking::No);
@@ -874,7 +874,7 @@ namespace UE::Mutable::Private
 					LastMeshOp = ConstantOp;
 
 					// We need to clone the mesh in the node because we will modify it.
-					TSharedPtr<FMesh> ClonedMesh = Mesh->Clone();
+					TManagedPtr<FMesh> ClonedMesh = Mesh->Clone();
 					ClonedMesh->EnsureSurfaceData();
 
 					ConstantOp->SetValue(ClonedMesh);
@@ -963,7 +963,7 @@ namespace UE::Mutable::Private
 				op->Source = baseResult.MeshOp;
 				op->Flags = 0;
 
-				TSharedPtr<FMesh> FormatMesh = MakeShared<FMesh>();
+				TManagedPtr<FMesh> FormatMesh = MakeManaged<FMesh>();
 
 				if (Node->VertexBuffers.GetBufferCount())
 				{
@@ -1373,10 +1373,10 @@ namespace UE::Mutable::Private
 		Requisites.Emplace(InOptions);
 
 		// Reference Mesh used to generate the UV layouts.
-		TSharedPtr< TSharedPtr<FMesh> > ReferenceMesh = MakeShared<TSharedPtr<FMesh>>();
+		TSharedPtr< TManagedPtr<FMesh> > ReferenceMesh = MakeShared<TManagedPtr<FMesh>>();
 
 		// Compile-time reference? Add task to resolve as requisite.
-		TSharedPtr<FMesh> Mesh = Node->ReferenceMesh;
+		TManagedPtr<FMesh> Mesh = Node->ReferenceMesh;
 		if (Mesh)
 		{
 			if (Mesh->IsReference())
diff --git a/Engine/Plugins/Mutable/Source/MutableTools/Private/MuT/CodeOptimiser.cpp b/Engine/Plugins/Mutable/Source/MutableTools/Private/MuT/CodeOptimiser.cpp
--- a/Engine/Plugins/Mutable/Source/MutableTools/Private/MuT/CodeOptimiser.cpp
+++ b/Engine/Plugins/Mutable/Source/MutableTools/Private/MuT/CodeOptimiser.cpp
@@ -203,7 +203,7 @@ namespace UE::Mutable::Private
 			{
 				MUTABLE_CPUPROFILER_SCOPE(ConstantMesh);
 
-				TSharedPtr<const FMesh> MeshBuild = System->GetPrivate()->BuildMesh(nullptr, Model, LocalParams.Get(), at, EMeshContentFlags::AllFlags);
+				TManagedPtr<const FMesh> MeshBuild = System->GetPrivate()->BuildMesh(nullptr, Model, LocalParams.Get(), at, EMeshContentFlags::AllFlags);
 
 				if (MeshBuild)
 				{
@@ -212,7 +212,7 @@ namespace UE::Mutable::Private
 					ConstantOp->Type = EOpType::ME_CONSTANT;
 					ConstantOp->SetValue(MeshBuild);
 					Result = ConstantOp;
-				  }
+				}
 				break;
 			}
 
@@ -220,14 +220,14 @@ namespace UE::Mutable::Private
 			{
 				MUTABLE_CPUPROFILER_SCOPE(ConstantImage);
 
-				TSharedPtr<const FImage> pImage = System->GetPrivate()->BuildImage(nullptr, Model, LocalParams.Get(), at, 0);
+				TManagedPtr<const FImage> ImageResult = System->GetPrivate()->BuildImage(nullptr, Model, LocalParams.Get(), at, 0);
 
-				if (pImage)
+				if (ImageResult)
 				{
 					UE::Mutable::Private::Ptr<ASTOpConstantResource> ConstantOp = new ASTOpConstantResource();
 					ConstantOp->SourceDataDescriptor = SourceDataDescriptor;
 					ConstantOp->Type = EOpType::IM_CONSTANT;
-					ConstantOp->SetValue(pImage);
+					ConstantOp->SetValue(ImageResult);
 					Result = ConstantOp;
 				}
 				break;
@@ -237,14 +237,14 @@ namespace UE::Mutable::Private
 			{
 				MUTABLE_CPUPROFILER_SCOPE(ConstantLayout);
 
-				TSharedPtr<const FLayout> pLayout = System->GetPrivate()->BuildLayout(nullptr, Model, LocalParams.Get(), at);
+				TManagedPtr<const FLayout> LayoutResult = System->GetPrivate()->BuildLayout(nullptr, Model, LocalParams.Get(), at);
 
-				if (pLayout)
+				if (LayoutResult)
 				{
-					UE::Mutable::Private::Ptr<ASTOpConstantResource> constantOp = new ASTOpConstantResource();
-					constantOp->Type = EOpType::LA_CONSTANT;
-					constantOp->SetValue( pLayout);
-					Result = constantOp;
+					UE::Mutable::Private::Ptr<ASTOpConstantResource> ConstantOp = new ASTOpConstantResource();
+					ConstantOp->Type = EOpType::LA_CONSTANT;
+					ConstantOp->SetValue(LayoutResult);
+					Result = ConstantOp;
 				}
 				break;
 			}
@@ -537,7 +537,7 @@ namespace UE::Mutable::Private
 						// Instead of generating the constant we resolve the reference, which also replaces the ASTOp.
 						if (bIsReference)
 						{
-							TSharedPtr< TSharedPtr<FImage> > ResolveImage = MakeShared<TSharedPtr<FImage>>();
+							TSharedPtr< TManagedPtr<FImage> > ResolveImage = MakeShared<TManagedPtr<FImage>>();
 
 							constexpr bool bRunImmediatlyIfPossible = false;
 							UE::Tasks::FTask ReferenceCompletion = InOptions->OptimisationOptions.ReferencedImageResourceProvider(ImageID, ResolveImage, bRunImmediatlyIfPossible);
@@ -586,7 +586,7 @@ namespace UE::Mutable::Private
 						// Instead of generating the constant we resolve the reference, which also replaces the ASTOp.
 						if (bIsReference)
 						{
-							TSharedPtr< TSharedPtr<FMesh> > ResolveMesh = MakeShared<TSharedPtr<FMesh>>();
+							TSharedPtr< TManagedPtr<FMesh> > ResolveMesh = MakeShared<TManagedPtr<FMesh>>();
 
 							constexpr bool bRunImmediatlyIfPossible = false;
 							UE::Tasks::FTask ReferenceCompletion = InOptions->OptimisationOptions.ReferencedMeshResourceProvider(MeshID, MeshMorph, ResolveMesh, bRunImmediatlyIfPossible);
@@ -848,11 +848,11 @@ namespace UE::Mutable::Private
 	{
 		UE::Mutable::Private::Ptr<ASTOp> AddMeshOp;
 		TArray<UE::Mutable::Private::Ptr<ASTOpConstantResource>> ContributingMeshes;
-		TSharedPtr<FSkeleton> FinalSkeleton;
+		TManagedPtr<FSkeleton> FinalSkeleton;
 
 		FAddMeshSkeleton( const UE::Mutable::Private::Ptr<ASTOp>& InAddMeshOp,
 						  TArray<UE::Mutable::Private::Ptr<ASTOpConstantResource>>& InContributingMeshes,
-						  const TSharedPtr<FSkeleton>& InFinalSkeleton )
+						  const TManagedPtr<FSkeleton>& InFinalSkeleton )
 		{
 			AddMeshOp = InAddMeshOp;
 			ContributingMeshes = MoveTemp(InContributingMeshes);
@@ -902,11 +902,11 @@ namespace UE::Mutable::Private
 					});
 
 					// Create a mesh just with the unified skeleton
-					TSharedPtr<FSkeleton> FinalSkeleton = MakeShared<FSkeleton>();
+					TManagedPtr<FSkeleton> FinalSkeleton = MakeManaged<FSkeleton>();
 					for (const auto& meshAt: subtreeMeshes)
 					{
-						TSharedPtr<const FMesh> pMesh = StaticCastSharedPtr<const FMesh>(meshAt->GetValue());
-						TSharedPtr<const FSkeleton> SourceSkeleton = pMesh ? pMesh->GetSkeleton() : nullptr;
+						TManagedPtr<const FMesh> pMesh = StaticCastManagedPtr<const FMesh>(meshAt->GetValue());
+						TManagedPtr<const FSkeleton> SourceSkeleton = pMesh ? pMesh->GetSkeleton() : nullptr;
 						if (SourceSkeleton)
 						{
 							ExtendSkeleton(FinalSkeleton.Get(),SourceSkeleton.Get());
@@ -932,10 +932,10 @@ namespace UE::Mutable::Private
 				{
 					if (Rep.ContributingMeshes.Contains(at))
 					{
-						TSharedPtr<const FMesh> pMesh = StaticCastSharedPtr<const FMesh>(typedOp->GetValue());
+						TManagedPtr<const FMesh> pMesh = StaticCastManagedPtr<const FMesh>(typedOp->GetValue());
 						pMesh->CheckIntegrity();
 
-						TSharedPtr<FMesh> NewMesh = MakeShared<FMesh>();
+						TManagedPtr<FMesh> NewMesh = MakeManaged<FMesh>();
 						bool bOutSuccess = false;
 						MeshRemapSkeleton(NewMesh.Get(), pMesh.Get(), Rep.FinalSkeleton, bOutSuccess);
 
@@ -1060,7 +1060,7 @@ namespace UE::Mutable::Private
 				if (n->GetOpType()==EOpType::ME_CONSTANT)
 				{
 					ASTOpConstantResource* typed = static_cast<ASTOpConstantResource*>(n.get());
-					auto pMesh = StaticCastSharedPtr<const FMesh>(typed->GetValue());
+					TManagedPtr<const FMesh> pMesh = StaticCastManagedPtr<const FMesh>(typed->GetValue());
 					pMesh->ResetStaticFormatFlags();
 					typed->SetValue(pMesh);
 				}
diff --git a/Engine/Plugins/Mutable/Source/MutableTools/Private/MuT/CodeOptimiserPrivate.cpp b/Engine/Plugins/Mutable/Source/MutableTools/Private/MuT/CodeOptimiserPrivate.cpp
--- a/Engine/Plugins/Mutable/Source/MutableTools/Private/MuT/CodeOptimiserPrivate.cpp
+++ b/Engine/Plugins/Mutable/Source/MutableTools/Private/MuT/CodeOptimiserPrivate.cpp
@@ -6,12 +6,12 @@ namespace UE::Mutable::Private
 {
 	bool DeduplicateFMaterial(TMap<TPassthroughObjectPtr<UMaterialInterface>, Ptr<ASTOpConstantResource>>& Ops, const Ptr<ASTOpConstantResource>& OpResource)
 	{
-		TSharedPtr<const FMaterial> CurrentValue = StaticCastSharedPtr<const FMaterial>(OpResource->GetValue());
+		TManagedPtr<const FMaterial> CurrentValue = StaticCastManagedPtr<const FMaterial>(OpResource->GetValue());
 		
 		if (Ptr<ASTOpConstantResource>* CachedOpResource = Ops.Find(CurrentValue->PassthroughObject))
 		{
-			TSharedPtr<const FMaterial> CachedOpResourceValue = StaticCastSharedPtr<const FMaterial>(CachedOpResource->get()->GetValue());
-			TSharedPtr<FMaterial> ClonedCachedOpResourceValue = CachedOpResourceValue->Clone();
+			TManagedPtr<const FMaterial> CachedOpResourceValue = StaticCastManagedPtr<const FMaterial>(CachedOpResource->get()->GetValue());
+			TManagedPtr<FMaterial> ClonedCachedOpResourceValue = CachedOpResourceValue->Clone();
 			
 			// Append onto the already processed ASTOp the parameters found in the ASTOp being processed
 			ClonedCachedOpResourceValue->ImageParameters.Append(CurrentValue->ImageParameters);
diff --git a/Engine/Plugins/Mutable/Source/MutableTools/Private/MuT/CodeOptimiserPrivate.h b/Engine/Plugins/Mutable/Source/MutableTools/Private/MuT/CodeOptimiserPrivate.h
--- a/Engine/Plugins/Mutable/Source/MutableTools/Private/MuT/CodeOptimiserPrivate.h
+++ b/Engine/Plugins/Mutable/Source/MutableTools/Private/MuT/CodeOptimiserPrivate.h
@@ -14,7 +14,7 @@ namespace UE::Mutable::Private
 	struct FResourceEntry
 	{
 		uint32 Hash = 0;
-		TSharedPtr<const T> Value;
+		TManagedPtr<const T> Value;
 
 		bool operator==(const FResourceEntry& Entry) const
 		{
@@ -32,7 +32,7 @@ namespace UE::Mutable::Private
 	bool Deduplicate(TMap<FResourceEntry<T>, Ptr<ASTOpConstantResource>>& Ops, const Ptr<ASTOpConstantResource>& OpResource)
 	{
 		uint32 Hash = OpResource->GetValueHash();
-		TSharedPtr<const T> Value = StaticCastSharedPtr<const T>(OpResource->GetValue());
+		TManagedPtr<const T> Value = StaticCastManagedPtr<const T>(OpResource->GetValue());
 
 		FResourceEntry<T> Entry(Hash, Value);
 		
diff --git a/Engine/Plugins/Mutable/Source/MutableTools/Private/MuT/Compiler.cpp b/Engine/Plugins/Mutable/Source/MutableTools/Private/MuT/Compiler.cpp
--- a/Engine/Plugins/Mutable/Source/MutableTools/Private/MuT/Compiler.cpp
+++ b/Engine/Plugins/Mutable/Source/MutableTools/Private/MuT/Compiler.cpp
@@ -745,7 +745,7 @@ namespace UE::Mutable::Private
 		{
 			MUTABLE_CPUPROFILER_SCOPE(GenerateRoms_ImageIds);
 
-			TArray<TSharedPtr<const FImage>> AllMips = MoveTemp(Program.ConstantImageLODsPermanent);
+			TArray<TManagedPtr<const FImage>> AllMips = MoveTemp(Program.ConstantImageLODsPermanent);
 			Program.ConstantImageLODsPermanent = {};
 			Program.ConstantImageLODsPermanent.Reserve(AllMips.Num());
 
@@ -753,7 +753,7 @@ namespace UE::Mutable::Private
 			ParallelFor(AllMips.Num(),
 				[EmbeddedDataBytesLimit, &Program, &AllMips, &RomDatas](uint32 MipIndex)
 				{
-					TSharedPtr<const FImage>& Resource = AllMips[MipIndex];
+					TManagedPtr<const FImage>& Resource = AllMips[MipIndex];
 
 					// Serialize to find out final size of this rom
 					FOutputSizeStream SizeStream;
@@ -890,7 +890,7 @@ namespace UE::Mutable::Private
 		}
 
 		// Meshes
-		TArray<TSharedPtr<const FMesh>> AllMeshes = MoveTemp(Program.ConstantMeshesPermanent);
+		TArray<TManagedPtr<const FMesh>> AllMeshes = MoveTemp(Program.ConstantMeshesPermanent);
 		Program.ConstantMeshesPermanent = {};
 		Program.ConstantMeshesPermanent.Reserve(AllMeshes.Num());
 
@@ -924,7 +924,7 @@ namespace UE::Mutable::Private
 			ParallelFor(AllMeshes.Num(),
 				[EmbeddedDataBytesLimit, &AllMeshes , &Program, &RomDatas](uint32 ResourceIndex)
 				{
-					TSharedPtr<const FMesh>& ResData = AllMeshes[ResourceIndex];
+					TManagedPtr<const FMesh>& ResData = AllMeshes[ResourceIndex];
 
 					// Serialize to memory, to find out final size of this rom
 					const FMesh* Resource = ResData.Get();
diff --git a/Engine/Plugins/Mutable/Source/MutableTools/Private/MuT/DataPacker.cpp b/Engine/Plugins/Mutable/Source/MutableTools/Private/MuT/DataPacker.cpp
--- a/Engine/Plugins/Mutable/Source/MutableTools/Private/MuT/DataPacker.cpp
+++ b/Engine/Plugins/Mutable/Source/MutableTools/Private/MuT/DataPacker.cpp
@@ -536,13 +536,13 @@ namespace UE::Mutable::Private
             if (n->GetOpType()==EOpType::IM_CONSTANT)
             {
 				ASTOpConstantResource* Typed = static_cast<ASTOpConstantResource*>(n.get());
-                TSharedPtr<const FImage> pOld = StaticCastSharedPtr<const FImage>(Typed->GetValue());
+				TManagedPtr<const FImage> pOld = StaticCastManagedPtr<const FImage>(Typed->GetValue());
 
 				FImageOperator ImOp = FImageOperator::GetDefault( Options->GetPrivate()->ImageFormatFunc );
 
 				// See if there is a better format for this image
 				FVector4f PlainColor;
-				if ( pOld->IsPlainColour(PlainColor) )
+				if (pOld->IsPlainColour(PlainColor))
 				{
 					// It is more efficient to just have an instruction for it instead, to avoid the overhead
 					// of data loading. 
@@ -560,23 +560,23 @@ namespace UE::Mutable::Private
 					ASTOp::Replace(n, NewPlain);
 
 				}
-				else if (ImageFormatAccumulator.SupportedFormats[Typed][(int32)EImageFormat::L_UBitRLE] )
-                {
-                    TSharedPtr<FImage> pNew = ImOp.ImagePixelFormat( ImageCompressionQuality, pOld.Get(), EImageFormat::L_UBitRLE );
+				else if (ImageFormatAccumulator.SupportedFormats[Typed][(int32)EImageFormat::L_UBitRLE])
+				{
+					TManagedPtr<FImage> pNew = ImOp.ImagePixelFormat( ImageCompressionQuality, pOld.Get(), EImageFormat::L_UBitRLE );
 
-                    // Only replace if the compression was worth!
-                    int32 oldSize = pOld->GetDataSize();
+					// Only replace if the compression was worth!
+					int32 oldSize = pOld->GetDataSize();
 					int32 newSize = pNew->GetDataSize();
-                    if (float(oldSize) > float(newSize) * OptimizeOptions.MinRLECompressionGain)
-                    {
+					if (float(oldSize) > float(newSize) * OptimizeOptions.MinRLECompressionGain)
+					{
 						Typed->SetValue(pNew);
-                    }
-                }
-                else if (ImageFormatAccumulator.SupportedFormats[Typed][(int32)EImageFormat::L_UByteRLE] )
-                {
-                    TSharedPtr<FImage> pNew = ImOp.ImagePixelFormat( ImageCompressionQuality, pOld.Get(), EImageFormat::L_UByteRLE );
+					}
+				}
+				else if (ImageFormatAccumulator.SupportedFormats[Typed][(int32)EImageFormat::L_UByteRLE])
+				{
+					TManagedPtr<FImage> pNew = ImOp.ImagePixelFormat( ImageCompressionQuality, pOld.Get(), EImageFormat::L_UByteRLE );
 
-                    // Only replace if the compression was worth!
+					// Only replace if the compression was worth!
 					int32 oldSize = pOld->GetDataSize();
 					int32 newSize = pNew->GetDataSize();
                     if (float(oldSize) > float(newSize) * OptimizeOptions.MinRLECompressionGain)
@@ -596,7 +596,7 @@ namespace UE::Mutable::Private
 		for (const TPair<Ptr<ASTOpConstantResource>, uint64>& Entry : MeshSemanticsAccumulator.RequiredSemanticsPerConstant)
 		{
 			ASTOpConstantResource* Op = Entry.Key.get();
-			TSharedPtr<FMesh> Mesh = static_cast<const FMesh*>(Op->GetValue().Get())->Clone();
+			TManagedPtr<FMesh> Mesh = StaticCastManagedPtr<const FMesh>(Op->GetValue())->Clone();
 			MeshRemoveUnusedBufferSemantics(Mesh.Get(), Entry.Value);
 			Op->SetValue(Mesh);
 		}
diff --git a/Engine/Plugins/Mutable/Source/MutableTools/Private/MuT/NodeLayout.cpp b/Engine/Plugins/Mutable/Source/MutableTools/Private/MuT/NodeLayout.cpp
--- a/Engine/Plugins/Mutable/Source/MutableTools/Private/MuT/NodeLayout.cpp
+++ b/Engine/Plugins/Mutable/Source/MutableTools/Private/MuT/NodeLayout.cpp
@@ -20,9 +20,9 @@ namespace UE::Mutable::Private
 	FNodeType NodeLayout::StaticType = FNodeType(Node::EType::Layout, Node::GetStaticType() );
 
 
-	TSharedPtr<FLayout> NodeLayout::BuildRuntimeLayout(uint32 BlockIDPrefix) const
+	TManagedPtr<FLayout> NodeLayout::BuildRuntimeLayout(uint32 BlockIDPrefix) const
 	{
-		TSharedPtr<FLayout> GeneratedLayout = MakeShared<FLayout>();
+		TManagedPtr<FLayout> GeneratedLayout = MakeManaged<FLayout>();
 		GeneratedLayout->Size = Size;
 		GeneratedLayout->MaxSize = MaxSize;
 		GeneratedLayout->Strategy = Strategy;
@@ -55,7 +55,7 @@ namespace UE::Mutable::Private
 	}
 
 
-	void NodeLayout::GenerateLayoutBlocks(const TSharedPtr<FMesh>& Mesh, int32 LayoutIndex )
+	void NodeLayout::GenerateLayoutBlocks(const UE::Mutable::Private::TManagedPtr<FMesh>& Mesh, int32 LayoutIndex )
 	{
 		if (!(Mesh && LayoutIndex >= 0 && Size.X > 0 && Size.Y > 0))
 		{
@@ -68,7 +68,7 @@ namespace UE::Mutable::Private
 		{
 			VerticesAlreadyClassified.SetNumZeroed(VertexCount);
 
-			TSharedPtr<FLayout> RuntimeLayout = BuildRuntimeLayout(0);
+			TManagedPtr<FLayout> RuntimeLayout = BuildRuntimeLayout(0);
 			FMutableLayoutClassifyResult ClassifyResult;
 			bool bNormalizeUVs = false;
 			bool bClampUVIslands = true;
@@ -293,7 +293,7 @@ namespace UE::Mutable::Private
 	}
 
 
-	void NodeLayout::GenerateLayoutBlocksFromUVIslands(const TSharedPtr<FMesh>& Mesh, int32 LayoutIndex)
+	void NodeLayout::GenerateLayoutBlocksFromUVIslands(const UE::Mutable::Private::TManagedPtr<FMesh>& Mesh, int32 LayoutIndex)
 	{
 		if (!(Mesh && LayoutIndex >= 0 && Size.X > 0 && Size.Y > 0))
 		{
@@ -555,7 +555,7 @@ namespace UE::Mutable::Private
 			// TODO: Size?
 			int32 SizeX = 1024;
 			int32 SizeY = 1024;
-			TSharedPtr<FImage> Mask = MakeShared<FImage>(SizeX, SizeY, 1, EImageFormat::L_UByte, EInitializationType::Black);
+			TManagedPtr<FImage> Mask = MakeManaged<FImage>(SizeX, SizeY, 1, EImageFormat::L_UByte, EInitializationType::Black);
 			const TArrayView<uint8> ImageData = Mask->DataStorage.GetLOD(0);
 
 			for (int32 TriangleIndex = 0; TriangleIndex < NumTriangles; ++TriangleIndex)
diff --git a/Engine/Plugins/Mutable/Source/MutableTools/Private/MuT/Table.cpp b/Engine/Plugins/Mutable/Source/MutableTools/Private/MuT/Table.cpp
--- a/Engine/Plugins/Mutable/Source/MutableTools/Private/MuT/Table.cpp
+++ b/Engine/Plugins/Mutable/Source/MutableTools/Private/MuT/Table.cpp
@@ -139,7 +139,7 @@ namespace UE::Mutable::Private
 	}
 
 
-    void FTable::SetCell(int32 Column, uint32 RowId, const TSharedPtr<FImage>& Value, const void* ErrorContext)
+    void FTable::SetCell(int32 Column, uint32 RowId, const UE::Mutable::Private::TManagedPtr<FImage>& Value, const void* ErrorContext)
 	{
 		int32 Row = m_pD->FindRow(RowId);
 		check( Row>=0 );
@@ -152,7 +152,7 @@ namespace UE::Mutable::Private
 	}
 
 
-    void FTable::SetCell( int32 Column, uint32 RowId, const TSharedPtr<FMesh>& Value, const void* ErrorContext )
+    void FTable::SetCell( int32 Column, uint32 RowId, const UE::Mutable::Private::TManagedPtr<FMesh>& Value, const void* ErrorContext )
 	{
 		int32 Row = m_pD->FindRow(RowId);
 		check( Row>=0 );
diff --git a/Engine/Plugins/Mutable/Source/MutableTools/Private/MuT/UnrealPixelFormatOverride.cpp b/Engine/Plugins/Mutable/Source/MutableTools/Private/MuT/UnrealPixelFormatOverride.cpp
--- a/Engine/Plugins/Mutable/Source/MutableTools/Private/MuT/UnrealPixelFormatOverride.cpp
+++ b/Engine/Plugins/Mutable/Source/MutableTools/Private/MuT/UnrealPixelFormatOverride.cpp
@@ -208,7 +208,7 @@ void MutableToImageCore(const UE::Mutable::Private::FImage* InMutable, FImage& C
 {
 	MUTABLE_CPUPROFILER_SCOPE(MutableToImageCore);
 
-	TSharedPtr<UE::Mutable::Private::FImage> TempMutable;
+	UE::Mutable::Private::TManagedPtr<UE::Mutable::Private::FImage> TempMutable;
 
 	ERawImageFormat::Type CoreImageFormat;
 	switch (InMutable->GetFormat())
