[Mutable] OpImageMultiCompose optimizations.

#rnx
#Synced-CL Unknown

diff --git a/Engine/Plugins/Mutable/Source/MutableRuntime/Internal/MuR/ParallelExecutionUtils.h b/Engine/Plugins/Mutable/Source/MutableRuntime/Internal/MuR/ParallelExecutionUtils.h
--- a/Engine/Plugins/Mutable/Source/MutableRuntime/Internal/MuR/ParallelExecutionUtils.h
+++ b/Engine/Plugins/Mutable/Source/MutableRuntime/Internal/MuR/ParallelExecutionUtils.h
@@ -7,6 +7,7 @@
 namespace UE::Mutable::Private::ParallelExecutionUtils
 {
     void InvokeBatchParallelFor(int32 Num, TFunctionRef<void(int32)> Body);
+    void InvokeBatchParallelForUnbalanced(int32 Num, TFunctionRef<void(int32)> Body);
     void InvokeBatchParallelForSingleThread(int32 Num, TFunctionRef<void(int32)> Body);
 }
 
diff --git a/Engine/Plugins/Mutable/Source/MutableRuntime/Private/MuR/OpImageMultiCompose.cpp b/Engine/Plugins/Mutable/Source/MutableRuntime/Private/MuR/OpImageMultiCompose.cpp
--- a/Engine/Plugins/Mutable/Source/MutableRuntime/Private/MuR/OpImageMultiCompose.cpp
+++ b/Engine/Plugins/Mutable/Source/MutableRuntime/Private/MuR/OpImageMultiCompose.cpp
@@ -6,6 +6,7 @@
 #include "MuR/MutableMath.h"
 #include "MuR/ImageFormatUtils.h"
 #include "MuR/ImageResizeUtils.h"
+#include "MuR/ParallelExecutionUtils.h"
 
 namespace UE::Mutable::Private
 {
@@ -38,6 +39,8 @@ namespace OpImageMultiComposeInternal
 		SubImageResize::SubImageResizeFuncType* ResizeFunc;
 		SubImageFormat::UncompressedFormatFuncType* SrcUncompressedFormatFunc;
 		SubImageFormat::UncompressedFormatFuncType* DstUncompressedFormatFunc;
+
+		bool bIsUncompressedReformat;
 	};
 
 	template<typename T>
@@ -71,12 +74,17 @@ namespace OpImageMultiComposeInternal
 
 	struct FScratchMemoryAllocation
 	{
+		static constexpr int32 MaxAllowedJobs = 8; 
+
+		FIntVector2 BatchSizeInBlocks;
+		int32 MaxParallelJobs; 
+
 		EActiveStages ActiveStages = EActiveStages::None; 
 		
 		TArray<uint8, FDefaultMemoryTrackingAllocator<MemoryCounters::FImageMemoryCounter>> Buffer;
 
-		TArrayView<uint8> SrcDecompressed;
-		TArrayView<uint8> DstDecompressed;
+		TArray<TArrayView<uint8>, TInlineAllocator<MaxAllowedJobs>> SrcDecompressedViews;
+		TArray<TArrayView<uint8>, TInlineAllocator<MaxAllowedJobs>> DstDecompressedViews;		
 	};
 
 	FImageFormatInfo GatherImageFormatInfo(const FImage& Dst, const FImage& Src)
@@ -86,20 +94,27 @@ namespace OpImageMultiComposeInternal
 		EImageFormat SrcDecompressedFormat = GetUncompressedFormat(Src.GetFormat());
 		EImageFormat DstDecompressedFormat = GetUncompressedFormat(Dst.GetFormat());
 		EImageFormat MostGenericFormat     = GetMostGenericFormat(SrcDecompressedFormat, DstDecompressedFormat);
+
+		const bool bIsUncompressedReformat = 
+				SrcDecompressedFormat != DstDecompressedFormat && 
+				SrcDecompressedFormat == Src.GetFormat()       &&  
+				DstDecompressedFormat == Dst.GetFormat();
+
 		return FImageFormatInfo
 		{
-			.Dst                   = GetImageFormatData(Dst.GetFormat()),
-			.Src                   = GetImageFormatData(Src.GetFormat()),
-			.DstDecompressed       = GetImageFormatData(DstDecompressedFormat),
-			.SrcDecompressedFormat = SrcDecompressedFormat,
-			.DstDecompressedFormat = DstDecompressedFormat,
-			.MostGenericFormat     = MostGenericFormat,
-			.SrcDecompressionFunc  = SelectDecompressionFunction(DstDecompressedFormat, Src.GetFormat()),
-			.DstDecompressionFunc  = SelectDecompressionFunction(DstDecompressedFormat, Dst.GetFormat()),
-			.DstCompressionFunc    = SelectCompressionFunction(Dst.GetFormat(), DstDecompressedFormat),
-			.ResizeFunc            = SubImageResize::SelectSubImageResizeFunc(DstDecompressedFormat),
+			.Dst                       = GetImageFormatData(Dst.GetFormat()),
+			.Src                       = GetImageFormatData(Src.GetFormat()),
+			.DstDecompressed           = GetImageFormatData(DstDecompressedFormat),
+			.SrcDecompressedFormat     = SrcDecompressedFormat,
+			.DstDecompressedFormat     = DstDecompressedFormat,
+			.MostGenericFormat         = MostGenericFormat,
+			.SrcDecompressionFunc      = SelectDecompressionFunction(DstDecompressedFormat, Src.GetFormat()),
+			.DstDecompressionFunc      = SelectDecompressionFunction(DstDecompressedFormat, Dst.GetFormat()),
+			.DstCompressionFunc        = SelectCompressionFunction(Dst.GetFormat(), DstDecompressedFormat),
+			.ResizeFunc                = SubImageResize::SelectSubImageResizeFunc(DstDecompressedFormat),
 			.SrcUncompressedFormatFunc = SubImageFormat::SelectUncompressedFormatFunction(SrcDecompressedFormat, DstDecompressedFormat),
 			.DstUncompressedFormatFunc = SubImageFormat::SelectUncompressedFormatFunction(DstDecompressedFormat, MostGenericFormat),
+			.bIsUncompressedReformat   = bIsUncompressedReformat,
 		};
 	}
 
@@ -159,45 +174,60 @@ namespace OpImageMultiComposeInternal
 	FScratchMemoryAllocation AllocateScratchMemory(const FImage& Dst, const FImage& Src, TConstArrayView<FIntRect> DstRects, TConstArrayView<FIntRect> SrcRects, const FImageFormatInfo& FormatInfo)
 	{
 		MUTABLE_CPUPROFILER_SCOPE(ImageMultiCompose_AllocateScratchMemory);
-		
-		// TODO: For now allocate for the largest block. This should allocate only for a specific memory target
-		// splitting the work in smaller batches.
+	
+		// This values can be tuned to control the maximum amount of scratch memory it will be used.
+		// Smaller batches may get better data locality but extra source block decompression from the 
+		// batch edge will be done. The current values seem to give good results with the tested cases.
+		constexpr int32 MaxNumBlocksPerBatchX = 1 << 4; 
+		constexpr int32 MaxNumBlocksPerBatchY = 1 << 3; 
 
 		FScratchMemoryAllocation Result;
-		
-		uint32 SrcDecompressedAllocSize = 0; 
-		uint32 DstDecompressedAllocSize = 0; 
+		Result.BatchSizeInBlocks = FIntVector2(MaxNumBlocksPerBatchX, MaxNumBlocksPerBatchY);
+	
+		FIntVector2 BatchSizeInPixels = Result.BatchSizeInBlocks * FIntVector2(FormatInfo.Dst.PixelsPerBlockX, FormatInfo.Dst.PixelsPerBlockY);
 
-		EActiveStages AllocationActiveStages = EActiveStages::None;
+		uint32 DstDecompressedAllocSize = BatchSizeInPixels.X * BatchSizeInPixels.Y * FormatInfo.DstDecompressed.BytesPerBlock;
 
+		uint32 SrcDecompressedAllocSize = 0; 
 		for (int32 RectIndex = 0, NumRects = DstRects.Num(); RectIndex < NumRects; ++RectIndex)
 		{
 			FIntRect DstRect = DstRects[RectIndex];
 			FIntRect SrcRect = SrcRects[RectIndex];
+	
+			if (DstRect.Area() == 0)
+			{
+				continue;
+			}
 
-			FIntRect AlignedSrcRect = FIntRect(
-					AlignDown(SrcRect.Min.X, FormatInfo.Src.PixelsPerBlockX),
-					AlignDown(SrcRect.Min.Y, FormatInfo.Src.PixelsPerBlockY),
-					AlignUp  (SrcRect.Max.X, FormatInfo.Src.PixelsPerBlockX),
-					AlignUp  (SrcRect.Max.Y, FormatInfo.Src.PixelsPerBlockY));
-			
-			FIntRect AlignedDstRect = FIntRect(
-					AlignDown(DstRect.Min.X, FormatInfo.Dst.PixelsPerBlockX),
-					AlignDown(DstRect.Min.Y, FormatInfo.Dst.PixelsPerBlockY),
-					AlignUp  (DstRect.Max.X, FormatInfo.Dst.PixelsPerBlockX),
-					AlignUp  (DstRect.Max.Y, FormatInfo.Dst.PixelsPerBlockY));
+			// Resize Dst Batch to get the Src 
+			FIntVector2 SrcDecBatchSize = BatchSizeInPixels * FIntVector2(SrcRect.Width(), SrcRect.Height());
+
+			SrcDecBatchSize = FIntVector2(
+					FMath::DivideAndRoundUp(SrcDecBatchSize.X, DstRect.Width()), 
+					FMath::DivideAndRoundUp(SrcDecBatchSize.Y, DstRect.Height()));
 
-			uint32 AlignedDecSrcRectBytes = AlignedSrcRect.Area()*FormatInfo.DstDecompressed.BytesPerBlock;
-			uint32 AlignedDecDstRectBytes = AlignedDstRect.Area()*FormatInfo.DstDecompressed.BytesPerBlock;
+			// Add 1 compression blocks per side to account for alignment. 
+			SrcDecBatchSize += FIntVector2(FormatInfo.Src.PixelsPerBlockX*2, FormatInfo.Src.PixelsPerBlockY*2);
+			
+			uint32 SrcDecBatchBytes = SrcDecBatchSize.X * SrcDecBatchSize.Y * FormatInfo.DstDecompressed.BytesPerBlock;
 
-			SrcDecompressedAllocSize = FMath::Max(AlignedDecSrcRectBytes, SrcDecompressedAllocSize);
-			DstDecompressedAllocSize = FMath::Max(AlignedDecDstRectBytes, DstDecompressedAllocSize);
+			SrcDecompressedAllocSize = FMath::Max(SrcDecBatchBytes, SrcDecompressedAllocSize);
 		}
+		
+		const int32 NumWorkers = LowLevelTasks::FScheduler::Get().GetNumWorkers();
+		Result.MaxParallelJobs = FMath::Clamp(NumWorkers - 1, 1, FScratchMemoryAllocation::MaxAllowedJobs);  
 
-		Result.Buffer.SetNumUninitialized(SrcDecompressedAllocSize + DstDecompressedAllocSize);
+		// Add a cache line of padding per job to reduce possible false sharing between parallel batches.
+		uint32 AllocationSizeInBytes = (SrcDecompressedAllocSize + DstDecompressedAllocSize + PLATFORM_CACHE_LINE_SIZE) * Result.MaxParallelJobs;
 
-		Result.SrcDecompressed = TArrayView(Result.Buffer.GetData(), SrcDecompressedAllocSize);
-		Result.DstDecompressed = TArrayView(Result.Buffer.GetData() + SrcDecompressedAllocSize, DstDecompressedAllocSize);    
+		Result.Buffer.SetNumUninitialized(AllocationSizeInBytes);
+
+		for (int32 JobId = 0; JobId < Result.MaxParallelJobs; ++JobId)
+		{
+			uint32 JobViewOffset = (SrcDecompressedAllocSize + DstDecompressedAllocSize + PLATFORM_CACHE_LINE_SIZE) * JobId;
+			Result.SrcDecompressedViews.Emplace(Result.Buffer.GetData() + JobViewOffset, SrcDecompressedAllocSize);
+			Result.DstDecompressedViews.Emplace(Result.Buffer.GetData() + JobViewOffset + SrcDecompressedAllocSize, DstDecompressedAllocSize);
+		}
 
 		// TODO: We are allocating for all possible stages but should only allocate for the used ones. 
 		Result.ActiveStages = EActiveStages::None;
@@ -232,10 +262,26 @@ namespace OpImageMultiComposeInternal
 		}
 	}
 
-	int32 ComputeNumBatches(const FIntRect& DstRect)
-	{
-		//TODO: Do batching to reduce memory requirements and parallelise.
-		return 1;	
+	FIntVector2 ComputeNumBatches(const FIntRect& DstRect, const FImageFormatInfo& FormatInfo, const FScratchMemoryAllocation& ScratchMemory)
+	{ 
+		FIntVector2 BatchSizeInBlocks = ScratchMemory.BatchSizeInBlocks;
+		FIntVector2 DstBlockSize      = FIntVector2(FormatInfo.Dst.PixelsPerBlockX, FormatInfo.Dst.PixelsPerBlockY);
+
+		FIntRect AlignedRect = FIntRect(
+				AlignDown(DstRect.Min.X, DstBlockSize.X),
+				AlignDown(DstRect.Min.Y, DstBlockSize.Y),
+				AlignUp  (DstRect.Max.X, DstBlockSize.X),
+				AlignUp  (DstRect.Max.Y, DstBlockSize.Y));
+
+		FIntVector2 RectSizeInBlocks = FIntVector2(
+				FMath::DivideAndRoundUp(AlignedRect.Width(),  DstBlockSize.X),
+				FMath::DivideAndRoundUp(AlignedRect.Height(), DstBlockSize.Y));
+
+		FIntVector2 NumBatches = FIntVector2(
+				FMath::DivideAndRoundUp(RectSizeInBlocks.X, BatchSizeInBlocks.X),
+				FMath::DivideAndRoundUp(RectSizeInBlocks.Y, BatchSizeInBlocks.Y));
+
+		return NumBatches;
 	}
 
 	struct FBatchWorkingSubImages
@@ -248,83 +294,124 @@ namespace OpImageMultiComposeInternal
 		TSubImage<uint8> DecDstAligned;
 
 		FVector2f SrcSubPixelOffset = FVector2f(0, 0);
+
+		EActiveStages ActiveStages;
 	};
 
 	FBatchWorkingSubImages ComputeWorkingSubImages(
-			int32 BatchIndex, 
+			int32 JobId,
+			FIntVector2 BatchId, 
 			FIntRect DstRect, FIntRect SrcRect, 
 			const TSubImage<uint8>& DstImage, const TConstSubImage<uint8>& SrcImage, 
 			const FImageFormatInfo& FormatInfo, 
-			const FScratchMemoryAllocation& ScratchMemory)
+			const FScratchMemoryAllocation& ScratchMemory,
+			EActiveStages ActiveStages)
 	{
 		FBatchWorkingSubImages Result;
+		Result.ActiveStages =  ActiveStages;
+
+		FIntVector2 BatchSizeInPixels = ScratchMemory.BatchSizeInBlocks * 
+				FIntVector2(FormatInfo.Dst.PixelsPerBlockX, FormatInfo.Dst.PixelsPerBlockY);
+	
+		FIntRect AlignedDstRect = FIntRect(
+				AlignDown(DstRect.Min.X, DstImage.NumElemsPerBlockX),
+				AlignDown(DstRect.Min.Y, DstImage.NumElemsPerBlockY),
+				AlignUp  (DstRect.Max.X, DstImage.NumElemsPerBlockX),
+				AlignUp  (DstRect.Max.Y, DstImage.NumElemsPerBlockY));	
+
+		// Dst batches are always aligned to the compression blocks.
+		FIntRect DstBatchRect = FIntRect(
+				BatchId.X*BatchSizeInPixels.X + AlignedDstRect.Min.X,
+				BatchId.Y*BatchSizeInPixels.Y + AlignedDstRect.Min.Y,
+				BatchId.X*BatchSizeInPixels.X + BatchSizeInPixels.X + AlignedDstRect.Min.X,
+				BatchId.Y*BatchSizeInPixels.Y + BatchSizeInPixels.Y + AlignedDstRect.Min.Y);
+		
+		// Clip the batch rect to the original dst layout rect to get the actual region to process.
+		FIntRect DstClippedBatchRect = DstBatchRect;
+		DstClippedBatchRect.Clip(DstRect);
+
+		// If all the batch area will be modified there is no need to decompress the dst image for this batch.
+		// Remove the stage from this batch ActiveStages.
+		if (DstBatchRect == DstClippedBatchRect)
+		{
+			EnumRemoveFlags(Result.ActiveStages, EActiveStages::DstDecompression);
+		}
+
+		// Compute the clipped dst rect coordinates respect the original layout rect to calculate the resizing
+		// sizes and offsets.
+		FBox2f DstRectForResize = FBox2f(
+				FVector2f(DstClippedBatchRect.Min.X - DstRect.Min.X, DstClippedBatchRect.Min.Y - DstRect.Min.Y),
+				FVector2f(DstClippedBatchRect.Max.X - DstRect.Min.X, DstClippedBatchRect.Max.Y - DstRect.Min.Y));
 
-		// TODO: For now assume there is only one batch for the whole Rect. This should be split in smaller batches
-		// to reduce memory requirements.
-		check(BatchIndex < 1); 
-		FIntRect DstBatchRect = DstRect;
+		FVector2f OneOverResizeFactor = FVector2f(SrcRect.Width(), SrcRect.Height()) / FVector2f(DstRect.Width(), DstRect.Height());
+		FBox2f SrcRectResized = FBox2f(DstRectForResize.Min * OneOverResizeFactor, DstRectForResize.Max * OneOverResizeFactor);
 
-		FIntRect SrcBatchRect = SrcRect;
+		Result.SrcSubPixelOffset = FVector2f(FMath::Frac(SrcRectResized.Min.X), FMath::Frac(SrcRectResized.Min.Y));
 
-		check(SrcBatchRect.Area() == SrcRect.Area())
+		// Translate the resized batch to src image coordinates.
+		FIntRect SrcBatchRect = FIntRect(
+				FMath::FloorToInt32(SrcRectResized.Min.X), 
+				FMath::FloorToInt32(SrcRectResized.Min.Y), 
+				FMath::CeilToInt32(SrcRectResized.Max.X ), 
+				FMath::CeilToInt32(SrcRectResized.Max.Y )) + SrcRect.Min; 
 
-		FIntRect AlignedSrcRect = FIntRect(
+		FIntRect AlignedBatchSrcRect = FIntRect(
 				AlignDown(SrcBatchRect.Min.X, SrcImage.NumElemsPerBlockX),
 				AlignDown(SrcBatchRect.Min.Y, SrcImage.NumElemsPerBlockY),
 				AlignUp  (SrcBatchRect.Max.X, SrcImage.NumElemsPerBlockX),
 				AlignUp  (SrcBatchRect.Max.Y, SrcImage.NumElemsPerBlockY));
-			
-		FIntRect AlignedDstRect = FIntRect(
-				AlignDown(DstBatchRect.Min.X, DstImage.NumElemsPerBlockX),
-				AlignDown(DstBatchRect.Min.Y, DstImage.NumElemsPerBlockY),
-				AlignUp  (DstBatchRect.Max.X, DstImage.NumElemsPerBlockX),
-				AlignUp  (DstBatchRect.Max.Y, DstImage.NumElemsPerBlockY));
 
-		int32 SrcDataOffsetInBlocks = (AlignedSrcRect.Min.Y / SrcImage.NumElemsPerBlockY) * (SrcImage.Size.X / SrcImage.NumElemsPerBlockX) + (AlignedSrcRect.Min.X / SrcImage.NumElemsPerBlockX); 
+		FIntRect AlignedBatchDstRect = FIntRect(
+				AlignDown(DstClippedBatchRect.Min.X, DstImage.NumElemsPerBlockX),
+				AlignDown(DstClippedBatchRect.Min.Y, DstImage.NumElemsPerBlockY),
+				AlignUp  (DstClippedBatchRect.Max.X, DstImage.NumElemsPerBlockX),
+				AlignUp  (DstClippedBatchRect.Max.Y, DstImage.NumElemsPerBlockY));
+
+		int32 SrcDataOffsetInBlocks = (AlignedBatchSrcRect.Min.Y / SrcImage.NumElemsPerBlockY) * (SrcImage.Size.X / SrcImage.NumElemsPerBlockX) + (AlignedBatchSrcRect.Min.X / SrcImage.NumElemsPerBlockX); 
 		int32 SrcDataOffsetInBytes =  SrcDataOffsetInBlocks * SrcImage.BytesPerBlock;
 
 		Result.Src = TConstSubImage<uint8> 
 		{
 			.DataView      = TConstArrayView<uint8>(SrcImage.DataView.GetData() + SrcDataOffsetInBytes, SrcImage.DataView.Num() - SrcDataOffsetInBytes),
 			.Size          = SrcImage.Size,
-			.SubSize       = FImageSize(AlignedSrcRect.Width(), AlignedSrcRect.Height()),
+			.SubSize       = FImageSize(AlignedBatchSrcRect.Width(), AlignedBatchSrcRect.Height()),
 			.BytesPerBlock = SrcImage.BytesPerBlock,
 			.NumElemsPerBlockX = SrcImage.NumElemsPerBlockX,
 			.NumElemsPerBlockY = SrcImage.NumElemsPerBlockY,
 		};
 
-		int32 DstDataOffsetInBlocks = (AlignedDstRect.Min.Y / DstImage.NumElemsPerBlockY) * (DstImage.Size.X / DstImage.NumElemsPerBlockX) + (AlignedDstRect.Min.X / DstImage.NumElemsPerBlockX); 
+		int32 DstDataOffsetInBlocks = (AlignedBatchDstRect.Min.Y / DstImage.NumElemsPerBlockY) * (DstImage.Size.X / DstImage.NumElemsPerBlockX) + (AlignedBatchDstRect.Min.X / DstImage.NumElemsPerBlockX); 
 		int32 DstDataOffsetInBytes =  DstDataOffsetInBlocks * DstImage.BytesPerBlock;
 		Result.Dst = TSubImage<uint8>
 		{
 			.DataView      = TArrayView(DstImage.DataView.GetData() + DstDataOffsetInBytes, DstImage.DataView.Num() - DstDataOffsetInBytes),
 			.Size          = DstImage.Size,
-			.SubSize       = FImageSize(AlignedDstRect.Width(), AlignedDstRect.Height()),
+			.SubSize       = FImageSize(AlignedBatchDstRect.Width(), AlignedBatchDstRect.Height()),
 			.BytesPerBlock = DstImage.BytesPerBlock,
 			.NumElemsPerBlockX = DstImage.NumElemsPerBlockX,
 			.NumElemsPerBlockY = DstImage.NumElemsPerBlockY,
 		};
 
-		if (!ScratchMemory.SrcDecompressed.IsEmpty())
+		if (!ScratchMemory.SrcDecompressedViews[JobId].IsEmpty())
 		{
 			Result.DecSrcAligned = TSubImage<uint8>
 			{
-				.DataView      = TArrayView(ScratchMemory.SrcDecompressed),
-				.Size          = FImageSize(AlignedSrcRect.Width(), AlignedSrcRect.Height()),
-				.SubSize       = FImageSize(AlignedSrcRect.Width(), AlignedSrcRect.Height()),
+				.DataView      = TArrayView(ScratchMemory.SrcDecompressedViews[JobId]),
+				.Size          = FImageSize(AlignedBatchSrcRect.Width(), AlignedBatchSrcRect.Height()),
+				.SubSize       = FImageSize(AlignedBatchSrcRect.Width(), AlignedBatchSrcRect.Height()),
 				.BytesPerBlock = FormatInfo.DstDecompressed.BytesPerBlock, 
 				.NumElemsPerBlockX = 1,
 				.NumElemsPerBlockY = 1,
 			};
 		}
 
-		if (!ScratchMemory.DstDecompressed.IsEmpty())
+		if (!ScratchMemory.DstDecompressedViews[JobId].IsEmpty())
 		{
 			Result.DecDstAligned = TSubImage<uint8>
 			{
-				.DataView      = TArrayView(ScratchMemory.DstDecompressed),
-				.Size          = FImageSize(AlignedDstRect.Width(), AlignedDstRect.Height()),
-				.SubSize       = FImageSize(AlignedDstRect.Width(), AlignedDstRect.Height()),
+				.DataView      = TArrayView(ScratchMemory.DstDecompressedViews[JobId]),
+				.Size          = FImageSize(AlignedBatchDstRect.Width(), AlignedBatchDstRect.Height()),
+				.SubSize       = FImageSize(AlignedBatchDstRect.Width(), AlignedBatchDstRect.Height()),
 				.BytesPerBlock = FormatInfo.DstDecompressed.BytesPerBlock, 
 				.NumElemsPerBlockX = 1,
 				.NumElemsPerBlockY = 1,
@@ -333,7 +420,7 @@ namespace OpImageMultiComposeInternal
 
 		if (!Result.DecDstAligned.DataView.IsEmpty())
 		{
-			FIntVector2 DstOffset = DstBatchRect.Min - AlignedDstRect.Min;
+			FIntVector2 DstOffset = DstClippedBatchRect.Min - AlignedBatchDstRect.Min;
 			int32 DecDstOffsetInBytes = (DstOffset.Y*Result.DecDstAligned.Size.X + DstOffset.X) * Result.DecDstAligned.BytesPerBlock;
 			
 			Result.DecDst = TSubImage<uint8>
@@ -341,7 +428,7 @@ namespace OpImageMultiComposeInternal
 				
 				.DataView      = TArrayView(Result.DecDstAligned.DataView.GetData() + DecDstOffsetInBytes, Result.DecDstAligned.DataView.Num() - DecDstOffsetInBytes),
 				.Size          = Result.DecDstAligned.Size,
-				.SubSize       = FImageSize(DstBatchRect.Width(), DstBatchRect.Height()),
+				.SubSize       = FImageSize(DstClippedBatchRect.Width(), DstClippedBatchRect.Height()),
 				.BytesPerBlock = FormatInfo.DstDecompressed.BytesPerBlock,
 				.NumElemsPerBlockX = 1,
 				.NumElemsPerBlockY = 1,
@@ -350,7 +437,7 @@ namespace OpImageMultiComposeInternal
 
 		if (!Result.DecSrcAligned.DataView.IsEmpty())
 		{
-			FIntVector2 SrcOffset = SrcBatchRect.Min - AlignedSrcRect.Min;
+			FIntVector2 SrcOffset = SrcBatchRect.Min - AlignedBatchSrcRect.Min;
 			int32 DecSrcOffsetInBytes = (SrcOffset.Y*Result.DecSrcAligned.Size.X + SrcOffset.X) * Result.DecSrcAligned.BytesPerBlock;
 
 			Result.DecSrc = TSubImage<uint8>
@@ -405,8 +492,6 @@ namespace OpImageMultiComposeInternal
 			EActiveStages ActiveStages = ComputeActiveStages(Base, Source, FormatInfo, DstRect, SrcRect);
 			check(EnumHasAllFlags(ScratchMemory.ActiveStages, ActiveStages)); 
 
-			int32 NumBatches = ComputeNumBatches(DstRect);
-
 			uint32 SourceLODIndex = SourceMipPerRect[RectIndex];
 
 			FIntVector2 SourceLODSize = Source.CalculateMipSize(SourceLODIndex); 
@@ -431,112 +516,135 @@ namespace OpImageMultiComposeInternal
 				.NumElemsPerBlockY = FormatInfo.Dst.PixelsPerBlockY,
 			};
 
-			for (int32 BatchIndex = 0; BatchIndex < NumBatches; ++BatchIndex)
-			{
-				FBatchWorkingSubImages WorkingSubImages = ComputeWorkingSubImages(BatchIndex, DstRect, SrcRect, BaseSubImage, SrcSubImage, FormatInfo, ScratchMemory);
-
-				if (EnumHasAnyFlags(ActiveStages, EActiveStages::DstDecompression))
-				{
-					MUTABLE_CPUPROFILER_SCOPE(ImageMultiCompose_DstDecompression);
-
-					const TSubImage<uint8>& StageSrc = WorkingSubImages.Dst;
-					const TSubImage<uint8>& StageDst = WorkingSubImages.DecDstAligned;
-
-					check(StageSrc.SubSize == StageDst.SubSize);
+			FIntVector2 NumBatches = ComputeNumBatches(DstRect, FormatInfo, ScratchMemory);
 
-					const uint8* FromPtr = StageSrc.DataView.GetData();
-					uint8* ToPtr         = StageDst.DataView.GetData();
+			// Split parallel work in rows, using a 2D grid approach could lead to better work distribution. 
+			constexpr int32 MinRowBatchesPerJob = 1;
+			const int32 NumRowBatchesPerJob = 
+				 FMath::Min(NumBatches.Y, FMath::Max(MinRowBatchesPerJob, FMath::DivideAndRoundUp(NumBatches.Y, ScratchMemory.MaxParallelJobs)));
 
-					miro::FImageSize FromSize = miro::FImageSize(StageSrc.Size.X, StageSrc.Size.Y);
-					miro::FImageSize ToSize   = miro::FImageSize(StageDst.Size.X, StageDst.Size.Y); 
-					miro::FImageSize SubSize  = miro::FImageSize(StageSrc.SubSize.X, StageSrc.SubSize.Y);
+			const int32 NumParallelJobs = FMath::DivideAndRoundUp(NumBatches.Y, NumRowBatchesPerJob); 
 
-					FormatInfo.DstDecompressionFunc(FromSize, ToSize, SubSize, FromPtr, ToPtr);	
-				}
+			ParallelExecutionUtils::InvokeBatchParallelFor(NumParallelJobs, 
+			[
+				&ScratchMemory, &FormatInfo, 
+				NumRowBatchesPerJob, NumBatches,
+				BaseSubImage, SrcSubImage,
+				ActiveStages,
+				DstRect, SrcRect,
+				CompressionQuality
+			](int32 JobId)
+			{
+				const int32 JobRowBegin = JobId*NumRowBatchesPerJob;
+				const int32 JobRowEnd   = FMath::Min(JobRowBegin + NumRowBatchesPerJob, NumBatches.Y);
 
-				if (EnumHasAnyFlags(ActiveStages, EActiveStages::SrcDecompression))
+				for (int32 BatchIdY = JobRowBegin; BatchIdY < JobRowEnd; ++BatchIdY)
 				{
-					MUTABLE_CPUPROFILER_SCOPE(ImageMultiCompose_SrcDecompression);
-
-					const TConstSubImage<uint8>& StageSrc = WorkingSubImages.Src;
-					const TSubImage<uint8>& StageDst = EnumHasAnyFlags(ActiveStages, EActiveStages::DstCompression | EActiveStages::Resize)
-							? WorkingSubImages.DecSrcAligned
-							: WorkingSubImages.Dst;
-		
-					const uint8* FromPtr = StageSrc.DataView.GetData();
-					uint8* ToPtr         = StageDst.DataView.GetData();
-
-					miro::FImageSize FromSize = miro::FImageSize(StageSrc.Size.X, StageSrc.Size.Y);
-					miro::FImageSize ToSize   = miro::FImageSize(StageDst.Size.X, StageDst.Size.Y); 
-					miro::FImageSize SubSize  = miro::FImageSize(StageDst.SubSize.X, StageDst.SubSize.Y);
-
-					const bool bUncompressedFormat = 
-							FormatInfo.SrcDecompressedFormat != FormatInfo.DstDecompressedFormat && 
-							FormatInfo.SrcDecompressedFormat == Source.GetFormat()               &&  
-							FormatInfo.DstDecompressedFormat == Base.GetFormat();
-
-					if (bUncompressedFormat)
+					for (int32 BatchIdX = 0; BatchIdX < NumBatches.X; ++BatchIdX)
 					{
-						FormatInfo.SrcUncompressedFormatFunc(StageSrc.Size, StageDst.Size, StageDst.SubSize, FromPtr, ToPtr); 
-					}
-					else
-					{
-						FormatInfo.SrcDecompressionFunc(FromSize, ToSize, SubSize, FromPtr, ToPtr);
-					}
-				}
+						FIntVector2 BatchId = FIntVector2(BatchIdX, BatchIdY);
+						FBatchWorkingSubImages WorkingSubImages = ComputeWorkingSubImages(JobId, BatchId, DstRect, SrcRect, BaseSubImage, SrcSubImage, FormatInfo, ScratchMemory, ActiveStages);
 
-				{
-					TConstSubImage<uint8> ConstDecSrc = WorkingSubImages.DecSrc.ToConst();
+						if (EnumHasAnyFlags(WorkingSubImages.ActiveStages, EActiveStages::DstDecompression))
+						{
+							MUTABLE_CPUPROFILER_SCOPE(ImageMultiCompose_DstDecompression);
 
-					const TConstSubImage<uint8>& StageSrc = EnumHasAnyFlags(ActiveStages, EActiveStages::SrcDecompression)
-							? ConstDecSrc
-							: WorkingSubImages.Src;
+							const TSubImage<uint8>& StageSrc = WorkingSubImages.Dst;
+							const TSubImage<uint8>& StageDst = WorkingSubImages.DecDstAligned;
 
-					const TSubImage<uint8>& StageDst = EnumHasAnyFlags(ActiveStages, EActiveStages::DstCompression)
-							? WorkingSubImages.DecDst
-							: WorkingSubImages.Dst;
+							check(StageSrc.SubSize == StageDst.SubSize);
 
-					if (EnumHasAnyFlags(ActiveStages, EActiveStages::Resize))
-					{
-						MUTABLE_CPUPROFILER_SCOPE(ImageMultiCompose_Resize);
+							const uint8* FromPtr = StageSrc.DataView.GetData();
+							uint8* ToPtr         = StageDst.DataView.GetData();
 
-						const uint8* FromPtr = StageSrc.DataView.GetData();
-						uint8* ToPtr         = StageDst.DataView.GetData();
+							miro::FImageSize FromSize = miro::FImageSize(StageSrc.Size.X, StageSrc.Size.Y);
+							miro::FImageSize ToSize   = miro::FImageSize(StageDst.Size.X, StageDst.Size.Y); 
+							miro::FImageSize SubSize  = miro::FImageSize(StageSrc.SubSize.X, StageSrc.SubSize.Y);
 
-						FVector2f SrcRectDimsF = FVector2f(SrcRect.Width(), SrcRect.Height());
-						FVector2f DstRectDimsF = FVector2f(DstRect.Width(), DstRect.Height());
+							FormatInfo.DstDecompressionFunc(FromSize, ToSize, SubSize, FromPtr, ToPtr);	
+						}
 
-						FVector2f OneOverScalingFactor = SrcRectDimsF / DstRectDimsF;
+						if (EnumHasAnyFlags(WorkingSubImages.ActiveStages, EActiveStages::SrcDecompression))
+						{
+							MUTABLE_CPUPROFILER_SCOPE(ImageMultiCompose_SrcDecompression);
 
-						FormatInfo.ResizeFunc(OneOverScalingFactor, WorkingSubImages.SrcSubPixelOffset, 
-								StageDst.Size, StageDst.SubSize, StageSrc.Size, StageSrc.SubSize, ToPtr, FromPtr);	
-					}
-					else 
-					{
-						MUTABLE_CPUPROFILER_SCOPE(ImageMultiCompose_Copy);
-						CopySubImage(StageDst, StageSrc);
+							const TConstSubImage<uint8>& StageSrc = WorkingSubImages.Src;
+							const TSubImage<uint8>& StageDst = EnumHasAnyFlags(ActiveStages, EActiveStages::DstCompression | EActiveStages::Resize)
+									? WorkingSubImages.DecSrcAligned
+									: WorkingSubImages.Dst;
+				
+							const uint8* FromPtr = StageSrc.DataView.GetData();
+							uint8* ToPtr         = StageDst.DataView.GetData();
+
+							if (FormatInfo.bIsUncompressedReformat)
+							{
+								FormatInfo.SrcUncompressedFormatFunc(StageSrc.Size, StageDst.Size, StageDst.SubSize, FromPtr, ToPtr); 
+							}
+							else
+							{
+								miro::FImageSize FromSize = miro::FImageSize(StageSrc.Size.X, StageSrc.Size.Y);
+								miro::FImageSize ToSize   = miro::FImageSize(StageDst.Size.X, StageDst.Size.Y); 
+								miro::FImageSize SubSize  = miro::FImageSize(StageDst.SubSize.X, StageDst.SubSize.Y);
+
+								FormatInfo.SrcDecompressionFunc(FromSize, ToSize, SubSize, FromPtr, ToPtr);
+							}
+						}
+
+						{
+							TConstSubImage<uint8> ConstDecSrc = WorkingSubImages.DecSrc.ToConst();
+
+							const TConstSubImage<uint8>& StageSrc = EnumHasAnyFlags(ActiveStages, EActiveStages::SrcDecompression)
+									? ConstDecSrc
+									: WorkingSubImages.Src;
+
+							const TSubImage<uint8>& StageDst = EnumHasAnyFlags(ActiveStages, EActiveStages::DstCompression)
+									? WorkingSubImages.DecDst
+									: WorkingSubImages.Dst;
+
+							if (EnumHasAnyFlags(WorkingSubImages.ActiveStages, EActiveStages::Resize))
+							{
+								MUTABLE_CPUPROFILER_SCOPE(ImageMultiCompose_Resize);
+
+								const uint8* FromPtr = StageSrc.DataView.GetData();
+								uint8* ToPtr         = StageDst.DataView.GetData();
+
+								FVector2f SrcRectDimsF = FVector2f(SrcRect.Width(), SrcRect.Height());
+								FVector2f DstRectDimsF = FVector2f(DstRect.Width(), DstRect.Height());
+
+								FVector2f OneOverScalingFactor = SrcRectDimsF / DstRectDimsF;
+
+								FormatInfo.ResizeFunc(OneOverScalingFactor, WorkingSubImages.SrcSubPixelOffset, 
+										StageDst.Size, StageDst.SubSize, StageSrc.Size, StageSrc.SubSize, ToPtr, FromPtr);	
+							}
+							else 
+							{
+								MUTABLE_CPUPROFILER_SCOPE(ImageMultiCompose_Copy);
+								CopySubImage(StageDst, StageSrc);
+							}
+						}
+
+						if (EnumHasAnyFlags(WorkingSubImages.ActiveStages, EActiveStages::DstCompression))
+						{
+							MUTABLE_CPUPROFILER_SCOPE(ImageMultiCompose_DstCompression);
+
+							const TSubImage<uint8>& StageSrc = WorkingSubImages.DecDstAligned;
+							const TSubImage<uint8>& StageDst = WorkingSubImages.Dst;
+
+							check(StageSrc.SubSize == StageDst.SubSize);
+
+							const uint8* FromPtr = StageSrc.DataView.GetData();
+							uint8* ToPtr         = StageDst.DataView.GetData();
+
+							miro::FImageSize FromSize = miro::FImageSize(StageSrc.Size.X, StageSrc.Size.Y);
+							miro::FImageSize ToSize   = miro::FImageSize(StageDst.Size.X, StageDst.Size.Y); 
+							miro::FImageSize SubSize  = miro::FImageSize(StageDst.SubSize.X, StageDst.SubSize.Y);
+
+							FormatInfo.DstCompressionFunc(FromSize, ToSize, SubSize, FromPtr, ToPtr, CompressionQuality);
+						}
 					}
 				}
+			});
 
-				if (EnumHasAnyFlags(ActiveStages, EActiveStages::DstCompression))
-				{
-					MUTABLE_CPUPROFILER_SCOPE(ImageMultiCompose_DstCompression);
-
-					const TSubImage<uint8>& StageSrc = WorkingSubImages.DecDstAligned;
-					const TSubImage<uint8>& StageDst = WorkingSubImages.Dst;
-
-					check(StageSrc.SubSize == StageDst.SubSize);
-
-					const uint8* FromPtr = StageSrc.DataView.GetData();
-					uint8* ToPtr         = StageDst.DataView.GetData();
-
-					miro::FImageSize FromSize = miro::FImageSize(StageSrc.Size.X, StageSrc.Size.Y);
-					miro::FImageSize ToSize   = miro::FImageSize(StageDst.Size.X, StageDst.Size.Y); 
-					miro::FImageSize SubSize  = miro::FImageSize(StageDst.SubSize.X, StageDst.SubSize.Y);
-
-					FormatInfo.DstCompressionFunc(FromSize, ToSize, SubSize, FromPtr, ToPtr, CompressionQuality);
-				}
-			}
 		}
 	}
 
@@ -548,8 +656,8 @@ namespace OpImageMultiComposeInternal
 		// Image multicompose only supports scaling down by less than 2 or scale up, for now pick the coord that 
 		// changes the most down so we never get values outside the limitations.
 		//
-		// NOTE: Doing this we may loose resolution if both dimentions change very differently, using layout packing
-		// startegies that mantain the original aspect ratio is advised to get the best results possible. 
+		// NOTE: Doing this we may loose resolution if both dimensions change very differently, using layout packing
+		// strategies that maintain the original aspect ratio is advised to get the best results possible. 
 		int32 MostChangingCoord = static_cast<int32>((DstRectSize.X - SrcRectSize.X) > (DstRectSize.Y - SrcRectSize.Y));
 
 		const float BestLOD = 
diff --git a/Engine/Plugins/Mutable/Source/MutableRuntime/Private/MuR/ParallelExecutionUtils.cpp b/Engine/Plugins/Mutable/Source/MutableRuntime/Private/MuR/ParallelExecutionUtils.cpp
--- a/Engine/Plugins/Mutable/Source/MutableRuntime/Private/MuR/ParallelExecutionUtils.cpp
+++ b/Engine/Plugins/Mutable/Source/MutableRuntime/Private/MuR/ParallelExecutionUtils.cpp
@@ -17,6 +17,18 @@ namespace UE::Mutable::Private::ParallelExecutionUtils
 		}
 	}
 
+	void InvokeBatchParallelForUnbalanced(int32 Num, TFunctionRef<void(int32)> Body)
+	{
+		if (Num == 1)
+		{
+			Body(0);
+		}
+		else if (Num > 1)
+		{
+			ParallelFor(Num, Body, EParallelForFlags::Unbalanced);
+		}
+	}
+
 	void InvokeBatchParallelForSingleThread(int32 Num, TFunctionRef<void(int32)> Body)
 	{
 		for (int32 BatchId = 0; BatchId < Num; ++BatchId)
