[Mutable] Remove stray unused variable from FMesh.

#rnx
#Synced-CL Unknown

diff --git a/Engine/Plugins/Mutable/Source/MutableRuntime/Internal/MuR/Mesh.h b/Engine/Plugins/Mutable/Source/MutableRuntime/Internal/MuR/Mesh.h
--- a/Engine/Plugins/Mutable/Source/MutableRuntime/Internal/MuR/Mesh.h
+++ b/Engine/Plugins/Mutable/Source/MutableRuntime/Internal/MuR/Mesh.h
@@ -441,7 +441,6 @@ namespace UE::Mutable::Private
 		/** If the mesh is a reference the referenced morph name is stored here. Otherwise it is an empty string. */
 		FString ReferencedMorph;
 
-		TManagedPtr<int32> A = nullptr;
 		/** 
 		 * Prefix for the unique IDs related to this mesh (vertices and layout blocks). Useful if the mesh stores them in an implicit, or relative way. 
 		 * See MeshVertexIdIterator for details.
