[mutable] Debugger: Implemented support for material constants.
- Added a new section into the Mutable Constants (debugger) so we can debug what FMaterials mutable is generating and how many params each one of them has.
- Added WIP Material Viewer to show the internal data of the mutable FMaterials. It will currently fail to produce complete images from the image parameters they contain.
- Other small fixes and tweaks.

#rnx

#Synced-CL 48762940

diff --git a/Engine/Plugins/Mutable/Source/CustomizableObjectEditor/Private/MuCOE/SMutableCodeViewer.cpp b/Engine/Plugins/Mutable/Source/CustomizableObjectEditor/Private/MuCOE/SMutableCodeViewer.cpp
--- a/Engine/Plugins/Mutable/Source/CustomizableObjectEditor/Private/MuCOE/SMutableCodeViewer.cpp
+++ b/Engine/Plugins/Mutable/Source/CustomizableObjectEditor/Private/MuCOE/SMutableCodeViewer.cpp
@@ -2,6 +2,7 @@
 
 #include "MuCOE/SMutableCodeViewer.h"
 
+#include "SMutableMaterialViewer.h"
 #include "Framework/MultiBox/MultiBoxBuilder.h"
 #include "Framework/Views/TableViewMetadata.h"
 #include "Misc/Paths.h"
@@ -115,7 +116,7 @@ public:
 	/** Depending on the state of the row returns one color or another to be used by the highlighting system */
 	FLinearColor GetHighlightColor() const
 	{
-		if (bShouldBeHiglighted)
+		if (bShouldBeHighlighted)
 		{
 			if (RowItem->DuplicatedOf)
 			{
@@ -210,13 +211,13 @@ public:
 	/** Marks the row to be highlighted */
 	void Highlight()
 	{
-		bShouldBeHiglighted = true;
+		bShouldBeHighlighted = true;
 	}
 
 	/** Resets the highlighting status */
 	void ResetHighlight()
 	{
-		bShouldBeHiglighted = false;
+		bShouldBeHighlighted = false;
 	}
 
 	/** Returns a reference to the Element this row is representing */
@@ -271,7 +272,7 @@ private:
 	/** Color shown on the extra data column when the resource is found to be State Constant */
 	const FLinearColor StateConstantBoxColor = FLinearColor(1,0,0,0.8);
 
-	bool bShouldBeHiglighted = false;
+	bool bShouldBeHighlighted = false;
 };
 
 
@@ -2112,6 +2113,13 @@ void SMutableCodeViewer::OnSelectionChanged(TSharedPtr<FMutableCodeTreeElement>
 		break;
 	}
 
+	case UE::Mutable::Private::EDataType::Material:
+	{
+		// Create or reuse the UI
+		PrepareMaterialViewer();
+		break;
+	}
+
 	case UE::Mutable::Private::EDataType::Instance:
 	{
 		// Create or reuse the UI
@@ -2838,6 +2846,15 @@ bool SMutableCodeViewer::IsConstantResourceUsedByOperation(const int32 IndexOnCo
 				break;
 			}
 
+		case UE::Mutable::Private::EDataType::Material:
+			{
+				if (OperationType == UE::Mutable::Private::EOpType::MI_CONSTANT)
+				{
+					bResourceLocated = IndexOnConstantsArray == InProgram.GetOpArgs<UE::Mutable::Private::OP::MaterialConstantArgs>(OperationAddress).Value;
+				}
+				break;
+			}
+
 		// Invalid types
 		case UE::Mutable::Private::EDataType::None:
 		default:
@@ -2962,18 +2979,40 @@ void SMutableCodeViewer::Tick(const FGeometry& AllottedGeometry, const double In
 		check (PreviewProjectorViewer);
 		UE::Mutable::Private::FProjector Projector = System->GetPrivate()->BuildProjector(ExternalResourceProvider, MutableModel, PreviewParameters.Get(), SelectedOperationAddress);
 		PreviewProjectorViewer->SetProjector(Projector);
+		break;
+	}
+	
+	case UE::Mutable::Private::EDataType::Material:
+	{
+		check (PreviewMaterialViewer);
+		// todo: Evaluate and feed the FMaterial into the viewer or make the viewer evaluate the contents and display them appropriately (mind the constants pannel)
+		PreviewMaterialViewer->SetMaterial(nullptr);
+		break;
 	}
 	
 	default:
 #if UE_BUILD_DEBUG || UE_BUILD_DEVELOPMENT
 		UE_LOG(LogMutable, Log, TEXT("There is no previewer for the selected type of Mutable object"))
 #endif
+		PreviewBorder->SetContent(SNullWidget::NullWidget);
 		// There is no viewer for this type.
 		break;
 	}
 
 	System->GetPrivate()->EndBuild();
+}
+
+
+UE::Mutable::Private::TManagedPtr<const UE::Mutable::Private::FImage> SMutableCodeViewer::BuildImage(UE::Mutable::Private::OP::ADDRESS TargetOpAddress) const
+{
+	UE::Mutable::Private::FSettings Settings;
+	const TSharedPtr<UE::Mutable::Private::FSystem> System = MakeShared<UE::Mutable::Private::FSystem>(Settings);
 
+	System->GetPrivate()->BeginBuild(MutableModel, Object->GetPrivate()->MeshIdRegistry, Object->GetPrivate()->ImageIdRegistry, Object->GetPrivate()->MaterialIdRegistry, Object->GetPrivate()->SkeletalMeshIdRegistry);
+	UE::Mutable::Private::TManagedPtr<const UE::Mutable::Private::FImage> MutableImage = System->GetPrivate()->BuildImage(ExternalResourceProvider, MutableModel, PreviewParameters.Get(), TargetOpAddress, 0);
+	System->GetPrivate()->EndBuild();
+	
+	return MutableImage;
 }
 
 
@@ -2994,6 +3033,7 @@ void SMutableCodeViewer::PrepareStringViewer()
 	PreviewBorder->SetContent(PreviewStringViewer.ToSharedRef());
 }
 
+
 void SMutableCodeViewer::PrepareImageViewer()
 {
 	if (!PreviewImageViewer)
@@ -3005,6 +3045,7 @@ void SMutableCodeViewer::PrepareImageViewer()
 	PreviewBorder->SetContent(PreviewImageViewer.ToSharedRef());
 }
 
+
 void SMutableCodeViewer::PrepareMeshViewer()
 {
 	if (!PreviewMeshViewer)
@@ -3015,6 +3056,20 @@ void SMutableCodeViewer::PrepareMeshViewer()
 	PreviewBorder->SetContent(PreviewMeshViewer.ToSharedRef());
 }
 
+
+void SMutableCodeViewer::PrepareMaterialViewer()
+{
+	if (!PreviewMaterialViewer)
+	{
+		PreviewMaterialViewer = 
+			SNew(SMutableMaterialViewer)
+			.HostCodeViewer(TWeakPtr<SMutableCodeViewer>(SharedThis(this)));		// So it can invoke the BuildImage member method WIP
+	}
+	
+	PreviewBorder->SetContent(PreviewMaterialViewer.ToSharedRef());
+}
+
+
 void SMutableCodeViewer::PrepareInstanceViewer()
 {
 	if (!PreviewInstanceViewer)
@@ -3025,6 +3080,7 @@ void SMutableCodeViewer::PrepareInstanceViewer()
 	PreviewBorder->SetContent(PreviewInstanceViewer.ToSharedRef());
 }
 
+
 void SMutableCodeViewer::PrepareLayoutViewer()
 {
 	if (!PreviewLayoutViewer)
@@ -3072,6 +3128,13 @@ void SMutableCodeViewer::PreviewMutableMesh(UE::Mutable::Private::TManagedPtr<co
 }
 
 
+void SMutableCodeViewer::PreviewMutableMaterial(UE::Mutable::Private::TManagedPtr<const UE::Mutable::Private::FMaterial> InMaterialPtr)
+{
+	PrepareMaterialViewer();
+	PreviewMaterialViewer->SetMaterial(InMaterialPtr);
+}
+
+
 void SMutableCodeViewer::PreviewMutableLayout(UE::Mutable::Private::TManagedPtr<const UE::Mutable::Private::FLayout> Layout)
 {
 	PrepareLayoutViewer();
diff --git a/Engine/Plugins/Mutable/Source/CustomizableObjectEditor/Private/MuCOE/SMutableCodeViewer.h b/Engine/Plugins/Mutable/Source/CustomizableObjectEditor/Private/MuCOE/SMutableCodeViewer.h
--- a/Engine/Plugins/Mutable/Source/CustomizableObjectEditor/Private/MuCOE/SMutableCodeViewer.h
+++ b/Engine/Plugins/Mutable/Source/CustomizableObjectEditor/Private/MuCOE/SMutableCodeViewer.h
@@ -39,6 +39,7 @@ class SMutableProjectorViewer;
 class SMutableScalarViewer;
 class SMutableSkeletonViewer;
 class SMutableStringViewer;
+class SMutableMaterialViewer;
 class STextComboBox;
 class SWidget;
 namespace UE::Mutable::Private { struct FProjector; }
@@ -66,6 +67,14 @@ public:
 	/** Clear the selected operation of the tree view.*/
 	void ClearSelectedTreeRow() const;
 	
+	/**
+	 * Evaluates the graph starting at the TargetOPAddress and returns an FImage resulting from that
+	 * @param TargetOpAddress The OP Address to use as the root to be computed.
+	 * @return A pointer to the FImage generated. May be null if the operation does not evaluate to an FImage
+	 * @note This method is experimental and should not be used as example
+	 */
+	UE::Mutable::Private::TManagedPtr<const UE::Mutable::Private::FImage>  BuildImage(UE::Mutable::Private::OP::ADDRESS TargetOpAddress) const;
+	
 private:
 	TStrongObjectPtr<UCustomizableObject> Object;
 	TStrongObjectPtr<UCustomizableObjectInstance> Instance;
@@ -113,6 +122,9 @@ private:
 	/** Widget used to show a preview of a mesh and the metadata it holds */
 	TSharedPtr<SMutableMeshViewer> PreviewMeshViewer;
 
+	/** Widget used to show a preview of a mutable FMaterial and the data it holds */
+	TSharedPtr<SMutableMaterialViewer> PreviewMaterialViewer;
+	
 	/** Widget used to show a preview of a mutable bool value */
 	TSharedPtr<SMutableBoolViewer> PreviewBoolViewer;
 
@@ -793,6 +805,7 @@ private:
 public:
 	void PreviewMutableImage (UE::Mutable::Private::TManagedPtr<const UE::Mutable::Private::FImage> InImagePtr);
 	void PreviewMutableMesh (UE::Mutable::Private::TManagedPtr<const UE::Mutable::Private::FMesh> InMeshPtr);
+	void PreviewMutableMaterial(UE::Mutable::Private::TManagedPtr<const UE::Mutable::Private::FMaterial> InMaterialPtr);
 	void PreviewMutableLayout(UE::Mutable::Private::TManagedPtr<const UE::Mutable::Private::FLayout> Layout);
 	void PreviewMutableSkeleton(UE::Mutable::Private::TManagedPtr<const UE::Mutable::Private::FSkeleton> Skeleton);
 	void PreviewMutablePhysics(UE::Mutable::Private::TManagedPtr<const UE::Mutable::Private::FPhysicsBody> Physics);
@@ -806,6 +819,7 @@ private:
 	void PrepareStringViewer();
 	void PrepareImageViewer();
 	void PrepareMeshViewer();
+	void PrepareMaterialViewer();
 	void PrepareInstanceViewer();
 	void PrepareLayoutViewer();
 	void PrepareProjectorViewer();
diff --git a/Engine/Plugins/Mutable/Source/CustomizableObjectEditor/Private/MuCOE/SMutableColorViewer.cpp b/Engine/Plugins/Mutable/Source/CustomizableObjectEditor/Private/MuCOE/SMutableColorViewer.cpp
--- a/Engine/Plugins/Mutable/Source/CustomizableObjectEditor/Private/MuCOE/SMutableColorViewer.cpp
+++ b/Engine/Plugins/Mutable/Source/CustomizableObjectEditor/Private/MuCOE/SMutableColorViewer.cpp
@@ -4,6 +4,7 @@
 
 #include "Widgets/SBoxPanel.h"
 #include "Widgets/Colors/SColorBlock.h"
+#include "Widgets/Layout/SBorder.h"
 #include "Widgets/Text/STextBlock.h"
 
 #define LOCTEXT_NAMESPACE "CustomizableObjectEditor"
@@ -91,48 +92,48 @@ void SMutableColorViewer::Construct(const FArguments& InArgs)
 
 				// Blue channel ----------------------------------------
 				+SVerticalBox::Slot()
-					.Padding(0, 1)
-					.AutoHeight()
-					[
-						SNew(SHorizontalBox)
+				.Padding(0, 1)
+				.AutoHeight()
+				[
+					SNew(SHorizontalBox)
 
-						+ SHorizontalBox::Slot()
+					+ SHorizontalBox::Slot()
 					.AutoWidth()
 					[
 						SNew(STextBlock).
 						Text(BlueColorValueTitle)
 					]
 
-				+ SHorizontalBox::Slot()
+					+ SHorizontalBox::Slot()
 					.AutoWidth()
 					[
 						SNew(STextBlock).
 						Text(this, &SMutableColorViewer::GetBlueValue)
 					]
-					]
+				]
 				// -----------------------------------------------------
 
 				// Alpha channel ----------------------------------------
 				+SVerticalBox::Slot()
-					.Padding(0, 1)
-					.AutoHeight()
-					[
-						SNew(SHorizontalBox)
+				.Padding(0, 1)
+				.AutoHeight()
+				[
+					SNew(SHorizontalBox)
 
-						+ SHorizontalBox::Slot()
+					+ SHorizontalBox::Slot()
 					.AutoWidth()
 					[
 						SNew(STextBlock).
 						Text(AlphaColorValueTitle)
 					]
 
-				+ SHorizontalBox::Slot()
+					+ SHorizontalBox::Slot()
 					.AutoWidth()
 					[
 						SNew(STextBlock).
 						Text(this, &SMutableColorViewer::GetAlphaValue)
 					]
-					]
+				]
 				// -----------------------------------------------------
 			]
 		
@@ -141,9 +142,12 @@ void SMutableColorViewer::Construct(const FArguments& InArgs)
 			.Padding(ColorPreviewHorizontalPadding,0)
 			.MaxWidth(120)
 			[
-				SAssignNew(ColorPreview, SColorBlock)
-				.UseSRGB(false)
-				.Color(this, &SMutableColorViewer::GetColor)
+				SNew(SBorder)
+				[
+					SAssignNew(ColorPreview, SColorBlock)
+					.UseSRGB(false)
+					.Color(this, &SMutableColorViewer::GetColor)
+				]
 			]
 			// -----------------------------------------------------
 		]
diff --git a/Engine/Plugins/Mutable/Source/CustomizableObjectEditor/Private/MuCOE/SMutableConstantsWidget.cpp b/Engine/Plugins/Mutable/Source/CustomizableObjectEditor/Private/MuCOE/SMutableConstantsWidget.cpp
--- a/Engine/Plugins/Mutable/Source/CustomizableObjectEditor/Private/MuCOE/SMutableConstantsWidget.cpp
+++ b/Engine/Plugins/Mutable/Source/CustomizableObjectEditor/Private/MuCOE/SMutableConstantsWidget.cpp
@@ -2,6 +2,8 @@
 
 #include "MuCOE/SMutableConstantsWidget.h"
 
+#include "AssetViewUtils.h"
+#include "Editor.h"
 #include "Algo/StableSort.h"
 #include "Framework/Views/TableViewMetadata.h"
 #include "MuCOE/SMutableCodeViewer.h"
@@ -12,6 +14,8 @@
 #include "Widgets/Layout/SSeparator.h"
 #include "Widgets/Views/STileView.h"
 #include "MuCOE/Widgets/MutableMultiPageListView.h"
+#include "Subsystems/AssetEditorSubsystem.h"
+#include "Widgets/Input/SRichTextHyperlink.h"
 
 class ITableRow;
 class SWidget;
@@ -98,11 +102,11 @@ namespace
 
 namespace MeshConstantTitles
 {
-	static const FName MeshID("Id");
+	static const FName MeshID("ID");
 	static const FName MeshVertices("Vertices");
 	static const FName MeshIndices("Indices");
-	static const FName MeshChannels("BufferChannels");
-	static const FName MeshMemory("Memory");
+	static const FName MeshChannels("Channels");
+	static const FName MeshMemory("Size");
 }
 
 class SMutableConstantMeshRow final : public SMultiColumnTableRow<TSharedPtr<FMutableConstantMeshElement>>
@@ -218,11 +222,11 @@ private:
 
 namespace ImageConstantTitles
 {
-	static const FName ImageID("Id");
+	static const FName ImageID("ID");
 	static const FName ImageSize("Resolution");
-	static const FName ImageMipMaps("MipMapCount");
+	static const FName ImageMipMaps("Mip Maps");
 	static const FName ImageFormat("Format");
-	static const FName ImageTotalMemory("MemorySize");
+	static const FName ImageTotalMemory("Size");
 }
 
 class SMutableConstantImageRow final : public SMultiColumnTableRow<TSharedPtr<FMutableConstantImageElement>>
@@ -301,6 +305,125 @@ private:
 };
 
 
+namespace MaterialConstantTitles
+{
+	static const FName MaterialID("ID");
+	static const FName Material("Material");
+	static const FName ImageParametersCount("Images");
+	static const FName ScalarParametersCount("Scalars");
+	static const FName ColorParametersCount("Colors");
+}
+
+class SMutableConstantMaterialRow final : public SMultiColumnTableRow<TSharedPtr<FMutableConstantMaterialElement>>
+{
+public:
+	void Construct(const FArguments& Args, const TSharedRef<STableViewBase>& InOwnerTableView, const TSharedPtr<FMutableConstantMaterialElement>& InRowElement)
+	{
+		check(InRowElement);
+		RowElement = InRowElement;
+		
+		SMultiColumnTableRow< TSharedPtr<FMutableConstantMaterialElement> >::Construct(
+	STableRow::FArguments()
+			.ShowSelection(true)
+			, InOwnerTableView
+		);
+	}
+	
+	
+	void OnMaterialNavigation() const
+	{
+		if (RowElement)
+		{
+			UE::Mutable::Private::TManagedPtr<const UE::Mutable::Private::FMaterial> MutableMaterial = RowElement->MutableMaterial;
+			if (MutableMaterial && MutableMaterial->PassthroughObject.IsResolved())
+			{
+				const UMaterialInterface* Material = MutableMaterial->PassthroughObject.Get();
+				check(Material);
+				GEditor->GetEditorSubsystem<UAssetEditorSubsystem>()->OpenEditorForAsset(Material);
+			}
+		}
+	}
+
+	
+	virtual TSharedRef<SWidget> GenerateWidgetForColumn(const FName& ColumnName) override
+	{
+		check(RowElement);
+
+		const UE::Mutable::Private::FMaterial* MutableMaterial = RowElement->MutableMaterial.Get();
+		if (!MutableMaterial)
+		{
+			return SNullWidget::NullWidget;
+		}
+		
+		// Index
+		if (ColumnName == MaterialConstantTitles::MaterialID)
+		{
+			const FText IndexAsText = FText::AsNumber(RowElement->IndexOnSourceVector);
+			return SNew(STextBlock).Text(IndexAsText);
+		}
+		
+		// Material
+		if (ColumnName == MaterialConstantTitles::Material)
+		{
+			// Try to get the actual object and if null then display the ID as fallback
+			if (MutableMaterial->PassthroughObject.IsResolved())
+			{
+				if (const UMaterialInterface* Material = MutableMaterial->PassthroughObject.Get())
+				{
+					return SNew(SHyperlink)
+						.Style(FAppStyle::Get(), TEXT("NavigationHyperlink"))
+						.Text(FText::FromString(Material->GetName()))
+						.OnNavigate(this, &SMutableConstantMaterialRow::OnMaterialNavigation);
+				}
+			}
+			else
+			{
+				FText MaterialIDText;
+				
+				UE::Mutable::Private::PASSTHROUGH_ID PassthroughMaterialID = MutableMaterial->PassthroughObject.GetId();
+				if (PassthroughMaterialID != UE::Mutable::Private::PASSTHROUGH_ID_INVALID)
+				{
+					MaterialIDText = FText::AsNumber(PassthroughMaterialID);
+				}
+				else
+				{
+					MaterialIDText = FText::FromString(TEXT("INVALID_ID"));
+				}
+				
+				return SNew(STextBlock).Text(MaterialIDText);
+			}
+		}
+		
+		// Image parameters count
+		if (ColumnName == MaterialConstantTitles::ImageParametersCount)
+		{
+			const int32 ImageParametersCount = MutableMaterial->ImageParameters.Num();
+			return SNew(STextBlock).Text(FText::AsNumber(ImageParametersCount));
+		}
+		
+		// Color Parameters Count
+		if (ColumnName == MaterialConstantTitles::ColorParametersCount)
+		{
+			const int32 ColorParametersCount = MutableMaterial->ColorParameters.Num();
+			return SNew(STextBlock).Text(FText::AsNumber(ColorParametersCount));
+		}
+		
+		// Scalar Parameters Count
+		if (ColumnName == MaterialConstantTitles::ScalarParametersCount)
+		{
+			const int32 ScalarParametersCount = MutableMaterial->ScalarParameters.Num();
+			return SNew(STextBlock).Text(FText::AsNumber(ScalarParametersCount));
+		}
+
+		checkNoEntry();
+		return SNullWidget::NullWidget;
+	}
+
+private:
+	TSharedPtr<FMutableConstantMaterialElement> RowElement;
+};
+
+
 class SMutableConstantLayoutRow final : public STableRow<TSharedPtr<FMutableConstantLayoutElement>>
 {
 public:
@@ -472,48 +595,48 @@ void SMutableConstantsWidget::Construct(const FArguments& InArgs,const UE::Mutab
 	constexpr float ProxyEntryHeight = 20;
 	
 	// Panel title
-	const FText ConstantsPanelTitle = LOCTEXT("ConstantsPannelName", "Constants : ");
+	const FText ConstantsPanelTitle = LOCTEXT("ConstantsPanelName", "Constants : ");
 
 	// Hack to allow us to later set the array to be used. if no array is provided then the children slate to contain them later will not exist
-	TArray<TSharedPtr<FMutableConstantImageElement>> TempImageElementsEmptyArray;
+	TArray<TSharedPtr<FMutableConstantImageElement>> DummyImageElementsArray;
 	
 	// Image Constants List View object and Handler
 	TSharedPtr<SListView<TSharedPtr<FMutableConstantImageElement>>> ConstantImagesListView =
 		SNew(SListView<TSharedPtr<FMutableConstantImageElement>>)
 	.OnGenerateRow(this,&SMutableConstantsWidget::OnGenerateImageRow)
 		// We do require to provide something here or the slate to contain the children will not get generated
-	.ListItemsSource(&TempImageElementsEmptyArray)		
+	.ListItemsSource(&DummyImageElementsArray)		
 	.OnSelectionChanged(this,&SMutableConstantsWidget::OnSelectedImageChanged)
 	.SelectionMode(ESelectionMode::Single)
 	.HeaderRow
 	(
 		SNew(SHeaderRow)
 		+ SHeaderRow::Column(ImageConstantTitles::ImageID)
-		.DefaultLabel(FText(LOCTEXT("ImageId", "ID")))
+		.DefaultLabel(FText::FromName(ImageConstantTitles::ImageID))
 		.OnSort(this, &SMutableConstantsWidget::OnImageTableSortRequested)
 		.SortMode(this, &SMutableConstantsWidget::GetImageListColumnSortMode, ImageConstantTitles::ImageID )
 		.FillWidth(0.28f)
 	
 		+ SHeaderRow::Column(ImageConstantTitles::ImageSize)
-		.DefaultLabel(FText(LOCTEXT("ImageResolution","Resolution")))
+		.DefaultLabel(FText::FromName(ImageConstantTitles::ImageSize))
 		.DefaultTooltip(FText(LOCTEXT("ImageResolutionColumnToolTip","Pixel resolution")))
 		.OnSort(this, &SMutableConstantsWidget::OnImageTableSortRequested)
 		.SortMode(this, &SMutableConstantsWidget::GetImageListColumnSortMode, ImageConstantTitles::ImageSize )
 		
 		+ SHeaderRow::Column(ImageConstantTitles::ImageMipMaps)
-		.DefaultLabel(FText(LOCTEXT("ImageMipMaps","Mip Maps")))
+		.DefaultLabel(FText::FromName(ImageConstantTitles::ImageMipMaps))
 		.DefaultTooltip(FText(LOCTEXT("ImageMipMapsColumnToolTip","Amount of Mip maps")))
 		.OnSort(this, &SMutableConstantsWidget::OnImageTableSortRequested)
 		.SortMode(this, &SMutableConstantsWidget::GetImageListColumnSortMode, ImageConstantTitles::ImageMipMaps )
 
 		+ SHeaderRow::Column(ImageConstantTitles::ImageFormat)
-		.DefaultLabel(FText(LOCTEXT("ImageFormat","Format")))
+		.DefaultLabel(FText::FromName(ImageConstantTitles::ImageFormat))
 		.DefaultTooltip(FText(LOCTEXT("ImageFormatColumnToolTip","Image Format")))
 		.OnSort(this, &SMutableConstantsWidget::OnImageTableSortRequested)
 		.SortMode(this, &SMutableConstantsWidget::GetImageListColumnSortMode, ImageConstantTitles::ImageFormat )
 
 		+ SHeaderRow::Column(ImageConstantTitles::ImageTotalMemory)
-		.DefaultLabel(FText(LOCTEXT("ImageMemorySize","Size")))
+		.DefaultLabel(FText::FromName(ImageConstantTitles::ImageTotalMemory))
 		.DefaultTooltip(FText(LOCTEXT("ImageMemorySizeColumnToolTip","Memory size")))
 		.OnSort(this, &SMutableConstantsWidget::OnImageTableSortRequested)
 		.SortMode(this, &SMutableConstantsWidget::GetImageListColumnSortMode, ImageConstantTitles::ImageTotalMemory )
@@ -521,51 +644,94 @@ void SMutableConstantsWidget::Construct(const FArguments& InArgs,const UE::Mutab
 	);
 
 
-	TArray<TSharedPtr<FMutableConstantMeshElement>> TempMeshElementsEmptyArray;
+	TArray<TSharedPtr<FMutableConstantMeshElement>> DummyMeshElementsArray;
 
 	// Handled list view for mutable constant meshes
 	TSharedPtr<SListView<TSharedPtr<FMutableConstantMeshElement>>> ConstantMeshesListView =
 		SNew(SListView<TSharedPtr<FMutableConstantMeshElement>>)
 	.OnGenerateRow(this,&SMutableConstantsWidget::OnGenerateMeshRow)
 		// We do require to provide something here or the slate to contain the children will not get generated
-	.ListItemsSource(&TempMeshElementsEmptyArray)		
+	.ListItemsSource(&DummyMeshElementsArray)		
 	.OnSelectionChanged(this,&SMutableConstantsWidget::OnSelectedMeshChanged)
 	.SelectionMode(ESelectionMode::Single)
 	.HeaderRow
 	(
 		SNew(SHeaderRow)
 		+ SHeaderRow::Column(MeshConstantTitles::MeshID)
-		.DefaultLabel(FText(LOCTEXT("MeshId", "ID")))
+		.DefaultLabel(FText::FromName(MeshConstantTitles::MeshID))
 		.OnSort(this, &SMutableConstantsWidget::OnMeshTableSortRequested)
 		.SortMode(this, &SMutableConstantsWidget::GetMeshListColumnSortMode, MeshConstantTitles::MeshID)
 		.FillWidth(0.28f)
 	
 		+ SHeaderRow::Column(MeshConstantTitles::MeshVertices)
-		.DefaultLabel(FText(LOCTEXT("MeshVerticesCount","Vertices")))
-		.DefaultTooltip(FText(LOCTEXT("MeshVerticesCountColumnToolTip","Amount of vertice")))
+		.DefaultLabel(FText::FromName(MeshConstantTitles::MeshVertices))
+		.DefaultTooltip(FText(LOCTEXT("MeshVerticesCountColumnToolTip","Amount of vertices")))
 		.OnSort(this, &SMutableConstantsWidget::OnMeshTableSortRequested)
 		.SortMode(this, &SMutableConstantsWidget::GetMeshListColumnSortMode, MeshConstantTitles::MeshVertices)
 
 		+ SHeaderRow::Column(MeshConstantTitles::MeshIndices)
-		.DefaultLabel(FText(LOCTEXT("MeshIndicesCount","Indices")))
+		.DefaultLabel(FText::FromName(MeshConstantTitles::MeshIndices))
 		.DefaultTooltip(FText(LOCTEXT("MeshIndicesCountColumnToolTip","Amount of indices")))
 		.OnSort(this, &SMutableConstantsWidget::OnMeshTableSortRequested)
 		.SortMode(this, &SMutableConstantsWidget::GetMeshListColumnSortMode, MeshConstantTitles::MeshIndices)
 
 		+ SHeaderRow::Column(MeshConstantTitles::MeshChannels)
-		.DefaultLabel(FText(LOCTEXT("MeshVertexChannelsCount","Channels")))
+		.DefaultLabel(FText::FromName(MeshConstantTitles::MeshChannels))
 		.DefaultTooltip(FText(LOCTEXT("MeshVertexChannelsCountColumnToolTip","Amount of channels")))
 		.OnSort(this, &SMutableConstantsWidget::OnMeshTableSortRequested)
 		.SortMode(this, &SMutableConstantsWidget::GetMeshListColumnSortMode, MeshConstantTitles::MeshChannels)
 		
 		+ SHeaderRow::Column(MeshConstantTitles::MeshMemory)
-		.DefaultLabel(FText(LOCTEXT("MeshMemory","Size")))
+		.DefaultLabel(FText::FromName(MeshConstantTitles::MeshMemory))
 		.DefaultTooltip(FText(LOCTEXT("MeshMemoryColumnToolTip","Memory size")))
 		.OnSort(this, &SMutableConstantsWidget::OnMeshTableSortRequested)
 		.SortMode(this, &SMutableConstantsWidget::GetMeshListColumnSortMode, MeshConstantTitles::MeshMemory)
 	);
 	
 
+	// Hack to allow us to later set the array to be used. if no array is provided then the children slate to contain them later will not exist
+	TArray<TSharedPtr<FMutableConstantMaterialElement>> DummyMaterialElementsArray;
+	
+	// Image Constants List View object and Handler
+	TSharedPtr<SListView<TSharedPtr<FMutableConstantMaterialElement>>> ConstantMaterialsListView =
+		SNew(SListView<TSharedPtr<FMutableConstantMaterialElement>>)
+	.OnGenerateRow(this,&SMutableConstantsWidget::OnGenerateMaterialRow)
+		// We do require to provide something here or the slate to contain the children will not get generated
+	.ListItemsSource(&DummyMaterialElementsArray)		
+	.OnSelectionChanged(this,&SMutableConstantsWidget::OnSelectedMaterialChanged)
+	.SelectionMode(ESelectionMode::Single)
+	.HeaderRow
+	(
+		SNew(SHeaderRow)
+		+ SHeaderRow::Column(MaterialConstantTitles::MaterialID)
+		.DefaultLabel(FText::FromName(MaterialConstantTitles::MaterialID))
+		.OnSort(this, &SMutableConstantsWidget::OnMaterialsTableSortRequested)
+		.SortMode(this, &SMutableConstantsWidget::GetMaterialListColumnSortMode, MaterialConstantTitles::MaterialID )
+		.FillWidth(0.20f)
+		
+		+ SHeaderRow::Column(MaterialConstantTitles::Material)
+		.DefaultLabel(FText::FromName(MaterialConstantTitles::Material))
+		.FillWidth(0.40f)
+	
+		+ SHeaderRow::Column(MaterialConstantTitles::ImageParametersCount)
+		.DefaultLabel(FText::FromName(MaterialConstantTitles::ImageParametersCount))
+		.OnSort(this, &SMutableConstantsWidget::OnMaterialsTableSortRequested)
+		.SortMode(this, &SMutableConstantsWidget::GetMaterialListColumnSortMode, MaterialConstantTitles::ImageParametersCount )
+		.FillWidth(0.20f)
+
+		+ SHeaderRow::Column(MaterialConstantTitles::ColorParametersCount)
+		.DefaultLabel(FText::FromName(MaterialConstantTitles::ColorParametersCount))
+		.OnSort(this, &SMutableConstantsWidget::OnMaterialsTableSortRequested)
+		.SortMode(this, &SMutableConstantsWidget::GetMaterialListColumnSortMode, MaterialConstantTitles::ColorParametersCount )
+		.FillWidth(0.20f)
+
+		+ SHeaderRow::Column(MaterialConstantTitles::ScalarParametersCount)
+		.DefaultLabel(FText::FromName(MaterialConstantTitles::ScalarParametersCount))
+		.OnSort(this, &SMutableConstantsWidget::OnMaterialsTableSortRequested)
+		.SortMode(this, &SMutableConstantsWidget::GetMaterialListColumnSortMode, MaterialConstantTitles::ScalarParametersCount )
+		.FillWidth(0.20f)
+	);
+	
 	// Child structure
 	this->ChildSlot
 	[
@@ -607,9 +773,9 @@ void SMutableConstantsWidget::Construct(const FArguments& InArgs,const UE::Mutab
 			.BodyContent()
 			[
 				// Custom slate that will handle the updating of the elements displayed by the ConstantImagesListView object
-					SAssignNew(ImageListViewHandler,SMutableMultiPageListView<TSharedPtr<FMutableConstantImageElement>>)
-						.HostedListView(ConstantImagesListView)
-						.ElementsToSeparateInPages(ConstantImageElements)
+				SAssignNew(ImageListViewHandler,SMutableMultiPageListView<TSharedPtr<FMutableConstantImageElement>>)
+					.HostedListView(ConstantImagesListView)
+					.ElementsToSeparateInPages(ConstantImageElements)
 			]
 		]
 
@@ -629,7 +795,7 @@ void SMutableConstantsWidget::Construct(const FArguments& InArgs,const UE::Mutab
 			.BodyContent()
 			[
 				// Custom slate that will handle the updating of the elements displayed by the ConstantMeshesListView object
-				SAssignNew(MeshListViewHandler ,SMutableMultiPageListView<TSharedPtr<FMutableConstantMeshElement>>)
+				SAssignNew(MeshListViewHandler, SMutableMultiPageListView<TSharedPtr<FMutableConstantMeshElement>>)
 				.HostedListView(ConstantMeshesListView)	
 				.ElementsToSeparateInPages(ConstantMeshElements)
 			]
@@ -802,6 +968,28 @@ void SMutableConstantsWidget::Construct(const FArguments& InArgs,const UE::Mutab
 				.SelectionMode(ESelectionMode::Single)
 			]
 		]
+		
+		+ SScrollBox::Slot()
+		[
+			SNew(SSeparator)
+		]
+
+		// Material Constants
+		+ SScrollBox::Slot()
+		.Padding(0,InBetweenListsVerticalPadding)
+		[
+			SAssignNew(MaterialsExpandableArea,SExpandableArea)
+			.OnAreaExpansionChanged(this,&SMutableConstantsWidget::OnMaterialsRegionExpansionChanged)
+			.InitiallyCollapsed(true)
+			.AreaTitle(this,&SMutableConstantsWidget::OnDrawMaterialsAreaTitle)
+			.BodyContent()
+			[
+				// Custom slate that will handle the updating of the elements displayed by the ConstantMeshesListView object
+				SAssignNew(MaterialistViewHandler ,SMutableMultiPageListView<TSharedPtr<FMutableConstantMaterialElement>>)
+				.HostedListView(ConstantMaterialsListView)	
+				.ElementsToSeparateInPages(ConstantMaterialElements)
+			]
+		]
 	];
 
 	// Store all the expandable areas so they are later reachable using loops
@@ -815,6 +1003,7 @@ void SMutableConstantsWidget::Construct(const FArguments& InArgs,const UE::Mutab
 	ExpandableAreas.Add(MatricesExpandableArea);
 	ExpandableAreas.Add(SkeletonsExpandableArea);
 	ExpandableAreas.Add(PhysicsExpandableArea);
+	ExpandableAreas.Add(MaterialsExpandableArea);
 }
 
 
@@ -861,6 +1050,13 @@ TSharedRef<ITableRow> SMutableConstantsWidget::OnGenerateMeshRow(
 	return Row;
 }
 
+TSharedRef<ITableRow> SMutableConstantsWidget::OnGenerateMaterialRow(TSharedPtr<FMutableConstantMaterialElement> MutableConstantMaterialElement,
+	const TSharedRef<STableViewBase>& OwnerTable) const
+{
+	TSharedRef<SMutableConstantMaterialRow> Row = SNew(SMutableConstantMaterialRow, OwnerTable, MutableConstantMaterialElement);
+	return Row;
+}
+
 TSharedRef<ITableRow> SMutableConstantsWidget::OnGenerateLayoutRow(
 	TSharedPtr<FMutableConstantLayoutElement> MutableConstantLayoutElement,
 	const TSharedRef<STableViewBase>& OwnerTable) const
@@ -917,6 +1113,7 @@ TSharedRef<ITableRow> SMutableConstantsWidget::OnGeneratePhysicsRow(
 	return Row;
 }
 
+
 void SMutableConstantsWidget::OnImageTableSortRequested(EColumnSortPriority::Type ColumnSortPriority, const FName& ColumnID,
                                                         EColumnSortMode::Type ColumnSortMode)
 {
@@ -1088,6 +1285,80 @@ EColumnSortMode::Type SMutableConstantsWidget::GetMeshListColumnSortMode(FName C
 	return MeshListSortMode;
 }
 
+void SMutableConstantsWidget::OnMaterialsTableSortRequested(EColumnSortPriority::Type ColumnSortPriority, const FName& ColumnID,
+	EColumnSortMode::Type ColumnSortMode)
+{
+	// If the colum has been sorted on one way now do the inverse way of sorting
+	MaterialConstantsLastSortedColumnID = ColumnID;
+	MaterialListSortMode = ColumnSortMode;
+
+	Algo::StableSort(*ConstantMaterialElements, [&](const TSharedPtr<FMutableConstantMaterialElement>& A, const TSharedPtr<FMutableConstantMaterialElement>& B)
+	{
+		if (!A || !B)
+		{
+			return false;
+		}
+			
+		// Sort by material id
+		if (ColumnID == MaterialConstantTitles::MaterialID)
+		{
+			if (ColumnSortMode == EColumnSortMode::Ascending)
+			{
+				return A->IndexOnSourceVector < B->IndexOnSourceVector;
+			}
+			return A->IndexOnSourceVector > B->IndexOnSourceVector;
+		}
+			
+		// Sort by Image Parameters count
+		if (ColumnID == MaterialConstantTitles::ImageParametersCount)
+		{
+			if (ColumnSortMode == EColumnSortMode::Ascending)
+			{
+				return A->MutableMaterial->ImageParameters.Num() < B->MutableMaterial->ImageParameters.Num();
+			}
+			return A->MutableMaterial->ImageParameters.Num() > B->MutableMaterial->ImageParameters.Num();
+		}
+
+		// Sort by Color Parameters count
+		if (ColumnID == MaterialConstantTitles::ColorParametersCount)
+		{
+			if (ColumnSortMode == EColumnSortMode::Ascending)
+			{
+				return A->MutableMaterial->ColorParameters.Num() < B->MutableMaterial->ColorParameters.Num();
+			}
+			return A->MutableMaterial->ColorParameters.Num() > B->MutableMaterial->ColorParameters.Num();
+		}
+		
+		// Sort by Scalar Parameters count
+		if (ColumnID == MaterialConstantTitles::ScalarParametersCount)
+		{
+			if (ColumnSortMode == EColumnSortMode::Ascending)
+			{
+				return A->MutableMaterial->ScalarParameters.Num() < B->MutableMaterial->ScalarParameters.Num();
+			}
+			return A->MutableMaterial->ScalarParameters.Num() > B->MutableMaterial->ScalarParameters.Num();
+		}
+	
+		return false;
+	});
+	
+	if (MaterialistViewHandler)
+	{
+		MaterialistViewHandler->RegeneratePage();
+	}
+}
+
+
+EColumnSortMode::Type SMutableConstantsWidget::GetMaterialListColumnSortMode(FName ColumnName) const
+{
+	if (MaterialConstantsLastSortedColumnID != ColumnName)
+	{
+		return EColumnSortMode::None;
+	}
+
+	return MaterialListSortMode;
+}
+
 
 #pragma endregion 
 
@@ -1174,6 +1445,15 @@ void SMutableConstantsWidget::OnPhysicsRegionExpansionChanged(bool bExpanded)
 }
 
 
+void SMutableConstantsWidget::OnMaterialsRegionExpansionChanged(bool bExpanded)
+{
+	if (bExpanded)
+	{
+		ContractExpandableAreas(MaterialsExpandableArea);
+	}
+}
+
+
 void SMutableConstantsWidget::ContractExpandableAreas(const TSharedPtr<SExpandableArea>& InException)
 {
 #if UE_BUILD_DEBUG || UE_BUILD_DEVELOPMENT
@@ -1210,6 +1490,7 @@ void SMutableConstantsWidget::LoadConstantElements()
 	LoadConstantCurves();
 	LoadConstantSkeletons();
 	LoadConstantPhysics();
+	LoadConstantMaterials();
 }
 
 void SMutableConstantsWidget::LoadConstantStrings()
@@ -1233,7 +1514,7 @@ void SMutableConstantsWidget::LoadConstantStrings()
 		ConstantStringElements.Add(ConstantStringElement);
 	}
 
-	// Cache the size in memory of the constants as a formatted text so it is able to be be later used by the UI
+	// Cache the size in memory of the constants as a formatted text so it is able to be later used by the UI
 	ConstantStringsFormattedSize = GenerateTextForSize(ConstantStringsAccumulatedSize);
 }
 
@@ -1272,7 +1553,7 @@ void SMutableConstantsWidget::LoadConstantImages()
 		ConstantImageElements->Add(ConstantImageElement);
 	}
 
-	// Cache the size in memory of the constants as a formatted text so it is able to be be later used by the UI
+	// Cache the size in memory of the constants as a formatted text so it is able to be later used by the UI
 	ConstantImagesFormattedSize = GenerateTextForSize(ConstantImagesAccumulatedSize);
 }
 
@@ -1319,7 +1600,7 @@ void SMutableConstantsWidget::LoadConstantMeshes()
 		ConstantMeshElements->Add(ConstantMeshElement);
 	}
 
-	// Cache the size in memory of the constants as a formatted text so it is able to be be later used by the UI
+	// Cache the size in memory of the constants as a formatted text so it is able to be later used by the UI
 	ConstantMeshesFormattedSize = GenerateTextForSize(ConstantMeshesAccumulatedSize);
 }
 
@@ -1344,7 +1625,7 @@ void SMutableConstantsWidget::LoadConstantLayouts()
 		ConstantLayoutElements.Add(ConstantLayoutElement);
 	}
 
-	// Cache the size in memory of the constants as a formatted text so it is able to be be later used by the UI
+	// Cache the size in memory of the constants as a formatted text so it is able to be later used by the UI
 	ConstantLayoutsFormattedSize = GenerateTextForSize(Stream.GetBufferSize());
 }
 
@@ -1368,7 +1649,7 @@ void SMutableConstantsWidget::LoadConstantSkeletons()
 		ConstantSkeletonElements.Add(ConstantSkeletonElement);
 	}
 
-	// Cache the size in memory of the constants as a formatted text so it is able to be be later used by the UI
+	// Cache the size in memory of the constants as a formatted text so it is able to be later used by the UI
 	ConstantSkeletonsFormattedSize = GenerateTextForSize(Stream.GetBufferSize());
 }
 
@@ -1392,11 +1673,44 @@ void SMutableConstantsWidget::LoadConstantPhysics()
 		ConstantPhysicsElements.Add(ConstantPhysicsElement);
 	}
 
-	// Cache the size in memory of the constants as a formatted text so it is able to be be later used by the UI
+	// Cache the size in memory of the constants as a formatted text so it is able to be later used by the UI
 	ConstantPhysicsFormattedSize = GenerateTextForSize(Stream.GetBufferSize());
 }
 
 
+void SMutableConstantsWidget::LoadConstantMaterials()
+{
+	check (MutableProgramPtr);
+	const int32 ConstantsCount = MutableProgramPtr->ConstantMaterials.Num();
+	
+	if (ConstantMaterialElements)
+	{
+		ConstantMaterialElements->Empty(ConstantsCount);
+	}
+	else
+	{
+		ConstantMaterialElements = MakeShared<TArray<TSharedPtr<FMutableConstantMaterialElement>>>();
+		ConstantMaterialElements->Reserve(ConstantsCount);
+	}
+	
+	UE::Mutable::Private::FOutputMemoryStream Stream;
+	UE::Mutable::Private::FOutputArchive Archive{&Stream};
+	
+	for (int32 MaterialIndex = 0; MaterialIndex < ConstantsCount; MaterialIndex++)
+	{
+		TSharedPtr<FMutableConstantMaterialElement> ConstantMaterialElement = MakeShared<FMutableConstantMaterialElement>();
+		ConstantMaterialElement->MutableMaterial = MutableProgramPtr->ConstantMaterials[MaterialIndex];
+		ConstantMaterialElement->IndexOnSourceVector = MaterialIndex;
+		ConstantMaterialElement->MutableMaterial->Serialise(Archive);
+
+		ConstantMaterialElements->Add(ConstantMaterialElement);
+	}
+
+	// Cache the size in memory of the constants as a formatted text so it is able to be later used by the UI
+	ConstantMaterialsFormattedSize = GenerateTextForSize(Stream.GetBufferSize());
+}
+
+
 void SMutableConstantsWidget::LoadConstantProjectors()
 {
 	check (MutableProgramPtr);
@@ -1417,7 +1731,7 @@ void SMutableConstantsWidget::LoadConstantProjectors()
 		ConstantProjectorElements.Add(ConstantProjectorElement);
 	}
 
-	// Cache the size in memory of the constants as a formatted text so it is able to be be later used by the UI
+	// Cache the size in memory of the constants as a formatted text so it is able to be later used by the UI
 	ConstantProjectorsFormattedSize = GenerateTextForSize(Stream.GetBufferSize());
 }
 
@@ -1441,7 +1755,7 @@ void SMutableConstantsWidget::LoadConstantMatrices()
 		ConstantMatrixElements.Add(ConstantMatrixElement);
 	}
 
-	// Cache the size in memory of the constants as a formatted text so it is able to be be later used by the UI
+	// Cache the size in memory of the constants as a formatted text so it is able to be later used by the UI
 	ConstantMatricesFormattedSize = GenerateTextForSize(Stream.GetBufferSize());
 }
 
@@ -1465,7 +1779,7 @@ void SMutableConstantsWidget::LoadConstantShapes()
 		ConstantShapeElements.Add(ConstantShapeElement);
 	}
 	
-	// Cache the size in memory of the constants as a formatted text so it is able to be be later used by the UI
+	// Cache the size in memory of the constants as a formatted text so it is able to be later used by the UI
 	ConstantShapesFormattedSize = GenerateTextForSize(Stream.GetBufferSize());
 }
 
@@ -1489,7 +1803,7 @@ void SMutableConstantsWidget::LoadConstantCurves()
 		ConstantCurveElements.Add(ConstantCurveElement);
 	}
 
-	// Cache the size in memory of the constants as a formatted text so it is able to be be later used by the UI
+	// Cache the size in memory of the constants as a formatted text so it is able to be later used by the UI
 	ConstantCurvesFormattedSize = GenerateTextForSize(Stream.GetBufferSize());
 }
 
@@ -1540,7 +1854,10 @@ void SMutableConstantsWidget::ClearSelectedConstantItems(UE::Mutable::Private::E
 	if (ExceptionDataType != UE::Mutable::Private::EDataType::Image)
 	{
 		ImageListViewHandler->ClearSelection();
-		//ConstantImagesListView->ClearSelection();
+	}
+	if (ExceptionDataType != UE::Mutable::Private::EDataType::Material)
+	{
+		MaterialistViewHandler->ClearSelection();
 	}
 
 	// todo: if you add more slates for the types update this to make sure they behave as the others
@@ -1552,7 +1869,7 @@ void SMutableConstantsWidget::OnSelectedStringChanged(
 {
 	if (MutableConstantStringElement)
 	{
-		static constexpr UE::Mutable::Private::EDataType SlateDataType = UE::Mutable::Private::EDataType::String;
+		constexpr UE::Mutable::Private::EDataType SlateDataType = UE::Mutable::Private::EDataType::String;
 
 		// Clear the selected CodeViewer row and all other constant viewer slates that are not the type provided.
 		MutableCodeViewerPtr->ClearSelectedTreeRow();
@@ -1570,7 +1887,7 @@ void SMutableConstantsWidget::OnSelectedImageChanged(
 {
 	if (MutableConstantImageElement)
 	{
-		static constexpr UE::Mutable::Private::EDataType SlateDataType = UE::Mutable::Private::EDataType::Image;
+		constexpr UE::Mutable::Private::EDataType SlateDataType = UE::Mutable::Private::EDataType::Image;
 
 		MutableCodeViewerPtr->ClearSelectedTreeRow();
 		ClearSelectedConstantItems(SlateDataType);
@@ -1586,7 +1903,7 @@ void SMutableConstantsWidget::OnSelectedMeshChanged(
 {
 	if (MutableConstantMeshElement)
 	{
-		static constexpr UE::Mutable::Private::EDataType SlateDataType = UE::Mutable::Private::EDataType::Mesh;
+		constexpr UE::Mutable::Private::EDataType SlateDataType = UE::Mutable::Private::EDataType::Mesh;
 		
 		MutableCodeViewerPtr->ClearSelectedTreeRow();
 		ClearSelectedConstantItems(SlateDataType);
@@ -1596,13 +1913,29 @@ void SMutableConstantsWidget::OnSelectedMeshChanged(
 	}
 }
 
+
+void SMutableConstantsWidget::OnSelectedMaterialChanged(TSharedPtr<FMutableConstantMaterialElement> MutableConstantMaterialElement,
+	ESelectInfo::Type SelectionType) const
+{
+	if (MutableConstantMaterialElement)
+	{
+		constexpr UE::Mutable::Private::EDataType SlateDataType = UE::Mutable::Private::EDataType::Material;
+		
+		MutableCodeViewerPtr->ClearSelectedTreeRow();
+		ClearSelectedConstantItems(SlateDataType);
+
+		MutableCodeViewerPtr->PreviewMutableMaterial(MutableConstantMaterialElement->MutableMaterial);
+		MutableCodeViewerPtr->CacheAddressesRelatedWithConstantResource(SlateDataType,MutableConstantMaterialElement->IndexOnSourceVector);
+	}
+}
+
 void SMutableConstantsWidget::OnSelectedLayoutChanged(
 	TSharedPtr<FMutableConstantLayoutElement> MutableConstantLayoutElement,
 	ESelectInfo::Type SelectionType) const
 {
 	if (MutableConstantLayoutElement)
 	{
-		static constexpr UE::Mutable::Private::EDataType SlateDataType = UE::Mutable::Private::EDataType::Layout;
+		constexpr UE::Mutable::Private::EDataType SlateDataType = UE::Mutable::Private::EDataType::Layout;
 		
 		MutableCodeViewerPtr->ClearSelectedTreeRow();
 		ClearSelectedConstantItems(SlateDataType);
@@ -1618,7 +1951,7 @@ void SMutableConstantsWidget::OnSelectedProjectorChanged(
 {
 	if (MutableConstantProjectorElement)
 	{
-		static constexpr UE::Mutable::Private::EDataType SlateDataType = UE::Mutable::Private::EDataType::Projector;
+		constexpr UE::Mutable::Private::EDataType SlateDataType = UE::Mutable::Private::EDataType::Projector;
 
 		MutableCodeViewerPtr->ClearSelectedTreeRow();
 		ClearSelectedConstantItems(SlateDataType);
@@ -1634,7 +1967,7 @@ void SMutableConstantsWidget::OnSelectedMatrixChanged(
 {
 	if (MutableConstantMatrixElement)
 	{
-		static constexpr UE::Mutable::Private::EDataType SlateDataType = UE::Mutable::Private::EDataType::Matrix;
+		constexpr UE::Mutable::Private::EDataType SlateDataType = UE::Mutable::Private::EDataType::Matrix;
 
 		MutableCodeViewerPtr->ClearSelectedTreeRow();
 		ClearSelectedConstantItems(SlateDataType);
@@ -1650,7 +1983,7 @@ void SMutableConstantsWidget::OnSelectedShapeChanged(
 {
 	if (MutableConstantShapeElement)
 	{
-		static constexpr UE::Mutable::Private::EDataType SlateDataType = UE::Mutable::Private::EDataType::Shape;
+		constexpr UE::Mutable::Private::EDataType SlateDataType = UE::Mutable::Private::EDataType::Shape;
 
 		MutableCodeViewerPtr->ClearSelectedTreeRow();
 		ClearSelectedConstantItems(SlateDataType);
@@ -1666,7 +1999,7 @@ void SMutableConstantsWidget::OnSelectedCurveChanged(
 {
 	if (MutableConstantCurveElement)
 	{
-		static constexpr UE::Mutable::Private::EDataType SlateDataType = UE::Mutable::Private::EDataType::Curve;
+		constexpr UE::Mutable::Private::EDataType SlateDataType = UE::Mutable::Private::EDataType::Curve;
 
 		MutableCodeViewerPtr->ClearSelectedTreeRow();
 		ClearSelectedConstantItems(SlateDataType);
@@ -1682,7 +2015,7 @@ void SMutableConstantsWidget::OnSelectedSkeletonChanged(
 {
 	if (MutableConstantSkeletonElement)
 	{
-		static constexpr UE::Mutable::Private::EDataType SlateDataType = UE::Mutable::Private::EDataType::Skeleton;
+		constexpr UE::Mutable::Private::EDataType SlateDataType = UE::Mutable::Private::EDataType::Skeleton;
 
 		MutableCodeViewerPtr->ClearSelectedTreeRow();
 		ClearSelectedConstantItems(SlateDataType);
@@ -1698,7 +2031,7 @@ void SMutableConstantsWidget::OnSelectedPhysicsChanged(
 {
 	if (MutableConstantPhysicsElement)
 	{
-		static constexpr UE::Mutable::Private::EDataType SlateDataType = UE::Mutable::Private::EDataType::PhysicsAsset;
+		constexpr UE::Mutable::Private::EDataType SlateDataType = UE::Mutable::Private::EDataType::PhysicsAsset;
 
 		MutableCodeViewerPtr->ClearSelectedTreeRow();
 		ClearSelectedConstantItems(SlateDataType);
@@ -1715,17 +2048,17 @@ void SMutableConstantsWidget::OnSelectedPhysicsChanged(
 
 FText SMutableConstantsWidget::OnDrawStringsAreaTitle() const
 {
-	return FText::Format( LOCTEXT("StringConstantsTitle", "String Constants ({0}) : {1} "),ConstantStringElements.Num(), ConstantStringsFormattedSize);
+	return FText::Format( LOCTEXT("StringConstantsTitle", "String Constants ({0}) : {1} "),  ConstantStringElements.Num() , ConstantStringsFormattedSize);
 }
 
 FText SMutableConstantsWidget::OnDrawImagesAreaTitle() const
 {
-	return FText::Format( LOCTEXT("ImageConstantsTitle", "Image Constants ({0}) : {1} "), ConstantImageElements->Num(), ConstantImagesFormattedSize);
+	return FText::Format( LOCTEXT("ImageConstantsTitle", "Image Constants ({0}) : {1} "), ConstantImageElements ? ConstantImageElements->Num() : 0, ConstantImagesFormattedSize);
 }
 
 FText SMutableConstantsWidget::OnDrawMeshesAreaTitle() const
 {
-	return FText::Format(LOCTEXT("MeshConstantsTitle", "Mesh Constants ({0}) : {1} "), ConstantMeshElements->Num(), ConstantMeshesFormattedSize) ;
+	return FText::Format(LOCTEXT("MeshConstantsTitle", "Mesh Constants ({0}) : {1} "), ConstantMeshElements ? ConstantMeshElements->Num() : 0, ConstantMeshesFormattedSize) ;
 }
 
 FText SMutableConstantsWidget::OnDrawLayoutsAreaTitle() const
@@ -1763,6 +2096,11 @@ FText SMutableConstantsWidget::OnDrawPhysicsAreaTitle() const
 	return FText::Format( LOCTEXT("PhysicsConstantsTitle", "Physics Constants ({0}) : {1} "), ConstantPhysicsElements.Num(), ConstantPhysicsFormattedSize);
 }
 
+FText SMutableConstantsWidget::OnDrawMaterialsAreaTitle() const
+{
+	return FText::Format( LOCTEXT("MaterialConstantsTitle", "Material Constants ({0}) : {1} "), ConstantMaterialElements ? ConstantMaterialElements->Num() : 0, ConstantMaterialsFormattedSize);
+}
+
 #pragma endregion
 
 
diff --git a/Engine/Plugins/Mutable/Source/CustomizableObjectEditor/Private/MuCOE/SMutableConstantsWidget.h b/Engine/Plugins/Mutable/Source/CustomizableObjectEditor/Private/MuCOE/SMutableConstantsWidget.h
--- a/Engine/Plugins/Mutable/Source/CustomizableObjectEditor/Private/MuCOE/SMutableConstantsWidget.h
+++ b/Engine/Plugins/Mutable/Source/CustomizableObjectEditor/Private/MuCOE/SMutableConstantsWidget.h
@@ -78,6 +78,14 @@ struct FMutableConstantPhysicsElement : public FMutableConstantElement
 	UE::Mutable::Private::TManagedPtr<const UE::Mutable::Private::FPhysicsBody> Physics;
 };
 
+/**
+* Cache object used for the generation of the ui elements related to the constant materials found on the model
+*/
+struct FMutableConstantMaterialElement : public FMutableConstantElement
+{
+	UE::Mutable::Private::TManagedPtr<const UE::Mutable::Private::FMaterial> MutableMaterial;
+};
+
 /**
 * Cache object used for the generation of the ui elements related to the constant projectors found on the model
 */
@@ -173,6 +181,9 @@ private:
 	TSharedPtr<SMutableMultiPageListView<TSharedPtr<FMutableConstantMeshElement>>> MeshListViewHandler;
 	TSharedPtr<TArray<TSharedPtr<FMutableConstantMeshElement>>> ConstantMeshElements;
 	
+	TSharedPtr<SMutableMultiPageListView<TSharedPtr<FMutableConstantMaterialElement>>> MaterialistViewHandler;
+	TSharedPtr<TArray<TSharedPtr<FMutableConstantMaterialElement>>> ConstantMaterialElements;
+	
 	TArray<TSharedPtr<FMutableConstantStringElement>> ConstantStringElements;
 	TArray<TSharedPtr<FMutableConstantLayoutElement>> ConstantLayoutElements;
 	TArray<TSharedPtr<FMutableConstantProjectorElement>> ConstantProjectorElements;
@@ -197,6 +208,7 @@ private:
 	void LoadConstantCurves();
 	void LoadConstantSkeletons();
 	void LoadConstantPhysics();
+	void LoadConstantMaterials();
 
 	/*
 	 * Proxy slates operation objects 
@@ -205,6 +217,7 @@ private:
 	void OnSelectedStringChanged(TSharedPtr<FMutableConstantStringElement> MutableConstantStringElement, ESelectInfo::Type SelectionType) const;
 	void OnSelectedImageChanged(TSharedPtr<FMutableConstantImageElement> MutableConstantImageElement, ESelectInfo::Type SelectionType) const;
 	void OnSelectedMeshChanged(TSharedPtr<FMutableConstantMeshElement> MutableConstantMeshElement, ESelectInfo::Type SelectionType) const;
+	void OnSelectedMaterialChanged(TSharedPtr<FMutableConstantMaterialElement> MutableConstantMaterialElement, ESelectInfo::Type SelectionType) const;
 	void OnSelectedLayoutChanged(TSharedPtr<FMutableConstantLayoutElement> MutableConstantLayoutElement, ESelectInfo::Type SelectionType) const;
 	void OnSelectedProjectorChanged(TSharedPtr<FMutableConstantProjectorElement> MutableConstantProjectorElement, ESelectInfo::Type SelectionType) const;
 	void OnSelectedMatrixChanged(TSharedPtr<FMutableConstantMatrixElement> MutableConstantMatrixElement, ESelectInfo::Type SelectionType) const;
@@ -212,6 +225,7 @@ private:
 	void OnSelectedCurveChanged(TSharedPtr<FMutableConstantCurveElement> MutableConstantCurveElement, ESelectInfo::Type SelectionType) const;
 	void OnSelectedSkeletonChanged(TSharedPtr<FMutableConstantSkeletonElement> MutableConstantSkeletonElement, ESelectInfo::Type SelectionType) const;
 	void OnSelectedPhysicsChanged(TSharedPtr<FMutableConstantPhysicsElement> MutableConstantPhysicsElement, ESelectInfo::Type SelectionType) const;
+	
 	/*
 	* UI updating methods
 	*/
@@ -219,6 +233,7 @@ private:
 	TSharedRef<ITableRow> OnGenerateStringRow(TSharedPtr<FMutableConstantStringElement> MutableConstantStringElement, const TSharedRef<STableViewBase>& OwnerTable) const;
 	TSharedRef<ITableRow> OnGenerateImageRow(TSharedPtr<FMutableConstantImageElement> MutableConstantImageElement, const TSharedRef<STableViewBase>& OwnerTable) const;
 	TSharedRef<ITableRow> OnGenerateMeshRow(TSharedPtr<FMutableConstantMeshElement> MutableConstantMeshElement, const TSharedRef<STableViewBase>& OwnerTable) const;
+	TSharedRef<ITableRow> OnGenerateMaterialRow(TSharedPtr<FMutableConstantMaterialElement> MutableConstantMaterialElement, const TSharedRef<STableViewBase>& OwnerTable) const;
 	TSharedRef<ITableRow> OnGenerateLayoutRow(TSharedPtr<FMutableConstantLayoutElement> MutableConstantLayoutElement, const TSharedRef<STableViewBase>& OwnerTable) const;
 	TSharedRef<ITableRow> OnGenerateProjectorRow(TSharedPtr<FMutableConstantProjectorElement> MutableConstantProjectorElement, const TSharedRef<STableViewBase>& OwnerTable)const;
 	TSharedRef<ITableRow> OnGenerateMatrixRow(TSharedPtr<FMutableConstantMatrixElement> MutableConstantMatrixElement, const TSharedRef<STableViewBase>& OwnerTable) const;
@@ -252,8 +267,19 @@ private:
 	void OnMeshTableSortRequested(EColumnSortPriority::Type ColumnSortPriority, const FName& ColumnID, EColumnSortMode::Type ColumnSortMode);
 	EColumnSortMode::Type GetMeshListColumnSortMode(FName ColumnName) const;
 
+	/**
+	 * Materials List sorting methods
+	 */
+	
+	FName MaterialConstantsLastSortedColumnID = "";
+	EColumnSortMode::Type MaterialListSortMode = EColumnSortMode::Type::None;
+	
+	void OnMaterialsTableSortRequested(EColumnSortPriority::Type ColumnSortPriority, const FName& ColumnID, EColumnSortMode::Type ColumnSortMode);
+	EColumnSortMode::Type GetMaterialListColumnSortMode(FName ColumnName) const;
+
+	
 	/*
-	 * Expandable areas objects and behaviours
+	 * Expandable areas objects and behaviors
 	 */
 
 	// Pointers to all expandable areas part of this slate 
@@ -267,6 +293,7 @@ private:
 	TSharedPtr<SExpandableArea> CurvesExpandableArea;
 	TSharedPtr<SExpandableArea> SkeletonsExpandableArea;
 	TSharedPtr<SExpandableArea> PhysicsExpandableArea;
+	TSharedPtr<SExpandableArea> MaterialsExpandableArea;
 
 	/** Array with all expandable areas set on this object. Used for dynamic expansion/contraction */
 	TArray<TSharedPtr<SExpandableArea>> ExpandableAreas;
@@ -285,7 +312,8 @@ private:
 	void OnCurvesRegionExpansionChanged(bool bExpanded);
 	void OnSkeletonsRegionExpansionChanged(bool bExpanded);
 	void OnPhysicsRegionExpansionChanged(bool bExpanded);
-	
+	void OnMaterialsRegionExpansionChanged(bool bExpanded);
+
 	/** Method called each time one expandable area changes expansion state
 	 * @param InException - Pointer to the expandable area that will not be contracted while all the others will
 	 */
@@ -306,7 +334,8 @@ private:
 	FText ConstantCurvesFormattedSize;
 	FText ConstantSkeletonsFormattedSize;
 	FText ConstantPhysicsFormattedSize;
-	
+	FText ConstantMaterialsFormattedSize;
+
 	/*
 	 * Callback methods used for drawing the titles of each of the constant expandable areas
 	 */
@@ -321,7 +350,8 @@ private:
 	FText OnDrawCurvesAreaTitle() const;
 	FText OnDrawSkeletonsAreaTitle() const;
 	FText OnDrawPhysicsAreaTitle() const;
-	
+	FText OnDrawMaterialsAreaTitle() const;
+
 };
 
 
diff --git a/Engine/Plugins/Mutable/Source/CustomizableObjectEditor/Private/MuCOE/SMutableImageViewer.h b/Engine/Plugins/Mutable/Source/CustomizableObjectEditor/Private/MuCOE/SMutableImageViewer.h
--- a/Engine/Plugins/Mutable/Source/CustomizableObjectEditor/Private/MuCOE/SMutableImageViewer.h
+++ b/Engine/Plugins/Mutable/Source/CustomizableObjectEditor/Private/MuCOE/SMutableImageViewer.h
@@ -41,12 +41,12 @@ public:
 	void Construct( const FArguments& InArgs );
 
 	// SWidgetInterface
-	int32 OnPaint(const FPaintArgs& Args, const FGeometry& AllottedGeometry, const FSlateRect& MyClippingRect, FSlateWindowElementList& OutDrawElements, int32 LayerId, const FWidgetStyle& InWidgetStyle, bool bParentEnabled ) const override;
-	FVector2D ComputeDesiredSize(float) const override;
+	virtual int32 OnPaint(const FPaintArgs& Args, const FGeometry& AllottedGeometry, const FSlateRect& MyClippingRect, FSlateWindowElementList& OutDrawElements, int32 LayerId, const FWidgetStyle& InWidgetStyle, bool bParentEnabled ) const override;
+	virtual FVector2D ComputeDesiredSize(float) const override;
 
 	// FGCObject interface
-	void AddReferencedObjects(FReferenceCollector& Collector) override;
-	FString GetReferencerName() const override;
+	virtual void AddReferencedObjects(FReferenceCollector& Collector) override;
+	virtual FString GetReferencerName() const override;
 
 	// Own interface
 
@@ -84,7 +84,7 @@ public:
 	void Construct(const FArguments& InArgs);
 
 	// SWidget interface
-	void Tick(const FGeometry& AllottedGeometry, const double InCurrentTime, const float InDeltaTime) override;
+	virtual void Tick(const FGeometry& AllottedGeometry, const double InCurrentTime, const float InDeltaTime) override;
 
 	// Own interface
 
diff --git a/Engine/Plugins/Mutable/Source/CustomizableObjectEditor/Private/MuCOE/SMutableMaterialViewer.cpp b/Engine/Plugins/Mutable/Source/CustomizableObjectEditor/Private/MuCOE/SMutableMaterialViewer.cpp
new file mode 100644
--- /dev/null
+++ b/Engine/Plugins/Mutable/Source/CustomizableObjectEditor/Private/MuCOE/SMutableMaterialViewer.cpp
@@ -0,0 +1,585 @@
+// Copyright Epic Games, Inc. All Rights Reserved.
+
+#include "MuCOE/SMutableMaterialViewer.h"
+
+#include "Editor.h"
+#include "PropertyCustomizationHelpers.h"
+#include "SMutableColorViewer.h"
+#include "SMutableImageViewer.h"
+#include "Components/HorizontalBox.h"
+#include "MuT/NodeMaterialModify.h"
+#include "MuT/TypeInfo.h"
+#include "Subsystems/AssetEditorSubsystem.h"
+#include "Widgets/Colors/SColorBlock.h"
+#include "Widgets/Input/SRichTextHyperlink.h"
+#include "Widgets/Text/STextBlock.h"
+#include "Widgets/Views/SListView.h"
+
+
+class ITableRow;
+class STableViewBase;
+class SWidget;
+class USkeleton;
+
+#define LOCTEXT_NAMESPACE "CustomizableObjectEditor"
+
+namespace MutableMaterialParameterColumns
+{
+	static const FName ParameterIndexID("Index");
+	static const FName ParameterNameID("Parameter");
+	static const FName ParameterLayerID("Layer");
+
+	// Image Params
+	static const FName ParameterImageAddressID("Image OP");
+	static const FName ParameterImageID("Image");
+	
+	// Color Params
+	static const FName ParameterColorID("Color");
+	
+	// Scalar Params
+	static const FName ParameterScalarID("Scalar");
+}
+
+class SMutableMaterialImageParameterRow : public SMultiColumnTableRow<TSharedPtr<FMutableMaterialImageParameterElement>>
+{
+public:
+	void Construct(const FArguments& Args, const TSharedRef<STableViewBase>& InOwnerTableView,
+				   const TSharedPtr<FMutableMaterialImageParameterElement>& InRowItem, TWeakPtr<SMutableCodeViewer> InHostCodeViewer)
+	{
+		MutableCodeViewer = InHostCodeViewer;
+		RowItem = InRowItem;
+		
+		SMultiColumnTableRow<TSharedPtr<FMutableMaterialImageParameterElement>>::Construct(
+			STableRow::FArguments()
+			.ShowSelection(true)
+			, InOwnerTableView
+		);
+	}
+
+	virtual TSharedRef<SWidget> GenerateWidgetForColumn(const FName& InColumnName) override
+	{
+		if (!RowItem)
+		{
+			return SNullWidget::NullWidget;
+		}
+		
+		// Index of the parameter within the array of material parameters of the same type
+		if (InColumnName == MutableMaterialParameterColumns::ParameterIndexID)
+		{
+			return 
+				SNew(STextBlock)
+				.Text(FText::AsNumber(RowItem->Index));
+		}
+
+		// Material Parameter Layer index
+		if (InColumnName == MutableMaterialParameterColumns::ParameterLayerID)
+		{
+			return 
+				SNew(STextBlock)
+				.Text(FText::AsNumber(RowItem->LayerIndex));
+		}
+		
+		// Material Parameter Name
+		if (InColumnName == MutableMaterialParameterColumns::ParameterNameID)
+		{
+			return 
+				SNew(STextBlock)
+				.Text(FText::FromName(RowItem->ParameterName));
+		}
+		
+		// OP ADDRESS
+		if (InColumnName == MutableMaterialParameterColumns::ParameterImageAddressID)
+		{
+			if (const UE::Mutable::Private::OP::ADDRESS* ImageAddress = RowItem->Value.ImageParameter.TryGet<UE::Mutable::Private::OP::ADDRESS>())
+			{
+				return 
+					SNew(STextBlock)
+					.Text(FText::AsNumber((int32)*ImageAddress));
+			}
+			else
+			{
+				return SNullWidget::NullWidget;
+			}
+		}
+		
+		
+		// Try showing the evaluated Image. This will currently fail and produce a blank image
+		if (InColumnName == MutableMaterialParameterColumns::ParameterImageID)
+		{
+			if (TSharedPtr<SMutableCodeViewer> PinnedCodeViewer = MutableCodeViewer.Pin())
+			{
+				if (const UE::Mutable::Private::OP::ADDRESS* ImageAddress = RowItem->Value.ImageParameter.TryGet<UE::Mutable::Private::OP::ADDRESS>())
+				{
+					UE::Mutable::Private::TManagedPtr<const UE::Mutable::Private::FImage> GeneratedImage = PinnedCodeViewer->BuildImage(*ImageAddress);
+					return 
+						SNew(SMutableImageViewer)
+						.Image(GeneratedImage);
+				}
+			}
+			
+			return SNullWidget::NullWidget;
+		}
+
+		// Invalid column name so no widget will be produced 
+		checkNoEntry();
+		return SNullWidget::NullWidget;
+	}
+
+private:
+	TSharedPtr<FMutableMaterialImageParameterElement> RowItem;
+	
+	TWeakPtr<SMutableCodeViewer> MutableCodeViewer;
+};
+
+
+class SMutableMaterialColorParameterRow : public SMultiColumnTableRow<TSharedPtr<FMutableMaterialColorParameterElement>>
+{
+	public:
+	void Construct(const FArguments& Args, const TSharedRef<STableViewBase>& InOwnerTableView,
+				   const TSharedPtr<FMutableMaterialColorParameterElement>& InRowItem)
+	{
+		RowItem = InRowItem;
+		
+		SMultiColumnTableRow<TSharedPtr<FMutableMaterialColorParameterElement>>::Construct(
+			STableRow::FArguments()
+			.ShowSelection(true)
+			, InOwnerTableView
+		);
+	}
+
+	virtual TSharedRef<SWidget> GenerateWidgetForColumn(const FName& InColumnName) override
+	{
+		if (!RowItem)
+		{
+			return SNullWidget::NullWidget;
+		}
+	
+		// Index of the parameter within the array of material parameters of the same type
+		if (InColumnName == MutableMaterialParameterColumns::ParameterIndexID)
+		{
+			return 
+				SNew(STextBlock)
+				.Text(FText::AsNumber(RowItem->Index));
+		}
+		
+		// Material Parameter Layer index
+		if (InColumnName == MutableMaterialParameterColumns::ParameterLayerID)
+		{
+			return 
+				SNew(STextBlock)
+				.Text(FText::AsNumber(RowItem->LayerIndex));
+		}
+		
+		// Material Parameter Name
+		if (InColumnName == MutableMaterialParameterColumns::ParameterNameID)
+		{
+			return 
+				SNew(STextBlock)
+				.Text(FText::FromName(RowItem->ParameterName));
+		}
+
+		// Generate the sub table here
+		if (InColumnName == MutableMaterialParameterColumns::ParameterColorID)
+		{
+			TSharedRef<SMutableColorViewer> Viewer = SNew(SMutableColorViewer);
+			Viewer->SetColor(RowItem->Value);
+			return Viewer;
+		}
+
+		// Invalid column name so no widget will be produced 
+		checkNoEntry();
+		return SNullWidget::NullWidget;
+	}
+
+private:
+	TSharedPtr<FMutableMaterialColorParameterElement> RowItem;
+};
+
+
+class SMutableMaterialScalarParameterRow : public SMultiColumnTableRow<TSharedPtr<FMutableMaterialScalarParameterElement>>
+{
+public:
+	void Construct(const FArguments& Args, const TSharedRef<STableViewBase>& InOwnerTableView,
+				   const TSharedPtr<FMutableMaterialScalarParameterElement>& InRowItem)
+	{
+		RowItem = InRowItem;
+		
+		SMultiColumnTableRow<TSharedPtr<FMutableMaterialScalarParameterElement>>::Construct(
+	STableRow::FArguments()
+			.ShowSelection(true)
+			, InOwnerTableView
+		);
+	}
+
+	virtual TSharedRef<SWidget> GenerateWidgetForColumn(const FName& InColumnName) override
+	{
+		if (!RowItem)
+		{
+			return SNullWidget::NullWidget;
+		}
+	
+		// Index of the parameter within the array of material parameters of the same type
+		if (InColumnName == MutableMaterialParameterColumns::ParameterIndexID)
+		{
+			return 
+				SNew(STextBlock)
+				.Text(FText::AsNumber(RowItem->Index));
+		}
+		
+		// Material Parameter Layer index
+		if (InColumnName == MutableMaterialParameterColumns::ParameterLayerID)
+		{
+			return 
+				SNew(STextBlock)
+				.Text(FText::AsNumber(RowItem->LayerIndex));
+		}
+		
+		// Material Parameter Name
+		if (InColumnName == MutableMaterialParameterColumns::ParameterNameID)
+		{
+			return 
+				SNew(STextBlock)
+				.Text(FText::FromName(RowItem->ParameterName));
+		}
+
+		// Generate the sub table here
+		if (InColumnName == MutableMaterialParameterColumns::ParameterScalarID)
+		{
+			return SNew(STextBlock).Text(FText::AsNumber(RowItem->Value));
+		}
+
+		// Invalid column name so no widget will be produced 
+		checkNoEntry();
+		return SNullWidget::NullWidget;
+	}
+
+private:
+	TSharedPtr<FMutableMaterialScalarParameterElement> RowItem;
+};
+
+
+void SMutableMaterialViewer::Construct(const FArguments& InArgs)
+{
+	MutableCodeViewer = InArgs._HostCodeViewer;
+	
+	const FText ParameterIndexNameLabel = FText::FromString(TEXT("Index"));
+	const FText ParameterLayerNameLabel = FText::FromString(TEXT("Layer"));
+	const FText ParameterNameLabel = FText::FromString(TEXT("Parameter"));
+	
+	ChildSlot
+	[
+		SNew(SVerticalBox)
+		
+		+ SVerticalBox::Slot()
+		.AutoHeight()
+		[
+			SNew(SVerticalBox)
+			
+			+ SVerticalBox::Slot()
+			.AutoHeight()
+			[
+				SNew (SHorizontalBox)
+
+				+ SHorizontalBox::Slot()
+				.Padding(5)
+				[
+					// Container for the passthrough material preview
+					SAssignNew(PreviewContainer, SBox)
+				]
+
+				+ SHorizontalBox::Slot()
+				.Padding(5)
+				[
+					SNew(SVerticalBox)
+					
+					// Image count
+					+ SVerticalBox::Slot()
+					[
+						SNew(STextBlock).
+						Text(this, &SMutableMaterialViewer::GetImageParametersCountAsText)
+					]
+					// Color Count
+					+ SVerticalBox::Slot()
+					[
+						SNew(STextBlock)
+						.Text(this, &SMutableMaterialViewer::GetColorParametersCountAsText)
+					]
+					// Scalar Count
+					+ SVerticalBox::Slot()
+					[
+						SNew(STextBlock)
+						.Text(this, &SMutableMaterialViewer::GetScalarParametersCountAsText)
+					]
+				]
+			]
+		]
+
+		+ SVerticalBox::Slot()
+		.Padding(2,5)
+		[
+			SNew(SVerticalBox)
+
+			// Image Parameters
+			+ SVerticalBox::Slot()
+			.Padding(5)
+			[
+				SAssignNew(ImageParameterListView, SListView<TSharedPtr<FMutableMaterialImageParameterElement>>)
+				.ListItemsSource(&ImageParameters)
+				.SelectionMode(ESelectionMode::None)
+				.OnGenerateRow(this, &SMutableMaterialViewer::OnGenerateImageParameterRow)
+				.HeaderRow
+				(
+					SNew(SHeaderRow)
+		
+					+ SHeaderRow::Column(MutableMaterialParameterColumns::ParameterIndexID)
+						.DefaultLabel(ParameterIndexNameLabel)
+						.FillWidth(0.2f)
+						
+					+ SHeaderRow::Column(MutableMaterialParameterColumns::ParameterLayerID)
+						.DefaultLabel(ParameterLayerNameLabel)
+						.FillWidth(0.2f)
+					
+					+ SHeaderRow::Column(MutableMaterialParameterColumns::ParameterNameID)
+						.DefaultLabel(ParameterNameLabel)
+						.FillWidth(0.2f)
+				
+					+ SHeaderRow::Column(MutableMaterialParameterColumns::ParameterImageAddressID)
+						.DefaultLabel(FText::FromString(TEXT("OP_ADDRESS")))
+						.FillWidth(0.15f)
+
+					+ SHeaderRow::Column(MutableMaterialParameterColumns::ParameterImageID)
+						.DefaultLabel(FText::FromString(TEXT("FImage")))
+						.FillWidth(0.45f)
+				)
+			]
+
+			+ SVerticalBox::Slot()
+			.Padding(5)
+			[
+				SAssignNew(ColorParameterListView, SListView<TSharedPtr<FMutableMaterialColorParameterElement>>)
+				.ListItemsSource(&ColorParameters)
+				.SelectionMode(ESelectionMode::None)
+				.OnGenerateRow(this, &SMutableMaterialViewer::OnGenerateColorParameterRow)
+				.HeaderRow
+				(
+					SNew(SHeaderRow)
+		
+					+ SHeaderRow::Column(MutableMaterialParameterColumns::ParameterIndexID)
+						.DefaultLabel(ParameterIndexNameLabel)
+						.FillWidth(0.2f)
+	
+					+ SHeaderRow::Column(MutableMaterialParameterColumns::ParameterLayerID)
+						.DefaultLabel(ParameterLayerNameLabel)
+						.FillWidth(0.2f)
+						
+					+ SHeaderRow::Column(MutableMaterialParameterColumns::ParameterNameID)
+						.DefaultLabel(ParameterNameLabel)
+						.FillWidth(0.2f)
+				
+					+ SHeaderRow::Column(MutableMaterialParameterColumns::ParameterColorID)
+						.DefaultLabel(FText::FromString(TEXT("Color")))
+						.FillWidth(0.6f)
+				)
+			]
+			
+			// Scalar Parameters
+			+ SVerticalBox::Slot()
+			.Padding(5)
+			[
+				SAssignNew(ScalarParameterListView, SListView<TSharedPtr<FMutableMaterialScalarParameterElement>>)
+				.ListItemsSource(&ScalarParameters)
+				.SelectionMode(ESelectionMode::None)
+				.OnGenerateRow(this, &SMutableMaterialViewer::OnGenerateScalarParameterRow)
+				.HeaderRow
+				(
+					SNew(SHeaderRow)
+						
+					+ SHeaderRow::Column(MutableMaterialParameterColumns::ParameterIndexID)
+						.DefaultLabel(ParameterIndexNameLabel)
+						.FillWidth(0.2f)
+					
+					+ SHeaderRow::Column(MutableMaterialParameterColumns::ParameterLayerID)
+						.DefaultLabel(ParameterLayerNameLabel)
+						.FillWidth(0.2f)
+						
+					+ SHeaderRow::Column(MutableMaterialParameterColumns::ParameterNameID)
+						.DefaultLabel(ParameterNameLabel)
+						.FillWidth(0.2f)
+								
+					+ SHeaderRow::Column(MutableMaterialParameterColumns::ParameterScalarID)
+						.DefaultLabel(FText::FromString(TEXT("Scalar")))
+						.FillWidth(0.6f)
+				)
+			]
+		]
+	];
+	
+	SetMaterial(InArgs._Material);
+}
+
+
+void SMutableMaterialViewer::SetMaterial(const UE::Mutable::Private::TManagedPtr<const UE::Mutable::Private::FMaterial>& InMaterial)
+{
+	if (InMaterial != HostedMaterial)
+	{
+		HostedMaterial = InMaterial;
+
+		if (HostedMaterial)
+		{
+			// Image Parameters
+			ImageParameters.Reset(HostedMaterial->ImageParameters.Num());
+			int32 ImageParameterIndex = 0;
+			for (const TTuple<UE::Mutable::Private::FParameterKey, UE::Mutable::Private::FMaterial::FImageParameterData>& ImageParameterTuple : HostedMaterial->ImageParameters)
+			{
+				TSharedPtr<FMutableMaterialImageParameterElement> ImageParameter = MakeShared<FMutableMaterialImageParameterElement>();
+				ImageParameter->Index = ImageParameterIndex++;
+				ImageParameter->ParameterName = ImageParameterTuple.Key.ParameterName;
+				ImageParameter->LayerIndex = ImageParameterTuple.Key.LayerIndex;
+				ImageParameter->Value = ImageParameterTuple.Value;
+				ImageParameters.Add(ImageParameter);
+			}
+			
+			// Color Parameters
+			int32 ColorParameterIndex = 0;
+			ColorParameters.Reset(HostedMaterial->ColorParameters.Num());
+			for (const TTuple<UE::Mutable::Private::FParameterKey, FVector4f>& ColorParameterTuple : HostedMaterial->ColorParameters)
+			{
+				TSharedPtr<FMutableMaterialColorParameterElement> ColorParameter = MakeShared<FMutableMaterialColorParameterElement>();
+				ColorParameter->Index = ColorParameterIndex++;
+				ColorParameter->ParameterName = ColorParameterTuple.Key.ParameterName;
+				ColorParameter->LayerIndex = ColorParameterTuple.Key.LayerIndex;
+				ColorParameter->Value = ColorParameterTuple.Value;
+				ColorParameters.Add(ColorParameter);
+			}
+			
+			// Scalar Parameters
+			ScalarParameters.Reset(HostedMaterial->ScalarParameters.Num());
+			int32 ScalarParameterIndex = 0;
+			for (const TTuple<UE::Mutable::Private::FParameterKey, float>& ScalarParameterTuple : HostedMaterial->ScalarParameters)
+			{
+				TSharedPtr<FMutableMaterialScalarParameterElement> ScalarParameter = MakeShared<FMutableMaterialScalarParameterElement>();
+				ScalarParameter->Index = ScalarParameterIndex++;
+				ScalarParameter->ParameterName = ScalarParameterTuple.Key.ParameterName;
+				ScalarParameter->LayerIndex = ScalarParameterTuple.Key.LayerIndex;
+				ScalarParameter->Value = ScalarParameterTuple.Value;
+				ScalarParameters.Add(ScalarParameter);
+			}
+		}
+		else
+		{
+			ImageParameters.Reset();
+			ColorParameters.Reset();
+			ScalarParameters.Reset();
+		}
+		
+		// Now that the Material has changed refresh the preview of it
+		check(PreviewContainer)
+		PreviewContainer->SetContent(GeneratePassthroughMaterialWidget());
+		
+		// Request the refreshing of the list views so they show up-to-date data
+		ImageParameterListView->RequestListRefresh();
+		ColorParameterListView->RequestListRefresh();
+		ScalarParameterListView->RequestListRefresh();
+	}
+}
+
+
+TSharedRef<SWidget> SMutableMaterialViewer::GeneratePassthroughMaterialWidget()
+{
+	if (HostedMaterial)
+	{
+		// Try to get the actual object and if null then display the ID as fallback
+		if (HostedMaterial->PassthroughObject.IsResolved())
+		{
+			if (const UMaterialInterface* Material = HostedMaterial->PassthroughObject.Get())
+			{
+				return SNew(SHyperlink)
+					.Style(FAppStyle::Get(), TEXT("NavigationHyperlink"))
+					.Text(FText::FromString(Material->GetName()))
+					.OnNavigate(this, &SMutableMaterialViewer::OnMaterialNavigation);
+			}
+		}
+		else
+		{
+			FText MaterialIDText;
+				
+			UE::Mutable::Private::PASSTHROUGH_ID PassthroughMaterialID = HostedMaterial->PassthroughObject.GetId();
+			if (PassthroughMaterialID != UE::Mutable::Private::PASSTHROUGH_ID_INVALID)
+			{
+				MaterialIDText = FText::AsNumber(PassthroughMaterialID);
+			}
+			else
+			{
+				MaterialIDText = FText::FromString(TEXT("INVALID_ID"));
+			}
+				
+			return SNew(STextBlock).Text(MaterialIDText);
+		}
+	}
+
+	return SNew(STextBlock).Text(FText::FromString("No FMaterial has been provided!"));
+}
+
+
+void SMutableMaterialViewer::OnMaterialNavigation() const
+{
+	if (HostedMaterial)
+	{
+		if (HostedMaterial && HostedMaterial->PassthroughObject.IsResolved())
+		{
+			const UMaterialInterface* Material = HostedMaterial->PassthroughObject.Get();
+			check(Material);
+			GEditor->GetEditorSubsystem<UAssetEditorSubsystem>()->OpenEditorForAsset(Material);
+		}
+	}
+}
+
+
+FText SMutableMaterialViewer::GetImageParametersCountAsText() const
+{
+	const FString ImageParametersCountTitleLabel {TEXT("Image Parameter Count : ")};
+	const FString ImageParametersAsNumber = FString::FormatAsNumber(ImageParameters.Num());
+	return FText::FromString(ImageParametersCountTitleLabel + ImageParametersAsNumber);
+}
+
+
+FText SMutableMaterialViewer::GetColorParametersCountAsText() const
+{
+	const FString ColorParametersCountTitleLabel {TEXT("Color Parameter Count : ")};
+	const FString ColorParametersAsNumber = FString::FormatAsNumber(ColorParameters.Num());
+	return FText::FromString(ColorParametersCountTitleLabel + ColorParametersAsNumber);
+}
+
+
+FText SMutableMaterialViewer::GetScalarParametersCountAsText() const
+{
+	const FString ScalarParametersCountTitleLabel {TEXT("Scalar Parameter Count : ")};
+	const FString ScalarParametersAsNumber = FString::FormatAsNumber(ScalarParameters.Num());
+	return FText::FromString(ScalarParametersCountTitleLabel + ScalarParametersAsNumber);
+}
+
+
+TSharedRef<ITableRow> SMutableMaterialViewer::OnGenerateImageParameterRow(TSharedPtr<FMutableMaterialImageParameterElement> MutableMaterialImageParameterElement,
+	const TSharedRef<STableViewBase>& Shared) const
+{
+	TSharedRef<SMutableMaterialImageParameterRow> Row = SNew(SMutableMaterialImageParameterRow, Shared, MutableMaterialImageParameterElement, MutableCodeViewer);
+	return Row;
+}
+
+
+TSharedRef<ITableRow> SMutableMaterialViewer::OnGenerateColorParameterRow(TSharedPtr<FMutableMaterialColorParameterElement> MutableMaterialColorParameterElement,
+	const TSharedRef<STableViewBase>& Shared) const
+{
+	TSharedRef<SMutableMaterialColorParameterRow> Row = SNew(SMutableMaterialColorParameterRow, Shared, MutableMaterialColorParameterElement);
+	return Row;
+}
+
+
+TSharedRef<ITableRow> SMutableMaterialViewer::OnGenerateScalarParameterRow(TSharedPtr<FMutableMaterialScalarParameterElement> MutableMaterialScalarParameterElement,
+	const TSharedRef<STableViewBase>& Shared) const
+{
+	TSharedRef<SMutableMaterialScalarParameterRow> Row = SNew(SMutableMaterialScalarParameterRow, Shared, MutableMaterialScalarParameterElement);
+	return Row;
+}
+
+#undef LOCTEXT_NAMESPACE
diff --git a/Engine/Plugins/Mutable/Source/CustomizableObjectEditor/Private/MuCOE/SMutableMaterialViewer.h b/Engine/Plugins/Mutable/Source/CustomizableObjectEditor/Private/MuCOE/SMutableMaterialViewer.h
new file mode 100644
--- /dev/null
+++ b/Engine/Plugins/Mutable/Source/CustomizableObjectEditor/Private/MuCOE/SMutableMaterialViewer.h
@@ -0,0 +1,90 @@
+// Copyright Epic Games, Inc. All Rights Reserved.
+
+#pragma once
+
+#include "SMutableCodeViewer.h"
+#include "Widgets/SCompoundWidget.h"
+#include "MuR/Material.h"
+#include "Widgets/Views/STableRow.h"
+#include "Widgets/Views/STableViewBase.h"
+
+// Base for all parameter types
+struct FMutableMaterialParameterListElement
+{
+	FName ParameterName;
+	int32 Index = -1;
+	int32 LayerIndex = -1;
+};
+
+
+struct FMutableMaterialImageParameterElement : public FMutableMaterialParameterListElement
+{
+	UE::Mutable::Private::FMaterial::FImageParameterData Value;		// !
+};
+
+
+struct FMutableMaterialColorParameterElement : public FMutableMaterialParameterListElement
+{
+	FVector4f Value;
+};
+
+
+struct FMutableMaterialScalarParameterElement : public FMutableMaterialParameterListElement
+{
+	float Value;
+};
+
+
+
+
+/** Widget designed to show the statistical data from a Mutable FMaterial*/
+class SMutableMaterialViewer : public SCompoundWidget
+{
+public:
+	SLATE_BEGIN_ARGS(SMutableMaterialViewer) {}
+		SLATE_ARGUMENT_DEFAULT(UE::Mutable::Private::TManagedPtr<const UE::Mutable::Private::FMaterial>, Material){nullptr};
+		SLATE_ARGUMENT_DEFAULT(TWeakPtr<SMutableCodeViewer>, HostCodeViewer) {nullptr};
+	SLATE_END_ARGS()
+
+	/** Builds the widget */
+	void Construct(const FArguments& InArgs);
+
+	/** Set the Mutable Material to be used for this widget */
+	void SetMaterial(const UE::Mutable::Private::TManagedPtr<const UE::Mutable::Private::FMaterial>& InMaterial);
+
+private:
+	
+	TSharedRef<SWidget> GeneratePassthroughMaterialWidget();
+	
+	// callbacks
+
+	/**
+	 * Callback method called when the material link is pressed. Will open the editor for the provided material
+	 */
+	void OnMaterialNavigation() const;
+	
+	FText GetImageParametersCountAsText() const;
+	FText GetColorParametersCountAsText() const;
+	FText GetScalarParametersCountAsText() const;
+	
+	TSharedRef<ITableRow> OnGenerateImageParameterRow(TSharedPtr<FMutableMaterialImageParameterElement> MutableMaterialImageParameterElement, const TSharedRef<STableViewBase>& Shared) const;
+	TSharedRef<ITableRow> OnGenerateColorParameterRow(TSharedPtr<FMutableMaterialColorParameterElement> MutableMaterialColorParameterElement, const TSharedRef<STableViewBase>& Shared) const;
+	TSharedRef<ITableRow> OnGenerateScalarParameterRow(TSharedPtr<FMutableMaterialScalarParameterElement> MutableMaterialScalarParameterElement, const TSharedRef<STableViewBase>& Shared) const;
+	
+	// The material we are showing some data of
+	UE::Mutable::Private::TManagedPtr<const UE::Mutable::Private::FMaterial> HostedMaterial;
+	
+	TWeakPtr<SMutableCodeViewer> MutableCodeViewer;
+	
+	TSharedPtr<SBox> PreviewContainer;
+	
+	// Backend of the parameter list views. Extracted from the provided Hosted Material
+	TArray<TSharedPtr<FMutableMaterialImageParameterElement>> ImageParameters;
+	TArray<TSharedPtr<FMutableMaterialColorParameterElement>> ColorParameters;
+	TArray<TSharedPtr<FMutableMaterialScalarParameterElement>> ScalarParameters;
+	
+	// Lists
+	TSharedPtr<SListView<TSharedPtr<FMutableMaterialImageParameterElement>>> ImageParameterListView;
+	TSharedPtr<SListView<TSharedPtr<FMutableMaterialColorParameterElement>>> ColorParameterListView;
+	TSharedPtr<SListView<TSharedPtr<FMutableMaterialScalarParameterElement>>> ScalarParameterListView;
+};
diff --git a/Engine/Plugins/Mutable/Source/MutableRuntime/Internal/MuR/Image.h b/Engine/Plugins/Mutable/Source/MutableRuntime/Internal/MuR/Image.h
--- a/Engine/Plugins/Mutable/Source/MutableRuntime/Internal/MuR/Image.h
+++ b/Engine/Plugins/Mutable/Source/MutableRuntime/Internal/MuR/Image.h
@@ -64,7 +64,7 @@ namespace UE::Mutable::Private
 		/** */
 		UE_API void Init(uint32 SizeX, uint32 SizeY, uint32 Lods, EImageFormat Format, EInitializationType InitType);
 
-		/** Clear the image to black colour. */
+		/** Clear the image to black color. */
 		UE_API void InitToBlack();
 
 		/** Return true if this is a reference to an engine image. */
@@ -105,13 +105,13 @@ namespace UE::Mutable::Private
 			// Set if the next flag has been calculated and its value is valid.
 			IF_IS_PLAIN_COLOUR_VALID = 1 << 0,
 
-			// If the previous flag is set and this one too, the image is single colour.
+			// If the previous flag is set and this one too, the image is single color.
 			IF_IS_PLAIN_COLOUR = 1 << 1,
 
-			// If this is set, the image shouldn't be scaled: it's contents is resoultion-dependent.
+			// If this is set, the image shouldn't be scaled: it's contents is resolution-dependent.
 			IF_CANNOT_BE_SCALED = 1 << 2,
 
-			// If this is set, the image has an updated relevancy map. This flag is not persisent.
+			// If this is set, the image has an updated relevancy map. This flag is not persistent.
 			IF_HAS_RELEVANCY_MAP = 1 << 3,
 			
 			/** For reference images, this indicates that they should be loaded into full images as soon as they are generated. */
