[mutable] Added Reference Image option for the Material Modify Node image pin options and updated the way the user sets the target param name for the Material Break and Modify nodes.
- Made the GenerateMutableSourceMaterial method for the case of the MaterialModify node mind the Reference texture set in the PinData of the pin being processed. It will use that texture, and, if not found, it will try to get it from the node connected to the pin.
- Fixed bug that prevented the "Pin Viewer" slate from updating when the Pin "Editable Name" got changed for the PassthroughSkeletalMeshComponent node.
- Made the Material Modify Node expose the param names as editable boxes from within the standard Pin UI to improve UX.
- Made the Material Break Node expose the param names as editable boxes from within the standard Pin UI to improve UX.
- Made the GenerateMutableSource methods not crash when a parameter of content "None" is found. It should be object of warning but not of a check failure.
- Removed the option to edit the name of a pin from the PinViewer as that option now lies in the node itself.
- Fixed issue with the LayoutIndex of the Material Modify and Break Parameters not updating the UI of the node.
- Fixed check failure when the Material Parameter Name is empty.
- Fixed some typos found along the way.

#rnx
#Synced-CL 48762940

diff --git a/Engine/Plugins/Mutable/Config/DefaultMutable.ini b/Engine/Plugins/Mutable/Config/DefaultMutable.ini
--- a/Engine/Plugins/Mutable/Config/DefaultMutable.ini
+++ b/Engine/Plugins/Mutable/Config/DefaultMutable.ini
@@ -29,6 +29,10 @@
 +PropertyRedirects=(OldName="/Script/CustomizableObject.CompileParams.bGatherReferneces",NewName="/Script/CustomizableObject.CompileParams.bGatherReferences")
 +PropertyRedirects=(OldName="/Script/CustomizableObject.CompileParams.bSkipIfOutOfDate",NewName="/Script/CustomizableObject.CompileParams.bSkipIfNotOutOfDate")
 
+; Pin Data renames
++ClassRedirects=(OldName="/Script/CustomizableObjectEditor.CONodeMaterialBreakPinData",NewName="/Script/CustomizableObjectEditor.CONodeMaterialBreakParameterPinData")
++ClassRedirects=(OldName="/Script/CustomizableObjectEditor.CONodeMaterialModifyPinData",NewName="/Script/CustomizableObjectEditor.CONodeMaterialModifyParameterPinData")
+
 ; Settings for the Mutable plugin
 [Features]
 
diff --git a/Engine/Plugins/Mutable/Source/CustomizableObjectEditor/Private/MuCOE/COVariable.cpp b/Engine/Plugins/Mutable/Source/CustomizableObjectEditor/Private/MuCOE/COVariable.cpp
--- a/Engine/Plugins/Mutable/Source/CustomizableObjectEditor/Private/MuCOE/COVariable.cpp
+++ b/Engine/Plugins/Mutable/Source/CustomizableObjectEditor/Private/MuCOE/COVariable.cpp
@@ -7,7 +7,6 @@
 #include "MuCOE/Nodes/CustomizableObjectNode.h"
 #include "PropertyCustomizationHelpers.h"
 #include "Widgets/Input/SComboBox.h"
-#include "Widgets/Input/SEditableTextBox.h"
 #include "Widgets/SWidget.h"
 
 
@@ -20,9 +19,8 @@ TSharedRef<IPropertyTypeCustomization> FCOVariableCustomization::MakeInstance()
 void FCOVariableCustomization::CustomizeHeader(TSharedRef<IPropertyHandle> StructPropertyHandle, FDetailWidgetRow& InHeaderRow, IPropertyTypeCustomizationUtils& InCustomizationUtils)
 {
 	// Get the property handles to edit the variable values
-	NamePropertyHandle = StructPropertyHandle->GetChildHandle("Name");
 	TypePropertyHandle = StructPropertyHandle->GetChildHandle("Type");
-	check(NamePropertyHandle && TypePropertyHandle);
+	check(TypePropertyHandle);
 
 	//Get parent node if it's defined inside one
 	TArray<UObject*> OuterObjects;
@@ -50,10 +48,8 @@ void FCOVariableCustomization::CustomizeHeader(TSharedRef<IPropertyHandle> Struc
 
 	InHeaderRow.NameContent()
 	[
-		SNew(SEditableTextBox)
-		.Text(this, &FCOVariableCustomization::GetVariableName)
-		.OnTextChanged(this, &FCOVariableCustomization::OnVariableNameChanged)
-		.OnTextCommitted(this, &FCOVariableCustomization::OnVariableNameCommitted)
+		SNew(STextBlock)
+		.Text(TypePropertyHandle->GetPropertyDisplayName())
 		.Font(IDetailLayoutBuilder::GetDetailFont())
 	]
 	.ValueContent()
@@ -99,53 +95,13 @@ bool FCOVariableCustomization::ShouldVariableTypePropertyBeEnabled() const
 	return false;
 }
 
-FText FCOVariableCustomization::GetVariableName() const
-{
-	FName VariableName;
-	NamePropertyHandle->GetValue(VariableName);
-
-	return FText::FromName(VariableName);
-}
-
 
 void FCOVariableCustomization::ResetToDefault()
 {
-	NamePropertyHandle->SetValue(FName("NewVar"));
 	TypePropertyHandle->SetValue(AllowedVariableTypes.Num() ? AllowedVariableTypes[0] : NAME_None);
 }
 
 
-void FCOVariableCustomization::OnVariableNameChanged(const FText& InNewText)
-{
-	FName NewName = FName(InNewText.ToString());
-
-	if (InNewText.IsEmpty() || NewName.IsNone())
-	{
-		return;
-	}
-
-	NamePropertyHandle->SetValue(NewName);
-}
-
-
-void FCOVariableCustomization::OnVariableNameCommitted(const FText& InNewText, ETextCommit::Type InTextCommit)
-{
-	FName NewName = FName(InNewText.ToString());
-
-	if (InNewText.IsEmpty() || NewName.IsNone())
-	{
-		return;
-	}
-
-	NamePropertyHandle->SetValue(NewName);
-
-	if (BaseObjectNode.IsValid())
-	{
-		BaseObjectNode->ReconstructNode();
-	}
-}
-
-
 TSharedRef<SWidget> FCOVariableCustomization::OnGenerateVariableTypeRow(FName Type)
 {
 	const FSlateBrush* IconBrush = FAppStyle::GetBrush(TEXT("Kismet.VariableList.TypeIcon"));
diff --git a/Engine/Plugins/Mutable/Source/CustomizableObjectEditor/Private/MuCOE/COVariable.h b/Engine/Plugins/Mutable/Source/CustomizableObjectEditor/Private/MuCOE/COVariable.h
--- a/Engine/Plugins/Mutable/Source/CustomizableObjectEditor/Private/MuCOE/COVariable.h
+++ b/Engine/Plugins/Mutable/Source/CustomizableObjectEditor/Private/MuCOE/COVariable.h
@@ -46,15 +46,8 @@ public:
 
 private:
 
-	/** Returns the name of the specified variable. Needed to update automatically the text if the name of the variable changes. */
-	FText GetVariableName() const;
-
 	void ResetToDefault();
 
-	/** Editable Text Box Callbacks */
-	void OnVariableNameChanged(const FText& InNewText);
-	void OnVariableNameCommitted(const FText& InNewText, ETextCommit::Type InTextCommit);
-
 	/** ComboBox Callbacks */
 	TSharedRef<SWidget> OnGenerateVariableTypeRow(FName Option);
 	bool ShouldVariableTypePropertyBeEnabled() const;
@@ -70,9 +63,6 @@ private:
 	
 	TArray<FName> AllowedVariableTypes;
 
-	/** Property handle of the variable name */
-	TSharedPtr<IPropertyHandle> NamePropertyHandle;
-
 	/** Property handle of the variable type */
 	TSharedPtr<IPropertyHandle> TypePropertyHandle;
 
diff --git a/Engine/Plugins/Mutable/Source/CustomizableObjectEditor/Private/MuCOE/CustomizableObjectEditorModule.cpp b/Engine/Plugins/Mutable/Source/CustomizableObjectEditor/Private/MuCOE/CustomizableObjectEditorModule.cpp
--- a/Engine/Plugins/Mutable/Source/CustomizableObjectEditor/Private/MuCOE/CustomizableObjectEditorModule.cpp
+++ b/Engine/Plugins/Mutable/Source/CustomizableObjectEditor/Private/MuCOE/CustomizableObjectEditorModule.cpp
@@ -1061,6 +1061,17 @@ TMap<FName, FGuid> FCustomizableObjectEditorModule::GetParticipatingObjects(cons
 				AddObject(MatNode->GetImageReferenceTexture(ImageIndex));
 			}
 		}
+			
+		else if (UCONodeMaterialModify* MaterialModifyNode = Cast<UCONodeMaterialModify>(&Node))
+		{
+			TArray<FEdGraphPinReference> MaterialTextureParameterPins;
+			MaterialModifyNode->GetMaterialTextureParameterPins(MaterialTextureParameterPins);
+			
+			for (FEdGraphPinReference MaterialTextureParameterPin : MaterialTextureParameterPins)
+			{
+				AddObjectSoft(MaterialModifyNode->GetTexturePinParameterReferenceTexture(*MaterialTextureParameterPin.Get()).ToSoftObjectPath());
+			}
+		}
 
 		else if (UCustomizableObjectNodeTable* TableNode = Cast<UCustomizableObjectNodeTable>(&Node))
 		{
diff --git a/Engine/Plugins/Mutable/Source/CustomizableObjectEditor/Private/MuCOE/GenerateMutableSource/GenerateMutableSourceColor.cpp b/Engine/Plugins/Mutable/Source/CustomizableObjectEditor/Private/MuCOE/GenerateMutableSource/GenerateMutableSourceColor.cpp
--- a/Engine/Plugins/Mutable/Source/CustomizableObjectEditor/Private/MuCOE/GenerateMutableSource/GenerateMutableSourceColor.cpp
+++ b/Engine/Plugins/Mutable/Source/CustomizableObjectEditor/Private/MuCOE/GenerateMutableSource/GenerateMutableSourceColor.cpp
@@ -397,8 +397,13 @@ UE::Mutable::Private::Ptr<UE::Mutable::Private::NodeColour> GenerateMutableSourc
 			UE::Mutable::Private::Ptr<UE::Mutable::Private::NodeColourMaterialBreak> MaterialBreak = new UE::Mutable::Private::NodeColourMaterialBreak();
 
 			FName ParameterName = TypedNodeMaterialBreak->GetPinParameterName(*Pin);
+			if (ParameterName.IsNone())
+			{
+				GenerationContext.Log(LOCTEXT("NoneParameterName", "'None' is not a valid Material Parameter name."), Node, EMessageSeverity::Warning);
+				return MaterialBreak;
+			}
+			
 			int32 LayerIndex = TypedNodeMaterialBreak->GetPinParameterLayerIndex(*Pin);
-			check(!ParameterName.IsNone());
 
 			// Set the name of the parameter that represents the pin
 			MaterialBreak->ParameterName = ParameterName;
diff --git a/Engine/Plugins/Mutable/Source/CustomizableObjectEditor/Private/MuCOE/GenerateMutableSource/GenerateMutableSourceFloat.cpp b/Engine/Plugins/Mutable/Source/CustomizableObjectEditor/Private/MuCOE/GenerateMutableSource/GenerateMutableSourceFloat.cpp
--- a/Engine/Plugins/Mutable/Source/CustomizableObjectEditor/Private/MuCOE/GenerateMutableSource/GenerateMutableSourceFloat.cpp
+++ b/Engine/Plugins/Mutable/Source/CustomizableObjectEditor/Private/MuCOE/GenerateMutableSource/GenerateMutableSourceFloat.cpp
@@ -422,8 +422,13 @@ UE::Mutable::Private::Ptr<UE::Mutable::Private::NodeScalar> GenerateMutableSourc
 			UE::Mutable::Private::Ptr<UE::Mutable::Private::NodeScalarMaterialBreak> MaterialBreak = new UE::Mutable::Private::NodeScalarMaterialBreak();
 
 			FName ParameterName = TypedNodeMaterialBreak->GetPinParameterName(*Pin);
+			if (ParameterName.IsNone())
+			{
+				GenerationContext.Log(LOCTEXT("NoneParameterName", "'None' is not a valid Material Parameter name."), Node, EMessageSeverity::Warning);
+				return MaterialBreak;
+			}
+			
 			int32 LayerIndex = TypedNodeMaterialBreak->GetPinParameterLayerIndex(*Pin);
-			check(!ParameterName.IsNone());
 
 			// Set the name of the parameter that represents the pin
 			MaterialBreak->ParameterName = ParameterName;
diff --git a/Engine/Plugins/Mutable/Source/CustomizableObjectEditor/Private/MuCOE/GenerateMutableSource/GenerateMutableSourceImage.cpp b/Engine/Plugins/Mutable/Source/CustomizableObjectEditor/Private/MuCOE/GenerateMutableSource/GenerateMutableSourceImage.cpp
--- a/Engine/Plugins/Mutable/Source/CustomizableObjectEditor/Private/MuCOE/GenerateMutableSource/GenerateMutableSourceImage.cpp
+++ b/Engine/Plugins/Mutable/Source/CustomizableObjectEditor/Private/MuCOE/GenerateMutableSource/GenerateMutableSourceImage.cpp
@@ -1231,8 +1231,13 @@ UE::Mutable::Private::Ptr<UE::Mutable::Private::NodeImage> GenerateMutableSource
 			UE::Mutable::Private::Ptr<UE::Mutable::Private::NodeImageMaterialBreak> MaterialBreak = new UE::Mutable::Private::NodeImageMaterialBreak();
 
 			FName ParameterName = TypedNodeMaterialBreak->GetPinParameterName(*Pin);
+			if (ParameterName.IsNone())
+			{
+				GenerationContext.Log(LOCTEXT("NoneParameterName", "'None' is not a valid Material Parameter name."), Node, EMessageSeverity::Warning);
+				return MaterialBreak;
+			}
+			
 			int32 LayerIndex = TypedNodeMaterialBreak->GetPinParameterLayerIndex(*Pin);
-			check(!ParameterName.IsNone());
 
 			// Set the name of the parameter that represents the pin
 			MaterialBreak->ParameterName = ParameterName;
diff --git a/Engine/Plugins/Mutable/Source/CustomizableObjectEditor/Private/MuCOE/GenerateMutableSource/GenerateMutableSourceMaterial.cpp b/Engine/Plugins/Mutable/Source/CustomizableObjectEditor/Private/MuCOE/GenerateMutableSource/GenerateMutableSourceMaterial.cpp
--- a/Engine/Plugins/Mutable/Source/CustomizableObjectEditor/Private/MuCOE/GenerateMutableSource/GenerateMutableSourceMaterial.cpp
+++ b/Engine/Plugins/Mutable/Source/CustomizableObjectEditor/Private/MuCOE/GenerateMutableSource/GenerateMutableSourceMaterial.cpp
@@ -168,7 +168,7 @@ UE::Mutable::Private::Ptr<UE::Mutable::Private::NodeMaterial> GenerateMutableSou
 		}
 
 		// Generate images
-		int32 NumImages = TypedNodeMaterialBase->GetNumParameters(EMaterialParameterType::Texture);
+		const int32 NumImages = TypedNodeMaterialBase->GetNumParameters(EMaterialParameterType::Texture);
 
 		for (int32 ImageIndex = 0; ImageIndex < NumImages; ++ImageIndex)
 		{
@@ -199,19 +199,17 @@ UE::Mutable::Private::Ptr<UE::Mutable::Private::NodeMaterial> GenerateMutableSou
 			}
 			else
 			{
-				UE::Mutable::Private::NodeImagePtr GroupProjectionImg;
-				UTexture2D* GroupProjectionReferenceTexture = nullptr;
-				const FNodeMaterialParameterId ImageId = TypedNodeMaterialBase->GetParameterId(EMaterialParameterType::Texture, ImageIndex);
-
-				FString MaterialImageId = FGroupProjectorImageInfo::GenerateId(TypedNodeMaterialBase, ImageIndex);
-				bool bShareProjectionTexturesBetweenLODs = false;
-				FGroupProjectorImageInfo* ProjectorInfo = GenerationContext.GroupProjectorLODCache.Find(MaterialImageId);
+				const FString MaterialImageId = FGroupProjectorImageInfo::GenerateId(TypedNodeMaterialBase, ImageIndex);
+				const FGroupProjectorImageInfo* ProjectorInfo = GenerationContext.GroupProjectorLODCache.Find(MaterialImageId);
 
 				if (!ProjectorInfo) // No previous LOD of this material generated the image.
 				{
 					bool bIsGroupProjectorImage = false;
+					bool bShareProjectionTexturesBetweenLODs = false;
+					
+					UTexture2D* GroupProjectionReferenceTexture = nullptr;
 
-					GroupProjectionImg = GenerateMutableSourceGroupProjector(GenerationContext.CurrentLOD, ImageIndex, MaterialOptions.SurfaceNode->Mesh, GenerationContext,
+					UE::Mutable::Private::NodeImagePtr GroupProjectionImg = GenerateMutableSourceGroupProjector(GenerationContext.CurrentLOD, ImageIndex, MaterialOptions.SurfaceNode->Mesh, GenerationContext,
 						TypedNodeMaterialBase, nullptr, bShareProjectionTexturesBetweenLODs, bIsGroupProjectorImage,
 						GroupProjectionReferenceTexture);
 
@@ -220,6 +218,8 @@ UE::Mutable::Private::Ptr<UE::Mutable::Private::NodeMaterial> GenerateMutableSou
 						// Get the reference texture
 						UTexture2D* ReferenceTexture = nullptr;
 						{
+							const FNodeMaterialParameterId ImageId = TypedNodeMaterialBase->GetParameterId(EMaterialParameterType::Texture, ImageIndex);
+							
 							//TODO(Max) UE-220247: Add support for multilayer materials
 							GenerationContext.CurrentMaterialTableParameter = ImageName.ToString();
 							GenerationContext.CurrentMaterialParameterId = ImageId.ParameterId.ToString();
@@ -749,11 +749,11 @@ UE::Mutable::Private::Ptr<UE::Mutable::Private::NodeMaterial> GenerateMutableSou
 		// - Add support for table node materials. This will be added when the current hack of the tables is removed. UE-314401
 		// - Add support for layout index once the pin viewer has all the needed properties. UE-317590
 		// - Add support for group projectors.
-		// - Remove the material editon part from the mesh section nodes. Talk about this with Gerard since we can lose the capacity to pass an image parameter of the reference material as mutable without linking any texture node.
-		// - Improve the dectection of the reference texture. 
-		//		|_ Find the first material constant (or parameter) node that contains the texture parameter and return it
-		//		|_ Or serach it in the pin data UE-318756
-		//		|_ Or use the reference image of the group projector generated (if any).
+		// - Remove the material edition part from the mesh section nodes. Talk about this with Gerard since we can lose the capacity to pass an image parameter of the reference material as mutable without linking any texture node.
+		// - Improve the detection of the reference texture. 
+		//		|_ Search it in the pin data
+		//		|_ Or instead find the first material constant (or parameter) node that contains the texture parameter and return it
+		//		|_ Or use the reference image of the group projector generated (if any). 
 
 		TArray<UEdGraphPin*> InputPins;
 		TypedNodeMaterialModify->GetInputPins(InputPins);
@@ -768,19 +768,30 @@ UE::Mutable::Private::Ptr<UE::Mutable::Private::NodeMaterial> GenerateMutableSou
 				continue;
 			}
 
-			FName PinType = InputPin->PinType.PinCategory;
-			FName ParameterName = TypedNodeMaterialModify->GetPinParameterName(*InputPin);
-			int32 LayerIndex = TypedNodeMaterialModify->GetPinParameterLayerIndex(*InputPin);
-
+			const FName ParameterName = TypedNodeMaterialModify->GetPinParameterName(*InputPin);
+			if (ParameterName.IsNone())
+			{
+				GenerationContext.Log(LOCTEXT("NoneParameterName", "'None' is not a valid Material Parameter name."), Node, EMessageSeverity::Warning);
+				continue;
+			}
+			
+			const int32 LayerIndex = TypedNodeMaterialModify->GetPinParameterLayerIndex(*InputPin);
 			UE::Mutable::Private::FParameterKey ParameterKey = { ParameterName, LayerIndex };
 
+			const FName PinType = InputPin->PinType.PinCategory;
 			if (PinType == UEdGraphSchema_CustomizableObject::PC_Texture)
 			{
 				UE::Mutable::Private::FImageParameterData ImageData;
 
 				if (const UEdGraphPin* FollowPin = FollowInputPin(*InputPin))
 				{
-					UTexture2D* ReferenceTexture = FindReferenceImage(FollowPin, GenerationContext); //TODO(Max - GM) UE - 318756
+					TObjectPtr<UTexture2D> ReferenceTexture = GenerationContext.LoadObject(TypedNodeMaterialModify->GetTexturePinParameterReferenceTexture(*InputPin));
+					if (!ReferenceTexture)
+					{
+						// Fallback if the user did not provide a reference texture
+						ReferenceTexture = FindReferenceImage(FollowPin, GenerationContext);
+					}
+					
 					FGeneratedImageProperties Props = GenerateImageProperties(Node, ReferenceTexture, ParameterKey, false, GenerationContext);
 					UE::Mutable::Private::Ptr<UE::Mutable::Private::NodeImage> ImageNode = GenerateMutableSourceImage(FollowPin, GenerationContext, Props.TextureSize);
 
@@ -790,7 +801,7 @@ UE::Mutable::Private::Ptr<UE::Mutable::Private::NodeMaterial> GenerateMutableSou
 					// Fill Image Data
 					ImageData.ImageNode = ImageNode;
 					ImageData.ImagePropertyIndex = Props.ImagePropertiesIndex;
-					ImageData.LayoutIndex = -1; // TODO(Max - GM) UE-317590
+					ImageData.LayoutIndex = TypedNodeMaterialModify->GetImagePinParameterUVIndex(*InputPin);
 				}
 
 				// We do not check Null images since it is possible to not have anything linked
diff --git a/Engine/Plugins/Mutable/Source/CustomizableObjectEditor/Private/MuCOE/Nodes/CONodeMaterialBreak.cpp b/Engine/Plugins/Mutable/Source/CustomizableObjectEditor/Private/MuCOE/Nodes/CONodeMaterialBreak.cpp
--- a/Engine/Plugins/Mutable/Source/CustomizableObjectEditor/Private/MuCOE/Nodes/CONodeMaterialBreak.cpp
+++ b/Engine/Plugins/Mutable/Source/CustomizableObjectEditor/Private/MuCOE/Nodes/CONodeMaterialBreak.cpp
@@ -16,11 +16,12 @@ struct FSlateBrush;
 
 #define LOCTEXT_NAMESPACE "CustomizableObjectEditor"
 
+const FName BreakMaterialMaterialParameterPinName = "Parameter";
 
 bool UCustomizableObjectNodeBreakMaterialRemapPins::Equal(const UCustomizableObjectNode& Node, const UEdGraphPin& OldPin, const UEdGraphPin& NewPin) const
 {
-	UCONodeMaterialBreakPinData* OldPinData = Cast<UCONodeMaterialBreakPinData>(Node.GetPinData(OldPin));
-	UCONodeMaterialBreakPinData* NewPinData = Cast<UCONodeMaterialBreakPinData>(Node.GetPinData(NewPin));
+	UCONodeMaterialBreakParameterPinData* OldPinData = Cast<UCONodeMaterialBreakParameterPinData>(Node.GetPinData(OldPin));
+	UCONodeMaterialBreakParameterPinData* NewPinData = Cast<UCONodeMaterialBreakParameterPinData>(Node.GetPinData(NewPin));
 
 	if (OldPinData && NewPinData)
 	{
@@ -55,24 +56,38 @@ void UCONodeMaterialBreak::AllocateDefaultPins(UCustomizableObjectNodeRemapPins*
 			continue;
 		}
 		
-		if (UCONodeMaterialBreakPinData* PinData = Cast<UCONodeMaterialBreakPinData>(GetPinData(*Pin)))
+		if (UCONodeMaterialBreakParameterPinData* PinData = Cast<UCONodeMaterialBreakParameterPinData>(GetPinData(*Pin)))
 		{
 			// Use the same pin data as the old pin since the new one is a copy
-			UEdGraphPin* NewPin = CustomCreatePin(EGPD_Output, PinData->VariableData.Type, PinData->VariableData.Name, PinData);
+			UEdGraphPin* NewPin = CustomCreatePin(EGPD_Output, PinData->VariableData.Type, BreakMaterialMaterialParameterPinName, PinData);
 
 			if (PinData->LayerIndex == 0)
 			{
-				NewPin->PinFriendlyName = FText::Format(LOCTEXT("BackgroundLayerText", "{0} - Background"), FText::FromName(PinData->VariableData.Name));
+				NewPin->PinFriendlyName = FText::Format(LOCTEXT("BackgroundLayerText", "{0} - Background"), FText::FromName(BreakMaterialMaterialParameterPinName));
 			}
 			else if (PinData->LayerIndex > 0)
 			{
-				NewPin->PinFriendlyName = FText::Format(LOCTEXT("LayerText", "{0} - Layer {1}"), FText::FromName(PinData->VariableData.Name), FText::FromString(FString::FromInt(PinData->LayerIndex)));
+				NewPin->PinFriendlyName = FText::Format(LOCTEXT("LayerText", "{0} - Layer {1}"), FText::FromName(BreakMaterialMaterialParameterPinName), FText::FromString(FString::FromInt(PinData->LayerIndex)));
 			}
 		}
 	}
 }
 
 
+void UCONodeMaterialBreakParameterPinData::PostEditChangeProperty(FPropertyChangedEvent& PropertyChangedEvent)
+{
+	Super::PostEditChangeProperty(PropertyChangedEvent);
+
+	if (const FProperty* PropertyThatChanged = PropertyChangedEvent.Property)
+	{
+		if (PropertyThatChanged->GetFName() == GET_MEMBER_NAME_CHECKED(UCONodeMaterialBreakParameterPinData, LayerIndex))
+		{
+			CastChecked<UCustomizableObjectNode>(GetOuter())->ReconstructNode();
+		}
+	}
+}
+
+
 FText UCONodeMaterialBreak::GetNodeTitle(ENodeTitleType::Type TitleType) const
 {
 	return LOCTEXT("Material_Break", "Break Material");
@@ -93,7 +108,7 @@ FText UCONodeMaterialBreak::GetTooltipText() const
 
 TSharedPtr<IDetailsView> UCONodeMaterialBreak::CustomizePinDetails(const UEdGraphPin& Pin) const
 {
-	if (UCONodeMaterialBreakPinData* PinData = Cast<UCONodeMaterialBreakPinData>(GetPinData(Pin)))
+	if (UCONodeMaterialBreakParameterPinData* PinData = Cast<UCONodeMaterialBreakParameterPinData>(GetPinData(Pin)))
 	{
 		FPropertyEditorModule& EditModule = FModuleManager::Get().GetModuleChecked<FPropertyEditorModule>("PropertyEditor");
 
@@ -145,7 +160,12 @@ UCustomizableObjectNodeRemapPins* UCONodeMaterialBreak::CreateRemapPinsDefault()
 
 TArray<FName> UCONodeMaterialBreak::GetAllowedPinViewerCreationTypes() const
 {
-	return { UEdGraphSchema_CustomizableObject::PC_Texture, UEdGraphSchema_CustomizableObject::PC_Color, UEdGraphSchema_CustomizableObject::PC_Float };
+	return { 
+		UEdGraphSchema_CustomizableObject::PC_Texture, 
+		UEdGraphSchema_CustomizableObject::PC_PassthroughTexture, 
+		UEdGraphSchema_CustomizableObject::PC_Color, 
+		UEdGraphSchema_CustomizableObject::PC_Float 
+	};
 }
 
 
@@ -158,27 +178,65 @@ void UCONodeMaterialBreak::AddPinFromUI()
 	
 	const FName NewPinCategory = UEdGraphSchema_CustomizableObject::PC_Float;
 	
-	UCONodeMaterialBreakPinData* PinData = NewObject<UCONodeMaterialBreakPinData>(this);
-	PinData->VariableData.Name = "NewPin";
+	UCONodeMaterialBreakParameterPinData* PinData = NewObject<UCONodeMaterialBreakParameterPinData>(this);
 	PinData->VariableData.Type = NewPinCategory;
 	PinData->VariableData.Id = FGuid::NewGuid();
 
-	if (CustomCreatePin(EEdGraphPinDirection::EGPD_Output, NewPinCategory, "NewPin", PinData))
+	if (CustomCreatePin(EEdGraphPinDirection::EGPD_Output, NewPinCategory, BreakMaterialMaterialParameterPinName, PinData))
 	{
 		GetGraph()->NotifyNodeChanged(this);
 	}
 }
 
 
+bool UCONodeMaterialBreak::CanRenamePin(const UEdGraphPin& Pin) const
+{
+	return Pin.Direction == EGPD_Output;
+}
+
+
+EEditablePinNameBoxVisibilityPolicy UCONodeMaterialBreak::GetEditablePinNameVisibilityPolicy(const UEdGraphPin& Pin) const
+{
+	return EEditablePinNameBoxVisibilityPolicy::ALWAYS_VISIBLE;
+}
+
+
+void UCONodeMaterialBreak::SetPinEditableName(const UEdGraphPin& Pin, const FText& Value)
+{
+	UCustomizableObjectNodePinData* PinData = GetPinData(Pin);
+	if (UCONodeMaterialBreakParameterPinData* CastedPinData = Cast<UCONodeMaterialBreakParameterPinData>(PinData))
+	{
+		CastedPinData->VariableData.Name = FName{Value.ToString()};
+		PinEditableNameChangedDelegate.Broadcast(Pin, Value);
+	}
+	else
+	{
+		Super::SetPinEditableName(Pin, Value);
+	}
+}
+
+
+FText UCONodeMaterialBreak::GetPinEditableName(const UEdGraphPin& Pin) const
+{
+	const UCustomizableObjectNodePinData* PinData = GetPinData(Pin);
+	if (const UCONodeMaterialBreakParameterPinData* CastedPinData = Cast<UCONodeMaterialBreakParameterPinData>(PinData))
+	{
+		return FText::FromName(CastedPinData->VariableData.Name);
+	}
+
+	return Super::GetPinEditableName(Pin);
+}
+
+
 FName UCONodeMaterialBreak::GetPinParameterName(const UEdGraphPin& Pin) const
 {
-	return GetPinData<UCONodeMaterialBreakPinData>(Pin).VariableData.Name;
+	return GetPinData<UCONodeMaterialBreakParameterPinData>(Pin).VariableData.Name;
 }
 
 
 int32 UCONodeMaterialBreak::GetPinParameterLayerIndex(const UEdGraphPin& Pin) const
 {
-	return GetPinData<UCONodeMaterialBreakPinData>(Pin).LayerIndex;
+	return GetPinData<UCONodeMaterialBreakParameterPinData>(Pin).LayerIndex;
 }
 
 
diff --git a/Engine/Plugins/Mutable/Source/CustomizableObjectEditor/Private/MuCOE/Nodes/CONodeMaterialBreak.h b/Engine/Plugins/Mutable/Source/CustomizableObjectEditor/Private/MuCOE/Nodes/CONodeMaterialBreak.h
--- a/Engine/Plugins/Mutable/Source/CustomizableObjectEditor/Private/MuCOE/Nodes/CONodeMaterialBreak.h
+++ b/Engine/Plugins/Mutable/Source/CustomizableObjectEditor/Private/MuCOE/Nodes/CONodeMaterialBreak.h
@@ -25,7 +25,7 @@ public:
 
 
 UCLASS(MinimalAPI)
-class UCONodeMaterialBreakPinData : public UCustomizableObjectNodePinData
+class UCONodeMaterialBreakParameterPinData : public UCustomizableObjectNodePinData
 {
 	GENERATED_BODY()
 
@@ -39,6 +39,11 @@ public:
 	UPROPERTY(EditAnywhere, Category = NoCategory)
 	int32 LayerIndex = -1;
 
+	/**
+	 * Handle change in the LayerIndex property
+	 * @param PropertyChangedEvent The Property Change Event
+	 */
+	virtual void PostEditChangeProperty(FPropertyChangedEvent& PropertyChangedEvent) override;
 };
 
 
@@ -65,6 +70,12 @@ public:
 	UE_API virtual TArray<FName> GetAllowedPinViewerCreationTypes() const override;
 	UE_API virtual bool IsExperimental() const override;
 
+	// UCustomizableObjectNode Interface : Editable area of the pin name management
+	virtual bool CanRenamePin(const UEdGraphPin& Pin) const override;
+	virtual EEditablePinNameBoxVisibilityPolicy GetEditablePinNameVisibilityPolicy(const UEdGraphPin& Pin) const override;
+	virtual void SetPinEditableName(const UEdGraphPin& Pin, const FText& Value) override;
+	virtual FText GetPinEditableName(const UEdGraphPin& Pin) const override;
+	
 	// Own interface
 	FName GetPinParameterName(const UEdGraphPin& Pin) const;
 	int32 GetPinParameterLayerIndex(const UEdGraphPin& Pin) const;
diff --git a/Engine/Plugins/Mutable/Source/CustomizableObjectEditor/Private/MuCOE/Nodes/CONodeMaterialModify.cpp b/Engine/Plugins/Mutable/Source/CustomizableObjectEditor/Private/MuCOE/Nodes/CONodeMaterialModify.cpp
--- a/Engine/Plugins/Mutable/Source/CustomizableObjectEditor/Private/MuCOE/Nodes/CONodeMaterialModify.cpp
+++ b/Engine/Plugins/Mutable/Source/CustomizableObjectEditor/Private/MuCOE/Nodes/CONodeMaterialModify.cpp
@@ -10,10 +10,13 @@
 #define LOCTEXT_NAMESPACE "CustomizableObjectEditor"
 
 
+const FName MaterialModifyMaterialParameterPinName = "Parameter";
+
+
 bool UCONodeMaterialModifyRemapPins::Equal(const UCustomizableObjectNode& Node, const UEdGraphPin& OldPin, const UEdGraphPin& NewPin) const
 {
-	UCONodeMaterialModifyPinData* OldPinData = Cast<UCONodeMaterialModifyPinData>(Node.GetPinData(OldPin));
-	UCONodeMaterialModifyPinData* NewPinData = Cast<UCONodeMaterialModifyPinData>(Node.GetPinData(NewPin));
+	const UCONodeMaterialModifyParameterPinData* OldPinData = Cast<UCONodeMaterialModifyParameterPinData>(Node.GetPinData(OldPin));
+	const UCONodeMaterialModifyParameterPinData* NewPinData = Cast<UCONodeMaterialModifyParameterPinData>(Node.GetPinData(NewPin));
 
 	if (OldPinData && NewPinData)
 	{
@@ -26,6 +29,20 @@ bool UCONodeMaterialModifyRemapPins::Equal(const UCustomizableObjectNode& Node,
 }
 
 
+void UCONodeMaterialModifyParameterPinData::PostEditChangeProperty(FPropertyChangedEvent& PropertyChangedEvent)
+{
+	Super::PostEditChangeProperty(PropertyChangedEvent);
+
+	if (const FProperty* PropertyThatChanged = PropertyChangedEvent.Property)
+	{
+		if (PropertyThatChanged->GetFName() == GET_MEMBER_NAME_CHECKED(UCONodeMaterialModifyParameterPinData, LayerIndex))
+		{
+			CastChecked<UCustomizableObjectNode>(GetOuter())->ReconstructNode();
+		}
+	}
+}
+
+
 FText UCONodeMaterialModify::GetNodeTitle(ENodeTitleType::Type TitleType) const
 {
 	return LOCTEXT("ModifyMaterialNodeTitle","Modify Material");
@@ -40,7 +57,7 @@ FLinearColor UCONodeMaterialModify::GetNodeTitleColor() const
 
 FText UCONodeMaterialModify::GetTooltipText() const
 {
-	return LOCTEXT("ModifyMaterialNodeTooltip","Modifys a new material instance from the Reference Material.");
+	return LOCTEXT("ModifyMaterialNodeTooltip","Modifies a new material instance from the Reference Material.");
 }
 
 
@@ -65,22 +82,51 @@ void UCONodeMaterialModify::AllocateDefaultPins(UCustomizableObjectNodeRemapPins
 	TArray<UEdGraphPin*> InputPins = GetAllNonOrphanPins();
 
 	// Regenerate the new pins from the old pin's pindata which may have changed
-	for (const UEdGraphPin* Pin : InputPins)
+	for (const UEdGraphPin* InputPin : InputPins)
 	{
-		if (Pin->Direction == EEdGraphPinDirection::EGPD_Input)
+		if (InputPin->Direction == EEdGraphPinDirection::EGPD_Input)
 		{
-			if (UCONodeMaterialModifyPinData* PinData = Cast<UCONodeMaterialModifyPinData>(GetPinData(*Pin)))
+			if (UCONodeMaterialModifyParameterPinData* OldPinData = Cast<UCONodeMaterialModifyParameterPinData>(GetPinData(*InputPin)))
 			{
-				// Use the same pin data as the old pin since the new one is a copy
-				UEdGraphPin* NewPin = CustomCreatePin(EGPD_Input, PinData->VariableData.Type, PinData->VariableData.Name, PinData);
-
-				if (PinData->LayerIndex == 0)
+				UEdGraphPin* NewPin = nullptr;
+				UCONodeMaterialModifyParameterPinData* NewPinData = nullptr;
+				
+				if (OldPinData->VariableData.Type == UEdGraphSchema_CustomizableObject::PC_Float)
+				{
+					NewPinData = NewObject<UCONodeMaterialModifyScalarParamPinData>(this);
+				}
+				else if (OldPinData->VariableData.Type == UEdGraphSchema_CustomizableObject::PC_Texture)
+				{
+					NewPinData = NewObject<UCONodeMaterialModifyTextureParamPinData>(this);
+				}
+				else if (OldPinData->VariableData.Type == UEdGraphSchema_CustomizableObject::PC_PassthroughTexture)
+				{
+					NewPinData = NewObject<UCONodeMaterialModifyPassthroughTextureParamPinData>(this);
+				}
+				else if (OldPinData->VariableData.Type == UEdGraphSchema_CustomizableObject::PC_Color)
 				{
-					NewPin->PinFriendlyName = FText::Format(LOCTEXT("BackgroundLayerText","{0} - Background"), FText::FromName(PinData->VariableData.Name));
+					NewPinData = NewObject<UCONodeMaterialModifyColorParamPinData>(this);
 				}
-				else if (PinData->LayerIndex > 0)
+				
+				// A new Pin data structure has been created because one pin has changed its type (from the UI)
+				check(NewPinData)
+				if (NewPinData)
 				{
-					NewPin->PinFriendlyName = FText::Format(LOCTEXT("LayerText", "{0} - Layer {1}"), FText::FromName(PinData->VariableData.Name), FText::FromString(FString::FromInt(PinData->LayerIndex)));
+					NewPinData->VariableData.Id = OldPinData->VariableData.Id;
+					NewPinData->VariableData.Name = OldPinData->VariableData.Name;
+					NewPinData->LayerIndex = OldPinData->LayerIndex;
+					
+					NewPin = CustomCreatePin(EGPD_Input, NewPinData->VariableData.Type, MaterialModifyMaterialParameterPinName, NewPinData);
+				}
+				
+				check(NewPin)
+				if (OldPinData->LayerIndex == 0)
+				{
+					NewPin->PinFriendlyName = FText::Format(LOCTEXT("BackgroundLayerText","{0} - Background"), FText::FromName(MaterialModifyMaterialParameterPinName));
+				}
+				else if (OldPinData->LayerIndex > 0)
+				{
+					NewPin->PinFriendlyName = FText::Format(LOCTEXT("LayerText", "{0} - Layer {1}"), FText::FromName(MaterialModifyMaterialParameterPinName), FText::FromString(FString::FromInt(OldPinData->LayerIndex)));
 				}
 			}
 		}
@@ -90,7 +136,7 @@ void UCONodeMaterialModify::AllocateDefaultPins(UCustomizableObjectNodeRemapPins
 
 TSharedPtr<IDetailsView> UCONodeMaterialModify::CustomizePinDetails(const UEdGraphPin& Pin) const
 {
-	if (UCONodeMaterialModifyPinData* PinData = Cast<UCONodeMaterialModifyPinData>(GetPinData(Pin)))
+	if (UCONodeMaterialModifyParameterPinData* PinData = Cast<UCONodeMaterialModifyParameterPinData>(GetPinData(Pin)))
 	{
 		FPropertyEditorModule& EditModule = FModuleManager::Get().GetModuleChecked<FPropertyEditorModule>("PropertyEditor");
 
@@ -143,15 +189,9 @@ void UCONodeMaterialModify::AddPinFromUI()
 
 	FScopedTransaction Transaction(LOCTEXT("CONodeMaterialModify_AddPinFromNodeUITransaction", "Add Pin from Node UI"));
 	Modify();
-
-	const FName NewPinCategory = UEdGraphSchema_CustomizableObject::PC_Float;
-
-	UCONodeMaterialModifyPinData* PinData = NewObject<UCONodeMaterialModifyPinData>(this);
-	PinData->VariableData.Name = "NewPin";
-	PinData->VariableData.Type = NewPinCategory;
-	PinData->VariableData.Id = FGuid::NewGuid();
-
-	if (CustomCreatePin(EEdGraphPinDirection::EGPD_Input, NewPinCategory, "NewPin", PinData))
+	
+	UCONodeMaterialModifyScalarParamPinData* ScalarParameterPinData = NewObject<UCONodeMaterialModifyScalarParamPinData>(this);
+	if (CustomCreatePin(EEdGraphPinDirection::EGPD_Input, ScalarParameterPinData->VariableData.Type, MaterialModifyMaterialParameterPinName, ScalarParameterPinData))
 	{
 		GetGraph()->NotifyNodeChanged(this);
 	}
@@ -166,13 +206,54 @@ bool UCONodeMaterialModify::CanPinBeRemoved(const UEdGraphPin& Pin) const
 
 FName UCONodeMaterialModify::GetPinParameterName(const UEdGraphPin& Pin) const
 {
-	return GetPinData<UCONodeMaterialModifyPinData>(Pin).VariableData.Name;
+	return GetPinData<UCONodeMaterialModifyParameterPinData>(Pin).VariableData.Name;
 }
 
 
 int32 UCONodeMaterialModify::GetPinParameterLayerIndex(const UEdGraphPin& Pin) const
 {
-	return GetPinData<UCONodeMaterialModifyPinData>(Pin).LayerIndex;
+	return GetPinData<UCONodeMaterialModifyParameterPinData>(Pin).LayerIndex;
+}
+
+
+TSoftObjectPtr<UTexture2D> UCONodeMaterialModify::GetTexturePinParameterReferenceTexture(const UEdGraphPin& Pin) const
+{
+	const UCONodeMaterialModifyTextureParamPinData& TextureParameterPnData = GetPinData<UCONodeMaterialModifyTextureParamPinData>(Pin);
+	return TextureParameterPnData.ReferenceTexture;
+}
+
+
+int32 UCONodeMaterialModify::GetImagePinParameterUVIndex(const UEdGraphPin& Pin) const
+{
+	const UCONodeMaterialModifyTextureParamPinData& TextureParameterPnData = GetPinData<UCONodeMaterialModifyTextureParamPinData>(Pin);
+	switch(TextureParameterPnData.UVLayoutMode)
+	{
+	case EMaterialModifyUVLayoutMode::Ignore:
+		return UCustomizableObjectNodeMaterialPinDataImage::UV_LAYOUT_IGNORE;
+	case EMaterialModifyUVLayoutMode::Index:
+		return TextureParameterPnData.UVLayout;
+	default:
+		unimplemented();
+	}
+		
+	return UCustomizableObjectNodeMaterialPinDataImage::UV_LAYOUT_IGNORE;
+}
+
+
+void UCONodeMaterialModify::GetMaterialTextureParameterPins(TArray<FEdGraphPinReference>& OutTextureParameterPins) const
+{
+	TArray<UEdGraphPin*> AllNonOrphanPins = GetAllNonOrphanPins();
+
+	for (UEdGraphPin* NonOrphanPin : AllNonOrphanPins)
+	{
+		if (NonOrphanPin && 
+			NonOrphanPin->Direction == EGPD_Input &&
+			NonOrphanPin->PinType.PinCategory == GetDefault<UEdGraphSchema_CustomizableObject>()->PC_Texture)
+		{
+			check(GetPinData(*NonOrphanPin)->IsA<UCONodeMaterialModifyTextureParamPinData>());
+			OutTextureParameterPins.Add(NonOrphanPin);
+		}
+	}
 }
 
 
@@ -182,4 +263,43 @@ bool UCONodeMaterialModify::IsExperimental() const
 }
 
 
-#undef LOCTEXT_NAMESPACE
\ No newline at end of file
+bool UCONodeMaterialModify::CanRenamePin(const UEdGraphPin& Pin) const
+{
+	return Pin.Direction == EGPD_Input && FEdGraphPinReference{&Pin} != MaterialPinRef;
+}
+
+
+EEditablePinNameBoxVisibilityPolicy UCONodeMaterialModify::GetEditablePinNameVisibilityPolicy(const UEdGraphPin& Pin) const
+{
+	return EEditablePinNameBoxVisibilityPolicy::ALWAYS_VISIBLE;
+}
+
+
+void UCONodeMaterialModify::SetPinEditableName(const UEdGraphPin& Pin, const FText& Value)
+{
+	UCustomizableObjectNodePinData* PinData = GetPinData(Pin);
+	if (UCONodeMaterialModifyParameterPinData* CastedPinData = Cast<UCONodeMaterialModifyParameterPinData>(PinData))
+	{
+		CastedPinData->VariableData.Name = FName{Value.ToString()};
+		PinEditableNameChangedDelegate.Broadcast(Pin, Value);
+	}
+	else
+	{
+		Super::SetPinEditableName(Pin, Value);
+	}
+}
+
+
+FText UCONodeMaterialModify::GetPinEditableName(const UEdGraphPin& Pin) const
+{
+	const UCustomizableObjectNodePinData* PinData = GetPinData(Pin);
+	if (const UCONodeMaterialModifyParameterPinData* CastedPinData = Cast<UCONodeMaterialModifyParameterPinData>(PinData))
+	{
+		return FText::FromName(CastedPinData->VariableData.Name);
+	}
+
+	return Super::GetPinEditableName(Pin);
+}
+
+
+#undef LOCTEXT_NAMESPACE
diff --git a/Engine/Plugins/Mutable/Source/CustomizableObjectEditor/Private/MuCOE/Nodes/CONodeMaterialModify.h b/Engine/Plugins/Mutable/Source/CustomizableObjectEditor/Private/MuCOE/Nodes/CONodeMaterialModify.h
--- a/Engine/Plugins/Mutable/Source/CustomizableObjectEditor/Private/MuCOE/Nodes/CONodeMaterialModify.h
+++ b/Engine/Plugins/Mutable/Source/CustomizableObjectEditor/Private/MuCOE/Nodes/CONodeMaterialModify.h
@@ -2,9 +2,12 @@
 
 #pragma once
 
+#include "CustomizableObjectNodeMaterial.h"
+#include "IDetailCustomization.h"
 #include "MuCOE/Nodes/CustomizableObjectNode.h"
 #include "MuCOE/RemapPins/CustomizableObjectNodeRemapPinsByName.h"
 #include "MuCOE/COVariable.h"
+#include "MuCOE/EdGraphSchema_CustomizableObject.h"
 #include "CONodeMaterialModify.generated.h"
 
 #define UE_API CUSTOMIZABLEOBJECTEDITOR_API
@@ -21,16 +24,30 @@ public:
 };
 
 
-// TODO(Max - GM)UE-318756 This should have the same properties than UCustomizableObjectNodeMaterialPinDataImage
-// TODO(Max - GM)UE-318756 Reformulate all the pin logic for this node
+// todo: Remove this enum and use the "EUVLayoutMode" so we support the option of getting the UV layout from the material  
+/** Image Pin, UV Layout Mode. */
+UENUM()
+enum class EMaterialModifyUVLayoutMode
+{
+	/* Texture should not be transformed by any layout. These textures will not be reduced automatically for LODs. */
+	Ignore,
+	/** User specified UV Index. */
+	Index
+};
+
+
 /** Base class for all Material Parameters. */
 UCLASS()
-class UCONodeMaterialModifyPinData : public UCustomizableObjectNodePinData
+class UCONodeMaterialModifyParameterPinData : public UCustomizableObjectNodePinData
 {
 	GENERATED_BODY()
 
 public:
-
+	UCONodeMaterialModifyParameterPinData()
+	{
+		VariableData.Id = FGuid::NewGuid();
+	}
+	
 	UPROPERTY(EditAnywhere, Category = NoCategory)
 	FCOVariable VariableData;
 	
@@ -38,8 +55,107 @@ public:
 	-1 if the material does not contain layers.*/
 	UPROPERTY(EditAnywhere, Category = NoCategory)
 	int32 LayerIndex = -1;
+
+	/**
+	 * Handle change in the LayerIndex property
+	 * @param PropertyChangedEvent The Property Change Event
+	 */
+	virtual void PostEditChangeProperty(FPropertyChangedEvent& PropertyChangedEvent) override;
+};
+
+
+/**
+ * PinData object representing pins of the PC_Texture type.
+ */
+UCLASS()
+class UCONodeMaterialModifyTextureParamPinData : public UCONodeMaterialModifyParameterPinData
+{
+	GENERATED_BODY()
+	
+public:
+	UCONodeMaterialModifyTextureParamPinData() 
+	{
+		VariableData.Type = UEdGraphSchema_CustomizableObject::PC_Texture;
+	}
+	
+
+	virtual void Copy(const UCustomizableObjectNodePinData& Other) override
+	{
+		if (const UCONodeMaterialModifyTextureParamPinData* PinDataOldPin = Cast<UCONodeMaterialModifyTextureParamPinData>(&Other))
+		{
+			VariableData = PinDataOldPin->VariableData;
+            
+            UVLayoutMode = PinDataOldPin->UVLayoutMode;
+            UVLayout = PinDataOldPin->UVLayout;
+            ReferenceTexture = PinDataOldPin->ReferenceTexture;
+		}
+	}
+
+	
+	constexpr static int32 UV_LAYOUT_IGNORE = -1;
+	
+	UPROPERTY(EditAnywhere, Category = NoCategory)
+	EMaterialModifyUVLayoutMode UVLayoutMode = EMaterialModifyUVLayoutMode::Ignore;
+	
+	/** Index of the UV channel that will be used with this image. It is necessary to apply the proper layout transformations to it. */
+	UPROPERTY(EditAnywhere, Category = NoCategory, meta = (EditCondition = "UVLayoutMode == EMaterialModifyUVLayoutMode::Index", EditConditionHides))
+	int32 UVLayout = -2;
+	
+	/** Reference Texture used to decide the texture properties of the mutable-generated textures
+	* connected to this material. If null, it will try to be guessed at compile time from
+	* the graph. */
+	UPROPERTY(EditAnywhere, Category = NoCategory) // Required to be EditAnywhere for the selector to work.
+	TSoftObjectPtr<UTexture2D> ReferenceTexture = nullptr;
+};
+
+/**
+ * PinData object representing pins of the PC_PassthroughTexture type.
+ */
+UCLASS()
+class UCONodeMaterialModifyPassthroughTextureParamPinData : public UCONodeMaterialModifyParameterPinData
+{
+	GENERATED_BODY()
+	
+public:
+	UCONodeMaterialModifyPassthroughTextureParamPinData()
+	{
+		VariableData.Type = UEdGraphSchema_CustomizableObject::PC_PassthroughTexture;
+	}
 };
 
+
+/**
+ * PinData object representing pins of the PC_Float type.
+ */
+UCLASS()
+class UCONodeMaterialModifyScalarParamPinData : public UCONodeMaterialModifyParameterPinData
+{
+	GENERATED_BODY()
+	
+public:
+	UCONodeMaterialModifyScalarParamPinData()
+	{
+		VariableData.Type = UEdGraphSchema_CustomizableObject::PC_Float;
+	}
+};
+
+
+/**
+ * PinData object representing pins of the PC_Color type.
+ */
+UCLASS()
+class UCONodeMaterialModifyColorParamPinData : public UCONodeMaterialModifyParameterPinData
+{
+	GENERATED_BODY()
+	
+public:
+	UCONodeMaterialModifyColorParamPinData()
+	{
+		VariableData.Type = UEdGraphSchema_CustomizableObject::PC_Color;
+	}
+};
+
+
 UCLASS()
 class UCONodeMaterialModify : public UCustomizableObjectNode
 {
@@ -63,11 +179,18 @@ public:
 	virtual TArray<FName> GetAllowedPinViewerCreationTypes() const override;
 	virtual bool IsExperimental() const override;
 
+	// Own Interface : Editable area of the pin name management
+	virtual bool CanRenamePin(const UEdGraphPin& Pin) const override;
+	virtual EEditablePinNameBoxVisibilityPolicy GetEditablePinNameVisibilityPolicy(const UEdGraphPin& Pin) const override;
+	virtual void SetPinEditableName(const UEdGraphPin& Pin, const FText& Value) override;
+	virtual FText GetPinEditableName(const UEdGraphPin& Pin) const override;
+	
 	// Own interface
 	FName GetPinParameterName(const UEdGraphPin& Pin) const;
 	int32 GetPinParameterLayerIndex(const UEdGraphPin& Pin) const;
-
-public:
+	TSoftObjectPtr<UTexture2D> GetTexturePinParameterReferenceTexture(const UEdGraphPin& Pin) const;
+	int32 GetImagePinParameterUVIndex(const UEdGraphPin& Pin) const;
+	void GetMaterialTextureParameterPins(TArray<FEdGraphPinReference>& OutTextureParameterPins) const;
 	
 	UPROPERTY()
 	FEdGraphPinReference MaterialPinRef;
diff --git a/Engine/Plugins/Mutable/Source/CustomizableObjectEditor/Private/MuCOE/Nodes/CONodeSkeletalMeshBreak.cpp b/Engine/Plugins/Mutable/Source/CustomizableObjectEditor/Private/MuCOE/Nodes/CONodeSkeletalMeshBreak.cpp
--- a/Engine/Plugins/Mutable/Source/CustomizableObjectEditor/Private/MuCOE/Nodes/CONodeSkeletalMeshBreak.cpp
+++ b/Engine/Plugins/Mutable/Source/CustomizableObjectEditor/Private/MuCOE/Nodes/CONodeSkeletalMeshBreak.cpp
@@ -75,7 +75,8 @@ void UCONodeSkeletalMeshBreak::SetPinEditableName(const UEdGraphPin& Pin, const
 	UCustomizableObjectNodePinData* PinData = GetPinData(Pin);
 	if (UCONodeSkeletalMeshBreakPinData* CastedPinData = Cast<UCONodeSkeletalMeshBreakPinData>(PinData))
 	{
-		CastedPinData->TargetMaterialSlotName =  FName{Value.ToString()};
+		CastedPinData->TargetMaterialSlotName = FName{Value.ToString()};
+		PinEditableNameChangedDelegate.Broadcast(Pin, Value);
 	}
 	else
 	{
