[Mutable] Fix double-counting due LLM::UpdateStats being called multiple times per frame.

- The change causing the issue was erroneously introduced in CL 48724838. Reverts MemoryTrackingUtils.cpp to its original state.

#rnx
#Synced-CL Unknown

diff --git a/Engine/Plugins/Mutable/Source/MutableRuntime/Private/MuR/MemoryTrackingUtils.cpp b/Engine/Plugins/Mutable/Source/MutableRuntime/Private/MuR/MemoryTrackingUtils.cpp
--- a/Engine/Plugins/Mutable/Source/MutableRuntime/Private/MuR/MemoryTrackingUtils.cpp
+++ b/Engine/Plugins/Mutable/Source/MutableRuntime/Private/MuR/MemoryTrackingUtils.cpp
@@ -16,7 +16,7 @@ namespace UE::Mutable::Private::Private
 	{
 		// This code will only be compiled if the global definition to enable LLM tracking is set to 1 for the host program
 		// E.g.: GlobalDefinitions.Add("LLM_ENABLED_IN_CONFIG=1");
-#if ENABLE_LOW_LEVEL_MEM_TRACKER /* && IS_PROGRAM*/
+#if ENABLE_LOW_LEVEL_MEM_TRACKER && IS_PROGRAM
 		FLowLevelMemTracker& MemTracker = FLowLevelMemTracker::Get();
 		if (MemTracker.IsEnabled())
 		{
@@ -38,8 +38,8 @@ namespace UE::Mutable::Private
 			TEXT("Enables Mutable plugin global memory tracking, peak and total."),
 			ECVF_Default);
 
-#if ENABLE_LOW_LEVEL_MEM_TRACKER /*&& IS_PROGRAM*/
-		static inline bool bEnableUpdateLLMStats = true;
+#if ENABLE_LOW_LEVEL_MEM_TRACKER && IS_PROGRAM
+		static inline bool bEnableUpdateLLMStats = false;
 		static inline FAutoConsoleVariableRef CVarEnableUpdateLLMStats = FAutoConsoleVariableRef(
 			TEXT("mutable.EnableUpdateLLMStats"),
 			bEnableUpdateLLMStats,
@@ -64,7 +64,7 @@ namespace UE::Mutable::Private
 		AbsoluteCounter += Differential;
 		AbsolutePeakValue = FMath::Max(AbsolutePeakValue, AbsoluteCounter);
 
-#if ENABLE_LOW_LEVEL_MEM_TRACKER /*&& IS_PROGRAM*/
+#if ENABLE_LOW_LEVEL_MEM_TRACKER && IS_PROGRAM
 		if (UNLIKELY(FMemoryTrackingConsoleFlags::bEnableUpdateLLMStats))
 		{
 			Private::UpdateLLMStats();
