[Mutable] Fixed Parameters cache.

#rnx

#Synced-CL 48900308

diff --git a/Engine/Plugins/Mutable/Source/CustomizableObjectEditor/Private/MuCOE/GenerateMutableSource/GenerateMutableSource.cpp b/Engine/Plugins/Mutable/Source/CustomizableObjectEditor/Private/MuCOE/GenerateMutableSource/GenerateMutableSource.cpp
--- a/Engine/Plugins/Mutable/Source/CustomizableObjectEditor/Private/MuCOE/GenerateMutableSource/GenerateMutableSource.cpp
+++ b/Engine/Plugins/Mutable/Source/CustomizableObjectEditor/Private/MuCOE/GenerateMutableSource/GenerateMutableSource.cpp
@@ -405,17 +405,25 @@ FGeneratedKey::FGeneratedKey(void* InFunctionAddress, const UEdGraphPin& InPin,
 	{
 		MaterialBreakParameterStack = GenerationContext.MaterialBreakParameterStack;
 	}
-	
-	// Do not use the Stack of macros as Key if the Node is a Parameter Node. 
-	// By doing so, we can define parameters inside mutable macros and reuse them without having repeated nodes.
-	MacroContext = Node.IsA(UCustomizableObjectNodeParameter::StaticClass()) ? TArray<const UCustomizableObjectNodeMacroInstance*>{} : GenerationContext.MacroNodesStack;
 
-	if (UseMesh)
+	if (const UCustomizableObjectNodeParameter* ParameterNode = Cast<UCustomizableObjectNodeParameter>(&Node))
+	{
+		// Do not use the Stack of macros as Key if the Node is a Parameter Node. 
+		// By doing so, we can define parameters inside mutable macros and reuse them without having repeated nodes.
+		MacroContext = TArray<const UCustomizableObjectNodeMacroInstance*>{};
+		ParameterName = ParameterNode->GetParameterName(&GenerationContext.MacroNodesStack);
+	}
+	else
 	{
-		Flags = GenerationContext.MeshGenerationFlags.Last();
-		LayoutFlags = GenerationContext.LayoutGenerationFlags.Last();
-		MeshMorphStack = GenerationContext.MeshMorphStack;
-		bOnlyConnectedLOD = InbOnlyConnectedLOD;
+		MacroContext = GenerationContext.MacroNodesStack;
+
+		if (UseMesh)
+		{
+			Flags = GenerationContext.MeshGenerationFlags.Last();
+			LayoutFlags = GenerationContext.LayoutGenerationFlags.Last();
+			MeshMorphStack = GenerationContext.MeshMorphStack;
+			bOnlyConnectedLOD = InbOnlyConnectedLOD;
+		}
 	}
 }
 
@@ -470,6 +478,15 @@ uint32 GetTypeHash(const FGeneratedSourceSkeletalMeshKey& Key)
 }
 
 
+uint32 GetTypeHash(const FGeneratedParameterKey& Key)
+{
+	uint32 Hash = GetTypeHash(Key.NodeId);
+	Hash = HashCombine(Hash, GetTypeHash(Key.ParameterName));
+
+	return Hash;
+}
+
+
 FGraphCycleKey::FGraphCycleKey(const UEdGraphPin& Pin, const FString& Id, const UCustomizableObjectNodeMacroInstance* MacroContext) :
 	Pin(Pin),
 	Id(Id),
diff --git a/Engine/Plugins/Mutable/Source/CustomizableObjectEditor/Private/MuCOE/GenerateMutableSource/GenerateMutableSource.h b/Engine/Plugins/Mutable/Source/CustomizableObjectEditor/Private/MuCOE/GenerateMutableSource/GenerateMutableSource.h
--- a/Engine/Plugins/Mutable/Source/CustomizableObjectEditor/Private/MuCOE/GenerateMutableSource/GenerateMutableSource.h
+++ b/Engine/Plugins/Mutable/Source/CustomizableObjectEditor/Private/MuCOE/GenerateMutableSource/GenerateMutableSource.h
@@ -197,6 +197,9 @@ public:
 	TArray<FMaterialBreakParameter> MaterialBreakParameterStack;
 	
 	int32 ReferenceTextureSize = -1;
+
+	/** Name of the parameter. A parameter node can be redefined if it is inside a macro and the name comes from outside the graph. */
+	FString ParameterName;
 };
 
 
@@ -204,6 +207,21 @@ uint32 GetTypeHash(const FMaterialBreakParameter& Key);
 
 uint32 GetTypeHash(const FGeneratedKey& Key);
 
+// Cache key of the generated parameter nodes
+struct FGeneratedParameterKey
+{
+	// ID of the node
+	FGuid NodeId;
+
+	// Name of the parameter.
+	FString ParameterName;
+
+	bool operator==(const FGeneratedParameterKey& Other) const = default;
+};
+
+uint32 GetTypeHash(const FGeneratedParameterKey& Key);
+
+
 // Struct used for external options of the GenerateMutableSourceX functions.
 struct FSourceExternalOptions
 {
@@ -707,20 +725,20 @@ struct FMutableGraphGenerationContext
 
 	TMap<FGeneratedGroupProjectorsKey, FGroupProjectorTempData> GeneratedGroupProjectors;
 
-	/** Key is the Node Uid. */
-	TMap<FGuid, UE::Mutable::Private::Ptr<UE::Mutable::Private::NodeScalarParameter>> GeneratedScalarParameters;
+	/** Key is the Node Uid and the parameter name. */
+	TMap<FGeneratedParameterKey, UE::Mutable::Private::Ptr<UE::Mutable::Private::NodeScalarParameter>> GeneratedScalarParameters;
 
-	/** Key is the Node Uid. */
-	TMap<FGuid, UE::Mutable::Private::Ptr<UE::Mutable::Private::NodeScalarEnumParameter>> GeneratedEnumParameters;
+	/** Key is the Node Uid and the parameter name. */
+	TMap<FGeneratedParameterKey, UE::Mutable::Private::Ptr<UE::Mutable::Private::NodeScalarEnumParameter>> GeneratedEnumParameters;
 
-	/** Key is the Node Uid. */
-	TMap<FGuid, UE::Mutable::Private::Ptr<UE::Mutable::Private::NodeImageParameter>> GeneratedImageParameters;
+	/** Key is the Node Uid and the parameter name. */
+	TMap<FGeneratedParameterKey, UE::Mutable::Private::Ptr<UE::Mutable::Private::NodeImageParameter>> GeneratedImageParameters;
 
-	/** Key is the Node Uid. */
-	TMap<FGuid, UE::Mutable::Private::Ptr<UE::Mutable::Private::NodeSkeletalMeshParameter>> GeneratedSkeletalMeshParameters;
+	/** Key is the Node Uid and the parameter name. */
+	TMap<FGeneratedParameterKey, UE::Mutable::Private::Ptr<UE::Mutable::Private::NodeSkeletalMeshParameter>> GeneratedSkeletalMeshParameters;
 
-	/** Key is the Node Uid. */
-	TMap<FGuid, UE::Mutable::Private::Ptr<UE::Mutable::Private::NodeMaterialParameter>> GeneratedMaterialParameters;
+	/** Key is the Node Uid and the parameter name. */
+	TMap<FGeneratedParameterKey, UE::Mutable::Private::Ptr<UE::Mutable::Private::NodeMaterialParameter>> GeneratedMaterialParameters;
 	
 	struct FGeneratedCompositeDataTablesData
 	{
diff --git a/Engine/Plugins/Mutable/Source/CustomizableObjectEditor/Private/MuCOE/GenerateMutableSource/GenerateMutableSourceExternal.cpp b/Engine/Plugins/Mutable/Source/CustomizableObjectEditor/Private/MuCOE/GenerateMutableSource/GenerateMutableSourceExternal.cpp
--- a/Engine/Plugins/Mutable/Source/CustomizableObjectEditor/Private/MuCOE/GenerateMutableSource/GenerateMutableSourceExternal.cpp
+++ b/Engine/Plugins/Mutable/Source/CustomizableObjectEditor/Private/MuCOE/GenerateMutableSource/GenerateMutableSourceExternal.cpp
@@ -146,14 +146,15 @@ UE::Mutable::Private::Ptr<UE::Mutable::Private::NodeExternal> GenerateMutableSou
 
 	if (const UCONodeExternalTypeParameter* TypeNodeExternalParameter = Cast<UCONodeExternalTypeParameter>(EditorNode))
 	{
-		UE::Mutable::Private::Ptr<UE::Mutable::Private::NodeExternalParameter> ExternalParameterNode = new UE::Mutable::Private::NodeExternalParameter();
+		FString ParameterName = TypeNodeExternalParameter->GetParameterName(&GenerationContext.MacroNodesStack);
 
-		ExternalParameterNode->Name = TypeNodeExternalParameter->GetParameterName(&GenerationContext.MacroNodesStack);
+		UE::Mutable::Private::Ptr<UE::Mutable::Private::NodeExternalParameter> ExternalParameterNode = new UE::Mutable::Private::NodeExternalParameter();
+		ExternalParameterNode->Name = ParameterName;
 		ExternalParameterNode->UID = GenerationContext.GetNodeIdUnique(Node).ToString();
 		ExternalParameterNode->DefaultValue = UE::Mutable::Private::MakeManaged<FInstancedStruct>(TypeNodeExternalParameter->DefaultValue);
 
-		GenerationContext.ExternalTypeParameterDefaultValues.Add(TypeNodeExternalParameter->GetParameterName(), TypeNodeExternalParameter->DefaultValue);
-		GenerationContext.ExternalTypeParameterTypes.Add(TypeNodeExternalParameter->GetParameterName(), TypeNodeExternalParameter->DefaultValue.GetScriptStruct());
+		GenerationContext.ExternalTypeParameterDefaultValues.Add(ParameterName, TypeNodeExternalParameter->DefaultValue);
+		GenerationContext.ExternalTypeParameterTypes.Add(ParameterName, TypeNodeExternalParameter->DefaultValue.GetScriptStruct());
 		
 		Result = ExternalParameterNode;
 	}
diff --git a/Engine/Plugins/Mutable/Source/CustomizableObjectEditor/Private/MuCOE/GenerateMutableSource/GenerateMutableSourceGroupProjector.cpp b/Engine/Plugins/Mutable/Source/CustomizableObjectEditor/Private/MuCOE/GenerateMutableSource/GenerateMutableSourceGroupProjector.cpp
--- a/Engine/Plugins/Mutable/Source/CustomizableObjectEditor/Private/MuCOE/GenerateMutableSource/GenerateMutableSourceGroupProjector.cpp
+++ b/Engine/Plugins/Mutable/Source/CustomizableObjectEditor/Private/MuCOE/GenerateMutableSource/GenerateMutableSourceGroupProjector.cpp
@@ -317,16 +317,25 @@ TOptional<FGroupProjectorTempData> GenerateMutableGroupProjector(UCustomizableOb
 
 	if (GroupProjectorTempData.NodeProjectorParameterPtr)
 	{
+		FString ParameterName = ProjParamNode->GetParameterName(&GenerationContext.MacroNodesStack);
+
 		// Use the projector parameter uid + num to identify parameters derived from this node
 		// TODO(Max UE-247783): Get Node UID taking into account the macro context
 		FGuid NumLayersParamUid = ProjParamNode->NodeGuid;
 		NumLayersParamUid.D += 1;
+		FGeneratedParameterKey NumLayerParamKey = { NumLayersParamUid, ParameterName };
+
 		FGuid SelectedPoseParamUid = ProjParamNode->NodeGuid;
 		SelectedPoseParamUid.D += 2;
+		FGeneratedParameterKey SelectedPoseParamKey = { SelectedPoseParamUid, ParameterName };
+
 		FGuid OpacityParamUid = ProjParamNode->NodeGuid;
 		OpacityParamUid.D += 3;
+		FGeneratedParameterKey OpacityParamKey = { OpacityParamUid, ParameterName };
+
 		FGuid SelectedImageParamUid = ProjParamNode->NodeGuid;
 		SelectedImageParamUid.D += 4;
+		FGeneratedParameterKey SelectedImageParamKey = { SelectedImageParamUid, ParameterName };
 
 		// Add to UCustomizableObjectNodeGroupProjectorParameter::OptionTextures those textures that are present in
 		// UCustomizableObjectNodeGroupProjectorParameter::OptionTexturesDataTable avoiding any repeated element
@@ -342,7 +351,7 @@ TOptional<FGroupProjectorTempData> GenerateMutableGroupProjector(UCustomizableOb
 		GroupProjectorTempData.CustomizableObjectNodeGroupProjectorParameter = ProjParamNode;
 
 		UE::Mutable::Private::Ptr<UE::Mutable::Private::NodeScalarParameter> NodeScalarParam;
-		if (UE::Mutable::Private::Ptr<UE::Mutable::Private::NodeScalarParameter>* FindResult = GenerationContext.GeneratedScalarParameters.Find(NumLayersParamUid))
+		if (UE::Mutable::Private::Ptr<UE::Mutable::Private::NodeScalarParameter>* FindResult = GenerationContext.GeneratedScalarParameters.Find(NumLayerParamKey))
 		{
 			NodeScalarParam = *FindResult;
 		}
@@ -357,7 +366,7 @@ TOptional<FGroupProjectorTempData> GenerateMutableGroupProjector(UCustomizableOb
 			GenerationContext.ParameterUIDataMap.Add(NodeScalarParamName,
 		FMutableParameterData(ProjParamNode->ParamUIMetadata, EMutableParameterType::Int));
 					
-			GenerationContext.GeneratedScalarParameters.Add(NumLayersParamUid, NodeScalarParam);
+			GenerationContext.GeneratedScalarParameters.Add(NumLayerParamKey, NodeScalarParam);
 		}
 
 		UE::Mutable::Private::Ptr<UE::Mutable::Private::NodeRangeFromScalar> NodeRangeFromScalar = new UE::Mutable::Private::NodeRangeFromScalar;
@@ -367,7 +376,7 @@ TOptional<FGroupProjectorTempData> GenerateMutableGroupProjector(UCustomizableOb
 		GroupProjectorTempData.NodeProjectorParameterPtr->SetRange(0, NodeRangeFromScalar);
 		
 		UE::Mutable::Private::Ptr<UE::Mutable::Private::NodeScalarEnumParameter> PoseEnumParameterNode;
-		if (UE::Mutable::Private::Ptr<UE::Mutable::Private::NodeScalarEnumParameter>* FindResult = GenerationContext.GeneratedEnumParameters.Find(SelectedPoseParamUid))
+		if (UE::Mutable::Private::Ptr<UE::Mutable::Private::NodeScalarEnumParameter>* FindResult = GenerationContext.GeneratedEnumParameters.Find(SelectedPoseParamKey))
 		{
 			PoseEnumParameterNode = *FindResult;
 		}
@@ -381,14 +390,14 @@ TOptional<FGroupProjectorTempData> GenerateMutableGroupProjector(UCustomizableOb
 			PoseEnumParameterNode->DefaultValue = 0;
 			GroupProjectorTempData.PoseOptionsParameter = PoseEnumParameterNode;
 
-			GenerationContext.GeneratedEnumParameters.Add(SelectedPoseParamUid, PoseEnumParameterNode);
+			GenerationContext.GeneratedEnumParameters.Add(SelectedPoseParamKey, PoseEnumParameterNode);
 
 			GenerationContext.ParameterUIDataMap.Add(PoseNodeEnumParamName,
 		FMutableParameterData(ProjParamNode->ParamUIMetadata, EMutableParameterType::Int));
 		}
 		
 		UE::Mutable::Private::Ptr<UE::Mutable::Private::NodeScalarParameter> OpacityParameterNode;
-		if (UE::Mutable::Private::Ptr<UE::Mutable::Private::NodeScalarParameter>* FindResult = GenerationContext.GeneratedScalarParameters.Find(OpacityParamUid))
+		if (UE::Mutable::Private::Ptr<UE::Mutable::Private::NodeScalarParameter>* FindResult = GenerationContext.GeneratedScalarParameters.Find(OpacityParamKey))
 		{
 			OpacityParameterNode = *FindResult;
 		}
@@ -402,7 +411,7 @@ TOptional<FGroupProjectorTempData> GenerateMutableGroupProjector(UCustomizableOb
 			OpacityParameterNode->Ranges.SetNum(1);
 			OpacityParameterNode->Ranges[0] = NodeRangeFromScalar;
 
-			GenerationContext.GeneratedScalarParameters.Add(OpacityParamUid, OpacityParameterNode);
+			GenerationContext.GeneratedScalarParameters.Add(OpacityParamKey, OpacityParameterNode);
 					
 			FMutableParamUIMetadata OpacityMetadata = ProjParamNode->ParamUIMetadata;
 			OpacityMetadata.ObjectFriendlyName = FString("Opacity");
@@ -451,7 +460,7 @@ TOptional<FGroupProjectorTempData> GenerateMutableGroupProjector(UCustomizableOb
 		}
 	
 		UE::Mutable::Private::Ptr<UE::Mutable::Private::NodeScalarEnumParameter> EnumParameterNode;
-		if (UE::Mutable::Private::Ptr<UE::Mutable::Private::NodeScalarEnumParameter>* FindResult = GenerationContext.GeneratedEnumParameters.Find(SelectedImageParamUid))
+		if (UE::Mutable::Private::Ptr<UE::Mutable::Private::NodeScalarEnumParameter>* FindResult = GenerationContext.GeneratedEnumParameters.Find(SelectedImageParamKey))
 		{
 			EnumParameterNode = *FindResult;
 		}
@@ -467,8 +476,8 @@ TOptional<FGroupProjectorTempData> GenerateMutableGroupProjector(UCustomizableOb
 			EnumParameterNode->Ranges.SetNum(1);
 			EnumParameterNode->Ranges[0] = NodeRangeFromScalar;
 
-			GenerationContext.GeneratedEnumParameters.Add(SelectedImageParamUid, EnumParameterNode);
-					
+			GenerationContext.GeneratedEnumParameters.Add(SelectedImageParamKey, EnumParameterNode);
+
 			FMutableParameterData ParameterUIData;
 			ParameterUIData.ParamUIMetadata = ProjParamNode->ParamUIMetadata;
 			ParameterUIData.Type = EMutableParameterType::Int;
diff --git a/Engine/Plugins/Mutable/Source/CustomizableObjectEditor/Private/MuCOE/GenerateMutableSource/GenerateMutableSourceImage.cpp b/Engine/Plugins/Mutable/Source/CustomizableObjectEditor/Private/MuCOE/GenerateMutableSource/GenerateMutableSourceImage.cpp
--- a/Engine/Plugins/Mutable/Source/CustomizableObjectEditor/Private/MuCOE/GenerateMutableSource/GenerateMutableSourceImage.cpp
+++ b/Engine/Plugins/Mutable/Source/CustomizableObjectEditor/Private/MuCOE/GenerateMutableSource/GenerateMutableSourceImage.cpp
@@ -1299,8 +1299,8 @@ UE::Mutable::Private::Ptr<UE::Mutable::Private::NodeImageParameter> GenerateMuta
 
 	UE::Mutable::Private::Ptr<UE::Mutable::Private::NodeImageParameter> Result;
 
-	const FGuid NodeGUID = ParameterNode->NodeGuid;
-	if (const UE::Mutable::Private::Ptr<UE::Mutable::Private::NodeImageParameter>* GeneratedImageParameter = GenerationContext.GeneratedImageParameters.Find(NodeGUID))
+	const FGeneratedParameterKey Key = { ParameterNode->NodeGuid, ParameterNode->GetParameterName(&GenerationContext.MacroNodesStack) };
+	if (const UE::Mutable::Private::Ptr<UE::Mutable::Private::NodeImageParameter>* GeneratedImageParameter = GenerationContext.GeneratedImageParameters.Find(Key))
 	{
 		return GeneratedImageParameter->get();
 	}
@@ -1330,7 +1330,7 @@ UE::Mutable::Private::Ptr<UE::Mutable::Private::NodeImageParameter> GenerateMuta
 		Result->SetMessageContext(ParameterNode);
 	}
 	
-	GenerationContext.GeneratedImageParameters.Add(NodeGUID,Result);
+	GenerationContext.GeneratedImageParameters.Add(Key, Result);
 	GenerationContext.GeneratedNodes.Add(ParameterNode);
 	
 	return Result;
diff --git a/Engine/Plugins/Mutable/Source/CustomizableObjectEditor/Private/MuCOE/GenerateMutableSource/GenerateMutableSourceMaterial.cpp b/Engine/Plugins/Mutable/Source/CustomizableObjectEditor/Private/MuCOE/GenerateMutableSource/GenerateMutableSourceMaterial.cpp
--- a/Engine/Plugins/Mutable/Source/CustomizableObjectEditor/Private/MuCOE/GenerateMutableSource/GenerateMutableSourceMaterial.cpp
+++ b/Engine/Plugins/Mutable/Source/CustomizableObjectEditor/Private/MuCOE/GenerateMutableSource/GenerateMutableSourceMaterial.cpp
@@ -70,8 +70,8 @@ UE::Mutable::Private::Ptr<UE::Mutable::Private::NodeMaterialParameter> GenerateM
 
 	UE::Mutable::Private::Ptr<UE::Mutable::Private::NodeMaterialParameter> Result;
 
-	const FGuid NodeGUID = ParameterNode->NodeGuid;
-	if (const UE::Mutable::Private::Ptr<UE::Mutable::Private::NodeMaterialParameter>* GeneratedMaterialParameter = GenerationContext.GeneratedMaterialParameters.Find(NodeGUID))
+	const FGeneratedParameterKey Key = { ParameterNode->NodeGuid, ParameterNode->GetParameterName(&GenerationContext.MacroNodesStack) };
+	if (const UE::Mutable::Private::Ptr<UE::Mutable::Private::NodeMaterialParameter>* GeneratedMaterialParameter = GenerationContext.GeneratedMaterialParameters.Find(Key))
 	{
 		return GeneratedMaterialParameter->get();
 	}
@@ -82,7 +82,7 @@ UE::Mutable::Private::Ptr<UE::Mutable::Private::NodeMaterialParameter> GenerateM
 		MaterialNode->Name = MaterialParameterNode->GetParameterName(&GenerationContext.MacroNodesStack);
 		MaterialNode->UID = GenerationContext.GetNodeIdUnique(ParameterNode).ToString();
 		
-		GenerationContext.MaterialParameterDefaultValues.Add(MaterialParameterNode->GetParameterName(), UE::Mutable::Private::LoadObject(MaterialParameterNode->DefaultValue));
+		GenerationContext.MaterialParameterDefaultValues.Add(MaterialParameterNode->GetParameterName(&GenerationContext.MacroNodesStack), UE::Mutable::Private::LoadObject(MaterialParameterNode->DefaultValue));
 		GenerationContext.ParameterUIDataMap.Add(MaterialParameterNode->GetParameterName(&GenerationContext.MacroNodesStack), FMutableParameterData(
 			MaterialParameterNode->ParamUIMetadata,
 			EMutableParameterType::Material));
@@ -99,7 +99,7 @@ UE::Mutable::Private::Ptr<UE::Mutable::Private::NodeMaterialParameter> GenerateM
 		Result->SetMessageContext(ParameterNode);
 	}
 	
-	GenerationContext.GeneratedMaterialParameters.Add(NodeGUID, Result);
+	GenerationContext.GeneratedMaterialParameters.Add(Key, Result);
 	GenerationContext.GeneratedNodes.Add(ParameterNode);
 	
 	return Result;
diff --git a/Engine/Plugins/Mutable/Source/CustomizableObjectEditor/Private/MuCOE/GenerateMutableSource/GenerateMutableSourceSkeletalMesh.cpp b/Engine/Plugins/Mutable/Source/CustomizableObjectEditor/Private/MuCOE/GenerateMutableSource/GenerateMutableSourceSkeletalMesh.cpp
--- a/Engine/Plugins/Mutable/Source/CustomizableObjectEditor/Private/MuCOE/GenerateMutableSource/GenerateMutableSourceSkeletalMesh.cpp
+++ b/Engine/Plugins/Mutable/Source/CustomizableObjectEditor/Private/MuCOE/GenerateMutableSource/GenerateMutableSourceSkeletalMesh.cpp
@@ -27,8 +27,8 @@ namespace UE::Mutable::Private
 
 		Ptr<NodeSkeletalMeshParameter> Result;
 
-		const FGuid NodeGUID = ParameterNode->NodeGuid;
-		if (const Ptr<NodeSkeletalMeshParameter>* GeneratedSkeletalMeshParameter = GenerationContext.GeneratedSkeletalMeshParameters.Find(NodeGUID))
+		const FGeneratedParameterKey Key = { ParameterNode->NodeGuid, ParameterNode->GetParameterName(&GenerationContext.MacroNodesStack) };
+		if (const Ptr<NodeSkeletalMeshParameter>* GeneratedSkeletalMeshParameter = GenerationContext.GeneratedSkeletalMeshParameters.Find(Key))
 		{
 			return GeneratedSkeletalMeshParameter->get();
 		}
@@ -58,7 +58,7 @@ namespace UE::Mutable::Private
 			Result->SetMessageContext(ParameterNode);
 		}
 	
-		GenerationContext.GeneratedSkeletalMeshParameters.Add(NodeGUID,Result);
+		GenerationContext.GeneratedSkeletalMeshParameters.Add(Key, Result);
 		GenerationContext.GeneratedNodes.Add(ParameterNode);
 	
 		return Result;
