[Mutable] Fixed an incorrect copy of an ASTChild which caused a possible dangling pointer in States with uncompressed texture.
+ This fixes the crash in the internat test asset Banditt_forRPG.

#rnx

#Synced-CL 48900308

diff --git a/Engine/Plugins/Mutable/Source/MutableTools/Private/MuT/ASTOpMaterialModify.cpp b/Engine/Plugins/Mutable/Source/MutableTools/Private/MuT/ASTOpMaterialModify.cpp
--- a/Engine/Plugins/Mutable/Source/MutableTools/Private/MuT/ASTOpMaterialModify.cpp
+++ b/Engine/Plugins/Mutable/Source/MutableTools/Private/MuT/ASTOpMaterialModify.cpp
@@ -61,7 +61,8 @@ namespace UE::Mutable::Private
 		Result->ParametersToModify.Reserve(ParametersToModify.Num());
 		for (const TPair<FParameterKey, ASTChild>& Parameter : ParametersToModify)
 		{
-			Result->ParametersToModify.Add(Parameter.Key, ASTChild(Result, MapChild(Parameter.Value.child())));
+			ASTChild& ClonedParam = Result->ParametersToModify.Emplace(Parameter.Key, Result.get());
+			ClonedParam = MapChild(Parameter.Value.child());
 		}
 
 		Result->ImagePropertyIndexMap = ImagePropertyIndexMap;
diff --git a/Engine/Plugins/Mutable/Source/MutableTools/Private/MuT/CodeOptimiser_Parameters.cpp b/Engine/Plugins/Mutable/Source/MutableTools/Private/MuT/CodeOptimiser_Parameters.cpp
--- a/Engine/Plugins/Mutable/Source/MutableTools/Private/MuT/CodeOptimiser_Parameters.cpp
+++ b/Engine/Plugins/Mutable/Source/MutableTools/Private/MuT/CodeOptimiser_Parameters.cpp
@@ -1910,7 +1910,7 @@ namespace UE::Mutable::Private
 								fop->FormatIfAlpha = uncompressedFormat;
 								fop->Source = ImageOp;
 
-								Parameter.Value = ASTChild(NewOp, fop);
+								Parameter.Value = fop;
 							}
 						}
 					}
