[Mutable] Fix crash when visualizing layout blocks due to a desynchronization between the MaxGridSize and the current GridSize.

#rnx
#Synced-CL 48901557

diff --git a/Engine/Plugins/Mutable/Source/CustomizableObjectEditor/Private/MuCOE/SCustomizableObjectLayoutEditor.cpp b/Engine/Plugins/Mutable/Source/CustomizableObjectEditor/Private/MuCOE/SCustomizableObjectLayoutEditor.cpp
--- a/Engine/Plugins/Mutable/Source/CustomizableObjectEditor/Private/MuCOE/SCustomizableObjectLayoutEditor.cpp
+++ b/Engine/Plugins/Mutable/Source/CustomizableObjectEditor/Private/MuCOE/SCustomizableObjectLayoutEditor.cpp
@@ -1490,6 +1490,25 @@ void SCustomizableObjectLayoutEditor::FillLayoutComboBoxOptions()
 {
 	if (CurrentLayout)
 	{
+		// Update LayoutGridSizes
+		{
+			const int32 MaxGridSize = CurrentLayout->GetMaxGridSize().GetMax();
+			LayoutGridSizes.Empty();
+
+			for (int32 Size = 1; Size <= MaxGridSize; Size *= 2)
+			{
+				LayoutGridSizes.Add(MakeShareable(new FString(FString::FromInt(Size))));
+			}
+
+			FInt32Point GridSize = CurrentLayout->GetGridSize();
+			if (GridSize.X > MaxGridSize || GridSize.Y > MaxGridSize)
+			{
+				GridSize.X = FMath::Clamp(GridSize.X, 1, MaxGridSize);
+				GridSize.Y = FMath::Clamp(GridSize.Y, 1, MaxGridSize);
+				CurrentLayout->SetGridSize(GridSize);
+			}
+		}
+
 		if (StrategyComboBox)
 		{
 			const uint32 StrategyIndex = (uint32)CurrentLayout->PackingStrategy;
