[Mutable] Prevent crash in ConvertSkeletalMeshToMutable during CO compilation when a SkeletalMesh is missing.

#rnx
#14108
#Synced-CL 48901557

diff --git a/Engine/Plugins/Mutable/Source/CustomizableObjectEditor/Private/MuCOE/CustomizableObjectCompileRunnable.cpp b/Engine/Plugins/Mutable/Source/CustomizableObjectEditor/Private/MuCOE/CustomizableObjectCompileRunnable.cpp
--- a/Engine/Plugins/Mutable/Source/CustomizableObjectEditor/Private/MuCOE/CustomizableObjectCompileRunnable.cpp
+++ b/Engine/Plugins/Mutable/Source/CustomizableObjectEditor/Private/MuCOE/CustomizableObjectCompileRunnable.cpp
@@ -282,7 +282,7 @@ uint32 FCustomizableObjectCompileRunnable::Run()
 						}
 						else
 						{
-							Streamable->RequestAsyncLoad(AssetsToStream, OnAssetLoadCallback, true);
+							Streamable->RequestAsyncLoad(AssetsToStream, OnAssetLoadCallback, FStreamableManager::AsyncLoadHighPriority - 1);
 						}
 					}
 				};
diff --git a/Engine/Plugins/Mutable/Source/CustomizableObjectEditor/Private/MuCOE/GenerateMutableSource/GenerateMutableSourceMesh.cpp b/Engine/Plugins/Mutable/Source/CustomizableObjectEditor/Private/MuCOE/GenerateMutableSource/GenerateMutableSourceMesh.cpp
--- a/Engine/Plugins/Mutable/Source/CustomizableObjectEditor/Private/MuCOE/GenerateMutableSource/GenerateMutableSourceMesh.cpp
+++ b/Engine/Plugins/Mutable/Source/CustomizableObjectEditor/Private/MuCOE/GenerateMutableSource/GenerateMutableSourceMesh.cpp
@@ -1182,8 +1182,12 @@ UE::Tasks::FTask ConvertSkeletalMeshToMutable(TSharedRef<FMeshConversionContext>
 	const uint32 MeshId = CityHash32(reinterpret_cast<const char*>(*MeshName), MeshName.Len() * sizeof(FString::ElementType));
 
 	const USkeletalMesh* SkeletalMesh = Cast<USkeletalMesh>(Source.Mesh.Get());
-
-	checkf(SkeletalMesh, TEXT("SkeletalMesh cast failed for mesh '%s' (MeshId: %u)"), *MeshName, MeshId);
+	if (!SkeletalMesh)
+	{
+		const FString Msg = FString::Printf(TEXT("The SkeletalMesh [%s] failed to load."), *MeshName);
+		Context.Log(FText::FromString(Msg), Source.MessageContext);
+		return UE::Tasks::MakeCompletedTask<void>();
+	}
 
 	const FSkeletalMeshModel* ImportedModel = SkeletalMesh->GetImportedModel();
 	if (!ImportedModel)
