[Mutable] Allow taking the ownnership of generated resources.

- Automatically disable all caches when benchmarking.
- Reuse Textures is now done automatically (whenever possible). Still can be disabled by CVar.

#Synced-CL 49065323

diff --git a/Engine/Plugins/Mutable/Source/CustomizableObject/Internal/MuCO/CustomizableObjectInstancePrivate.h b/Engine/Plugins/Mutable/Source/CustomizableObject/Internal/MuCO/CustomizableObjectInstancePrivate.h
--- a/Engine/Plugins/Mutable/Source/CustomizableObject/Internal/MuCO/CustomizableObjectInstancePrivate.h
+++ b/Engine/Plugins/Mutable/Source/CustomizableObject/Internal/MuCO/CustomizableObjectInstancePrivate.h
@@ -226,17 +226,23 @@ class UCustomizableInstancePrivate : public UObject
 {
 public:
 	GENERATED_BODY()
+	
+	/** Generated Skeletal Meshes. */
+	UPROPERTY(Transient, VisibleAnywhere, Category = NoCategory)
+	TMap<FName, TObjectPtr<USkeletalMesh>> SkeletalMeshes; 	// TODO GMT Should be used for display only
 
-	/** The generated skeletal meshes for this Instance. They may be null if the component is empty. */
+	/** Generated Materials. */
 	UPROPERTY(Transient, VisibleAnywhere, Category = NoCategory)
-	TMap<FName, TObjectPtr<USkeletalMesh>> SkeletalMeshes;
-	
-	UPROPERTY(Transient, VisibleAnywhere, Category = NoCategory, meta=(DisplayName = "Materials"))
-	TArray<TObjectPtr<UMaterialInterface>> GeneratedMaterials;
+	TArray<TObjectPtr<UMaterialInterface>> Materials; // TODO GMT Should be used for display only
 
-	UPROPERTY(Transient, VisibleAnywhere, Category = NoCategory, meta=(DisplayName = "Textures"))
-	TArray<TObjectPtr<UTexture>> GeneratedTextures;
+	/** Generated Textures. */
+	UPROPERTY(Transient, VisibleAnywhere, Category = NoCategory)
+	TArray<TObjectPtr<UTexture>> Textures; // TODO GMT Should be used for display only
 
+	/** If check, Mutable will keep the ownership of the generated resources. This means that Mutable may decide to modify or reuse them for performance considerations. */
+	UPROPERTY(EditAnywhere, Category = UpdateOptions)
+	bool bKeepOwnershipOfGeneratedResources = true;
+	
 	// Indices of the parameters that are relevant for the given parameter values.
 	// This only gets updated if parameter decorations are generated.
 	TArray<int> RelevantParameters;
@@ -385,8 +391,6 @@ public:
 	/** Automatic update required.
 	 * Set to true when a Customizable Object Instance Usage requires an automatic update (e.g., component reattached). */
 	bool bAutomaticUpdateRequired = false;
-
-	bool bReuseInstanceTextures = false;
 	
 	// TEMP VARIABLE to check the Min desired LODs for this instance
 	TWeakObjectPtr<UCustomizableObjectInstanceUsage> NearestToActor;
diff --git a/Engine/Plugins/Mutable/Source/CustomizableObject/Internal/MuCO/CustomizableObjectPrivate.h b/Engine/Plugins/Mutable/Source/CustomizableObject/Internal/MuCO/CustomizableObjectPrivate.h
--- a/Engine/Plugins/Mutable/Source/CustomizableObject/Internal/MuCO/CustomizableObjectPrivate.h
+++ b/Engine/Plugins/Mutable/Source/CustomizableObject/Internal/MuCO/CustomizableObjectPrivate.h
@@ -897,10 +897,7 @@ struct FMutableStateData
 	/** If this is enabled, texture streaming won't be used for this state, and full images will be generated when an instance is first updated. */
 	UPROPERTY()
 	bool bDisableTextureStreaming = false;
-
-	UPROPERTY()
-	bool bReuseInstanceTextures = false;
-
+	
 	UPROPERTY()
 	TMap<FString, FString> ForcedParameterValues;
 
@@ -1769,6 +1766,8 @@ enum class ECustomizableObjectVersions : int32
 	
 	DeduplicateFMaterialsSharingPassthroughObject,
 
+	RemovedReuseTextures,
+
 	// -----<new versions can be added above this line>--------
 	LastCustomizableObjectVersion
 };
diff --git a/Engine/Plugins/Mutable/Source/CustomizableObject/Internal/MuCO/CustomizableObjectSystemPrivate.h b/Engine/Plugins/Mutable/Source/CustomizableObject/Internal/MuCO/CustomizableObjectSystemPrivate.h
--- a/Engine/Plugins/Mutable/Source/CustomizableObject/Internal/MuCO/CustomizableObjectSystemPrivate.h
+++ b/Engine/Plugins/Mutable/Source/CustomizableObject/Internal/MuCO/CustomizableObjectSystemPrivate.h
@@ -579,10 +579,12 @@ public:
 	/** When this option is enabled it will reuse the Mutable core instance and its temp data between updates.  */
 	bool bLiveUpdateMode = false;
 
-	bool bReuseInstanceTextures = false;
+	bool bUseResueTextureCache = false;
 	
 	bool bUseMeshCache = false;
 	
+	bool bUseCaches = true;
+	
 	/** Whether the mesh to generate should support Mesh LOD streaming or not. */
 	bool bStreamMeshLODs = false;
 
@@ -605,6 +607,8 @@ public:
 	bool bIsOnlyGenerateRequestedLODsEnabled = false;
 	
 	bool bIsUpdateAborted = false;
+		
+	bool bKeepOwnershipOfGeneratedResources = false;
 
 	/** This option comes from the operation request. It is used to reduce the number of mipmaps that mutable must generate for images.  */
 	int32 MipsToSkip = 0;
diff --git a/Engine/Plugins/Mutable/Source/CustomizableObject/Private/MuCO/CustomizableObject.cpp b/Engine/Plugins/Mutable/Source/CustomizableObject/Private/MuCO/CustomizableObject.cpp
--- a/Engine/Plugins/Mutable/Source/CustomizableObject/Private/MuCO/CustomizableObject.cpp
+++ b/Engine/Plugins/Mutable/Source/CustomizableObject/Private/MuCO/CustomizableObject.cpp
@@ -2125,7 +2125,6 @@ FArchive& operator<<(FArchive& Ar, FMutableStateData& Struct)
 	Ar << Struct.bLiveUpdateMode;
 	Ar << Struct.bDisableMeshStreaming;
 	Ar << Struct.bDisableTextureStreaming;
-	Ar << Struct.bReuseInstanceTextures;
 	Ar << Struct.ForcedParameterValues;
 
 	return Ar;
diff --git a/Engine/Plugins/Mutable/Source/CustomizableObject/Private/MuCO/CustomizableObjectInstance.cpp b/Engine/Plugins/Mutable/Source/CustomizableObject/Private/MuCO/CustomizableObjectInstance.cpp
--- a/Engine/Plugins/Mutable/Source/CustomizableObject/Private/MuCO/CustomizableObjectInstance.cpp
+++ b/Engine/Plugins/Mutable/Source/CustomizableObject/Private/MuCO/CustomizableObjectInstance.cpp
@@ -112,8 +112,11 @@ UTexture2D* CreateTexture(const FName& TextureName)
 void UCustomizableInstancePrivate::InvalidateGeneratedData()
 {
 	SkeletalMeshStatus = ESkeletalMeshStatus::NotGenerated;
-	SkeletalMeshes.Reset();
 
+	SkeletalMeshes.Empty();
+	Materials.Empty();
+	Textures.Empty();
+	
 	CommittedDescriptor = {};
 	CommittedDescriptorHash = {};
 
@@ -121,8 +124,6 @@ void UCustomizableInstancePrivate::InvalidateGeneratedData()
 	FCustomizableInstanceComponentData TemplateComponentData;
 	TemplateComponentData.LastComponentDescriptor.MeshDescriptors.Init({}, MAX_MESH_LOD_COUNT);
 	ComponentsData.Init(TemplateComponentData, ComponentsData.Num());
-
-	GeneratedMaterials.Empty();
 }
 
 
@@ -279,9 +280,6 @@ void UCustomizableInstancePrivate::PostDuplicate(bool bDuplicateForPIE)
 
 	// Invalidate all generated data to avoid modifying resources shared between CO instances.
 	InvalidateGeneratedData();
-
-	// Empty after duplicating or ReleasingMutableResources may free textures used by the other CO instance.
-	GeneratedTextures.Empty();
 }
 
 
@@ -5568,7 +5566,7 @@ FName TextureName(const FInstanceUpdateData::FImage& Image, bool bBake)
 	return MakeUniqueObjectName(GetTransientPackage(), UTexture2D::StaticClass(), *MutableTextureName);
 }
 
-UTexture* BuildTexture(const TSharedRef<FUpdateContextPrivate>& Context, const int32 ImageInex, const bool bInReuseTexture)
+UTexture* BuildTexture(const TSharedRef<FUpdateContextPrivate>& Context, const int32 ImageInex, const bool bCanReuseTexture)
 {
 	const TSharedRef<const FInstanceUpdateData::FImage>& Image = Context->InstanceUpdateData.Images[ImageInex];
 	
@@ -5584,8 +5582,8 @@ UTexture* BuildTexture(const TSharedRef<FUpdateContextPrivate>& Context, const i
 	{
 		UTexture* Texture = Image->Image->PassthroughObject.Get(); // TODO GMT Why not also cache Passthrough?
 
-		Context->Instance->GetPrivate()->GeneratedTextures.AddUnique(Texture);
-
+		Context->Instance->GetPrivate()->Textures.AddUnique(Texture);
+		
 		return Texture;
 	}
 
@@ -5602,9 +5600,8 @@ UTexture* BuildTexture(const TSharedRef<FUpdateContextPrivate>& Context, const i
 	
 	UE::Mutable::Private::TManagedPtr<const UE::Mutable::Private::FImage> MutableImage = Image->Image;
 
-	UTexture2D* Texture; // Texture generated by mutable
+	UTexture2D* Texture = nullptr; // Texture generated by mutable
 
-	const bool bReuseTexture = bInReuseTexture && Context->bReuseInstanceTextures;
 	TMap<FTextureReuseCacheKey, TStrongObjectPtr<UTexture2D>>& LastTextureReuseCache = Context->Instance->GetPrivate()->LastTextureReuseCache; 
 	TMap<FTextureReuseCacheKey, TStrongObjectPtr<UTexture2D>>& CurrentTextureReuseCache = Context->Instance->GetPrivate()->CurrentTextureReuseCache;
 	
@@ -5613,8 +5610,11 @@ UTexture* BuildTexture(const TSharedRef<FUpdateContextPrivate>& Context, const i
 	// If the mutable image is null, it must be in the cache
 	if (!MutableImage)
 	{
-		Texture = TextureCache.Get(TextureCacheKey);
-		check(Texture);
+		if (Context->bUseCaches)
+		{
+			Texture = TextureCache.Get(TextureCacheKey);
+			check(Texture);
+		}
 	}
 	else
 	{
@@ -5625,7 +5625,7 @@ UTexture* BuildTexture(const TSharedRef<FUpdateContextPrivate>& Context, const i
     	const FMutableModelImageProperties& Props = ModelResources.ImageProperties[ImageKey];
     	
 		UTexture2D* ReusedTexture = nullptr;
-		if (bReuseTexture)
+		if (Context->bUseCaches &&  Context->bUseResueTextureCache && bCanReuseTexture)
 		{
 			if (TStrongObjectPtr<UTexture2D>* ReuseTexture = LastTextureReuseCache.Find(TextureReuseCacheKey))
 			{
@@ -5735,9 +5735,12 @@ UTexture* BuildTexture(const TSharedRef<FUpdateContextPrivate>& Context, const i
 
 	if (Texture)
 	{
-		Context->Instance->GetPrivate()->GeneratedTextures.Add(Texture);
-
-		if (bReuseTexture)
+		Context->Instance->GetPrivate()->Textures.Add(Texture);
+	}
+		
+	if (Context->bUseCaches)
+	{
+		if (Context->bUseResueTextureCache && bCanReuseTexture)
 		{
 			CurrentTextureReuseCache.Add(TextureReuseCacheKey, TStrongObjectPtr(Texture));
 		}
@@ -5760,7 +5763,7 @@ struct FBuildMaterialOut
 };
 
 
-FBuildMaterialOut BuildMaterial(const TSharedRef<FUpdateContextPrivate>& Context, const int32 MaterialIndex, const bool bReuseMaterial)
+FBuildMaterialOut BuildMaterial(const TSharedRef<FUpdateContextPrivate>& Context, const int32 MaterialIndex, const bool bCanReuseMaterial)
 {
 	FBuildMaterialOut Out;
 
@@ -5782,8 +5785,6 @@ FBuildMaterialOut BuildMaterial(const TSharedRef<FUpdateContextPrivate>& Context
 	
 	if (!Material->Material->PassthroughObject.Get())
 	{
-		Out.Material = UMaterial::GetDefaultMaterial(MD_Surface);
-		
 		return Out;
 	}
 
@@ -5796,7 +5797,7 @@ FBuildMaterialOut BuildMaterial(const TSharedRef<FUpdateContextPrivate>& Context
 
 	UMaterialInterface* MaterialInterface = nullptr;
 	
-	if (bReuseMaterial)
+	if (Context->bUseCaches && bCanReuseMaterial)
 	{
 		if (const TStrongObjectPtr<UMaterialInterface>* ReuseMaterial = LastMaterialReuseCache.Find(MaterialReuseCacheKey))
 		{
@@ -5898,7 +5899,7 @@ FBuildMaterialOut BuildMaterial(const TSharedRef<FUpdateContextPrivate>& Context
 		{
 			const TSharedRef<const FInstanceUpdateData::FImage>& Image = Context->InstanceUpdateData.Images[Material->FirstImage + ImageIndex];
 				
-			UTexture* Texture = BuildTexture(Context, Material->FirstImage + ImageIndex, bReuseMaterial);
+			UTexture* Texture = BuildTexture(Context, Material->FirstImage + ImageIndex, bCanReuseMaterial);
 				
 			if (Image->MaterialLayer < 0)
 			{
@@ -5919,14 +5920,21 @@ FBuildMaterialOut BuildMaterial(const TSharedRef<FUpdateContextPrivate>& Context
 
 	if (MaterialInterface)
 	{
-		Context->Instance->GetPrivate()->GeneratedMaterials.AddUnique(MaterialInterface);
-
-		if (bReuseMaterial)
+		Context->Instance->GetPrivate()->Materials.AddUnique(MaterialInterface);
+	}
+		
+	if (Context->bUseCaches)
+	{
+		if (bCanReuseMaterial)
 		{
-			CurrentMaterialReuseCache.Add(MaterialReuseCacheKey, TStrongObjectPtr(MaterialInterface)); 			
+			CurrentMaterialReuseCache.Add(MaterialReuseCacheKey, TStrongObjectPtr(MaterialInterface));
+		}
+		else
+		{
+			// Add it to the CO cache. Currently we do not have one for Materials.
 		}
 	}
-	
+
 	Context->MaterialUpdateCache.Add(MaterialHash, TStrongObjectPtr(MaterialInterface));
 
 	Out.Material = MaterialInterface;
@@ -6046,7 +6054,7 @@ void UCustomizableInstancePrivate::BuildMaterials(const TSharedRef<FUpdateContex
 			SkeletalMesh->SetOverlayMaterial(OverlayMaterial);
 		}
 
-		TArray<FSkeletalMaterial> Materials;
+		TArray<FSkeletalMaterial> SkeletalMaterials;
 		
 		// SurfaceId per MaterialSlotIndex
 		TArray<int32> SurfaceIdToMaterialIndex;
@@ -6093,8 +6101,8 @@ void UCustomizableInstancePrivate::BuildMaterials(const TSharedRef<FUpdateContex
 				SurfaceIdToMaterialIndex.Add(Surface->SurfaceId);
 
 				// Add and set up the material data for this slot
-				const int32 MaterialSlotIndex = Materials.Num();
-				FSkeletalMaterial& MaterialSlot = Materials.AddDefaulted_GetRef();
+				const int32 MaterialSlotIndex = SkeletalMaterials.Num();
+				FSkeletalMaterial& MaterialSlot = SkeletalMaterials.AddDefaulted_GetRef();
 
 				uint32 UsedSurfaceMetadataId = Surface->SurfaceMetadataId;
 				
@@ -6201,9 +6209,9 @@ void UCustomizableInstancePrivate::BuildMaterials(const TSharedRef<FUpdateContex
 			}
 
 			// Force recreate render state after replacing the materials to avoid a crash in the render pipeline if the old materials are GCed while in use.
-			RecreateRenderStateOnInstanceComponent[InstanceComponentIndex] |= SkeletalMesh->GetResourceForRendering()->IsInitialized() && SkeletalMesh->GetMaterials() != Materials;
+			RecreateRenderStateOnInstanceComponent[InstanceComponentIndex] |= SkeletalMesh->GetResourceForRendering()->IsInitialized() && SkeletalMesh->GetMaterials() != SkeletalMaterials;
 
-			SkeletalMesh->SetMaterials(Materials);
+			SkeletalMesh->SetMaterials(SkeletalMaterials);
 
 #if WITH_EDITOR
 			if (GEditor && RecreateRenderStateOnInstanceComponent[InstanceComponentIndex])
@@ -6244,7 +6252,7 @@ void UCustomizableInstancePrivate::BuildMaterials(const TSharedRef<FUpdateContex
 			const int32 ObjectComponentIndex = ModelResources.ComponentNamesPerObjectComponent.IndexOfByKey(ComponentName);
 			
 			int32 InstanceComponentIndex = -1;
-			for (int32 CurrentInstanceIndex=0; CurrentInstanceIndex<Context->InstanceUpdateData.Components.Num(); ++CurrentInstanceIndex)
+			for (int32 CurrentInstanceIndex = 0; CurrentInstanceIndex<Context->InstanceUpdateData.Components.Num(); ++CurrentInstanceIndex)
 			{
 				if (Context->InstanceUpdateData.Components[CurrentInstanceIndex]->Id.GetValue() == ObjectComponentIndex)
 				{
@@ -6294,9 +6302,9 @@ void UCustomizableObjectInstance::SetReplacePhysicsAssets(bool bReplaceEnabled)
 }
 
 
-void UCustomizableObjectInstance::SetReuseInstanceTextures(bool bTextureReuseEnabled)
+void UCustomizableObjectInstance::SetKeepOwnershipOfGeneratedResources(bool bEnabled)
 {
-	GetPrivate()->bReuseInstanceTextures = bTextureReuseEnabled;
+	GetPrivate()->bKeepOwnershipOfGeneratedResources = bEnabled;
 }
 
 
diff --git a/Engine/Plugins/Mutable/Source/CustomizableObject/Private/MuCO/CustomizableObjectSystem.cpp b/Engine/Plugins/Mutable/Source/CustomizableObject/Private/MuCO/CustomizableObjectSystem.cpp
--- a/Engine/Plugins/Mutable/Source/CustomizableObject/Private/MuCO/CustomizableObjectSystem.cpp
+++ b/Engine/Plugins/Mutable/Source/CustomizableObject/Private/MuCO/CustomizableObjectSystem.cpp
@@ -274,7 +274,7 @@ FUpdateContextPrivate::FUpdateContextPrivate(UCustomizableObjectInstance& InInst
 	PixelFormatOverride = SystemPrivate->ImageFormatOverrideFunc;
 #endif
 
-	bReuseInstanceTextures = InInstance.GetPrivate()->bReuseInstanceTextures;
+	bKeepOwnershipOfGeneratedResources = InInstance.GetPrivate()->bKeepOwnershipOfGeneratedResources;
 	
 	bValid = true;
 }
@@ -1372,8 +1372,7 @@ void UpdateSkeletalMesh(const TSharedRef<FUpdateContextPrivate>& Context)
 		TArray<TObjectPtr<UMaterialInterface>> NewOverridenMaterials;
 
 		const bool bIsTransientMesh = SkeletalMesh ? SkeletalMesh->HasAllFlags(EObjectFlags::RF_Transient) : false;
-		const bool bIsUsingMeshCache = CustomizableObject->bEnableMeshCache && UCustomizableObjectSystem::IsMeshCacheEnabled();
-		const bool bUseOverrideMaterials = !bIsTransientMesh || bIsUsingMeshCache;
+		const bool bUseOverrideMaterials = !bIsTransientMesh || Context->bUseMeshCache;
 		if (bUseOverrideMaterials)
 		{
 			if (FCustomizableInstanceComponentData* ComponentData = CustomizableObjectInstance->GetPrivate()->GetComponentData(Usage->GetComponentName()))
@@ -1874,7 +1873,7 @@ bool UCustomizableObjectSystem::CheckIfDiskOrMipUpdateOperationsPending(const UC
 	{
 		if (IsValid(*CustomizableObjectInstance) && CustomizableObjectInstance->GetCustomizableObject() == &Object)
 		{
-			for (const TObjectPtr<UTexture>& GeneratedTexture : CustomizableObjectInstance->GetPrivate()->GeneratedTextures)
+			for (const TObjectPtr<UTexture>& GeneratedTexture : CustomizableObjectInstance->GetPrivate()->Textures) // TODO GMT This approach is incorrect. An update may clear the generated Textures array but these may live longer.
 			{
 				if (GeneratedTexture->HasPendingInitOrStreaming())
 				{
@@ -2595,7 +2594,7 @@ namespace impl
 	{
 		MUTABLE_CPUPROFILER_SCOPE(Task_Mutable_GetImages)
 		
-		const TSharedPtr<TArray<UE::Mutable::Private::FImageId>> ImagesInThisInstance = MakeShared<TArray<UE::Mutable::Private::FImageId>>();
+		const TSharedPtr<TArray<UE::Mutable::Private::FImageId>> RequestedGetImages = MakeShared<TArray<UE::Mutable::Private::FImageId>>();
 
 		UE::Tasks::FTask LastTask;
 		for (int32 ImageIndex = 0; ImageIndex < OperationData->InstanceUpdateData.Images.Num(); ++ImageIndex)
@@ -2618,8 +2617,8 @@ namespace impl
 			}
 			else
 			{
-				const bool bCached = ImagesInThisInstance->Contains(Image->ImageID) || // See if it is cached from this same instance (can happen with LODs)
-					(UCustomizableObjectSystem::ShouldReuseTexturesBetweenInstances() && OperationData->CachedTextures.Contains(Image->ImageID)); // See if it is cached from another instance
+				const bool bCached = RequestedGetImages->Contains(Image->ImageID) || // See if it is cached from this same instance (can happen with LODs)
+					OperationData->CachedTextures.Contains(Image->ImageID); // See if it is cached from another instance
 			
 				if (bCached)
 				{
@@ -2627,7 +2626,7 @@ namespace impl
 					continue;
 				}
 
-				ImagesInThisInstance->Add(Image->ImageID);
+				RequestedGetImages->Add(Image->ImageID);
 			
 				const int32 MaxSize = FMath::Max(Image->FullImageSizeX, Image->FullImageSizeY);
 				const int32 FullLODCount = FMath::CeilLogTwo(MaxSize) + 1;
@@ -2991,16 +2990,13 @@ namespace impl
 		// If the program aborted, AreAnyLoadsPending may return true, as the program doesn't cancel
 		// its own loads on exit (maybe it should?)
 		
+		if (OperationData->bUseCaches)
 		{
 			FTextureCache& TextureCache = CustomizableObject->GetPrivate()->TextureCache;
 		
 			for (const TSharedRef<FInstanceUpdateData::FImage>& Image : OperationData->InstanceUpdateData.Images)
 			{
-#if	WITH_EDITOR
 				const FTextureCache::FId TextureCacheKey = FTextureCache::FId(Image->ImageID, OperationData->MipsToSkip, OperationData->bBake);
-#else
-				const FTextureCache::FId TextureCacheKey = FTextureCache::FId(Image->ImageID, OperationData->MipsToSkip);
-#endif
 				if (const UTexture2D* Texture = TextureCache.Get(TextureCacheKey))
 				{
 					OperationData->Objects.Add(TStrongObjectPtr(Texture));
@@ -3209,13 +3205,9 @@ namespace impl
 			Operation->bLiveUpdateMode = false;
 		}
 		
-		Operation->bReuseInstanceTextures = !Operation->bBake && CVarEnableMutableReuseInstanceTextures.GetValueOnGameThread() && (Operation->bReuseInstanceTextures || (StateData ? StateData->bReuseInstanceTextures : false));
-
-		if (Operation->bReuseInstanceTextures && !Operation->bNeverStream)
-		{
-			UE_LOG(LogMutable, Warning, TEXT("Instance texture reuse requires that the current Mutable state is in non-streaming mode. Change it in the Mutable graph base node in the state definition."));
-			Operation->bReuseInstanceTextures = false;
-		}
+		Operation->bUseCaches = Operation->bKeepOwnershipOfGeneratedResources && !Operation->bBake && !FLogBenchmarkUtil::IsBenchmarkingReportingEnabled();
+		Operation->bUseMeshCache = Operation->bUseCaches && Object->bEnableMeshCache && !Operation->bLiveUpdateMode && CVarEnableMeshCache.GetValueOnAnyThread(true);
+		Operation->bUseResueTextureCache = Operation->bUseCaches && Operation->bNeverStream && CVarEnableMutableReuseInstanceTextures.GetValueOnGameThread();
 		
 		if (!Operation->bLiveUpdateMode && InstancePrivate->LiveUpdateModeInstanceID != 0)
 		{
@@ -3274,16 +3266,12 @@ namespace impl
 		// Task: Mutable Update and GetMesh
 		//-------------------------------------------------------------
 		Operation->InstanceID = Operation->bLiveUpdateMode ? InstancePrivate->LiveUpdateModeInstanceID : 0;
-		Operation->bUseMeshCache = Object->bEnableMeshCache && !Operation->bBake && !Operation->bLiveUpdateMode && UCustomizableObjectSystem::IsMeshCacheEnabled(true);
 		
 		Operation->bStreamMeshLODs = IsStreamingEnabled(*Object, State);
 
 #if WITH_EDITOR
 		Operation->PixelFormatOverride = SystemPrivate->ImageFormatOverrideFunc;
 #endif
-
-		InstancePrivate->GeneratedMaterials.Empty();
-		InstancePrivate->GeneratedTextures.Empty();
 		
 		bool bRequestAllLODs = !Operation->bIsOnlyGenerateRequestedLODsEnabled ||
 			!SystemPrivate->CurrentInstanceLODManagement->IsOnlyGenerateRequestedLODLevelsEnabled();
@@ -4183,24 +4171,12 @@ void UCustomizableObjectSystem::EndBenchmark()
 }
 
 
-bool UCustomizableObjectSystem::IsMeshCacheEnabled(bool bCheckCVarOnGameThread /** = false */)
-{
-	return UCustomizableObjectSystemPrivate::IsUsingBenchmarkingSettings() ? false : CVarEnableMeshCache.GetValueOnAnyThread(bCheckCVarOnGameThread);
-}
-
-
 bool UCustomizableObjectSystem::ShouldClearWorkingMemoryOnUpdateEnd()
 {
 	return UCustomizableObjectSystemPrivate::IsUsingBenchmarkingSettings() ? true : CVarClearWorkingMemoryOnUpdateEnd.GetValueOnAnyThread();
 }
 
 
-bool UCustomizableObjectSystem::ShouldReuseTexturesBetweenInstances()
-{
-	return UCustomizableObjectSystemPrivate::IsUsingBenchmarkingSettings() ? false : CVarReuseImagesBetweenInstances.GetValueOnAnyThread();
-}
-
-
 void UCustomizableObjectSystem::SetWorkingMemory(int32 KBytes)
 {
 	CVarWorkingMemoryKB->Set( KBytes );
diff --git a/Engine/Plugins/Mutable/Source/CustomizableObject/Public/MuCO/CustomizableObjectInstance.h b/Engine/Plugins/Mutable/Source/CustomizableObject/Public/MuCO/CustomizableObjectInstance.h
--- a/Engine/Plugins/Mutable/Source/CustomizableObject/Public/MuCO/CustomizableObjectInstance.h
+++ b/Engine/Plugins/Mutable/Source/CustomizableObject/Public/MuCO/CustomizableObjectInstance.h
@@ -790,10 +790,9 @@ public:
 	UFUNCTION(BlueprintCallable, Category = CustomizableObjectInstance)
 	CUSTOMIZABLEOBJECT_API void SetReplacePhysicsAssets(bool bReplaceEnabled);
 
-	/** Enables the reuse of all possible textures when the instance is updated without any changes in geometry or state (the first update after creation doesn't reuse any)
-	  * It will only work if the textures aren't compressed, so set the instance to a Mutable state with texture compression disabled
-	  * WARNING! If texture reuse is enabled, do NOT keep external references to the textures of the instance. The instance owns the textures. */
-	CUSTOMIZABLEOBJECT_API void SetReuseInstanceTextures(bool bTextureReuseEnabled);
+	/** If true, Mutable will keep the ownership of the generated resources. This means that Mutable may decide to modify or reuse them for performance considerations. */
+	UFUNCTION(BlueprintCallable, Category = CustomizableObjectInstance)
+	CUSTOMIZABLEOBJECT_API void SetKeepOwnershipOfGeneratedResources(bool bEnabled);
 	
 	// The following methods should only be used in an LOD management class
 	CUSTOMIZABLEOBJECT_API void SetIsBeingUsedByComponentInPlay(bool bIsUsedByComponent);
diff --git a/Engine/Plugins/Mutable/Source/CustomizableObject/Public/MuCO/CustomizableObjectSystem.h b/Engine/Plugins/Mutable/Source/CustomizableObject/Public/MuCO/CustomizableObjectSystem.h
--- a/Engine/Plugins/Mutable/Source/CustomizableObject/Public/MuCO/CustomizableObjectSystem.h
+++ b/Engine/Plugins/Mutable/Source/CustomizableObject/Public/MuCO/CustomizableObjectSystem.h
@@ -261,25 +261,12 @@ public:
 	UFUNCTION(BlueprintCallable, BlueprintPure, Category = CustomizableObjectSystem)
 	UE_API int32 GetWorkingMemory() const;
 
-	/**
-	 * Get if the mutable mesh cache for the instance meshes is enabled or not.
-	 * @param bCheckCVarOnGameThread Force the checking of the relevant CVar to use the GameThread as target thread
-	 * @return True if it is enabled and false otherwise
-	 */
-	static UE_API bool IsMeshCacheEnabled(bool bCheckCVarOnGameThread = false);
-
 	/**
 	 * Get if mutable should clear its working memory between instance updates.
 	 * @return True if the clearing of memory will be performed, false otherwise.
 	 */
 	static UE_API bool ShouldClearWorkingMemoryOnUpdateEnd();
 
-	/**
-	 * Get if mutable will try to reuse textures in between instances.
-	 * @return True if mutable will try to reuse them, false otherwise.
-	 */
-	static UE_API bool ShouldReuseTexturesBetweenInstances();
-	
 private:
 	UPROPERTY(Transient)
 	TObjectPtr<UCustomizableObjectSystemPrivate> Private = nullptr;
diff --git a/Engine/Plugins/Mutable/Source/CustomizableObjectEditor/Private/MuCOE/GenerateMutableSource/GenerateMutableSource.cpp b/Engine/Plugins/Mutable/Source/CustomizableObjectEditor/Private/MuCOE/GenerateMutableSource/GenerateMutableSource.cpp
--- a/Engine/Plugins/Mutable/Source/CustomizableObjectEditor/Private/MuCOE/GenerateMutableSource/GenerateMutableSource.cpp
+++ b/Engine/Plugins/Mutable/Source/CustomizableObjectEditor/Private/MuCOE/GenerateMutableSource/GenerateMutableSource.cpp
@@ -1000,7 +1000,6 @@ UE::Mutable::Private::Ptr<UE::Mutable::Private::NodeObject> GenerateMutableSourc
 			StateUIData.bDisableMeshStreaming = State.bDisableMeshStreaming;
 			StateUIData.bDisableTextureStreaming = State.bDisableTextureStreaming;
 			StateUIData.bLiveUpdateMode = State.bLiveUpdateMode;
-			StateUIData.bReuseInstanceTextures = State.bReuseInstanceTextures;
 			StateUIData.ForcedParameterValues = State.ForcedParameterValues;
 
 			GenerationContext.StateUIDataMap.Add(State.Name, StateUIData);
diff --git a/Engine/Plugins/Mutable/Source/CustomizableObjectEditor/Private/MuCOE/Nodes/CustomizableObjectNodeObject.h b/Engine/Plugins/Mutable/Source/CustomizableObjectEditor/Private/MuCOE/Nodes/CustomizableObjectNodeObject.h
--- a/Engine/Plugins/Mutable/Source/CustomizableObjectEditor/Private/MuCOE/Nodes/CustomizableObjectNodeObject.h
+++ b/Engine/Plugins/Mutable/Source/CustomizableObjectEditor/Private/MuCOE/Nodes/CustomizableObjectNodeObject.h
@@ -51,12 +51,6 @@ struct FCustomizableObjectState
 	UPROPERTY(EditAnywhere, Category = CustomizableObject)
 	bool bLiveUpdateMode = false;
 
-	// Enables the reuse of all possible textures when the instance is updated without any changes in geometry or state (the first update after creation doesn't reuse any)
-	// It will only work if the textures aren't compressed, so set the instance to a Mutable state with texture compression disabled
-	// WARNING! If texture reuse is enabled, do NOT keep external references to the textures of the instance. The instance owns the textures.
-	UPROPERTY(EditAnywhere, Category = CustomizableObject)
-	bool bReuseInstanceTextures = false;
-
 	UPROPERTY(EditAnywhere, Category = CustomizableObject)
 	bool bBuildOnlyFirstLOD = false;
 
diff --git a/Engine/Plugins/Mutable/Source/CustomizableObjectEditor/Private/MuCOE/SCustomizableObjectEditorTextureAnalyzer.cpp b/Engine/Plugins/Mutable/Source/CustomizableObjectEditor/Private/MuCOE/SCustomizableObjectEditorTextureAnalyzer.cpp
--- a/Engine/Plugins/Mutable/Source/CustomizableObjectEditor/Private/MuCOE/SCustomizableObjectEditorTextureAnalyzer.cpp
+++ b/Engine/Plugins/Mutable/Source/CustomizableObjectEditor/Private/MuCOE/SCustomizableObjectEditorTextureAnalyzer.cpp
@@ -139,7 +139,7 @@ void SCustomizableObjecEditorTextureAnalyzer::FillTextureAnalyzerTable(UCustomiz
 		const UCustomizableObject* CustomizableObject = PreviewInstance->GetCustomizableObject();
 		check(CustomizableObject)
 
-		for (const TObjectPtr<UTexture>& GeneratedTexture : PreviewInstance->GetPrivate()->GeneratedTextures)
+		for (const TObjectPtr<UTexture>& GeneratedTexture : PreviewInstance->GetPrivate()->Textures)
 		{
 			UTexture2D* Texture = Cast<UTexture2D>(GeneratedTexture);
 			if (!Texture || !Texture->HasAllFlags(EObjectFlags::RF_Transient))
