[Mutable] Added PassthroughObjectFactory.

#Synced-CL 49065323

diff --git a/Engine/Plugins/Mutable/Source/CustomizableObjectEditor/Private/MuCOE/CustomizableObjectCompiler.cpp b/Engine/Plugins/Mutable/Source/CustomizableObjectEditor/Private/MuCOE/CustomizableObjectCompiler.cpp
--- a/Engine/Plugins/Mutable/Source/CustomizableObjectEditor/Private/MuCOE/CustomizableObjectCompiler.cpp
+++ b/Engine/Plugins/Mutable/Source/CustomizableObjectEditor/Private/MuCOE/CustomizableObjectCompiler.cpp
@@ -960,7 +960,6 @@ void FCustomizableObjectCompiler::CompileInternal(bool bAsync)
 			ModelResources->RuntimeReferencedTextures.Add(Pair.Key);
 		}
 
-		ModelResources->PassthroughObjects = MoveTemp(GenerationContext.PassthroughObjects);
 		ModelResources->TextureParameterDefaultValues = MoveTemp(GenerationContext.TextureParameterDefaultValues);
 		ModelResources->SkeletalMeshParameterDefaultValues = MoveTemp(GenerationContext.SkeletalMeshParameterDefaultValues);
 		ModelResources->MaterialParameterDefaultValues = MoveTemp(GenerationContext.MaterialParameterDefaultValues);
@@ -1674,6 +1673,8 @@ void FCustomizableObjectCompiler::FinishCompilationTask()
 	{
 		TObjectPtr<UModelResources> ModelResources = PlatformData->ModelResources.Get();
 
+		ModelResources->PassthroughObjects = MoveTemp(CompilationContext->PassthroughObjectFactory.GetAll());
+
 		ModelResources->ExternalOperations = MoveTemp(CompileTask->ExternalOperations);
 
 		for (TPair<UE::Mutable::Private::OP::ADDRESS, FInstancedStruct> ExternalOperation : ModelResources->ExternalOperations)
diff --git a/Engine/Plugins/Mutable/Source/CustomizableObjectEditor/Private/MuCOE/GenerateMutableSource/GenerateMutableSource.cpp b/Engine/Plugins/Mutable/Source/CustomizableObjectEditor/Private/MuCOE/GenerateMutableSource/GenerateMutableSource.cpp
--- a/Engine/Plugins/Mutable/Source/CustomizableObjectEditor/Private/MuCOE/GenerateMutableSource/GenerateMutableSource.cpp
+++ b/Engine/Plugins/Mutable/Source/CustomizableObjectEditor/Private/MuCOE/GenerateMutableSource/GenerateMutableSource.cpp
@@ -1776,7 +1776,7 @@ UE::Mutable::Private::TManagedPtr<UE::Mutable::Private::FImage> GenerateImageCon
 
 	if (bIsReference)
 	{
-		Id = GenerationContext.AddPassthroughObject(*Texture);
+		Id = GenerationContext.CompilationContext->PassthroughObjectFactory.Add(*Texture);
 	}
 	else if (bIsCompileTime)
 	{
@@ -1967,24 +1967,3 @@ uint32 GetTypeHash(const FSharedSurfaceKey& Key)
 }
 
 
-UE::Mutable::Private::PASSTHROUGH_ID FMutableGraphGenerationContext::AddPassthroughObject(UObject& Object)
-{
-	TSoftObjectPtr<UObject> SoftObjectPtr(&Object);
-		
-	UE::Mutable::Private::PASSTHROUGH_ID Id = UE::Mutable::Private::GetPassthroughId(SoftObjectPtr);
-	while (TSoftObjectPtr<UObject>* PassthroughObject = PassthroughObjects.Find(Id))
-	{
-		if (*PassthroughObject == &Object)
-		{
-			return Id;
-		}
-		else
-		{
-			++Id;
-		}
-	}
-	
-	PassthroughObjects.Add(Id, SoftObjectPtr);
-
-	return Id;
-}
diff --git a/Engine/Plugins/Mutable/Source/CustomizableObjectEditor/Private/MuCOE/GenerateMutableSource/GenerateMutableSource.h b/Engine/Plugins/Mutable/Source/CustomizableObjectEditor/Private/MuCOE/GenerateMutableSource/GenerateMutableSource.h
--- a/Engine/Plugins/Mutable/Source/CustomizableObjectEditor/Private/MuCOE/GenerateMutableSource/GenerateMutableSource.h
+++ b/Engine/Plugins/Mutable/Source/CustomizableObjectEditor/Private/MuCOE/GenerateMutableSource/GenerateMutableSource.h
@@ -13,6 +13,7 @@
 #include "MuCO/LoadUtils.h"
 #include "MuCOE/CustomizableObjectEditorLogger.h"
 #include "MuCOE/ExtensionDataCompilerInterface.h"
+#include "MuCOE/PassthroughObjectFactory.h"
 #include "MuCOE/Nodes/CustomizableObjectNodeMaterialBase.h"
 #include "MuCOE/Nodes/CustomizableObjectNodeObject.h"
 #include "MuCOE/Nodes/CustomizableObjectNodeTable.h"
@@ -612,6 +613,8 @@ struct FMutableCompilationContext
 
 	/** Bone identifiers that had a collision. */
 	TMap<FString, UE::Mutable::Private::FBoneName> RemappedBoneNames;
+	
+	UE::Mutable::Private::FPassthroughObjectFactory PassthroughObjectFactory;
 
 private:
 
@@ -791,8 +794,6 @@ struct FMutableGraphGenerationContext
 	// Cache of runtime pass-through images and their IDs used in the core to identify them
 	// These textures will become mutable images in the compiled model.
 	TMap<TSoftObjectPtr<UTexture2D>, FGeneratedReferencedTexture> CompileTimeTextureMap;
-
-	TMap<UE::Mutable::Private::PASSTHROUGH_ID, TSoftObjectPtr<UObject>> PassthroughObjects;
 	
 	// Mutable meshes already build for source UStaticMesh or USkeletalMesh.
 	struct FGeneratedMeshData
@@ -877,8 +878,6 @@ struct FMutableGraphGenerationContext
 
 	/** Get the reference skeletal mesh associated to the current mesh component being generated */
 	FMutableComponentInfo* GetCurrentComponentInfo();
-
-	UE::Mutable::Private::PASSTHROUGH_ID AddPassthroughObject(UObject& Object);
 	
 	UObject* LoadObject(const FSoftObjectPtr& SoftObject);
 	
diff --git a/Engine/Plugins/Mutable/Source/CustomizableObjectEditor/Private/MuCOE/GenerateMutableSource/GenerateMutableSourceComponent.cpp b/Engine/Plugins/Mutable/Source/CustomizableObjectEditor/Private/MuCOE/GenerateMutableSource/GenerateMutableSourceComponent.cpp
--- a/Engine/Plugins/Mutable/Source/CustomizableObjectEditor/Private/MuCOE/GenerateMutableSource/GenerateMutableSourceComponent.cpp
+++ b/Engine/Plugins/Mutable/Source/CustomizableObjectEditor/Private/MuCOE/GenerateMutableSource/GenerateMutableSourceComponent.cpp
@@ -85,7 +85,7 @@ UE::Mutable::Private::Ptr<UE::Mutable::Private::NodeComponent> GenerateMutableSo
 		else if(UMaterialInterface* OverlayMaterial = GenerationContext.LoadObject(TypedComponentMesh->GetOverlayMaterial()))
 		{
 			UE::Mutable::Private::Ptr<UE::Mutable::Private::NodeMaterialConstant> MaterialNode = new UE::Mutable::Private::NodeMaterialConstant();
-			MaterialNode->MaterialId = GenerationContext.AddPassthroughObject(*OverlayMaterial);
+			MaterialNode->MaterialId = GenerationContext.CompilationContext->PassthroughObjectFactory.Add(*OverlayMaterial);
 			NodeComponentNew->OverlayMaterial = MaterialNode;
 		}
 
diff --git a/Engine/Plugins/Mutable/Source/CustomizableObjectEditor/Private/MuCOE/GenerateMutableSource/GenerateMutableSourceMaterial.cpp b/Engine/Plugins/Mutable/Source/CustomizableObjectEditor/Private/MuCOE/GenerateMutableSource/GenerateMutableSourceMaterial.cpp
--- a/Engine/Plugins/Mutable/Source/CustomizableObjectEditor/Private/MuCOE/GenerateMutableSource/GenerateMutableSourceMaterial.cpp
+++ b/Engine/Plugins/Mutable/Source/CustomizableObjectEditor/Private/MuCOE/GenerateMutableSource/GenerateMutableSourceMaterial.cpp
@@ -162,7 +162,7 @@ UE::Mutable::Private::Ptr<UE::Mutable::Private::NodeMaterial> GenerateMutableSou
 		if (!BaseMaterial->MaterialSource)
 		{
 			UE::Mutable::Private::Ptr<UE::Mutable::Private::NodeMaterialConstant> MaterialNode = new UE::Mutable::Private::NodeMaterialConstant();
-			MaterialNode->MaterialId = GenerationContext.AddPassthroughObject(*GenerationContext.CurrentReferencedMaterial.Get());
+			MaterialNode->MaterialId = GenerationContext.CompilationContext->PassthroughObjectFactory.Add(*GenerationContext.CurrentReferencedMaterial.Get());
 
 			BaseMaterial->MaterialSource = MaterialNode;
 		}
@@ -426,7 +426,7 @@ UE::Mutable::Private::Ptr<UE::Mutable::Private::NodeMaterial> GenerateMutableSou
 		{
 			MaterialNode = new UE::Mutable::Private::NodeMaterialConstant();
 			GenerationContext.CurrentReferencedMaterial = TStrongObjectPtr(Material);
-			MaterialNode->MaterialId = GenerationContext.AddPassthroughObject(*Material);
+			MaterialNode->MaterialId = GenerationContext.CompilationContext->PassthroughObjectFactory.Add(*Material);
 
 			// Store the parameter that needs to be processed
 			const FMaterialBreakParameter CurrentBreakParameter = GenerationContext.GetCurrentMaterialBreakParameter();
@@ -872,7 +872,7 @@ UE::Mutable::Private::Ptr<UE::Mutable::Private::NodeMaterial> GenerateMutableSou
 	{
 		//This node will add a default value in case of error
 		UE::Mutable::Private::Ptr<UE::Mutable::Private::NodeMaterialConstant> ConstantValue = new UE::Mutable::Private::NodeMaterialConstant();
-		ConstantValue->MaterialId = GenerationContext.AddPassthroughObject(*GenerationContext.CurrentReferencedMaterial.Get());
+		ConstantValue->MaterialId = GenerationContext.CompilationContext->PassthroughObjectFactory.Add(*GenerationContext.CurrentReferencedMaterial.Get());
 
 		Result = ConstantValue;
 
diff --git a/Engine/Plugins/Mutable/Source/CustomizableObjectEditor/Private/MuCOE/GenerateMutableSource/GenerateMutableSourceTable.cpp b/Engine/Plugins/Mutable/Source/CustomizableObjectEditor/Private/MuCOE/GenerateMutableSource/GenerateMutableSourceTable.cpp
--- a/Engine/Plugins/Mutable/Source/CustomizableObjectEditor/Private/MuCOE/GenerateMutableSource/GenerateMutableSourceTable.cpp
+++ b/Engine/Plugins/Mutable/Source/CustomizableObjectEditor/Private/MuCOE/GenerateMutableSource/GenerateMutableSourceTable.cpp
@@ -399,7 +399,7 @@ bool FillTableColumn(const UCustomizableObjectNodeTable* TableNode, UE::Mutable:
 					CurrentColumn = MutableTable->AddColumn(ColumnName, UE::Mutable::Private::ETableColumnType::Material);
 				}
 
-				int32 ReferenceMaterialId = GenerationContext.AddPassthroughObject(*MaterialInstance);
+				int32 ReferenceMaterialId = GenerationContext.CompilationContext->PassthroughObjectFactory.Add(*MaterialInstance);
 				MutableTable->SetCell(CurrentColumn, RowId, ReferenceMaterialId);
 
 				return true;
@@ -1198,7 +1198,7 @@ bool GenerateNoneRow(const UCustomizableObjectNodeTable& TableNode, const UEdGra
 			ColumnId = MutableTable->AddColumn(ColumnName, UE::Mutable::Private::ETableColumnType::Material);
 		}
 
-		int32 MaterialId = GenerationContext.AddPassthroughObject(*GenerationContext.CurrentReferencedMaterial.Get());
+		int32 MaterialId = GenerationContext.CompilationContext->PassthroughObjectFactory.Add(*GenerationContext.CurrentReferencedMaterial.Get());
 		MutableTable->SetCell(ColumnId, RowId, MaterialId);
 	}
 	else if (Pin->PinType.PinCategory == Schema->PC_Color)
diff --git a/Engine/Plugins/Mutable/Source/CustomizableObjectEditor/Private/MuCOE/PassthroughObjectFactory.cpp b/Engine/Plugins/Mutable/Source/CustomizableObjectEditor/Private/MuCOE/PassthroughObjectFactory.cpp
new file mode 100644
--- /dev/null
+++ b/Engine/Plugins/Mutable/Source/CustomizableObjectEditor/Private/MuCOE/PassthroughObjectFactory.cpp
@@ -0,0 +1,35 @@
+// Copyright Epic Games, Inc. All Rights Reserved.
+
+#include "PassthroughObjectFactory.h"
+
+
+namespace UE::Mutable::Private
+{
+	PASSTHROUGH_ID FPassthroughObjectFactory::Add(UObject& Object)
+	{
+		TSoftObjectPtr<UObject> SoftObjectPtr(&Object);
+		
+		PASSTHROUGH_ID Id = GetPassthroughId(SoftObjectPtr);
+		while (TSoftObjectPtr<UObject>* PassthroughObject = PassthroughObjects.Find(Id))
+		{
+			if (*PassthroughObject == &Object)
+			{
+				return Id;
+			}
+			else
+			{
+				++Id;
+			}
+		}
+	
+		PassthroughObjects.Add(Id, SoftObjectPtr);
+
+		return Id;
+	}
+
+
+	TMap<PASSTHROUGH_ID, TSoftObjectPtr<UObject>>& FPassthroughObjectFactory::GetAll()
+	{
+		return PassthroughObjects;
+	}
+}
diff --git a/Engine/Plugins/Mutable/Source/CustomizableObjectEditor/Private/MuCOE/PassthroughObjectFactory.h b/Engine/Plugins/Mutable/Source/CustomizableObjectEditor/Private/MuCOE/PassthroughObjectFactory.h
new file mode 100644
--- /dev/null
+++ b/Engine/Plugins/Mutable/Source/CustomizableObjectEditor/Private/MuCOE/PassthroughObjectFactory.h
@@ -0,0 +1,19 @@
+// Copyright Epic Games, Inc. All Rights Reserved.
+
+#pragma once
+#include "MuR/PassthroughObject.h"
+
+
+namespace UE::Mutable::Private
+{
+	class FPassthroughObjectFactory
+	{
+	public:
+		PASSTHROUGH_ID Add(UObject& Object);
+		
+		TMap<PASSTHROUGH_ID, TSoftObjectPtr<UObject>>& GetAll();
+		
+	private:
+		TMap<PASSTHROUGH_ID, TSoftObjectPtr<UObject>> PassthroughObjects;
+	};	
+}
