[Mutable] Change Break Skeletal Mesh node color.

#Synced-CL 49065323

diff --git a/Engine/Plugins/Mutable/Source/CustomizableObjectEditor/Private/MuCOE/Nodes/CONodeSkeletalMeshBreak.cpp b/Engine/Plugins/Mutable/Source/CustomizableObjectEditor/Private/MuCOE/Nodes/CONodeSkeletalMeshBreak.cpp
--- a/Engine/Plugins/Mutable/Source/CustomizableObjectEditor/Private/MuCOE/Nodes/CONodeSkeletalMeshBreak.cpp
+++ b/Engine/Plugins/Mutable/Source/CustomizableObjectEditor/Private/MuCOE/Nodes/CONodeSkeletalMeshBreak.cpp
@@ -23,7 +23,7 @@ FText UCONodeSkeletalMeshBreak::GetTooltipText() const
 
 FLinearColor UCONodeSkeletalMeshBreak::GetNodeTitleColor() const
 {
-	return UEdGraphSchema_CustomizableObject::GetPinTypeColor(UEdGraphSchema_CustomizableObject::PC_Material);
+	return UEdGraphSchema_CustomizableObject::GetPinTypeColor(UEdGraphSchema_CustomizableObject::PC_PassthroughSkeletalMesh);
 }
 
 
