[Mutable] Added new Passthrough Parameter pin type.

- Forced all pins to have a LOCTEXT.
- Fix naming issues with Node Switch pins.

#Synced-CL 49066841

diff --git a/Engine/Plugins/Experimental/HairStrandsMutable/Source/HairStrandsMutableEditor/Private/CustomizableObjectNodeGroomConstant.cpp b/Engine/Plugins/Experimental/HairStrandsMutable/Source/HairStrandsMutableEditor/Private/CustomizableObjectNodeGroomConstant.cpp
--- a/Engine/Plugins/Experimental/HairStrandsMutable/Source/HairStrandsMutableEditor/Private/CustomizableObjectNodeGroomConstant.cpp
+++ b/Engine/Plugins/Experimental/HairStrandsMutable/Source/HairStrandsMutableEditor/Private/CustomizableObjectNodeGroomConstant.cpp
@@ -31,8 +31,7 @@ FText UCustomizableObjectNodeGroomConstant::GetTooltipText() const
 
 void UCustomizableObjectNodeGroomConstant::AllocateDefaultPins(UCustomizableObjectNodeRemapPins* RemapPins)
 {
-	UEdGraphPin* OutputPin = CustomCreatePin(EGPD_Output, UHairStrandsMutableExtension::GroomPinType, UHairStrandsMutableExtension::GroomsBaseNodePinName);
-	OutputPin->bDefaultValueIsIgnored = true;
+	UEdGraphPin* OutputPin = CustomCreatePin(EGPD_Output, UHairStrandsMutableExtension::GroomPinType, UHairStrandsMutableExtension::GroomsBaseNodePinName, UHairStrandsMutableExtension::GroomNodeCategory);
 }
 
 bool UCustomizableObjectNodeGroomConstant::ShouldAddToContextMenu(FText& OutCategory) const
diff --git a/Engine/Plugins/Experimental/RigLogicMutable/Source/RigLogicMutableEditor/Private/CustomizableObjectNodeDNAConstant.cpp b/Engine/Plugins/Experimental/RigLogicMutable/Source/RigLogicMutableEditor/Private/CustomizableObjectNodeDNAConstant.cpp
--- a/Engine/Plugins/Experimental/RigLogicMutable/Source/RigLogicMutableEditor/Private/CustomizableObjectNodeDNAConstant.cpp
+++ b/Engine/Plugins/Experimental/RigLogicMutable/Source/RigLogicMutableEditor/Private/CustomizableObjectNodeDNAConstant.cpp
@@ -33,7 +33,7 @@ FText UCustomizableObjectNodeDNAConstant::GetTooltipText() const
 
 void UCustomizableObjectNodeDNAConstant::AllocateDefaultPins(UCustomizableObjectNodeRemapPins* RemapPins)
 {
-	UEdGraphPin* OutputPin = CustomCreatePin(EGPD_Output, URigLogicMutableExtension::DNAPinType, URigLogicMutableExtension::DNABaseNodePinName);
+	UEdGraphPin* OutputPin = CustomCreatePin(EGPD_Output, URigLogicMutableExtension::DNAPinType, URigLogicMutableExtension::DNABaseNodePinName, FText::FromName(URigLogicMutableExtension::DNABaseNodePinName));
 	OutputPin->bDefaultValueIsIgnored = true;
 }
 
diff --git a/Engine/Plugins/Mutable/Source/CustomizableObject/Internal/MuCO/CustomizableObjectCustomVersion.h b/Engine/Plugins/Mutable/Source/CustomizableObject/Internal/MuCO/CustomizableObjectCustomVersion.h
--- a/Engine/Plugins/Mutable/Source/CustomizableObject/Internal/MuCO/CustomizableObjectCustomVersion.h
+++ b/Engine/Plugins/Mutable/Source/CustomizableObject/Internal/MuCO/CustomizableObjectCustomVersion.h
@@ -211,6 +211,8 @@ struct FCustomizableObjectCustomVersion
 		PinSkeletalMesh,
 
 		NodeStaticStringRemoveHiddenInputPin,
+		
+		PassthroughParameter,
 
 		// -----<new versions can be added above this line>-------------------------------------------------
 		VersionPlusOne,
diff --git a/Engine/Plugins/Mutable/Source/CustomizableObjectEditor/Private/MuCOE/EdGraphSchema_CustomizableObject.cpp b/Engine/Plugins/Mutable/Source/CustomizableObjectEditor/Private/MuCOE/EdGraphSchema_CustomizableObject.cpp
--- a/Engine/Plugins/Mutable/Source/CustomizableObjectEditor/Private/MuCOE/EdGraphSchema_CustomizableObject.cpp
+++ b/Engine/Plugins/Mutable/Source/CustomizableObjectEditor/Private/MuCOE/EdGraphSchema_CustomizableObject.cpp
@@ -298,9 +298,11 @@ const FName UEdGraphSchema_CustomizableObject::PC_MeshSection("material");
 const FName UEdGraphSchema_CustomizableObject::PC_Modifier("modifier");
 const FName UEdGraphSchema_CustomizableObject::PC_Mesh("mesh");
 const FName UEdGraphSchema_CustomizableObject::PC_SkeletalMesh("skeletalMesh");
-const FName UEdGraphSchema_CustomizableObject::PC_PassthroughSkeletalMesh("passThroughSkeletalMesh");
+const FName UEdGraphSchema_CustomizableObject::PC_SkeletalMesh_Passthrough("passThroughSkeletalMesh");
+const FName UEdGraphSchema_CustomizableObject::PC_SkeletalMesh_PassthroughParameter("passThroughParameterSkeletalMesh");
 const FName UEdGraphSchema_CustomizableObject::PC_Texture("image");
-const FName UEdGraphSchema_CustomizableObject::PC_PassthroughTexture("passThroughImage");
+const FName UEdGraphSchema_CustomizableObject::PC_Texture_Passthrough("passThroughImage");
+const FName UEdGraphSchema_CustomizableObject::PC_Texture_PassthroughParameter("passThroughParameterImage");
 const FName UEdGraphSchema_CustomizableObject::PC_Projector("projector");
 const FName UEdGraphSchema_CustomizableObject::PC_GroupProjector("groupProjector");
 const FName UEdGraphSchema_CustomizableObject::PC_Color("color");
@@ -316,7 +318,7 @@ const FName UEdGraphSchema_CustomizableObject::PC_String("string");
 
 // Add more pin types to this array if needed
 const TArray<FName> UEdGraphSchema_CustomizableObject::SupportedMacroPinTypes = { PC_Object, PC_Component, PC_MeshSection, PC_Modifier, PC_Mesh, PC_Texture,
-PC_PassthroughTexture, PC_Projector, PC_GroupProjector, PC_Color, PC_Float, PC_Enum, PC_Stack, PC_Material, PC_PoseAsset, PC_Transform, PC_String };
+	PC_Texture_Passthrough, PC_Projector, PC_GroupProjector, PC_Color, PC_Float, PC_Enum, PC_Stack, PC_Material, PC_PoseAsset, PC_Transform, PC_String };
 
 // Node categories
 const FText UEdGraphSchema_CustomizableObject::NC_Experimental(FText::FromString("Experimental"));
@@ -711,7 +713,7 @@ void UEdGraphSchema_CustomizableObject::GetGraphContextActions(FGraphContextMenu
 	
 	// External Pin Nodes
 	{
-		TArray<FName> PinTypes({ PC_MeshSection, PC_Modifier, PC_Mesh, PC_Texture, PC_Projector, PC_GroupProjector, PC_Color, PC_Float, PC_Bool, PC_Enum, PC_Transform, PC_Stack, PC_PassthroughTexture, PC_Material, PC_PoseAsset, PC_Component });
+		TArray<FName> PinTypes({ PC_MeshSection, PC_Modifier, PC_Mesh, PC_Texture, PC_Projector, PC_GroupProjector, PC_Color, PC_Float, PC_Bool, PC_Enum, PC_Transform, PC_Stack, PC_Texture_Passthrough, PC_Texture_Passthrough, PC_Material, PC_PoseAsset, PC_Component });
 
 		// Extensions
 		for (const FRegisteredCustomizableObjectPinType& PinType : ICustomizableObjectModule::Get().GetExtendedPinTypes())
@@ -894,14 +896,16 @@ FLinearColor UEdGraphSchema_CustomizableObject::GetPinTypeColor(const FName& Typ
 		{
 			const FLinearColor Color = AssetDefinitionRegistry->GetAssetDefinitionForClass(USkeletalMesh::StaticClass())->GetAssetColor();
 			PinCategoryColors.Add(PC_SkeletalMesh, Color);
-			PinCategoryColors.Add(PC_PassthroughSkeletalMesh, Color); 	
+			PinCategoryColors.Add(PC_SkeletalMesh_Passthrough, Color); 	
+			PinCategoryColors.Add(PC_SkeletalMesh_PassthroughParameter, Color); 	
 			
 		}
 
 		{
 			const FLinearColor Color = AssetDefinitionRegistry->GetAssetDefinitionForClass(UTexture::StaticClass())->GetAssetColor();
 			PinCategoryColors.Add(PC_Texture, Color);
-			PinCategoryColors.Add(PC_PassthroughTexture, Color);	
+			PinCategoryColors.Add(PC_Texture_Passthrough, Color);	
+			PinCategoryColors.Add(PC_Texture_PassthroughParameter, Color);	
 		}
 
 		{
@@ -1527,7 +1531,11 @@ FName UEdGraphSchema_CustomizableObject::GetPinCategoryName(const FName& PinCate
 	{
 		return TEXT("Skeletal Mesh");
 	}
-	else if (PinCategory == PC_PassthroughSkeletalMesh)
+	else if (PinCategory == PC_SkeletalMesh_Passthrough)
+	{
+		return TEXT("Passthrough Skeletal Mesh");
+	}
+	else if (PinCategory == PC_SkeletalMesh_PassthroughParameter)
 	{
 		return TEXT("Passthrough Skeletal Mesh");
 	}
@@ -1535,9 +1543,13 @@ FName UEdGraphSchema_CustomizableObject::GetPinCategoryName(const FName& PinCate
 	{
 		return TEXT("Texture");
 	}
-	else if (PinCategory == PC_PassthroughTexture)
+	else if (PinCategory == PC_Texture_Passthrough)
 	{
-		return TEXT("PassThrough Texture");
+		return TEXT("Passthrough Texture");
+	}
+	else if (PinCategory == PC_Texture_PassthroughParameter)
+	{
+		return TEXT("Passthrough Texture");
 	}
 	else if (PinCategory == PC_Projector)
 	{
@@ -1657,7 +1669,11 @@ FText UEdGraphSchema_CustomizableObject::GetPinCategoryFriendlyName(const FName&
 	{
 		return LOCTEXT("Skeletal_Mesh_Pin_Category", "Skeletal Mesh");
 	}
-	else if (PinCategory == PC_PassthroughSkeletalMesh)
+	else if (PinCategory == PC_SkeletalMesh_Passthrough)
+	{
+		return LOCTEXT("Passthrough_Skeletal_Mesh_Pin_Category", "Passthrough Skeletal Mesh");
+	}
+	else if (PinCategory == PC_SkeletalMesh_PassthroughParameter)
 	{
 		return LOCTEXT("Passthrough_Skeletal_Mesh_Pin_Category", "Passthrough Skeletal Mesh");
 	}
@@ -1665,7 +1681,11 @@ FText UEdGraphSchema_CustomizableObject::GetPinCategoryFriendlyName(const FName&
 	{
 		return LOCTEXT("Image_Pin_Category", "Texture");
 	}
-	else if (PinCategory == PC_PassthroughTexture)
+	else if (PinCategory == PC_Texture_Passthrough)
+	{
+		return LOCTEXT("Passthrough_Texture_Pin_Category", "Passthrough Texture");
+	}
+	else if (PinCategory == PC_Texture_PassthroughParameter)
 	{
 		return LOCTEXT("Passthrough_Texture_Pin_Category", "Passthrough Texture");
 	}
@@ -1747,8 +1767,10 @@ FText UEdGraphSchema_CustomizableObject::GetPinCategoryFriendlyName(const UScrip
 
 bool UEdGraphSchema_CustomizableObject::IsPassthrough(const FName& PinCategory)
 {
-	return PinCategory == PC_PassthroughTexture ||
-		PinCategory == PC_PassthroughSkeletalMesh;
+	return PinCategory == PC_Texture_Passthrough || 
+		PinCategory == PC_Texture_PassthroughParameter ||
+		PinCategory == PC_SkeletalMesh_Passthrough ||
+		PinCategory == PC_SkeletalMesh_PassthroughParameter;
 }
 
 TSharedPtr<IAssetReferenceFilter> UEdGraphSchema_CustomizableObject::MakeAssetReferenceFilter(const UEdGraph* Graph)
diff --git a/Engine/Plugins/Mutable/Source/CustomizableObjectEditor/Private/MuCOE/GenerateMutableSource/GenerateMutableSource.cpp b/Engine/Plugins/Mutable/Source/CustomizableObjectEditor/Private/MuCOE/GenerateMutableSource/GenerateMutableSource.cpp
--- a/Engine/Plugins/Mutable/Source/CustomizableObjectEditor/Private/MuCOE/GenerateMutableSource/GenerateMutableSource.cpp
+++ b/Engine/Plugins/Mutable/Source/CustomizableObjectEditor/Private/MuCOE/GenerateMutableSource/GenerateMutableSource.cpp
@@ -565,9 +565,9 @@ UTexture2D* FindReferenceImage(const UEdGraphPin* Pin, FMutableGraphGenerationCo
 
 	else if (const UCustomizableObjectNodeTextureSwitch* TypedNodeSwitch = Cast<UCustomizableObjectNodeTextureSwitch>(Node))
 	{
-		for (int SelectorIndex = 0; !Result && SelectorIndex < TypedNodeSwitch->GetNumElements(); ++SelectorIndex)
+		for (int SelectorIndex = 0; !Result && SelectorIndex < TypedNodeSwitch->SwitchPins.Num(); ++SelectorIndex)
 		{
-			if (const UEdGraphPin* ConnectedPin = FollowInputPin(*TypedNodeSwitch->GetElementPin(SelectorIndex)))
+			if (const UEdGraphPin* ConnectedPin = FollowInputPin(*TypedNodeSwitch->SwitchPins[SelectorIndex].Get()))
 			{
 				Result = FindReferenceImage(ConnectedPin, GenerationContext);
 			}
@@ -576,9 +576,9 @@ UTexture2D* FindReferenceImage(const UEdGraphPin* Pin, FMutableGraphGenerationCo
 
 	else if (const UCustomizableObjectNodePassThroughTextureSwitch* TypedNodePassThroughSwitch = Cast<UCustomizableObjectNodePassThroughTextureSwitch>(Node))
 	{
-		for (int32 SelectorIndex = 0; !Result && SelectorIndex < TypedNodePassThroughSwitch->GetNumElements(); ++SelectorIndex)
+		for (int32 SelectorIndex = 0; !Result && SelectorIndex < TypedNodePassThroughSwitch->SwitchPins.Num(); ++SelectorIndex)
 		{
-			if (const UEdGraphPin* ConnectedPin = FollowInputPin(*TypedNodePassThroughSwitch->GetElementPin(SelectorIndex)))
+			if (const UEdGraphPin* ConnectedPin = FollowInputPin(*TypedNodePassThroughSwitch->SwitchPins[SelectorIndex].Get()))
 			{
 				Result = FindReferenceImage(ConnectedPin, GenerationContext);
 			}
diff --git a/Engine/Plugins/Mutable/Source/CustomizableObjectEditor/Private/MuCOE/GenerateMutableSource/GenerateMutableSourceColor.cpp b/Engine/Plugins/Mutable/Source/CustomizableObjectEditor/Private/MuCOE/GenerateMutableSource/GenerateMutableSourceColor.cpp
--- a/Engine/Plugins/Mutable/Source/CustomizableObjectEditor/Private/MuCOE/GenerateMutableSource/GenerateMutableSourceColor.cpp
+++ b/Engine/Plugins/Mutable/Source/CustomizableObjectEditor/Private/MuCOE/GenerateMutableSource/GenerateMutableSourceColor.cpp
@@ -96,7 +96,7 @@ UE::Mutable::Private::Ptr<UE::Mutable::Private::NodeColour> GenerateMutableSourc
 	{
 		Result = [&]() -> UE::Mutable::Private::Ptr<UE::Mutable::Private::NodeColourSwitch>
 		{
-			if (const int32 NumParameters = FollowInputPinArray(*TypedNodeColorSwitch->SwitchParameter()).Num();
+			if (const int32 NumParameters = FollowInputPinArray(*TypedNodeColorSwitch->SwitchParameterPinReference.Get()).Num();
 				NumParameters != 1)
 			{
 				const FText Message = NumParameters
@@ -107,7 +107,7 @@ UE::Mutable::Private::Ptr<UE::Mutable::Private::NodeColour> GenerateMutableSourc
 				return nullptr;
 			}
 
-			const UEdGraphPin* EnumPin = FollowInputPin(*TypedNodeColorSwitch->SwitchParameter());
+			const UEdGraphPin* EnumPin = FollowInputPin(*TypedNodeColorSwitch->SwitchParameterPinReference.Get());
 			UE::Mutable::Private::Ptr<UE::Mutable::Private::NodeScalar> SwitchParam = EnumPin ? GenerateMutableSourceFloat(EnumPin, GenerationContext) : nullptr;
 
 			// Switch Param not generated
@@ -131,7 +131,7 @@ UE::Mutable::Private::Ptr<UE::Mutable::Private::NodeColour> GenerateMutableSourc
 				return nullptr;
 			}
 
-			const int32 NumSwitchOptions = TypedNodeColorSwitch->GetNumElements();
+			const int32 NumSwitchOptions = TypedNodeColorSwitch->SwitchPins.Num();
 
 			UE::Mutable::Private::NodeScalarEnumParameter* EnumParameter = static_cast<UE::Mutable::Private::NodeScalarEnumParameter*>(SwitchParam.get());
 			if (NumSwitchOptions != EnumParameter->Options.Num())
@@ -146,7 +146,7 @@ UE::Mutable::Private::Ptr<UE::Mutable::Private::NodeColour> GenerateMutableSourc
 
 			for (int SelectorIndex = 0; SelectorIndex < NumSwitchOptions; ++SelectorIndex)
 			{
-				const UEdGraphPin* const ColorPin = TypedNodeColorSwitch->GetElementPin(SelectorIndex);
+				const UEdGraphPin* const ColorPin = TypedNodeColorSwitch->SwitchPins[SelectorIndex].Get();
 				if (const UEdGraphPin* ConnectedPin = FollowInputPin(*ColorPin))
 				{
 					SwitchNode->Options[SelectorIndex] = GenerateMutableSourceColor(ConnectedPin, GenerationContext);
diff --git a/Engine/Plugins/Mutable/Source/CustomizableObjectEditor/Private/MuCOE/GenerateMutableSource/GenerateMutableSourceComponent.cpp b/Engine/Plugins/Mutable/Source/CustomizableObjectEditor/Private/MuCOE/GenerateMutableSource/GenerateMutableSourceComponent.cpp
--- a/Engine/Plugins/Mutable/Source/CustomizableObjectEditor/Private/MuCOE/GenerateMutableSource/GenerateMutableSourceComponent.cpp
+++ b/Engine/Plugins/Mutable/Source/CustomizableObjectEditor/Private/MuCOE/GenerateMutableSource/GenerateMutableSourceComponent.cpp
@@ -211,7 +211,7 @@ UE::Mutable::Private::Ptr<UE::Mutable::Private::NodeComponent> GenerateMutableSo
 		// Using a lambda so control flow is easier to manage.
 		Result = [&]()
 			{
-				const UEdGraphPin* SwitchParameter = TypedNodeSwitch->SwitchParameter();
+				const UEdGraphPin* SwitchParameter = TypedNodeSwitch->SwitchParameterPinReference.Get();
 
 				// Check Switch Parameter arity preconditions.
 				if (const UEdGraphPin* EnumPin = FollowInputPin(*SwitchParameter))
@@ -239,7 +239,7 @@ UE::Mutable::Private::Ptr<UE::Mutable::Private::NodeComponent> GenerateMutableSo
 						return Result;
 					}
 
-					const int32 NumSwitchOptions = TypedNodeSwitch->GetNumElements();
+					const int32 NumSwitchOptions = TypedNodeSwitch->SwitchPins.Num();
 
 					UE::Mutable::Private::NodeScalarEnumParameter* EnumParameter = static_cast<UE::Mutable::Private::NodeScalarEnumParameter*>(SwitchParam.get());
 					if (NumSwitchOptions != EnumParameter->Options.Num())
@@ -254,7 +254,7 @@ UE::Mutable::Private::Ptr<UE::Mutable::Private::NodeComponent> GenerateMutableSo
 
 					for (int32 SelectorIndex = 0; SelectorIndex < NumSwitchOptions; ++SelectorIndex)
 					{
-						if (const UEdGraphPin* ConnectedPin = FollowInputPin(*TypedNodeSwitch->GetElementPin(SelectorIndex)))
+						if (const UEdGraphPin* ConnectedPin = FollowInputPin(*TypedNodeSwitch->SwitchPins[SelectorIndex].Get()))
 						{
 							UE::Mutable::Private::Ptr<UE::Mutable::Private::NodeComponent> ChildNode = GenerateMutableSourceComponent(ConnectedPin, GenerationContext);
 							if (ChildNode)
diff --git a/Engine/Plugins/Mutable/Source/CustomizableObjectEditor/Private/MuCOE/GenerateMutableSource/GenerateMutableSourceFloat.cpp b/Engine/Plugins/Mutable/Source/CustomizableObjectEditor/Private/MuCOE/GenerateMutableSource/GenerateMutableSourceFloat.cpp
--- a/Engine/Plugins/Mutable/Source/CustomizableObjectEditor/Private/MuCOE/GenerateMutableSource/GenerateMutableSourceFloat.cpp
+++ b/Engine/Plugins/Mutable/Source/CustomizableObjectEditor/Private/MuCOE/GenerateMutableSource/GenerateMutableSourceFloat.cpp
@@ -125,7 +125,7 @@ UE::Mutable::Private::Ptr<UE::Mutable::Private::NodeScalar> GenerateMutableSourc
 		// Using a lambda so control flow is easier to manage.
 		Result = [&]() -> UE::Mutable::Private::Ptr<UE::Mutable::Private::NodeScalar>
 		{
-			const UEdGraphPin* SwitchParameter = TypedNodeFloatSwitch->SwitchParameter();
+			const UEdGraphPin* SwitchParameter = TypedNodeFloatSwitch->SwitchParameterPinReference.Get();
 
 			// Check Switch Parameter arity preconditions.
 			if (const int32 NumParameters = FollowInputPinArray(*SwitchParameter).Num();
@@ -163,7 +163,7 @@ UE::Mutable::Private::Ptr<UE::Mutable::Private::NodeScalar> GenerateMutableSourc
 				return nullptr;
 			}
 
-			const int32 NumSwitchOptions = TypedNodeFloatSwitch->GetNumElements();
+			const int32 NumSwitchOptions = TypedNodeFloatSwitch->SwitchPins.Num();
 
 			UE::Mutable::Private::NodeScalarEnumParameter* EnumParameter = static_cast<UE::Mutable::Private::NodeScalarEnumParameter*>(SwitchParam.get());
 			if (NumSwitchOptions != EnumParameter->Options.Num())
@@ -178,12 +178,9 @@ UE::Mutable::Private::Ptr<UE::Mutable::Private::NodeScalar> GenerateMutableSourc
 
 			for (int SelectorIndex = 0; SelectorIndex < NumSwitchOptions; ++SelectorIndex)
 			{
-				if (const UEdGraphPin* const FloatPin = TypedNodeFloatSwitch->GetElementPin(SelectorIndex))
+				if (const UEdGraphPin* ConnectedPin = FollowInputPin(*TypedNodeFloatSwitch->SwitchPins[SelectorIndex].Get()))
 				{
-					if (const UEdGraphPin* ConnectedPin = FollowInputPin(*FloatPin))
-					{
-						SwitchNode->Options[SelectorIndex] = GenerateMutableSourceFloat(ConnectedPin, GenerationContext);
-					}
+					SwitchNode->Options[SelectorIndex] = GenerateMutableSourceFloat(ConnectedPin, GenerationContext);
 				}
 			}
 
diff --git a/Engine/Plugins/Mutable/Source/CustomizableObjectEditor/Private/MuCOE/GenerateMutableSource/GenerateMutableSourceImage.cpp b/Engine/Plugins/Mutable/Source/CustomizableObjectEditor/Private/MuCOE/GenerateMutableSource/GenerateMutableSourceImage.cpp
--- a/Engine/Plugins/Mutable/Source/CustomizableObjectEditor/Private/MuCOE/GenerateMutableSource/GenerateMutableSourceImage.cpp
+++ b/Engine/Plugins/Mutable/Source/CustomizableObjectEditor/Private/MuCOE/GenerateMutableSource/GenerateMutableSourceImage.cpp
@@ -349,7 +349,7 @@ UE::Mutable::Private::Ptr<UE::Mutable::Private::NodeImage> GenerateMutableSource
 	{
 		Result = [&]()
 		{
-			const UEdGraphPin* SwitchParameter = TypedNodeTextureSwitch->SwitchParameter();
+			const UEdGraphPin* SwitchParameter = TypedNodeTextureSwitch->SwitchParameterPinReference.Get();
 
 			// Check Switch Parameter arity preconditions.
 			if (const UEdGraphPin* ConnectedPin = FollowInputPin(*SwitchParameter))
@@ -372,7 +372,7 @@ UE::Mutable::Private::Ptr<UE::Mutable::Private::NodeImage> GenerateMutableSource
 					return Result;
 				}
 
-				const int32 NumSwitchOptions = TypedNodeTextureSwitch->GetNumElements();
+				const int32 NumSwitchOptions = TypedNodeTextureSwitch->SwitchPins.Num();
 
 				UE::Mutable::Private::NodeScalarEnumParameter* EnumParameter = static_cast<UE::Mutable::Private::NodeScalarEnumParameter*>(SwitchParam.get());
 				if (NumSwitchOptions != EnumParameter->Options.Num())
@@ -387,7 +387,7 @@ UE::Mutable::Private::Ptr<UE::Mutable::Private::NodeImage> GenerateMutableSource
 
 				for (int SelectorIndex = 0; SelectorIndex < NumSwitchOptions; ++SelectorIndex)
 				{
-					if (const UEdGraphPin* TexturePin = FollowInputPin(*TypedNodeTextureSwitch->GetElementPin(SelectorIndex)))
+					if (const UEdGraphPin* TexturePin = FollowInputPin(*TypedNodeTextureSwitch->SwitchPins[SelectorIndex].Get()))
 					{
 						SwitchNode->Options[SelectorIndex] = GenerateMutableSourceImage(TexturePin, GenerationContext, ReferenceTextureSize);
 					}
@@ -917,7 +917,7 @@ UE::Mutable::Private::Ptr<UE::Mutable::Private::NodeImage> GenerateMutableSource
 	{
 		Result = [&]()
 		{
-			const UEdGraphPin* SwitchParameter = TypedNodePassThroughTextureSwitch->SwitchParameter();
+			const UEdGraphPin* SwitchParameter = TypedNodePassThroughTextureSwitch->SwitchParameterPinReference.Get();
 
 			// Check Switch Parameter arity preconditions.
 			if (const UEdGraphPin* ConnectedPin = FollowInputPin(*SwitchParameter))
@@ -940,7 +940,7 @@ UE::Mutable::Private::Ptr<UE::Mutable::Private::NodeImage> GenerateMutableSource
 					return Result;
 				}
 
-				const int32 NumSwitchOptions = TypedNodePassThroughTextureSwitch->GetNumElements();
+				const int32 NumSwitchOptions = TypedNodePassThroughTextureSwitch->SwitchPins.Num();
 
 				UE::Mutable::Private::NodeScalarEnumParameter* EnumParameter = static_cast<UE::Mutable::Private::NodeScalarEnumParameter*>(SwitchParam.get());
 				if (NumSwitchOptions != EnumParameter->Options.Num())
@@ -956,7 +956,7 @@ UE::Mutable::Private::Ptr<UE::Mutable::Private::NodeImage> GenerateMutableSource
 
 				for (int32 SelectorIndex = 0; SelectorIndex < NumSwitchOptions; ++SelectorIndex)
 				{
-					if (const UEdGraphPin* TexturePin = FollowInputPin(*TypedNodePassThroughTextureSwitch->GetElementPin(SelectorIndex)))
+					if (const UEdGraphPin* TexturePin = FollowInputPin(*TypedNodePassThroughTextureSwitch->SwitchPins[SelectorIndex].Get()))
 					{
 						UE::Mutable::Private::Ptr<UE::Mutable::Private::NodeImage> PassThroughImage = GenerateMutableSourceImage(TexturePin, GenerationContext, ReferenceTextureSize);
 						SwitchNode->Options[SelectorIndex] = PassThroughImage;
@@ -1087,7 +1087,7 @@ UE::Mutable::Private::Ptr<UE::Mutable::Private::NodeImage> GenerateMutableSource
 							}
 
 							// Encoding the texture mode of the type to allow different texture modes from the same pin
-							if (Pin->PinType.PinCategory == Schema->PC_PassthroughTexture)
+							if (Pin->PinType.PinCategory == Schema->PC_Texture_Passthrough)
 							{
 								ColumnName += "--PassThrough";
 							}
@@ -1116,7 +1116,7 @@ UE::Mutable::Private::Ptr<UE::Mutable::Private::NodeImage> GenerateMutableSource
 							ImageTableNode->bNoneOption = TypedNodeTable->bAddNoneOption;
 							ImageTableNode->DefaultRowName = TypedNodeTable->DefaultRowName.ToString();
 
-							bool bIsPassthrough = (Pin->PinType.PinCategory == Schema->PC_PassthroughTexture);
+							bool bIsPassthrough = (Pin->PinType.PinCategory == Schema->PC_Texture_Passthrough);
 
 							UTexture2D* DefaultTexture2D = Cast<UTexture2D>(DefaultTexture);
 							if (!bIsPassthrough && DefaultTexture2D)
diff --git a/Engine/Plugins/Mutable/Source/CustomizableObjectEditor/Private/MuCOE/GenerateMutableSource/GenerateMutableSourceMaterial.cpp b/Engine/Plugins/Mutable/Source/CustomizableObjectEditor/Private/MuCOE/GenerateMutableSource/GenerateMutableSourceMaterial.cpp
--- a/Engine/Plugins/Mutable/Source/CustomizableObjectEditor/Private/MuCOE/GenerateMutableSource/GenerateMutableSourceMaterial.cpp
+++ b/Engine/Plugins/Mutable/Source/CustomizableObjectEditor/Private/MuCOE/GenerateMutableSource/GenerateMutableSourceMaterial.cpp
@@ -474,7 +474,7 @@ UE::Mutable::Private::Ptr<UE::Mutable::Private::NodeMaterial> GenerateMutableSou
 						}
 					}
 					
-					else if (CurrentBreakParameter.ParameterType == UEdGraphSchema_CustomizableObject::PC_PassthroughTexture)
+					else if (CurrentBreakParameter.ParameterType == UEdGraphSchema_CustomizableObject::PC_Texture_Passthrough)
 					{
 						UTexture* PassthroughTexture = nullptr;
 						const bool bPassthroughParameterFound = Material->GetTextureParameterValue(ParameterInfo, PassthroughTexture);
@@ -516,7 +516,7 @@ UE::Mutable::Private::Ptr<UE::Mutable::Private::NodeMaterial> GenerateMutableSou
 	{
 		Result = [&]()
 			{
-				const UEdGraphPin* SwitchParameter = TypedNodeMaterialSwitch->SwitchParameter();
+				const UEdGraphPin* SwitchParameter = TypedNodeMaterialSwitch->SwitchParameterPinReference.Get();
 
 				// Check Switch Parameter arity preconditions.
 				if (const UEdGraphPin* ConnectedPin = FollowInputPin(*SwitchParameter))
@@ -540,7 +540,7 @@ UE::Mutable::Private::Ptr<UE::Mutable::Private::NodeMaterial> GenerateMutableSou
 						return Result;
 					}
 
-					const int32 NumSwitchOptions = TypedNodeMaterialSwitch->GetNumElements();
+					const int32 NumSwitchOptions = TypedNodeMaterialSwitch->SwitchPins.Num();
 
 					UE::Mutable::Private::NodeScalarEnumParameter* EnumParameter = static_cast<UE::Mutable::Private::NodeScalarEnumParameter*>(SwitchParam.get());
 					if (NumSwitchOptions != EnumParameter->Options.Num())
@@ -555,7 +555,7 @@ UE::Mutable::Private::Ptr<UE::Mutable::Private::NodeMaterial> GenerateMutableSou
 
 					for (int SelectorIndex = 0; SelectorIndex < NumSwitchOptions; ++SelectorIndex)
 					{
-						if (const UEdGraphPin* MaterialPin = FollowInputPin(*TypedNodeMaterialSwitch->GetElementPin(SelectorIndex)))
+						if (const UEdGraphPin* MaterialPin = FollowInputPin(*TypedNodeMaterialSwitch->SwitchPins[SelectorIndex].Get()))
 						{
 							SwitchNode->Options[SelectorIndex] = GenerateMutableSourceMaterial(MaterialPin, GenerationContext, MaterialOptions);
 						}
@@ -666,7 +666,7 @@ UE::Mutable::Private::Ptr<UE::Mutable::Private::NodeMaterial> GenerateMutableSou
 					MaterialNode->ImageParametersToCompile.Add({ FName(ParameterInfo.Name), ParameterInfo.Index }, FormatNode);
 				}
 
-				else if (CurrentBreakParameter.ParameterType == UEdGraphSchema_CustomizableObject::PC_PassthroughTexture)
+				else if (CurrentBreakParameter.ParameterType == UEdGraphSchema_CustomizableObject::PC_Texture_Passthrough)
 				{
 					//TODO(Max - GM): Add support for Passthrough texture parameters
 					////Find reference Texture from the material parameter
@@ -755,13 +755,10 @@ UE::Mutable::Private::Ptr<UE::Mutable::Private::NodeMaterial> GenerateMutableSou
 		//		|_ Or instead find the first material constant (or parameter) node that contains the texture parameter and return it
 		//		|_ Or use the reference image of the group projector generated (if any). 
 
-		TArray<UEdGraphPin*> InputPins;
-		TypedNodeMaterialModify->GetInputPins(InputPins);
-
 		// Tracking the index of the image. Needed for the FGeneratedImagePropertyKey
 		int32 ImageIndex = 0;
 	
-		for (const UEdGraphPin* InputPin : InputPins)
+		for (const UEdGraphPin* InputPin : TypedNodeMaterialModify->GetAllNonOrphanPins())
 		{
 			if (InputPin == TypedNodeMaterialModify->MaterialPinRef.Get())
 			{
@@ -810,7 +807,7 @@ UE::Mutable::Private::Ptr<UE::Mutable::Private::NodeMaterial> GenerateMutableSou
 				ImageIndex++;
 			}
 
-			else if (PinType == UEdGraphSchema_CustomizableObject::PC_PassthroughTexture)
+			else if (PinType == UEdGraphSchema_CustomizableObject::PC_Texture_Passthrough)
 			{
 				UE::Mutable::Private::FImageParameterData ImageData;
 
diff --git a/Engine/Plugins/Mutable/Source/CustomizableObjectEditor/Private/MuCOE/GenerateMutableSource/GenerateMutableSourceMesh.cpp b/Engine/Plugins/Mutable/Source/CustomizableObjectEditor/Private/MuCOE/GenerateMutableSource/GenerateMutableSourceMesh.cpp
--- a/Engine/Plugins/Mutable/Source/CustomizableObjectEditor/Private/MuCOE/GenerateMutableSource/GenerateMutableSourceMesh.cpp
+++ b/Engine/Plugins/Mutable/Source/CustomizableObjectEditor/Private/MuCOE/GenerateMutableSource/GenerateMutableSourceMesh.cpp
@@ -4096,7 +4096,7 @@ UE::Mutable::Private::Ptr<UE::Mutable::Private::NodeMesh> GenerateMutableSourceM
 		// Using a lambda so control flow is easier to manage.
 		Result = [&]()
 		{
-			const UEdGraphPin* SwitchParameter = TypedNodeMeshSwitch->SwitchParameter();
+			const UEdGraphPin* SwitchParameter = TypedNodeMeshSwitch->SwitchParameterPinReference.Get();
 
 			// Check Switch Parameter arity preconditions.
 			if (const UEdGraphPin* EnumPin = FollowInputPin(*SwitchParameter))
@@ -4124,7 +4124,7 @@ UE::Mutable::Private::Ptr<UE::Mutable::Private::NodeMesh> GenerateMutableSourceM
 					return Result;
 				}
 
-				const int32 NumSwitchOptions = TypedNodeMeshSwitch->GetNumElements();
+				const int32 NumSwitchOptions = TypedNodeMeshSwitch->SwitchPins.Num();
 
 				UE::Mutable::Private::NodeScalarEnumParameter* EnumParameter = static_cast<UE::Mutable::Private::NodeScalarEnumParameter*>(SwitchParam.get());
 				if (NumSwitchOptions != EnumParameter->Options.Num())
@@ -4139,7 +4139,7 @@ UE::Mutable::Private::Ptr<UE::Mutable::Private::NodeMesh> GenerateMutableSourceM
 
 				for (int SelectorIndex = 0; SelectorIndex < NumSwitchOptions; ++SelectorIndex)
 				{
-					if (const UEdGraphPin* ConnectedPin = FollowInputPin(*TypedNodeMeshSwitch->GetElementPin(SelectorIndex)))
+					if (const UEdGraphPin* ConnectedPin = FollowInputPin(*TypedNodeMeshSwitch->SwitchPins[SelectorIndex].Get()))
 					{
 						Result = GenerateMutableSourceMesh(ConnectedPin, GenerationContext, BaseMeshData, false, bOnlyConnectedLOD);
 						SwitchNode->Options[SelectorIndex] = Result;
diff --git a/Engine/Plugins/Mutable/Source/CustomizableObjectEditor/Private/MuCOE/GenerateMutableSource/GenerateMutableSourceSurface.cpp b/Engine/Plugins/Mutable/Source/CustomizableObjectEditor/Private/MuCOE/GenerateMutableSource/GenerateMutableSourceSurface.cpp
--- a/Engine/Plugins/Mutable/Source/CustomizableObjectEditor/Private/MuCOE/GenerateMutableSource/GenerateMutableSourceSurface.cpp
+++ b/Engine/Plugins/Mutable/Source/CustomizableObjectEditor/Private/MuCOE/GenerateMutableSource/GenerateMutableSourceSurface.cpp
@@ -287,7 +287,7 @@ UE::Mutable::Private::Ptr<UE::Mutable::Private::NodeSurface> GenerateMutableSour
 		// Using a lambda so control flow is easier to manage.
 		Result = [&]()
 		{
-			const UEdGraphPin* SwitchParameter = TypedNodeSwitch->SwitchParameter();
+			const UEdGraphPin* SwitchParameter = TypedNodeSwitch->SwitchParameterPinReference.Get();
 
 			// Check Switch Parameter arity preconditions.
 			if (const UEdGraphPin* EnumPin = FollowInputPin(*SwitchParameter))
@@ -315,7 +315,7 @@ UE::Mutable::Private::Ptr<UE::Mutable::Private::NodeSurface> GenerateMutableSour
 					return Result;
 				}
 
-				const int32 NumSwitchOptions = TypedNodeSwitch->GetNumElements();
+				const int32 NumSwitchOptions = TypedNodeSwitch->SwitchPins.Num();
 
 				UE::Mutable::Private::NodeScalarEnumParameter* EnumParameter = static_cast<UE::Mutable::Private::NodeScalarEnumParameter*>(SwitchParam.get());
 				if (NumSwitchOptions != EnumParameter->Options.Num())
@@ -330,7 +330,7 @@ UE::Mutable::Private::Ptr<UE::Mutable::Private::NodeSurface> GenerateMutableSour
 
 				for (int32 SelectorIndex = 0; SelectorIndex < NumSwitchOptions; ++SelectorIndex)
 				{
-					if (const UEdGraphPin* ConnectedPin = FollowInputPin(*TypedNodeSwitch->GetElementPin(SelectorIndex)))
+					if (const UEdGraphPin* ConnectedPin = FollowInputPin(*TypedNodeSwitch->SwitchPins[SelectorIndex].Get()))
 					{
 						UE::Mutable::Private::NodeSurfacePtr ChildNode = GenerateMutableSourceSurface(ConnectedPin, GenerationContext);
 						if (ChildNode)
diff --git a/Engine/Plugins/Mutable/Source/CustomizableObjectEditor/Private/MuCOE/GraphTraversal.cpp b/Engine/Plugins/Mutable/Source/CustomizableObjectEditor/Private/MuCOE/GraphTraversal.cpp
--- a/Engine/Plugins/Mutable/Source/CustomizableObjectEditor/Private/MuCOE/GraphTraversal.cpp
+++ b/Engine/Plugins/Mutable/Source/CustomizableObjectEditor/Private/MuCOE/GraphTraversal.cpp
@@ -627,11 +627,11 @@ const UEdGraphPin* FindMeshBaseSource(const UEdGraphPin& Pin, const bool bOnlyLo
 
 	else if (const UCustomizableObjectNodeMeshSwitch* TypedNodeSwitch = Cast<UCustomizableObjectNodeMeshSwitch>(Node))
 	{
-		if (const UEdGraphPin* EnumParameterPin = FollowInputPin(*TypedNodeSwitch->SwitchParameter()))
+		if (const UEdGraphPin* EnumParameterPin = FollowInputPin(*TypedNodeSwitch->SwitchParameterPinReference.Get()))
 		{
 			if (const UCustomizableObjectNodeEnumParameter* EnumNode = Cast<UCustomizableObjectNodeEnumParameter>(EnumParameterPin->GetOwningNode()))
 			{
-				if (const UEdGraphPin* DefaultPin = TypedNodeSwitch->GetElementPin(EnumNode->DefaultIndex))
+				if (const UEdGraphPin* DefaultPin = TypedNodeSwitch->SwitchPins[EnumNode->DefaultIndex].Get())
 				{
 					if (const UEdGraphPin* ConnectedPin = FollowInputPin(*DefaultPin))
 					{
diff --git a/Engine/Plugins/Mutable/Source/CustomizableObjectEditor/Private/MuCOE/Nodes/CONodeExternalOperation.cpp b/Engine/Plugins/Mutable/Source/CustomizableObjectEditor/Private/MuCOE/Nodes/CONodeExternalOperation.cpp
--- a/Engine/Plugins/Mutable/Source/CustomizableObjectEditor/Private/MuCOE/Nodes/CONodeExternalOperation.cpp
+++ b/Engine/Plugins/Mutable/Source/CustomizableObjectEditor/Private/MuCOE/Nodes/CONodeExternalOperation.cpp
@@ -57,7 +57,7 @@ void UCONodeExternalOperation::AllocateDefaultPins(UCustomizableObjectNodeRemapP
 		
 		const FName Category = UEdGraphSchema_CustomizableObject::GetPinCategoryName(Pair.Value);
 
-		UEdGraphPin* Pin = CustomCreatePin(EGPD_Input, Category, Pair.Key, PinData);
+		UEdGraphPin* Pin = CustomCreatePin(EGPD_Input, Category, Pair.Key, FText::FromName(Pair.Key), PinData);
 		InputPins.Add(Pin);
 	}
 
@@ -69,7 +69,7 @@ void UCONodeExternalOperation::AllocateDefaultPins(UCustomizableObjectNodeRemapP
 		PinData->Name = Pair.Key;
 
 		const FName Category = UEdGraphSchema_CustomizableObject::GetPinCategoryName(Pair.Value);
-		OutputPin = CustomCreatePin(EGPD_Output, Category, Pair.Key, PinData);
+		OutputPin = CustomCreatePin(EGPD_Output, Category, Pair.Key, FText::FromName(Pair.Key), PinData);
 	}
 
 	CachedOperationName = Operation->GetName();
diff --git a/Engine/Plugins/Mutable/Source/CustomizableObjectEditor/Private/MuCOE/Nodes/CONodeMaterialBreak.cpp b/Engine/Plugins/Mutable/Source/CustomizableObjectEditor/Private/MuCOE/Nodes/CONodeMaterialBreak.cpp
--- a/Engine/Plugins/Mutable/Source/CustomizableObjectEditor/Private/MuCOE/Nodes/CONodeMaterialBreak.cpp
+++ b/Engine/Plugins/Mutable/Source/CustomizableObjectEditor/Private/MuCOE/Nodes/CONodeMaterialBreak.cpp
@@ -17,6 +17,7 @@ struct FSlateBrush;
 #define LOCTEXT_NAMESPACE "CustomizableObjectEditor"
 
 const FName BreakMaterialMaterialParameterPinName = "Parameter";
+const FText BreakMaterialMaterialParameterPinText = LOCTEXT("Parameter", "Parameter");
 
 bool UCustomizableObjectNodeBreakMaterialRemapPins::Equal(const UCustomizableObjectNode& Node, const UEdGraphPin& OldPin, const UEdGraphPin& NewPin) const
 {
@@ -36,15 +37,8 @@ bool UCustomizableObjectNodeBreakMaterialRemapPins::Equal(const UCustomizableObj
 
 void UCONodeMaterialBreak::AllocateDefaultPins(UCustomizableObjectNodeRemapPins* RemapPins)
 {
-	const FName Type = UEdGraphSchema_CustomizableObject::PC_Material;
-	const FName PinName = UEdGraphSchema_CustomizableObject::GetPinCategoryName(Type);
-	const FText PinFriendlyName = UEdGraphSchema_CustomizableObject::GetPinCategoryFriendlyName(Type);
+	MaterialPinRef = CustomCreatePin(EGPD_Input, UEdGraphSchema_CustomizableObject::PC_Material);
 	
-	UEdGraphPin* MaterialPin = CustomCreatePin(EGPD_Input, Type, PinName);
-	MaterialPin->PinFriendlyName = PinFriendlyName;
-
-	MaterialPinRef = FEdGraphPinReference(MaterialPin);
-
 	// Get old pins
 	TArray<UEdGraphPin*> OldPins = GetAllNonOrphanPins();
 
@@ -58,17 +52,19 @@ void UCONodeMaterialBreak::AllocateDefaultPins(UCustomizableObjectNodeRemapPins*
 		
 		if (UCONodeMaterialBreakParameterPinData* PinData = Cast<UCONodeMaterialBreakParameterPinData>(GetPinData(*Pin)))
 		{
-			// Use the same pin data as the old pin since the new one is a copy
-			UEdGraphPin* NewPin = CustomCreatePin(EGPD_Output, PinData->VariableData.Type, BreakMaterialMaterialParameterPinName, PinData);
-
+			FText FriendlyName;
+			
 			if (PinData->LayerIndex == 0)
 			{
-				NewPin->PinFriendlyName = FText::Format(LOCTEXT("BackgroundLayerText", "{0} - Background"), FText::FromName(BreakMaterialMaterialParameterPinName));
+				FriendlyName = FText::Format(LOCTEXT("BackgroundLayerText", "{0} - Background"), BreakMaterialMaterialParameterPinText);
 			}
 			else if (PinData->LayerIndex > 0)
 			{
-				NewPin->PinFriendlyName = FText::Format(LOCTEXT("LayerText", "{0} - Layer {1}"), FText::FromName(BreakMaterialMaterialParameterPinName), FText::FromString(FString::FromInt(PinData->LayerIndex)));
+				FriendlyName = FText::Format(LOCTEXT("LayerText", "{0} - Layer {1}"), BreakMaterialMaterialParameterPinText, FText::FromString(FString::FromInt(PinData->LayerIndex)));
 			}
+
+			// Use the same pin data as the old pin since the new one is a copy
+			CustomCreatePin(EGPD_Output, PinData->VariableData.Type, BreakMaterialMaterialParameterPinName, FriendlyName, PinData);
 		}
 	}
 }
@@ -162,7 +158,7 @@ TArray<FName> UCONodeMaterialBreak::GetAllowedPinViewerCreationTypes() const
 {
 	return { 
 		UEdGraphSchema_CustomizableObject::PC_Texture, 
-		UEdGraphSchema_CustomizableObject::PC_PassthroughTexture, 
+		UEdGraphSchema_CustomizableObject::PC_Texture_Passthrough, 
 		UEdGraphSchema_CustomizableObject::PC_Color, 
 		UEdGraphSchema_CustomizableObject::PC_Float 
 	};
@@ -182,7 +178,7 @@ void UCONodeMaterialBreak::AddPinFromUI()
 	PinData->VariableData.Type = NewPinCategory;
 	PinData->VariableData.Id = FGuid::NewGuid();
 
-	if (CustomCreatePin(EEdGraphPinDirection::EGPD_Output, NewPinCategory, BreakMaterialMaterialParameterPinName, PinData))
+	if (CustomCreatePin(EEdGraphPinDirection::EGPD_Output, NewPinCategory, BreakMaterialMaterialParameterPinName, BreakMaterialMaterialParameterPinText, PinData))
 	{
 		GetGraph()->NotifyNodeChanged(this);
 	}
diff --git a/Engine/Plugins/Mutable/Source/CustomizableObjectEditor/Private/MuCOE/Nodes/CONodeMaterialConstant.cpp b/Engine/Plugins/Mutable/Source/CustomizableObjectEditor/Private/MuCOE/Nodes/CONodeMaterialConstant.cpp
--- a/Engine/Plugins/Mutable/Source/CustomizableObjectEditor/Private/MuCOE/Nodes/CONodeMaterialConstant.cpp
+++ b/Engine/Plugins/Mutable/Source/CustomizableObjectEditor/Private/MuCOE/Nodes/CONodeMaterialConstant.cpp
@@ -14,12 +14,7 @@ struct FSlateBrush;
 
 void UCONodeMaterialConstant::AllocateDefaultPins(UCustomizableObjectNodeRemapPins* RemapPins)
 {
-	const FName Type = UEdGraphSchema_CustomizableObject::PC_Material;
-	const FName PinName = UEdGraphSchema_CustomizableObject::GetPinCategoryName(Type);
-	const FText PinFriendlyName = UEdGraphSchema_CustomizableObject::GetPinCategoryFriendlyName(Type);
-	
-	UEdGraphPin* ValuePin = CustomCreatePin(EGPD_Output, Type, PinName);
-	ValuePin->PinFriendlyName = PinFriendlyName;
+	CustomCreatePin(EGPD_Output, UEdGraphSchema_CustomizableObject::PC_Material);
 }
 
 
diff --git a/Engine/Plugins/Mutable/Source/CustomizableObjectEditor/Private/MuCOE/Nodes/CONodeMaterialModify.cpp b/Engine/Plugins/Mutable/Source/CustomizableObjectEditor/Private/MuCOE/Nodes/CONodeMaterialModify.cpp
--- a/Engine/Plugins/Mutable/Source/CustomizableObjectEditor/Private/MuCOE/Nodes/CONodeMaterialModify.cpp
+++ b/Engine/Plugins/Mutable/Source/CustomizableObjectEditor/Private/MuCOE/Nodes/CONodeMaterialModify.cpp
@@ -11,6 +11,7 @@
 
 
 const FName MaterialModifyMaterialParameterPinName = "Parameter";
+const FText MaterialModifyMaterialParameterPinText = LOCTEXT("Parameter", "Parameter");
 
 
 bool UCONodeMaterialModifyRemapPins::Equal(const UCustomizableObjectNode& Node, const UEdGraphPin& OldPin, const UEdGraphPin& NewPin) const
@@ -68,14 +69,10 @@ void UCONodeMaterialModify::AllocateDefaultPins(UCustomizableObjectNodeRemapPins
 	const FName PinName = UEdGraphSchema_CustomizableObject::GetPinCategoryName(Type);
 	const FText PinFriendlyName = UEdGraphSchema_CustomizableObject::GetPinCategoryFriendlyName(Type);
 	
-	UEdGraphPin* InputMaterialPin = CustomCreatePin(EGPD_Input, Type, PinName);
-	InputMaterialPin->PinFriendlyName = PinFriendlyName;
-
-	MaterialPinRef = FEdGraphPinReference(InputMaterialPin);
-
+	MaterialPinRef = CustomCreatePin(EGPD_Input, Type, PinName, PinFriendlyName);
+	
 	//Output Material Pin
-	UEdGraphPin* OutputMaterialPin = CustomCreatePin(EGPD_Output, Type, PinName);
-	OutputMaterialPin->PinFriendlyName = PinFriendlyName;
+	CustomCreatePin(EGPD_Output, Type, PinName, PinFriendlyName);
 
 	// Parameter Pins
 	// Get old pins
@@ -99,7 +96,7 @@ void UCONodeMaterialModify::AllocateDefaultPins(UCustomizableObjectNodeRemapPins
 				{
 					NewPinData = NewObject<UCONodeMaterialModifyTextureParamPinData>(this);
 				}
-				else if (OldPinData->VariableData.Type == UEdGraphSchema_CustomizableObject::PC_PassthroughTexture)
+				else if (OldPinData->VariableData.Type == UEdGraphSchema_CustomizableObject::PC_Texture_Passthrough)
 				{
 					NewPinData = NewObject<UCONodeMaterialModifyPassthroughTextureParamPinData>(this);
 				}
@@ -116,17 +113,17 @@ void UCONodeMaterialModify::AllocateDefaultPins(UCustomizableObjectNodeRemapPins
 					NewPinData->VariableData.Name = OldPinData->VariableData.Name;
 					NewPinData->LayerIndex = OldPinData->LayerIndex;
 					
-					NewPin = CustomCreatePin(EGPD_Input, NewPinData->VariableData.Type, MaterialModifyMaterialParameterPinName, NewPinData);
+					NewPin = CustomCreatePin(EGPD_Input, NewPinData->VariableData.Type, MaterialModifyMaterialParameterPinName, MaterialModifyMaterialParameterPinText, NewPinData);
 				}
 				
 				check(NewPin)
 				if (OldPinData->LayerIndex == 0)
 				{
-					NewPin->PinFriendlyName = FText::Format(LOCTEXT("BackgroundLayerText","{0} - Background"), FText::FromName(MaterialModifyMaterialParameterPinName));
+					NewPin->PinFriendlyName = FText::Format(LOCTEXT("BackgroundLayerText","{0} - Background"), MaterialModifyMaterialParameterPinText);
 				}
 				else if (OldPinData->LayerIndex > 0)
 				{
-					NewPin->PinFriendlyName = FText::Format(LOCTEXT("LayerText", "{0} - Layer {1}"), FText::FromName(MaterialModifyMaterialParameterPinName), FText::FromString(FString::FromInt(OldPinData->LayerIndex)));
+					NewPin->PinFriendlyName = FText::Format(LOCTEXT("LayerText", "{0} - Layer {1}"), MaterialModifyMaterialParameterPinText, FText::FromString(FString::FromInt(OldPinData->LayerIndex)));
 				}
 			}
 		}
@@ -177,7 +174,7 @@ UCustomizableObjectNodeRemapPins* UCONodeMaterialModify::CreateRemapPinsDefault(
 TArray<FName> UCONodeMaterialModify::GetAllowedPinViewerCreationTypes() const
 {
 	return { GetDefault<UEdGraphSchema_CustomizableObject>()->PC_Texture,
-		GetDefault<UEdGraphSchema_CustomizableObject>()->PC_PassthroughTexture,
+		GetDefault<UEdGraphSchema_CustomizableObject>()->PC_Texture_Passthrough,
 		GetDefault<UEdGraphSchema_CustomizableObject>()->PC_Color,
 		GetDefault<UEdGraphSchema_CustomizableObject>()->PC_Float };
 }
@@ -191,7 +188,7 @@ void UCONodeMaterialModify::AddPinFromUI()
 	Modify();
 	
 	UCONodeMaterialModifyScalarParamPinData* ScalarParameterPinData = NewObject<UCONodeMaterialModifyScalarParamPinData>(this);
-	if (CustomCreatePin(EEdGraphPinDirection::EGPD_Input, ScalarParameterPinData->VariableData.Type, MaterialModifyMaterialParameterPinName, ScalarParameterPinData))
+	if (CustomCreatePin(EEdGraphPinDirection::EGPD_Input, ScalarParameterPinData->VariableData.Type, MaterialModifyMaterialParameterPinName, MaterialModifyMaterialParameterPinText, ScalarParameterPinData))
 	{
 		GetGraph()->NotifyNodeChanged(this);
 	}
diff --git a/Engine/Plugins/Mutable/Source/CustomizableObjectEditor/Private/MuCOE/Nodes/CONodeMaterialModify.h b/Engine/Plugins/Mutable/Source/CustomizableObjectEditor/Private/MuCOE/Nodes/CONodeMaterialModify.h
--- a/Engine/Plugins/Mutable/Source/CustomizableObjectEditor/Private/MuCOE/Nodes/CONodeMaterialModify.h
+++ b/Engine/Plugins/Mutable/Source/CustomizableObjectEditor/Private/MuCOE/Nodes/CONodeMaterialModify.h
@@ -119,7 +119,7 @@ class UCONodeMaterialModifyPassthroughTextureParamPinData : public UCONodeMateri
 public:
 	UCONodeMaterialModifyPassthroughTextureParamPinData()
 	{
-		VariableData.Type = UEdGraphSchema_CustomizableObject::PC_PassthroughTexture;
+		VariableData.Type = UEdGraphSchema_CustomizableObject::PC_Texture_Passthrough;
 	}
 };
 
diff --git a/Engine/Plugins/Mutable/Source/CustomizableObjectEditor/Private/MuCOE/Nodes/CONodeModifierTransformWithBone.cpp b/Engine/Plugins/Mutable/Source/CustomizableObjectEditor/Private/MuCOE/Nodes/CONodeModifierTransformWithBone.cpp
--- a/Engine/Plugins/Mutable/Source/CustomizableObjectEditor/Private/MuCOE/Nodes/CONodeModifierTransformWithBone.cpp
+++ b/Engine/Plugins/Mutable/Source/CustomizableObjectEditor/Private/MuCOE/Nodes/CONodeModifierTransformWithBone.cpp
@@ -20,15 +20,9 @@ FText UCONodeModifierTransformWithBone::GetTooltipText() const
 
 void UCONodeModifierTransformWithBone::AllocateDefaultPins(UCustomizableObjectNodeRemapPins* RemapPins)
 {
-	const UEdGraphSchema_CustomizableObject* Schema = GetDefault<UEdGraphSchema_CustomizableObject>();
+	TransformPin = CustomCreatePin(EGPD_Input, UEdGraphSchema_CustomizableObject::PC_Transform);
 
-	UEdGraphPin* NewTransformPin = CustomCreatePin(EGPD_Input, Schema->PC_Transform, FName("Transform"));
-	NewTransformPin->bDefaultValueIsIgnored = true;
-	TransformPin = NewTransformPin;
-
-	UEdGraphPin* OutPin = CustomCreatePin(EGPD_Output, Schema->PC_Modifier, FName("Modifier"));
-	OutPin->bDefaultValueIsIgnored = true;
-	OutputPin = OutPin;
+	OutputPin = CustomCreatePin(EGPD_Output, UEdGraphSchema_CustomizableObject::PC_Modifier);
 
 	//Create Node Modifier Common Pins
 	Super::AllocateDefaultPins(RemapPins);
diff --git a/Engine/Plugins/Mutable/Source/CustomizableObjectEditor/Private/MuCOE/Nodes/CONodeSkeletalMeshBreak.cpp b/Engine/Plugins/Mutable/Source/CustomizableObjectEditor/Private/MuCOE/Nodes/CONodeSkeletalMeshBreak.cpp
--- a/Engine/Plugins/Mutable/Source/CustomizableObjectEditor/Private/MuCOE/Nodes/CONodeSkeletalMeshBreak.cpp
+++ b/Engine/Plugins/Mutable/Source/CustomizableObjectEditor/Private/MuCOE/Nodes/CONodeSkeletalMeshBreak.cpp
@@ -23,14 +23,14 @@ FText UCONodeSkeletalMeshBreak::GetTooltipText() const
 
 FLinearColor UCONodeSkeletalMeshBreak::GetNodeTitleColor() const
 {
-	return UEdGraphSchema_CustomizableObject::GetPinTypeColor(UEdGraphSchema_CustomizableObject::PC_PassthroughSkeletalMesh);
+	return UEdGraphSchema_CustomizableObject::GetPinTypeColor(UEdGraphSchema_CustomizableObject::PC_SkeletalMesh_Passthrough);
 }
 
 
 void UCONodeSkeletalMeshBreak::AllocateDefaultPins(UCustomizableObjectNodeRemapPins* RemapPins)
 {
 	// Passthrough skm input
-	InputPassthroughMesh = CustomCreatePinSimple(EGPD_Input, UEdGraphSchema_CustomizableObject::PC_PassthroughSkeletalMesh);
+	InputPassthroughMesh = CustomCreatePin(EGPD_Input, UEdGraphSchema_CustomizableObject::PC_SkeletalMesh_Passthrough);
 	
 	// Regenerate the new pins from the old pin's PinData which may have changed
 	OutputMaterialPins.Reset();
@@ -108,7 +108,7 @@ void UCONodeSkeletalMeshBreak::AddPinFromUI()
 	FScopedTransaction Transaction(LOCTEXT("UCONodeSkeletalMeshBreak_AddPinFromNodeUITransaction", "Add Material Output Pin from Node UI"));
 	Modify();
 
-	UCONodeSkeletalMeshBreakPinData* PinData = NewObject<UCONodeSkeletalMeshBreakPinData>(this);;
+	UCONodeSkeletalMeshBreakPinData* PinData = NewObject<UCONodeSkeletalMeshBreakPinData>(this);
 	if (CreateMaterialOutputPin(PinData))
 	{
 		GetGraph()->NotifyNodeChanged(this);
@@ -141,9 +141,9 @@ bool UCONodeSkeletalMeshBreak::CanRenamePin(const UEdGraphPin& Pin) const
 }
 
 
-bool UCONodeSkeletalMeshBreak::CreateMaterialOutputPin(const UCONodeSkeletalMeshBreakPinData* InPinData)
+bool UCONodeSkeletalMeshBreak::CreateMaterialOutputPin(UCONodeSkeletalMeshBreakPinData* InPinData)
 {
-	UEdGraphPin* NewPin = CustomCreatePin(EGPD_Output, UEdGraphSchema_CustomizableObject::PC_Material, FName{"Material Slot"}, (UCONodeSkeletalMeshBreakPinData*)InPinData);;
+	UEdGraphPin* NewPin = CustomCreatePin(EGPD_Output, UEdGraphSchema_CustomizableObject::PC_Material, TEXT("Material Slot"), LOCTEXT("MaterialSlot", "Material Slot"), InPinData);
 	if (NewPin)
 	{
 		OutputMaterialPins.Add(NewPin);
diff --git a/Engine/Plugins/Mutable/Source/CustomizableObjectEditor/Private/MuCOE/Nodes/CONodeSkeletalMeshBreak.h b/Engine/Plugins/Mutable/Source/CustomizableObjectEditor/Private/MuCOE/Nodes/CONodeSkeletalMeshBreak.h
--- a/Engine/Plugins/Mutable/Source/CustomizableObjectEditor/Private/MuCOE/Nodes/CONodeSkeletalMeshBreak.h
+++ b/Engine/Plugins/Mutable/Source/CustomizableObjectEditor/Private/MuCOE/Nodes/CONodeSkeletalMeshBreak.h
@@ -54,7 +54,7 @@ private:
 
 	bool IsMaterialOutputPin(const UEdGraphPin& Pin) const;
 
-	bool CreateMaterialOutputPin(const UCONodeSkeletalMeshBreakPinData* InPinData);
+	bool CreateMaterialOutputPin(UCONodeSkeletalMeshBreakPinData* InPinData);
 
 protected:
 	
diff --git a/Engine/Plugins/Mutable/Source/CustomizableObjectEditor/Private/MuCOE/Nodes/CustomizableObjectNode.cpp b/Engine/Plugins/Mutable/Source/CustomizableObjectEditor/Private/MuCOE/Nodes/CustomizableObjectNode.cpp
--- a/Engine/Plugins/Mutable/Source/CustomizableObjectEditor/Private/MuCOE/Nodes/CustomizableObjectNode.cpp
+++ b/Engine/Plugins/Mutable/Source/CustomizableObjectEditor/Private/MuCOE/Nodes/CustomizableObjectNode.cpp
@@ -40,46 +40,25 @@ bool UCustomizableObjectNode::IsSingleOutputNode() const
 }
 
 
-UEdGraphPin* UCustomizableObjectNode::CustomCreatePinSimple(EEdGraphPinDirection Direction, const FName& Category, bool bIsArray)
+UEdGraphPin* UCustomizableObjectNode::CustomCreatePin(EEdGraphPinDirection Direction, const FName& Category, UCustomizableObjectNodePinData* PinData)
 {
 	const FName NewPinName = UEdGraphSchema_CustomizableObject::GetPinCategoryName(Category);
 	const FText NewPinFriendlyName = UEdGraphSchema_CustomizableObject::GetPinCategoryFriendlyName(Category);
 
-	UEdGraphPin* Pin = CreatePin(Direction, Category, NewPinName);
-	check(Pin);
-	Pin->PinFriendlyName = NewPinFriendlyName;
-	
-	if (bIsArray)
-	{
-		Pin->PinType.ContainerType = EPinContainerType::Array;
-	}
-
-	return Pin;
+	return CustomCreatePin(Direction, Category, NewPinName, NewPinFriendlyName, PinData);
 }
 
 
-UEdGraphPin* UCustomizableObjectNode::CustomCreatePin(EEdGraphPinDirection Direction, const FName& Category, const FName& Name, bool bIsArray)
+UEdGraphPin* UCustomizableObjectNode::CustomCreatePin(EEdGraphPinDirection Direction, const FName& Category, const FName& Name, const FText& FriendlyName, UCustomizableObjectNodePinData* PinData)
 {
 	UEdGraphPin* Pin = CreatePin(Direction, Category, Name);
-
-	Pin->PinFriendlyName = FText::FromName(Name);
-	if (bIsArray)
-	{
-		Pin->PinType.ContainerType = EPinContainerType::Array;
-	}
-
-	return Pin;
-}
-
-
-UEdGraphPin* UCustomizableObjectNode::CustomCreatePin(EEdGraphPinDirection Direction, const FName& Category, const FName& Name, UCustomizableObjectNodePinData* PinData)
-{
-	UEdGraphPin* Pin = CreatePin(Direction, Category, Name);
-	if (Pin && PinData)
+	Pin->PinFriendlyName = FriendlyName;
+	
+	if (PinData)
 	{
 		AddPinData(*Pin, *PinData);		
 	}
-	
+
 	return Pin;
 }
 
@@ -135,55 +114,6 @@ bool UCustomizableObjectNode::ShouldAddToContextMenu(FText& OutCategory) const
 }
 
 
-void UCustomizableObjectNode::GetInputPins(TArray<class UEdGraphPin*>& OutInputPins) const
-{
-	OutInputPins.Empty();
-
-	for (int32 PinIndex = 0; PinIndex < Pins.Num(); PinIndex++)
-	{
-		if (Pins[PinIndex]->Direction == EGPD_Input)
-		{
-			OutInputPins.Add(Pins[PinIndex]);
-		}
-	}
-}
-
-
-void UCustomizableObjectNode::GetOutputPins(TArray<class UEdGraphPin*>& OutOutputPins) const
-{
-	OutOutputPins.Empty();
-
-	for (int32 PinIndex = 0; PinIndex < Pins.Num(); PinIndex++)
-	{
-		if (Pins[PinIndex]->Direction == EGPD_Output)
-		{
-			OutOutputPins.Add(Pins[PinIndex]);
-		}
-	}
-}
-
-
-UEdGraphPin* UCustomizableObjectNode::GetOutputPin(int32 OutputIndex) const
-{
-	for (int32 PinIndex = 0, FoundOutputs = 0; PinIndex < Pins.Num(); PinIndex++)
-	{
-		if (Pins[PinIndex]->Direction == EGPD_Output)
-		{
-			if (OutputIndex == FoundOutputs)
-			{
-				return Pins[PinIndex];
-			}
-			else
-			{
-				FoundOutputs++;
-			}
-		}
-	}
-
-	return NULL;
-}
-
-
 void UCustomizableObjectNode::SetRefreshNodeWarning()
 {
 	if (!bHasCompilerMessage && ErrorType < EMessageSeverity::Warning)
@@ -441,7 +371,10 @@ bool UCustomizableObjectNode::CanConnect( const UEdGraphPin* InOwnedInputPin, co
 
 	bOutArePinsCompatible = InOwnedInputPin->PinType.PinCategory == InOutputPin->PinType.PinCategory ||
 		InOwnedInputPin->PinType.PinCategory == UEdGraphSchema_CustomizableObject::PC_Wildcard ||
-		InOutputPin->PinType.PinCategory == UEdGraphSchema_CustomizableObject::PC_Wildcard;
+		InOutputPin->PinType.PinCategory == UEdGraphSchema_CustomizableObject::PC_Wildcard ||
+		(InOutputPin->PinType.PinCategory == UEdGraphSchema_CustomizableObject::PC_Texture_PassthroughParameter && InOwnedInputPin->PinType.PinCategory == UEdGraphSchema_CustomizableObject::PC_Texture_Passthrough) ||
+		(InOutputPin->PinType.PinCategory == UEdGraphSchema_CustomizableObject::PC_SkeletalMesh_PassthroughParameter && InOwnedInputPin->PinType.PinCategory == UEdGraphSchema_CustomizableObject::PC_SkeletalMesh_Passthrough);
+		
 	
 	return bOutArePinsCompatible;
 }
diff --git a/Engine/Plugins/Mutable/Source/CustomizableObjectEditor/Private/MuCOE/Nodes/CustomizableObjectNodeAnimationPose.cpp b/Engine/Plugins/Mutable/Source/CustomizableObjectEditor/Private/MuCOE/Nodes/CustomizableObjectNodeAnimationPose.cpp
--- a/Engine/Plugins/Mutable/Source/CustomizableObjectEditor/Private/MuCOE/Nodes/CustomizableObjectNodeAnimationPose.cpp
+++ b/Engine/Plugins/Mutable/Source/CustomizableObjectEditor/Private/MuCOE/Nodes/CustomizableObjectNodeAnimationPose.cpp
@@ -25,14 +25,9 @@ UCustomizableObjectNodeAnimationPose::UCustomizableObjectNodeAnimationPose() : S
 
 void UCustomizableObjectNodeAnimationPose::AllocateDefaultPins(UCustomizableObjectNodeRemapPins* RemapPins)
 {
-	UEdGraphPin* InMeshPin = CustomCreatePinSimple(EGPD_Input, UEdGraphSchema_CustomizableObject::PC_Mesh);
-	InMeshPin->bDefaultValueIsIgnored = true;
-
-	UEdGraphPin* TablePosePin = CustomCreatePinSimple(EGPD_Input, UEdGraphSchema_CustomizableObject::PC_PoseAsset);
-	TablePosePin->bDefaultValueIsIgnored = true;
-
-	UEdGraphPin* OutMeshPin = CustomCreatePinSimple(EGPD_Output,  UEdGraphSchema_CustomizableObject::PC_Mesh);
-	OutMeshPin->bDefaultValueIsIgnored = true;
+	CustomCreatePin(EGPD_Input, UEdGraphSchema_CustomizableObject::PC_Mesh);
+	CustomCreatePin(EGPD_Input, UEdGraphSchema_CustomizableObject::PC_PoseAsset);
+	CustomCreatePin(EGPD_Output,  UEdGraphSchema_CustomizableObject::PC_Mesh);
 }
 
 void UCustomizableObjectNodeAnimationPose::BackwardsCompatibleFixup(int32 CustomizableObjectCustomVersion)
diff --git a/Engine/Plugins/Mutable/Source/CustomizableObjectEditor/Private/MuCOE/Nodes/CustomizableObjectNodeColorArithmeticOp.cpp b/Engine/Plugins/Mutable/Source/CustomizableObjectEditor/Private/MuCOE/Nodes/CustomizableObjectNodeColorArithmeticOp.cpp
--- a/Engine/Plugins/Mutable/Source/CustomizableObjectEditor/Private/MuCOE/Nodes/CustomizableObjectNodeColorArithmeticOp.cpp
+++ b/Engine/Plugins/Mutable/Source/CustomizableObjectEditor/Private/MuCOE/Nodes/CustomizableObjectNodeColorArithmeticOp.cpp
@@ -32,21 +32,9 @@ void UCustomizableObjectNodeColorArithmeticOp::PostEditChangeProperty(FPropertyC
 
 void UCustomizableObjectNodeColorArithmeticOp::AllocateDefaultPins(UCustomizableObjectNodeRemapPins* RemapPins)
 {
-	const UEdGraphSchema_CustomizableObject* Schema = GetDefault<UEdGraphSchema_CustomizableObject>();
-
-	UEdGraphPin* XPin = CustomCreatePin(EGPD_Input, Schema->PC_Color, FName("A"));
-	XPin->bDefaultValueIsIgnored = true;
-	XPin->PinFriendlyName = FText();
-
-	FName PinName = TEXT("B");
-	UEdGraphPin* YPin = CustomCreatePin(EGPD_Input, Schema->PC_Color, PinName);
-	YPin->bDefaultValueIsIgnored = true;
-	YPin->PinFriendlyName = FText();
-
-	PinName = TEXT("Result");
-	UEdGraphPin* ColorPin = CustomCreatePin(EGPD_Output, Schema->PC_Color,PinName);
-	ColorPin->bDefaultValueIsIgnored = true;
-	ColorPin->PinFriendlyName = FText();
+	CustomCreatePin(EGPD_Input, UEdGraphSchema_CustomizableObject::PC_Color, TEXT("A"), LOCTEXT("A", "A"));
+	CustomCreatePin(EGPD_Input, UEdGraphSchema_CustomizableObject::PC_Color,  TEXT("B"), LOCTEXT("B", "B"));
+	CustomCreatePin(EGPD_Output, UEdGraphSchema_CustomizableObject::PC_Color, TEXT("Result"), LOCTEXT("Result", "Result"));
 }
 
 
diff --git a/Engine/Plugins/Mutable/Source/CustomizableObjectEditor/Private/MuCOE/Nodes/CustomizableObjectNodeColorConstant.cpp b/Engine/Plugins/Mutable/Source/CustomizableObjectEditor/Private/MuCOE/Nodes/CustomizableObjectNodeColorConstant.cpp
--- a/Engine/Plugins/Mutable/Source/CustomizableObjectEditor/Private/MuCOE/Nodes/CustomizableObjectNodeColorConstant.cpp
+++ b/Engine/Plugins/Mutable/Source/CustomizableObjectEditor/Private/MuCOE/Nodes/CustomizableObjectNodeColorConstant.cpp
@@ -21,23 +21,9 @@ struct FSlateBrush;
 #define LOCTEXT_NAMESPACE "CustomizableObjectEditor"
 
 
-UCustomizableObjectNodeColorConstant::UCustomizableObjectNodeColorConstant()
-	: Super()
-{
-	Value = FLinearColor(1,1,1,1);
-	bCollapsed = true;
-}
-
-
 void UCustomizableObjectNodeColorConstant::AllocateDefaultPins(UCustomizableObjectNodeRemapPins* RemapPins)
 {
-	const FName Type = UEdGraphSchema_CustomizableObject::PC_Color;
-	const FName PinName = UEdGraphSchema_CustomizableObject::GetPinCategoryName(Type);
-	const FText PinFriendlyName = UEdGraphSchema_CustomizableObject::GetPinCategoryFriendlyName(Type);
-	
-	UEdGraphPin* ValuePin = CustomCreatePin(EGPD_Output, Type, PinName);
-	ValuePin->PinFriendlyName = PinFriendlyName;
-	ValuePin->bDefaultValueIsIgnored = true;
+	CustomCreatePin(EGPD_Output, UEdGraphSchema_CustomizableObject::PC_Color);
 }
 
 
diff --git a/Engine/Plugins/Mutable/Source/CustomizableObjectEditor/Private/MuCOE/Nodes/CustomizableObjectNodeColorConstant.h b/Engine/Plugins/Mutable/Source/CustomizableObjectEditor/Private/MuCOE/Nodes/CustomizableObjectNodeColorConstant.h
--- a/Engine/Plugins/Mutable/Source/CustomizableObjectEditor/Private/MuCOE/Nodes/CustomizableObjectNodeColorConstant.h
+++ b/Engine/Plugins/Mutable/Source/CustomizableObjectEditor/Private/MuCOE/Nodes/CustomizableObjectNodeColorConstant.h
@@ -81,12 +81,6 @@ class UCustomizableObjectNodeColorConstant : public UCustomizableObjectNode
 {
 public:
 	GENERATED_BODY()
-
-	UE_API UCustomizableObjectNodeColorConstant();
-
-	/**  */
-	UPROPERTY(EditAnywhere, Category=CustomizableObject, meta = (DontUpdateWhileEditing))
-	FLinearColor Value;
 	
 	// Begin EdGraphNode interface
 	UE_API virtual FText GetNodeTitle(ENodeTitleType::Type TitleType) const override;
@@ -102,8 +96,11 @@ public:
 	// Creates the SGraph Node widget for the Color Editor
 	UE_API virtual TSharedPtr<SGraphNode> CreateVisualWidget() override;
 
+	UPROPERTY(EditAnywhere, Category=CustomizableObject, meta = (DontUpdateWhileEditing))
+	FLinearColor Value = FLinearColor(1,1,1,1);
+
 	// Determines if the Node is collapsed or not
-	bool bCollapsed;
+	bool bCollapsed = true;
 
 };
 
diff --git a/Engine/Plugins/Mutable/Source/CustomizableObjectEditor/Private/MuCOE/Nodes/CustomizableObjectNodeColorFromFloats.cpp b/Engine/Plugins/Mutable/Source/CustomizableObjectEditor/Private/MuCOE/Nodes/CustomizableObjectNodeColorFromFloats.cpp
--- a/Engine/Plugins/Mutable/Source/CustomizableObjectEditor/Private/MuCOE/Nodes/CustomizableObjectNodeColorFromFloats.cpp
+++ b/Engine/Plugins/Mutable/Source/CustomizableObjectEditor/Private/MuCOE/Nodes/CustomizableObjectNodeColorFromFloats.cpp
@@ -27,22 +27,11 @@ void UCustomizableObjectNodeColorFromFloats::PostEditChangeProperty(FPropertyCha
 
 void UCustomizableObjectNodeColorFromFloats::AllocateDefaultPins(UCustomizableObjectNodeRemapPins* RemapPins)
 {
-	const UEdGraphSchema_CustomizableObject* Schema = GetDefault<UEdGraphSchema_CustomizableObject>();
-
-	UEdGraphPin* RPin = CustomCreatePin(EGPD_Input, Schema->PC_Float, FName("R"));
-	RPin->bDefaultValueIsIgnored = true;
-
-	UEdGraphPin* GPin = CustomCreatePin(EGPD_Input, Schema->PC_Float, FName("G"));
-	GPin->bDefaultValueIsIgnored = true;
-
-	UEdGraphPin* BPin = CustomCreatePin(EGPD_Input, Schema->PC_Float, FName("B"));
-	BPin->bDefaultValueIsIgnored = true;
-
-	UEdGraphPin* APin = CustomCreatePin(EGPD_Input, Schema->PC_Float, FName("A"));
-	APin->bDefaultValueIsIgnored = true;
-
-	UEdGraphPin* ColorPin = CustomCreatePin(EGPD_Output, Schema->PC_Color, FName("Color"));
-	ColorPin->bDefaultValueIsIgnored = true;
+	CustomCreatePin(EGPD_Input, UEdGraphSchema_CustomizableObject::PC_Float, TEXT("R"), LOCTEXT("R", "R"));
+	CustomCreatePin(EGPD_Input, UEdGraphSchema_CustomizableObject::PC_Float, TEXT("G"), LOCTEXT("G", "G"));
+	CustomCreatePin(EGPD_Input, UEdGraphSchema_CustomizableObject::PC_Float, TEXT("B"), LOCTEXT("B", "B"));
+	CustomCreatePin(EGPD_Input, UEdGraphSchema_CustomizableObject::PC_Float, TEXT("A"), LOCTEXT("A", "A"));
+	CustomCreatePin(EGPD_Output, UEdGraphSchema_CustomizableObject::PC_Color, TEXT("Color"), LOCTEXT("Color", "Color"));
 }
 
 
diff --git a/Engine/Plugins/Mutable/Source/CustomizableObjectEditor/Private/MuCOE/Nodes/CustomizableObjectNodeColorSwitch.cpp b/Engine/Plugins/Mutable/Source/CustomizableObjectEditor/Private/MuCOE/Nodes/CustomizableObjectNodeColorSwitch.cpp
--- a/Engine/Plugins/Mutable/Source/CustomizableObjectEditor/Private/MuCOE/Nodes/CustomizableObjectNodeColorSwitch.cpp
+++ b/Engine/Plugins/Mutable/Source/CustomizableObjectEditor/Private/MuCOE/Nodes/CustomizableObjectNodeColorSwitch.cpp
@@ -10,17 +10,6 @@
 #define LOCTEXT_NAMESPACE "CustomizableObjectEditor"
 
 
-void UCustomizableObjectNodeColorSwitch::BackwardsCompatibleFixup(int32 CustomizableObjectCustomVersion)
-{
-	Super::BackwardsCompatibleFixup(CustomizableObjectCustomVersion);
-
-	if (CustomizableObjectCustomVersion == FCustomizableObjectCustomVersion::BugPinsSwitch)
-	{
-		OutputPinReference = FindPin(TEXT("Color"));	
-	}
-}
-
-
 FName UCustomizableObjectNodeColorSwitch::GetCategory() const
 {
 	return UEdGraphSchema_CustomizableObject::PC_Color;
diff --git a/Engine/Plugins/Mutable/Source/CustomizableObjectEditor/Private/MuCOE/Nodes/CustomizableObjectNodeColorSwitch.h b/Engine/Plugins/Mutable/Source/CustomizableObjectEditor/Private/MuCOE/Nodes/CustomizableObjectNodeColorSwitch.h
--- a/Engine/Plugins/Mutable/Source/CustomizableObjectEditor/Private/MuCOE/Nodes/CustomizableObjectNodeColorSwitch.h
+++ b/Engine/Plugins/Mutable/Source/CustomizableObjectEditor/Private/MuCOE/Nodes/CustomizableObjectNodeColorSwitch.h
@@ -15,9 +15,6 @@ class UCustomizableObjectNodeColorSwitch : public UCustomizableObjectNodeSwitchB
 public:
 	GENERATED_BODY()
 
-	// UCustomizableObjectNode interface
-	UE_API virtual void BackwardsCompatibleFixup(int32 CustomizableObjectCustomVersion) override;
-
 	// UCustomizableObjectNodeSwitchBase interface
 	UE_API virtual FName GetCategory() const override;
 };
diff --git a/Engine/Plugins/Mutable/Source/CustomizableObjectEditor/Private/MuCOE/Nodes/CustomizableObjectNodeColorToSRGB.cpp b/Engine/Plugins/Mutable/Source/CustomizableObjectEditor/Private/MuCOE/Nodes/CustomizableObjectNodeColorToSRGB.cpp
--- a/Engine/Plugins/Mutable/Source/CustomizableObjectEditor/Private/MuCOE/Nodes/CustomizableObjectNodeColorToSRGB.cpp
+++ b/Engine/Plugins/Mutable/Source/CustomizableObjectEditor/Private/MuCOE/Nodes/CustomizableObjectNodeColorToSRGB.cpp
@@ -11,23 +11,10 @@ class UCustomizableObjectNodeRemapPins;
 #define LOCTEXT_NAMESPACE "CustomizableObjectEditor"
 
 
-UCustomizableObjectNodeColorToSRGB::UCustomizableObjectNodeColorToSRGB()
-	: Super()
-{
-}
-
-
 void UCustomizableObjectNodeColorToSRGB::AllocateDefaultPins(UCustomizableObjectNodeRemapPins* RemapPins)
 {
-	const UEdGraphSchema_CustomizableObject* Schema = GetDefault<UEdGraphSchema_CustomizableObject>();
-	
-	UEdGraphPin* InPin = CustomCreatePin(EGPD_Input, Schema->PC_Color, FName("Linear"));
-	InPin->bDefaultValueIsIgnored = true;
-	InputPin = InPin;
-
-	UEdGraphPin* OutPin = CustomCreatePin(EGPD_Output, Schema->PC_Color, FName("sRGB"));
-	OutPin->bDefaultValueIsIgnored = true;
-	OutputPin = OutPin;
+	InputPin = CustomCreatePin(EGPD_Input, UEdGraphSchema_CustomizableObject::PC_Color, TEXT("Linear"), LOCTEXT("Linear", "Linear"));
+	OutputPin = CustomCreatePin(EGPD_Output, UEdGraphSchema_CustomizableObject::PC_Color, TEXT("sRGB"), LOCTEXT("sRGB", "sRGB"));
 }
 
 
diff --git a/Engine/Plugins/Mutable/Source/CustomizableObjectEditor/Private/MuCOE/Nodes/CustomizableObjectNodeColorToSRGB.h b/Engine/Plugins/Mutable/Source/CustomizableObjectEditor/Private/MuCOE/Nodes/CustomizableObjectNodeColorToSRGB.h
--- a/Engine/Plugins/Mutable/Source/CustomizableObjectEditor/Private/MuCOE/Nodes/CustomizableObjectNodeColorToSRGB.h
+++ b/Engine/Plugins/Mutable/Source/CustomizableObjectEditor/Private/MuCOE/Nodes/CustomizableObjectNodeColorToSRGB.h
@@ -19,8 +19,6 @@ class UCustomizableObjectNodeColorToSRGB : public UCustomizableObjectNode
 public:
 	GENERATED_BODY()
 
-	UE_API UCustomizableObjectNodeColorToSRGB();
-
 	// Begin EdGraphNode interface
 	UE_API FText GetNodeTitle(ENodeTitleType::Type TitleType) const override;
 	UE_API FLinearColor GetNodeTitleColor() const override;
diff --git a/Engine/Plugins/Mutable/Source/CustomizableObjectEditor/Private/MuCOE/Nodes/CustomizableObjectNodeComponent.cpp b/Engine/Plugins/Mutable/Source/CustomizableObjectEditor/Private/MuCOE/Nodes/CustomizableObjectNodeComponent.cpp
--- a/Engine/Plugins/Mutable/Source/CustomizableObjectEditor/Private/MuCOE/Nodes/CustomizableObjectNodeComponent.cpp
+++ b/Engine/Plugins/Mutable/Source/CustomizableObjectEditor/Private/MuCOE/Nodes/CustomizableObjectNodeComponent.cpp
@@ -15,8 +15,8 @@
 
 void UCustomizableObjectNodeComponent::AllocateDefaultPins(UCustomizableObjectNodeRemapPins* RemapPins)
 {
-	OutputPin = CustomCreatePin(EGPD_Output, UEdGraphSchema_CustomizableObject::PC_Component, FName(TEXT("Component")));
-	ComponentNamePin = CustomCreatePin(EGPD_Input, UEdGraphSchema_CustomizableObject::PC_String, FName("Name"));
+	OutputPin = CustomCreatePin(EGPD_Output, UEdGraphSchema_CustomizableObject::PC_Component, TEXT("Component"), LOCTEXT("Component", "Component"));
+	ComponentNamePin = CustomCreatePin(EGPD_Input, UEdGraphSchema_CustomizableObject::PC_String, TEXT("Name"), LOCTEXT("Name", "Name"));
 }
 
 
@@ -40,7 +40,7 @@ void UCustomizableObjectNodeComponent::BackwardsCompatibleFixup(int32 Customizab
 	{
 		if (!ComponentNamePin.Get())
 		{
-			ComponentNamePin = CustomCreatePin(EGPD_Input, UEdGraphSchema_CustomizableObject::PC_String, FName("Name"));
+			ComponentNamePin = CustomCreatePin(EGPD_Input, UEdGraphSchema_CustomizableObject::PC_String, TEXT("Name"), LOCTEXT("Name", "Name"));
 		}
 	}
 }
diff --git a/Engine/Plugins/Mutable/Source/CustomizableObjectEditor/Private/MuCOE/Nodes/CustomizableObjectNodeComponentMesh.cpp b/Engine/Plugins/Mutable/Source/CustomizableObjectEditor/Private/MuCOE/Nodes/CustomizableObjectNodeComponentMesh.cpp
--- a/Engine/Plugins/Mutable/Source/CustomizableObjectEditor/Private/MuCOE/Nodes/CustomizableObjectNodeComponentMesh.cpp
+++ b/Engine/Plugins/Mutable/Source/CustomizableObjectEditor/Private/MuCOE/Nodes/CustomizableObjectNodeComponentMesh.cpp
@@ -63,11 +63,7 @@ void UCustomizableObjectNodeComponentMesh::AllocateDefaultPins(UCustomizableObje
 	Super::AllocateDefaultPins(RemapPins);
 
 	// Mesh Component Pins
-	const FString PinFriendlyName = TEXT("Overlay Material");
-	const FString PinName = PinFriendlyName + FString(TEXT("_Input_Pin"));
-	UEdGraphPin* OverlayMaterialPin = CustomCreatePin(EGPD_Input, UEdGraphSchema_CustomizableObject::PC_Material, FName(*PinName));
-	OverlayMaterialPin->PinFriendlyName = FText::FromString(PinFriendlyName);
-	OverlayMaterialPin->bDefaultValueIsIgnored = true;
+	UEdGraphPin* OverlayMaterialPin = CustomCreatePin(EGPD_Input, UEdGraphSchema_CustomizableObject::PC_Material, TEXT("Overlay Material_Input_Pin"), LOCTEXT("OverlayMaterial", "Overlay Material"));
 	OverlayMaterialPin->PinToolTip = "Pin for an Overlay Material from a Table Node";
 	
 	// Base Mesh Interface Pins
@@ -76,7 +72,9 @@ void UCustomizableObjectNodeComponentMesh::AllocateDefaultPins(UCustomizableObje
 	{
 		FString LODName = FString::Printf(TEXT("LOD %d"), LODIndex);
 
-		UEdGraphPin* Pin = CustomCreatePin(EGPD_Input, UEdGraphSchema_CustomizableObject::PC_MeshSection, FName(*LODName), true);
+		UEdGraphPin* Pin = CustomCreatePin(EGPD_Input, UEdGraphSchema_CustomizableObject::PC_MeshSection, FName(*LODName), FText::FromString(LODName));
+		Pin->PinType.ContainerType = EPinContainerType::Array;
+
 		LODPins.Add(Pin);
 	}
 }
diff --git a/Engine/Plugins/Mutable/Source/CustomizableObjectEditor/Private/MuCOE/Nodes/CustomizableObjectNodeComponentMeshAddTo.cpp b/Engine/Plugins/Mutable/Source/CustomizableObjectEditor/Private/MuCOE/Nodes/CustomizableObjectNodeComponentMeshAddTo.cpp
--- a/Engine/Plugins/Mutable/Source/CustomizableObjectEditor/Private/MuCOE/Nodes/CustomizableObjectNodeComponentMeshAddTo.cpp
+++ b/Engine/Plugins/Mutable/Source/CustomizableObjectEditor/Private/MuCOE/Nodes/CustomizableObjectNodeComponentMeshAddTo.cpp
@@ -24,12 +24,14 @@ void UCustomizableObjectNodeComponentMeshAddTo::AllocateDefaultPins(UCustomizabl
 	{
 		FString LODName = FString::Printf(TEXT("LOD %d"), LODIndex);
 
-		UEdGraphPin* Pin = CustomCreatePin(EGPD_Input, UEdGraphSchema_CustomizableObject::PC_MeshSection, FName(*LODName), true);
+		UEdGraphPin* Pin = CustomCreatePin(EGPD_Input, UEdGraphSchema_CustomizableObject::PC_MeshSection, FName(*LODName), FText::FromString(LODName));
+		Pin->PinType.ContainerType = EPinContainerType::Array;
+
 		LODPins.Add(Pin);
 	}
 
-	ParentComponentNamePin = CustomCreatePin(EGPD_Input, UEdGraphSchema_CustomizableObject::PC_String, FName("Component Name"));
-	OutputPin = CustomCreatePin(EGPD_Output, UEdGraphSchema_CustomizableObject::PC_Component, FName(TEXT("Component")));
+	ParentComponentNamePin = CustomCreatePin(EGPD_Input, UEdGraphSchema_CustomizableObject::PC_String, TEXT("Component Name"), LOCTEXT("ComponentName", "Component Name"));
+	OutputPin = CustomCreatePin(EGPD_Output, UEdGraphSchema_CustomizableObject::PC_Component);
 }
 
 
@@ -72,7 +74,7 @@ void UCustomizableObjectNodeComponentMeshAddTo::BackwardsCompatibleFixup(int32 C
 	{
 		if (!ParentComponentNamePin.Get())
 		{
-			ParentComponentNamePin = CustomCreatePin(EGPD_Input, UEdGraphSchema_CustomizableObject::PC_String, FName("Component Name"));
+			ParentComponentNamePin = CustomCreatePin(EGPD_Input, UEdGraphSchema_CustomizableObject::PC_String, TEXT("Component Name"), LOCTEXT("ComponentName", "Component Name"));
 		}
 	}
 }
diff --git a/Engine/Plugins/Mutable/Source/CustomizableObjectEditor/Private/MuCOE/Nodes/CustomizableObjectNodeComponentPassthroughMesh.cpp b/Engine/Plugins/Mutable/Source/CustomizableObjectEditor/Private/MuCOE/Nodes/CustomizableObjectNodeComponentPassthroughMesh.cpp
--- a/Engine/Plugins/Mutable/Source/CustomizableObjectEditor/Private/MuCOE/Nodes/CustomizableObjectNodeComponentPassthroughMesh.cpp
+++ b/Engine/Plugins/Mutable/Source/CustomizableObjectEditor/Private/MuCOE/Nodes/CustomizableObjectNodeComponentPassthroughMesh.cpp
@@ -35,7 +35,7 @@ void UCustomizableObjectNodeComponentPassthroughMesh::BackwardsCompatibleFixup(i
 			}
 		}
 		
-		SkeletalMeshPin = CustomCreatePin(EGPD_Input, UEdGraphSchema_CustomizableObject::PC_PassthroughSkeletalMesh, TEXT("Skeletal Mesh"));
+		SkeletalMeshPin = CustomCreatePin(EGPD_Input, UEdGraphSchema_CustomizableObject::PC_SkeletalMesh_Passthrough);
 	}
 }
 
@@ -44,7 +44,7 @@ void UCustomizableObjectNodeComponentPassthroughMesh::AllocateDefaultPins(UCusto
 {
 	Super::AllocateDefaultPins(RemapPins);
 
-	SkeletalMeshPin = CustomCreatePin(EGPD_Input, UEdGraphSchema_CustomizableObject::PC_PassthroughSkeletalMesh, TEXT("Passthrough Skeletal Mesh"));
+	SkeletalMeshPin = CustomCreatePin(EGPD_Input, UEdGraphSchema_CustomizableObject::PC_SkeletalMesh_Passthrough);
 
 	// Get old pins
 	const TArray<UEdGraphPin*> AllPins = GetAllNonOrphanPins();
@@ -205,8 +205,7 @@ bool UCustomizableObjectNodeComponentPassthroughMesh::CustomRemovePin(UEdGraphPi
 
 UEdGraphPin* UCustomizableObjectNodeComponentPassthroughMesh::CreateMaterialSlotOverridePin(UCONodeComponentPassthroughMeshMaterialOverridePinData* InPinData)
 {
-	return CustomCreatePin(EEdGraphPinDirection::EGPD_Input, UEdGraphSchema_CustomizableObject::PC_Material,
-	"Override Material", InPinData);
+	return CustomCreatePin(EGPD_Input, UEdGraphSchema_CustomizableObject::PC_Material, TEXT("Override Material"), LOCTEXT("OverrideMaterial", "Override Material"), InPinData);
 }
 
 
diff --git a/Engine/Plugins/Mutable/Source/CustomizableObjectEditor/Private/MuCOE/Nodes/CustomizableObjectNodeCopyMaterial.cpp b/Engine/Plugins/Mutable/Source/CustomizableObjectEditor/Private/MuCOE/Nodes/CustomizableObjectNodeCopyMaterial.cpp
--- a/Engine/Plugins/Mutable/Source/CustomizableObjectEditor/Private/MuCOE/Nodes/CustomizableObjectNodeCopyMaterial.cpp
+++ b/Engine/Plugins/Mutable/Source/CustomizableObjectEditor/Private/MuCOE/Nodes/CustomizableObjectNodeCopyMaterial.cpp
@@ -39,17 +39,11 @@ void UCustomizableObjectNodeCopyMaterial::AllocateDefaultPins(UCustomizableObjec
 	const UEdGraphSchema_CustomizableObject* Schema = GetDefault<UEdGraphSchema_CustomizableObject>();
 
 	// Input pins
-	UEdGraphPin* Pin = CustomCreatePin(EGPD_Input, Schema->PC_Mesh, FName(CustomizableObjectNodeCopyMaterial::MeshPinName));
-	Pin->PinFriendlyName = LOCTEXT("MeshPin", "Mesh");
-	Pin->bDefaultValueIsIgnored = true;
-
-	Pin = CustomCreatePin(EGPD_Input, Schema->PC_MeshSection, FName(CustomizableObjectNodeCopyMaterial::MeshSectionPinName));
-	Pin->PinFriendlyName = LOCTEXT("BaseMeshSectionPin", "Base Mesh Section");;
-	Pin->bDefaultValueIsIgnored = true;
+	CustomCreatePin(EGPD_Input, Schema->PC_Mesh, FName(CustomizableObjectNodeCopyMaterial::MeshPinName), LOCTEXT("MeshPin", "Mesh"));
+	CustomCreatePin(EGPD_Input, Schema->PC_MeshSection, FName(CustomizableObjectNodeCopyMaterial::MeshSectionPinName), LOCTEXT("BaseMeshSectionPin", "Base Mesh Section"));
 
 	// Output pins
-	Pin = CustomCreatePin(EGPD_Output, Schema->PC_MeshSection, FName(CustomizableObjectNodeCopyMaterial::OutputPinName));
-	Pin->PinFriendlyName = LOCTEXT("MeshSectionPin", "Mesh Section");
+	CustomCreatePin(EGPD_Output, Schema->PC_MeshSection, FName(CustomizableObjectNodeCopyMaterial::OutputPinName), LOCTEXT("MeshSectionPin", "Mesh Section"));
 }
 
 
diff --git a/Engine/Plugins/Mutable/Source/CustomizableObjectEditor/Private/MuCOE/Nodes/CustomizableObjectNodeCurve.cpp b/Engine/Plugins/Mutable/Source/CustomizableObjectEditor/Private/MuCOE/Nodes/CustomizableObjectNodeCurve.cpp
--- a/Engine/Plugins/Mutable/Source/CustomizableObjectEditor/Private/MuCOE/Nodes/CustomizableObjectNodeCurve.cpp
+++ b/Engine/Plugins/Mutable/Source/CustomizableObjectEditor/Private/MuCOE/Nodes/CustomizableObjectNodeCurve.cpp
@@ -45,8 +45,16 @@ void UCustomizableObjectNodeCurve::Serialize(FArchive& Ar)
 		&& CurveAsset)
 	{
 		TArray<FRichCurveEditInfo> Curves = CurveAsset->GetCurves();
+
 		TArray<UEdGraphPin*> OutputPins;
-		GetOutputPins(OutputPins);
+		for (int32 PinIndex = 0; PinIndex < Pins.Num(); PinIndex++)
+		{
+			if (Pins[PinIndex]->Direction == EGPD_Output)
+			{
+				OutputPins.Add(Pins[PinIndex]);
+			}
+		}
+		
 		for (int32 i = 0; i < Curves.Num(); ++i)
 		{
 			if (UEdGraphPin* CurvePin = FindPin(FString::Printf(TEXT("Curve %d"), i)))
@@ -73,21 +81,16 @@ void UCustomizableObjectNodeCurve::AllocateDefaultPins(UCustomizableObjectNodeRe
 
 	FString RemapIndex = Helper_GetPinName(InputPin());
 
-	UEdGraphPin* InputPin = CustomCreatePin(EGPD_Input, Schema->PC_Float, FName(TEXT("Input")));
-
-	InputPin->bDefaultValueIsIgnored = true;
-
+	CustomCreatePin(EGPD_Input, Schema->PC_Float, TEXT("Input"), LOCTEXT("Input", "Input"));
+	
 	if (CurveAsset)
 	{
 		TArray<FRichCurveEditInfo> Curves = CurveAsset->GetCurves();
-		TArray<UEdGraphPin*> OutputPins;
-		GetOutputPins(OutputPins);
 
 		for (int32 i = 0; i < Curves.Num(); ++i)
 		{
 			FString PinName = *(FString("Curve ") + Curves[i].CurveName.ToString());
-			UEdGraphPin* CurvePin = CustomCreatePin(EGPD_Output, Schema->PC_Float, FName(*PinName));
-			CurvePin->bDefaultValueIsIgnored = true;
+			CustomCreatePin(EGPD_Output, Schema->PC_Float, FName(*PinName), FText::FromString(PinName));
 		}
 	}
 }
diff --git a/Engine/Plugins/Mutable/Source/CustomizableObjectEditor/Private/MuCOE/Nodes/CustomizableObjectNodeExposePin.cpp b/Engine/Plugins/Mutable/Source/CustomizableObjectEditor/Private/MuCOE/Nodes/CustomizableObjectNodeExposePin.cpp
--- a/Engine/Plugins/Mutable/Source/CustomizableObjectEditor/Private/MuCOE/Nodes/CustomizableObjectNodeExposePin.cpp
+++ b/Engine/Plugins/Mutable/Source/CustomizableObjectEditor/Private/MuCOE/Nodes/CustomizableObjectNodeExposePin.cpp
@@ -29,12 +29,12 @@ void UCustomizableObjectNodeExposePin::PostEditChangeProperty(FPropertyChangedEv
 
 void UCustomizableObjectNodeExposePin::AllocateDefaultPins(UCustomizableObjectNodeRemapPins* RemapPins)
 {
-	const FName PinName = UEdGraphSchema_CustomizableObject::GetPinCategoryName(PinType);
-	const FText PinFriendlyName = UEdGraphSchema_CustomizableObject::GetPinCategoryFriendlyName(PinType);
-	
+	UEdGraphPin* InputPin = CustomCreatePin(EGPD_Input, PinType);
 	const bool bIsArrayPinCategory = PinType == UEdGraphSchema_CustomizableObject::PC_GroupProjector;
-	UEdGraphPin* InputPin = CustomCreatePin(EGPD_Input, PinType, PinName, bIsArrayPinCategory);
-	InputPin->PinFriendlyName = PinFriendlyName;
+	if (bIsArrayPinCategory)
+	{
+		InputPin->PinType.ContainerType = EPinContainerType::Array;
+	}
 }
 
 
@@ -130,7 +130,7 @@ void UCustomizableObjectNodeExposePin::BackwardsCompatibleFixup(int32 Customizab
 				Pin->PinName = TEXT("Texture");
 				Pin->PinFriendlyName = LOCTEXT("Image_Pin_Category", "Texture");
 			}
-			else if (PinType == UEdGraphSchema_CustomizableObject::PC_PassthroughTexture)
+			else if (PinType == UEdGraphSchema_CustomizableObject::PC_Texture_Passthrough)
 			{
 				Pin->PinName = TEXT("PassThrough Texture");
 				Pin->PinFriendlyName = LOCTEXT("PassThrough_Image_Pin_Category", "PassThrough Texture");
diff --git a/Engine/Plugins/Mutable/Source/CustomizableObjectEditor/Private/MuCOE/Nodes/CustomizableObjectNodeExtensionDataSwitch.cpp b/Engine/Plugins/Mutable/Source/CustomizableObjectEditor/Private/MuCOE/Nodes/CustomizableObjectNodeExtensionDataSwitch.cpp
--- a/Engine/Plugins/Mutable/Source/CustomizableObjectEditor/Private/MuCOE/Nodes/CustomizableObjectNodeExtensionDataSwitch.cpp
+++ b/Engine/Plugins/Mutable/Source/CustomizableObjectEditor/Private/MuCOE/Nodes/CustomizableObjectNodeExtensionDataSwitch.cpp
@@ -21,19 +21,15 @@ bool UCustomizableObjectNodeExtensionDataSwitch::ShouldAddToContextMenu(FText& O
 	return true;
 }
 
-FString UCustomizableObjectNodeExtensionDataSwitch::GetPinPrefix() const
-{
-	return FString::Format(TEXT("{0} "), { GetOutputPinName() });
-}
 
 UE::Mutable::Private::Ptr<UE::Mutable::Private::NodeExtensionData> UCustomizableObjectNodeExtensionDataSwitch::GenerateMutableNode(FExtensionDataCompilerInterface& InCompilerInterface) const
 {
-	const int32 NumParameters = FollowInputPinArray(*SwitchParameter()).Num();
+	const int32 NumParameters = FollowInputPinArray(*SwitchParameterPinReference.Get()).Num();
 
-	const UEdGraphPin* EnumPin = FollowInputPin(*SwitchParameter());
+	const UEdGraphPin* EnumPin = FollowInputPin(*SwitchParameterPinReference.Get());
 	UE::Mutable::Private::Ptr<UE::Mutable::Private::NodeScalar> SwitchParam = EnumPin ? GenerateMutableSourceFloat(EnumPin, InCompilerInterface.GenerationContext) : nullptr;
 
-	const int32 NumSwitchOptions = GetNumElements();
+	const int32 NumSwitchOptions = SwitchPins.Num();
 
 	UE::Mutable::Private::Ptr<UE::Mutable::Private::NodeExtensionDataSwitch> SwitchNode = new UE::Mutable::Private::NodeExtensionDataSwitch;
 	SwitchNode->SetParameter(SwitchParam);
@@ -41,7 +37,7 @@ UE::Mutable::Private::Ptr<UE::Mutable::Private::NodeExtensionData> UCustomizable
 
 	for (int32 OptionIndex = 0; OptionIndex < NumSwitchOptions; ++OptionIndex)
 	{
-		if (const UEdGraphPin* ExtensionDataPin = GetElementPin(OptionIndex))
+		if (const UEdGraphPin* ExtensionDataPin = SwitchPins[OptionIndex].Get())
 		{
 			if (const UEdGraphPin* ConnectedPin = FollowInputPin(*ExtensionDataPin))
 			{
diff --git a/Engine/Plugins/Mutable/Source/CustomizableObjectEditor/Private/MuCOE/Nodes/CustomizableObjectNodeExtensionDataSwitch.h b/Engine/Plugins/Mutable/Source/CustomizableObjectEditor/Private/MuCOE/Nodes/CustomizableObjectNodeExtensionDataSwitch.h
--- a/Engine/Plugins/Mutable/Source/CustomizableObjectEditor/Private/MuCOE/Nodes/CustomizableObjectNodeExtensionDataSwitch.h
+++ b/Engine/Plugins/Mutable/Source/CustomizableObjectEditor/Private/MuCOE/Nodes/CustomizableObjectNodeExtensionDataSwitch.h
@@ -24,7 +24,6 @@ public:
 
 	//~Begin UCustomizableObjectNodeSwitchBase interface
 	UE_API virtual bool ShouldAddToContextMenu(FText& OutCategory) const override;
-	UE_API virtual FString GetPinPrefix() const override;
 	//~End UCustomizableObjectNodeSwitchBase interface
 
 	//~Begin ICustomizableObjectExtensionNode interface
diff --git a/Engine/Plugins/Mutable/Source/CustomizableObjectEditor/Private/MuCOE/Nodes/CustomizableObjectNodeExtensionDataVariation.cpp b/Engine/Plugins/Mutable/Source/CustomizableObjectEditor/Private/MuCOE/Nodes/CustomizableObjectNodeExtensionDataVariation.cpp
--- a/Engine/Plugins/Mutable/Source/CustomizableObjectEditor/Private/MuCOE/Nodes/CustomizableObjectNodeExtensionDataVariation.cpp
+++ b/Engine/Plugins/Mutable/Source/CustomizableObjectEditor/Private/MuCOE/Nodes/CustomizableObjectNodeExtensionDataVariation.cpp
@@ -41,19 +41,15 @@ void UCustomizableObjectNodeExtensionDataVariation::AllocateDefaultPins(UCustomi
 {
 	const UEdGraphSchema_CustomizableObject* Schema = GetDefault<UEdGraphSchema_CustomizableObject>();
 
-	UEdGraphPin* OutputPin = CustomCreatePin(EGPD_Output, GetCategory(), GetOutputPinName());
-	OutputPin->bDefaultValueIsIgnored = true;
+	CustomCreatePin(EGPD_Output, GetCategory(), GetOutputPinName(), FText::FromName(GetDefaultPinName()));
 
 	for (int32 VariationIndex = Variations.Num() - 1; VariationIndex >= 0; --VariationIndex)
 	{
-		constexpr bool bIsArray = false;
-		UEdGraphPin* VariationPin = CustomCreatePin(EGPD_Input, GetCategory(), GetVariationPinName(VariationIndex), bIsArray);
-		VariationPin->bDefaultValueIsIgnored = true;
-		VariationPin->PinFriendlyName = FText::Format(LOCTEXT("VariationPinFriendlyName", "Variation {0} [{1}]"), { VariationIndex, FText::FromString(Variations[VariationIndex].Tag) });
+		FText FriendlyName = FText::Format(LOCTEXT("VariationPinFriendlyName", "Variation {0} [{1}]"), { VariationIndex, FText::FromString(Variations[VariationIndex].Tag) });
+		CustomCreatePin(EGPD_Input, GetCategory(), GetVariationPinName(VariationIndex), FriendlyName);
 	}
 
-	UEdGraphPin* DefaultVariation = CustomCreatePin(EGPD_Input, GetCategory(), GetDefaultPinName(), false);
-	DefaultVariation->bDefaultValueIsIgnored = true;
+	CustomCreatePin(EGPD_Input, GetCategory(), GetDefaultPinName(), FText::FromName(GetDefaultPinName()));
 }
 
 UE::Mutable::Private::Ptr<UE::Mutable::Private::NodeExtensionData> UCustomizableObjectNodeExtensionDataVariation::GenerateMutableNode(FExtensionDataCompilerInterface& InCompilerInterface) const
diff --git a/Engine/Plugins/Mutable/Source/CustomizableObjectEditor/Private/MuCOE/Nodes/CustomizableObjectNodeExternalPin.cpp b/Engine/Plugins/Mutable/Source/CustomizableObjectEditor/Private/MuCOE/Nodes/CustomizableObjectNodeExternalPin.cpp
--- a/Engine/Plugins/Mutable/Source/CustomizableObjectEditor/Private/MuCOE/Nodes/CustomizableObjectNodeExternalPin.cpp
+++ b/Engine/Plugins/Mutable/Source/CustomizableObjectEditor/Private/MuCOE/Nodes/CustomizableObjectNodeExternalPin.cpp
@@ -74,7 +74,7 @@ void UCustomizableObjectNodeExternalPin::BackwardsCompatibleFixup(int32 Customiz
 				Pin->PinName = TEXT("Texture");
 				Pin->PinFriendlyName = LOCTEXT("Image_Pin_Category", "Texture");
 			}
-			else if (PinType == UEdGraphSchema_CustomizableObject::PC_PassthroughTexture)
+			else if (PinType == UEdGraphSchema_CustomizableObject::PC_Texture_Passthrough)
 			{
 				Pin->PinName = TEXT("PassThrough Texture");
 				Pin->PinFriendlyName = LOCTEXT("PassThrough_Image_Pin_Category", "PassThrough Texture");
@@ -231,12 +231,13 @@ void UCustomizableObjectNodeExternalPin::PropagateConnectionChanged()
 
 void UCustomizableObjectNodeExternalPin::AllocateDefaultPins(UCustomizableObjectNodeRemapPins* RemapPins)
 {
-	const FName PinName = UEdGraphSchema_CustomizableObject::GetPinCategoryName(PinType);
-	const FText PinFriendlyName = UEdGraphSchema_CustomizableObject::GetPinCategoryFriendlyName(PinType);
-
+	UEdGraphPin* OutputPin = CustomCreatePin(EGPD_Output, PinType);
 	const bool bIsArrayPinCategory = PinType == UEdGraphSchema_CustomizableObject::PC_GroupProjector;
-	UEdGraphPin* OutputPin = CustomCreatePin(EGPD_Output, PinType, PinName, bIsArrayPinCategory);
-	OutputPin->PinFriendlyName = PinFriendlyName;
+	if (bIsArrayPinCategory)
+	{
+		OutputPin->PinType.ContainerType = EPinContainerType::Array;
+	}
+
 }
 
 
diff --git a/Engine/Plugins/Mutable/Source/CustomizableObjectEditor/Private/MuCOE/Nodes/CustomizableObjectNodeFloatArithmeticOp.cpp b/Engine/Plugins/Mutable/Source/CustomizableObjectEditor/Private/MuCOE/Nodes/CustomizableObjectNodeFloatArithmeticOp.cpp
--- a/Engine/Plugins/Mutable/Source/CustomizableObjectEditor/Private/MuCOE/Nodes/CustomizableObjectNodeFloatArithmeticOp.cpp
+++ b/Engine/Plugins/Mutable/Source/CustomizableObjectEditor/Private/MuCOE/Nodes/CustomizableObjectNodeFloatArithmeticOp.cpp
@@ -21,21 +21,9 @@ UCustomizableObjectNodeFloatArithmeticOp::UCustomizableObjectNodeFloatArithmetic
 
 void UCustomizableObjectNodeFloatArithmeticOp::AllocateDefaultPins(UCustomizableObjectNodeRemapPins* RemapPins)
 {
-	const UEdGraphSchema_CustomizableObject* Schema = GetDefault<UEdGraphSchema_CustomizableObject>();
-
-	UEdGraphPin* XPin = CustomCreatePin(EGPD_Input, Schema->PC_Float, FName("A"));
-	XPin->bDefaultValueIsIgnored = true;
-	XPin->PinFriendlyName = FText();
-
-	FName PinName = TEXT("B");
-	UEdGraphPin* YPin = CustomCreatePin(EGPD_Input, Schema->PC_Float, PinName);
-	YPin->bDefaultValueIsIgnored = true;
-	YPin->PinFriendlyName = FText();
-
-	PinName = TEXT("Result");
-	UEdGraphPin* ResultPin = CustomCreatePin(EGPD_Output, Schema->PC_Float, PinName);
-	ResultPin->bDefaultValueIsIgnored = true;
-	ResultPin->PinFriendlyName = FText();
+	CustomCreatePin(EGPD_Input, UEdGraphSchema_CustomizableObject::PC_Float, TEXT("A"), LOCTEXT("A", "A"));
+	CustomCreatePin(EGPD_Input, UEdGraphSchema_CustomizableObject::PC_Float, TEXT("B"), LOCTEXT("B", "B"));
+	CustomCreatePin(EGPD_Output, UEdGraphSchema_CustomizableObject::PC_Float, TEXT("Result"), LOCTEXT("Result", "Result"));
 }
 
 
diff --git a/Engine/Plugins/Mutable/Source/CustomizableObjectEditor/Private/MuCOE/Nodes/CustomizableObjectNodeFloatConstant.cpp b/Engine/Plugins/Mutable/Source/CustomizableObjectEditor/Private/MuCOE/Nodes/CustomizableObjectNodeFloatConstant.cpp
--- a/Engine/Plugins/Mutable/Source/CustomizableObjectEditor/Private/MuCOE/Nodes/CustomizableObjectNodeFloatConstant.cpp
+++ b/Engine/Plugins/Mutable/Source/CustomizableObjectEditor/Private/MuCOE/Nodes/CustomizableObjectNodeFloatConstant.cpp
@@ -19,23 +19,9 @@ struct FSlateBrush;
 #define LOCTEXT_NAMESPACE "CustomizableObjectEditor"
 
 
-UCustomizableObjectNodeFloatConstant::UCustomizableObjectNodeFloatConstant()
-	: Super()
-{
-	Value = 1.0f;
-	bCollapsed = true;
-}
-
-
 void UCustomizableObjectNodeFloatConstant::AllocateDefaultPins(UCustomizableObjectNodeRemapPins* RemapPins)
 {
-	const FName Type = UEdGraphSchema_CustomizableObject::PC_Float;
-	const FName PinName = UEdGraphSchema_CustomizableObject::GetPinCategoryName(Type);
-	const FText PinFriendlyName = UEdGraphSchema_CustomizableObject::GetPinCategoryFriendlyName(Type);
-	
-	UEdGraphPin* ValuePin = CustomCreatePin(EGPD_Output, Type, PinName);
-	ValuePin->PinFriendlyName = PinFriendlyName;
-	ValuePin->bDefaultValueIsIgnored = true;
+	CustomCreatePin(EGPD_Output, UEdGraphSchema_CustomizableObject::PC_Float);
 }
 
 
diff --git a/Engine/Plugins/Mutable/Source/CustomizableObjectEditor/Private/MuCOE/Nodes/CustomizableObjectNodeFloatConstant.h b/Engine/Plugins/Mutable/Source/CustomizableObjectEditor/Private/MuCOE/Nodes/CustomizableObjectNodeFloatConstant.h
--- a/Engine/Plugins/Mutable/Source/CustomizableObjectEditor/Private/MuCOE/Nodes/CustomizableObjectNodeFloatConstant.h
+++ b/Engine/Plugins/Mutable/Source/CustomizableObjectEditor/Private/MuCOE/Nodes/CustomizableObjectNodeFloatConstant.h
@@ -69,12 +69,6 @@ class UCustomizableObjectNodeFloatConstant : public UCustomizableObjectNode
 public:
 	GENERATED_BODY()
 
-	UE_API UCustomizableObjectNodeFloatConstant();
-
-	/**  */
-	UPROPERTY(EditAnywhere, Category = CustomizableObject, meta = (DontUpdateWhileEditing))
-	float Value;
-
 	// Begin EdGraphNode interface
 	UE_API virtual FText GetNodeTitle(ENodeTitleType::Type TitleType) const override;
 	UE_API virtual FLinearColor GetNodeTitleColor() const override;
@@ -88,8 +82,11 @@ public:
 	// Creates the SGraph Node widget for the thumbnail
 	UE_API virtual TSharedPtr<SGraphNode> CreateVisualWidget() override;
 
+	UPROPERTY(EditAnywhere, Category = CustomizableObject, meta = (DontUpdateWhileEditing))
+	float Value = 1.0f;
+
 	// Determines if the Node is collapsed or not
-	bool bCollapsed;
+	bool bCollapsed = true;
 };
 
 #undef UE_API
diff --git a/Engine/Plugins/Mutable/Source/CustomizableObjectEditor/Private/MuCOE/Nodes/CustomizableObjectNodeFloatSwitch.cpp b/Engine/Plugins/Mutable/Source/CustomizableObjectEditor/Private/MuCOE/Nodes/CustomizableObjectNodeFloatSwitch.cpp
--- a/Engine/Plugins/Mutable/Source/CustomizableObjectEditor/Private/MuCOE/Nodes/CustomizableObjectNodeFloatSwitch.cpp
+++ b/Engine/Plugins/Mutable/Source/CustomizableObjectEditor/Private/MuCOE/Nodes/CustomizableObjectNodeFloatSwitch.cpp
@@ -10,17 +10,6 @@
 #define LOCTEXT_NAMESPACE "CustomizableObjectEditor"
 
 
-void UCustomizableObjectNodeFloatSwitch::BackwardsCompatibleFixup(int32 CustomizableObjectCustomVersion)
-{
-	Super::BackwardsCompatibleFixup(CustomizableObjectCustomVersion);
-	
-	if (CustomizableObjectCustomVersion == FCustomizableObjectCustomVersion::BugPinsSwitch)
-	{
-		OutputPinReference = FindPin(TEXT("Float"));	
-	}
-}
-
-
 FName UCustomizableObjectNodeFloatSwitch::GetCategory() const
 {
 	return UEdGraphSchema_CustomizableObject::PC_Float;
diff --git a/Engine/Plugins/Mutable/Source/CustomizableObjectEditor/Private/MuCOE/Nodes/CustomizableObjectNodeFloatSwitch.h b/Engine/Plugins/Mutable/Source/CustomizableObjectEditor/Private/MuCOE/Nodes/CustomizableObjectNodeFloatSwitch.h
--- a/Engine/Plugins/Mutable/Source/CustomizableObjectEditor/Private/MuCOE/Nodes/CustomizableObjectNodeFloatSwitch.h
+++ b/Engine/Plugins/Mutable/Source/CustomizableObjectEditor/Private/MuCOE/Nodes/CustomizableObjectNodeFloatSwitch.h
@@ -15,9 +15,6 @@ class UCustomizableObjectNodeFloatSwitch : public UCustomizableObjectNodeSwitchB
 public:
 	GENERATED_BODY()
 
-	// UCustomizableObjectNode interface
-	UE_API virtual void BackwardsCompatibleFixup(int32 CustomizableObjectCustomVersion) override;
-
 	// UCustomizableObjectNodeSwitchBase interface
 	UE_API virtual FName GetCategory() const override;
 };
diff --git a/Engine/Plugins/Mutable/Source/CustomizableObjectEditor/Private/MuCOE/Nodes/CustomizableObjectNodeMacroInstance.cpp b/Engine/Plugins/Mutable/Source/CustomizableObjectEditor/Private/MuCOE/Nodes/CustomizableObjectNodeMacroInstance.cpp
--- a/Engine/Plugins/Mutable/Source/CustomizableObjectEditor/Private/MuCOE/Nodes/CustomizableObjectNodeMacroInstance.cpp
+++ b/Engine/Plugins/Mutable/Source/CustomizableObjectEditor/Private/MuCOE/Nodes/CustomizableObjectNodeMacroInstance.cpp
@@ -137,7 +137,7 @@ void UCustomizableObjectNodeMacroInstance::AllocateDefaultPins(UCustomizableObje
 			PinData->VariableId = Variable->UniqueId;
 
 			const FName PinName = Variable->Name;
-			UEdGraphPin* Pin = CustomCreatePin(PinDirection, Variable->PinCategoryType, PinName, PinData);
+			CustomCreatePin(PinDirection, Variable->PinCategoryType, PinName, FText::FromName(PinName), PinData);
 		}
 	}
 }
diff --git a/Engine/Plugins/Mutable/Source/CustomizableObjectEditor/Private/MuCOE/Nodes/CustomizableObjectNodeMaterial.cpp b/Engine/Plugins/Mutable/Source/CustomizableObjectEditor/Private/MuCOE/Nodes/CustomizableObjectNodeMaterial.cpp
--- a/Engine/Plugins/Mutable/Source/CustomizableObjectEditor/Private/MuCOE/Nodes/CustomizableObjectNodeMaterial.cpp
+++ b/Engine/Plugins/Mutable/Source/CustomizableObjectEditor/Private/MuCOE/Nodes/CustomizableObjectNodeMaterial.cpp
@@ -178,7 +178,7 @@ FName UCustomizableObjectNodeMaterial::NodePinModeToImagePinMode(const ENodePinM
 	case ENodePinMode::Mutable:
 		return UEdGraphSchema_CustomizableObject::PC_Texture;
 	case ENodePinMode::Passthrough:
-		return UEdGraphSchema_CustomizableObject::PC_PassthroughTexture;
+		return UEdGraphSchema_CustomizableObject::PC_Texture_Passthrough;
 	default:
 		unimplemented();
 		return {};
@@ -195,7 +195,7 @@ FName UCustomizableObjectNodeMaterial::GetImagePinMode(const EPinMode PinMode) c
 	case EPinMode::Mutable:
 		return UEdGraphSchema_CustomizableObject::PC_Texture;
 	case EPinMode::Passthrough:
-		return UEdGraphSchema_CustomizableObject::PC_PassthroughTexture;
+		return UEdGraphSchema_CustomizableObject::PC_Texture_Passthrough;
 	default:
 		unimplemented();
 		return {};
@@ -325,25 +325,21 @@ void UCustomizableObjectNodeMaterial::AllocateDefaultPins(UCustomizableObjectNod
 	{
 		const FString PinFriendlyName = TEXT("Mesh");
 		const FString PinName = PinFriendlyName + FString(TEXT("_Input_Pin"));
-		UEdGraphPin* MeshPin = CustomCreatePin(EGPD_Input, Schema->PC_Mesh, FName(*PinName));
-		MeshPin->PinFriendlyName = FText::FromString(PinFriendlyName);
-		MeshPin->bDefaultValueIsIgnored = true;
+		CustomCreatePin(EGPD_Input, Schema->PC_Mesh, FName(*PinName), FText::FromString(PinFriendlyName));
 	}
 	
 	{
 		const FString PinFriendlyName = UEdGraphSchema_CustomizableObject::GetPinCategoryFriendlyName(Schema->PC_Material).ToString();
 		const FString PinName = FString(TEXT("Material_Input_Pin"));
-		UEdGraphPin* TableMaterialPin = CustomCreatePin(EGPD_Input, Schema->PC_Material, FName(*PinName));
-		TableMaterialPin->PinFriendlyName = FText::FromString(PinFriendlyName);
-		TableMaterialPin->bDefaultValueIsIgnored = true;
+		UEdGraphPin* TableMaterialPin = CustomCreatePin(EGPD_Input, Schema->PC_Material, FName(*PinName), FText::FromString(PinFriendlyName));
 		TableMaterialPin->PinToolTip = "Pin for a Material from a Table Node";
 	}
 
 	{
 		const FString PinFriendlyName = "Enable Tags";
 		const FString PinName = PinFriendlyName + FString(TEXT("_Input_Pin"));
-		UEdGraphPin* TagsPin = CustomCreatePin(EGPD_Input, Schema->PC_String, FName(PinName), true);
-		TagsPin->PinFriendlyName = FText::FromString(PinFriendlyName);
+		UEdGraphPin* TagsPin = CustomCreatePin(EGPD_Input, Schema->PC_String, FName(PinName), FText::FromString(PinFriendlyName));
+		TagsPin->PinType.ContainerType = EPinContainerType::Array;
 		TagsPin->PinToolTip = "List of Tags that this node will Enable";
 		EnableTagsPinRef = TagsPin;
 	}
@@ -356,8 +352,7 @@ void UCustomizableObjectNodeMaterial::AllocateDefaultPins(UCustomizableObjectNod
 	{
 		const FString PinFriendlyName = TEXT("Mesh Section");
 		const FString PinName = PinFriendlyName + FString(TEXT("_Output_Pin"));
-		UEdGraphPin* OutputPin = CustomCreatePin(EGPD_Output, Schema->PC_MeshSection, FName(*PinName));
-		OutputPin->PinFriendlyName = FText::FromString(PinFriendlyName);
+		CustomCreatePin(EGPD_Output, Schema->PC_MeshSection, FName(*PinName), FText::FromString(PinFriendlyName));
 	}
 }
 
@@ -691,9 +686,9 @@ void UCustomizableObjectNodeMaterial::BackwardsCompatibleFixup(int32 Customizabl
 						PinData.PinMode = EPinMode::Mutable;
 					}
 				}
-				else if (LinkedPin->PinType.PinCategory == UEdGraphSchema_CustomizableObject::PC_PassthroughTexture)
+				else if (LinkedPin->PinType.PinCategory == UEdGraphSchema_CustomizableObject::PC_Texture_Passthrough)
 				{
-					Pin->PinType.PinCategory = UEdGraphSchema_CustomizableObject::PC_PassthroughTexture;
+					Pin->PinType.PinCategory = UEdGraphSchema_CustomizableObject::PC_Texture_Passthrough;
 
 					if (PinData.PinMode == EPinMode::Default &&
 						TextureParametersMode == ENodePinMode::Mutable)
@@ -710,7 +705,7 @@ void UCustomizableObjectNodeMaterial::BackwardsCompatibleFixup(int32 Customizabl
 			{
 				if (!IsImageMutableMode(*Pin))
 				{
-					Pin->PinType.PinCategory = UEdGraphSchema_CustomizableObject::PC_PassthroughTexture;
+					Pin->PinType.PinCategory = UEdGraphSchema_CustomizableObject::PC_Texture_Passthrough;
 				}				
 			}
 		}
@@ -786,8 +781,8 @@ void UCustomizableObjectNodeMaterial::BackwardsCompatibleFixup(int32 Customizabl
 		{
 			const FString PinFriendlyName = "Enable Tags";
 			const FString PinName = PinFriendlyName + FString(TEXT("_Input_Pin"));
-			UEdGraphPin* TagsPin = CustomCreatePin(EGPD_Input, UEdGraphSchema_CustomizableObject::PC_String, FName(PinName), true);
-			TagsPin->PinFriendlyName = FText::FromString(PinFriendlyName);
+			UEdGraphPin* TagsPin = CustomCreatePin(EGPD_Input, UEdGraphSchema_CustomizableObject::PC_String, FName(PinName),  FText::FromString(PinFriendlyName));
+			TagsPin->PinType.ContainerType = EPinContainerType::Array;
 			TagsPin->PinToolTip = "List of Tags that this node will Enable";
 			EnableTagsPinRef = TagsPin;
 		}
@@ -1444,7 +1439,7 @@ void UCustomizableObjectNodeMaterial::AllocateDefaultParameterPins(const EMateri
 			}
 			else
 			{
-				PinCategory = UEdGraphSchema_CustomizableObject::PC_PassthroughTexture;
+				PinCategory = UEdGraphSchema_CustomizableObject::PC_Texture_Passthrough;
 			}
 			break;
 
@@ -1460,9 +1455,8 @@ void UCustomizableObjectNodeMaterial::AllocateDefaultParameterPins(const EMateri
 			check(false); // Type not contemplated
 		}		
 		
-		UEdGraphPin* Pin = CustomCreatePin(EGPD_Input, PinCategory, PinName, PinData);
+		UEdGraphPin* Pin = CustomCreatePin(EGPD_Input, PinCategory, PinName, FText::FromName(PinName), PinData);
 		Pin->bHidden = true;
-		Pin->bDefaultValueIsIgnored = true;
 
 		PinsParameterMap.Add(PinData->MaterialParameterId, FEdGraphPinReference(Pin));
 	}
@@ -1536,7 +1530,7 @@ void UCustomizableObjectNodeMaterial::PinConnectionListChanged(UEdGraphPin* Pin)
 				{
 					PinMode = EPinMode::Mutable;
 				}
-				else if (LinkedPin->PinType.PinCategory == UEdGraphSchema_CustomizableObject::PC_PassthroughTexture)
+				else if (LinkedPin->PinType.PinCategory == UEdGraphSchema_CustomizableObject::PC_Texture_Passthrough || LinkedPin->PinType.PinCategory == UEdGraphSchema_CustomizableObject::PC_Texture_PassthroughParameter)
 				{
 					PinMode = EPinMode::Passthrough;
 				}
@@ -1565,13 +1559,8 @@ bool UCustomizableObjectNodeMaterial::CanConnect(const UEdGraphPin* InOwnedInput
 
 	if (InOwnedInputPin && InOutputPin)
 	{
-		if ((InOwnedInputPin->PinType.PinCategory == Schema->PC_Texture && InOutputPin->PinType.PinCategory == Schema->PC_PassthroughTexture) ||
-			(InOwnedInputPin->PinType.PinCategory == Schema->PC_PassthroughTexture && InOutputPin->PinType.PinCategory == Schema->PC_Texture))
-		{
-			return true;
-		}
-		
-		if (InOwnedInputPin->PinType.PinCategory == Schema->PC_Mesh)
+		if ((InOwnedInputPin->PinType.PinCategory == Schema->PC_Texture && (InOutputPin->PinType.PinCategory == Schema->PC_Texture_Passthrough || InOutputPin->PinType.PinCategory == Schema->PC_Texture_PassthroughParameter)) ||
+			(InOwnedInputPin->PinType.PinCategory == Schema->PC_Texture_Passthrough && InOutputPin->PinType.PinCategory == Schema->PC_Texture))
 		{
 			return true;
 		}
diff --git a/Engine/Plugins/Mutable/Source/CustomizableObjectEditor/Private/MuCOE/Nodes/CustomizableObjectNodeMaterialSwitch.cpp b/Engine/Plugins/Mutable/Source/CustomizableObjectEditor/Private/MuCOE/Nodes/CustomizableObjectNodeMaterialSwitch.cpp
--- a/Engine/Plugins/Mutable/Source/CustomizableObjectEditor/Private/MuCOE/Nodes/CustomizableObjectNodeMaterialSwitch.cpp
+++ b/Engine/Plugins/Mutable/Source/CustomizableObjectEditor/Private/MuCOE/Nodes/CustomizableObjectNodeMaterialSwitch.cpp
@@ -35,22 +35,28 @@ void UCustomizableObjectNodeMaterialSwitch::BackwardsCompatibleFixup(int32 Custo
 			UEdGraphPin* NewPin = FindPin(TEXT("Material"), EEdGraphPinDirection::EGPD_Output);
 			if (!NewPin)
 			{
-				NewPin = CustomCreatePin(EGPD_Output, Schema->PC_MeshSection, FName("Material"));
+				NewPin = CustomCreatePin(EGPD_Output, Schema->PC_MeshSection, TEXT("Material"), LOCTEXT("Material", "Material"));
 			}
 
 			FGuid PinId = NewPin->PinId;
 			NewPin->CopyPersistentDataFromOldPin(*OldPin);
 			NewPin->PinId = PinId;
 			NewPin->bHidden = OldPin->bHidden;
-
-			OutputPinReference = NewPin;
-
+			
 			CustomRemovePin(*OldPin);
 		}
 
 		// Inputs
-		int32 InputCount = ReloadingElementsNames.Num();
-		for (int32 InputIndex=0; InputIndex<InputCount; ++InputIndex)
+		int32 InputCount = -1;
+		for (UEdGraphPin* Pin : GetAllNonOrphanPins())
+		{
+			if (Pin->Direction == EEdGraphPinDirection::EGPD_Input)
+			{
+				InputCount++;
+			}
+		}
+		
+		for (int32 InputIndex=0; InputIndex < InputCount; ++InputIndex)
 		{
 			OldPin = FindPin(FString::Printf(TEXT("Mesh Section %d "),InputIndex), EEdGraphPinDirection::EGPD_Input);
 			if (OldPin)
@@ -59,7 +65,7 @@ void UCustomizableObjectNodeMaterialSwitch::BackwardsCompatibleFixup(int32 Custo
 				UEdGraphPin* NewPin = FindPin(NewName, EEdGraphPinDirection::EGPD_Input);
 				if (!NewPin)
 				{
-					NewPin = CustomCreatePin(EEdGraphPinDirection::EGPD_Input, Schema->PC_MeshSection, FName(*NewName));
+					NewPin = CustomCreatePin(EEdGraphPinDirection::EGPD_Input, Schema->PC_MeshSection, FName(*NewName), FText::FromString(NewName));
 				}
 
 				FGuid PinId = NewPin->PinId;
diff --git a/Engine/Plugins/Mutable/Source/CustomizableObjectEditor/Private/MuCOE/Nodes/CustomizableObjectNodeMeshClipDeform.cpp b/Engine/Plugins/Mutable/Source/CustomizableObjectEditor/Private/MuCOE/Nodes/CustomizableObjectNodeMeshClipDeform.cpp
--- a/Engine/Plugins/Mutable/Source/CustomizableObjectEditor/Private/MuCOE/Nodes/CustomizableObjectNodeMeshClipDeform.cpp
+++ b/Engine/Plugins/Mutable/Source/CustomizableObjectEditor/Private/MuCOE/Nodes/CustomizableObjectNodeMeshClipDeform.cpp
@@ -14,11 +14,9 @@ void UCustomizableObjectNodeModifierClipDeform::AllocateDefaultPins(UCustomizabl
 {
 	const UEdGraphSchema_CustomizableObject* Schema = GetDefault<UEdGraphSchema_CustomizableObject>();
 	
-	UEdGraphPin* ClipMeshPin = CustomCreatePin(EGPD_Input, Schema->PC_Mesh, FName("Clip Shape"));
-	ClipMeshPin->bDefaultValueIsIgnored = true;
-	
-	UEdGraphPin* OutputPin = CustomCreatePin(EGPD_Output, Schema->PC_Modifier, FName("Modifier"));
-	OutputPin->bDefaultValueIsIgnored = true;
+	CustomCreatePin(EGPD_Input, Schema->PC_Mesh, TEXT("Clip Shape"), LOCTEXT("ClipShape", "Clip Shape"));
+
+	CustomCreatePin(EGPD_Output, Schema->PC_Modifier, TEXT("Modifier"), LOCTEXT("Modifier", "Modifier"));
 
 	//Create Node Modifier Common Pins
 	Super::AllocateDefaultPins(RemapPins);
diff --git a/Engine/Plugins/Mutable/Source/CustomizableObjectEditor/Private/MuCOE/Nodes/CustomizableObjectNodeMeshClipMorph.cpp b/Engine/Plugins/Mutable/Source/CustomizableObjectEditor/Private/MuCOE/Nodes/CustomizableObjectNodeMeshClipMorph.cpp
--- a/Engine/Plugins/Mutable/Source/CustomizableObjectEditor/Private/MuCOE/Nodes/CustomizableObjectNodeMeshClipMorph.cpp
+++ b/Engine/Plugins/Mutable/Source/CustomizableObjectEditor/Private/MuCOE/Nodes/CustomizableObjectNodeMeshClipMorph.cpp
@@ -270,7 +270,7 @@ void UCustomizableObjectNodeModifierClipMorph::AllocateDefaultPins(UCustomizable
 {
 	const UEdGraphSchema_CustomizableObject* Schema = GetDefault<UEdGraphSchema_CustomizableObject>();
 	
-	CustomCreatePin(EGPD_Output, Schema->PC_Modifier, FName("Modifier"));
+	CustomCreatePin(EGPD_Output, Schema->PC_Modifier, TEXT("Modifier"), LOCTEXT("Modifier", "Modifier"));
 
 	//Create Node Modifier Common Pins
 	Super::AllocateDefaultPins(RemapPins);
diff --git a/Engine/Plugins/Mutable/Source/CustomizableObjectEditor/Private/MuCOE/Nodes/CustomizableObjectNodeMeshClipWithMesh.cpp b/Engine/Plugins/Mutable/Source/CustomizableObjectEditor/Private/MuCOE/Nodes/CustomizableObjectNodeMeshClipWithMesh.cpp
--- a/Engine/Plugins/Mutable/Source/CustomizableObjectEditor/Private/MuCOE/Nodes/CustomizableObjectNodeMeshClipWithMesh.cpp
+++ b/Engine/Plugins/Mutable/Source/CustomizableObjectEditor/Private/MuCOE/Nodes/CustomizableObjectNodeMeshClipWithMesh.cpp
@@ -22,10 +22,9 @@ void UCustomizableObjectNodeModifierClipWithMesh::AllocateDefaultPins(UCustomiza
 {
 	const UEdGraphSchema_CustomizableObject* Schema = GetDefault<UEdGraphSchema_CustomizableObject>();
 
-	UEdGraphPin* ClipMeshPin = CustomCreatePin(EGPD_Input, Schema->PC_Mesh, FName("Clip Mesh"));
-	ClipMeshPin->bDefaultValueIsIgnored = true;
+	CustomCreatePin(EGPD_Input, Schema->PC_Mesh, TEXT("Clip Mesh"), LOCTEXT("ClipMesh", "Clip Mesh"));
 
-	UEdGraphPin* OutputPin_p = CustomCreatePin(EGPD_Output, Schema->PC_Modifier, FName("Modifier"));
+	CustomCreatePin(EGPD_Output, Schema->PC_Modifier, TEXT("Modifier"), LOCTEXT("Modifier", "Modifier"));
 
 	//Create Node Modifier Common Pins
 	Super::AllocateDefaultPins(RemapPins);
diff --git a/Engine/Plugins/Mutable/Source/CustomizableObjectEditor/Private/MuCOE/Nodes/CustomizableObjectNodeMeshMorph.cpp b/Engine/Plugins/Mutable/Source/CustomizableObjectEditor/Private/MuCOE/Nodes/CustomizableObjectNodeMeshMorph.cpp
--- a/Engine/Plugins/Mutable/Source/CustomizableObjectEditor/Private/MuCOE/Nodes/CustomizableObjectNodeMeshMorph.cpp
+++ b/Engine/Plugins/Mutable/Source/CustomizableObjectEditor/Private/MuCOE/Nodes/CustomizableObjectNodeMeshMorph.cpp
@@ -28,22 +28,11 @@ UCustomizableObjectNodeMeshMorph::UCustomizableObjectNodeMeshMorph()
 
 void UCustomizableObjectNodeMeshMorph::AllocateDefaultPins(UCustomizableObjectNodeRemapPins* RemapPins)
 {
-	const UEdGraphSchema_CustomizableObject* Schema = GetDefault<UEdGraphSchema_CustomizableObject>();
+	CustomCreatePin(EGPD_Output, UEdGraphSchema_CustomizableObject::PC_Mesh, TEXT("Mesh"), LOCTEXT("Mesh", "Mesh"));
 
-	FString PinName = TEXT("Mesh");
-	UEdGraphPin* OutputPin = CustomCreatePin(EGPD_Output, Schema->PC_Mesh, FName(*PinName));
-	OutputPin->bDefaultValueIsIgnored = true;
-
-	PinName = TEXT("Mesh");
-	UEdGraphPin* MeshPin = CustomCreatePin(EGPD_Input, Schema->PC_Mesh, FName(*PinName));
-	MeshPin->bDefaultValueIsIgnored = true;
-
-	PinName = TEXT("Factor");
-	UEdGraphPin* FactorPin = CustomCreatePin(EGPD_Input, Schema->PC_Float, FName(*PinName));
-	FactorPin->bDefaultValueIsIgnored = true;
-
-	PinName = TEXT("Morph Target Name");
-	MorphTargetNamePinRef = CustomCreatePin(EGPD_Input, Schema->PC_String, FName(*PinName));
+	CustomCreatePin(EGPD_Input, UEdGraphSchema_CustomizableObject::PC_Mesh, TEXT("Mesh"), LOCTEXT("Mesh", "Mesh"));
+	CustomCreatePin(EGPD_Input, UEdGraphSchema_CustomizableObject::PC_Float, TEXT("Factor"), LOCTEXT("Factor", "Factor"));
+	MorphTargetNamePinRef = CustomCreatePin(EGPD_Input, UEdGraphSchema_CustomizableObject::PC_String, TEXT("Morph Target Name"), LOCTEXT("MorphTargetName", "Morph Target Name"));
 }
 
 
@@ -172,8 +161,7 @@ void UCustomizableObjectNodeMeshMorph::BackwardsCompatibleFixup(int32 Customizab
 	{
 		if (!MorphTargetNamePinRef.Get())
 		{
-			FString PinName = TEXT("Morph Target Name");
-			MorphTargetNamePinRef = CustomCreatePin(EGPD_Input, UEdGraphSchema_CustomizableObject::PC_String, FName(*PinName));
+			MorphTargetNamePinRef = CustomCreatePin(EGPD_Input, UEdGraphSchema_CustomizableObject::PC_String, TEXT("Morph Target Name"), LOCTEXT("MorphTargetName", "Morph Target Name"));
 		}
 	}
 }
diff --git a/Engine/Plugins/Mutable/Source/CustomizableObjectEditor/Private/MuCOE/Nodes/CustomizableObjectNodeMeshMorphStackApplication.cpp b/Engine/Plugins/Mutable/Source/CustomizableObjectEditor/Private/MuCOE/Nodes/CustomizableObjectNodeMeshMorphStackApplication.cpp
--- a/Engine/Plugins/Mutable/Source/CustomizableObjectEditor/Private/MuCOE/Nodes/CustomizableObjectNodeMeshMorphStackApplication.cpp
+++ b/Engine/Plugins/Mutable/Source/CustomizableObjectEditor/Private/MuCOE/Nodes/CustomizableObjectNodeMeshMorphStackApplication.cpp
@@ -21,12 +21,12 @@ void UCustomizableObjectNodeMeshMorphStackApplication::AllocateDefaultPins(UCust
 {
 	// Input pins
 	
-	CustomCreatePinSimple(EGPD_Input, UEdGraphSchema_CustomizableObject::PC_Mesh);
-	CustomCreatePinSimple(EGPD_Input, UEdGraphSchema_CustomizableObject::PC_Stack);
+	CustomCreatePin(EGPD_Input, UEdGraphSchema_CustomizableObject::PC_Mesh);
+	CustomCreatePin(EGPD_Input, UEdGraphSchema_CustomizableObject::PC_Stack);
 	
 	// Output pins
 	
-	CustomCreatePinSimple(EGPD_Output, UEdGraphSchema_CustomizableObject::PC_Mesh);
+	CustomCreatePin(EGPD_Output, UEdGraphSchema_CustomizableObject::PC_Mesh);
 }
 
 
diff --git a/Engine/Plugins/Mutable/Source/CustomizableObjectEditor/Private/MuCOE/Nodes/CustomizableObjectNodeMeshMorphStackDefinition.cpp b/Engine/Plugins/Mutable/Source/CustomizableObjectEditor/Private/MuCOE/Nodes/CustomizableObjectNodeMeshMorphStackDefinition.cpp
--- a/Engine/Plugins/Mutable/Source/CustomizableObjectEditor/Private/MuCOE/Nodes/CustomizableObjectNodeMeshMorphStackDefinition.cpp
+++ b/Engine/Plugins/Mutable/Source/CustomizableObjectEditor/Private/MuCOE/Nodes/CustomizableObjectNodeMeshMorphStackDefinition.cpp
@@ -23,20 +23,16 @@ void UCustomizableObjectNodeMeshMorphStackDefinition::AllocateDefaultPins(UCusto
 	UpdateMorphList();
 	
 	// Input pins
-	UEdGraphPin* MeshPin = CustomCreatePin(EGPD_Input, Schema->PC_Mesh, FName(TEXT("Mesh")));
-	MeshPin->bDefaultValueIsIgnored = true;
+	CustomCreatePin(EGPD_Input, Schema->PC_Mesh, TEXT("Mesh"), LOCTEXT("Mesh", "Mesh"));
 
 	for (int32 i = 0; i < MorphNames.Num(); ++i)
 	{
 		FString PinName = MorphNames[i] + FString("_Morph");
-		UEdGraphPin* Pin = CustomCreatePin(EGPD_Input, Schema->PC_Float, FName(*PinName));
-		Pin->PinFriendlyName = FText::FromString(MorphNames[i]);
-		Pin->bDefaultValueIsIgnored = true;
+		UEdGraphPin* Pin = CustomCreatePin(EGPD_Input, Schema->PC_Float, FName(*PinName), FText::FromString(MorphNames[i]));
 	}
 
 	// Output pins
-	UEdGraphPin* StackPin = CustomCreatePin(EGPD_Output, Schema->PC_Stack, FName(TEXT("Stack")));
-	StackPin->bDefaultValueIsIgnored = true;
+	CustomCreatePin(EGPD_Output, Schema->PC_Stack, TEXT("Stack"), LOCTEXT("Stack", "Stack"));
 }
 
 
diff --git a/Engine/Plugins/Mutable/Source/CustomizableObjectEditor/Private/MuCOE/Nodes/CustomizableObjectNodeMeshReshape.cpp b/Engine/Plugins/Mutable/Source/CustomizableObjectEditor/Private/MuCOE/Nodes/CustomizableObjectNodeMeshReshape.cpp
--- a/Engine/Plugins/Mutable/Source/CustomizableObjectEditor/Private/MuCOE/Nodes/CustomizableObjectNodeMeshReshape.cpp
+++ b/Engine/Plugins/Mutable/Source/CustomizableObjectEditor/Private/MuCOE/Nodes/CustomizableObjectNodeMeshReshape.cpp
@@ -13,27 +13,13 @@ class UCustomizableObjectNodeRemapPins;
 #define LOCTEXT_NAMESPACE "CustomizableObjectEditor"
 
 
-UCustomizableObjectNodeMeshReshape::UCustomizableObjectNodeMeshReshape()
-	: Super()
-{
-}
-
 
 void UCustomizableObjectNodeMeshReshape::AllocateDefaultPins(UCustomizableObjectNodeRemapPins* RemapPins)
 {
-	const UEdGraphSchema_CustomizableObject* Schema = GetDefault<UEdGraphSchema_CustomizableObject>();
-
-	UEdGraphPin* OutputPin = CustomCreatePin(EGPD_Output, Schema->PC_Mesh, FName("Mesh"));
-	OutputPin->bDefaultValueIsIgnored = true;
-
-	UEdGraphPin* MeshPin = CustomCreatePin(EGPD_Input, Schema->PC_Mesh, FName("Base Mesh"));
-	MeshPin->bDefaultValueIsIgnored = true;
-
-	MeshPin = CustomCreatePin(EGPD_Input, Schema->PC_Mesh, FName("Base Shape"));
-	MeshPin->bDefaultValueIsIgnored = true;
-
-	MeshPin = CustomCreatePin(EGPD_Input, Schema->PC_Mesh, FName("Target Shape"));
-	MeshPin->bDefaultValueIsIgnored = true;
+	CustomCreatePin(EGPD_Output,UEdGraphSchema_CustomizableObject::PC_Mesh, TEXT("Mesh"), LOCTEXT("Mesh", "Mesh"));
+	CustomCreatePin(EGPD_Input, UEdGraphSchema_CustomizableObject::PC_Mesh, TEXT("Base Mesh"), LOCTEXT("BaseMesh", "Base Mesh"));
+	CustomCreatePin(EGPD_Input, UEdGraphSchema_CustomizableObject::PC_Mesh, TEXT("Base Shape"), LOCTEXT("BaseShape", "Base Shape"));
+	CustomCreatePin(EGPD_Input, UEdGraphSchema_CustomizableObject::PC_Mesh, TEXT("Target Shape"), LOCTEXT("TargetShape", "Target Shape"));
 }
 
 
diff --git a/Engine/Plugins/Mutable/Source/CustomizableObjectEditor/Private/MuCOE/Nodes/CustomizableObjectNodeMeshReshape.h b/Engine/Plugins/Mutable/Source/CustomizableObjectEditor/Private/MuCOE/Nodes/CustomizableObjectNodeMeshReshape.h
--- a/Engine/Plugins/Mutable/Source/CustomizableObjectEditor/Private/MuCOE/Nodes/CustomizableObjectNodeMeshReshape.h
+++ b/Engine/Plugins/Mutable/Source/CustomizableObjectEditor/Private/MuCOE/Nodes/CustomizableObjectNodeMeshReshape.h
@@ -46,9 +46,7 @@ class UCustomizableObjectNodeMeshReshape : public UCustomizableObjectNode
 {
 public:
 	GENERATED_BODY()
-
-	UE_API UCustomizableObjectNodeMeshReshape();
-
+	
 	// Begin EdGraphNode interface
 	UE_API FText GetNodeTitle(ENodeTitleType::Type TitleType) const override;
 	UE_API FLinearColor GetNodeTitleColor() const override;
diff --git a/Engine/Plugins/Mutable/Source/CustomizableObjectEditor/Private/MuCOE/Nodes/CustomizableObjectNodeMeshSwitch.cpp b/Engine/Plugins/Mutable/Source/CustomizableObjectEditor/Private/MuCOE/Nodes/CustomizableObjectNodeMeshSwitch.cpp
--- a/Engine/Plugins/Mutable/Source/CustomizableObjectEditor/Private/MuCOE/Nodes/CustomizableObjectNodeMeshSwitch.cpp
+++ b/Engine/Plugins/Mutable/Source/CustomizableObjectEditor/Private/MuCOE/Nodes/CustomizableObjectNodeMeshSwitch.cpp
@@ -10,17 +10,6 @@
 #define LOCTEXT_NAMESPACE "CustomizableObjectEditor"
 
 
-void UCustomizableObjectNodeMeshSwitch::BackwardsCompatibleFixup(int32 CustomizableObjectCustomVersion)
-{
-	Super::BackwardsCompatibleFixup(CustomizableObjectCustomVersion);
-	
-	if (CustomizableObjectCustomVersion == FCustomizableObjectCustomVersion::BugPinsSwitch)
-	{
-		OutputPinReference = FindPin(TEXT("Mesh"));	
-	}
-}
-
-
 FName UCustomizableObjectNodeMeshSwitch::GetCategory() const
 {
 	return UEdGraphSchema_CustomizableObject::PC_Mesh;
diff --git a/Engine/Plugins/Mutable/Source/CustomizableObjectEditor/Private/MuCOE/Nodes/CustomizableObjectNodeMeshSwitch.h b/Engine/Plugins/Mutable/Source/CustomizableObjectEditor/Private/MuCOE/Nodes/CustomizableObjectNodeMeshSwitch.h
--- a/Engine/Plugins/Mutable/Source/CustomizableObjectEditor/Private/MuCOE/Nodes/CustomizableObjectNodeMeshSwitch.h
+++ b/Engine/Plugins/Mutable/Source/CustomizableObjectEditor/Private/MuCOE/Nodes/CustomizableObjectNodeMeshSwitch.h
@@ -15,9 +15,6 @@ class UCustomizableObjectNodeMeshSwitch : public UCustomizableObjectNodeSwitchBa
 public:
 	GENERATED_BODY()
 
-	// UCustomizableObjectNode interface
-	UE_API virtual void BackwardsCompatibleFixup(int32 CustomizableObjectCustomVersion) override;
-	
 	// UCustomizableObjectNodeSwitchBase interface
 	UE_API virtual FName GetCategory() const override;
 };
diff --git a/Engine/Plugins/Mutable/Source/CustomizableObjectEditor/Private/MuCOE/Nodes/CustomizableObjectNodeModifierBase.cpp b/Engine/Plugins/Mutable/Source/CustomizableObjectEditor/Private/MuCOE/Nodes/CustomizableObjectNodeModifierBase.cpp
--- a/Engine/Plugins/Mutable/Source/CustomizableObjectEditor/Private/MuCOE/Nodes/CustomizableObjectNodeModifierBase.cpp
+++ b/Engine/Plugins/Mutable/Source/CustomizableObjectEditor/Private/MuCOE/Nodes/CustomizableObjectNodeModifierBase.cpp
@@ -13,6 +13,8 @@
 
 #include UE_INLINE_GENERATED_CPP_BY_NAME(CustomizableObjectNodeModifierBase)
 
+#define LOCTEXT_NAMESPACE "CustomizableObjectEditor"
+
 
 FLinearColor UCustomizableObjectNodeModifierBase::GetNodeTitleColor() const
 {
@@ -156,7 +158,7 @@ void UCustomizableObjectNodeModifierBase::BackwardsCompatibleFixup(int32 Customi
 			UEdGraphPin* NewPin = FindPin(TEXT("Modifier")); 
 			if (!NewPin)
 			{
-				NewPin = CustomCreatePin(EGPD_Output, Schema->PC_Modifier, FName("Modifier"));
+				NewPin = CustomCreatePin(EGPD_Output, Schema->PC_Modifier, TEXT("Modifier"), LOCTEXT("Modifier", "Modifier"));
 			}
 
 			FGuid PinId = NewPin->PinId;
@@ -269,7 +271,9 @@ void UCustomizableObjectNodeModifierBase::BackwardsCompatibleFixup(int32 Customi
 	{
 		if (!RequiredTagsPinRef.Get())
 		{
-			RequiredTagsPinRef = CustomCreatePin(EGPD_Input, UEdGraphSchema_CustomizableObject::PC_String, "Target Tags", true);
+			UEdGraphPin* Pin = CustomCreatePin(EGPD_Input, UEdGraphSchema_CustomizableObject::PC_String, TEXT("Target Tags"), LOCTEXT("TargetTags", "Target Tags"));
+			Pin->PinType.ContainerType = EPinContainerType::Array;
+			RequiredTagsPinRef = Pin;
 		}
 	}
 }
@@ -346,5 +350,10 @@ FString UCustomizableObjectNodeModifierBase::MakeNodeAutoTag( UEdGraphNode* Node
 
 void UCustomizableObjectNodeModifierBase::AllocateDefaultPins(UCustomizableObjectNodeRemapPins* RemapPins)
 {
-	RequiredTagsPinRef = CustomCreatePin(EGPD_Input, UEdGraphSchema_CustomizableObject::PC_String, "Target Tags", true);
+	UEdGraphPin* Pin = CustomCreatePin(EGPD_Input, UEdGraphSchema_CustomizableObject::PC_String, TEXT("Target Tags"), LOCTEXT("TargetTags", "Target Tags"));
+	Pin->PinType.ContainerType = EPinContainerType::Array;
+	RequiredTagsPinRef = Pin;
 }
+
+
+#undef LOCTEXT_NAMESPACE
diff --git a/Engine/Plugins/Mutable/Source/CustomizableObjectEditor/Private/MuCOE/Nodes/CustomizableObjectNodeModifierClipWithUVMask.cpp b/Engine/Plugins/Mutable/Source/CustomizableObjectEditor/Private/MuCOE/Nodes/CustomizableObjectNodeModifierClipWithUVMask.cpp
--- a/Engine/Plugins/Mutable/Source/CustomizableObjectEditor/Private/MuCOE/Nodes/CustomizableObjectNodeModifierClipWithUVMask.cpp
+++ b/Engine/Plugins/Mutable/Source/CustomizableObjectEditor/Private/MuCOE/Nodes/CustomizableObjectNodeModifierClipWithUVMask.cpp
@@ -18,10 +18,8 @@ void UCustomizableObjectNodeModifierClipWithUVMask::AllocateDefaultPins(UCustomi
 {
 	const UEdGraphSchema_CustomizableObject* Schema = GetDefault<UEdGraphSchema_CustomizableObject>();
 
-	UEdGraphPin* ClipMaskPin = CustomCreatePin(EGPD_Input, Schema->PC_Texture, FName("Clip Mask"));
-	ClipMaskPin->bDefaultValueIsIgnored = true;
-
-	UEdGraphPin* OutputPin_p = CustomCreatePin(EGPD_Output, Schema->PC_Modifier, FName("Modifier"));
+	CustomCreatePin(EGPD_Input, Schema->PC_Texture, TEXT("Clip Mask"), LOCTEXT("ClipMask", "Clip Mask"));
+	CustomCreatePin(EGPD_Output, Schema->PC_Modifier, TEXT("Modifier"), LOCTEXT("Modifier", "Modifier"));
 
 	//Create Node Modifier Common Pins
 	Super::AllocateDefaultPins(RemapPins);
diff --git a/Engine/Plugins/Mutable/Source/CustomizableObjectEditor/Private/MuCOE/Nodes/CustomizableObjectNodeModifierEditMeshSection.cpp b/Engine/Plugins/Mutable/Source/CustomizableObjectEditor/Private/MuCOE/Nodes/CustomizableObjectNodeModifierEditMeshSection.cpp
--- a/Engine/Plugins/Mutable/Source/CustomizableObjectEditor/Private/MuCOE/Nodes/CustomizableObjectNodeModifierEditMeshSection.cpp
+++ b/Engine/Plugins/Mutable/Source/CustomizableObjectEditor/Private/MuCOE/Nodes/CustomizableObjectNodeModifierEditMeshSection.cpp
@@ -34,22 +34,20 @@ void UCustomizableObjectNodeModifierEditMeshSection::AllocateDefaultPins(UCustom
 			PinEditImageData->ImageParamId = GetParameterId(EMaterialParameterType::Texture, ImageIndex);
 			
 			const FName ImageName = GetParameterName(EMaterialParameterType::Texture, ImageIndex);
-			UEdGraphPin* PinImage = CustomCreatePin(EGPD_Input, Schema->PC_Texture, ImageName, PinEditImageData);
+			UEdGraphPin* PinImage = CustomCreatePin(EGPD_Input, Schema->PC_Texture, ImageName, FText::FromName(ImageName), PinEditImageData);
 			PinImage->bHidden = true;
-			PinImage->bDefaultValueIsIgnored = true;
 
 			PinsParameterMap.Add(PinEditImageData->ImageParamId, FEdGraphPinReference(PinImage));
 
 			FString PinMaskName = ImageName.ToString() + FString(TEXT(" Mask"));
-			UEdGraphPin* PinMask = CustomCreatePin(EGPD_Input, Schema->PC_Texture, *PinMaskName);
+			UEdGraphPin* PinMask = CustomCreatePin(EGPD_Input, Schema->PC_Texture, *PinMaskName, FText::FromString(PinMaskName));
 			PinMask->bHidden = true;
-			PinMask->bDefaultValueIsIgnored = true;
 			
 			PinEditImageData->PinMask = FEdGraphPinReference(PinMask);
 		}
 	}
 
-	CustomCreatePin(EGPD_Output, Schema->PC_Modifier, TEXT("Modifier"));
+	CustomCreatePin(EGPD_Output, Schema->PC_Modifier, TEXT("Modifier"), LOCTEXT("Modifier", "Modifier"));
 }
 
 
diff --git a/Engine/Plugins/Mutable/Source/CustomizableObjectEditor/Private/MuCOE/Nodes/CustomizableObjectNodeModifierExtendMeshSection.cpp b/Engine/Plugins/Mutable/Source/CustomizableObjectEditor/Private/MuCOE/Nodes/CustomizableObjectNodeModifierExtendMeshSection.cpp
--- a/Engine/Plugins/Mutable/Source/CustomizableObjectEditor/Private/MuCOE/Nodes/CustomizableObjectNodeModifierExtendMeshSection.cpp
+++ b/Engine/Plugins/Mutable/Source/CustomizableObjectEditor/Private/MuCOE/Nodes/CustomizableObjectNodeModifierExtendMeshSection.cpp
@@ -115,7 +115,8 @@ void UCustomizableObjectNodeModifierExtendMeshSection::BackwardsCompatibleFixup(
 	{
 		if (!EnableTagsPinRef.Get())
 		{
-			UEdGraphPin* EnableTagsPin = CustomCreatePin(EGPD_Input, UEdGraphSchema_CustomizableObject::PC_String, FName("Enable Tags"), true);
+			UEdGraphPin* EnableTagsPin = CustomCreatePin(EGPD_Input, UEdGraphSchema_CustomizableObject::PC_String, TEXT("Enable Tags"), LOCTEXT("EnableTags", "Enable Tags"));
+			EnableTagsPin->PinType.ContainerType = EPinContainerType::Array;
 			EnableTagsPin->PinToolTip = "List of Tags that this node will Enable";
 			EnableTagsPinRef = EnableTagsPin;
 		}
@@ -127,14 +128,15 @@ void UCustomizableObjectNodeModifierExtendMeshSection::AllocateDefaultPins(UCust
 {
 	const UEdGraphSchema_CustomizableObject* Schema = GetDefault<UEdGraphSchema_CustomizableObject>();
 
-	CustomCreatePin(EGPD_Input, Schema->PC_Mesh, FName("Add Mesh") );
+	CustomCreatePin(EGPD_Input, Schema->PC_Mesh, TEXT("Add Mesh"), LOCTEXT("AddMesh", "Add Mesh"));
 
 	//Create Node Modifier Common Pins
 	Super::AllocateDefaultPins(RemapPins);
 
 	// Create enable tags pin
 	{
-		UEdGraphPin* EnableTagsPin = CustomCreatePin(EGPD_Input, Schema->PC_String, FName("Enable Tags"), true);
+		UEdGraphPin* EnableTagsPin = CustomCreatePin(EGPD_Input, Schema->PC_String, TEXT("Enable Tags"), LOCTEXT("EnableTags", "Enable Tags"));
+		EnableTagsPin->PinType.ContainerType = EPinContainerType::Array;
 		EnableTagsPin->PinToolTip = "List of Tags that this node will Enable";
 		EnableTagsPinRef = EnableTagsPin;
 	}
@@ -145,8 +147,7 @@ void UCustomizableObjectNodeModifierExtendMeshSection::AllocateDefaultPins(UCust
 		for (int32 ImageIndex = 0; ImageIndex < NumImages; ++ImageIndex)
 		{
 			const FName ImageName = GetParameterName(EMaterialParameterType::Texture, ImageIndex);
-			UEdGraphPin* PinImage = CustomCreatePin(EGPD_Input, Schema->PC_Texture, ImageName);
-			PinImage->bDefaultValueIsIgnored = true;
+			UEdGraphPin* PinImage = CustomCreatePin(EGPD_Input, Schema->PC_Texture, ImageName, FText::FromName(ImageName));
 
 			const FNodeMaterialParameterId ImageId = GetParameterId(EMaterialParameterType::Texture, ImageIndex);
 				
@@ -155,7 +156,7 @@ void UCustomizableObjectNodeModifierExtendMeshSection::AllocateDefaultPins(UCust
 	}
 	// End texture pins
 
-	CustomCreatePin(EGPD_Output, Schema->PC_Modifier, TEXT("Modifier"));
+	CustomCreatePin(EGPD_Output, Schema->PC_Modifier, TEXT("Modifier"), LOCTEXT("Modifier", "Modifier"));
 }
 
 
diff --git a/Engine/Plugins/Mutable/Source/CustomizableObjectEditor/Private/MuCOE/Nodes/CustomizableObjectNodeModifierMorphMeshSection.cpp b/Engine/Plugins/Mutable/Source/CustomizableObjectEditor/Private/MuCOE/Nodes/CustomizableObjectNodeModifierMorphMeshSection.cpp
--- a/Engine/Plugins/Mutable/Source/CustomizableObjectEditor/Private/MuCOE/Nodes/CustomizableObjectNodeModifierMorphMeshSection.cpp
+++ b/Engine/Plugins/Mutable/Source/CustomizableObjectEditor/Private/MuCOE/Nodes/CustomizableObjectNodeModifierMorphMeshSection.cpp
@@ -22,16 +22,10 @@ void UCustomizableObjectNodeModifierMorphMeshSection::AllocateDefaultPins(UCusto
 {
 	const UEdGraphSchema_CustomizableObject* Schema = GetDefault<UEdGraphSchema_CustomizableObject>();
 
-	FString PinName = TEXT("Modifier");
-	UEdGraphPin* ModifierPin = CustomCreatePin(EGPD_Output, Schema->PC_Modifier, FName(PinName));
-	ModifierPin->bDefaultValueIsIgnored = true;
+	CustomCreatePin(EGPD_Output, Schema->PC_Modifier, TEXT("Modifier"), LOCTEXT("Modifier", "Modifier"));
+	CustomCreatePin(EGPD_Input, Schema->PC_Float, TEXT("Factor"), LOCTEXT("Factor", "Factor"));
 
-	PinName = TEXT("Factor");
-	UEdGraphPin* FactorPin = CustomCreatePin(EGPD_Input, Schema->PC_Float, FName(*PinName));
-	FactorPin->bDefaultValueIsIgnored = true;
-
-	PinName = TEXT("Morph Target Name");
-	MorphTargetNamePinRef = CustomCreatePin(EGPD_Input, Schema->PC_String, FName(*PinName));
+	MorphTargetNamePinRef = CustomCreatePin(EGPD_Input, Schema->PC_String,  TEXT("Morph Target Name"), LOCTEXT("MorphTargetName", "Morph Target Name"));
 
 	//Create Node Modifier Common Pins
 	Super::AllocateDefaultPins(RemapPins);
@@ -76,7 +70,7 @@ void UCustomizableObjectNodeModifierMorphMeshSection::BackwardsCompatibleFixup(i
 	{
 		if (!MorphTargetNamePinRef.Get())
 		{
-			MorphTargetNamePinRef = CustomCreatePin(EGPD_Input, UEdGraphSchema_CustomizableObject::PC_String, FName("Morph Target Name"));
+			MorphTargetNamePinRef = CustomCreatePin(EGPD_Input, UEdGraphSchema_CustomizableObject::PC_String, TEXT("Morph Target Name"), LOCTEXT("MorphTargetName", "Morph Target Name"));
 		}
 	}
 }
diff --git a/Engine/Plugins/Mutable/Source/CustomizableObjectEditor/Private/MuCOE/Nodes/CustomizableObjectNodeModifierRemoveMesh.cpp b/Engine/Plugins/Mutable/Source/CustomizableObjectEditor/Private/MuCOE/Nodes/CustomizableObjectNodeModifierRemoveMesh.cpp
--- a/Engine/Plugins/Mutable/Source/CustomizableObjectEditor/Private/MuCOE/Nodes/CustomizableObjectNodeModifierRemoveMesh.cpp
+++ b/Engine/Plugins/Mutable/Source/CustomizableObjectEditor/Private/MuCOE/Nodes/CustomizableObjectNodeModifierRemoveMesh.cpp
@@ -17,10 +17,8 @@ void UCustomizableObjectNodeModifierRemoveMesh::AllocateDefaultPins(UCustomizabl
 {
 	const UEdGraphSchema_CustomizableObject* Schema = GetDefault<UEdGraphSchema_CustomizableObject>();
 
-	UEdGraphPin* RemoveMeshPin = CustomCreatePin(EGPD_Input, Schema->PC_Mesh, FName("Remove Mesh") );
-	RemoveMeshPin->bDefaultValueIsIgnored = true;
-
-	UEdGraphPin* OutputPin = CustomCreatePin(EGPD_Output, Schema->PC_Modifier, FName("Modifier"));
+	CustomCreatePin(EGPD_Input, Schema->PC_Mesh, TEXT("Remove Mesh"), LOCTEXT("RemoveMesh", "Remove Mesh"));
+	CustomCreatePin(EGPD_Output, UEdGraphSchema_CustomizableObject::PC_Modifier);
 
 	//Create Node Modifier Common Pins
 	Super::AllocateDefaultPins(RemapPins);
diff --git a/Engine/Plugins/Mutable/Source/CustomizableObjectEditor/Private/MuCOE/Nodes/CustomizableObjectNodeModifierRemoveMeshBlocks.cpp b/Engine/Plugins/Mutable/Source/CustomizableObjectEditor/Private/MuCOE/Nodes/CustomizableObjectNodeModifierRemoveMeshBlocks.cpp
--- a/Engine/Plugins/Mutable/Source/CustomizableObjectEditor/Private/MuCOE/Nodes/CustomizableObjectNodeModifierRemoveMeshBlocks.cpp
+++ b/Engine/Plugins/Mutable/Source/CustomizableObjectEditor/Private/MuCOE/Nodes/CustomizableObjectNodeModifierRemoveMeshBlocks.cpp
@@ -223,9 +223,7 @@ void UCustomizableObjectNodeModifierRemoveMeshBlocks::BackwardsCompatibleFixup(i
 
 void UCustomizableObjectNodeModifierRemoveMeshBlocks::AllocateDefaultPins(UCustomizableObjectNodeRemapPins* RemapPins)
 {
-	const UEdGraphSchema_CustomizableObject* Schema = GetDefault<UEdGraphSchema_CustomizableObject>();
-
-	UEdGraphPin* OutputPin = CustomCreatePin(EGPD_Output, Schema->PC_Modifier, FName("Modifier"));
+	CustomCreatePin(EGPD_Output, UEdGraphSchema_CustomizableObject::PC_Modifier);
 }
 
 
diff --git a/Engine/Plugins/Mutable/Source/CustomizableObjectEditor/Private/MuCOE/Nodes/CustomizableObjectNodeModifierTransformInMesh.cpp b/Engine/Plugins/Mutable/Source/CustomizableObjectEditor/Private/MuCOE/Nodes/CustomizableObjectNodeModifierTransformInMesh.cpp
--- a/Engine/Plugins/Mutable/Source/CustomizableObjectEditor/Private/MuCOE/Nodes/CustomizableObjectNodeModifierTransformInMesh.cpp
+++ b/Engine/Plugins/Mutable/Source/CustomizableObjectEditor/Private/MuCOE/Nodes/CustomizableObjectNodeModifierTransformInMesh.cpp
@@ -28,13 +28,9 @@ void UCustomizableObjectNodeModifierTransformInMesh::AllocateDefaultPins(UCustom
 {
 	const UEdGraphSchema_CustomizableObject* Schema = GetDefault<UEdGraphSchema_CustomizableObject>();
 
-	UEdGraphPin* ClipMeshPin = CustomCreatePin(EGPD_Input, Schema->PC_Mesh, FName(BoundingMeshPinName));
-	ClipMeshPin->bDefaultValueIsIgnored = true;
-
-	UEdGraphPin* TransformPin = CustomCreatePin(EGPD_Input, Schema->PC_Transform, FName(TransformPinName));
-	TransformPin->bDefaultValueIsIgnored = true;
-	
-	(void)CustomCreatePin(EGPD_Output, Schema->PC_Modifier, FName(OutputPinName));
+	CustomCreatePin(EGPD_Input, Schema->PC_Mesh, FName(BoundingMeshPinName), FText::FromString(BoundingMeshPinName));
+	CustomCreatePin(EGPD_Input, Schema->PC_Transform, FName(TransformPinName), FText::FromString(TransformPinName));
+	CustomCreatePin(EGPD_Output, Schema->PC_Modifier, FName(OutputPinName), FText::FromString(TransformPinName));
 
 	//Create Node Modifier Common Pins
 	Super::AllocateDefaultPins(RemapPins);
diff --git a/Engine/Plugins/Mutable/Source/CustomizableObjectEditor/Private/MuCOE/Nodes/CustomizableObjectNodeObject.cpp b/Engine/Plugins/Mutable/Source/CustomizableObjectEditor/Private/MuCOE/Nodes/CustomizableObjectNodeObject.cpp
--- a/Engine/Plugins/Mutable/Source/CustomizableObjectEditor/Private/MuCOE/Nodes/CustomizableObjectNodeObject.cpp
+++ b/Engine/Plugins/Mutable/Source/CustomizableObjectEditor/Private/MuCOE/Nodes/CustomizableObjectNodeObject.cpp
@@ -132,7 +132,8 @@ void UCustomizableObjectNodeObject::BackwardsCompatibleFixup(int32 CustomizableO
 		UEdGraphPin* OldPin = FindPin(TEXT("Modifiers"));
 		if (!OldPin)
 		{
-			CustomCreatePin(EGPD_Input, Schema->PC_Modifier, ModifiersPinName, true);
+			UEdGraphPin* Pin = CustomCreatePin(EGPD_Input, Schema->PC_Modifier, ModifiersPinName, LOCTEXT("Modifiers", "Modifiers"));
+			Pin->PinType.ContainerType = EPinContainerType::Array;
 		}
 	}
 	
@@ -149,11 +150,12 @@ void UCustomizableObjectNodeObject::BackwardsCompatibleFixup(int32 CustomizableO
 			{
 				FString LODName = FString::Printf(TEXT("LOD %d"), NodeComponentLODIndex);
 
-				UEdGraphPin* Pin = Node->CustomCreatePin(EGPD_Input, Schema->PC_MeshSection, FName(*LODName), true);
+				UEdGraphPin* Pin = Node->CustomCreatePin(EGPD_Input, Schema->PC_MeshSection, FName(*LODName), FText::FromString(LODName));
+				Pin->PinType.ContainerType = EPinContainerType::Array;
 				NodeComponentMesh->GetLODPins().Add(Pin);
 			}
 						
-			NodeComponentMesh->SetOutputPin(Node->CustomCreatePin(EGPD_Output, Schema->PC_Component, TEXT("Component")));
+			NodeComponentMesh->SetOutputPin(Node->CustomCreatePin(EGPD_Output, Schema->PC_Component, TEXT("Component"), LOCTEXT("Component", "Component")));
 		};
 
 		auto NodeObjectAllocateDefaultPins = [](UCustomizableObjectNode* Node, UCustomizableObjectNodeRemapPins*)
@@ -162,12 +164,15 @@ void UCustomizableObjectNodeObject::BackwardsCompatibleFixup(int32 CustomizableO
 
 			const UEdGraphSchema_CustomizableObject* Schema = GetDefault<UEdGraphSchema_CustomizableObject>();
 
-			NodeObject->CustomCreatePin(EGPD_Input, Schema->PC_Component, ComponentsPinName, true);
-	
-			NodeObject->CustomCreatePin(EGPD_Input, Schema->PC_Modifier, ModifiersPinName, true);
+			UEdGraphPin* ComponentPin = NodeObject->CustomCreatePin(EGPD_Input, Schema->PC_Component, ComponentsPinName, LOCTEXT("Components", "Components"));
+			ComponentPin->PinType.ContainerType = EPinContainerType::Array;
+
+			UEdGraphPin* ModifierPin = NodeObject->CustomCreatePin(EGPD_Input, Schema->PC_Modifier, ModifiersPinName, LOCTEXT("Modifiers", "Modifiers"));
+			ModifierPin->PinType.ContainerType = EPinContainerType::Array;
+
+			UEdGraphPin* ChildrenPin = NodeObject->CustomCreatePin(EGPD_Input, Schema->PC_Object, ChildrenPinName, LOCTEXT("Children", "Children"));
+			ChildrenPin->PinType.ContainerType = EPinContainerType::Array;
 
-			NodeObject->CustomCreatePin(EGPD_Input, Schema->PC_Object, ChildrenPinName, true);
-	
 			for (const FRegisteredObjectNodeInputPin& Pin : ICustomizableObjectModule::Get().GetAdditionalObjectNodePins())
 			{
 				// Use the global pin name here to prevent extensions using the same pin names from
@@ -175,12 +180,14 @@ void UCustomizableObjectNodeObject::BackwardsCompatibleFixup(int32 CustomizableO
 				//
 				// This also prevents extension pins from clashing with the built-in pins from this node,
 				// such as "Object".
-				UEdGraphPin* GraphPin = NodeObject->CustomCreatePin(EGPD_Input, Pin.InputPin.PinType, Pin.GlobalPinName, Pin.InputPin.bIsArray);
-
-				GraphPin->PinFriendlyName = Pin.InputPin.DisplayName;
+				UEdGraphPin* GraphPin = NodeObject->CustomCreatePin(EGPD_Input, Pin.InputPin.PinType, Pin.GlobalPinName, Pin.InputPin.DisplayName);
+				if (Pin.InputPin.bIsArray)
+				{
+					GraphPin->PinType.ContainerType = EPinContainerType::Array;
+				}
 			}
 
-			UEdGraphPin* OutputPin = NodeObject->CustomCreatePin(EGPD_Output, Schema->PC_Object, OutputPinName);
+			UEdGraphPin* OutputPin = NodeObject->CustomCreatePin(EGPD_Output, Schema->PC_Object, OutputPinName, LOCTEXT("Output", "Output"));
 
 			if (NodeObject->bIsBase)
 			{
@@ -243,7 +250,8 @@ void UCustomizableObjectNodeObject::BackwardsCompatibleFixup(int32 CustomizableO
 
 		if (!ComponentsPin()) // Some old nodes do not have the component pin.
 		{
-			CustomCreatePin(EGPD_Input, Schema->PC_Component, ComponentsPinName, true);
+			UEdGraphPin* Pin = CustomCreatePin(EGPD_Input, Schema->PC_Component, ComponentsPinName, LOCTEXT("Components", "Components"));
+			Pin->PinType.ContainerType = EPinContainerType::Array;
 		}
 		
 		if (!ParentObject && // Is a root object
@@ -344,16 +352,16 @@ void UCustomizableObjectNodeObject::BackwardsCompatibleFixup(int32 CustomizableO
 
 					[&]() // Lambda to ease the control flow.
 					{
-						if (!NodeMaterialSwitch->GetNumElements()) // We should at least have a component to know were to connect this material. If not, not supported.
+						if (!NodeMaterialSwitch->SwitchPins.Num()) // We should at least have a component to know were to connect this material. If not, not supported.
 						{
 							return;
 						}
 
 						FName ComponentName;
 						bool bFirst = true;
-						for (int32 ElementIndex = 0; ElementIndex < NodeMaterialSwitch->GetNumElements(); ++ElementIndex)
+						for (int32 ElementIndex = 0; ElementIndex < NodeMaterialSwitch->SwitchPins.Num(); ++ElementIndex)
 						{
-							if (UEdGraphPin* ConnectedPin = FollowInputPin(*NodeMaterialSwitch->GetElementPin(ElementIndex)))
+							if (UEdGraphPin* ConnectedPin = FollowInputPin(*NodeMaterialSwitch->SwitchPins[ElementIndex].Get()))
 							{
 								if (UCustomizableObjectNodeMaterialBase* FirstNodeMaterialBase = Cast<UCustomizableObjectNodeMaterialBase>(ConnectedPin->GetOwningNode()))
 								{
@@ -540,7 +548,7 @@ void UCustomizableObjectNodeObject::BackwardsCompatibleFixup(int32 CustomizableO
 	{
 		if (!NamePin.Get())
 		{
-			NamePin = CustomCreatePin(EGPD_Input, UEdGraphSchema_CustomizableObject::PC_String, "Name");
+			NamePin = CustomCreatePin(EGPD_Input, UEdGraphSchema_CustomizableObject::PC_String, TEXT("Name"), LOCTEXT("Name", "Name"));
 		}
 	}
 }
@@ -581,13 +589,16 @@ void UCustomizableObjectNodeObject::AllocateDefaultPins(UCustomizableObjectNodeR
 {
 	const UEdGraphSchema_CustomizableObject* Schema = GetDefault<UEdGraphSchema_CustomizableObject>();
 
-	NamePin = CustomCreatePin(EGPD_Input, UEdGraphSchema_CustomizableObject::PC_String, "Name");
+	NamePin = CustomCreatePin(EGPD_Input, UEdGraphSchema_CustomizableObject::PC_String, TEXT("Name"), LOCTEXT("Name", "Name"));
 	
-	CustomCreatePin(EGPD_Input, Schema->PC_Component, ComponentsPinName, true);
+	UEdGraphPin* ComponentsPin = CustomCreatePin(EGPD_Input, Schema->PC_Component, ComponentsPinName, LOCTEXT("Components", "Components"));
+	ComponentsPin->PinType.ContainerType = EPinContainerType::Array;
 	
-	CustomCreatePin(EGPD_Input, Schema->PC_Modifier, ModifiersPinName, true);
+	UEdGraphPin* ModifiersPin = CustomCreatePin(EGPD_Input, Schema->PC_Modifier, ModifiersPinName, LOCTEXT("Modifiers", "Modifiers"));
+	ModifiersPin->PinType.ContainerType = EPinContainerType::Array;
 
-	CustomCreatePin(EGPD_Input, Schema->PC_Object, ChildrenPinName, true);
+	UEdGraphPin* ChildrenPin = CustomCreatePin(EGPD_Input, Schema->PC_Object, ChildrenPinName, LOCTEXT("Children", "Children"));
+	ChildrenPin->PinType.ContainerType = EPinContainerType::Array;
 	
 	for (const FRegisteredObjectNodeInputPin& Pin : ICustomizableObjectModule::Get().GetAdditionalObjectNodePins())
 	{
@@ -596,12 +607,15 @@ void UCustomizableObjectNodeObject::AllocateDefaultPins(UCustomizableObjectNodeR
 		//
 		// This also prevents extension pins from clashing with the built-in pins from this node,
 		// such as "Object".
-		UEdGraphPin* GraphPin = CustomCreatePin(EGPD_Input, Pin.InputPin.PinType, Pin.GlobalPinName, Pin.InputPin.bIsArray);
+		UEdGraphPin* GraphPin = CustomCreatePin(EGPD_Input, Pin.InputPin.PinType, Pin.GlobalPinName, Pin.InputPin.DisplayName);
+		if (Pin.InputPin.bIsArray)
+		{
+			GraphPin->PinType.ContainerType = EPinContainerType::Array;
 
-		GraphPin->PinFriendlyName = Pin.InputPin.DisplayName;
+		}
 	}
 
-	UEdGraphPin* OutputPin = CustomCreatePin(EGPD_Output, Schema->PC_Object, OutputPinName);
+	UEdGraphPin* OutputPin = CustomCreatePin(EGPD_Output, Schema->PC_Object, OutputPinName, LOCTEXT("Output", "Output"));
 
 	if (bIsBase)
 	{
diff --git a/Engine/Plugins/Mutable/Source/CustomizableObjectEditor/Private/MuCOE/Nodes/CustomizableObjectNodeObjectGroup.cpp b/Engine/Plugins/Mutable/Source/CustomizableObjectEditor/Private/MuCOE/Nodes/CustomizableObjectNodeObjectGroup.cpp
--- a/Engine/Plugins/Mutable/Source/CustomizableObjectEditor/Private/MuCOE/Nodes/CustomizableObjectNodeObjectGroup.cpp
+++ b/Engine/Plugins/Mutable/Source/CustomizableObjectEditor/Private/MuCOE/Nodes/CustomizableObjectNodeObjectGroup.cpp
@@ -129,13 +129,15 @@ void UCustomizableObjectNodeObjectGroup::AllocateDefaultPins(UCustomizableObject
 {
 	const UEdGraphSchema_CustomizableObject* Schema = GetDefault<UEdGraphSchema_CustomizableObject>();
 
-	NamePin = CustomCreatePin(EGPD_Input, UEdGraphSchema_CustomizableObject::PC_String, FName("Name"));
+	NamePin = CustomCreatePin(EGPD_Input, UEdGraphSchema_CustomizableObject::PC_String, TEXT("Name"), LOCTEXT("Name", "Name"));
 
-	CustomCreatePin(EGPD_Input, Schema->PC_Object, FName("Objects"), true);
+	UEdGraphPin* ObjectsPin = CustomCreatePin(EGPD_Input, Schema->PC_Object, TEXT("Objects"), LOCTEXT("Objects", "Objects"));
+	ObjectsPin->PinType.ContainerType = EPinContainerType::Array;
 	
-	CustomCreatePin(EGPD_Input, Schema->PC_GroupProjector, FName("Projectors"), true);
-
-	CustomCreatePin(EGPD_Output, Schema->PC_Object, FName("Group"));
+	UEdGraphPin* ProjectorsPin = CustomCreatePin(EGPD_Input, Schema->PC_GroupProjector, TEXT("Projectors"), LOCTEXT("Projectors", "Projectors"));
+	ProjectorsPin->PinType.ContainerType = EPinContainerType::Array;
+	
+	CustomCreatePin(EGPD_Output, Schema->PC_Object, TEXT("Group"), LOCTEXT("Group", "Group"));
 }
 
 
@@ -153,7 +155,7 @@ void UCustomizableObjectNodeObjectGroup::BackwardsCompatibleFixup(int32 Customiz
 	{
 		if (!NamePin.Get())
 		{
-			NamePin = CustomCreatePin(EGPD_Input, UEdGraphSchema_CustomizableObject::PC_String, FName("Name"));
+			NamePin = CustomCreatePin(EGPD_Input, UEdGraphSchema_CustomizableObject::PC_String, TEXT("Name"), LOCTEXT("Name", "Name"));
 		}
 	}
 }
diff --git a/Engine/Plugins/Mutable/Source/CustomizableObjectEditor/Private/MuCOE/Nodes/CustomizableObjectNodeParameter.cpp b/Engine/Plugins/Mutable/Source/CustomizableObjectEditor/Private/MuCOE/Nodes/CustomizableObjectNodeParameter.cpp
--- a/Engine/Plugins/Mutable/Source/CustomizableObjectEditor/Private/MuCOE/Nodes/CustomizableObjectNodeParameter.cpp
+++ b/Engine/Plugins/Mutable/Source/CustomizableObjectEditor/Private/MuCOE/Nodes/CustomizableObjectNodeParameter.cpp
@@ -22,11 +22,11 @@ void UCustomizableObjectNodeParameter::AllocateDefaultPins(UCustomizableObjectNo
 	const FName PinName = UEdGraphSchema_CustomizableObject::GetPinCategoryName(Type);
 	const FText PinFriendlyName = GetCategoryFriendlyName();
 	
-	UEdGraphPin* ValuePin = CustomCreatePin(EGPD_Output, Type, PinName);
+	UEdGraphPin* ValuePin = CustomCreatePin(EGPD_Output, Type, PinName, PinFriendlyName);
 	ValuePin->PinFriendlyName = PinFriendlyName;
 	ValuePin->bDefaultValueIsIgnored = true;
 
-	NamePin = CustomCreatePin(EGPD_Input, UEdGraphSchema_CustomizableObject::PC_String, FName("Name"));
+	NamePin = CustomCreatePin(EGPD_Input, UEdGraphSchema_CustomizableObject::PC_String, TEXT("Name"), LOCTEXT("Name", "Name"));
 }
 
 
@@ -44,7 +44,7 @@ void UCustomizableObjectNodeParameter::BackwardsCompatibleFixup(int32 Customizab
 	{
 		if (!NamePin.Get())
 		{
-			NamePin = CustomCreatePin(EGPD_Input, UEdGraphSchema_CustomizableObject::PC_String, FName("Name"));
+			NamePin = CustomCreatePin(EGPD_Input, UEdGraphSchema_CustomizableObject::PC_String, TEXT("Name"), LOCTEXT("Name", "Name"));
 		}
 	}
 }
@@ -88,7 +88,7 @@ void UCustomizableObjectNodeParameter::SetPinEditableName(const UEdGraphPin& Pin
 
 FText UCustomizableObjectNodeParameter::GetNodeTitle(ENodeTitleType::Type TitleType) const
 {
-	return FText::Format(LOCTEXT("ParameterTitle_ListView", "{0} Parameter"), GetCategoryFriendlyName());
+	return FText::Format(LOCTEXT("ParameterTitle_ListView", "{0} Parameter"), UEdGraphSchema_CustomizableObject::GetPinCategoryFriendlyName(GetCategory()));
 }
 
 
diff --git a/Engine/Plugins/Mutable/Source/CustomizableObjectEditor/Private/MuCOE/Nodes/CustomizableObjectNodePassThroughTexture.cpp b/Engine/Plugins/Mutable/Source/CustomizableObjectEditor/Private/MuCOE/Nodes/CustomizableObjectNodePassThroughTexture.cpp
--- a/Engine/Plugins/Mutable/Source/CustomizableObjectEditor/Private/MuCOE/Nodes/CustomizableObjectNodePassThroughTexture.cpp
+++ b/Engine/Plugins/Mutable/Source/CustomizableObjectEditor/Private/MuCOE/Nodes/CustomizableObjectNodePassThroughTexture.cpp
@@ -39,8 +39,7 @@ void UCustomizableObjectNodePassThroughTexture::BackwardsCompatibleFixup(int32 C
 
 void UCustomizableObjectNodePassThroughTexture::AllocateDefaultPins(UCustomizableObjectNodeRemapPins* RemapPins)
 {
-	UEdGraphPin* PinImagePin = CustomCreatePinSimple(EGPD_Output, UEdGraphSchema_CustomizableObject::PC_PassthroughTexture);
-	PinImagePin->bDefaultValueIsIgnored = true;
+	CustomCreatePin(EGPD_Output, UEdGraphSchema_CustomizableObject::PC_Texture_Passthrough);
 }
 
 
@@ -62,7 +61,7 @@ FText UCustomizableObjectNodePassThroughTexture::GetNodeTitle(ENodeTitleType::Ty
 
 FLinearColor UCustomizableObjectNodePassThroughTexture::GetNodeTitleColor() const
 {
-	return UEdGraphSchema_CustomizableObject::GetPinTypeColor(UEdGraphSchema_CustomizableObject::PC_PassthroughTexture);
+	return UEdGraphSchema_CustomizableObject::GetPinTypeColor(UEdGraphSchema_CustomizableObject::PC_Texture_Passthrough);
 }
 
 
diff --git a/Engine/Plugins/Mutable/Source/CustomizableObjectEditor/Private/MuCOE/Nodes/CustomizableObjectNodePassThroughTextureSwitch.cpp b/Engine/Plugins/Mutable/Source/CustomizableObjectEditor/Private/MuCOE/Nodes/CustomizableObjectNodePassThroughTextureSwitch.cpp
--- a/Engine/Plugins/Mutable/Source/CustomizableObjectEditor/Private/MuCOE/Nodes/CustomizableObjectNodePassThroughTextureSwitch.cpp
+++ b/Engine/Plugins/Mutable/Source/CustomizableObjectEditor/Private/MuCOE/Nodes/CustomizableObjectNodePassThroughTextureSwitch.cpp
@@ -11,7 +11,7 @@
 
 FName UCustomizableObjectNodePassThroughTextureSwitch::GetCategory() const
 {
-	return UEdGraphSchema_CustomizableObject::PC_PassthroughTexture;
+	return UEdGraphSchema_CustomizableObject::PC_Texture_Passthrough;
 }
 
 
diff --git a/Engine/Plugins/Mutable/Source/CustomizableObjectEditor/Private/MuCOE/Nodes/CustomizableObjectNodePassThroughTextureVariation.cpp b/Engine/Plugins/Mutable/Source/CustomizableObjectEditor/Private/MuCOE/Nodes/CustomizableObjectNodePassThroughTextureVariation.cpp
--- a/Engine/Plugins/Mutable/Source/CustomizableObjectEditor/Private/MuCOE/Nodes/CustomizableObjectNodePassThroughTextureVariation.cpp
+++ b/Engine/Plugins/Mutable/Source/CustomizableObjectEditor/Private/MuCOE/Nodes/CustomizableObjectNodePassThroughTextureVariation.cpp
@@ -13,7 +13,7 @@
 // UCustomizableObjectNodeVariation interface
 FName UCustomizableObjectNodePassThroughTextureVariation::GetCategory() const
 {
-	return UEdGraphSchema_CustomizableObject::PC_PassthroughTexture;
+	return UEdGraphSchema_CustomizableObject::PC_Texture_Passthrough;
 }
 
 
diff --git a/Engine/Plugins/Mutable/Source/CustomizableObjectEditor/Private/MuCOE/Nodes/CustomizableObjectNodeProjectorConstant.cpp b/Engine/Plugins/Mutable/Source/CustomizableObjectEditor/Private/MuCOE/Nodes/CustomizableObjectNodeProjectorConstant.cpp
--- a/Engine/Plugins/Mutable/Source/CustomizableObjectEditor/Private/MuCOE/Nodes/CustomizableObjectNodeProjectorConstant.cpp
+++ b/Engine/Plugins/Mutable/Source/CustomizableObjectEditor/Private/MuCOE/Nodes/CustomizableObjectNodeProjectorConstant.cpp
@@ -44,13 +44,7 @@ void UCustomizableObjectNodeProjectorConstant::PostEditChangeProperty(FPropertyC
 
 void UCustomizableObjectNodeProjectorConstant::AllocateDefaultPins(UCustomizableObjectNodeRemapPins* RemapPins)
 {
-	const FName Type = UEdGraphSchema_CustomizableObject::PC_Projector;
-	const FName PinName = UEdGraphSchema_CustomizableObject::GetPinCategoryName(Type);
-	const FText PinFriendlyName = UEdGraphSchema_CustomizableObject::GetPinCategoryFriendlyName(Type);
-	
-	UEdGraphPin* ValuePin = CustomCreatePin(EGPD_Output, Type, PinName);
-	ValuePin->PinFriendlyName = PinFriendlyName;
-	ValuePin->bDefaultValueIsIgnored = true;
+	CustomCreatePin(EGPD_Output, UEdGraphSchema_CustomizableObject::PC_Projector);
 }
 
 
diff --git a/Engine/Plugins/Mutable/Source/CustomizableObjectEditor/Private/MuCOE/Nodes/CustomizableObjectNodeReroute.cpp b/Engine/Plugins/Mutable/Source/CustomizableObjectEditor/Private/MuCOE/Nodes/CustomizableObjectNodeReroute.cpp
--- a/Engine/Plugins/Mutable/Source/CustomizableObjectEditor/Private/MuCOE/Nodes/CustomizableObjectNodeReroute.cpp
+++ b/Engine/Plugins/Mutable/Source/CustomizableObjectEditor/Private/MuCOE/Nodes/CustomizableObjectNodeReroute.cpp
@@ -15,9 +15,7 @@
 
 void UCustomizableObjectNodeReroute::AllocateDefaultPins(UCustomizableObjectNodeRemapPins* RemapPins)
 {
-	UEdGraphPin* MyInputPin = CreatePin(EGPD_Input, UEdGraphSchema_CustomizableObject::PC_Wildcard, TEXT("InputPin"));
-	MyInputPin->bDefaultValueIsIgnored = true;
-
+	CreatePin(EGPD_Input, UEdGraphSchema_CustomizableObject::PC_Wildcard, TEXT("InputPin"));
 	CreatePin(EGPD_Output, UEdGraphSchema_CustomizableObject::PC_Wildcard, TEXT("OutputPin"));
 }
 
diff --git a/Engine/Plugins/Mutable/Source/CustomizableObjectEditor/Private/MuCOE/Nodes/CustomizableObjectNodeSkeletalMesh.cpp b/Engine/Plugins/Mutable/Source/CustomizableObjectEditor/Private/MuCOE/Nodes/CustomizableObjectNodeSkeletalMesh.cpp
--- a/Engine/Plugins/Mutable/Source/CustomizableObjectEditor/Private/MuCOE/Nodes/CustomizableObjectNodeSkeletalMesh.cpp
+++ b/Engine/Plugins/Mutable/Source/CustomizableObjectEditor/Private/MuCOE/Nodes/CustomizableObjectNodeSkeletalMesh.cpp
@@ -30,9 +30,6 @@ struct FSlateBrush;
 #define LOCTEXT_NAMESPACE "CustomizableObjectEditor"
 
 
-/** Default node pin configuration pin name (node does not have an skeletal mesh). */
-static const TCHAR* SKELETAL_MESH_PIN_NAME = TEXT("Skeletal Mesh");
-
 
 void UCustomizableObjectNodeSkeletalMesh::LoadObjects()
 {
@@ -81,7 +78,7 @@ void UCustomizableObjectNodeSkeletalMesh::AllocateDefaultPins(UCustomizableObjec
 		UCustomizableObjectNodeSkeletalMeshPinDataMesh* PinData = NewObject<UCustomizableObjectNodeSkeletalMeshPinDataMesh>(this);
 		PinData->Init(-1, -1, -1);
 		
-		DefaultPin = CustomCreatePin(EGPD_Output, Schema->PC_Mesh, FName(SKELETAL_MESH_PIN_NAME), PinData);
+		DefaultPin = CustomCreatePin(EGPD_Output, Schema->PC_Mesh, TEXT("Mesh"), LOCTEXT("Mesh", "Mesh"), PinData);
 		return;
 	}
 	else
@@ -118,8 +115,8 @@ void UCustomizableObjectNodeSkeletalMesh::AllocateDefaultPins(UCustomizableObjec
 					
 					const FString MeshPinName = FString::Printf(TEXT("LOD %i - Section %i - Mesh"), LODIndex, SectionIndex);
 					
-					UEdGraphPin* Pin = CustomCreatePin(EGPD_Output, Schema->PC_Mesh, FName(*MeshPinName), PinData);
-					Pin->PinFriendlyName = FText::FromString(FString::Printf(TEXT("LOD %i - %s"), LODIndex, *Section));
+					FText FriendlyName = FText::FromString(FString::Printf(TEXT("LOD %i - %s"), LODIndex, *Section));
+					CustomCreatePin(EGPD_Output, Schema->PC_Mesh, FName(*MeshPinName), FriendlyName, PinData);
 				}
 				
 				
@@ -143,9 +140,9 @@ void UCustomizableObjectNodeSkeletalMesh::AllocateDefaultPins(UCustomizableObjec
 						const FMaterialParameterInfo& ImageInfo = ImageInfos[Index];
 						const FString ImageNameStr = *ImageInfo.Name.ToString();
 						const FString TexturePinName = FString::Printf(TEXT("LOD %i - Section %i - Texture Parameter %s"), LODIndex, SectionIndex, *ImageNameStr);
-						
-						UEdGraphPin* Pin = CustomCreatePin(EGPD_Output, Schema->PC_Texture, FName(*TexturePinName), PinData);
-						Pin->PinFriendlyName = FText::FromString(FString::Printf(TEXT("LOD %i - %s - %s"), LODIndex, *Section, *ImageNameStr));
+
+						FText FriendlyName = FText::FromString(FString::Printf(TEXT("LOD %i - %s - %s"), LODIndex, *Section, *ImageNameStr));
+						UEdGraphPin* Pin = CustomCreatePin(EGPD_Output, Schema->PC_Texture, FName(*TexturePinName), FriendlyName, PinData);
 						Pin->bHidden = true;
 					}
 				}
@@ -676,6 +673,17 @@ void UCustomizableObjectNodeSkeletalMesh::BackwardsCompatibleFixup(int32 Customi
 			}
 		}
 	}
+	
+	if (CustomizableObjectCustomVersion == FCustomizableObjectCustomVersion::PassthroughParameter)
+	{
+		for (UEdGraphPin* Pin : Pins)
+		{
+			if (Pin->PinName == TEXT("Skeletal Mesh"))
+			{
+				Pin->PinName = TEXT("Mesh");
+			}
+		}
+	}
 }
 
 
diff --git a/Engine/Plugins/Mutable/Source/CustomizableObjectEditor/Private/MuCOE/Nodes/CustomizableObjectNodeSkeletalMeshParameter.cpp b/Engine/Plugins/Mutable/Source/CustomizableObjectEditor/Private/MuCOE/Nodes/CustomizableObjectNodeSkeletalMeshParameter.cpp
--- a/Engine/Plugins/Mutable/Source/CustomizableObjectEditor/Private/MuCOE/Nodes/CustomizableObjectNodeSkeletalMeshParameter.cpp
+++ b/Engine/Plugins/Mutable/Source/CustomizableObjectEditor/Private/MuCOE/Nodes/CustomizableObjectNodeSkeletalMeshParameter.cpp
@@ -4,7 +4,6 @@
 
 #include "MuCOE/EdGraphSchema_CustomizableObject.h"
 #include "MuCOE/CustomizableObjectLayout.h"
-#include "MuCOE/CustomizableObjectMacroLibrary/CustomizableObjectMacroLibrary.h"
 #include "MuCO/CustomizableObjectCustomVersion.h"
 #include "Rendering/SkeletalMeshModel.h"
 #include "Engine/SkeletalMesh.h"
@@ -17,11 +16,11 @@ class UCustomizableObjectNodeRemapPins;
 
 void UCustomizableObjectNodeSkeletalMeshParameter::AllocateDefaultPins(UCustomizableObjectNodeRemapPins* RemapPins)
 {
-	NamePin = CustomCreatePin(EGPD_Input, UEdGraphSchema_CustomizableObject::PC_String, FName("Name"));
+	NamePin = CustomCreatePin(EGPD_Input, UEdGraphSchema_CustomizableObject::PC_String, TEXT("Name"), LOCTEXT("Name", "Name"));
 
 	const UEdGraphSchema_CustomizableObject* Schema = GetDefault<UEdGraphSchema_CustomizableObject>();
 	
-	PassthroughPin = CustomCreatePin(EGPD_Output, Schema->PC_PassthroughSkeletalMesh, TEXT("Passthrough Skeletal Mesh"));
+	PassthroughPin = CustomCreatePin(EGPD_Output, Schema->PC_SkeletalMesh_PassthroughParameter);
 	
 	USkeletalMesh* SkeletalMesh = ReferenceValue.LoadSynchronous();
 
@@ -55,8 +54,7 @@ void UCustomizableObjectNodeSkeletalMeshParameter::AllocateDefaultPins(UCustomiz
 					const FString MeshPinName = FString::Printf(TEXT("LOD %i - Section %i: %s"), LODIndex, SectionIndex, *MaterialName.ToString());
 					const FText MeshPinFriendlyName = FText::Format(LOCTEXT("SkeletalMeshParameterMeshPin", "LOD {0} - Section {1}: {2}"), LODIndex, SectionIndex, FText::FromString(*MaterialName.ToString()));
 					
-					UEdGraphPin* Pin = CustomCreatePin(EGPD_Output, Schema->PC_Mesh, FName(*MeshPinName), PinData);
-					Pin->PinFriendlyName = MeshPinFriendlyName;
+					CustomCreatePin(EGPD_Output, Schema->PC_Mesh, FName(*MeshPinName), MeshPinFriendlyName, PinData);
 				}
 			}			
 		}
@@ -90,13 +88,24 @@ void UCustomizableObjectNodeSkeletalMeshParameter::BackwardsCompatibleFixup(int3
 	{
 		if (!NamePin.Get())
 		{
-			NamePin = CustomCreatePin(EGPD_Input, UEdGraphSchema_CustomizableObject::PC_String, FName("Name"));
+			NamePin = CustomCreatePin(EGPD_Input, UEdGraphSchema_CustomizableObject::PC_String, TEXT("Name"), LOCTEXT("Name", "Name"));
 		}
 	}
 
 	if (CustomizableObjectCustomVersion == FCustomizableObjectCustomVersion::PinSkeletalMesh)
 	{
-		PassthroughPin = CustomCreatePin(EGPD_Output, UEdGraphSchema_CustomizableObject::PC_PassthroughSkeletalMesh, TEXT("Passthrough Skeletal Mesh"));
+		PassthroughPin = CustomCreatePin(EGPD_Output, UEdGraphSchema_CustomizableObject::PC_SkeletalMesh_Passthrough);
+	}
+	
+	if (CustomizableObjectCustomVersion == FCustomizableObjectCustomVersion::PassthroughParameter)
+	{
+		for (UEdGraphPin* Pin : Pins)
+		{
+			if (Pin->PinType.PinCategory ==  UEdGraphSchema_CustomizableObject::PC_SkeletalMesh_Passthrough)
+			{
+				Pin->PinType.PinCategory =  UEdGraphSchema_CustomizableObject::PC_SkeletalMesh_PassthroughParameter;
+			}
+		}
 	}
 }
 
diff --git a/Engine/Plugins/Mutable/Source/CustomizableObjectEditor/Private/MuCOE/Nodes/CustomizableObjectNodeStaticMesh.cpp b/Engine/Plugins/Mutable/Source/CustomizableObjectEditor/Private/MuCOE/Nodes/CustomizableObjectNodeStaticMesh.cpp
--- a/Engine/Plugins/Mutable/Source/CustomizableObjectEditor/Private/MuCOE/Nodes/CustomizableObjectNodeStaticMesh.cpp
+++ b/Engine/Plugins/Mutable/Source/CustomizableObjectEditor/Private/MuCOE/Nodes/CustomizableObjectNodeStaticMesh.cpp
@@ -121,8 +121,7 @@ void UCustomizableObjectNodeStaticMesh::AllocateDefaultPins(UCustomizableObjectN
 				
 				// Mesh
 				const FString MeshName = MaterialName + TEXT(" Mesh");
-				UEdGraphPin* Pin = CustomCreatePin(EGPD_Output, Schema->PC_Mesh, FName(*MeshName));
-				Pin->bDefaultValueIsIgnored = true;
+				UEdGraphPin* Pin = CustomCreatePin(EGPD_Output, Schema->PC_Mesh, FName(*MeshName), FText::FromString(MeshName));
 				LODs[LODIndex].Materials[MaterialIndex].MeshPinRef.SetPin(Pin);
 
 				// Images
@@ -136,8 +135,7 @@ void UCustomizableObjectNodeStaticMesh::AllocateDefaultPins(UCustomizableObjectN
 					for (int32 ImageIndex = 0; ImageIndex < ImageNames.Num(); ++ImageIndex)
 					{
 						FString ImageName = ImageNames[ImageIndex].Name.ToString();
-						Pin = CustomCreatePin(EGPD_Output, Schema->PC_Texture, *ImageName);
-						Pin->bDefaultValueIsIgnored = true;
+						Pin = CustomCreatePin(EGPD_Output, Schema->PC_Texture, *ImageName, FText::FromString(ImageName));
 						LODs[LODIndex].Materials[MaterialIndex].ImagePinsRef.Add(Pin);
 					}
 				}
@@ -148,11 +146,7 @@ void UCustomizableObjectNodeStaticMesh::AllocateDefaultPins(UCustomizableObjectN
 	}
 	else
 	{
-		UEdGraphPin* Pin = CustomCreatePin(EGPD_Output, Schema->PC_Mesh, FName(STATIC_MESH_PIN_NAME));
-		Pin->bDefaultValueIsIgnored = true;
-		Pin->PinFriendlyName = FText::FromString(STATIC_MESH_PIN_NAME);
-
-		DefaultPin = FEdGraphPinReference(Pin);
+		DefaultPin = CustomCreatePin(EGPD_Output, Schema->PC_Mesh, FName(STATIC_MESH_PIN_NAME), FText::FromString(STATIC_MESH_PIN_NAME));
 	}
 }
 
@@ -421,8 +415,7 @@ void UCustomizableObjectNodeStaticMesh::BackwardsCompatibleFixup(int32 Customiza
 						
 						// Mesh
 						const FString MeshName = MaterialName + TEXT(" Mesh");
-						UEdGraphPin* Pin = CustomCreatePin(EGPD_Output, UEdGraphSchema_CustomizableObject::PC_Mesh, FName(*MeshName));
-						Pin->bDefaultValueIsIgnored = true;
+						UEdGraphPin* Pin = CustomCreatePin(EGPD_Output, UEdGraphSchema_CustomizableObject::PC_Mesh, FName(*MeshName), FText::FromString(MeshName));
 						LODs[LODIndex].Materials[MaterialIndex].MeshPinRef.SetPin(Pin);
 
 						// Images
@@ -437,8 +430,7 @@ void UCustomizableObjectNodeStaticMesh::BackwardsCompatibleFixup(int32 Customiza
 								for (int32 ImageIndex = 0; ImageIndex < ImageNames.Num(); ++ImageIndex)
 								{
 									FString ImageName = ImageNames[ImageIndex].Name.ToString();
-									Pin = CustomCreatePin(EGPD_Output, UEdGraphSchema_CustomizableObject::PC_Texture, *ImageName);
-									Pin->bDefaultValueIsIgnored = true;
+									Pin = CustomCreatePin(EGPD_Output, UEdGraphSchema_CustomizableObject::PC_Texture, *ImageName, FText::FromString(ImageName));
 									LODs[LODIndex].Materials[MaterialIndex].ImagePinsRef.Add(Pin);
 								}
 							}	
diff --git a/Engine/Plugins/Mutable/Source/CustomizableObjectEditor/Private/MuCOE/Nodes/CustomizableObjectNodeStaticString.cpp b/Engine/Plugins/Mutable/Source/CustomizableObjectEditor/Private/MuCOE/Nodes/CustomizableObjectNodeStaticString.cpp
--- a/Engine/Plugins/Mutable/Source/CustomizableObjectEditor/Private/MuCOE/Nodes/CustomizableObjectNodeStaticString.cpp
+++ b/Engine/Plugins/Mutable/Source/CustomizableObjectEditor/Private/MuCOE/Nodes/CustomizableObjectNodeStaticString.cpp
@@ -32,7 +32,7 @@ FLinearColor UCustomizableObjectNodeStaticString::GetNodeTitleColor() const
 
 void UCustomizableObjectNodeStaticString::AllocateDefaultPins(UCustomizableObjectNodeRemapPins* RemapPins)
 {
-	CustomCreatePin(EGPD_Output, UEdGraphSchema_CustomizableObject::PC_String, FName("Value"));
+	CustomCreatePin(EGPD_Output, UEdGraphSchema_CustomizableObject::PC_String, TEXT("Value"), LOCTEXT("Value", "Value"));
 }
 
 
diff --git a/Engine/Plugins/Mutable/Source/CustomizableObjectEditor/Private/MuCOE/Nodes/CustomizableObjectNodeSwitchBase.cpp b/Engine/Plugins/Mutable/Source/CustomizableObjectEditor/Private/MuCOE/Nodes/CustomizableObjectNodeSwitchBase.cpp
--- a/Engine/Plugins/Mutable/Source/CustomizableObjectEditor/Private/MuCOE/Nodes/CustomizableObjectNodeSwitchBase.cpp
+++ b/Engine/Plugins/Mutable/Source/CustomizableObjectEditor/Private/MuCOE/Nodes/CustomizableObjectNodeSwitchBase.cpp
@@ -3,6 +3,7 @@
 #include "MuCOE/Nodes/CustomizableObjectNodeSwitchBase.h"
 
 #include "MuCO/CustomizableObjectCustomVersion.h"
+#include "MuCOE/CustomizableObjectPin.h"
 #include "MuCOE/EdGraphSchema_CustomizableObject.h"
 #include "MuCOE/GraphTraversal.h"
 #include "MuCOE/Nodes/CustomizableObjectNodeEnumParameter.h"
@@ -13,50 +14,11 @@ class UCustomizableObjectNodeRemapPins;
 
 #define LOCTEXT_NAMESPACE "CustomizableObjectEditor"
 
-void UCustomizableObjectNodeSwitchBase::ReloadEnumParam()
-{
-	ReloadingElementsNames.Empty();
-
-	// Get the names of the enum parameter for the element pins
-	if (const UEdGraphPin* EnumPin = SwitchParameter())
-	{
-		if (const UEdGraphPin* LinkedPin = FollowInputPin(*EnumPin))
-		{
-			if (UCustomizableObjectNodeEnumParameter* EnumNode = Cast<UCustomizableObjectNodeEnumParameter>(LinkedPin->GetOwningNode()))
-			{
-				for (int i = 0; i < EnumNode->Values.Num(); ++i)
-				{
-					ReloadingElementsNames.Add(EnumNode->Values[i].Name);
-				}
-			}
-		}
-	}
-}
-
-
-void UCustomizableObjectNodeSwitchBase::ReconstructNode(UCustomizableObjectNodeRemapPins* RemapPins)
-{
-	ReloadEnumParam();
-
-	Super::ReconstructNode(RemapPins);
-}
-
-
-UEdGraphPin* UCustomizableObjectNodeSwitchBase::OutputPin() const
-{
-	return OutputPinReference.Get();
-}
-
-
-UEdGraphPin* UCustomizableObjectNodeSwitchBase::SwitchParameter() const
-{
-	return SwitchParameterPinReference.Get();
-}
 
 
 void UCustomizableObjectNodeSwitchBase::PinConnectionListChanged(UEdGraphPin* Pin)
 {
-	if (Pin == SwitchParameter())
+	if (Pin == SwitchParameterPinReference.Get())
 	{
 		LinkPostEditChangePropertyDelegate(*Pin);
 	}
@@ -68,39 +30,26 @@ void UCustomizableObjectNodeSwitchBase::PinConnectionListChanged(UEdGraphPin* Pi
 void UCustomizableObjectNodeSwitchBase::AllocateDefaultPins(UCustomizableObjectNodeRemapPins* RemapPins)
 {
 	const UEdGraphSchema_CustomizableObject* Schema = GetDefault<UEdGraphSchema_CustomizableObject>();
-
-	UEdGraphPin* OutputPin = CustomCreatePin(EGPD_Output, GetCategory(), FName(UEdGraphSchema_CustomizableObject::GetPinCategoryName(GetCategory()).ToString()));
-	OutputPin->bDefaultValueIsIgnored = true;
-	OutputPin->PinFriendlyName = UEdGraphSchema_CustomizableObject::GetPinCategoryFriendlyName(GetCategory());
-
-	OutputPinReference = FEdGraphPinReference(OutputPin);
-
-	UEdGraphPin* SwitchParameterPin = CustomCreatePin(EGPD_Input, Schema->PC_Enum, FName(TEXT("Switch Parameter")));
-	SwitchParameterPin->bDefaultValueIsIgnored = true;
-	SwitchParameterPin->SetOwningNode(this);
-
-	SwitchParameterPinReference = FEdGraphPinReference(SwitchParameterPin);
 	
-	for (int32 LayerIndex = 0; LayerIndex < ReloadingElementsNames.Num(); ++LayerIndex)
-	{
-		FString PinName = GetPinPrefix(LayerIndex);
-		UEdGraphPin* InputPin = CustomCreatePin(EGPD_Input, GetCategory(), FName(*PinName));
-		FString PinDisplayName = ReloadingElementsNames[LayerIndex].IsEmpty() ? PinName : ReloadingElementsNames[LayerIndex];
-		const FString CategoryFriendlyName = UEdGraphSchema_CustomizableObject::GetPinCategoryFriendlyName(GetCategory()).ToString();
-		InputPin->PinFriendlyName = FText::FromString(CategoryFriendlyName + TEXT(" ") + FString::FromInt(LayerIndex));
-		InputPin->bDefaultValueIsIgnored = true;
-		InputPin->SetOwningNode(this);
-	}
-}
-
+	CustomCreatePin(EGPD_Output, GetCategory(), FName(UEdGraphSchema_CustomizableObject::GetPinCategoryName(GetCategory()).ToString()), UEdGraphSchema_CustomizableObject::GetPinCategoryFriendlyName(GetCategory()));
 
-void UCustomizableObjectNodeSwitchBase::Serialize(FArchive& Ar)
-{
-	Super::Serialize(Ar);
+	UEdGraphPin* OldEnumPin = SwitchParameterPinReference.Get();
+	SwitchParameterPinReference = CustomCreatePin(EGPD_Input, Schema->PC_Enum, TEXT("Switch Parameter"), LOCTEXT("SwitchParameter", "Switch Parameter"));
 
-	if (Ar.CustomVer(FCustomizableObjectCustomVersion::GUID) < FCustomizableObjectCustomVersion::BugPinsSwitch)
+	SwitchPins.Empty();	
+	if (OldEnumPin) // May be null the first time the node is created.
 	{
-		SwitchParameterPinReference = FindPin(TEXT("Switch Parameter"));	
+		if (const UEdGraphPin* LinkedPin = FollowInputPin(*OldEnumPin))
+		{
+			if (UCustomizableObjectNodeEnumParameter* EnumNode = Cast<UCustomizableObjectNodeEnumParameter>(LinkedPin->GetOwningNode()))
+			{
+				for (int i = 0; i < EnumNode->Values.Num(); ++i)
+				{
+					const FString OptionName = EnumNode->Values[i].Name;
+					SwitchPins.Add(CustomCreatePin(EGPD_Input, GetCategory(), FName(*OptionName), FText::FromString(OptionName)));
+				}
+			}
+		}
 	}
 }
 
@@ -125,38 +74,14 @@ FText UCustomizableObjectNodeSwitchBase::GetTooltipText() const
 
 void UCustomizableObjectNodeSwitchBase::EnumParameterPostEditChangeProperty(FPostEditChangePropertyDelegateParameters& Parameters)
 {
-	if (const UEdGraphPin* SwitchPin = SwitchParameter())
+	if (const UEdGraphPin* ConnectedPin = FollowInputPin(*SwitchParameterPinReference.Get()); ConnectedPin && ConnectedPin->GetOwningNode() == Parameters.Node)
 	{
-		if (const UEdGraphPin* ConnectedPin = FollowInputPin(*SwitchPin); ConnectedPin && ConnectedPin->GetOwningNode() == Parameters.Node)
-		{
-			Super::ReconstructNode();
-		}
-		else if (UCustomizableObjectNode* EnumNode = Cast<UCustomizableObjectNode>(Parameters.Node))
-		{
-			EnumNode->PostEditChangePropertyDelegate.RemoveDynamic(this, &UCustomizableObjectNodeSwitchBase::EnumParameterPostEditChangeProperty);
-		}
+		Super::ReconstructNode();
 	}
- }
-
-
-int32 UCustomizableObjectNodeSwitchBase::GetNumElements() const
-{
-	int32 Count = 0;
-	for (UEdGraphPin* Pin : GetAllNonOrphanPins())
+	else if (UCustomizableObjectNode* EnumNode = Cast<UCustomizableObjectNode>(Parameters.Node))
 	{
-		if (Pin->GetName().StartsWith(GetPinPrefix()))
-		{
-			Count++;
-		}
+		EnumNode->PostEditChangePropertyDelegate.RemoveDynamic(this, &UCustomizableObjectNodeSwitchBase::EnumParameterPostEditChangeProperty);
 	}
-
-	return Count;
-}
-
-
-FString UCustomizableObjectNodeSwitchBase::GetPinPrefix(int32 Index) const
-{
-	return GetPinPrefix() + FString::FromInt(Index) + TEXT(" ");
 }
 
 
@@ -164,10 +89,7 @@ void UCustomizableObjectNodeSwitchBase::PostPasteNode()
 {
 	Super::PostPasteNode();
 
-	if (UEdGraphPin* SwitchPin = SwitchParameter())
-	{
-		LinkPostEditChangePropertyDelegate(*SwitchPin);
-	}
+	LinkPostEditChangePropertyDelegate(*SwitchParameterPinReference.Get());
 }
 
 
@@ -175,13 +97,28 @@ void UCustomizableObjectNodeSwitchBase::BackwardsCompatibleFixup(int32 Customiza
 {
 	Super::BackwardsCompatibleFixup(CustomizableObjectCustomVersion);
 
-	if (CustomizableObjectCustomVersion == FCustomizableObjectCustomVersion::ChangedSwitchNodesInputPinsFriendlyNames)
+	if (CustomizableObjectCustomVersion == FCustomizableObjectCustomVersion::BugPinsSwitch)
+	{
+		SwitchParameterPinReference = FindPin(TEXT("Switch Parameter"));	
+	}
+	
+	else if (CustomizableObjectCustomVersion == FCustomizableObjectCustomVersion::ChangedSwitchNodesInputPinsFriendlyNames)
 	{
 		const FString CategoryFriendlyName = UEdGraphSchema_CustomizableObject::GetPinCategoryFriendlyName(GetCategory()).ToString();
+		const FString PinPrefix = UEdGraphSchema_CustomizableObject::GetPinCategoryName(GetCategory()).ToString() + " ";
 
-		for (int32 ElementIndex = 0; ElementIndex < GetNumElements(); ++ElementIndex)
+		int32 NumOptions = 0;
+		for (UEdGraphPin* Pin : GetAllNonOrphanPins())
 		{
-			UEdGraphPin* ElementPin = GetElementPin(ElementIndex);
+			if (Pin->GetName().StartsWith(PinPrefix))
+			{
+				NumOptions++;
+			}
+		}
+		
+		for (int32 ElementIndex = 0; ElementIndex < NumOptions; ++ElementIndex)
+		{
+			UEdGraphPin* ElementPin = FindPin(PinPrefix + FString::FromInt(ElementIndex) + TEXT(" "));
 
 			if (ElementPin->PinFriendlyName.ToString().StartsWith(TEXT("Material")))
 			{
@@ -189,6 +126,47 @@ void UCustomizableObjectNodeSwitchBase::BackwardsCompatibleFixup(int32 Customiza
 			}
 		}
 	}
+	
+	else if (CustomizableObjectCustomVersion == FCustomizableObjectCustomVersion::PassthroughParameter)
+	{
+		const FString PinPrefix = UEdGraphSchema_CustomizableObject::GetPinCategoryName(GetCategory()).ToString() + " ";
+
+		int32 NumOptions = 0;
+		for (UEdGraphPin* Pin : GetAllNonOrphanPins())
+		{
+			if (Pin->GetName().StartsWith(PinPrefix))
+			{
+				NumOptions++;
+			}
+		}
+		
+		TArray<FString> OptionNames;
+		if (const UEdGraphPin* LinkedPin = FollowInputPin(*SwitchParameterPinReference.Get()))
+		{
+			if (UCustomizableObjectNodeEnumParameter* EnumNode = Cast<UCustomizableObjectNodeEnumParameter>(LinkedPin->GetOwningNode()))
+			{
+				for (int i = 0; i < EnumNode->Values.Num(); ++i)
+				{
+					OptionNames.Add(EnumNode->Values[i].Name);
+				}
+			}
+		}
+		
+		for (int32 ElementIndex = 0; ElementIndex < NumOptions; ++ElementIndex)
+		{
+			UEdGraphPin* ElementPin = FindPin(PinPrefix + FString::FromInt(ElementIndex) + TEXT(" "));
+			if (OptionNames.IsValidIndex(ElementIndex))
+			{
+				ElementPin->PinName = FName(OptionNames[ElementIndex]);
+				ElementPin->PinFriendlyName = FText::FromName(ElementPin->PinName);
+				SwitchPins.Add(ElementPin);
+			}
+			else
+			{
+				ElementPin->bOrphanedPin = true;
+			}
+		}
+	}
 }
 
 
@@ -196,7 +174,7 @@ void UCustomizableObjectNodeSwitchBase::PostBackwardsCompatibleFixup()
 {
 	Super::PostBackwardsCompatibleFixup();
 
-	if (UEdGraphPin* SwitchPin = SwitchParameter())
+	if (UEdGraphPin* SwitchPin = SwitchParameterPinReference.Get()) // New nodes execute AllocateDefaultPins after PostBackwardsCompatibleFixup.
 	{
 		LinkPostEditChangePropertyDelegate(*SwitchPin);
 	}
@@ -230,11 +208,5 @@ FString UCustomizableObjectNodeSwitchBase::GetOutputPinName() const
 }
 
 
-FString UCustomizableObjectNodeSwitchBase::GetPinPrefix() const
-{
-	return UEdGraphSchema_CustomizableObject::GetPinCategoryName(GetCategory()).ToString() + " ";
-}
-
-
 #undef LOCTEXT_NAMESPACE
 
diff --git a/Engine/Plugins/Mutable/Source/CustomizableObjectEditor/Private/MuCOE/Nodes/CustomizableObjectNodeSwitchBase.h b/Engine/Plugins/Mutable/Source/CustomizableObjectEditor/Private/MuCOE/Nodes/CustomizableObjectNodeSwitchBase.h
--- a/Engine/Plugins/Mutable/Source/CustomizableObjectEditor/Private/MuCOE/Nodes/CustomizableObjectNodeSwitchBase.h
+++ b/Engine/Plugins/Mutable/Source/CustomizableObjectEditor/Private/MuCOE/Nodes/CustomizableObjectNodeSwitchBase.h
@@ -21,73 +21,42 @@ class UCustomizableObjectNodeSwitchBase : public UCustomizableObjectNode
 {
 public:
 	GENERATED_BODY()
-
-	// UObject interface
-	UE_API virtual void Serialize(FArchive& Ar) override;
 	
 	// Begin EdGraphNode interface
-	UE_API virtual FText GetNodeTitle(ENodeTitleType::Type TitleType) const override;
-	UE_API virtual FLinearColor GetNodeTitleColor() const override;
-	UE_API virtual FText GetTooltipText() const override;
-	UE_API virtual void PinConnectionListChanged(UEdGraphPin* Pin) override;
-	UE_API virtual void PostPasteNode() override;
+	virtual FText GetNodeTitle(ENodeTitleType::Type TitleType) const override;
+	virtual FLinearColor GetNodeTitleColor() const override;
+	virtual FText GetTooltipText() const override;
+	virtual void PinConnectionListChanged(UEdGraphPin* Pin) override;
+	virtual void PostPasteNode() override;
 
 	// UCustomizableObjectNode interface
-	UE_API virtual void BackwardsCompatibleFixup(int32 CustomizableObjectCustomVersion) override;
-	UE_API virtual void PostBackwardsCompatibleFixup() override;
-	UE_API virtual void AllocateDefaultPins(UCustomizableObjectNodeRemapPins* RemapPins) override;
-	UE_API virtual void ReconstructNode(UCustomizableObjectNodeRemapPins* RemapPins) override;
+	virtual void BackwardsCompatibleFixup(int32 CustomizableObjectCustomVersion) override;
+	virtual void PostBackwardsCompatibleFixup() override;
+	virtual void AllocateDefaultPins(UCustomizableObjectNodeRemapPins* RemapPins) override;
 
 	/** Get the output pin category. Override. */
 	virtual FName GetCategory() const PURE_VIRTUAL(UCustomizableObjectNodeSwitchBase::GetCategory, return {}; );
-
-	UE_API UEdGraphPin* OutputPin() const;
-
-	UE_API UEdGraphPin* SwitchParameter() const;
-
-	UEdGraphPin* GetElementPin(int32 Index) const
-	{
-		return FindPin(GetPinPrefix(Index));
-	}
-
-	UE_API int32 GetNumElements() const;
-
+	
 	/** Links the PostEditChangeProperty delegate */
-	UE_API void LinkPostEditChangePropertyDelegate(const UEdGraphPin& Pin);
+	void LinkPostEditChangePropertyDelegate(const UEdGraphPin& Pin);
 
 	/** Get the ouput pin name. Override. */
-	UE_API virtual FString GetOutputPinName() const;
-
-private:
-	/** Get the pin prefix. Used for retrocompatibility. Override. */
-	UE_API virtual FString GetPinPrefix() const;
-
-protected:
-	UPROPERTY()
-	FEdGraphPinReference OutputPinReference;
-
-	UPROPERTY()
-	TArray<FString> ReloadingElementsNames;
+	virtual FString GetOutputPinName() const;
 
 private:
-	/** Get the pin prefix with index. Used for retrocompatibility.*/
-	UE_API FString GetPinPrefix(int32 Index) const;
-
-	UE_API void ReloadEnumParam();
-
 	/** Last NodeEnumParameter connected. Used to remove the callback once disconnected. */
 	TWeakObjectPtr<UCustomizableObjectNode> LastNodeEnumParameterConnected;
 
 	/** NodeEnumParameter property changed callback function. Reconstructs the node. */
 	UFUNCTION()
-	UE_API void EnumParameterPostEditChangeProperty(FPostEditChangePropertyDelegateParameters& Parameters);
-
-	/** The node has to be reconstructed. PinConnectionListChanged(...) can not reconstruct the node, flag used to reconstruct the node on NodeConnectionListChanged(). */
-	bool bMarkReconstruct = false;
-
+	void EnumParameterPostEditChangeProperty(FPostEditChangePropertyDelegateParameters& Parameters);
 
+public:
 	UPROPERTY()
 	FEdGraphPinReference SwitchParameterPinReference;
+	
+	UPROPERTY()
+	TArray<FEdGraphPinReference> SwitchPins;
 };
 
 #undef UE_API
diff --git a/Engine/Plugins/Mutable/Source/CustomizableObjectEditor/Private/MuCOE/Nodes/CustomizableObjectNodeTable.cpp b/Engine/Plugins/Mutable/Source/CustomizableObjectEditor/Private/MuCOE/Nodes/CustomizableObjectNodeTable.cpp
--- a/Engine/Plugins/Mutable/Source/CustomizableObjectEditor/Private/MuCOE/Nodes/CustomizableObjectNodeTable.cpp
+++ b/Engine/Plugins/Mutable/Source/CustomizableObjectEditor/Private/MuCOE/Nodes/CustomizableObjectNodeTable.cpp
@@ -66,7 +66,7 @@ bool UCustomizableObjectNodeTableRemapPins::Equal(const UCustomizableObjectNode&
 
 	// In this case pin type may have changed but we consider them the same type
 	if (OldPin.PinType.PinCategory == UEdGraphSchema_CustomizableObject::PC_Texture ||
-		OldPin.PinType.PinCategory == UEdGraphSchema_CustomizableObject::PC_PassthroughTexture)
+		OldPin.PinType.PinCategory == UEdGraphSchema_CustomizableObject::PC_Texture_Passthrough)
 	{
 		return true;
 	}
@@ -544,7 +544,7 @@ void UCustomizableObjectNodeTable::AllocateDefaultPins(UCustomizableObjectNodeRe
 									PinData->Material = MatIndex;
 
 									// Create a new pin for each lod and mesh section
-									UEdGraphPin* MeshPin = CustomCreatePin(EGPD_Output, Schema->PC_Mesh, FName(*TableMeshPinName), PinData);
+									UEdGraphPin* MeshPin = CustomCreatePin(EGPD_Output, Schema->PC_Mesh, FName(*TableMeshPinName), FText::FromString(TableMeshPinName), PinData);
 									MeshPin->SafeSetHidden(false);
 									MeshPin->PinToolTip = TableMeshPinName;
 
@@ -596,7 +596,7 @@ void UCustomizableObjectNodeTable::AllocateDefaultPins(UCustomizableObjectNodeRe
 									PinData->Material = MatIndex;
 
 									// Create a new pin for each lod and mesh section
-									UEdGraphPin* MeshPin = CustomCreatePin(EGPD_Output, Schema->PC_Mesh, FName(*TableMeshPinName), PinData);
+									UEdGraphPin* MeshPin = CustomCreatePin(EGPD_Output, Schema->PC_Mesh, FName(*TableMeshPinName), FText::FromString(TableMeshPinName), PinData);
 									MeshPin->SafeSetHidden(false);
 									MeshPin->PinToolTip = TableMeshPinName;
 								}
@@ -620,7 +620,7 @@ void UCustomizableObjectNodeTable::AllocateDefaultPins(UCustomizableObjectNodeRe
 						PinData->bIsNotTexture2D = false;
 						PinData->NodeTable = this;
 
-						FName PinCategory = PinData->ImageMode == ETableTextureType::PASSTHROUGH_TEXTURE ? Schema->PC_PassthroughTexture : Schema->PC_Texture;
+						FName PinCategory = PinData->ImageMode == ETableTextureType::PASSTHROUGH_TEXTURE ? Schema->PC_Texture_Passthrough : Schema->PC_Texture;
 
 						for (UEdGraphPin* OldPin : OldPins)
 						{
@@ -629,13 +629,13 @@ void UCustomizableObjectNodeTable::AllocateDefaultPins(UCustomizableObjectNodeRe
 							{
 								if (OldPinData->ColumnPropertyName == PropertyName)
 								{
-									PinCategory = OldPinData->ImageMode == ETableTextureType::PASSTHROUGH_TEXTURE ? Schema->PC_PassthroughTexture : Schema->PC_Texture;
+									PinCategory = OldPinData->ImageMode == ETableTextureType::PASSTHROUGH_TEXTURE ? Schema->PC_Texture_Passthrough : Schema->PC_Texture;
 									break;
 								}
 							}
 						}
 
-						OutPin = CustomCreatePin(EGPD_Output, PinCategory, FName(*PinName), PinData);
+						OutPin = CustomCreatePin(EGPD_Output, PinCategory, FName(*PinName), FText::FromString(PinName), PinData);
 					}
 
 					else if (Object->IsA(UTexture::StaticClass()))
@@ -650,7 +650,7 @@ void UCustomizableObjectNodeTable::AllocateDefaultPins(UCustomizableObjectNodeRe
 						PinData->bIsNotTexture2D = true;
 						PinData->NodeTable = this;
 
-						OutPin = CustomCreatePin(EGPD_Output, Schema->PC_PassthroughTexture, FName(*PinName), PinData);
+						OutPin = CustomCreatePin(EGPD_Output, Schema->PC_Texture_Passthrough, FName(*PinName), FText::FromString(PinName), PinData);
 					}
 
 					else if (Object->IsA(UMaterialInterface::StaticClass()))
@@ -660,7 +660,7 @@ void UCustomizableObjectNodeTable::AllocateDefaultPins(UCustomizableObjectNodeRe
 						PinData->ColumnDisplayName = PinName;
 						PinData->StructColumnId = PropertyID;
 
-						OutPin = CustomCreatePin(EGPD_Output, Schema->PC_Material, FName(*PinName), PinData);
+						OutPin = CustomCreatePin(EGPD_Output, Schema->PC_Material, FName(*PinName), FText::FromString(PinName), PinData);
 					}
 
 					else if (Object->IsA(UPoseAsset::StaticClass()))
@@ -670,7 +670,7 @@ void UCustomizableObjectNodeTable::AllocateDefaultPins(UCustomizableObjectNodeRe
 						PinData->ColumnDisplayName = PinName;
 						PinData->StructColumnId = PropertyID;
 
-						OutPin = CustomCreatePin(EGPD_Output, Schema->PC_PoseAsset, FName(*PinName), PinData);
+						OutPin = CustomCreatePin(EGPD_Output, Schema->PC_PoseAsset, FName(*PinName), FText::FromString(PinName), PinData);
 					}
 				}
 				else
@@ -695,7 +695,7 @@ void UCustomizableObjectNodeTable::AllocateDefaultPins(UCustomizableObjectNodeRe
 				PinData->ColumnDisplayName = PinName;
 				PinData->StructColumnId = PropertyID;
 
-				OutPin = CustomCreatePin(EGPD_Output, Schema->PC_Color, FName(*PinName), PinData);
+				OutPin = CustomCreatePin(EGPD_Output, Schema->PC_Color, FName(*PinName), FText::FromString(PinName), PinData);
 			}
 		}
 
@@ -706,7 +706,7 @@ void UCustomizableObjectNodeTable::AllocateDefaultPins(UCustomizableObjectNodeRe
 			PinData->ColumnDisplayName = PinName;
 			PinData->StructColumnId = PropertyID;
 
-			OutPin = CustomCreatePin(EGPD_Output, Schema->PC_Float, FName(*PinName), PinData);
+			OutPin = CustomCreatePin(EGPD_Output, Schema->PC_Float, FName(*PinName), FText::FromString(PinName), PinData);
 		}
 		
 		else if (const FNumericProperty* NumDoubleProperty = CastField<FDoubleProperty>(ColumnProperty))
@@ -716,7 +716,7 @@ void UCustomizableObjectNodeTable::AllocateDefaultPins(UCustomizableObjectNodeRe
 			PinData->ColumnDisplayName = PinName;
 			PinData->StructColumnId = PropertyID;
 
-			OutPin = CustomCreatePin(EGPD_Output, Schema->PC_Float, FName(*PinName), PinData);
+			OutPin = CustomCreatePin(EGPD_Output, Schema->PC_Float, FName(*PinName), FText::FromString(PinName), PinData);
 		}
 
 		else if (const FObjectProperty* ObjectProperty = CastField<FObjectProperty>(ColumnProperty))
@@ -848,7 +848,7 @@ bool UCustomizableObjectNodeTable::IsNodeOutDatedAndNeedsRefresh()
 				}
 				else if (Object->IsA(UTexture2D::StaticClass()))
 				{
-					if (CheckPinUpdated(PinName, Schema->PC_Texture) && CheckPinUpdated(PinName, Schema->PC_PassthroughTexture))
+					if (CheckPinUpdated(PinName, Schema->PC_Texture) && CheckPinUpdated(PinName, Schema->PC_Texture_Passthrough))
 					{
 						bNeedsUpdate = true;
 					}
@@ -857,7 +857,7 @@ bool UCustomizableObjectNodeTable::IsNodeOutDatedAndNeedsRefresh()
 				}
 				else if (Object->IsA(UTexture2DArray::StaticClass()))
 				{
-					if (CheckPinUpdated(PinName, Schema->PC_PassthroughTexture))
+					if (CheckPinUpdated(PinName, Schema->PC_Texture_Passthrough))
 					{
 						bNeedsUpdate = true;
 					}
@@ -971,7 +971,7 @@ void UCustomizableObjectNodeTable::RemapPinsData(const TMap<UEdGraphPin*, UEdGra
 				}
 			}
 		}
-		else if (Pair.Key->PinType.PinCategory == Schema->PC_Texture || Pair.Key->PinType.PinCategory == Schema->PC_PassthroughTexture)
+		else if (Pair.Key->PinType.PinCategory == Schema->PC_Texture || Pair.Key->PinType.PinCategory == Schema->PC_Texture_Passthrough)
 		{
 			UCustomizableObjectNodeTableImagePinData* PinDataOldPin = Cast<UCustomizableObjectNodeTableImagePinData>(GetPinData(*(Pair.Key)));
 			UCustomizableObjectNodeTableImagePinData* PinDataNewPin = Cast<UCustomizableObjectNodeTableImagePinData>(GetPinData(*(Pair.Value)));
@@ -999,7 +999,8 @@ bool UCustomizableObjectNodeTable::IsPinRelevant(const UEdGraphPin* Pin) const
 	return Pin->Direction == EGPD_Input &&
 		(Pin->PinType.PinCategory == Schema->PC_Material ||
 		Pin->PinType.PinCategory == Schema->PC_Texture ||
-		Pin->PinType.PinCategory == Schema->PC_PassthroughTexture ||
+		Pin->PinType.PinCategory == Schema->PC_Texture_Passthrough ||
+		Pin->PinType.PinCategory == Schema->PC_Texture_PassthroughParameter ||
 		Pin->PinType.PinCategory == Schema->PC_Color ||
 		Pin->PinType.PinCategory == Schema->PC_Float ||
 		Pin->PinType.PinCategory == Schema->PC_Mesh);
diff --git a/Engine/Plugins/Mutable/Source/CustomizableObjectEditor/Private/MuCOE/Nodes/CustomizableObjectNodeTexture.cpp b/Engine/Plugins/Mutable/Source/CustomizableObjectEditor/Private/MuCOE/Nodes/CustomizableObjectNodeTexture.cpp
--- a/Engine/Plugins/Mutable/Source/CustomizableObjectEditor/Private/MuCOE/Nodes/CustomizableObjectNodeTexture.cpp
+++ b/Engine/Plugins/Mutable/Source/CustomizableObjectEditor/Private/MuCOE/Nodes/CustomizableObjectNodeTexture.cpp
@@ -31,9 +31,7 @@ void UCustomizableObjectNodeTexture::AllocateDefaultPins(UCustomizableObjectNode
 {
 	const UEdGraphSchema_CustomizableObject* Schema = GetDefault<UEdGraphSchema_CustomizableObject>();
 
-	FString PinName = TEXT("Texture");
-	UEdGraphPin* PinImagePin = CustomCreatePin(EGPD_Output, Schema->PC_Texture, FName(*PinName));
-	PinImagePin->bDefaultValueIsIgnored = true;
+	CustomCreatePin(EGPD_Output, Schema->PC_Texture);
 }
 
 
diff --git a/Engine/Plugins/Mutable/Source/CustomizableObjectEditor/Private/MuCOE/Nodes/CustomizableObjectNodeTextureBinarise.cpp b/Engine/Plugins/Mutable/Source/CustomizableObjectEditor/Private/MuCOE/Nodes/CustomizableObjectNodeTextureBinarise.cpp
--- a/Engine/Plugins/Mutable/Source/CustomizableObjectEditor/Private/MuCOE/Nodes/CustomizableObjectNodeTextureBinarise.cpp
+++ b/Engine/Plugins/Mutable/Source/CustomizableObjectEditor/Private/MuCOE/Nodes/CustomizableObjectNodeTextureBinarise.cpp
@@ -14,14 +14,13 @@ class UCustomizableObjectNodeRemapPins;
 
 void UCustomizableObjectNodeTextureBinarise::AllocateDefaultPins(UCustomizableObjectNodeRemapPins* RemapPins)
 {
-	UEdGraphPin* OutputPin = CustomCreatePinSimple(EGPD_Output, UEdGraphSchema_CustomizableObject::PC_Texture);
+	UEdGraphPin* OutputPin = CustomCreatePin(EGPD_Output, UEdGraphSchema_CustomizableObject::PC_Texture);
 	OutputPin->bDefaultValueIsIgnored = true;
 	
-	UEdGraphPin* ImagePin = CustomCreatePinSimple(EGPD_Input, UEdGraphSchema_CustomizableObject::PC_Texture);
+	UEdGraphPin* ImagePin = CustomCreatePin(EGPD_Input, UEdGraphSchema_CustomizableObject::PC_Texture);
 	BaseImagePinReference = FEdGraphPinReference(ImagePin);
 	
-	const FName PinName = TEXT("Threshold");
-	CustomCreatePin(EGPD_Input, UEdGraphSchema_CustomizableObject::PC_Float, PinName);
+	CustomCreatePin(EGPD_Input, UEdGraphSchema_CustomizableObject::PC_Float,  TEXT("Threshold"),  LOCTEXT("Threshold", "Threshold"));
 }
 
 
diff --git a/Engine/Plugins/Mutable/Source/CustomizableObjectEditor/Private/MuCOE/Nodes/CustomizableObjectNodeTextureColourMap.cpp b/Engine/Plugins/Mutable/Source/CustomizableObjectEditor/Private/MuCOE/Nodes/CustomizableObjectNodeTextureColourMap.cpp
--- a/Engine/Plugins/Mutable/Source/CustomizableObjectEditor/Private/MuCOE/Nodes/CustomizableObjectNodeTextureColourMap.cpp
+++ b/Engine/Plugins/Mutable/Source/CustomizableObjectEditor/Private/MuCOE/Nodes/CustomizableObjectNodeTextureColourMap.cpp
@@ -14,23 +14,10 @@ class UCustomizableObjectNodeRemapPins;
 
 void UCustomizableObjectNodeTextureColourMap::AllocateDefaultPins(UCustomizableObjectNodeRemapPins* RemapPins)
 {
-	const UEdGraphSchema_CustomizableObject* Schema = GetDefault<UEdGraphSchema_CustomizableObject>();
-
-	FString PinName = TEXT("Texture");
-	UEdGraphPin* OutputPin = CustomCreatePin(EGPD_Output, Schema->PC_Texture, FName(*PinName));
-	OutputPin->bDefaultValueIsIgnored = true;
-
-	PinName = TEXT("Base");
-	UEdGraphPin* SourcePin = CustomCreatePin(EGPD_Input, Schema->PC_Texture, FName(*PinName));
-	SourcePin->bDefaultValueIsIgnored = true;
-
-	PinName = TEXT("Mask");
-	UEdGraphPin* MaskPin = CustomCreatePin(EGPD_Input, Schema->PC_Texture, FName(*PinName));
-	MaskPin->bDefaultValueIsIgnored = true;
-
-	PinName = TEXT("Map");
-	UEdGraphPin* GradientPin = CustomCreatePin(EGPD_Input, Schema->PC_Texture, FName(*PinName));
-	GradientPin->bDefaultValueIsIgnored = true;
+	CustomCreatePin(EGPD_Output, UEdGraphSchema_CustomizableObject::PC_Texture, TEXT("Texture"), LOCTEXT("Texture", "Texture"));
+	CustomCreatePin(EGPD_Input, UEdGraphSchema_CustomizableObject::PC_Texture, TEXT("Base"), LOCTEXT("Base", "Base"));
+	CustomCreatePin(EGPD_Input, UEdGraphSchema_CustomizableObject::PC_Texture, TEXT("Mask"), LOCTEXT("Mask", "Mask"));
+	CustomCreatePin(EGPD_Input, UEdGraphSchema_CustomizableObject::PC_Texture, TEXT("Map"), LOCTEXT("Map", "Map"));
 }
 
 
diff --git a/Engine/Plugins/Mutable/Source/CustomizableObjectEditor/Private/MuCOE/Nodes/CustomizableObjectNodeTextureFromChannels.cpp b/Engine/Plugins/Mutable/Source/CustomizableObjectEditor/Private/MuCOE/Nodes/CustomizableObjectNodeTextureFromChannels.cpp
--- a/Engine/Plugins/Mutable/Source/CustomizableObjectEditor/Private/MuCOE/Nodes/CustomizableObjectNodeTextureFromChannels.cpp
+++ b/Engine/Plugins/Mutable/Source/CustomizableObjectEditor/Private/MuCOE/Nodes/CustomizableObjectNodeTextureFromChannels.cpp
@@ -14,27 +14,12 @@ class UCustomizableObjectNodeRemapPins;
 
 void UCustomizableObjectNodeTextureFromChannels::AllocateDefaultPins(UCustomizableObjectNodeRemapPins* RemapPins)
 {
-	const UEdGraphSchema_CustomizableObject* Schema = GetDefault<UEdGraphSchema_CustomizableObject>();
+	CustomCreatePin(EGPD_Output, UEdGraphSchema_CustomizableObject::PC_Texture, TEXT("Texture"), LOCTEXT("Texture", "Texture"));
 
-	FString PinName = TEXT("Texture");
-	UEdGraphPin* OutputPin = CustomCreatePin(EGPD_Output, Schema->PC_Texture, FName(*PinName));
-	OutputPin->bDefaultValueIsIgnored = true;
-
-	PinName = TEXT("R");
-	UEdGraphPin* RPin = CustomCreatePin(EGPD_Input, Schema->PC_Texture, FName(*PinName));
-	RPin->bDefaultValueIsIgnored = true;
-
-	PinName = TEXT("G");
-	UEdGraphPin* GPin = CustomCreatePin(EGPD_Input, Schema->PC_Texture, FName(*PinName));
-	GPin->bDefaultValueIsIgnored = true;
-
-	PinName = TEXT("B");
-	UEdGraphPin* BPin = CustomCreatePin(EGPD_Input, Schema->PC_Texture, FName(*PinName));
-	BPin->bDefaultValueIsIgnored = true;
-
-	PinName = TEXT("A");
-	UEdGraphPin* APin = CustomCreatePin(EGPD_Input, Schema->PC_Texture, FName(*PinName));
-	APin->bDefaultValueIsIgnored = true;
+	CustomCreatePin(EGPD_Input, UEdGraphSchema_CustomizableObject::PC_Texture, TEXT("R"), LOCTEXT("R", "R"));
+	CustomCreatePin(EGPD_Input, UEdGraphSchema_CustomizableObject::PC_Texture, TEXT("G"), LOCTEXT("G", "G"));
+	CustomCreatePin(EGPD_Input, UEdGraphSchema_CustomizableObject::PC_Texture, TEXT("B"), LOCTEXT("B", "B"));
+	CustomCreatePin(EGPD_Input, UEdGraphSchema_CustomizableObject::PC_Texture, TEXT("A"), LOCTEXT("A", "A"));
 }
 
 
diff --git a/Engine/Plugins/Mutable/Source/CustomizableObjectEditor/Private/MuCOE/Nodes/CustomizableObjectNodeTextureFromColor.cpp b/Engine/Plugins/Mutable/Source/CustomizableObjectEditor/Private/MuCOE/Nodes/CustomizableObjectNodeTextureFromColor.cpp
--- a/Engine/Plugins/Mutable/Source/CustomizableObjectEditor/Private/MuCOE/Nodes/CustomizableObjectNodeTextureFromColor.cpp
+++ b/Engine/Plugins/Mutable/Source/CustomizableObjectEditor/Private/MuCOE/Nodes/CustomizableObjectNodeTextureFromColor.cpp
@@ -14,18 +14,12 @@ class UCustomizableObjectNodeRemapPins;
 
 void UCustomizableObjectNodeTextureFromColor::AllocateDefaultPins(UCustomizableObjectNodeRemapPins* RemapPins)
 {
-	const UEdGraphSchema_CustomizableObject* Schema = GetDefault<UEdGraphSchema_CustomizableObject>();
-
-	FString PinName = TEXT("Texture");
-	UEdGraphPin* OutputPin = CustomCreatePin(EGPD_Output, Schema->PC_Texture, FName(*PinName));
-	OutputPin->bDefaultValueIsIgnored = true;
-
-	PinName = TEXT("Color");
-	UEdGraphPin* RPin = CustomCreatePin(EGPD_Input, Schema->PC_Color, FName(*PinName));
-	RPin->bDefaultValueIsIgnored = true;
+	CustomCreatePin(EGPD_Output, UEdGraphSchema_CustomizableObject::PC_Texture);
+	CustomCreatePin(EGPD_Input, UEdGraphSchema_CustomizableObject::PC_Color);
 }
 
 
+
 void UCustomizableObjectNodeTextureFromColor::BackwardsCompatibleFixup(int32 CustomizableObjectCustomVersion)
 {
 	Super::BackwardsCompatibleFixup(CustomizableObjectCustomVersion);
diff --git a/Engine/Plugins/Mutable/Source/CustomizableObjectEditor/Private/MuCOE/Nodes/CustomizableObjectNodeTextureFromFloats.cpp b/Engine/Plugins/Mutable/Source/CustomizableObjectEditor/Private/MuCOE/Nodes/CustomizableObjectNodeTextureFromFloats.cpp
--- a/Engine/Plugins/Mutable/Source/CustomizableObjectEditor/Private/MuCOE/Nodes/CustomizableObjectNodeTextureFromFloats.cpp
+++ b/Engine/Plugins/Mutable/Source/CustomizableObjectEditor/Private/MuCOE/Nodes/CustomizableObjectNodeTextureFromFloats.cpp
@@ -14,27 +14,12 @@ class UCustomizableObjectNodeRemapPins;
 
 void UCustomizableObjectNodeTextureFromFloats::AllocateDefaultPins(UCustomizableObjectNodeRemapPins* RemapPins)
 {
-	const UEdGraphSchema_CustomizableObject* Schema = GetDefault<UEdGraphSchema_CustomizableObject>();
+	CustomCreatePin(EGPD_Output, UEdGraphSchema_CustomizableObject::PC_Texture, TEXT("Texture"), LOCTEXT("Texture", "Texture"));
 
-	FString PinName = TEXT("Texture");
-	UEdGraphPin* OutputPin = CustomCreatePin(EGPD_Output, Schema->PC_Texture, FName(*PinName));
-	OutputPin->bDefaultValueIsIgnored = true;
-
-	PinName = TEXT("R");
-	UEdGraphPin* RPin = CustomCreatePin(EGPD_Input, Schema->PC_Float, FName(*PinName));
-	RPin->bDefaultValueIsIgnored = true;
-
-	PinName = TEXT("G");
-	UEdGraphPin* GPin = CustomCreatePin(EGPD_Input, Schema->PC_Float, FName(*PinName));
-	GPin->bDefaultValueIsIgnored = true;
-
-	PinName = TEXT("B");
-	UEdGraphPin* BPin = CustomCreatePin(EGPD_Input, Schema->PC_Float, FName(*PinName));
-	BPin->bDefaultValueIsIgnored = true;
-
-	PinName = TEXT("A");
-	UEdGraphPin* APin = CustomCreatePin(EGPD_Input, Schema->PC_Float, FName(*PinName));
-	APin->bDefaultValueIsIgnored = true;
+	CustomCreatePin(EGPD_Input, UEdGraphSchema_CustomizableObject::PC_Float, TEXT("R"), LOCTEXT("R", "R"));
+	CustomCreatePin(EGPD_Input, UEdGraphSchema_CustomizableObject::PC_Float, TEXT("G"), LOCTEXT("G", "G"));
+	CustomCreatePin(EGPD_Input, UEdGraphSchema_CustomizableObject::PC_Float, TEXT("B"), LOCTEXT("B", "B"));
+	CustomCreatePin(EGPD_Input, UEdGraphSchema_CustomizableObject::PC_Float, TEXT("A"), LOCTEXT("A", "A"));
 }
 
 
diff --git a/Engine/Plugins/Mutable/Source/CustomizableObjectEditor/Private/MuCOE/Nodes/CustomizableObjectNodeTextureInterpolate.cpp b/Engine/Plugins/Mutable/Source/CustomizableObjectEditor/Private/MuCOE/Nodes/CustomizableObjectNodeTextureInterpolate.cpp
--- a/Engine/Plugins/Mutable/Source/CustomizableObjectEditor/Private/MuCOE/Nodes/CustomizableObjectNodeTextureInterpolate.cpp
+++ b/Engine/Plugins/Mutable/Source/CustomizableObjectEditor/Private/MuCOE/Nodes/CustomizableObjectNodeTextureInterpolate.cpp
@@ -34,19 +34,14 @@ void UCustomizableObjectNodeTextureInterpolate::AllocateDefaultPins(UCustomizabl
 {
 	const UEdGraphSchema_CustomizableObject* Schema = GetDefault<UEdGraphSchema_CustomizableObject>();
 
-	FString PinName = TEXT("Texture");
-	UEdGraphPin* OutputPin = CustomCreatePin(EGPD_Output, Schema->PC_Texture, FName(*PinName));
-	OutputPin->bDefaultValueIsIgnored = true;
+	CustomCreatePin(EGPD_Output, Schema->PC_Texture, TEXT("Texture"), LOCTEXT("Texture", "Texture"));
 	
-	PinName = TEXT("Factor");
-	UEdGraphPin* FactorPin = CustomCreatePin(EGPD_Input, Schema->PC_Float, FName(*PinName));
-	FactorPin->bDefaultValueIsIgnored = true;
+	CustomCreatePin(EGPD_Input, Schema->PC_Float, TEXT("Factor"), LOCTEXT("Factor", "Factor"));
 
 	for ( int LayerIndex = 0; LayerIndex < NumTargets; ++LayerIndex )
 	{
-		PinName = FString::Printf(TEXT("Target %d"), LayerIndex);
-		UEdGraphPin* LayerPin = CustomCreatePin(EGPD_Input, Schema->PC_Texture, FName(*PinName));
-		LayerPin->bDefaultValueIsIgnored = true;
+		FString PinName = FString::Printf(TEXT("Target %d"), LayerIndex);
+		CustomCreatePin(EGPD_Input, Schema->PC_Texture, FName(*PinName), FText::FromString(PinName));
 	}
 }
 
diff --git a/Engine/Plugins/Mutable/Source/CustomizableObjectEditor/Private/MuCOE/Nodes/CustomizableObjectNodeTextureInvert.cpp b/Engine/Plugins/Mutable/Source/CustomizableObjectEditor/Private/MuCOE/Nodes/CustomizableObjectNodeTextureInvert.cpp
--- a/Engine/Plugins/Mutable/Source/CustomizableObjectEditor/Private/MuCOE/Nodes/CustomizableObjectNodeTextureInvert.cpp
+++ b/Engine/Plugins/Mutable/Source/CustomizableObjectEditor/Private/MuCOE/Nodes/CustomizableObjectNodeTextureInvert.cpp
@@ -14,10 +14,10 @@ class UCustomizableObjectNodeRemapPins;
 
 void UCustomizableObjectNodeTextureInvert::AllocateDefaultPins(UCustomizableObjectNodeRemapPins* RemapPins)
 {
-	UEdGraphPin* OutputPin = CustomCreatePinSimple(EGPD_Output, UEdGraphSchema_CustomizableObject::PC_Texture);
+	UEdGraphPin* OutputPin = CustomCreatePin(EGPD_Output, UEdGraphSchema_CustomizableObject::PC_Texture);
 	OutputPin->bDefaultValueIsIgnored = true;
 	
-	UEdGraphPin* ImagePin = CustomCreatePinSimple(EGPD_Input, UEdGraphSchema_CustomizableObject::PC_Texture);
+	UEdGraphPin* ImagePin = CustomCreatePin(EGPD_Input, UEdGraphSchema_CustomizableObject::PC_Texture);
 	BaseImagePinReference = FEdGraphPinReference(ImagePin);
 }
 
diff --git a/Engine/Plugins/Mutable/Source/CustomizableObjectEditor/Private/MuCOE/Nodes/CustomizableObjectNodeTextureLayer.cpp b/Engine/Plugins/Mutable/Source/CustomizableObjectEditor/Private/MuCOE/Nodes/CustomizableObjectNodeTextureLayer.cpp
--- a/Engine/Plugins/Mutable/Source/CustomizableObjectEditor/Private/MuCOE/Nodes/CustomizableObjectNodeTextureLayer.cpp
+++ b/Engine/Plugins/Mutable/Source/CustomizableObjectEditor/Private/MuCOE/Nodes/CustomizableObjectNodeTextureLayer.cpp
@@ -51,23 +51,17 @@ void UCustomizableObjectNodeTextureLayer::AllocateDefaultPins(UCustomizableObjec
 {
 	const UEdGraphSchema_CustomizableObject* Schema = GetDefault<UEdGraphSchema_CustomizableObject>();
 
-	UEdGraphPin* OutputPin = CustomCreatePin(EGPD_Output, Schema->PC_Texture, FName(TEXT("Texture")));
-	OutputPin->bDefaultValueIsIgnored = true;
-
-	OutputPinReference = FEdGraphPinReference(OutputPin);
-
-	UEdGraphPin* BasePin = CustomCreatePin(EGPD_Input, Schema->PC_Texture, FName(TEXT("Base")));
-	BasePin->bDefaultValueIsIgnored = true;
+	OutputPinReference =CustomCreatePin(EGPD_Output, Schema->PC_Texture, TEXT("Texture"), LOCTEXT("Texture", "Texture"));
+	
+	CustomCreatePin(EGPD_Input, Schema->PC_Texture, TEXT("Base"), LOCTEXT("Base", "Base"));
 	
 	for (int32 LayerIndex = 0; LayerIndex < Layers.Num() ; ++LayerIndex )
 	{
-		FString PinName = FString::Printf( TEXT("Layer %d "), LayerIndex );
-		UEdGraphPin* LayerPin = CustomCreatePin(EGPD_Input, Schema->PC_Texture, *PinName);
-		LayerPin->bDefaultValueIsIgnored = true;
+		FString PinName = FString::Printf(TEXT("Layer %d "), LayerIndex);
+		CustomCreatePin(EGPD_Input, Schema->PC_Texture, *PinName, FText::FromString(PinName));
 
 		PinName = FString::Printf( TEXT("Mask %d "), LayerIndex );
-		UEdGraphPin* MaskPin = CustomCreatePin(EGPD_Input, Schema->PC_Texture, *PinName);
-		MaskPin->bDefaultValueIsIgnored = true;
+		CustomCreatePin(EGPD_Input, Schema->PC_Texture, *PinName, FText::FromString(PinName));
 	}
 }
 
diff --git a/Engine/Plugins/Mutable/Source/CustomizableObjectEditor/Private/MuCOE/Nodes/CustomizableObjectNodeTextureParameter.cpp b/Engine/Plugins/Mutable/Source/CustomizableObjectEditor/Private/MuCOE/Nodes/CustomizableObjectNodeTextureParameter.cpp
--- a/Engine/Plugins/Mutable/Source/CustomizableObjectEditor/Private/MuCOE/Nodes/CustomizableObjectNodeTextureParameter.cpp
+++ b/Engine/Plugins/Mutable/Source/CustomizableObjectEditor/Private/MuCOE/Nodes/CustomizableObjectNodeTextureParameter.cpp
@@ -23,15 +23,7 @@ void UCustomizableObjectNodeTextureParameter::AllocateDefaultPins(UCustomizableO
 {
 	Super::AllocateDefaultPins(RemapPins);
 
-	const FName TypePassthrough = UEdGraphSchema_CustomizableObject::PC_PassthroughTexture;
-	
-	const FName PinName = UEdGraphSchema_CustomizableObject::GetPinCategoryName(TypePassthrough);
-	const FText PinFriendlyName = UEdGraphSchema_CustomizableObject::GetPinCategoryFriendlyName(TypePassthrough);
-	
-	UEdGraphPin* Pin = CustomCreatePin(EGPD_Output, TypePassthrough, PinName);
-	Pin->PinFriendlyName = PinFriendlyName;
-
-	PassthroughPin = Pin;
+	PassthroughPin = CustomCreatePin(EGPD_Output, UEdGraphSchema_CustomizableObject::PC_Texture_PassthroughParameter);
 }
 
 
@@ -54,15 +46,17 @@ void UCustomizableObjectNodeTextureParameter::BackwardsCompatibleFixup(int32 Cus
 	}
 	else if (CustomizableObjectCustomVersion == FCustomizableObjectCustomVersion::NodeTextureParameterPassthroughPin)
 	{
-		const FName TypePassthrough = UEdGraphSchema_CustomizableObject::PC_PassthroughTexture;
-	
-		const FName PinName = UEdGraphSchema_CustomizableObject::GetPinCategoryName(TypePassthrough);
-		const FText PinFriendlyName = UEdGraphSchema_CustomizableObject::GetPinCategoryFriendlyName(TypePassthrough);
-	
-		UEdGraphPin* Pin = CustomCreatePin(EGPD_Output, TypePassthrough, PinName);
-		Pin->PinFriendlyName = PinFriendlyName;
-
-		PassthroughPin = Pin;
+		PassthroughPin = CustomCreatePin(EGPD_Output, UEdGraphSchema_CustomizableObject::PC_Texture_Passthrough);
+	}
+	else if (CustomizableObjectCustomVersion == FCustomizableObjectCustomVersion::PassthroughParameter)
+	{
+		for (UEdGraphPin* Pin : Pins)
+		{
+			if (Pin->PinType.PinCategory ==  UEdGraphSchema_CustomizableObject::PC_Texture_Passthrough)
+			{
+				Pin->PinType.PinCategory =  UEdGraphSchema_CustomizableObject::PC_Texture_PassthroughParameter;
+			}
+		}
 	}
 }
 
diff --git a/Engine/Plugins/Mutable/Source/CustomizableObjectEditor/Private/MuCOE/Nodes/CustomizableObjectNodeTextureProject.cpp b/Engine/Plugins/Mutable/Source/CustomizableObjectEditor/Private/MuCOE/Nodes/CustomizableObjectNodeTextureProject.cpp
--- a/Engine/Plugins/Mutable/Source/CustomizableObjectEditor/Private/MuCOE/Nodes/CustomizableObjectNodeTextureProject.cpp
+++ b/Engine/Plugins/Mutable/Source/CustomizableObjectEditor/Private/MuCOE/Nodes/CustomizableObjectNodeTextureProject.cpp
@@ -83,39 +83,27 @@ void UCustomizableObjectNodeTextureProject::AllocateDefaultPins(UCustomizableObj
 	for ( uint32 LayerIndex = 0; LayerIndex < Textures; ++LayerIndex )
 	{
 		FString PinName = FString::Printf(TEXT("Texture %u"), LayerIndex);
-		UEdGraphPin* TexturePin = CustomCreatePin(EGPD_Input, Schema->PC_Texture, FName(*PinName));
-		TexturePin->bDefaultValueIsIgnored = true;
-
+		UEdGraphPin* TexturePin = CustomCreatePin(EGPD_Input, Schema->PC_Texture, FName(*PinName), FText::FromString(PinName));
+		
 		TexturePinsReferences.Add(FEdGraphPinReference(TexturePin));
 		
 		PinName = FString::Printf(TEXT("Texture %u"), LayerIndex);
-		UEdGraphPin* OutputPin = CustomCreatePin(EGPD_Output, Schema->PC_Texture, FName(*PinName));
-		OutputPin->bDefaultValueIsIgnored = true;
+		UEdGraphPin* OutputPin = CustomCreatePin(EGPD_Output, Schema->PC_Texture, FName(*PinName), FText::FromString(PinName));
 
 		OutputPinsReferences.Add(FEdGraphPinReference(OutputPin));
 	}
 
-	FString PinName = TEXT("Mesh");
-	UEdGraphPin* MeshPin = CustomCreatePin(EGPD_Input, Schema->PC_Mesh, FName(*PinName));
-	MeshPin->bDefaultValueIsIgnored = true;
+	CustomCreatePin(EGPD_Input, Schema->PC_Mesh, TEXT("Mesh"), LOCTEXT("Mesh", "Mesh"));
 
-	PinName = TEXT("Mesh Mask");
-	UEdGraphPin* MeshMaskPin = CustomCreatePin(EGPD_Input, Schema->PC_Texture, FName(*PinName));
-	MeshMaskPin->bDefaultValueIsIgnored = true;
+	CustomCreatePin(EGPD_Input, Schema->PC_Texture, TEXT("Mesh Mask"), LOCTEXT("MeshMask", "Mesh Mask"));
 
-	PinName = TEXT("Fade Start Angle");
-	UEdGraphPin* AngleFadeStartPin = CustomCreatePin(EGPD_Input, Schema->PC_Float, FName(*PinName));
-	AngleFadeStartPin->bDefaultValueIsIgnored = false;
+	UEdGraphPin* AngleFadeStartPin = CustomCreatePin(EGPD_Input, Schema->PC_Float, TEXT("Fade Start Angle"), LOCTEXT("FadeStartAngle", "Fade Start Angle"));
 	AngleFadeStartPin->DefaultValue = "90";
 
-	PinName = TEXT("Fade End Angle");
-	UEdGraphPin* AngleFadeEndPin = CustomCreatePin(EGPD_Input, Schema->PC_Float, FName(*PinName));
-	AngleFadeEndPin->bDefaultValueIsIgnored = false;
+	UEdGraphPin* AngleFadeEndPin = CustomCreatePin(EGPD_Input, Schema->PC_Float, TEXT("Fade End Angle"), LOCTEXT("FadeEndAngle", "Fade End Angle"));
 	AngleFadeEndPin->DefaultValue = "95";
-
-	PinName = TEXT("Projector");
-	UEdGraphPin* ProjectorPin = CustomCreatePin(EGPD_Input, Schema->PC_Projector, FName(*PinName));
-	ProjectorPin->bDefaultValueIsIgnored = true;
+	
+	CustomCreatePin(EGPD_Input, Schema->PC_Projector, TEXT("Projector"), LOCTEXT("Projector", "Projector"));
 }
 
 
diff --git a/Engine/Plugins/Mutable/Source/CustomizableObjectEditor/Private/MuCOE/Nodes/CustomizableObjectNodeTextureSample.cpp b/Engine/Plugins/Mutable/Source/CustomizableObjectEditor/Private/MuCOE/Nodes/CustomizableObjectNodeTextureSample.cpp
--- a/Engine/Plugins/Mutable/Source/CustomizableObjectEditor/Private/MuCOE/Nodes/CustomizableObjectNodeTextureSample.cpp
+++ b/Engine/Plugins/Mutable/Source/CustomizableObjectEditor/Private/MuCOE/Nodes/CustomizableObjectNodeTextureSample.cpp
@@ -13,30 +13,15 @@ struct FPropertyChangedEvent;
 #define LOCTEXT_NAMESPACE "CustomizableObjectEditor"
 
 
-UCustomizableObjectNodeTextureSample::UCustomizableObjectNodeTextureSample()
-	: Super()
-{
-
-}
-
-
 void UCustomizableObjectNodeTextureSample::AllocateDefaultPins(UCustomizableObjectNodeRemapPins* RemapPins)
 {
 	// Output
-	
-	UEdGraphPin* ColorPin = CustomCreatePinSimple(EGPD_Output, UEdGraphSchema_CustomizableObject::PC_Color);
-	ColorPin->bDefaultValueIsIgnored = true;
+	CustomCreatePin(EGPD_Output, UEdGraphSchema_CustomizableObject::PC_Color);
 	
 	// Inputs
-	
-	UEdGraphPin* TexturePin = CustomCreatePinSimple(EGPD_Input, UEdGraphSchema_CustomizableObject::PC_Texture);
-	TexturePin->bDefaultValueIsIgnored = true;
-
-	UEdGraphPin* XPin = CustomCreatePin(EGPD_Input, UEdGraphSchema_CustomizableObject::PC_Float,  FName(TEXT("X")));
-	XPin->bDefaultValueIsIgnored = true;
-
-	UEdGraphPin* YPin = CustomCreatePin(EGPD_Input, UEdGraphSchema_CustomizableObject::PC_Float, FName(TEXT("Y")));
-	YPin->bDefaultValueIsIgnored = true;
+	CustomCreatePin(EGPD_Input, UEdGraphSchema_CustomizableObject::PC_Texture);
+	CustomCreatePin(EGPD_Input, UEdGraphSchema_CustomizableObject::PC_Float,  TEXT("X"), LOCTEXT("X", "X"));
+	CustomCreatePin(EGPD_Input, UEdGraphSchema_CustomizableObject::PC_Float, TEXT("Y"), LOCTEXT("Y", "Y"));
 }
 
 void UCustomizableObjectNodeTextureSample::BackwardsCompatibleFixup(int32 CustomizableObjectCustomVersion)
diff --git a/Engine/Plugins/Mutable/Source/CustomizableObjectEditor/Private/MuCOE/Nodes/CustomizableObjectNodeTextureSample.h b/Engine/Plugins/Mutable/Source/CustomizableObjectEditor/Private/MuCOE/Nodes/CustomizableObjectNodeTextureSample.h
--- a/Engine/Plugins/Mutable/Source/CustomizableObjectEditor/Private/MuCOE/Nodes/CustomizableObjectNodeTextureSample.h
+++ b/Engine/Plugins/Mutable/Source/CustomizableObjectEditor/Private/MuCOE/Nodes/CustomizableObjectNodeTextureSample.h
@@ -23,8 +23,6 @@ class UCustomizableObjectNodeTextureSample : public UCustomizableObjectNode
 public:
 	GENERATED_BODY()
 
-	UE_API UCustomizableObjectNodeTextureSample();
-
 	// Begin EdGraphNode interface
 	UE_API virtual FText GetNodeTitle(ENodeTitleType::Type TitleType) const override;
 	UE_API virtual FLinearColor GetNodeTitleColor() const override;
diff --git a/Engine/Plugins/Mutable/Source/CustomizableObjectEditor/Private/MuCOE/Nodes/CustomizableObjectNodeTextureSaturate.cpp b/Engine/Plugins/Mutable/Source/CustomizableObjectEditor/Private/MuCOE/Nodes/CustomizableObjectNodeTextureSaturate.cpp
--- a/Engine/Plugins/Mutable/Source/CustomizableObjectEditor/Private/MuCOE/Nodes/CustomizableObjectNodeTextureSaturate.cpp
+++ b/Engine/Plugins/Mutable/Source/CustomizableObjectEditor/Private/MuCOE/Nodes/CustomizableObjectNodeTextureSaturate.cpp
@@ -14,15 +14,11 @@ class UCustomizableObjectNodeRemapPins;
 
 void UCustomizableObjectNodeTextureSaturate::AllocateDefaultPins(UCustomizableObjectNodeRemapPins* RemapPins)
 {
-	UEdGraphPin* OutputPin = CustomCreatePinSimple(EGPD_Output, UEdGraphSchema_CustomizableObject::PC_Texture);
-	OutputPin->bDefaultValueIsIgnored = true;
+	CustomCreatePin(EGPD_Output, UEdGraphSchema_CustomizableObject::PC_Texture);
 	
-	UEdGraphPin* ImagePin = CustomCreatePinSimple(EGPD_Input, UEdGraphSchema_CustomizableObject::PC_Texture);
-	BaseImagePinReference = FEdGraphPinReference(ImagePin);
+	BaseImagePinReference = CustomCreatePin(EGPD_Input, UEdGraphSchema_CustomizableObject::PC_Texture);
 	
-	UEdGraphPin* FactorPin = CustomCreatePin(EGPD_Input, UEdGraphSchema_CustomizableObject::PC_Float, FName(TEXT("Factor")));
-	FactorPin->bDefaultValueIsIgnored = true;
-	FactorPinReference = FEdGraphPinReference(FactorPin);
+	FactorPinReference = CustomCreatePin(EGPD_Input, UEdGraphSchema_CustomizableObject::PC_Float, TEXT("Factor"), LOCTEXT("Factor", "Factor"));
 }
 
 
diff --git a/Engine/Plugins/Mutable/Source/CustomizableObjectEditor/Private/MuCOE/Nodes/CustomizableObjectNodeTextureSwitch.cpp b/Engine/Plugins/Mutable/Source/CustomizableObjectEditor/Private/MuCOE/Nodes/CustomizableObjectNodeTextureSwitch.cpp
--- a/Engine/Plugins/Mutable/Source/CustomizableObjectEditor/Private/MuCOE/Nodes/CustomizableObjectNodeTextureSwitch.cpp
+++ b/Engine/Plugins/Mutable/Source/CustomizableObjectEditor/Private/MuCOE/Nodes/CustomizableObjectNodeTextureSwitch.cpp
@@ -14,11 +14,6 @@ void UCustomizableObjectNodeTextureSwitch::BackwardsCompatibleFixup(int32 Custom
 {
 	Super::BackwardsCompatibleFixup(CustomizableObjectCustomVersion);
 	
-	if (CustomizableObjectCustomVersion == FCustomizableObjectCustomVersion::PinsNamesImageToTexture)
-	{
-		OutputPinReference = FindPin(TEXT("Image"));
-	}
-
 	if (CustomizableObjectCustomVersion == FCustomizableObjectCustomVersion::FixPinsNamesImageToTexture2)
 	{
 		if (UEdGraphPin* TexturePin = FindPin(TEXT("Image")))
diff --git a/Engine/Plugins/Mutable/Source/CustomizableObjectEditor/Private/MuCOE/Nodes/CustomizableObjectNodeTextureToChannels.cpp b/Engine/Plugins/Mutable/Source/CustomizableObjectEditor/Private/MuCOE/Nodes/CustomizableObjectNodeTextureToChannels.cpp
--- a/Engine/Plugins/Mutable/Source/CustomizableObjectEditor/Private/MuCOE/Nodes/CustomizableObjectNodeTextureToChannels.cpp
+++ b/Engine/Plugins/Mutable/Source/CustomizableObjectEditor/Private/MuCOE/Nodes/CustomizableObjectNodeTextureToChannels.cpp
@@ -36,37 +36,15 @@ void UCustomizableObjectNodeTextureToChannels::AllocateDefaultPins(UCustomizable
 {
 	const UEdGraphSchema_CustomizableObject* Schema = GetDefault<UEdGraphSchema_CustomizableObject>();
 
-	{
-		FString PinName = TEXT("Texture");
-		UEdGraphPin* InputPin = CustomCreatePin(EGPD_Input, Schema->PC_Texture, FName(*PinName));
-		InputPin->bDefaultValueIsIgnored = true;
-
-		InputPinReference = FEdGraphPinReference(InputPin);
-	}
-
-	{
-		FString PinName = TEXT("R");
-		UEdGraphPin* RPin = CustomCreatePin(EGPD_Output, Schema->PC_Texture, FName(*PinName));
-		RPin->bDefaultValueIsIgnored = true;
-	}
+	InputPinReference = CustomCreatePin(EGPD_Input, Schema->PC_Texture, TEXT("Texture"), LOCTEXT("Texture", "Texture"));
 
-	{
-		FString PinName = TEXT("G");
-		UEdGraphPin* GPin = CustomCreatePin(EGPD_Output, Schema->PC_Texture, FName(*PinName));
-		GPin->bDefaultValueIsIgnored = true;
-	}
+	CustomCreatePin(EGPD_Output, Schema->PC_Texture, TEXT("R"), LOCTEXT("R", "R"));
 
-	{
-		FString PinName = TEXT("B");
-		UEdGraphPin* BPin = CustomCreatePin(EGPD_Output, Schema->PC_Texture, FName(*PinName));
-		BPin->bDefaultValueIsIgnored = true;
-	}
-
-	{
-		FString PinName = TEXT("A");
-		UEdGraphPin* APin = CustomCreatePin(EGPD_Output, Schema->PC_Texture, FName(*PinName));
-		APin->bDefaultValueIsIgnored = true;
-	}
+	CustomCreatePin(EGPD_Output, Schema->PC_Texture, TEXT("G"), LOCTEXT("G", "G"));
+	
+	CustomCreatePin(EGPD_Output, Schema->PC_Texture, TEXT("B"), LOCTEXT("B", "B"));
+	
+	CustomCreatePin(EGPD_Output, Schema->PC_Texture, TEXT("A"), LOCTEXT("A", "A"));
 }
 
 
diff --git a/Engine/Plugins/Mutable/Source/CustomizableObjectEditor/Private/MuCOE/Nodes/CustomizableObjectNodeTextureTransform.cpp b/Engine/Plugins/Mutable/Source/CustomizableObjectEditor/Private/MuCOE/Nodes/CustomizableObjectNodeTextureTransform.cpp
--- a/Engine/Plugins/Mutable/Source/CustomizableObjectEditor/Private/MuCOE/Nodes/CustomizableObjectNodeTextureTransform.cpp
+++ b/Engine/Plugins/Mutable/Source/CustomizableObjectEditor/Private/MuCOE/Nodes/CustomizableObjectNodeTextureTransform.cpp
@@ -13,27 +13,15 @@ class UCustomizableObjectNodeRemapPins;
 
 void UCustomizableObjectNodeTextureTransform::AllocateDefaultPins(UCustomizableObjectNodeRemapPins* RemapPins)
 {
-	UEdGraphPin* OutputPin = CustomCreatePinSimple(EGPD_Output, UEdGraphSchema_CustomizableObject::PC_Texture);
-	OutputPin->bDefaultValueIsIgnored = true;
+	CustomCreatePin(EGPD_Output, UEdGraphSchema_CustomizableObject::PC_Texture);
 	
+	BaseImagePinReference = CustomCreatePin(EGPD_Input, UEdGraphSchema_CustomizableObject::PC_Texture);
 	
-	UEdGraphPin* ImagePin = CustomCreatePinSimple(EGPD_Input, UEdGraphSchema_CustomizableObject::PC_Texture);
-	BaseImagePinReference = FEdGraphPinReference(ImagePin);
-	
-	UEdGraphPin* OffsetXPin = CustomCreatePin(EGPD_Input, UEdGraphSchema_CustomizableObject::PC_Float, FName(TEXT("Offset X")));
-	OffsetXPin->bDefaultValueIsIgnored = true;
-	
-	UEdGraphPin* OffsetYPin = CustomCreatePin(EGPD_Input, UEdGraphSchema_CustomizableObject::PC_Float, FName(TEXT("Offset Y")));
-	OffsetYPin->bDefaultValueIsIgnored = true;
-
-	UEdGraphPin* ScaleXPin = CustomCreatePin(EGPD_Input, UEdGraphSchema_CustomizableObject::PC_Float, FName(TEXT("Scale X")));
-	ScaleXPin->bDefaultValueIsIgnored = true;
-
-	UEdGraphPin* ScaleYPin = CustomCreatePin(EGPD_Input, UEdGraphSchema_CustomizableObject::PC_Float, FName(TEXT("Scale Y")));
-	ScaleYPin->bDefaultValueIsIgnored = true;
-	
-	UEdGraphPin* RotationPin = CustomCreatePin(EGPD_Input, UEdGraphSchema_CustomizableObject::PC_Float, FName(TEXT("Rotation")));
-	RotationPin->bDefaultValueIsIgnored = true;
+	CustomCreatePin(EGPD_Input, UEdGraphSchema_CustomizableObject::PC_Float, TEXT("Offset X"), LOCTEXT("OffsetX", "Offset X"));
+	CustomCreatePin(EGPD_Input, UEdGraphSchema_CustomizableObject::PC_Float, TEXT("Offset Y"), LOCTEXT("OffsetY", "Offset Y"));
+	CustomCreatePin(EGPD_Input, UEdGraphSchema_CustomizableObject::PC_Float, TEXT("Scale X"), LOCTEXT("ScaleX", "Scale X"));
+	CustomCreatePin(EGPD_Input, UEdGraphSchema_CustomizableObject::PC_Float, TEXT("Scale Y"), LOCTEXT("ScaleY", "Scale Y"));
+	CustomCreatePin(EGPD_Input, UEdGraphSchema_CustomizableObject::PC_Float, TEXT("Rotation"), LOCTEXT("Rotation", "Rotation"));
 }
 
 void UCustomizableObjectNodeTextureTransform::BackwardsCompatibleFixup(int32 CustomizableObjectCustomVersion)
diff --git a/Engine/Plugins/Mutable/Source/CustomizableObjectEditor/Private/MuCOE/Nodes/CustomizableObjectNodeTransformConstant.cpp b/Engine/Plugins/Mutable/Source/CustomizableObjectEditor/Private/MuCOE/Nodes/CustomizableObjectNodeTransformConstant.cpp
--- a/Engine/Plugins/Mutable/Source/CustomizableObjectEditor/Private/MuCOE/Nodes/CustomizableObjectNodeTransformConstant.cpp
+++ b/Engine/Plugins/Mutable/Source/CustomizableObjectEditor/Private/MuCOE/Nodes/CustomizableObjectNodeTransformConstant.cpp
@@ -230,13 +230,7 @@ FText UCustomizableObjectNodeTransformConstant::GetTooltipText() const
 
 void UCustomizableObjectNodeTransformConstant::AllocateDefaultPins(UCustomizableObjectNodeRemapPins* RemapPins)
 {
-	const FName Type = UEdGraphSchema_CustomizableObject::PC_Transform;
-	const FName PinName = UEdGraphSchema_CustomizableObject::GetPinCategoryName(Type);
-	const FText PinFriendlyName = UEdGraphSchema_CustomizableObject::GetPinCategoryFriendlyName(Type);
-	
-	UEdGraphPin* ValuePin = CustomCreatePin(EGPD_Output, Type, PinName);
-	ValuePin->PinFriendlyName = PinFriendlyName;
-	ValuePin->bDefaultValueIsIgnored = true;
+	CustomCreatePin(EGPD_Output, UEdGraphSchema_CustomizableObject::PC_Transform);
 }
 
 
diff --git a/Engine/Plugins/Mutable/Source/CustomizableObjectEditor/Private/MuCOE/Nodes/CustomizableObjectNodeTunnel.cpp b/Engine/Plugins/Mutable/Source/CustomizableObjectEditor/Private/MuCOE/Nodes/CustomizableObjectNodeTunnel.cpp
--- a/Engine/Plugins/Mutable/Source/CustomizableObjectEditor/Private/MuCOE/Nodes/CustomizableObjectNodeTunnel.cpp
+++ b/Engine/Plugins/Mutable/Source/CustomizableObjectEditor/Private/MuCOE/Nodes/CustomizableObjectNodeTunnel.cpp
@@ -98,8 +98,7 @@ void UCustomizableObjectNodeTunnel::AllocateDefaultPins(UCustomizableObjectNodeR
 			UCustomizableObjectNodeTunnelPinData* PinData = NewObject<UCustomizableObjectNodeTunnelPinData>(this);
 			PinData->VariableId = Variable->UniqueId;
 
-			const FName PinName = Variable->Name;
-			UEdGraphPin* Pin = CustomCreatePin(PinDirection, Variable->PinCategoryType, PinName, PinData);
+			CustomCreatePin(PinDirection, Variable->PinCategoryType, Variable->Name, FText::FromName(Variable->Name), PinData);
 		}
 	}
 }
diff --git a/Engine/Plugins/Mutable/Source/CustomizableObjectEditor/Private/MuCOE/Nodes/CustomizableObjectNodeVariation.cpp b/Engine/Plugins/Mutable/Source/CustomizableObjectEditor/Private/MuCOE/Nodes/CustomizableObjectNodeVariation.cpp
--- a/Engine/Plugins/Mutable/Source/CustomizableObjectEditor/Private/MuCOE/Nodes/CustomizableObjectNodeVariation.cpp
+++ b/Engine/Plugins/Mutable/Source/CustomizableObjectEditor/Private/MuCOE/Nodes/CustomizableObjectNodeVariation.cpp
@@ -33,9 +33,7 @@ void UCustomizableObjectNodeVariation::AllocateDefaultPins(UCustomizableObjectNo
 	const bool bIsInputPinArray = IsInputPinArray();
 
 	{
-		const FName PinName = FName(UEdGraphSchema_CustomizableObject::GetPinCategoryName(GetCategory()).ToString());
-		UEdGraphPin* Pin = CustomCreatePin(EGPD_Output, Category, PinName);
-		Pin->PinFriendlyName = UEdGraphSchema_CustomizableObject::GetPinCategoryFriendlyName(GetCategory());
+		CustomCreatePin(EGPD_Output, Category);
 	}
 	
 	VariationsPins.SetNum(VariationsData.Num());
@@ -44,7 +42,11 @@ void UCustomizableObjectNodeVariation::AllocateDefaultPins(UCustomizableObjectNo
 	for (int32 VariationIndex = VariationsData.Num() - 1; VariationIndex >= 0; --VariationIndex)
 	{
 		const FName PinName = FName(FString::Printf( TEXT("Variation %d"), VariationIndex));
-		UEdGraphPin* VariationPin = CustomCreatePin(EGPD_Input, Category, PinName, bIsInputPinArray);
+		UEdGraphPin* VariationPin = CustomCreatePin(EGPD_Input, Category, PinName, FText::FromName(PinName));
+		if (bIsInputPinArray)
+		{
+			VariationPin->PinType.ContainerType = EPinContainerType::Array;
+		}
 
 		FString TagName = GetTagDisplayName(VariationsData[VariationIndex].Tag);
 		VariationPin->PinFriendlyName = FText::Format(LOCTEXT("Variation_Pin_FriendlyName", "Variation {0} [{1}]"), VariationIndex, FText::FromString(TagName));
@@ -53,12 +55,16 @@ void UCustomizableObjectNodeVariation::AllocateDefaultPins(UCustomizableObjectNo
 
 		FText VariationTagSufix = LOCTEXT("TagSufix", " - Tag");
 		const FName TagPinName = FName(PinName.ToString() + VariationTagSufix.ToString());
-		UEdGraphPin* VariationTagPin = CustomCreatePin(EGPD_Input, Schema->PC_String, TagPinName);
+		UEdGraphPin* VariationTagPin = CustomCreatePin(EGPD_Input, Schema->PC_String, TagPinName, FText::FromName(TagPinName));
 
 		VariationTagPins[VariationIndex] = VariationTagPin;
 	}
 
-	CustomCreatePin(EGPD_Input, Category, FName(TEXT("Default")), bIsInputPinArray);
+	UEdGraphPin* Pin = CustomCreatePin(EGPD_Input, Category, TEXT("Default"), LOCTEXT("Default", "Default"));
+	if (bIsInputPinArray)
+	{
+		Pin->PinType.ContainerType = EPinContainerType::Array;
+	}
 }
 
 
@@ -205,7 +211,7 @@ void UCustomizableObjectNodeVariation::BackwardsCompatibleFixup(int32 Customizab
 				const FName PinName = FName(FString::Printf(TEXT("Variation %d"), VariationIndex));
 				FText VariationTagSufix = LOCTEXT("TagSufix", " - Tag");
 				const FName TagPinName = FName(PinName.ToString() + VariationTagSufix.ToString());
-				UEdGraphPin* VariationTagPin = CustomCreatePin(EGPD_Input, UEdGraphSchema_CustomizableObject::PC_String, TagPinName);
+				UEdGraphPin* VariationTagPin = CustomCreatePin(EGPD_Input, UEdGraphSchema_CustomizableObject::PC_String, TagPinName, FText::FromName(TagPinName));
 				VariationTagPins[VariationIndex] = VariationTagPin;
 			}
 		}
diff --git a/Engine/Plugins/Mutable/Source/CustomizableObjectEditor/Private/MuCOE/Nodes/SCustomizableObjectNodeMaterial.cpp b/Engine/Plugins/Mutable/Source/CustomizableObjectEditor/Private/MuCOE/Nodes/SCustomizableObjectNodeMaterial.cpp
--- a/Engine/Plugins/Mutable/Source/CustomizableObjectEditor/Private/MuCOE/Nodes/SCustomizableObjectNodeMaterial.cpp
+++ b/Engine/Plugins/Mutable/Source/CustomizableObjectEditor/Private/MuCOE/Nodes/SCustomizableObjectNodeMaterial.cpp
@@ -17,7 +17,7 @@ void SCustomizableObjectNodeMaterial::Construct(const FArguments& InArgs, UCusto
 
 TSharedPtr<SGraphPin> SCustomizableObjectNodeMaterial::CreatePinWidget(UEdGraphPin* Pin) const
 {
-	if ((Pin->PinType.PinCategory == UEdGraphSchema_CustomizableObject::PC_Texture || Pin->PinType.PinCategory == UEdGraphSchema_CustomizableObject::PC_PassthroughTexture) &&
+	if ((Pin->PinType.PinCategory == UEdGraphSchema_CustomizableObject::PC_Texture || Pin->PinType.PinCategory == UEdGraphSchema_CustomizableObject::PC_Texture_Passthrough) &&
 		Pin->Direction == EGPD_Input)
 	{
 		return SNew(SCustomizableObjectNodeMaterialPinImage, Pin);
diff --git a/Engine/Plugins/Mutable/Source/CustomizableObjectEditor/Public/MuCOE/EdGraphSchema_CustomizableObject.h b/Engine/Plugins/Mutable/Source/CustomizableObjectEditor/Public/MuCOE/EdGraphSchema_CustomizableObject.h
--- a/Engine/Plugins/Mutable/Source/CustomizableObjectEditor/Public/MuCOE/EdGraphSchema_CustomizableObject.h
+++ b/Engine/Plugins/Mutable/Source/CustomizableObjectEditor/Public/MuCOE/EdGraphSchema_CustomizableObject.h
@@ -33,9 +33,11 @@ public:
 	static UE_API const FName PC_Modifier;
 	static UE_API const FName PC_Mesh;
 	static UE_API const FName PC_SkeletalMesh;
-	static UE_API const FName PC_PassthroughSkeletalMesh;
+	static UE_API const FName PC_SkeletalMesh_Passthrough;
+	static UE_API const FName PC_SkeletalMesh_PassthroughParameter;
 	static UE_API const FName PC_Texture;
-	static UE_API const FName PC_PassthroughTexture;
+	static UE_API const FName PC_Texture_Passthrough;
+	static UE_API const FName PC_Texture_PassthroughParameter;
 	static UE_API const FName PC_Projector;
 	static UE_API const FName PC_GroupProjector;
 	static UE_API const FName PC_Color;
diff --git a/Engine/Plugins/Mutable/Source/CustomizableObjectEditor/Public/MuCOE/Nodes/CustomizableObjectNode.h b/Engine/Plugins/Mutable/Source/CustomizableObjectEditor/Public/MuCOE/Nodes/CustomizableObjectNode.h
--- a/Engine/Plugins/Mutable/Source/CustomizableObjectEditor/Public/MuCOE/Nodes/CustomizableObjectNode.h
+++ b/Engine/Plugins/Mutable/Source/CustomizableObjectEditor/Public/MuCOE/Nodes/CustomizableObjectNode.h
@@ -175,11 +175,7 @@ public:
 	 * Some nodes are added manually in the graph editor code and don't need to do this.
 	 */
 	UE_API virtual bool ShouldAddToContextMenu(FText& OutCategory) const;
-
-	UE_API void GetInputPins(TArray<class UEdGraphPin*>& OutInputPins) const;
-	UE_API void GetOutputPins(TArray<class UEdGraphPin*>& OutOutputPins) const;
-	UE_API UEdGraphPin* GetOutputPin(int32 OutputIndex) const;
-
+	
 	virtual bool IsNodeOutDatedAndNeedsRefresh() { return false; }
 	virtual FString GetRefreshMessage() const { return "Refresh Node."; }
 	UE_API void SetRefreshNodeWarning();
@@ -245,16 +241,11 @@ public:
 	/** If returns true, only a single link can go out of every output pin. False by default. */
 	UE_API virtual bool IsSingleOutputNode() const;
 
-	/** Wrapper of the pin creation method. Will use the PinName and PinFriendlyname returned by the "UEdGraphSchema_CustomizableObject" */
-	UE_API UEdGraphPin* CustomCreatePinSimple( EEdGraphPinDirection Direction, const FName& Category, bool bIsArray=false );
+	/** Wrapper of the pin creation method. Will use the PinName and PinFriendlyName returned by the "UEdGraphSchema_CustomizableObject" */
+	UE_API UEdGraphPin* CustomCreatePin( EEdGraphPinDirection Direction, const FName& Category, UCustomizableObjectNodePinData* PinData = nullptr);
 	
 	// Wrapper of the pin creation method, to cope with multiple Unreal Engine versions
-	UE_API UEdGraphPin* CustomCreatePin( EEdGraphPinDirection Direction, const FName& Category, const FName& Name, bool bIsArray=false );
-
-	/** Create pin with attached data.
-	 *
-     * @param PinData Pin data to be saved. */
-	UE_API UEdGraphPin* CustomCreatePin(EEdGraphPinDirection Direction, const FName& Category, const FName& Name, UCustomizableObjectNodePinData* PinData);
+	UE_API UEdGraphPin* CustomCreatePin(EEdGraphPinDirection Direction, const FName& Category, const FName& Name, const FText& FriendlyName, UCustomizableObjectNodePinData* PinData = nullptr);
 	
 	virtual void AddAttachedErrorData(const FAttachedErrorDataView& ) {};
 	virtual void ResetAttachedErrorData() {};
