[Mutable] Now scalar values from material breaks are stored into the coderunnerbegin stack.
+ This fixes a crash when a float pin from break material node was linked to a material modify node.

#rnx

#Synced-CL 49100904

diff --git a/Engine/Plugins/Mutable/Source/MutableRuntime/Private/MuR/CodeRunnerBegin.cpp b/Engine/Plugins/Mutable/Source/MutableRuntime/Private/MuR/CodeRunnerBegin.cpp
--- a/Engine/Plugins/Mutable/Source/MutableRuntime/Private/MuR/CodeRunnerBegin.cpp
+++ b/Engine/Plugins/Mutable/Source/MutableRuntime/Private/MuR/CodeRunnerBegin.cpp
@@ -2361,9 +2361,18 @@ namespace UE::Mutable::Private
 					StoreNone(Item);
 					break;
 				}
+		
+			case EOpType::SC_MATERIAL_BREAK:
+			{
+				OP::MaterialBreakArgs Args = Program.GetOpArgs<OP::MaterialBreakArgs>(Item.At);
+				PushOp(FScheduledOpInline(Args.Material, Item));
+
+				// Hack: we need to store any value here. TODO GMT: remove this.
+				StoreFloat(Item, 0.0f);
+				break;
+			}
 
 			case EOpType::IM_MATERIAL_BREAK:
-			case EOpType::SC_MATERIAL_BREAK:
 			case EOpType::CO_MATERIAL_BREAK:
 				{
 					OP::MaterialBreakArgs Args = Program.GetOpArgs<OP::MaterialBreakArgs>(Item.At);
diff --git a/Engine/Plugins/Mutable/Source/MutableTools/Private/MuT/ASTOpMaterialBreak.cpp b/Engine/Plugins/Mutable/Source/MutableTools/Private/MuT/ASTOpMaterialBreak.cpp
--- a/Engine/Plugins/Mutable/Source/MutableTools/Private/MuT/ASTOpMaterialBreak.cpp
+++ b/Engine/Plugins/Mutable/Source/MutableTools/Private/MuT/ASTOpMaterialBreak.cpp
@@ -38,8 +38,7 @@ namespace UE::Mutable::Private
 			const ASTOpMaterialBreak* Other = static_cast<const ASTOpMaterialBreak*>(&OtherUntyped);
 			return Type == Other->Type &&
 				Material == Other->Material &&
-				ParameterName == Other->ParameterName &&
-				LayerIndex == Other->LayerIndex;
+				ParameterKey == Other->ParameterKey;
 		}
 		return false;
 	}
@@ -49,9 +48,8 @@ namespace UE::Mutable::Private
 	{
 		Ptr<ASTOpMaterialBreak> Result = new ASTOpMaterialBreak();
 		Result->Type = Type;
-		Result->ParameterName = ParameterName;
+		Result->ParameterKey = ParameterKey;
 		Result->Material = MapChild(Material.child());
-		Result->LayerIndex = LayerIndex;
 
 		return Result;
 	}
@@ -66,9 +64,8 @@ namespace UE::Mutable::Private
 	uint64 ASTOpMaterialBreak::Hash() const
 	{
 		uint64 res = std::hash<uint64>()(uint64(Type));
-		hash_combine(res, GetTypeHash(ParameterName));
+		hash_combine(res, GetTypeHash(ParameterKey));
 		hash_combine(res, Material.child().get());
-		hash_combine(res, GetTypeHash(LayerIndex));
 		
 		return res;
 	}
@@ -97,8 +94,8 @@ namespace UE::Mutable::Private
 				Args.Material = Material->linkedAddress;
 			}
 
-			Args.ParameterName = Program.AddConstant(ParameterName.ToString());
-			Args.LayerIndex = LayerIndex;
+			Args.ParameterName = Program.AddConstant(ParameterKey.ParameterName.ToString());
+			Args.LayerIndex = ParameterKey.LayerIndex;
 
 			AppendCode(Program.ByteCode, Args);
 		}
diff --git a/Engine/Plugins/Mutable/Source/MutableTools/Private/MuT/ASTOpMaterialBreak.h b/Engine/Plugins/Mutable/Source/MutableTools/Private/MuT/ASTOpMaterialBreak.h
--- a/Engine/Plugins/Mutable/Source/MutableTools/Private/MuT/ASTOpMaterialBreak.h
+++ b/Engine/Plugins/Mutable/Source/MutableTools/Private/MuT/ASTOpMaterialBreak.h
@@ -22,9 +22,7 @@ namespace UE::Mutable::Private
 
 		ASTChild Material;
 
-		FName ParameterName;
-
-		int32 LayerIndex;
+		FParameterKey ParameterKey;
 
 	public:
 
diff --git a/Engine/Plugins/Mutable/Source/MutableTools/Private/MuT/CodeGenerator_Color.cpp b/Engine/Plugins/Mutable/Source/MutableTools/Private/MuT/CodeGenerator_Color.cpp
--- a/Engine/Plugins/Mutable/Source/MutableTools/Private/MuT/CodeGenerator_Color.cpp
+++ b/Engine/Plugins/Mutable/Source/MutableTools/Private/MuT/CodeGenerator_Color.cpp
@@ -558,9 +558,8 @@ namespace UE::Mutable::Private
 	{
 		Ptr<ASTOpMaterialBreak> MaterialOp = new ASTOpMaterialBreak();
 
-		// Set the parameter name that this break node will generate
-		MaterialOp->ParameterName = Typed->ParameterName;
-		MaterialOp->LayerIndex = Typed->LayerIndex;
+		// Set the key that identifies the break node parameter
+		MaterialOp->ParameterKey = { Typed->ParameterName, Typed->LayerIndex };
 
 		// Generate Material Source
 		FMaterialGenerationOptions MaterialSourceOptions;
diff --git a/Engine/Plugins/Mutable/Source/MutableTools/Private/MuT/CodeGenerator_Image.cpp b/Engine/Plugins/Mutable/Source/MutableTools/Private/MuT/CodeGenerator_Image.cpp
--- a/Engine/Plugins/Mutable/Source/MutableTools/Private/MuT/CodeGenerator_Image.cpp
+++ b/Engine/Plugins/Mutable/Source/MutableTools/Private/MuT/CodeGenerator_Image.cpp
@@ -1968,9 +1968,8 @@ namespace UE::Mutable::Private
 	{
 		Ptr<ASTOpMaterialBreak> MaterialOp = new ASTOpMaterialBreak();
 
-		// Set the parameter name that this break node will generate
-		MaterialOp->ParameterName = Node->ParameterName;
-		MaterialOp->LayerIndex = Node->LayerIndex;
+		// Set the parameter key that this break node will generate
+		MaterialOp->ParameterKey = { Node->ParameterName, Node->LayerIndex };
 
 		// Generate Material Source
 		FMaterialGenerationResult MaterialSourceResult;
diff --git a/Engine/Plugins/Mutable/Source/MutableTools/Private/MuT/CodeGenerator_Scalar.cpp b/Engine/Plugins/Mutable/Source/MutableTools/Private/MuT/CodeGenerator_Scalar.cpp
--- a/Engine/Plugins/Mutable/Source/MutableTools/Private/MuT/CodeGenerator_Scalar.cpp
+++ b/Engine/Plugins/Mutable/Source/MutableTools/Private/MuT/CodeGenerator_Scalar.cpp
@@ -439,9 +439,8 @@ namespace UE::Mutable::Private
 	{
 		Ptr<ASTOpMaterialBreak> MaterialOp = new ASTOpMaterialBreak();
 
-		// Set the parameter name that this break node will generate
-		MaterialOp->ParameterName = Node->ParameterName;
-		MaterialOp->LayerIndex = Node->LayerIndex;
+		// Set the key that identifies the break node parameter
+		MaterialOp->ParameterKey = { Node->ParameterName, Node->LayerIndex };
 
 		// Generate Material Source
 		FMaterialGenerationOptions MaterialSourceOptions;
