[Mutable] Added Asset User Data from Skeletal Mesh parameters.
+ Now, the skeletal meshes generated during the CO instance update that come from a SKM parameter, will contain their asset usser data.

#rnx

#Synced-CL 49100904

diff --git a/Engine/Plugins/Mutable/Source/CustomizableObject/Private/MuCO/CustomizableObjectInstance.cpp b/Engine/Plugins/Mutable/Source/CustomizableObject/Private/MuCO/CustomizableObjectInstance.cpp
--- a/Engine/Plugins/Mutable/Source/CustomizableObject/Private/MuCO/CustomizableObjectInstance.cpp
+++ b/Engine/Plugins/Mutable/Source/CustomizableObject/Private/MuCO/CustomizableObjectInstance.cpp
@@ -5099,6 +5099,25 @@ UE::Tasks::FTask LoadAdditionalAssetsAndData(const TSharedRef<FUpdateContextPriv
 						ComponentData->StreamedResourceIndex.AddUnique(ResourceIndex);
 					}
 				}
+
+				// Add AssetUserData from Mesh Parameters
+				for (const TStrongObjectPtr<USkeletalMesh>& StrongMesh : MutableMesh->SkeletalMeshes)
+				{
+					if (USkeletalMesh* SkeletalMesh = StrongMesh.Get())
+					{
+						if (const TArray<UAssetUserData*>* AssetUserDataArray = SkeletalMesh->GetAssetUserDataArray())
+						{
+							for (UAssetUserData* AssetUserData : *AssetUserDataArray)
+							{
+								// No need to load since the SKM is a StrongObjectPtr
+								if (AssetUserData)
+								{
+									ComponentData->AssetUserDataArray.Add(AssetUserData);
+								}
+							}
+						}
+					}
+				}
 				
 				for (int32 TagIndex = 0; TagIndex < MutableMesh->GetTagCount(); ++TagIndex)
 				{
