[Mutable] Added optional evaluation to CodeRunnerBegin.

- Removed CodeRunnerBegin System dependency.
- Typed constant indices.
- Simplified RelevantParameters.

#Synced-CL 49159615

diff --git a/Engine/Plugins/Mutable/Source/CustomizableObject/Internal/MuCO/CustomizableObjectInstancePrivate.h b/Engine/Plugins/Mutable/Source/CustomizableObject/Internal/MuCO/CustomizableObjectInstancePrivate.h
--- a/Engine/Plugins/Mutable/Source/CustomizableObject/Internal/MuCO/CustomizableObjectInstancePrivate.h
+++ b/Engine/Plugins/Mutable/Source/CustomizableObject/Internal/MuCO/CustomizableObjectInstancePrivate.h
@@ -245,7 +245,7 @@ public:
 	
 	// Indices of the parameters that are relevant for the given parameter values.
 	// This only gets updated if parameter decorations are generated.
-	TArray<int> RelevantParameters;
+	TArray<uint32> RelevantParameters;
 
 	// Only used in LiveUpdateMode to reuse core instances between updates and their temp data to speed up updates, but spend way more memory
 	UE::Mutable::Private::FInstance::FID LiveUpdateModeInstanceID = 0;
@@ -341,6 +341,8 @@ public:
 	/** Finds in ProjectorParameters a parameter with name ParamName, returns the index if found, -1 otherwise. */
 	UE_API int32 FindProjectorParameterNameIndex(const FString& ParamName) const;
 	
+	UE_API bool IsParameterRelevant(uint32 ParameterIndex) const;
+
 	UE_API UCustomizableObjectInstance* GetPublic() const;
 
 	// If any components are using this instance, they will store the min of their distances to the player here every frame for LOD purposes
diff --git a/Engine/Plugins/Mutable/Source/CustomizableObject/Internal/MuCO/CustomizableObjectSystemPrivate.h b/Engine/Plugins/Mutable/Source/CustomizableObject/Internal/MuCO/CustomizableObjectSystemPrivate.h
--- a/Engine/Plugins/Mutable/Source/CustomizableObject/Internal/MuCO/CustomizableObjectSystemPrivate.h
+++ b/Engine/Plugins/Mutable/Source/CustomizableObject/Internal/MuCO/CustomizableObjectSystemPrivate.h
@@ -569,7 +569,7 @@ public:
 	bool bForceHighPriority = false;
 	
 	FInstanceUpdateData InstanceUpdateData;
-	TArray<int32> RelevantParametersInProgress;
+	TArray<uint32> RelevantParameters;
 
 	const FInstanceUpdateData::FComponent* GetComponentUpdateData(FCustomizableObjectInstanceComponentIndex InstanceComponentIndex);
 
diff --git a/Engine/Plugins/Mutable/Source/CustomizableObject/Internal/MuCO/FMutableTaskGraph.h b/Engine/Plugins/Mutable/Source/CustomizableObject/Internal/MuCO/FMutableTaskGraph.h
--- a/Engine/Plugins/Mutable/Source/CustomizableObject/Internal/MuCO/FMutableTaskGraph.h
+++ b/Engine/Plugins/Mutable/Source/CustomizableObject/Internal/MuCO/FMutableTaskGraph.h
@@ -43,7 +43,7 @@ class FMutableTaskGraph
 	struct FGameThreadTask
 	{
 		FString DebugName;
-		TUniqueFunction<void()> Body;
+		TUniqueFunction<void(float)> Body;
 		TArray<UE::Tasks::FTask> Prerequisites;
 		bool bLockMutableThread = false;
 	
@@ -56,7 +56,7 @@ public:
 
 	~FMutableTaskGraph();
 
-	void AddGameThreadTask(const TCHAR* DebugName, TUniqueFunction<void()>&& TaskBody, bool bLockMutableThread = false, const TArray<UE::Tasks::FTask>& Prerequisites = {});
+	void AddGameThreadTask(const TCHAR* DebugName, TUniqueFunction<void(float)>&& TaskBody, bool bLockMutableThread = false, const TArray<UE::Tasks::FTask>& Prerequisites = {});
 	
 	/** Create and launch a task on the Mutable Thread with Normal priority. */
 	UE::Tasks::FTask AddMutableThreadTask(const TCHAR* DebugName, TUniqueFunction<void()>&& TaskBody, const TArray<UE::Tasks::FTask>& Prerequisites = {});
diff --git a/Engine/Plugins/Mutable/Source/CustomizableObject/Private/MuCO/CustomizableObjectInstance.cpp b/Engine/Plugins/Mutable/Source/CustomizableObject/Private/MuCO/CustomizableObjectInstance.cpp
--- a/Engine/Plugins/Mutable/Source/CustomizableObject/Private/MuCO/CustomizableObjectInstance.cpp
+++ b/Engine/Plugins/Mutable/Source/CustomizableObject/Private/MuCO/CustomizableObjectInstance.cpp
@@ -588,10 +588,10 @@ void UCustomizableObjectInstance::SetCurrentState(const FString& StateName)
 }
 
 
-bool UCustomizableObjectInstance::IsParameterRelevant(int32 ParameterIndex) const
+bool UCustomizableInstancePrivate::IsParameterRelevant(uint32 ParameterIndex) const
 {
 	// This should have been precalculated in the last update if the appropriate flag in the instance was set.
-	return GetPrivate()->RelevantParameters.Contains(ParameterIndex);
+	return RelevantParameters.Contains(ParameterIndex);
 }
 
 
diff --git a/Engine/Plugins/Mutable/Source/CustomizableObject/Private/MuCO/CustomizableObjectSystem.cpp b/Engine/Plugins/Mutable/Source/CustomizableObject/Private/MuCO/CustomizableObjectSystem.cpp
--- a/Engine/Plugins/Mutable/Source/CustomizableObject/Private/MuCO/CustomizableObjectSystem.cpp
+++ b/Engine/Plugins/Mutable/Source/CustomizableObject/Private/MuCO/CustomizableObjectSystem.cpp
@@ -1952,9 +1952,7 @@ namespace impl
 
 		check(OperationData->Parameters);
 		check(OperationData->InstanceID != 0);
-
-		OperationData->RelevantParametersInProgress.Empty();
-
+		
 		// This must run in the mutable thread.
 		check(UCustomizableObjectSystem::GetInstance() != nullptr);
 		check(UCustomizableObjectSystem::GetInstance()->GetPrivate() != nullptr);
@@ -1962,20 +1960,7 @@ namespace impl
 		// Update the parameter relevancy.
 		{
 			MUTABLE_CPUPROFILER_SCOPE(ParameterRelevancy)
-
-			const int32 NumParameters = OperationData->Parameters->GetCount();
-
-			TArray<bool> Relevant;
-			Relevant.SetNumZeroed(NumParameters);
-			OperationData->MutableSystem->GetParameterRelevancy(OperationData->InstanceID, OperationData->Parameters, Relevant.GetData());
-
-			for (int32 ParamIndex = 0; ParamIndex < NumParameters; ++ParamIndex)
-			{
-				if (Relevant[ParamIndex])
-				{
-					OperationData->RelevantParametersInProgress.Add(ParamIndex);
-				}
-			}
+			OperationData->MutableSystem->GetParameterRelevancy(OperationData->InstanceID, OperationData->Parameters, OperationData->RelevantParameters);
 		}
 
 		OperationData->TaskGetMeshesTime = FPlatformTime::Seconds() - OperationData->TaskGetMeshesTimeStart;
@@ -2777,46 +2762,45 @@ namespace impl
 	}
 
 	
-	void Task_Game_Callbacks(const TSharedRef<FUpdateContextPrivate>& OperationData)
+	void Task_Game_Callbacks(const TSharedRef<FUpdateContextPrivate>& Context)
 	{
 		MUTABLE_CPUPROFILER_SCOPE(Task_Game_Callbacks)
-		FMutableScopeTimer Timer(OperationData->TaskCallbacksTime);
+		FMutableScopeTimer Timer(Context->TaskCallbacksTime);
 
 		check(IsInGameThread());
 
 		UCustomizableObjectSystem* System = UCustomizableObjectSystem::GetInstance();
 		if (!System || !System->IsValidLowLevel() || System->HasAnyFlags(RF_BeginDestroyed))
 		{
-			OperationData->UpdateResult = EUpdateResult::Error;
-			FinishUpdateGlobal(OperationData);
+			Context->UpdateResult = EUpdateResult::Error;
+			FinishUpdateGlobal(Context);
 			return;
 		}
 
-		UCustomizableObjectInstance* CustomizableObjectInstance = OperationData->Instance.Get();
-
-		// TODO: Review checks.
-		if (!CustomizableObjectInstance || !CustomizableObjectInstance->IsValidLowLevel() )
+		UCustomizableObjectInstance* Instance = Context->Instance.Get();
+		if (!Instance)
 		{
-			OperationData->UpdateResult = EUpdateResult::Error;
-			FinishUpdateGlobal(OperationData);
+			Context->UpdateResult = EUpdateResult::Error;
+			FinishUpdateGlobal(Context);
 			return;
 		}
 
-		UCustomizableObject* CustomizableObject = CustomizableObjectInstance->GetCustomizableObject();
-		if (!CustomizableObject)
+		UCustomizableObject* Object = Instance->GetCustomizableObject();
+		if (!Object)
 		{
-			OperationData->UpdateResult = EUpdateResult::Error;
-			FinishUpdateGlobal(OperationData);
+			Context->UpdateResult = EUpdateResult::Error;
+			FinishUpdateGlobal(Context);
 			return;
 		}
-		
+
+		Instance->GetPrivate()->RelevantParameters = MoveTemp(Context->RelevantParameters);
 
 		// Actual work
 		// TODO MTBL-391: Review This hotfix
-		UpdateSkeletalMesh(OperationData);
+		UpdateSkeletalMesh(Context);
 
 		// End Update
-		FinishUpdateGlobal(OperationData);
+		FinishUpdateGlobal(Context);
 	}
 
 
@@ -2919,7 +2903,7 @@ namespace impl
 		//-------------------------------------------------------------
 		CustomizableObjectSystemPrivate->MutableTaskGraph.AddGameThreadTask(
 			TEXT("Task_Game_Callbacks"),
-			[OperationData]()
+			[OperationData](float RemainingTime)
 			{
 				Task_Game_Callbacks(OperationData);
 			},
@@ -2973,13 +2957,6 @@ namespace impl
 				ObjectInstancePrivateData->LiveUpdateModeInstanceID = OperationData->InstanceID;
 			}
 		}
-
-		
-		if (OperationData->GetCapturedDescriptor().GetBuildParameterRelevancy())
-		{
-			// Relevancy
-			ObjectInstancePrivateData->RelevantParameters = OperationData->RelevantParametersInProgress;
-		}
 		
 		// TODO: If this is the first code that runs after the CO program has finished AND if it's
 		// guaranteed that the next CO program hasn't started yet, we need to call ClearActiveObject
@@ -3042,7 +3019,7 @@ namespace impl
 		//-------------------------------------------------------------
 		SystemPrivate->MutableTaskGraph.AddGameThreadTask(
 			TEXT("Task_Game_ConvertResources"),
-			[OperationData]()
+			[OperationData](float RemainingTime)
 			{
 				Task_Game_ConvertResources(OperationData);
 			},
@@ -3523,7 +3500,7 @@ namespace impl
 
 		SystemPrivate->MutableTaskGraph.AddGameThreadTask(
 		TEXT("Task_Game_LockCache"),
-		[Operation]()
+		[Operation](float RemainingTime)
 		{
 			impl::Task_Game_LockCache(Operation);
 		},
@@ -4305,7 +4282,7 @@ void UCustomizableObjectSystemPrivate::StartUpdateSkeletalMesh(const TSharedRef<
 
 	MutableTaskGraph.AddGameThreadTask(
 		TEXT("Task_Game_StartUpdate"),
-		[Context]
+		[Context](float RemainingTime)
 		{
 			impl::Task_Game_StartUpdate(Context);
 		},
diff --git a/Engine/Plugins/Mutable/Source/CustomizableObject/Private/MuCO/FMutableTaskGraph.cpp b/Engine/Plugins/Mutable/Source/CustomizableObject/Private/MuCO/FMutableTaskGraph.cpp
--- a/Engine/Plugins/Mutable/Source/CustomizableObject/Private/MuCO/FMutableTaskGraph.cpp
+++ b/Engine/Plugins/Mutable/Source/CustomizableObject/Private/MuCO/FMutableTaskGraph.cpp
@@ -71,9 +71,9 @@ FMutableTaskGraph::~FMutableTaskGraph()
 }
 
 
-void FMutableTaskGraph::AddGameThreadTask(const TCHAR* DebugName, TUniqueFunction<void()>&& TaskBody, bool bLockMutableThread, const TArray<UE::Tasks::FTask>& Prerequisites)
+void FMutableTaskGraph::AddGameThreadTask(const TCHAR* DebugName, TUniqueFunction<void(float)>&& TaskBody, bool bLockMutableThread, const TArray<UE::Tasks::FTask>& Prerequisites)
 {
-	GameThreadTasks.Enqueue({ DebugName, Forward<TUniqueFunction<void()>>(TaskBody), Prerequisites, bLockMutableThread });
+	GameThreadTasks.Enqueue({ DebugName, Forward<TUniqueFunction<void(float)>>(TaskBody), Prerequisites, bLockMutableThread });
 }
 
 
@@ -288,10 +288,13 @@ void FMutableTaskGraph::TryLaunchGameThreadTask()
 {
 	MUTABLE_CPUPROFILER_SCOPE(AdvanceCurrentOperation);
 
-	double CurrentTime = FPlatformTime::Seconds(); 
-	const double EndTime = CurrentTime + CVarGameThreadTaskMaxTime.GetValueOnGameThread();
+	const double TickTime = CVarGameThreadTaskMaxTime.GetValueOnGameThread();
+	
+	const double StartTime = FPlatformTime::Seconds();
+	double RemainingTime = TickTime;
+	bool bLoop = true;
 
-	while (CurrentTime < EndTime)
+	while (bLoop)
 	{
 		const FGameThreadTask* PendingTask = GameThreadTasks.Peek();
 		if (!PendingTask)
@@ -318,10 +321,11 @@ void FMutableTaskGraph::TryLaunchGameThreadTask()
 		GameThreadTasks.Dequeue(CurrentTask); // Dequeue now since current task can enqueue new tasks.
 
 		// Careful! From this point, PendingTask is no longer valid.
-		
-		CurrentTask.Body();
 
-		CurrentTime = FPlatformTime::Seconds(); 
+		CurrentTask.Body(RemainingTime);
+
+		RemainingTime = TickTime - (FPlatformTime::Seconds() - StartTime);
+		bLoop = RemainingTime > 0.0;
 	}
 
 	UnlockMutableThread();
diff --git a/Engine/Plugins/Mutable/Source/CustomizableObject/Public/MuCO/CustomizableObjectInstance.h b/Engine/Plugins/Mutable/Source/CustomizableObject/Public/MuCO/CustomizableObjectInstance.h
--- a/Engine/Plugins/Mutable/Source/CustomizableObject/Public/MuCO/CustomizableObjectInstance.h
+++ b/Engine/Plugins/Mutable/Source/CustomizableObject/Public/MuCO/CustomizableObjectInstance.h
@@ -433,10 +433,7 @@ public:
 
 	/** Immediately destroy the Mutable Core Live Update Instance attached to this (if exists). */
 	CUSTOMIZABLEOBJECT_API void DestroyLiveUpdateInstance();
-
-	/** Return true if changing the parameter would affect the Instance given its current generation. */
-	CUSTOMIZABLEOBJECT_API bool IsParameterRelevant(int32 ParameterIndex) const;
-
+	
 	/** Return true if the given parameter has any effect in the current object state, and considering the current values of the other parameters. */
 	UFUNCTION(BlueprintCallable, Category = CustomizableObjectInstance)
 	CUSTOMIZABLEOBJECT_API bool IsParameterRelevant(const FString& ParamName) const;
diff --git a/Engine/Plugins/Mutable/Source/CustomizableObjectEditor/Private/MuCOE/CustomizableInstanceDetails.cpp b/Engine/Plugins/Mutable/Source/CustomizableObjectEditor/Private/MuCOE/CustomizableInstanceDetails.cpp
--- a/Engine/Plugins/Mutable/Source/CustomizableObjectEditor/Private/MuCOE/CustomizableInstanceDetails.cpp
+++ b/Engine/Plugins/Mutable/Source/CustomizableObjectEditor/Private/MuCOE/CustomizableInstanceDetails.cpp
@@ -742,7 +742,7 @@ bool FCustomizableInstanceDetails::GenerateParametersView(IDetailCategoryBuilder
 			FParameterInfo ParameterSortInfo;
 			ParameterSortInfo.ParamIndexInObject = CustomizableObject->GetPrivate()->GetStateParameterIndex(CustomInstance->GetPrivate()->GetState(), ParamIndexInState);
 
-			if (CustomInstance->IsParameterRelevant(ParameterSortInfo.ParamIndexInObject) && IsVisible(ParameterSortInfo.ParamIndexInObject))
+			if (CustomInstance->GetPrivate()->IsParameterRelevant(ParameterSortInfo.ParamIndexInObject) && IsVisible(ParameterSortInfo.ParamIndexInObject))
 			{
 				ParameterSortInfo.ParamName = CustomizableObject->GetParameterName(ParameterSortInfo.ParamIndexInObject);
 				ParameterSortInfo.ParamUIOrder = CustomizableObject->GetParameterUIMetadata(ParameterSortInfo.ParamName).UIOrder;
@@ -809,7 +809,7 @@ bool FCustomizableInstanceDetails::GenerateParametersView(IDetailCategoryBuilder
 		{
 			FParameterInfo ParameterSortInfo;
 			ParameterSortInfo.ParamIndexInObject = ParamIndexInObject;
-			if ((!CustomInstance->GetPrivate()->bShowOnlyRelevantParameters || CustomInstance->IsParameterRelevant(ParameterSortInfo.ParamIndexInObject)) && IsVisible(ParameterSortInfo.ParamIndexInObject))
+			if ((!CustomInstance->GetPrivate()->bShowOnlyRelevantParameters || CustomInstance->GetPrivate()->IsParameterRelevant(ParameterSortInfo.ParamIndexInObject)) && IsVisible(ParameterSortInfo.ParamIndexInObject))
 			{
 				ParameterSortInfo.ParamName = CustomizableObject->GetParameterName(ParameterSortInfo.ParamIndexInObject);
 				ParameterSortInfo.ParamUIOrder = CustomizableObject->GetParameterUIMetadata(ParameterSortInfo.ParamName).UIOrder;
diff --git a/Engine/Plugins/Mutable/Source/MutableRuntime/Internal/MuR/Operations.h b/Engine/Plugins/Mutable/Source/MutableRuntime/Internal/MuR/Operations.h
--- a/Engine/Plugins/Mutable/Source/MutableRuntime/Internal/MuR/Operations.h
+++ b/Engine/Plugins/Mutable/Source/MutableRuntime/Internal/MuR/Operations.h
@@ -439,7 +439,11 @@ namespace UE::Mutable::Private
     struct TOperation
     {
         typedef ADDRESS_TYPE ADDRESS;
-        typedef ADDRESS_TYPE CONSTANT_STRING_INDEX;
+        typedef ADDRESS_TYPE PARAMETER;
+    	typedef ADDRESS_TYPE CONSTANT_SHAPE_INDEX;
+    	typedef ADDRESS_TYPE CONSTANT_MATRIX_INDEX;
+    	typedef ADDRESS_TYPE CONSTANT_EXTENSIONDATA_INDEX;
+    	typedef ADDRESS_TYPE CONSTANT_STRING_INDEX;
 		typedef ADDRESS_TYPE LIST_UINT32_ADDRESS;
 
         /** Arguments for every operation type. */
@@ -465,12 +469,12 @@ namespace UE::Mutable::Private
 
     	struct MatrixConstantArgs
     	{
-    		ADDRESS value;
+    		CONSTANT_MATRIX_INDEX value;
     	};
 
         struct ResourceConstantArgs
         {
-            ADDRESS value;
+            CONSTANT_EXTENSIONDATA_INDEX value;
         };
 
         struct MeshConstantArgs
@@ -492,12 +496,12 @@ namespace UE::Mutable::Private
     	
 		struct ParameterArgs
 		{
-			ADDRESS variable;
+			PARAMETER variable;
 		};
 
 		struct MaterialParameterArgs
 		{
-			ADDRESS Variable;
+			PARAMETER Variable;
 
 			LIST_UINT32_ADDRESS ScalarParameterNames;
 			LIST_UINT32_ADDRESS ColorParameterNames;
@@ -513,7 +517,7 @@ namespace UE::Mutable::Private
 
     	struct MeshSkeletalMeshBreakArgs
     	{
-    		ADDRESS SkeletalMeshParameter;
+    		PARAMETER SkeletalMeshParameter;
     		uint8 LOD = 0;
 			uint8 Section = 0;
 			uint32 MeshID = 0;
@@ -741,7 +745,9 @@ namespace UE::Mutable::Private
 
         struct ImageComposeArgs
         {
-            ADDRESS layout, base, blockImage;
+            ADDRESS layout;
+        	ADDRESS base;
+        	ADDRESS blockImage;
             ADDRESS mask;
             uint64 BlockId;
         };
@@ -990,7 +996,7 @@ namespace UE::Mutable::Private
 		{
 			ADDRESS Source;
 			ADDRESS MorphShape;
-			ADDRESS VertexSelectionShapeOrBone;
+			CONSTANT_SHAPE_INDEX VertexSelectionShapeOrBone;
 
 			float Dist, Factor, MaxBoneRadius;
 
diff --git a/Engine/Plugins/Mutable/Source/MutableRuntime/Internal/MuR/System.h b/Engine/Plugins/Mutable/Source/MutableRuntime/Internal/MuR/System.h
--- a/Engine/Plugins/Mutable/Source/MutableRuntime/Internal/MuR/System.h
+++ b/Engine/Plugins/Mutable/Source/MutableRuntime/Internal/MuR/System.h
@@ -215,16 +215,11 @@ namespace UE::Mutable::Private
         UE_API void ReleaseInstance(FInstance::FID InstanceID);
 
 		//! Calculate the relevancy of every parameter. Some parameters may be unused depending on
-		//! the values of other parameters. This method will set to true the flags for parameters
-		//! that are relevant, and to false otherwise. This is useful to hide irrelevant parameters
-		//! in dynamic user interfaces.
-        //! \param pModel The model used to create the FParameters instance.
-        //! \param pParameters Parameter set that we want to find the relevancy of.
-        //! \param pFlags is a pointer to a preallocated array of booleans that contains at least
-		//! pParameters->GetCount() elements.
-        UE_API void GetParameterRelevancy( FInstance::FID InstanceID,
-									const TSharedPtr<const FParameters>& FParameters,
-									bool* Flags );
+		//! the values of other parameters. This is useful to hide irrelevant parameters in
+		//! dynamic user interfaces.
+        UE_API void GetParameterRelevancy(FInstance::FID InstanceID,
+        	const TSharedPtr<const FParameters>& FParameters,
+			TArray<uint32>& OutRelevantParameters);
 
 		UE_API static FMeshId GetMeshId(const TSharedPtr<FMeshIdRegistry>& Registry, const TSharedPtr<const FModel>& Model, const FParameters* Params, OP::ADDRESS RootAt);
     	
diff --git a/Engine/Plugins/Mutable/Source/MutableRuntime/Private/MuR/CodeRunnerBegin.cpp b/Engine/Plugins/Mutable/Source/MutableRuntime/Private/MuR/CodeRunnerBegin.cpp
--- a/Engine/Plugins/Mutable/Source/MutableRuntime/Private/MuR/CodeRunnerBegin.cpp
+++ b/Engine/Plugins/Mutable/Source/MutableRuntime/Private/MuR/CodeRunnerBegin.cpp
@@ -1,6 +1,6 @@
 // Copyright Epic Games, Inc. All Rights Reserved.
 
-#include "CodeRunnerBegin.h"
+#include "MuR/CodeRunnerBegin.h"
 
 #include "CodeRunner.h"
 #include "MuR/MutableTrace.h"
@@ -11,13 +11,28 @@
 
 namespace UE::Mutable::Private
 {
-	CodeRunnerBegin::CodeRunnerBegin(FSystem& InSystem, const TSharedPtr<const FModel>& InModel, const TSharedPtr<const FParameters>& InParams, uint32 InLodMask) :
+	CodeRunnerBegin::CodeRunnerBegin(
+			const TSharedPtr<const FModel>& InModel,
+			const TSharedPtr<const FParameters>& InParams,
+			const TSharedPtr<FProgramCache>& InProgramCache,
+			const TSharedPtr<FMeshIdRegistry>& InMeshIdRegistry,
+			const TSharedPtr<FImageIdRegistry>& InImageIdRegistry,
+			const TSharedPtr<FMaterialIdRegistry>& InMaterialIdRegistry,
+			const TSharedPtr<FSkeletalMeshIdRegistry>& InSkeletalMeshIdRegistry,
+			const TSharedPtr<FPassthroughObjectLoader>& InPassthroughObjectLoader,
+			const FScheduledOpInline& InRoot,
+			uint32 InLodMask) :
 		Model(InModel),
 		Params(InParams),
-		System(InSystem),
 		Program(InModel->GetPrivate()->Program),
 		LodMask(InLodMask),
-		Executed(InModel->GetPrivate()->Program.OpAddress.Num())
+		Executed(InModel->GetPrivate()->Program.OpAddress.Num()),
+		ProgramCache(InProgramCache),
+		MeshIdRegistry(InMeshIdRegistry),
+		ImageIdRegistry(InImageIdRegistry),
+		MaterialIdRegistry(InMaterialIdRegistry),
+		SkeletalMeshIdRegistry(InSkeletalMeshIdRegistry),
+		PassthroughObjectLoader(InPassthroughObjectLoader)
 	{
 	}
 
@@ -83,12 +98,6 @@ namespace UE::Mutable::Private
 	}
 	
 
-	FWorkingMemoryManager& CodeRunnerBegin::GetMemory() const
-	{
-		return System.GetPrivate()->WorkingMemoryManager;
-	}
-
-
 	FScheduledOpInline CodeRunnerBegin::PopOp()
 	{
 		return Items.Pop(EAllowShrinking::No);
@@ -301,59 +310,127 @@ namespace UE::Mutable::Private
 
 			if (Executed.Contains(Item))
 			{
+				bool bExecute = false;
+				
 				switch (GetOpDataType(Type))
 				{
 				case EDataType::Bool:
 					{
-						FStackValue Value;
-						Value.Set<bool>(ResultsBool[Item]);
+						if (Item.bEvaluate)
+						{
+							if (bool* Result = ResultsBool.Find(Item))
+							{
+								FStackValue Value;
+								Value.Set<bool>(*Result);
 					
-						Stack.Push(Value);
+								Stack.Push(Value);
+							}
+							else
+							{
+								bExecute = true;
+							}
+						}
+						
 						break;
 					}
 
 				case EDataType::Int:
 					{
-						FStackValue Value;
-						Value.Set<int32>(ResultsInt[Item]);
+						if (Item.bEvaluate)
+						{
+							if (int32* Result = ResultsInt.Find(Item))
+							{
+								FStackValue Value;
+								Value.Set<int32>(*Result);
 					
-						Stack.Push(Value);
+								Stack.Push(Value);
+							}
+							else
+							{
+								bExecute = true;
+							}
+						}
+						
 						break;
 					}
 					
 				case EDataType::Instance:
 					{
-						FStackValue Value;
-						Value.Set<TManagedPtr<const FInstance>>(ResultsInstance[Item]);
+						if (Item.bEvaluate)
+						{
+							if (TManagedPtr<const FInstance>* Result = ResultsInstance.Find(Item))
+							{
+								FStackValue Value;
+								Value.Set<TManagedPtr<const FInstance>>(*Result);
 					
-						Stack.Push(Value);
+								Stack.Push(Value);
+							}
+							else
+							{
+								bExecute = true;
+							}
+						}
+						
 						break;
 					}
 					
 				case EDataType::String:
 					{
-						FStackValue Value;
-						Value.Set<TManagedPtr<const String>>(ResultsString[Item]);
+						if (Item.bEvaluate)
+						{
+							if (TManagedPtr<const String>* Result = ResultsString.Find(Item))
+							{
+								FStackValue Value;
+								Value.Set<TManagedPtr<const String>>(*Result);
 					
-						Stack.Push(Value);
+								Stack.Push(Value);
+							}
+							else
+							{
+								bExecute = true;
+							}
+						}
+						
 						break;
 					}
 								
 				case EDataType::ExtensionData:
 					{
-						FStackValue Value;
-						Value.Set<TManagedPtr<const FExtensionData>>(ResultsExtensionData[Item]);
+						if (Item.bEvaluate)
+						{
+							if (TManagedPtr<const FExtensionData>* Result = ResultsExtensionData.Find(Item))
+							{
+								FStackValue Value;
+								Value.Set<TManagedPtr<const FExtensionData>>(*Result);
 					
-						Stack.Push(Value);
+								Stack.Push(Value);
+							}
+							else
+							{
+								bExecute = true;
+							}
+						}
+
 						break;
 					}
 
 				case EDataType::Scalar:
 					{
-						FStackValue Value;
-						Value.Set<float>(ResultsFloat[Item]);
+						if (Item.bEvaluate)
+						{
+							if (float* Result = ResultsFloat.Find(Item))
+							{
+								FStackValue Value;
+								Value.Set<float>(*Result);
 					
-						Stack.Push(Value);
+								Stack.Push(Value);
+							}
+							else
+							{
+								bExecute = true;
+							}
+						}
+						
 						break;
 					}
 					
@@ -370,8 +447,11 @@ namespace UE::Mutable::Private
 				default:
 					unimplemented();
 				}
-
-				continue;
+				
+				if (!bExecute)
+				{
+					continue;
+				}
 			}
 			
 			switch (Type)
@@ -390,18 +470,23 @@ namespace UE::Mutable::Private
 			case EOpType::IM_PARAMETER_FROM_MATERIAL:
 			case EOpType::IS_PARAMETER:
 				{
+					check(!Item.bEvaluate);
+
 					StoreNone(Item);
 					break;
 				}
 			
+			case EOpType::IN_CONDITIONAL:
+				check(Item.bEvaluate);
+
 			case EOpType::CO_CONDITIONAL:
 			case EOpType::ED_CONDITIONAL:
 			case EOpType::IM_CONDITIONAL:
-			case EOpType::IN_CONDITIONAL:
 			case EOpType::LA_CONDITIONAL:
 			case EOpType::ME_CONDITIONAL:
 			case EOpType::SC_CONDITIONAL:
 			case EOpType::MI_CONDITIONAL:
+			case EOpType::NU_CONDITIONAL:
 				{
 					OP::ConditionalArgs Args = Program.GetOpArgs<OP::ConditionalArgs>(Item.At);
 
@@ -409,8 +494,8 @@ namespace UE::Mutable::Private
 					{
 					case 0:
 						{
-							PushOp(FScheduledOpInline( Item,1));
-							PushOp(FScheduledOpInline( Args.condition, Item));
+							PushOp(FScheduledOpInline(Item,1));
+							PushOp(FScheduledOpInline(Args.condition, Item, true));
 
 							break;
 						}
@@ -433,77 +518,51 @@ namespace UE::Mutable::Private
 
 					case 2:
 						{
-							switch (GetOpDataType(Type))
+							if (Item.bEvaluate)
 							{
-							case EDataType::Int:
-								StoreInt(Item, LoadInt(Item.CustomState));
-								break;
+								switch (GetOpDataType(Type))
+								{
+								case EDataType::Int:
+									StoreInt(Item, LoadInt(Item.CustomState));
+									break;
 							
-							case EDataType::String:
-								StoreString(Item, LoadString(Item.CustomState));
-								break;
+								case EDataType::String:
+									StoreString(Item, LoadString(Item.CustomState));
+									break;
 
-							case EDataType::Instance:
-								StoreInstance(Item, LoadInstance(Item.CustomState));
-								break;
+								case EDataType::Instance:
+									StoreInstance(Item, LoadInstance(Item.CustomState));
+									break;
 
-							case EDataType::ExtensionData:
-								StoreExtensionData(Item, LoadExtensionData(Item.CustomState));
-								break;
+								case EDataType::ExtensionData:
+									StoreExtensionData(Item, LoadExtensionData(Item.CustomState));
+									break;
 							
-							case EDataType::Scalar:		
-								StoreFloat(Item, LoadFloat(Item.CustomState));
-								break;
+								case EDataType::Scalar:		
+									StoreFloat(Item, LoadFloat(Item.CustomState));
+									break;
 
-							case EDataType::Color:		
-							case EDataType::Projector:	
-							case EDataType::Mesh:		
-							case EDataType::Image:		
-							case EDataType::Layout:		
-							case EDataType::Material:		
-								StoreNone(Item);
-								break;
+								case EDataType::Color:		
+								case EDataType::Projector:	
+								case EDataType::Mesh:		
+								case EDataType::Image:		
+								case EDataType::Layout:		
+								case EDataType::Material:		
+									StoreNone(Item);
+									break;
 
-							case EDataType::Bool:
-								check(false);
+								case EDataType::Bool:
+									check(false);
 								
-							default:
-								unimplemented();
+								default:
+									unimplemented();
+								}
+							}
+							else
+							{
+								StoreNone(Item);
 							}
-
-							break;
-						}
-
-					default:
-						unimplemented();
-					}
-				
-					break;
-				}
-
-			case EOpType::NU_CONDITIONAL:
-				{
-					OP::ConditionalArgs Args = Program.GetOpArgs<OP::ConditionalArgs>(Item.At);
-
-					switch(Item.Stage)
-					{
-					case 0:
-						{
-							PushOp(FScheduledOpInline(Item,1));
-							PushOp(FScheduledOpInline(Args.condition, Item));
-
-							break;
-						}
-
-					case 1:
-						{
-							const bool Value = LoadBool(Args.condition);
 							
-							OP::ADDRESS ResultAt = Value ? Args.yes : Args.no;
-
-							PushOp(FScheduledOpInline(ResultAt, Item));
-
-							StoreNone(Item);
 							break;
 						}
 
@@ -513,7 +572,7 @@ namespace UE::Mutable::Private
 				
 					break;
 				}
-				
+			
 			case EOpType::CO_SWITCH:
 			case EOpType::ED_SWITCH:
 			case EOpType::IM_SWITCH:
@@ -545,47 +604,54 @@ namespace UE::Mutable::Private
 							if (VarAddress)
 							{
 								PushOp(FScheduledOpInline(Item, 1));
-								PushOp(FScheduledOpInline(VarAddress, Item));
+								PushOp(FScheduledOpInline(VarAddress, Item, true));
 							}
 							else
 							{
-								switch (GetOpDataType(Type))
+								if (Item.bEvaluate)
 								{
-								case EDataType::Bool:
-									StoreBool(Item, false);
-									break;
+									switch (GetOpDataType(Type))
+									{
+									case EDataType::Bool:
+										StoreBool(Item, false);
+										break;
 								
-								case EDataType::Int:
-									StoreInt(Item, 0);
-									break;
+									case EDataType::Int:
+										StoreInt(Item, 0);
+										break;
 								
-								case EDataType::Instance:
-									StoreInstance(Item, nullptr);
-									break;
+									case EDataType::Instance:
+										StoreInstance(Item, nullptr);
+										break;
 
-								case EDataType::String:
-									StoreString(Item, nullptr);
-									break;
+									case EDataType::String:
+										StoreString(Item, nullptr);
+										break;
 								
-								case EDataType::ExtensionData:
-									StoreExtensionData(Item, MakeManaged<FExtensionData>());
-									break;
+									case EDataType::ExtensionData:
+										StoreExtensionData(Item, MakeManaged<FExtensionData>());
+										break;
 
-								case EDataType::Scalar:
-									StoreFloat(Item, 0.0f);
-									break;
+									case EDataType::Scalar:
+										StoreFloat(Item, 0.0f);
+										break;
 									
-								case EDataType::Color:			
-								case EDataType::Projector:		
-								case EDataType::Mesh:			
-								case EDataType::Image:			
-								case EDataType::Layout:			
-								case EDataType::Material:			
-									StoreNone(Item);
-									break;
+									case EDataType::Color:			
+									case EDataType::Projector:		
+									case EDataType::Mesh:			
+									case EDataType::Image:			
+									case EDataType::Layout:			
+									case EDataType::Material:			
+										StoreNone(Item);
+										break;
 								
-								default:
-									unimplemented()
+									default:
+										unimplemented()
+									}
+								}
+								else
+								{
+									StoreNone(Item);
 								}
 							}
 							break;
@@ -653,43 +719,50 @@ namespace UE::Mutable::Private
 						
 					case 2:
 						{
-							switch (GetOpDataType(Type))
+							if (Item.bEvaluate)
 							{
-							case EDataType::Bool:
-								StoreBool(Item, LoadBool(Item.CustomState));
-								break;
+								switch (GetOpDataType(Type))
+								{
+								case EDataType::Bool:
+									StoreBool(Item, LoadBool(Item.CustomState));
+									break;
 
-							case EDataType::Int:
-								StoreInt(Item, LoadInt(Item.CustomState));
-								break;
+								case EDataType::Int:
+									StoreInt(Item, LoadInt(Item.CustomState));
+									break;
 
-							case EDataType::Instance:
-								StoreInstance(Item, LoadInstance(Item.CustomState));
-								break;
+								case EDataType::Instance:
+									StoreInstance(Item, LoadInstance(Item.CustomState));
+									break;
 								
-							case EDataType::String:
-								StoreString(Item, LoadString(Item.CustomState));
-								break;
+								case EDataType::String:
+									StoreString(Item, LoadString(Item.CustomState));
+									break;
 								
-							case EDataType::ExtensionData:
-								StoreExtensionData(Item, LoadExtensionData(Item.CustomState));
-								break;
+								case EDataType::ExtensionData:
+									StoreExtensionData(Item, LoadExtensionData(Item.CustomState));
+									break;
 
-							case EDataType::Scalar:
-								StoreFloat(Item, LoadFloat(Item.CustomState));
-								break;
+								case EDataType::Scalar:
+									StoreFloat(Item, LoadFloat(Item.CustomState));
+									break;
 								
-							case EDataType::Color:		
-							case EDataType::Projector:	
-							case EDataType::Mesh:		
-							case EDataType::Image:		
-							case EDataType::Layout:		
-							case EDataType::Material:		
-								StoreNone(Item);
-								break;
+								case EDataType::Color:		
+								case EDataType::Projector:	
+								case EDataType::Mesh:		
+								case EDataType::Image:		
+								case EDataType::Layout:		
+								case EDataType::Material:		
+									StoreNone(Item);
+									break;
 
-							default:
-								unimplemented();
+								default:
+									unimplemented();
+								}
+							}
+							else
+							{
+								StoreNone(Item);
 							}
 
 							break;
@@ -704,6 +777,8 @@ namespace UE::Mutable::Private
 
 			case EOpType::IN_ADDMESH:
 				{
+					check(Item.bEvaluate);
+
 					OP::InstanceAddArgs Args = Program.GetOpArgs<OP::InstanceAddArgs>(Item.At);
 					switch (Item.Stage)
 					{
@@ -711,7 +786,7 @@ namespace UE::Mutable::Private
 						{
 							PushOp(FScheduledOpInline(Item, 1));
 							PushOp(FScheduledOpInline(Args.instance, Item));
-							PushOp(FScheduledOpInline(Args.value, Item));
+							PushOp(FScheduledOpInline(Args.value, Item, false));
 							
 							break;
 						}
@@ -736,7 +811,7 @@ namespace UE::Mutable::Private
 							
 							if (Args.value)
 							{
-								FMeshId MeshId = FSystem::GetMeshId(System.GetPrivate()->WorkingMemoryManager.CurrentMeshIdRegistry, Model, Params.Get(), Args.value);
+								FMeshId MeshId = FSystem::GetMeshId(MeshIdRegistry, Model, Params.Get(), Args.value);
 								OP::ADDRESS NameAd = Args.name;
 								check(NameAd < (uint32)Program.ConstantStrings.Num());
 								const FString& Name = Program.ConstantStrings[NameAd];
@@ -756,6 +831,8 @@ namespace UE::Mutable::Private
 
 			case EOpType::IN_ADDMATERIAL:
 				{
+					check(Item.bEvaluate);
+
 					OP::InstanceAddMaterialArgs Args = Program.GetOpArgs<OP::InstanceAddMaterialArgs>(Item.At);
 					switch (Item.Stage)
 					{
@@ -763,7 +840,7 @@ namespace UE::Mutable::Private
 						{
 							PushOp(FScheduledOpInline(Item, 1));
 							PushOp(FScheduledOpInline(Args.Instance, Item));
-							PushOp(FScheduledOpInline(Args.Material, Item));
+							PushOp(FScheduledOpInline(Args.Material, Item, false));
 							
 							break;
 						}
@@ -789,7 +866,7 @@ namespace UE::Mutable::Private
 
 							if (Args.Material)
 							{
-        						FMaterialId MaterialId = FSystem::GetMaterialId(System.GetPrivate()->WorkingMemoryManager.CurrentMaterialIdRegistry, Model, Params.Get(), Args.Material);
+        						FMaterialId MaterialId = FSystem::GetMaterialId(MaterialIdRegistry, Model, Params.Get(), Args.Material);
 								Instance->SetMaterialId(0, 0, 0, MaterialId);
 							}
 							
@@ -806,6 +883,8 @@ namespace UE::Mutable::Private
 					
 			case EOpType::IN_ADDOVERLAYMATERIAL:
 				{
+					check(Item.bEvaluate);
+
 					OP::InstanceAddMaterialArgs Args = Program.GetOpArgs<OP::InstanceAddMaterialArgs>(Item.At);
 					switch (Item.Stage)
 					{
@@ -813,7 +892,7 @@ namespace UE::Mutable::Private
 						{
 							PushOp(FScheduledOpInline(Item, 1));
 							PushOp(FScheduledOpInline(Args.Instance, Item));
-							PushOp(FScheduledOpInline(Args.Material, Item));
+							PushOp(FScheduledOpInline(Args.Material, Item, false));
 
 							break;
 						}
@@ -838,7 +917,7 @@ namespace UE::Mutable::Private
 
 							if (Args.Material)
 							{
-        						FMaterialId MaterialId = FSystem::GetMaterialId(System.GetPrivate()->WorkingMemoryManager.CurrentMaterialIdRegistry, Model, Params.Get(), Args.Material);
+        						FMaterialId MaterialId = FSystem::GetMaterialId(MaterialIdRegistry, Model, Params.Get(), Args.Material);
 								Instance->SetOverlayMaterialId(0, MaterialId);
 							}
 							
@@ -887,7 +966,7 @@ namespace UE::Mutable::Private
 
 							if (Args.Material)
 							{
-								FMaterialId MaterialId = FSystem::GetMaterialId(System.GetPrivate()->WorkingMemoryManager.CurrentMaterialIdRegistry, Model, Params.Get(), Args.Material);
+								FMaterialId MaterialId = FSystem::GetMaterialId(MaterialIdRegistry, Model, Params.Get(), Args.Material);
 								const FString& SlotName = Program.ConstantStrings[Args.SlotName];
 								Instance->AddOverrideMaterial(0, FName(SlotName), MaterialId);
 							}
@@ -905,6 +984,8 @@ namespace UE::Mutable::Private
 				
 			case EOpType::IN_ADDCOMPONENT:
 				{
+					check(Item.bEvaluate);
+
 					OP::InstanceAddArgs Args = Program.GetOpArgs<OP::InstanceAddArgs>(Item.At);
 					switch (Item.Stage)
 					{
@@ -963,6 +1044,8 @@ namespace UE::Mutable::Private
 				
 			case EOpType::IN_ADDSTRING:
 				{
+					check(Item.bEvaluate);
+
 					OP::InstanceAddArgs Args = Program.GetOpArgs<OP::InstanceAddArgs>( Item.At );
 					switch (Item.Stage)
 					{
@@ -970,7 +1053,7 @@ namespace UE::Mutable::Private
 						{
 							PushOp(FScheduledOpInline(Item, 1));
 							PushOp(FScheduledOpInline(Args.instance, Item));
-							PushOp(FScheduledOpInline(Args.value, Item));
+							PushOp(FScheduledOpInline(Args.value, Item, false));
 
 							break;
 						}
@@ -1006,6 +1089,8 @@ namespace UE::Mutable::Private
 				
 			case EOpType::IN_ADDSURFACE:
 				{
+					check(Item.bEvaluate);
+
 					OP::InstanceAddArgs Args = Program.GetOpArgs<OP::InstanceAddArgs>(Item.At);
 		            switch (Item.Stage)
 		            {
@@ -1078,6 +1163,8 @@ namespace UE::Mutable::Private
 
 			case EOpType::IN_ADDLOD:
 				{
+					check(Item.bEvaluate);
+
 					const uint8* Data = Program.GetOpArgsPointer(Item.At);
 
 					uint8 LODCount;
@@ -1156,6 +1243,8 @@ namespace UE::Mutable::Private
 				
 			case EOpType::IN_ADDSKELETALMESH:
 				{
+					check(Item.bEvaluate);
+
 					OP::InstanceAddArgs Args = Program.GetOpArgs<OP::InstanceAddArgs>(Item.At);
 					switch (Item.Stage)
 					{
@@ -1163,7 +1252,7 @@ namespace UE::Mutable::Private
 						{
 							PushOp(FScheduledOpInline(Item, 1));
 							PushOp(FScheduledOpInline(Args.instance, Item));
-							PushOp(FScheduledOpInline(Args.value, Item));
+							PushOp(FScheduledOpInline(Args.value, Item, false));
 							
 							break;
 						}
@@ -1188,7 +1277,7 @@ namespace UE::Mutable::Private
 							
 							if (Args.value)
 							{
-        						FSkeletalMeshId SkeletalMeshId = FSystem::GetSkeletalMeshId(System.GetPrivate()->WorkingMemoryManager.CurrentSkeletalMeshIdRegistry, Model, Params.Get(), Args.value);
+        						FSkeletalMeshId SkeletalMeshId = FSystem::GetSkeletalMeshId(SkeletalMeshIdRegistry, Model, Params.Get(), Args.value);
 								Instance->SetSkeletalMeshId(0, SkeletalMeshId);
 							}
 							
@@ -1205,6 +1294,8 @@ namespace UE::Mutable::Private
 				
 			case EOpType::IN_ADDEXTENSIONDATA:
 				{
+					check(Item.bEvaluate);
+
 					OP::InstanceAddExtensionDataArgs Args = Program.GetOpArgs<OP::InstanceAddExtensionDataArgs>(Item.At);
 
 					switch (Item.Stage)
@@ -1258,9 +1349,17 @@ namespace UE::Mutable::Private
 				
 			case EOpType::BO_CONSTANT:
 				{
-					OP::BoolConstantArgs Args = Program.GetOpArgs<OP::BoolConstantArgs>(Item.At);
+					if (Item.bEvaluate)
+					{
+						OP::BoolConstantArgs Args = Program.GetOpArgs<OP::BoolConstantArgs>(Item.At);
 
-					StoreBool(Item, Args.bValue);
+						StoreBool(Item, Args.bValue);
+					}
+					else
+					{
+						StoreNone(Item);	
+					}
+					
 					break;
 				}
 
@@ -1268,14 +1367,22 @@ namespace UE::Mutable::Private
 				{
 					OP::ParameterArgs Args = Program.GetOpArgs<OP::ParameterArgs>(Item.At);
 
-					FScheduledOp Op;
-					Op.At = Item.At;
-					Op.ExecutionIndex = Item.ExecutionIndex;
+					if (Item.bEvaluate)
+					{
+						FScheduledOp Op;
+						Op.At = Item.At;
+						Op.ExecutionIndex = Item.ExecutionIndex;
         		
-					TSharedPtr<FRangeIndex> Index = BuildCurrentOpRangeIndex(Op, *Params, *Model.Get(), *GetMemory().CurrentInstanceCache, Args.variable);
-					const bool Value = Params->GetBoolValue(Args.variable, Index.Get());
+						TSharedPtr<FRangeIndex> Index = BuildCurrentOpRangeIndex(Op, *Params, *Model.Get(), *ProgramCache, Args.variable);
+						const bool Value = Params->GetBoolValue(Args.variable, Index.Get());
+					
+						StoreBool(Item, Value);
+					}
+					else
+					{
+						StoreNone(Item);	
+					}
 					
-					StoreBool(Item, Value);
 					break;
 				}
 
@@ -1283,30 +1390,39 @@ namespace UE::Mutable::Private
 				{
 					OP::BoolBinaryArgs Args = Program.GetOpArgs<OP::BoolBinaryArgs>(Item.At);
 
-					switch (Item.Stage)
+					if (Item.bEvaluate)
 					{
-					case 0:
+						switch (Item.Stage)
 						{
-							PushOp(FScheduledOpInline(Item, 1));
-							PushOp(FScheduledOpInline(Args.A, Item));
-							PushOp(FScheduledOpInline(Args.B, Item));
+						case 0:
+							{
+								PushOp(FScheduledOpInline(Item, 1));
+								PushOp(FScheduledOpInline(Args.A, Item));
+								PushOp(FScheduledOpInline(Args.B, Item));
 
-							break;
-						}
+								break;
+							}
 
-					case 1:
-						{
-							const bool ValueB = LoadBool(Args.A);
-							const bool ValueA = LoadBool(Args.B);
+						case 1:
+							{
+								const bool ValueB = LoadBool(Args.A);
+								const bool ValueA = LoadBool(Args.B);
 
-							const bool Result = ValueA && ValueB;
+								const bool Result = ValueA && ValueB;
 
-							StoreBool(Item, Result);
-							break;
-						}
+								StoreBool(Item, Result);
+								break;
+							}
 
-					default:
-						unimplemented();
+						default:
+							unimplemented();
+						}
+					}
+					else
+					{
+						PushOp(FScheduledOpInline(Args.A, Item));
+						PushOp(FScheduledOpInline(Args.B, Item));
+						StoreNone(Item);	
 					}
 					
 					break;
@@ -1316,30 +1432,39 @@ namespace UE::Mutable::Private
 				{
 					OP::BoolBinaryArgs Args = Program.GetOpArgs<OP::BoolBinaryArgs>(Item.At);
 
-					switch (Item.Stage)
+					if (Item.bEvaluate)
 					{
-					case 0:
+						switch (Item.Stage)
 						{
-							PushOp(FScheduledOpInline(Item, 1));
-							PushOp(FScheduledOpInline(Args.A, Item));
-							PushOp(FScheduledOpInline(Args.B, Item));
+						case 0:
+							{
+								PushOp(FScheduledOpInline(Item, 1));
+								PushOp(FScheduledOpInline(Args.A, Item));
+								PushOp(FScheduledOpInline(Args.B, Item));
 
-							break;
-						}
+								break;
+							}
 
-					case 1:
-						{
-							const bool ValueB = LoadBool(Args.A);
-							const bool ValueA = LoadBool(Args.B);
+						case 1:
+							{
+								const bool ValueB = LoadBool(Args.A);
+								const bool ValueA = LoadBool(Args.B);
 
-							const bool Result = ValueA || ValueB;
+								const bool Result = ValueA || ValueB;
 							
-							StoreBool(Item, Result);
-							break;
-						}
+								StoreBool(Item, Result);
+								break;
+							}
 
-					default:
-						unimplemented();
+						default:
+							unimplemented();
+						}	
+					}
+					else
+					{
+						PushOp(FScheduledOpInline(Args.A, Item));
+						PushOp(FScheduledOpInline(Args.B, Item));
+						StoreNone(Item);	
 					}
 					
 					break;
@@ -1349,28 +1474,36 @@ namespace UE::Mutable::Private
 				{
 					OP::BoolNotArgs Args = Program.GetOpArgs<OP::BoolNotArgs>(Item.At);
 
-					switch ( Item.Stage )
+					if (Item.bEvaluate)
 					{
-					case 0:
+						switch (Item.Stage)
 						{
-							PushOp(FScheduledOpInline(Item, 1));
-							PushOp(FScheduledOpInline(Args.A, Item));
+						case 0:
+							{
+								PushOp(FScheduledOpInline(Item, 1));
+								PushOp(FScheduledOpInline(Args.A, Item));
 
-							break;
-						}
+								break;
+							}
 					
-					case 1:
-						{
-							const bool Value = LoadBool(Args.A);
+						case 1:
+							{
+								const bool Value = LoadBool(Args.A);
 							
-							const bool Result = !Value;
+								const bool Result = !Value;
 							
-							StoreBool(Item, Result);
-							break;
-						}
+								StoreBool(Item, Result);
+								break;
+							}
 
-					default:
-						unimplemented();
+						default:
+							unimplemented();
+						}
+					}
+					else
+					{
+						PushOp(FScheduledOpInline(Args.A, Item));
+						StoreNone(Item);	
 					}
 					
 					break;
@@ -1380,28 +1513,36 @@ namespace UE::Mutable::Private
 				{
 					OP::BoolEqualScalarConstArgs Args = Program.GetOpArgs<OP::BoolEqualScalarConstArgs>(Item.At);
 
-					switch ( Item.Stage )
+					if (Item.bEvaluate)
 					{
-					case 0:
+						switch (Item.Stage)
 						{
-							PushOp(FScheduledOpInline(Item, 1));
-							PushOp(FScheduledOpInline(Args.Value, Item));
+						case 0:
+							{
+								PushOp(FScheduledOpInline(Item, 1));
+								PushOp(FScheduledOpInline(Args.Value, Item));
 
-							break;
-						}
+								break;
+							}
 					
-					case 1:
-						{
-							const int32 Value = LoadInt(Args.Value);
+						case 1:
+							{
+								const int32 Value = LoadInt(Args.Value);
 							
-							const bool Result = Value == Args.Constant;
+								const bool Result = Value == Args.Constant;
 
-							StoreBool(Item, Result);
-							break;
-						}
+								StoreBool(Item, Result);
+								break;
+							}
 
-					default:
-						unimplemented();
+						default:
+							unimplemented();
+						}
+					}
+					else
+					{
+						PushOp(FScheduledOpInline(Args.Value, Item));
+						StoreNone(Item);	
 					}
 					
 					break;
@@ -1409,9 +1550,17 @@ namespace UE::Mutable::Private
 
 			case EOpType::SC_CONSTANT:
 				{
-					OP::ScalarConstantArgs Args = Program.GetOpArgs<OP::ScalarConstantArgs>(Item.At);
+					if (Item.bEvaluate)
+					{
+						OP::ScalarConstantArgs Args = Program.GetOpArgs<OP::ScalarConstantArgs>(Item.At);
+						
+						StoreFloat(Item, Args.Value);
+					}
+					else
+					{
+						StoreNone(Item);
+					}
 
-					StoreFloat(Item, Args.Value);
 					break;
 				}
 			
@@ -1420,65 +1569,83 @@ namespace UE::Mutable::Private
 				{
 					OP::ParameterArgs Args = Program.GetOpArgs<OP::ParameterArgs>(Item.At);
 
-					FScheduledOp Op;
-					Op.At = Item.At;
-					Op.ExecutionIndex = Item.ExecutionIndex;
+					if (Item.bEvaluate)
+					{
+						FScheduledOp Op;
+						Op.At = Item.At;
+						Op.ExecutionIndex = Item.ExecutionIndex;
 					
-					TSharedPtr<FRangeIndex> Index = BuildCurrentOpRangeIndex(Op, *Params, *Model.Get(), *GetMemory().CurrentInstanceCache, Args.variable);
-					const float Result = Params->GetFloatValue( Args.variable, Index.Get());
+						TSharedPtr<FRangeIndex> Index = BuildCurrentOpRangeIndex(Op, *Params, *Model.Get(), *ProgramCache, Args.variable);
+						const float Result = Params->GetFloatValue( Args.variable, Index.Get());
 
-					StoreFloat(Item, Result);
+						StoreFloat(Item, Result);
+					}
+					else
+					{
+						StoreNone(Item);
+					}
+					
 					break;
 				}
 				
 			case EOpType::SC_ARITHMETIC:
 				{
 					OP::ArithmeticArgs Args = Program.GetOpArgs<OP::ArithmeticArgs>(Item.At);
-					switch (Item.Stage)
+
+					if (Item.bEvaluate)
 					{
-					case 0:
+						switch (Item.Stage)
 						{
-							PushOp(FScheduledOpInline(Item, 1));
-							PushOp(FScheduledOpInline(Args.A, Item));
-							PushOp(FScheduledOpInline(Args.B, Item));
-
-							break;
-						}
+						case 0:
+							{
+								PushOp(FScheduledOpInline(Item, 1));
+								PushOp(FScheduledOpInline(Args.A, Item));
+								PushOp(FScheduledOpInline(Args.B, Item));
 
-					case 1:
-						{
-							float ValueA = LoadFloat(Args.A);
-							float ValueB = LoadFloat(Args.B);
+								break;
+							}
 
-							float Result = 1.0f;
-							switch (Args.Operation)
+						case 1:
 							{
-							case OP::ArithmeticArgs::ADD:
-								Result = ValueA + ValueB;
-								break;
+								float ValueA = LoadFloat(Args.A);
+								float ValueB = LoadFloat(Args.B);
 
-							case OP::ArithmeticArgs::MULTIPLY:
-								Result = ValueA * ValueB;
-								break;
+								float Result = 1.0f;
+								switch (Args.Operation)
+								{
+								case OP::ArithmeticArgs::ADD:
+									Result = ValueA + ValueB;
+									break;
 
-							case OP::ArithmeticArgs::SUBTRACT:
-								Result = ValueA - ValueB;
-								break;
+								case OP::ArithmeticArgs::MULTIPLY:
+									Result = ValueA * ValueB;
+									break;
 
-							case OP::ArithmeticArgs::DIVIDE:
-								Result = ValueA / ValueB;
-								break;
+								case OP::ArithmeticArgs::SUBTRACT:
+									Result = ValueA - ValueB;
+									break;
+
+								case OP::ArithmeticArgs::DIVIDE:
+									Result = ValueA / ValueB;
+									break;
+
+								default:
+									unimplemented();
+								}
 
-							default:
-								unimplemented();
+								StoreFloat(Item, Result);
+								break;
 							}
 
-							StoreFloat(Item, Result);
-							break;
+						default:
+							unimplemented();
 						}
-
-					default:
-						unimplemented();
+					}
+					else
+					{
+						PushOp(FScheduledOpInline(Args.A, Item));
+						PushOp(FScheduledOpInline(Args.B, Item));
+						StoreNone(Item);
 					}
 					
 					break;
@@ -1487,29 +1654,38 @@ namespace UE::Mutable::Private
 			case EOpType::SC_CURVE:
 				{
 					OP::ScalarCurveArgs Args = Program.GetOpArgs<OP::ScalarCurveArgs>(Item.At);
-					switch (Item.Stage)
+
+					if (Item.bEvaluate)
 					{
-					case 0:
+						switch (Item.Stage)
 						{
-							PushOp(FScheduledOpInline(Item, 1));
-							PushOp(FScheduledOpInline(Args.time, Item));
+						case 0:
+							{
+								PushOp(FScheduledOpInline(Item, 1));
+								PushOp(FScheduledOpInline(Args.time, Item));
 
-							break;
-						}
+								break;
+							}
 
-					case 1:
-						{
-							const float Time = LoadFloat(Args.time);
+						case 1:
+							{
+								const float Time = LoadFloat(Args.time);
 
-							const FRichCurve& Curve = Program.ConstantCurves[Args.curve];
-							float Result = Curve.Eval(Time);
+								const FRichCurve& Curve = Program.ConstantCurves[Args.curve];
+								float Result = Curve.Eval(Time);
 
-							StoreFloat(Item, Result);
-							break;
-						}
+								StoreFloat(Item, Result);
+								break;
+							}
 
-					default:
-						unimplemented();
+						default:
+							unimplemented();
+						}	
+					}
+					else
+					{
+						PushOp(FScheduledOpInline(Args.time, Item));
+						StoreNone(Item);
 					}
 					
 					break;
@@ -1517,6 +1693,8 @@ namespace UE::Mutable::Private
 				
 			case EOpType::ED_CONSTANT:
 				{
+					check(!Item.bEvaluate);
+
 					OP::ResourceConstantArgs Args = Program.GetOpArgs<OP::ResourceConstantArgs>(Item.At);
 
 					// Assume the ROM has been loaded previously
@@ -1531,15 +1709,23 @@ namespace UE::Mutable::Private
 
 			case EOpType::ST_CONSTANT:
 				{
-					OP::ResourceConstantArgs Args = Program.GetOpArgs<OP::ResourceConstantArgs>(Item.At);
+					if (Item.bEvaluate)
+					{
+						OP::ResourceConstantArgs Args = Program.GetOpArgs<OP::ResourceConstantArgs>(Item.At);
 				
-					check(Args.value < static_cast<uint32>(Model->GetPrivate()->Program.ConstantStrings.Num()));
+						check(Args.value < static_cast<uint32>(Model->GetPrivate()->Program.ConstantStrings.Num()));
 
-					const FString& Result = Program.ConstantStrings[Args.value];
+						const FString& Result = Program.ConstantStrings[Args.value];
 
-					TManagedPtr<const String> Value = MakeManaged<String>(Result);
+						TManagedPtr<const String> Value = MakeManaged<String>(Result);
 
-					StoreString(Item, Value);
+						StoreString(Item, Value);
+					}
+					else
+					{
+						StoreNone(Item);
+					}
+					
 					break;
 				}
 
@@ -1547,28 +1733,44 @@ namespace UE::Mutable::Private
 				{
 					OP::ParameterArgs Args = Program.GetOpArgs<OP::ParameterArgs>(Item.At);
 
-					FScheduledOp Op;
-					Op.At = Item.At;
-					Op.ExecutionIndex = Item.ExecutionIndex;
+					if (Item.bEvaluate)
+					{
+						FScheduledOp Op;
+						Op.At = Item.At;
+						Op.ExecutionIndex = Item.ExecutionIndex;
 				
-					TSharedPtr<FRangeIndex> Index = BuildCurrentOpRangeIndex(Op, *Params, *Model.Get(), *GetMemory().CurrentInstanceCache, Args.variable);
+						TSharedPtr<FRangeIndex> Index = BuildCurrentOpRangeIndex(Op, *Params, *Model.Get(), *ProgramCache, Args.variable);
 
-					FString Result;
-					Params->GetStringValue(Args.variable, Result, Index.Get());
+						FString Result;
+						Params->GetStringValue(Args.variable, Result, Index.Get());
 
-					TManagedPtr<const String> Value = MakeManaged<String>(Result);
+						TManagedPtr<const String> Value = MakeManaged<String>(Result);
+					
+						StoreString(Item, Value);
+					}
+					else
+					{
+						StoreNone(Item);
+					}
 					
-					StoreString(Item, Value);
 					break;
 				}
 
 			case EOpType::NU_CONSTANT:
 				{
-					OP::IntConstantArgs Args = Program.GetOpArgs<OP::IntConstantArgs>(Item.At);
+					if (Item.bEvaluate)
+					{
+						OP::IntConstantArgs Args = Program.GetOpArgs<OP::IntConstantArgs>(Item.At);
 
-					const int32 Value = Args.Value;
+						const int32 Value = Args.Value;
+					
+						StoreInt(Item, Value);
+					}
+					else
+					{
+						StoreNone(Item);
+					}
 					
-					StoreInt(Item, Value);
 					break;
 				}
 
@@ -1576,14 +1778,24 @@ namespace UE::Mutable::Private
 				{
 					OP::ParameterArgs Args = Program.GetOpArgs<OP::ParameterArgs>(Item.At);
 
-					const int32 Value = Params->GetIntValue(Args.variable);
+					if (Item.bEvaluate)
+					{
+						const int32 Value = Params->GetIntValue(Args.variable);
 
-					StoreInt(Item, Value);
+						StoreInt(Item, Value);
+					}
+					else
+					{
+						StoreNone(Item);
+					}
+					
 					break;
 				}
 			
 			case EOpType::CO_ARITHMETIC:
 				{
+					check(!Item.bEvaluate);
+
 					OP::ArithmeticArgs Args = Program.GetOpArgs<OP::ArithmeticArgs>(Item.At);
 					PushOp(FScheduledOpInline(Args.A, Item));
 					PushOp(FScheduledOpInline(Args.B, Item));
@@ -1594,12 +1806,16 @@ namespace UE::Mutable::Private
 	
 			case EOpType::CO_FROMSCALARS:
 				{
+					check(!Item.bEvaluate);
+
 					StoreNone(Item);
 					break;
 				}
 				
 			case EOpType::CO_SAMPLEIMAGE:
 				{
+					check(!Item.bEvaluate);
+
 					OP::ColourSampleImageArgs Args = Program.GetOpArgs<OP::ColourSampleImageArgs>(Item.At);
 					PushOp(FScheduledOpInline(Args.Image, Item));
 
@@ -1609,6 +1825,8 @@ namespace UE::Mutable::Private
 	
 			case EOpType::CO_SWIZZLE:
 				{
+					check(!Item.bEvaluate);
+
 					OP::ColourSwizzleArgs Args = Program.GetOpArgs<OP::ColourSwizzleArgs>(Item.At);
 					PushOp(FScheduledOpInline(Args.sources[0], Item));
 					PushOp(FScheduledOpInline(Args.sources[1], Item));
@@ -1621,6 +1839,8 @@ namespace UE::Mutable::Private
 
 			case EOpType::ME_CONSTANT:
 				{
+					check(!Item.bEvaluate);
+
 					OP::MeshConstantArgs Args = Program.GetOpArgs<OP::MeshConstantArgs>(Item.At);
 
 					StoreNone(Item);
@@ -1629,8 +1849,11 @@ namespace UE::Mutable::Private
 			
 			case EOpType::ME_APPLYLAYOUT:
 				{
+					check(!Item.bEvaluate);
+
 					OP::MeshApplyLayoutArgs Args = Program.GetOpArgs<OP::MeshApplyLayoutArgs>(Item.At);
 					PushOp(FScheduledOpInline(Args.Mesh, Item));
+					//PushOp(FScheduledOpInline(Args.Layout, Item));
 
 					StoreNone(Item);
 					break;
@@ -1639,8 +1862,11 @@ namespace UE::Mutable::Private
 	
 			case EOpType::ME_PREPARELAYOUT:
 				{
+					check(!Item.bEvaluate);
+
 					OP::MeshPrepareLayoutArgs Args = Program.GetOpArgs<OP::MeshPrepareLayoutArgs>(Item.At);
 					PushOp(FScheduledOpInline(Args.Mesh, Item));
+					//PushOp(FScheduledOpInline(Args.Layout, Item));
 
 					StoreNone(Item);
 					break;
@@ -1648,6 +1874,8 @@ namespace UE::Mutable::Private
 	
 			case EOpType::ME_DIFFERENCE:
 				{
+					check(!Item.bEvaluate);
+
 					const uint8* data = Program.GetOpArgsPointer(Item.At);
 
 					OP::ADDRESS BaseAt = 0;
@@ -1667,6 +1895,8 @@ namespace UE::Mutable::Private
 	
 			case EOpType::ME_MORPH:
 				{
+					check(!Item.bEvaluate);
+
 					const uint8* Data = Program.GetOpArgsPointer(Item.At);
 
 					OP::CONSTANT_STRING_INDEX NameAt = 0;
@@ -1685,6 +1915,7 @@ namespace UE::Mutable::Private
 					FMemory::Memcpy(&TargetAt, Data, sizeof(OP::ADDRESS)); 
 					Data += sizeof(OP::ADDRESS);
 
+					//PushOp(FScheduledOpInline(FactorAt, Item));
 					PushOp(FScheduledOpInline(BaseAt, Item));
 					PushOp(FScheduledOpInline(TargetAt, Item));
 
@@ -1694,6 +1925,8 @@ namespace UE::Mutable::Private
 	
 			case EOpType::ME_MERGE:
 				{
+					check(!Item.bEvaluate);
+
 					OP::MeshMergeArgs Args = Program.GetOpArgs<OP::MeshMergeArgs>(Item.At);
 					PushOp(FScheduledOpInline(Args.Base, Item));
 					PushOp(FScheduledOpInline(Args.Added, Item));
@@ -1704,6 +1937,8 @@ namespace UE::Mutable::Private
 	
 			case EOpType::ME_MASKCLIPMESH:
 				{
+					check(!Item.bEvaluate);
+
 					OP::MeshMaskClipMeshArgs Args = Program.GetOpArgs<OP::MeshMaskClipMeshArgs>(Item.At);
 					PushOp(FScheduledOpInline(Args.source, Item));
 					PushOp(FScheduledOpInline(Args.clip, Item));
@@ -1714,6 +1949,8 @@ namespace UE::Mutable::Private
 	
 			case EOpType::ME_MASKCLIPUVMASK:
 				{
+					check(!Item.bEvaluate);
+
 					OP::MeshMaskClipUVMaskArgs Args = Program.GetOpArgs<OP::MeshMaskClipUVMaskArgs>(Item.At);
 					PushOp(FScheduledOpInline(Args.Source, Item));
 					PushOp(FScheduledOpInline(Args.UVSource, Item));
@@ -1726,6 +1963,8 @@ namespace UE::Mutable::Private
 	
 			case EOpType::ME_MASKDIFF:
 				{
+					check(!Item.bEvaluate);
+
 					OP::MeshMaskDiffArgs Args = Program.GetOpArgs<OP::MeshMaskDiffArgs>(Item.At);
 					PushOp(FScheduledOpInline(Args.Source, Item));
 					PushOp(FScheduledOpInline(Args.Fragment, Item));
@@ -1736,6 +1975,8 @@ namespace UE::Mutable::Private
 
 			case EOpType::ME_FORMAT:
 				{
+					check(!Item.bEvaluate);
+
 					OP::MeshFormatArgs Args = Program.GetOpArgs<OP::MeshFormatArgs>(Item.At);
 					PushOp(FScheduledOpInline(Args.source, Item));
 					PushOp(FScheduledOpInline(Args.format, Item));
@@ -1746,6 +1987,8 @@ namespace UE::Mutable::Private
 	
 			case EOpType::ME_EXTRACTLAYOUTBLOCK:
 				{
+					check(!Item.bEvaluate);
+
 					const uint8* Data = Program.GetOpArgsPointer(Item.At);
 
 					OP::ADDRESS Source;
@@ -1760,8 +2003,11 @@ namespace UE::Mutable::Private
 	
 			case EOpType::ME_TRANSFORM:
 				{
+					check(!Item.bEvaluate);
+
 					OP::MeshTransformArgs Args = Program.GetOpArgs<OP::MeshTransformArgs>(Item.At);
 					PushOp(FScheduledOpInline(Args.source, Item));
+					//PushOp(FScheduledOpInline(Args.matrix, Item));
 
 					StoreNone(Item);
 					break;
@@ -1769,6 +2015,8 @@ namespace UE::Mutable::Private
 	
 			case EOpType::ME_CLIPMORPHPLANE:
 				{
+					check(!Item.bEvaluate);
+
 					OP::MeshClipMorphPlaneArgs Args = Program.GetOpArgs<OP::MeshClipMorphPlaneArgs>(Item.At);
 					PushOp(FScheduledOpInline(Args.Source, Item));
 
@@ -1778,6 +2026,8 @@ namespace UE::Mutable::Private
 			
 			case EOpType::ME_CLIPWITHMESH:
 				{
+					check(!Item.bEvaluate);
+
 					OP::MeshClipWithMeshArgs Args = Program.GetOpArgs<OP::MeshClipWithMeshArgs>(Item.At);
 					PushOp(FScheduledOpInline(Args.Source, Item));
 					PushOp(FScheduledOpInline(Args.ClipMesh, Item));
@@ -1788,6 +2038,8 @@ namespace UE::Mutable::Private
 				
 			case EOpType::ME_CLIPDEFORM:
 				{
+					check(!Item.bEvaluate);
+
 					OP::MeshClipDeformArgs Args = Program.GetOpArgs<OP::MeshClipDeformArgs>(Item.At);
 					PushOp(FScheduledOpInline(Args.mesh, Item));
 					PushOp(FScheduledOpInline(Args.clipShape, Item));
@@ -1798,6 +2050,8 @@ namespace UE::Mutable::Private
 
 			case EOpType::ME_APPLYPOSE:
 				{
+					check(!Item.bEvaluate);
+
 					OP::MeshApplyPoseArgs Args = Program.GetOpArgs<OP::MeshApplyPoseArgs>(Item.At);
 					PushOp(FScheduledOpInline(Args.base, Item));
 					PushOp(FScheduledOpInline(Args.pose, Item));
@@ -1808,6 +2062,8 @@ namespace UE::Mutable::Private
 	
 			case EOpType::ME_BINDSHAPE:
 				{
+					check(!Item.bEvaluate);
+
 					OP::MeshBindShapeArgs Args = Program.GetOpArgs<OP::MeshBindShapeArgs>(Item.At);
 					PushOp(FScheduledOpInline(Args.mesh, Item));
 					PushOp(FScheduledOpInline(Args.shape, Item));
@@ -1818,6 +2074,8 @@ namespace UE::Mutable::Private
 	
 			case EOpType::ME_APPLYSHAPE:
 				{
+					check(!Item.bEvaluate);
+
 					OP::MeshApplyShapeArgs Args = Program.GetOpArgs<OP::MeshApplyShapeArgs>(Item.At);
 					PushOp(FScheduledOpInline(Args.mesh, Item));
 					PushOp(FScheduledOpInline(Args.shape, Item));
@@ -1828,6 +2086,8 @@ namespace UE::Mutable::Private
 	
 			case EOpType::ME_MORPHRESHAPE:
 				{
+					check(!Item.bEvaluate);
+
 					OP::MeshMorphReshapeArgs Args = Program.GetOpArgs<OP::MeshMorphReshapeArgs>(Item.At);
 					PushOp(FScheduledOpInline(Args.Morph, Item));
 					PushOp(FScheduledOpInline(Args.Reshape, Item));
@@ -1838,6 +2098,8 @@ namespace UE::Mutable::Private
 	
 			case EOpType::ME_SETSKELETON:
 				{
+					check(!Item.bEvaluate);
+
 					OP::MeshSetSkeletonArgs Args = Program.GetOpArgs<OP::MeshSetSkeletonArgs>(Item.At);
 					PushOp(FScheduledOpInline(Args.Source, Item));
 					PushOp(FScheduledOpInline(Args.Skeleton, Item));
@@ -1848,6 +2110,8 @@ namespace UE::Mutable::Private
 	
 			case EOpType::ME_REMOVEMASK:
 				{
+					check(!Item.bEvaluate);
+
 					const uint8* Data = Program.GetOpArgsPointer(Item.At);
 
 					OP::ADDRESS Source;
@@ -1878,7 +2142,7 @@ namespace UE::Mutable::Private
 
 								if (Condition)
 								{
-									PushOp(FScheduledOpInline(Condition, Item));								
+									PushOp(FScheduledOpInline(Condition, Item, true));								
 								}
 
 								OP::ADDRESS Mask;
@@ -1926,18 +2190,19 @@ namespace UE::Mutable::Private
 
 			case EOpType::ME_ADDMETADATA:
 				{
+					check(!Item.bEvaluate);
+
 					const OP::MeshAddMetadataArgs OpArgs = Program.GetOpArgs<OP::MeshAddMetadataArgs>(Item.At);
 					PushOp(FScheduledOpInline(OpArgs.Source, Item));
-
-					using OpEnumFlags = OP::MeshAddMetadataArgs::EnumFlags;
-					const bool bIsSkeletonList = EnumHasAnyFlags(OpArgs.Flags, OpEnumFlags::IsSkeletonList);
-
+					
 					StoreNone(Item);
 					break;
 				}
 				
 			case EOpType::ME_PROJECT:
 				{
+					check(!Item.bEvaluate);
+
 					OP::MeshProjectArgs Args = Program.GetOpArgs<OP::MeshProjectArgs>(Item.At);
 					PushOp(FScheduledOpInline(Args.Mesh, Item));
 					PushOp(FScheduledOpInline(Args.Projector, Item));
@@ -1948,6 +2213,8 @@ namespace UE::Mutable::Private
 
 			case EOpType::ME_OPTIMIZESKINNING:
 				{
+					check(!Item.bEvaluate);
+					
 					OP::MeshOptimizeSkinningArgs Args = Program.GetOpArgs<OP::MeshOptimizeSkinningArgs>(Item.At);
 					PushOp(FScheduledOpInline(Args.source, Item));
 
@@ -1957,9 +2224,12 @@ namespace UE::Mutable::Private
 
 			case EOpType::ME_TRANSFORMWITHMESH:
 				{
+					check(!Item.bEvaluate);
+					
 					OP::MeshTransformWithinMeshArgs Args = Program.GetOpArgs<OP::MeshTransformWithinMeshArgs>(Item.At);
 					PushOp(FScheduledOpInline(Args.sourceMesh, Item));
 					PushOp(FScheduledOpInline(Args.boundingMesh, Item));
+					//PushOp(FScheduledOpInline(Args.matrix, Item));
 
 					StoreNone(Item);
 					break;
@@ -1967,8 +2237,11 @@ namespace UE::Mutable::Private
 
 			case EOpType::ME_TRANSFORMWITHBONE:
 				{
+					check(!Item.bEvaluate);
+					
 					OP::MeshTransformWithBoneArgs Args = Program.GetOpArgs<OP::MeshTransformWithBoneArgs>(Item.At);
 					PushOp(FScheduledOpInline(Args.SourceMesh, Item));
+					//PushOp(FScheduledOpInline(Args.Matrix, Item));
 
 					StoreNone(Item);
 					break;
@@ -1976,6 +2249,8 @@ namespace UE::Mutable::Private
 
 			case EOpType::ME_SKELETALMESH_BREAK:
 				{
+					check(!Item.bEvaluate);
+					
 					OP::MeshSkeletalMeshBreakArgs Args = Program.GetOpArgs<OP::MeshSkeletalMeshBreakArgs>(Item.At);
 					PushOp(FScheduledOpInline(Args.SkeletalMeshParameter, Item));
 
@@ -1986,11 +2261,13 @@ namespace UE::Mutable::Private
 			
 			case EOpType::IM_REFERENCE:
 				{
+					check(!Item.bEvaluate);
+					
 					OP::ResourceReferenceArgs Args = Program.GetOpArgs<OP::ResourceReferenceArgs>(Item.At);
 
 					if (!Args.ForceLoad) // Not a passthrough. Only happens with None optimization.
 					{
-						GetMemory().CurrentPassthroughObjectLoader->Add<UTexture>(Args.ID);
+						PassthroughObjectLoader->Add<UTexture>(Args.ID);
 					}
 
 					StoreNone(Item);
@@ -1999,9 +2276,12 @@ namespace UE::Mutable::Private
 				
 			case EOpType::IM_LAYERCOLOUR:
 				{
+					check(!Item.bEvaluate);
+					
 					OP::ImageLayerColourArgs Args = Program.GetOpArgs<OP::ImageLayerColourArgs>(Item.At);
 					PushOp(FScheduledOpInline(Args.base, Item));
 					PushOp(FScheduledOpInline(Args.mask, Item));
+					//PushOp(FScheduledOpInline(Args.colour, Item));
 
 					StoreNone(Item);
 					break;
@@ -2009,6 +2289,8 @@ namespace UE::Mutable::Private
 
 			case EOpType::IM_LAYER:
 				{
+					check(!Item.bEvaluate);
+					
 					OP::ImageLayerArgs Args = Program.GetOpArgs<OP::ImageLayerArgs>(Item.At);	
 					PushOp(FScheduledOpInline(Args.base, Item));
 					PushOp(FScheduledOpInline(Args.mask, Item));
@@ -2020,6 +2302,8 @@ namespace UE::Mutable::Private
 
 			case EOpType::IM_MULTILAYER:
 				{
+					check(!Item.bEvaluate);
+					
 					OP::ImageMultiLayerArgs Args = Program.GetOpArgs<OP::ImageMultiLayerArgs>(Item.At);
 
 					switch (Item.Stage)
@@ -2027,7 +2311,7 @@ namespace UE::Mutable::Private
 					case 0:
 						{
 							PushOp(FScheduledOpInline(Item, 1));
-							PushOp(FScheduledOpInline(Args.rangeSize, Item));
+							PushOp(FScheduledOpInline(Args.rangeSize, Item, true));
 							PushOp(FScheduledOpInline(Args.base, Item));
 
 							break;
@@ -2051,11 +2335,11 @@ namespace UE::Mutable::Private
 
 							for (int32 IterationIndex = 0; IterationIndex < NumIterations; ++IterationIndex)
 							{
-								ExecutionIndex Index = GetMemory().CurrentInstanceCache->GetRangeIndex(Item.ExecutionIndex);
+								ExecutionIndex Index = ProgramCache->GetRangeIndex(Item.ExecutionIndex);
 								Index.SetFromModelRangeIndex(Args.rangeId, IterationIndex);
 
 								FScheduledOpInline ItemCopy = Item;
-								ItemCopy.ExecutionIndex = GetMemory().CurrentInstanceCache->GetRangeIndexIndex(Index);
+								ItemCopy.ExecutionIndex = ProgramCache->GetRangeIndexIndex(Index);
 
 								PushOp(FScheduledOpInline(Args.mask, Item));
 								PushOp(FScheduledOpInline(Args.blended, Item));
@@ -2074,6 +2358,8 @@ namespace UE::Mutable::Private
 
 			case EOpType::IM_NORMALCOMPOSITE:
 				{
+					check(!Item.bEvaluate);
+					
 					OP::ImageNormalCompositeArgs Args = Program.GetOpArgs<OP::ImageNormalCompositeArgs>(Item.At);
 					PushOp(FScheduledOpInline(Args.base, Item));
 					PushOp(FScheduledOpInline(Args.normal, Item));
@@ -2084,6 +2370,8 @@ namespace UE::Mutable::Private
 	
 			case EOpType::IM_PIXELFORMAT:
 				{
+					check(!Item.bEvaluate);
+					
 					OP::ImagePixelFormatArgs Args = Program.GetOpArgs<OP::ImagePixelFormatArgs>(Item.At);
 					PushOp(FScheduledOpInline(Args.source, Item));
 
@@ -2093,6 +2381,8 @@ namespace UE::Mutable::Private
 
 			case EOpType::IM_MIPMAP:
 				{
+					check(!Item.bEvaluate);
+					
 					OP::ImageMipmapArgs Args = Program.GetOpArgs<OP::ImageMipmapArgs>(Item.At);
 					PushOp(FScheduledOpInline(Args.source, Item));
 
@@ -2102,6 +2392,8 @@ namespace UE::Mutable::Private
 	
 			case EOpType::IM_RESIZE:
 				{
+					check(!Item.bEvaluate);
+					
 					OP::ImageResizeArgs Args = Program.GetOpArgs<OP::ImageResizeArgs>(Item.At);
 					PushOp(FScheduledOpInline(Args.Source, Item));
 
@@ -2111,6 +2403,8 @@ namespace UE::Mutable::Private
 	
 			case EOpType::IM_RESIZELIKE:
 				{
+					check(!Item.bEvaluate);
+					
 					OP::ImageResizeLikeArgs Args = Program.GetOpArgs<OP::ImageResizeLikeArgs>(Item.At);
 					PushOp(FScheduledOpInline(Args.Source, Item));
 					PushOp(FScheduledOpInline(Args.SizeSource, Item));
@@ -2121,6 +2415,8 @@ namespace UE::Mutable::Private
 
 			case EOpType::IM_RESIZEREL:
 				{
+					check(!Item.bEvaluate);
+					
 					OP::ImageResizeRelArgs Args = Program.GetOpArgs<OP::ImageResizeRelArgs>(Item.At);
 					PushOp(FScheduledOpInline(Args.Source, Item));
 
@@ -2130,6 +2426,8 @@ namespace UE::Mutable::Private
 	
 			case EOpType::IM_BLANKLAYOUT:
 				{
+					check(!Item.bEvaluate);
+					
 					OP::ImageBlankLayoutArgs Args = Program.GetOpArgs<OP::ImageBlankLayoutArgs>(Item.At);
 					PushOp(FScheduledOpInline(Args.Layout, Item));
 
@@ -2139,9 +2437,13 @@ namespace UE::Mutable::Private
 
 			case EOpType::IM_COMPOSE:
 				{
+					check(!Item.bEvaluate);
+					
 					OP::ImageComposeArgs Args = Program.GetOpArgs<OP::ImageComposeArgs>(Item.At);
+					//PushOp(FScheduledOpInline(Args.layout, Item));
 					PushOp(FScheduledOpInline(Args.base, Item));
 					PushOp(FScheduledOpInline(Args.blockImage, Item)); 
+					//PushOp(FScheduledOpInline(Args.mask, Item)); 
 
 					StoreNone(Item);
 					break;
@@ -2159,7 +2461,11 @@ namespace UE::Mutable::Private
 	
 			case EOpType::IM_INTERPOLATE:
 				{
+					check(!Item.bEvaluate);
+					
 					OP::ImageInterpolateArgs Args = Program.GetOpArgs<OP::ImageInterpolateArgs>(Item.At);
+
+					//PushOp(FScheduledOpInline(Args.Factor, Item));
 					
 					for (int32 ImageIndex = 0; ImageIndex < MUTABLE_OP_MAX_INTERPOLATE_COUNT; ++ImageIndex)
 					{
@@ -2177,7 +2483,10 @@ namespace UE::Mutable::Private
 	
 			case EOpType::IM_SATURATE:
 				{
+					check(!Item.bEvaluate);
+					
 					OP::ImageSaturateArgs Args = Program.GetOpArgs<OP::ImageSaturateArgs>(Item.At);
+					//PushOp(FScheduledOpInline(Args.Factor, Item, false));
 					PushOp(FScheduledOpInline(Args.Base, Item));
 
 					StoreNone(Item);
@@ -2186,6 +2495,8 @@ namespace UE::Mutable::Private
 
 			case EOpType::IM_LUMINANCE:
 				{
+					check(!Item.bEvaluate);
+					
 					OP::ImageLuminanceArgs Args = Program.GetOpArgs<OP::ImageLuminanceArgs>(Item.At);
 					PushOp(FScheduledOpInline(Args.Base, Item));
 
@@ -2195,6 +2506,8 @@ namespace UE::Mutable::Private
 	
 			case EOpType::IM_SWIZZLE:
 				{
+					check(!Item.bEvaluate);
+					
 					OP::ImageSwizzleArgs Args = Program.GetOpArgs<OP::ImageSwizzleArgs>(Item.At);
 				
 					TArray<OP::ADDRESS, TFixedAllocator<4>> ValidArgs;
@@ -2212,8 +2525,12 @@ namespace UE::Mutable::Private
 	
 			case EOpType::IM_COLOURMAP:
 				{
+					check(!Item.bEvaluate);
+					
 					OP::ImageColourMapArgs Args = Program.GetOpArgs<OP::ImageColourMapArgs>(Item.At);
 					PushOp(FScheduledOpInline(Args.Base, Item));
+					//PushOp(FScheduledOpInline(Args.Mask, Item));
+					//PushOp(FScheduledOpInline(Args.Map, Item));
 
 					StoreNone(Item);
 					break;
@@ -2221,8 +2538,11 @@ namespace UE::Mutable::Private
 	
 			case EOpType::IM_BINARISE:
 				{
+					check(!Item.bEvaluate);
+					
 					OP::ImageBinariseArgs Args = Program.GetOpArgs<OP::ImageBinariseArgs>(Item.At);
 					PushOp(FScheduledOpInline(Args.Base, Item));
+					//PushOp(FScheduledOpInline(Args.Threshold, Item));
 
 					StoreNone(Item);
 					break;
@@ -2230,6 +2550,8 @@ namespace UE::Mutable::Private
 	
 			case EOpType::IM_INVERT:
 				{
+					check(!Item.bEvaluate);
+					
 					OP::ImageInvertArgs Args = Program.GetOpArgs<OP::ImageInvertArgs>(Item.At);
 					PushOp(FScheduledOpInline(Args.Base, Item));
 
@@ -2239,12 +2561,16 @@ namespace UE::Mutable::Private
 
 			case EOpType::IM_PLAINCOLOUR:
 				{
+					check(!Item.bEvaluate);
+					
 					StoreNone(Item);
 					break;
 				}
 	
 			case EOpType::IM_CROP:
 				{
+					check(!Item.bEvaluate);
+					
 					OP::ImageCropArgs Args = Program.GetOpArgs<OP::ImageCropArgs>(Item.At);
 					PushOp(FScheduledOpInline(Args.source, Item));
 
@@ -2254,6 +2580,8 @@ namespace UE::Mutable::Private
 	
 			case EOpType::IM_PATCH:
 				{
+					check(!Item.bEvaluate);
+					
 					OP::ImagePatchArgs Args = Program.GetOpArgs<OP::ImagePatchArgs>(Item.At);
 					PushOp(FScheduledOpInline(Args.base, Item));
 					PushOp(FScheduledOpInline(Args.patch, Item));
@@ -2264,8 +2592,14 @@ namespace UE::Mutable::Private
 	
 			case EOpType::IM_RASTERMESH:
 				{
+					check(!Item.bEvaluate);
+					
 					OP::ImageRasterMeshArgs Args = Program.GetOpArgs<OP::ImageRasterMeshArgs>(Item.At);
 					PushOp(FScheduledOpInline(Args.mesh, Item));
+					//PushOp(FScheduledOpInline(Args.image, Item));
+					//PushOp(FScheduledOpInline(Args.angleFadeProperties, Item));
+					//PushOp(FScheduledOpInline(Args.mask, Item));
+					//PushOp(FScheduledOpInline(Args.projector, Item));
 
 					StoreNone(Item);
 					break;
@@ -2273,6 +2607,8 @@ namespace UE::Mutable::Private
 
 			case EOpType::IM_MAKEGROWMAP:
 				{
+					check(!Item.bEvaluate);
+					
 					OP::ImageMakeGrowMapArgs Args = Program.GetOpArgs<OP::ImageMakeGrowMapArgs>(Item.At);
 					PushOp(FScheduledOpInline(Args.mask, Item));
 
@@ -2282,8 +2618,11 @@ namespace UE::Mutable::Private
 	
 			case EOpType::IM_DISPLACE:
 				{
+					check(!Item.bEvaluate);
+					
 					OP::ImageDisplaceArgs Args = Program.GetOpArgs<OP::ImageDisplaceArgs>(Item.At);
 					PushOp(FScheduledOpInline(Args.Source, Item));
+					//PushOp(FScheduledOpInline(Args.DisplacementMap, Item));
 
 					StoreNone(Item);
 					break;
@@ -2291,8 +2630,15 @@ namespace UE::Mutable::Private
 	
 			case EOpType::IM_TRANSFORM:
 				{
+					check(!Item.bEvaluate);
+					
 					OP::ImageTransformArgs Args = Program.GetOpArgs<OP::ImageTransformArgs>(Item.At);
 					PushOp(FScheduledOpInline(Args.Base, Item));
+					//PushOp(FScheduledOpInline(Args.OffsetX, Item));
+					//PushOp(FScheduledOpInline(Args.OffsetY, Item));
+					//PushOp(FScheduledOpInline(Args.ScaleX, Item));
+					//PushOp(FScheduledOpInline(Args.ScaleY, Item));
+					//PushOp(FScheduledOpInline(Args.Rotation, Item));
 
 					StoreNone(Item);
 					break;
@@ -2300,6 +2646,8 @@ namespace UE::Mutable::Private
 
 			case EOpType::IM_PARAMETER_CONVERT:
 				{
+					check(!Item.bEvaluate);
+					
 					OP::ImageParameterConvertArgs Args = Program.GetOpArgs<OP::ImageParameterConvertArgs>(Item.At);
 					PushOp(FScheduledOpInline(Args.ImageParameter, Item));
 
@@ -2309,6 +2657,8 @@ namespace UE::Mutable::Private
 				
 			case EOpType::LA_MERGE:
 				{
+					check(!Item.bEvaluate);
+					
 					OP::LayoutMergeArgs Args = Program.GetOpArgs<OP::LayoutMergeArgs>(Item.At);
 					PushOp(FScheduledOpInline(Args.Base, Item));
 					PushOp(FScheduledOpInline(Args.Added, Item));
@@ -2319,6 +2669,8 @@ namespace UE::Mutable::Private
 	
 			case EOpType::LA_PACK:
 				{
+					check(!Item.bEvaluate);
+					
 					OP::LayoutPackArgs Args = Program.GetOpArgs<OP::LayoutPackArgs>(Item.At);
 					PushOp(FScheduledOpInline( Args.Source, Item));
 
@@ -2328,6 +2680,8 @@ namespace UE::Mutable::Private
 
 			case EOpType::LA_FROMMESH:
 				{
+					check(!Item.bEvaluate);
+					
 					OP::LayoutFromMeshArgs Args = Program.GetOpArgs<OP::LayoutFromMeshArgs>(Item.At);
 					PushOp(FScheduledOpInline(Args.Mesh, Item));
 
@@ -2337,6 +2691,8 @@ namespace UE::Mutable::Private
 
 			case EOpType::LA_REMOVEBLOCKS:
 				{
+					check(!Item.bEvaluate);
+					
 					OP::LayoutRemoveBlocksArgs Args = Program.GetOpArgs<OP::LayoutRemoveBlocksArgs>(Item.At);
 					PushOp(FScheduledOpInline(Args.Source, Item));
 					PushOp(FScheduledOpInline(Args.ReferenceLayout, Item));
@@ -2355,8 +2711,10 @@ namespace UE::Mutable::Private
 
 			case EOpType::MI_CONSTANT:
 				{
+					check(!Item.bEvaluate);
+					
 					OP::MaterialConstantArgs Args = Program.GetOpArgs<OP::MaterialConstantArgs>(Item.At);
-					GetMemory().CurrentPassthroughObjectLoader->Add<UMaterialInterface>(Args.ID);
+					PassthroughObjectLoader->Add<UMaterialInterface>(Args.ID);
 
 					StoreNone(Item);
 					break;
@@ -2375,6 +2733,8 @@ namespace UE::Mutable::Private
 			case EOpType::IM_MATERIAL_BREAK:
 			case EOpType::CO_MATERIAL_BREAK:
 				{
+					check(!Item.bEvaluate);
+					
 					OP::MaterialBreakArgs Args = Program.GetOpArgs<OP::MaterialBreakArgs>(Item.At);
 					PushOp(FScheduledOpInline(Args.Material, Item));
 
@@ -2389,6 +2749,8 @@ namespace UE::Mutable::Private
 			case EOpType::MI_EXTERNAL:
 			case EOpType::IS_EXTERNAL:
 				{
+					check(!Item.bEvaluate);
+					
 					const uint8* Data = Program.GetOpArgsPointer(Item.At);
 
 					OP::ExternalArgs Args = Program.GetOpArgs<OP::ExternalArgs>(Item.At);
@@ -2563,16 +2925,11 @@ namespace UE::Mutable::Private
 							EDataType DataType = GetOpDataType(OpType);
 
 							// Evaluate non lazy branches here
-							if ( DataType == EDataType::Image)
+							if (DataType == EDataType::Image)
 							{
 								// Jump Image Parameter Index
 								Data += sizeof(int32);
 							}
-							else if (DataType == EDataType::Scalar)
-							{
-								// Pop float value
-								LoadFloat(ParameterOperation);
-							}
 						}
 
 						StoreNone(Item);
diff --git a/Engine/Plugins/Mutable/Source/MutableRuntime/Private/MuR/CodeRunnerBegin.h b/Engine/Plugins/Mutable/Source/MutableRuntime/Private/MuR/CodeRunnerBegin.h
--- a/Engine/Plugins/Mutable/Source/MutableRuntime/Private/MuR/CodeRunnerBegin.h
+++ b/Engine/Plugins/Mutable/Source/MutableRuntime/Private/MuR/CodeRunnerBegin.h
@@ -18,30 +18,41 @@ namespace UE::Mutable::Private
 
 	struct FScheduledOpInline
 	{
-		explicit FScheduledOpInline(OP::ADDRESS InAt = 0, const uint16 InStage = 0, const uint16 InExecutionIndex = 0) :
+		explicit FScheduledOpInline(OP::ADDRESS InAt = 0, const uint16 InStage = 0, const uint16 InExecutionIndex = 0, bool bInEvaluate = true) :
 			At(InAt),
 			Stage(InStage),
-		 	ExecutionIndex(InExecutionIndex)
+			bEvaluate(bInEvaluate),
+			ExecutionIndex(InExecutionIndex)
 		{
 		}
-		
+
 		explicit FScheduledOpInline(OP::ADDRESS InAt, const FScheduledOpInline& InOpTemplate) :
 			At(InAt),
-     		ExecutionIndex(InOpTemplate.ExecutionIndex)
+			bEvaluate(InOpTemplate.bEvaluate),
+			ExecutionIndex(InOpTemplate.ExecutionIndex)
+		{
+		}
+		
+		explicit FScheduledOpInline(OP::ADDRESS InAt, const FScheduledOpInline& InOpTemplate, bool bInEvaluate) :
+			At(InAt),
+			bEvaluate(bInEvaluate),
+			ExecutionIndex(InOpTemplate.ExecutionIndex)
 		{
 		}
 		
 		FScheduledOpInline(const FScheduledOpInline& InOpTemplate, uint16 InStage) :
 			At(InOpTemplate.At),
          	Stage(InStage),
-            ExecutionIndex(InOpTemplate.ExecutionIndex)
+			bEvaluate(InOpTemplate.bEvaluate),
+			ExecutionIndex(InOpTemplate.ExecutionIndex)
 		{
 		}
 
 		bool operator==(const FScheduledOpInline& Other) const;
 
 		OP::ADDRESS At = 0;
-		uint16 Stage = 0;
+		uint8 Stage = 0;
+		uint8 bEvaluate = 0;
 		uint16 ExecutionIndex = 0;
 		OP::ADDRESS CustomState = 0;
 	};
@@ -98,12 +109,21 @@ namespace UE::Mutable::Private
 	class CodeRunnerBegin
 	{
 	public:
-		CodeRunnerBegin(FSystem& InSystem, const TSharedPtr<const FModel>& InModel, const TSharedPtr<const FParameters>& InParams, uint32 InLodMask);
+		MUTABLERUNTIME_API CodeRunnerBegin(
+			const TSharedPtr<const FModel>& InModel,
+			const TSharedPtr<const FParameters>& InParams,
+			const TSharedPtr<FProgramCache>& InProgramCache,
+			const TSharedPtr<FMeshIdRegistry>& InMeshIdRegistry,
+			const TSharedPtr<FImageIdRegistry>& InImageIdRegistry,
+			const TSharedPtr<FMaterialIdRegistry>& InMaterialIdRegistry,
+			const TSharedPtr<FSkeletalMeshIdRegistry>& InSkeletalMeshIdRegistry,
+			const TSharedPtr<FPassthroughObjectLoader>& InPassthroughObjectLoader,
+			const FScheduledOpInline& InRoot,
+			uint32 InLodMask);
 
 		TManagedPtr<const FInstance> RunCode(const FScheduledOpInline& Root);
-
-		FWorkingMemoryManager& GetMemory() const;
 		
+	private:
 		FScheduledOpInline PopOp();
 		
 		void PushOp(const FScheduledOpInline& Item);
@@ -137,7 +157,6 @@ namespace UE::Mutable::Private
 		TSharedPtr<const FModel> Model;
 		TSharedPtr<const FParameters> Params;
 
-		FSystem& System;
 		FProgram& Program;
 		
 		uint32 LodMask;
@@ -154,5 +173,14 @@ namespace UE::Mutable::Private
 		TMap<FCacheAddressInline, TManagedPtr<const FInstance>> ResultsInstance;
 		TMap<FCacheAddressInline, TManagedPtr<const String>> ResultsString;
 		TMap<FCacheAddressInline, TManagedPtr<const FExtensionData>> ResultsExtensionData;
+
+		TSharedPtr<FProgramCache> ProgramCache;
+
+		TSharedPtr<FMeshIdRegistry> MeshIdRegistry;
+		TSharedPtr<FImageIdRegistry> ImageIdRegistry;
+		TSharedPtr<FMaterialIdRegistry> MaterialIdRegistry;
+		TSharedPtr<FSkeletalMeshIdRegistry> SkeletalMeshIdRegistry;
+
+		TSharedPtr<FPassthroughObjectLoader> PassthroughObjectLoader;
 	};
 }
diff --git a/Engine/Plugins/Mutable/Source/MutableRuntime/Private/MuR/System.cpp b/Engine/Plugins/Mutable/Source/MutableRuntime/Private/MuR/System.cpp
--- a/Engine/Plugins/Mutable/Source/MutableRuntime/Private/MuR/System.cpp
+++ b/Engine/Plugins/Mutable/Source/MutableRuntime/Private/MuR/System.cpp
@@ -247,47 +247,30 @@ namespace UE::Mutable::Private
 		MUTABLE_CPUPROFILER_SCOPE(BeginUpdate_GameThread);
 
     	FLiveInstance* LiveInstance = m_pD->FindLiveInstance(InInstanceID);
-    	if (!LiveInstance)
+    	check(LiveInstance);
+		check(InParams);
+
+    	FScheduledOpInline Root(LiveInstance->Model->GetPrivate()->Program.States[InStateIndex].Root, 0, 0, true);
+    	CodeRunnerBegin Runner(
+    		LiveInstance->Model,
+    		InParams,
+    		LiveInstance->Cache,
+    		InMeshIdRegistry,
+    		InImageIdRegistry,
+    		InMaterialIdRegistry,
+    		InSkeletalMeshIdRegistry,
+    		LiveInstance->PassthroughObjectLoader,
+    		Root,
+    		InLodMask);
+
+    	TManagedPtr<const FInstance> Result = Runner.RunCode(FScheduledOpInline(Root, 0));
+    	if (!Result)
     	{
-    		UE_LOG(LogMutableCore, Error, TEXT("Invalid instance id in mutable update."));
-    		return nullptr;
+    		// In case of failure return an empty instance, to prevent following code to have to check it every time
+    		Result = MakeManaged<FInstance>();
     	}
     	
-		if (!InParams)
-		{
-			UE_LOG(LogMutableCore, Error, TEXT("Invalid parameters in mutable update."));
-			return nullptr;
-		}
-
-    	m_pD->WorkingMemoryManager.CurrentInstanceCache = LiveInstance->Cache;
-    	m_pD->WorkingMemoryManager.CurrentPassthroughObjectLoader = LiveInstance->PassthroughObjectLoader;
-
-		m_pD->WorkingMemoryManager.CurrentMeshIdRegistry = InMeshIdRegistry;
-    	m_pD->WorkingMemoryManager.CurrentImageIdRegistry = InImageIdRegistry;
-    	m_pD->WorkingMemoryManager.CurrentMaterialIdRegistry = InMaterialIdRegistry;
-    	m_pD->WorkingMemoryManager.CurrentSkeletalMeshIdRegistry = InSkeletalMeshIdRegistry;
-    	
-    	
-		OP::ADDRESS RootAt = LiveInstance->Model->GetPrivate()->Program.States[InStateIndex].Root;
-
-    	CodeRunnerBegin Runner(*this, LiveInstance->Model, InParams, InLodMask);
-    	
-    	TManagedPtr<const FInstance> Result = Runner.RunCode(FScheduledOpInline(RootAt, 0));
-    	if (!Result)
-		{
-			// In case of failure return an empty instance, to prevent following code to have to check it every time
-			Result = MakeManaged<FInstance>();
-		}
-
-    	m_pD->WorkingMemoryManager.CurrentInstanceCache = nullptr;
-    	m_pD->WorkingMemoryManager.CurrentPassthroughObjectLoader = nullptr;
-
-    	m_pD->WorkingMemoryManager.CurrentMeshIdRegistry = nullptr;
-    	m_pD->WorkingMemoryManager.CurrentImageIdRegistry = nullptr;
-    	m_pD->WorkingMemoryManager.CurrentMaterialIdRegistry = nullptr;
-    	m_pD->WorkingMemoryManager.CurrentSkeletalMeshIdRegistry = nullptr;
-
-		return Result;
+    	return Result;
 	}
 
 	
@@ -1077,28 +1060,22 @@ namespace UE::Mutable::Private
     {
     public:
 
-        RelevantParameterVisitor
-            (
+        RelevantParameterVisitor(
                 FSystem::Private* InSystem,
 				const TSharedPtr<const FModel>& InModel,
                 const TSharedPtr<const FParameters>& InParams,
-                bool* InFlags
-            )
-            : UniqueDiscreteCoveredCodeVisitor<>( InSystem, InModel, InParams, FSystem::AllLODs )
+                TArray<uint32>& OutRelevantParameters) :
+			UniqueDiscreteCoveredCodeVisitor<>(InSystem, InModel, InParams, FSystem::AllLODs),
+		    RelevantParameters(OutRelevantParameters)
         {
-            Flags = InFlags;
+            OP::ADDRESS At = InModel->GetPrivate()->Program.States[0].Root;
 
-            FMemory::Memset( Flags, 0, sizeof(bool)*InParams->GetCount() );
-
-            OP::ADDRESS at = InModel->GetPrivate()->Program.States[0].Root;
-
-            Run( at );
+            Run(At);
         }
 
-
-        bool Visit( OP::ADDRESS at, FProgram& Program ) override
+        bool Visit(OP::ADDRESS At, FProgram& Program) override
         {
-            switch ( Program.GetOpType(at) )
+            switch (Program.GetOpType(At))
             {
             case EOpType::BO_PARAMETER:
             case EOpType::NU_PARAMETER:
@@ -1111,29 +1088,27 @@ namespace UE::Mutable::Private
 			case EOpType::MI_PARAMETER:
 			case EOpType::IS_PARAMETER:
             {
-				OP::ParameterArgs args = Program.GetOpArgs<OP::ParameterArgs>(at);
+				OP::ParameterArgs args = Program.GetOpArgs<OP::ParameterArgs>(At);
 				OP::ADDRESS ParamIndex = args.variable;
-                Flags[ParamIndex] = true;
+                RelevantParameters.Add(ParamIndex);
                 break;
             }
             default:
                 break;
             }
 
-            return UniqueDiscreteCoveredCodeVisitor<>::Visit( at, Program );
+            return UniqueDiscreteCoveredCodeVisitor<>::Visit( At, Program );
         }
 
     private:
-
-        //! Non-owned result buffer
-        bool* Flags;
+        TArray<uint32>& RelevantParameters;
     };
 
 
     //---------------------------------------------------------------------------------------------
-    void FSystem::GetParameterRelevancy( FInstance::FID InstanceID,
-                                        const TSharedPtr<const FParameters>& FParameters,
-                                        bool* Flags )
+    void FSystem::GetParameterRelevancy(FInstance::FID InstanceID,
+    	const TSharedPtr<const FParameters>& FParameters,
+        TArray<uint32>& OutRelevantParameters)
     {
 		LLM_SCOPE_BYNAME(TEXT("MutableRuntime"));
 
@@ -1143,7 +1118,7 @@ namespace UE::Mutable::Private
 		m_pD->WorkingMemoryManager.CurrentInstanceCache = LiveInstance->Cache;
 		m_pD->WorkingMemoryManager.CurrentPassthroughObjectLoader = LiveInstance->PassthroughObjectLoader;
 		
-		RelevantParameterVisitor visitor( m_pD, LiveInstance->Model, FParameters, Flags );
+		RelevantParameterVisitor visitor(m_pD, LiveInstance->Model, FParameters, OutRelevantParameters);
 
 		m_pD->WorkingMemoryManager.CurrentInstanceCache = nullptr;
 		m_pD->WorkingMemoryManager.CurrentPassthroughObjectLoader = nullptr;
diff --git a/Engine/Plugins/Mutable/Source/MutableTools/Private/MuT/CodeOptimiser.h b/Engine/Plugins/Mutable/Source/MutableTools/Private/MuT/CodeOptimiser.h
--- a/Engine/Plugins/Mutable/Source/MutableTools/Private/MuT/CodeOptimiser.h
+++ b/Engine/Plugins/Mutable/Source/MutableTools/Private/MuT/CodeOptimiser.h
@@ -3,7 +3,6 @@
 #pragma once
 
 #include "MuR/Operations.h"
-#include "MuR/CodeVisitor.h"
 #include "MuT/CompilerPrivate.h"
 #include "MuT/AST.h"
 
diff --git a/Engine/Plugins/Mutable/Source/MutableTools/Private/MuT/DataPacker.cpp b/Engine/Plugins/Mutable/Source/MutableTools/Private/MuT/DataPacker.cpp
--- a/Engine/Plugins/Mutable/Source/MutableTools/Private/MuT/DataPacker.cpp
+++ b/Engine/Plugins/Mutable/Source/MutableTools/Private/MuT/DataPacker.cpp
@@ -5,7 +5,6 @@
 #include "Containers/Array.h"
 #include "HAL/PlatformCrt.h"
 #include "HAL/UnrealMemory.h"
-#include "MuR/CodeVisitor.h"
 #include "MuR/Image.h"
 #include "MuR/Layout.h"
 #include "MuR/Mesh.h"
