[Mutable] Add backwards compatibility to Macro libraries.

#rnx
#Synced-CL 49108708

diff --git a/Engine/Plugins/Mutable/Source/CustomizableObjectEditor/Private/MuCOE/CustomizableObjectMacroLibrary/CustomizableObjectMacroLibrary.cpp b/Engine/Plugins/Mutable/Source/CustomizableObjectEditor/Private/MuCOE/CustomizableObjectMacroLibrary/CustomizableObjectMacroLibrary.cpp
--- a/Engine/Plugins/Mutable/Source/CustomizableObjectEditor/Private/MuCOE/CustomizableObjectMacroLibrary/CustomizableObjectMacroLibrary.cpp
+++ b/Engine/Plugins/Mutable/Source/CustomizableObjectEditor/Private/MuCOE/CustomizableObjectMacroLibrary/CustomizableObjectMacroLibrary.cpp
@@ -5,12 +5,57 @@
 #include "MuCOE/CustomizableObjectGraph.h"
 #include "MuCOE/Nodes/CustomizableObjectNodeTunnel.h"
 #include "MuCOE/EdGraphSchema_CustomizableObject.h"
+#include "MuCO/ICustomizableObjectEditorModule.h"
+#include "MuCO/CustomizableObjectCustomVersion.h"
+
 
 #include UE_INLINE_GENERATED_CPP_BY_NAME(CustomizableObjectMacroLibrary)
 
 #define LOCTEXT_NAMESPACE "CustomizableObjectMacroLibrary"
 
 
+void UCustomizableObjectMacroLibrary::PostLoad()
+{
+	Super::PostLoad();
+
+	const int32 CustomVersion = FMath::Max(GetLinkerCustomVersion(FCustomizableObjectCustomVersion::GUID), (int32)FCustomizableObjectCustomVersion::EnableMutableMacros);
+
+	for (TObjectPtr<UCustomizableObjectMacro>& Macro : Macros)
+	{
+		if (ensure(Macro))
+		{
+			Macro->ConditionalPostLoad();
+		}
+	}
+
+	if (ICustomizableObjectEditorModule* Module = ICustomizableObjectEditorModule::Get())
+	{
+		for (TObjectPtr<UCustomizableObjectMacro>& Macro : Macros)
+		{
+			UEdGraph* MacroGraph = Macro ? Macro->Graph : nullptr;
+			if (ensure(MacroGraph))
+			{
+				MacroGraph->ConditionalPostLoad();
+
+				for (int32 Version = CustomVersion + 1; Version <= FCustomizableObjectCustomVersion::LatestVersion; ++Version)
+				{
+					// Execute backwards compatible code for all nodes. It requires all nodes to be loaded.
+					Module->BackwardsCompatibleFixup(*MacroGraph, Version);
+				}
+
+				Module->PostBackwardsCompatibleFixup(*MacroGraph);
+			}
+		}
+	}
+}
+
+void UCustomizableObjectMacroLibrary::Serialize(FArchive& Ar)
+{
+	Super::Serialize(Ar);
+
+	Ar.UsingCustomVersion(FCustomizableObjectCustomVersion::GUID);
+}
+
 UCustomizableObjectMacro* UCustomizableObjectMacroLibrary::AddMacro()
 {
 	UCustomizableObjectMacro* NewMacro = nullptr;
diff --git a/Engine/Plugins/Mutable/Source/CustomizableObjectEditor/Private/MuCOE/CustomizableObjectMacroLibrary/CustomizableObjectMacroLibrary.h b/Engine/Plugins/Mutable/Source/CustomizableObjectEditor/Private/MuCOE/CustomizableObjectMacroLibrary/CustomizableObjectMacroLibrary.h
--- a/Engine/Plugins/Mutable/Source/CustomizableObjectEditor/Private/MuCOE/CustomizableObjectMacroLibrary/CustomizableObjectMacroLibrary.h
+++ b/Engine/Plugins/Mutable/Source/CustomizableObjectEditor/Private/MuCOE/CustomizableObjectMacroLibrary/CustomizableObjectMacroLibrary.h
@@ -90,6 +90,10 @@ public:
 
 	UCustomizableObjectMacroLibrary() {};
 
+	// UObject interface
+	UE_API virtual void PostLoad() override;
+	UE_API virtual void Serialize(FArchive& Ar) override;
+
 	/** Creates a new macro with the basic nodes and adds it to the library. */
 	UE_API UCustomizableObjectMacro* AddMacro();
 
diff --git a/Engine/Plugins/Mutable/Source/CustomizableObjectEditor/Private/MuCOE/Nodes/CustomizableObjectNodeComponentPassthroughMesh.cpp b/Engine/Plugins/Mutable/Source/CustomizableObjectEditor/Private/MuCOE/Nodes/CustomizableObjectNodeComponentPassthroughMesh.cpp
--- a/Engine/Plugins/Mutable/Source/CustomizableObjectEditor/Private/MuCOE/Nodes/CustomizableObjectNodeComponentPassthroughMesh.cpp
+++ b/Engine/Plugins/Mutable/Source/CustomizableObjectEditor/Private/MuCOE/Nodes/CustomizableObjectNodeComponentPassthroughMesh.cpp
@@ -34,8 +34,11 @@ void UCustomizableObjectNodeComponentPassthroughMesh::BackwardsCompatibleFixup(i
 				CustomRemovePin(*Pin);
 			}
 		}
-		
-		SkeletalMeshPin = CustomCreatePin(EGPD_Input, UEdGraphSchema_CustomizableObject::PC_SkeletalMesh_Passthrough);
+
+		if (!SkeletalMeshPin.Get())
+		{
+			SkeletalMeshPin = CustomCreatePin(EGPD_Input, UEdGraphSchema_CustomizableObject::PC_SkeletalMesh_Passthrough);
+		}
 	}
 }
 
diff --git a/Engine/Plugins/Mutable/Source/CustomizableObjectEditor/Private/MuCOE/Nodes/CustomizableObjectNodeTextureParameter.cpp b/Engine/Plugins/Mutable/Source/CustomizableObjectEditor/Private/MuCOE/Nodes/CustomizableObjectNodeTextureParameter.cpp
--- a/Engine/Plugins/Mutable/Source/CustomizableObjectEditor/Private/MuCOE/Nodes/CustomizableObjectNodeTextureParameter.cpp
+++ b/Engine/Plugins/Mutable/Source/CustomizableObjectEditor/Private/MuCOE/Nodes/CustomizableObjectNodeTextureParameter.cpp
@@ -46,7 +46,10 @@ void UCustomizableObjectNodeTextureParameter::BackwardsCompatibleFixup(int32 Cus
 	}
 	else if (CustomizableObjectCustomVersion == FCustomizableObjectCustomVersion::NodeTextureParameterPassthroughPin)
 	{
-		PassthroughPin = CustomCreatePin(EGPD_Output, UEdGraphSchema_CustomizableObject::PC_Texture_Passthrough);
+		if (!PassthroughPin.Get())
+		{
+			PassthroughPin = CustomCreatePin(EGPD_Output, UEdGraphSchema_CustomizableObject::PC_Texture_Passthrough);
+		}
 	}
 	else if (CustomizableObjectCustomVersion == FCustomizableObjectCustomVersion::PassthroughParameter)
 	{
