[Mutable] Fix crash when opening an empty CO.

#Synced-CL 49159615

diff --git a/Engine/Plugins/Mutable/Source/MutableRuntime/Internal/MuR/SystemPrivate.h b/Engine/Plugins/Mutable/Source/MutableRuntime/Internal/MuR/SystemPrivate.h
--- a/Engine/Plugins/Mutable/Source/MutableRuntime/Internal/MuR/SystemPrivate.h
+++ b/Engine/Plugins/Mutable/Source/MutableRuntime/Internal/MuR/SystemPrivate.h
@@ -1328,7 +1328,10 @@ namespace UE::Mutable::Private
 
 		bool AquireCachedResult(FCacheAddress Address)
 		{
-			check(Address.At)
+			if (!Address.At)
+			{
+				return false;
+			}
 
 			check(Address.At < OpExecutionData.size_code());
 
