[mutable] Improved MaterialPin related nodes (MaterialModify and MaterialBreak mostly)
- Removed FCOVariable object as it was adding complexity to the Pin category change interactions, and now, instead, we change the type directly as it is simpler.
- Made the CategoryName dropdown, whose backend was previously in FCOVariable depend on the actual pin instead (it removes the man in the middle that was the FCOVariable structure).
- Made the changes in the editable name of multiple types work with the transactions system (undo/redo)
- Made the PinViewer able to also edit the "Editable Name" of the node pins. This expands the capabilities of the PinViewer without the need for major architecture changes.
- Removed custom remapping for the MaterialBreak and MaterialModify nodes as they do not require such complex code as they do not require a refresh when adding or removing pins.
- Other minor changes and fixes.

#rnx
#Synced-CL 49270271

diff --git a/Engine/Plugins/Mutable/Source/CustomizableObjectEditor/Private/MuCOE/COVariable.cpp b/Engine/Plugins/Mutable/Source/CustomizableObjectEditor/Private/MuCOE/COVariable.cpp
deleted file mode 100644
diff --git a/Engine/Plugins/Mutable/Source/CustomizableObjectEditor/Private/MuCOE/COVariable.h b/Engine/Plugins/Mutable/Source/CustomizableObjectEditor/Private/MuCOE/COVariable.h
deleted file mode 100644
diff --git a/Engine/Plugins/Mutable/Source/CustomizableObjectEditor/Private/MuCOE/CustomizableObjectEditorModule.cpp b/Engine/Plugins/Mutable/Source/CustomizableObjectEditor/Private/MuCOE/CustomizableObjectEditorModule.cpp
--- a/Engine/Plugins/Mutable/Source/CustomizableObjectEditor/Private/MuCOE/CustomizableObjectEditorModule.cpp
+++ b/Engine/Plugins/Mutable/Source/CustomizableObjectEditor/Private/MuCOE/CustomizableObjectEditorModule.cpp
@@ -87,7 +87,6 @@
 #include "MuCOE/Widgets/CustomizableObjectVariationCustomization.h"
 #include "MuCOE/Widgets/CustomizableObjectLODReductionSettings.h"
 #include "MuCOE/Widgets/CustomizableObjectNodeTableCompilationFilterEditor.h"
-#include "MuCOE/COVariable.h"
 #include "PropertyEditorModule.h"
 #include "AssetRegistry/AssetRegistryModule.h"
 #include "MuCO/CustomizableObjectInstance.h"
@@ -289,7 +288,6 @@ void FCustomizableObjectEditorModule::StartupModule()
 	// We need to cache those two FNames as if we want to get them on "ShutdownModule" we get "None" FNames and an ASAN error in Linux.
 	MeshReshapeBoneReferenceUStructName = FMeshReshapeBoneReference::StaticStruct()->GetFName();
 	BoneToRemoveUStructName = FBoneToRemove::StaticStruct()->GetFName();
-	COVariableUStructName = FCOVariable::StaticStruct()->GetFName();
 
 	// Custom properties
 	PropertyModule.RegisterCustomPropertyTypeLayout("CustomizableObjectIdentifier", FOnGetPropertyTypeCustomizationInstance::CreateStatic(&FCustomizableObjectIdentifierCustomization::MakeInstance));
@@ -299,7 +297,6 @@ void FCustomizableObjectEditorModule::StartupModule()
 	PropertyModule.RegisterCustomPropertyTypeLayout(FCustomizableObjectVariation::StaticStruct()->GetFName(), FOnGetPropertyTypeCustomizationInstance::CreateStatic(&FCustomizableObjectVariationCustomization::MakeInstance));
 	PropertyModule.RegisterCustomPropertyTypeLayout(FCustomizableObjectTextureVariation::StaticStruct()->GetFName(), FOnGetPropertyTypeCustomizationInstance::CreateStatic(&FCustomizableObjectVariationCustomization::MakeInstance));
 	PropertyModule.RegisterCustomPropertyTypeLayout(FTableNodeCompilationFilter::StaticStruct()->GetFName(), FOnGetPropertyTypeCustomizationInstance::CreateStatic(&FCustomizableObjectNodeTableCompilationFilterEditor::MakeInstance));
-	PropertyModule.RegisterCustomPropertyTypeLayout(COVariableUStructName, FOnGetPropertyTypeCustomizationInstance::CreateStatic(&FCOVariableCustomization::MakeInstance));
 	PropertyModule.RegisterCustomPropertyTypeLayout("InstancedStruct", FOnGetPropertyTypeCustomizationInstance::CreateStatic(&FExternalTypeCustomization::MakeInstance), MakeShared<FExternalTypeTypeIdentifier>());
 
 	PropertyModule.NotifyCustomizationModuleChanged();
@@ -347,7 +344,6 @@ void FCustomizableObjectEditorModule::ShutdownModule()
 		PropertyModule.UnregisterCustomPropertyTypeLayout(MeshReshapeBoneReferenceUStructName);
 		PropertyModule.UnregisterCustomPropertyTypeLayout(BoneToRemoveUStructName);
 		PropertyModule.UnregisterCustomPropertyTypeLayout(NAME_StrProperty);
-		PropertyModule.UnregisterCustomPropertyTypeLayout(COVariableUStructName);
 
 		PropertyModule.NotifyCustomizationModuleChanged();
 	}
diff --git a/Engine/Plugins/Mutable/Source/CustomizableObjectEditor/Private/MuCOE/CustomizableObjectEditorModule.h b/Engine/Plugins/Mutable/Source/CustomizableObjectEditor/Private/MuCOE/CustomizableObjectEditorModule.h
--- a/Engine/Plugins/Mutable/Source/CustomizableObjectEditor/Private/MuCOE/CustomizableObjectEditorModule.h
+++ b/Engine/Plugins/Mutable/Source/CustomizableObjectEditor/Private/MuCOE/CustomizableObjectEditorModule.h
@@ -106,7 +106,6 @@ private:
 
 	FName MeshReshapeBoneReferenceUStructName;
 	FName BoneToRemoveUStructName;
-	FName COVariableUStructName;
 
 	/** Cook requests. */
 	TMap<TWeakObjectPtr<UCustomizableObject>, TArray<TSharedRef<FCompilationRequest>>> CookCompileRequests;
diff --git a/Engine/Plugins/Mutable/Source/CustomizableObjectEditor/Private/MuCOE/EdGraphSchema_CustomizableObject.cpp b/Engine/Plugins/Mutable/Source/CustomizableObjectEditor/Private/MuCOE/EdGraphSchema_CustomizableObject.cpp
--- a/Engine/Plugins/Mutable/Source/CustomizableObjectEditor/Private/MuCOE/EdGraphSchema_CustomizableObject.cpp
+++ b/Engine/Plugins/Mutable/Source/CustomizableObjectEditor/Private/MuCOE/EdGraphSchema_CustomizableObject.cpp
@@ -122,6 +122,7 @@
 #include "Nodes/CONodeExternalTypeParameter.h"
 #include "Nodes/CONodeSkeletalMeshBreak.h"
 #include "UObject/Class.h"
+#include "Widgets/SMutablePinTypeSelector.h"
 
 #include UE_INLINE_GENERATED_CPP_BY_NAME(EdGraphSchema_CustomizableObject)
 
@@ -1145,6 +1146,16 @@ void UEdGraphSchema_CustomizableObject::GetContextMenuActions(UToolMenu* Menu, U
 						})));
 				}
 				
+				const TArray<FName> AllowedTypes = Node->GetPinAllowedTypes(*Pin);
+				if (AllowedTypes.Num())
+				{
+					FEdGraphPinReference PinReference = FEdGraphPinReference(const_cast<UEdGraphPin*>(Pin));
+					TSharedRef<SMutablePinTypeSelector> TypeSelectionSlate = SNew(SMutablePinTypeSelector).PinReference(PinReference);
+					
+					FToolMenuSection& Section = Menu->FindOrAddSection("EdGraphSchemaPinActions");
+					Section.AddEntry(FToolMenuEntry::InitWidget("Pin Type Selector", TypeSelectionSlate,  LOCTEXT("SetPinType_Label", "Pin Type"), true));
+				}
+				
 				if (TSharedPtr<IDetailsView> Widget = Node->CustomizePinDetails(*Pin))
 				{
 					Node->PostReconstructNodeDelegate.AddLambda([WeakMenu = Widget->AsWeak()](){
diff --git a/Engine/Plugins/Mutable/Source/CustomizableObjectEditor/Private/MuCOE/Nodes/CONodeMaterialBreak.cpp b/Engine/Plugins/Mutable/Source/CustomizableObjectEditor/Private/MuCOE/Nodes/CONodeMaterialBreak.cpp
--- a/Engine/Plugins/Mutable/Source/CustomizableObjectEditor/Private/MuCOE/Nodes/CONodeMaterialBreak.cpp
+++ b/Engine/Plugins/Mutable/Source/CustomizableObjectEditor/Private/MuCOE/Nodes/CONodeMaterialBreak.cpp
@@ -6,6 +6,7 @@
 #include "MuCOE/EdGraphSchema_CustomizableObject.h"
 #include "PropertyEditorModule.h"
 #include "ScopedTransaction.h"
+#include "MuCOE/RemapPins/CustomizableObjectNodeRemapPinsByPosition.h"
 
 
 #include UE_INLINE_GENERATED_CPP_BY_NAME(CONodeMaterialBreak)
@@ -19,21 +20,6 @@ struct FSlateBrush;
 const FName BreakMaterialMaterialParameterPinName = "Parameter";
 const FText BreakMaterialMaterialParameterPinText = LOCTEXT("Parameter", "Parameter");
 
-bool UCustomizableObjectNodeBreakMaterialRemapPins::Equal(const UCustomizableObjectNode& Node, const UEdGraphPin& OldPin, const UEdGraphPin& NewPin) const
-{
-	UCONodeMaterialBreakParameterPinData* OldPinData = Cast<UCONodeMaterialBreakParameterPinData>(Node.GetPinData(OldPin));
-	UCONodeMaterialBreakParameterPinData* NewPinData = Cast<UCONodeMaterialBreakParameterPinData>(Node.GetPinData(NewPin));
-
-	if (OldPinData && NewPinData)
-	{
-		return OldPin.Direction == NewPin.Direction && OldPinData->VariableData.Type == NewPinData->VariableData.Type && OldPinData->VariableData.Id == NewPinData->VariableData.Id;
-	}
-	else
-	{
-		return Super::Equal(Node, OldPin, NewPin);
-	}
-}
-
 
 void UCONodeMaterialBreak::AllocateDefaultPins(UCustomizableObjectNodeRemapPins* RemapPins)
 {
@@ -52,19 +38,17 @@ void UCONodeMaterialBreak::AllocateDefaultPins(UCustomizableObjectNodeRemapPins*
 		
 		if (UCONodeMaterialBreakParameterPinData* PinData = Cast<UCONodeMaterialBreakParameterPinData>(GetPinData(*Pin)))
 		{
-			FText FriendlyName;
-			
+			// Use the same pin data as the old pin since the new one is a copy
+			UEdGraphPin* NewPin = CustomCreatePin(EGPD_Output, Pin->PinType.PinCategory, BreakMaterialMaterialParameterPinName, BreakMaterialMaterialParameterPinText, PinData);
+
 			if (PinData->LayerIndex == 0)
 			{
-				FriendlyName = FText::Format(LOCTEXT("BackgroundLayerText", "{0} - Background"), BreakMaterialMaterialParameterPinText);
+				NewPin->PinFriendlyName = FText::Format(LOCTEXT("BackgroundLayerText", "{0} - Background"), FText::FromName(BreakMaterialMaterialParameterPinName));
 			}
 			else if (PinData->LayerIndex > 0)
 			{
-				FriendlyName = FText::Format(LOCTEXT("LayerText", "{0} - Layer {1}"), BreakMaterialMaterialParameterPinText, FText::FromString(FString::FromInt(PinData->LayerIndex)));
+				NewPin->PinFriendlyName = FText::Format(LOCTEXT("LayerText", "{0} - Layer {1}"), FText::FromName(BreakMaterialMaterialParameterPinName), FText::FromString(FString::FromInt(PinData->LayerIndex)));
 			}
-
-			// Use the same pin data as the old pin since the new one is a copy
-			CustomCreatePin(EGPD_Output, PinData->VariableData.Type, BreakMaterialMaterialParameterPinName, FriendlyName, PinData);
 		}
 	}
 }
@@ -150,7 +134,7 @@ EAddPinNodeButtonLocation UCONodeMaterialBreak::GetAddPinButtonNodeSide() const
 
 UCustomizableObjectNodeRemapPins* UCONodeMaterialBreak::CreateRemapPinsDefault() const
 {
-	return NewObject<UCustomizableObjectNodeBreakMaterialRemapPins>();
+	return NewObject<UCustomizableObjectNodeRemapPinsByPosition>();
 }
 
 
@@ -165,6 +149,22 @@ TArray<FName> UCONodeMaterialBreak::GetAllowedPinViewerCreationTypes() const
 }
 
 
+TArray<FName> UCONodeMaterialBreak::GetPinAllowedTypes(const UEdGraphPin& Pin) const
+{
+	if (Pin.Direction == EGPD_Output)
+	{
+		const TArray<FName> PinAllowedTypes = GetAllowedPinViewerCreationTypes();
+		check(PinAllowedTypes.Contains(Pin.PinType.PinCategory));
+		
+		return PinAllowedTypes;
+	}
+	else
+	{	
+		return Super::GetPinAllowedTypes(Pin);
+	}
+}
+
+
 void UCONodeMaterialBreak::AddPinFromUI()
 {
 	Super::AddPinFromUI();
@@ -173,11 +173,7 @@ void UCONodeMaterialBreak::AddPinFromUI()
 	Modify();
 	
 	const FName NewPinCategory = UEdGraphSchema_CustomizableObject::PC_Float;
-	
 	UCONodeMaterialBreakParameterPinData* PinData = NewObject<UCONodeMaterialBreakParameterPinData>(this);
-	PinData->VariableData.Type = NewPinCategory;
-	PinData->VariableData.Id = FGuid::NewGuid();
-
 	if (CustomCreatePin(EEdGraphPinDirection::EGPD_Output, NewPinCategory, BreakMaterialMaterialParameterPinName, BreakMaterialMaterialParameterPinText, PinData))
 	{
 		GetGraph()->NotifyNodeChanged(this);
@@ -202,7 +198,10 @@ void UCONodeMaterialBreak::SetPinEditableName(const UEdGraphPin& Pin, const FTex
 	UCustomizableObjectNodePinData* PinData = GetPinData(Pin);
 	if (UCONodeMaterialBreakParameterPinData* CastedPinData = Cast<UCONodeMaterialBreakParameterPinData>(PinData))
 	{
-		CastedPinData->VariableData.Name = FName{Value.ToString()};
+		FScopedTransaction LocalTransaction(LOCTEXT("UCONodeMaterialBreak_SetNamePinEditableNameTransaction", "Change Pin Name"));
+		CastedPinData->Modify();
+		
+		CastedPinData->ParameterName = FName{Value.ToString()};
 		PinEditableNameChangedDelegate.Broadcast(Pin, Value);
 	}
 	else
@@ -217,7 +216,7 @@ FText UCONodeMaterialBreak::GetPinEditableName(const UEdGraphPin& Pin) const
 	const UCustomizableObjectNodePinData* PinData = GetPinData(Pin);
 	if (const UCONodeMaterialBreakParameterPinData* CastedPinData = Cast<UCONodeMaterialBreakParameterPinData>(PinData))
 	{
-		return FText::FromName(CastedPinData->VariableData.Name);
+		return FText::FromName(CastedPinData->ParameterName);
 	}
 
 	return Super::GetPinEditableName(Pin);
@@ -226,7 +225,7 @@ FText UCONodeMaterialBreak::GetPinEditableName(const UEdGraphPin& Pin) const
 
 FName UCONodeMaterialBreak::GetPinParameterName(const UEdGraphPin& Pin) const
 {
-	return GetPinData<UCONodeMaterialBreakParameterPinData>(Pin).VariableData.Name;
+	return GetPinData<UCONodeMaterialBreakParameterPinData>(Pin).ParameterName;
 }
 
 
diff --git a/Engine/Plugins/Mutable/Source/CustomizableObjectEditor/Private/MuCOE/Nodes/CONodeMaterialBreak.h b/Engine/Plugins/Mutable/Source/CustomizableObjectEditor/Private/MuCOE/Nodes/CONodeMaterialBreak.h
--- a/Engine/Plugins/Mutable/Source/CustomizableObjectEditor/Private/MuCOE/Nodes/CONodeMaterialBreak.h
+++ b/Engine/Plugins/Mutable/Source/CustomizableObjectEditor/Private/MuCOE/Nodes/CONodeMaterialBreak.h
@@ -3,8 +3,6 @@
 #pragma once
 
 #include "MuCOE/Nodes/CustomizableObjectNode.h"
-#include "MuCOE/COVariable.h"
-#include "MuCOE/RemapPins/CustomizableObjectNodeRemapPinsByName.h"
 
 #include "CONodeMaterialBreak.generated.h"
 
@@ -14,31 +12,21 @@ namespace ENodeTitleType { enum Type : int; }
 class UCustomizableObjectNodeRemapPins;
 
 
-UCLASS()
-class UCustomizableObjectNodeBreakMaterialRemapPins: public UCustomizableObjectNodeRemapPinsByName
-{
-public:
-	GENERATED_BODY()
-
-	virtual bool Equal(const UCustomizableObjectNode& Node, const UEdGraphPin& OldPin, const UEdGraphPin& NewPin) const override;
-};
-
-
 UCLASS(MinimalAPI)
 class UCONodeMaterialBreakParameterPinData : public UCustomizableObjectNodePinData
 {
 	GENERATED_BODY()
 
 public:
-
-	UPROPERTY(EditAnywhere, Category = NoCategory)
-	FCOVariable VariableData;
-
+	
 	/** Layer index of this parameter inside a Layered Material.
 	-1 if the material does not contain layers.*/
 	UPROPERTY(EditAnywhere, Category = NoCategory)
 	int32 LayerIndex = -1;
-
+	
+	UPROPERTY()
+	FName ParameterName;
+	
 	/**
 	 * Handle change in the LayerIndex property
 	 * @param PropertyChangedEvent The Property Change Event
@@ -55,20 +43,21 @@ public:
 	GENERATED_BODY()
 
 	// Begin EdGraphNode interface
-	UE_API virtual FText GetNodeTitle(ENodeTitleType::Type TitleType) const override;
-	UE_API virtual FLinearColor GetNodeTitleColor() const override;
-	UE_API virtual FText GetTooltipText() const override;
+	virtual FText GetNodeTitle(ENodeTitleType::Type TitleType) const override;
+	virtual FLinearColor GetNodeTitleColor() const override;
+	virtual FText GetTooltipText() const override;
 
 	// UCustomizableObjectNode interface
-	UE_API virtual void AllocateDefaultPins(UCustomizableObjectNodeRemapPins* RemapPins) override;
-	UE_API virtual TSharedPtr<IDetailsView> CustomizePinDetails(const UEdGraphPin& Pin) const override;
-	UE_API virtual bool HasPinViewer() const override;
-	UE_API virtual EAddPinNodeButtonLocation GetAddPinButtonNodeSide() const override;
-	UE_API virtual void AddPinFromUI() override;
-	UE_API virtual bool CanPinBeRemoved(const UEdGraphPin& Pin) const override;
-	UE_API virtual UCustomizableObjectNodeRemapPins* CreateRemapPinsDefault() const override;
-	UE_API virtual TArray<FName> GetAllowedPinViewerCreationTypes() const override;
-	UE_API virtual bool IsExperimental() const override;
+	virtual void AllocateDefaultPins(UCustomizableObjectNodeRemapPins* RemapPins) override;
+	virtual TSharedPtr<IDetailsView> CustomizePinDetails(const UEdGraphPin& Pin) const override;
+	virtual bool HasPinViewer() const override;
+	virtual EAddPinNodeButtonLocation GetAddPinButtonNodeSide() const override;
+	virtual void AddPinFromUI() override;
+	virtual bool CanPinBeRemoved(const UEdGraphPin& Pin) const override;
+	virtual UCustomizableObjectNodeRemapPins* CreateRemapPinsDefault() const override;
+	virtual TArray<FName> GetAllowedPinViewerCreationTypes() const override;
+	virtual TArray<FName> GetPinAllowedTypes(const UEdGraphPin& Pin) const override;
+	virtual bool IsExperimental() const override;
 
 	// UCustomizableObjectNode Interface : Editable area of the pin name management
 	virtual bool CanRenamePin(const UEdGraphPin& Pin) const override;
diff --git a/Engine/Plugins/Mutable/Source/CustomizableObjectEditor/Private/MuCOE/Nodes/CONodeMaterialModify.cpp b/Engine/Plugins/Mutable/Source/CustomizableObjectEditor/Private/MuCOE/Nodes/CONodeMaterialModify.cpp
--- a/Engine/Plugins/Mutable/Source/CustomizableObjectEditor/Private/MuCOE/Nodes/CONodeMaterialModify.cpp
+++ b/Engine/Plugins/Mutable/Source/CustomizableObjectEditor/Private/MuCOE/Nodes/CONodeMaterialModify.cpp
@@ -6,6 +6,7 @@
 #include "MuCOE/EdGraphSchema_CustomizableObject.h"
 #include "PropertyEditorModule.h"
 #include "ScopedTransaction.h"
+#include "MuCOE/RemapPins/CustomizableObjectNodeRemapPinsByPosition.h"
 
 #define LOCTEXT_NAMESPACE "CustomizableObjectEditor"
 
@@ -14,22 +15,6 @@ const FName MaterialModifyMaterialParameterPinName = "Parameter";
 const FText MaterialModifyMaterialParameterPinText = LOCTEXT("Parameter", "Parameter");
 
 
-bool UCONodeMaterialModifyRemapPins::Equal(const UCustomizableObjectNode& Node, const UEdGraphPin& OldPin, const UEdGraphPin& NewPin) const
-{
-	const UCONodeMaterialModifyParameterPinData* OldPinData = Cast<UCONodeMaterialModifyParameterPinData>(Node.GetPinData(OldPin));
-	const UCONodeMaterialModifyParameterPinData* NewPinData = Cast<UCONodeMaterialModifyParameterPinData>(Node.GetPinData(NewPin));
-
-	if (OldPinData && NewPinData)
-	{
-		return OldPin.Direction == NewPin.Direction && OldPinData->VariableData.Type == NewPinData->VariableData.Type && OldPinData->VariableData.Id == NewPinData->VariableData.Id;
-	}
-	else
-	{
-		return Super::Equal(Node, OldPin, NewPin);
-	}
-}
-
-
 void UCONodeMaterialModifyParameterPinData::PostEditChangeProperty(FPropertyChangedEvent& PropertyChangedEvent)
 {
 	Super::PostEditChangeProperty(PropertyChangedEvent);
@@ -85,45 +70,39 @@ void UCONodeMaterialModify::AllocateDefaultPins(UCustomizableObjectNodeRemapPins
 		{
 			if (UCONodeMaterialModifyParameterPinData* OldPinData = Cast<UCONodeMaterialModifyParameterPinData>(GetPinData(*InputPin)))
 			{
-				UEdGraphPin* NewPin = nullptr;
 				UCONodeMaterialModifyParameterPinData* NewPinData = nullptr;
 				
-				if (OldPinData->VariableData.Type == UEdGraphSchema_CustomizableObject::PC_Float)
+				if (InputPin->PinType.PinCategory ==  UEdGraphSchema_CustomizableObject::PC_Float )
 				{
 					NewPinData = NewObject<UCONodeMaterialModifyScalarParamPinData>(this);
 				}
-				else if (OldPinData->VariableData.Type == UEdGraphSchema_CustomizableObject::PC_Texture)
+				else if (InputPin->PinType.PinCategory == UEdGraphSchema_CustomizableObject::PC_Texture )
 				{
 					NewPinData = NewObject<UCONodeMaterialModifyTextureParamPinData>(this);
 				}
-				else if (OldPinData->VariableData.Type == UEdGraphSchema_CustomizableObject::PC_Texture_Passthrough)
+				else if (InputPin->PinType.PinCategory == UEdGraphSchema_CustomizableObject::PC_Texture_Passthrough )
 				{
 					NewPinData = NewObject<UCONodeMaterialModifyPassthroughTextureParamPinData>(this);
 				}
-				else if (OldPinData->VariableData.Type == UEdGraphSchema_CustomizableObject::PC_Color)
+				else if (InputPin->PinType.PinCategory == UEdGraphSchema_CustomizableObject::PC_Color )
 				{
 					NewPinData = NewObject<UCONodeMaterialModifyColorParamPinData>(this);
 				}
 				
-				// A new Pin data structure has been created because one pin has changed its type (from the UI)
 				check(NewPinData)
 				if (NewPinData)
 				{
-					NewPinData->VariableData.Id = OldPinData->VariableData.Id;
-					NewPinData->VariableData.Name = OldPinData->VariableData.Name;
-					NewPinData->LayerIndex = OldPinData->LayerIndex;
-					
-					NewPin = CustomCreatePin(EGPD_Input, NewPinData->VariableData.Type, MaterialModifyMaterialParameterPinName, MaterialModifyMaterialParameterPinText, NewPinData);
-				}
+					UEdGraphPin* NewPin = CustomCreatePin(EGPD_Input, InputPin->PinType.PinCategory, MaterialModifyMaterialParameterPinName, MaterialModifyMaterialParameterPinText,  NewPinData);
+					check(NewPin)
 				
-				check(NewPin)
-				if (OldPinData->LayerIndex == 0)
-				{
-					NewPin->PinFriendlyName = FText::Format(LOCTEXT("BackgroundLayerText","{0} - Background"), MaterialModifyMaterialParameterPinText);
-				}
-				else if (OldPinData->LayerIndex > 0)
-				{
-					NewPin->PinFriendlyName = FText::Format(LOCTEXT("LayerText", "{0} - Layer {1}"), MaterialModifyMaterialParameterPinText, FText::FromString(FString::FromInt(OldPinData->LayerIndex)));
+					if (OldPinData->LayerIndex == 0)
+					{
+						NewPin->PinFriendlyName = FText::Format(LOCTEXT("BackgroundLayerText","{0} - Background"), FText::FromName(MaterialModifyMaterialParameterPinName));
+					}
+					else if (OldPinData->LayerIndex > 0)
+					{
+						NewPin->PinFriendlyName = FText::Format(LOCTEXT("LayerText", "{0} - Layer {1}"), FText::FromName(MaterialModifyMaterialParameterPinName), FText::FromString(FString::FromInt(OldPinData->LayerIndex)));
+					}
 				}
 			}
 		}
@@ -167,7 +146,7 @@ EAddPinNodeButtonLocation UCONodeMaterialModify::GetAddPinButtonNodeSide() const
 
 UCustomizableObjectNodeRemapPins* UCONodeMaterialModify::CreateRemapPinsDefault() const
 {
-	return NewObject<UCONodeMaterialModifyRemapPins>();
+	return NewObject<UCustomizableObjectNodeRemapPinsByPosition>();
 }
 
 
@@ -180,6 +159,22 @@ TArray<FName> UCONodeMaterialModify::GetAllowedPinViewerCreationTypes() const
 }
 
 
+TArray<FName> UCONodeMaterialModify::GetPinAllowedTypes(const UEdGraphPin& Pin) const
+{
+	if (Pin.Direction == EGPD_Input && FEdGraphPinReference{&Pin} != MaterialPinRef)
+	{
+		const TArray<FName> PinAllowedTypes = GetAllowedPinViewerCreationTypes();
+		check(PinAllowedTypes.Contains(Pin.PinType.PinCategory));
+		
+		return PinAllowedTypes;
+	}
+	else
+	{	
+		return Super::GetPinAllowedTypes(Pin);
+	}
+}
+
+
 void UCONodeMaterialModify::AddPinFromUI()
 {
 	Super::AddPinFromUI();
@@ -188,7 +183,8 @@ void UCONodeMaterialModify::AddPinFromUI()
 	Modify();
 	
 	UCONodeMaterialModifyScalarParamPinData* ScalarParameterPinData = NewObject<UCONodeMaterialModifyScalarParamPinData>(this);
-	if (CustomCreatePin(EEdGraphPinDirection::EGPD_Input, ScalarParameterPinData->VariableData.Type, MaterialModifyMaterialParameterPinName, MaterialModifyMaterialParameterPinText, ScalarParameterPinData))
+	if (UEdGraphPin* NewPin = CustomCreatePin(EEdGraphPinDirection::EGPD_Input, UEdGraphSchema_CustomizableObject::PC_Float,
+		MaterialModifyMaterialParameterPinName, MaterialModifyMaterialParameterPinText, ScalarParameterPinData))
 	{
 		GetGraph()->NotifyNodeChanged(this);
 	}
@@ -203,7 +199,7 @@ bool UCONodeMaterialModify::CanPinBeRemoved(const UEdGraphPin& Pin) const
 
 FName UCONodeMaterialModify::GetPinParameterName(const UEdGraphPin& Pin) const
 {
-	return GetPinData<UCONodeMaterialModifyParameterPinData>(Pin).VariableData.Name;
+	return GetPinData<UCONodeMaterialModifyParameterPinData>(Pin).ParameterName;
 }
 
 
@@ -277,7 +273,10 @@ void UCONodeMaterialModify::SetPinEditableName(const UEdGraphPin& Pin, const FTe
 	UCustomizableObjectNodePinData* PinData = GetPinData(Pin);
 	if (UCONodeMaterialModifyParameterPinData* CastedPinData = Cast<UCONodeMaterialModifyParameterPinData>(PinData))
 	{
-		CastedPinData->VariableData.Name = FName{Value.ToString()};
+		FScopedTransaction LocalTransaction(LOCTEXT("UCONodeMaterialModify_SetNamePinEditableNameTransaction", "Change Pin Name"));
+		CastedPinData->Modify();
+		
+		CastedPinData->ParameterName = FName{Value.ToString()};
 		PinEditableNameChangedDelegate.Broadcast(Pin, Value);
 	}
 	else
@@ -292,7 +291,7 @@ FText UCONodeMaterialModify::GetPinEditableName(const UEdGraphPin& Pin) const
 	const UCustomizableObjectNodePinData* PinData = GetPinData(Pin);
 	if (const UCONodeMaterialModifyParameterPinData* CastedPinData = Cast<UCONodeMaterialModifyParameterPinData>(PinData))
 	{
-		return FText::FromName(CastedPinData->VariableData.Name);
+		return FText::FromName(CastedPinData->ParameterName);
 	}
 
 	return Super::GetPinEditableName(Pin);
diff --git a/Engine/Plugins/Mutable/Source/CustomizableObjectEditor/Private/MuCOE/Nodes/CONodeMaterialModify.h b/Engine/Plugins/Mutable/Source/CustomizableObjectEditor/Private/MuCOE/Nodes/CONodeMaterialModify.h
--- a/Engine/Plugins/Mutable/Source/CustomizableObjectEditor/Private/MuCOE/Nodes/CONodeMaterialModify.h
+++ b/Engine/Plugins/Mutable/Source/CustomizableObjectEditor/Private/MuCOE/Nodes/CONodeMaterialModify.h
@@ -6,7 +6,6 @@
 #include "IDetailCustomization.h"
 #include "MuCOE/Nodes/CustomizableObjectNode.h"
 #include "MuCOE/RemapPins/CustomizableObjectNodeRemapPinsByName.h"
-#include "MuCOE/COVariable.h"
 #include "MuCOE/EdGraphSchema_CustomizableObject.h"
 #include "CONodeMaterialModify.generated.h"
 
@@ -14,15 +13,6 @@
 
 class UCustomizableObjectNodeRemapPinsByName;
 
-UCLASS()
-class UCONodeMaterialModifyRemapPins: public UCustomizableObjectNodeRemapPinsByName
-{
-public:
-	GENERATED_BODY()
-
-	virtual bool Equal(const UCustomizableObjectNode& Node, const UEdGraphPin& OldPin, const UEdGraphPin& NewPin) const override;
-};
-
 
 // todo: Remove this enum and use the "EUVLayoutMode" so we support the option of getting the UV layout from the material  
 /** Image Pin, UV Layout Mode. */
@@ -43,24 +33,30 @@ class UCONodeMaterialModifyParameterPinData : public UCustomizableObjectNodePinD
 	GENERATED_BODY()
 
 public:
-	UCONodeMaterialModifyParameterPinData()
-	{
-		VariableData.Id = FGuid::NewGuid();
-	}
-	
-	UPROPERTY(EditAnywhere, Category = NoCategory)
-	FCOVariable VariableData;
 	
 	/** Layer index of this parameter inside a Layered Material.
 	-1 if the material does not contain layers.*/
 	UPROPERTY(EditAnywhere, Category = NoCategory)
 	int32 LayerIndex = -1;
-
+	
+	UPROPERTY()
+	FName ParameterName;
+	
 	/**
 	 * Handle change in the LayerIndex property
 	 * @param PropertyChangedEvent The Property Change Event
 	 */
 	virtual void PostEditChangeProperty(FPropertyChangedEvent& PropertyChangedEvent) override;
+	
+
+	virtual void Copy(const UCustomizableObjectNodePinData& Other) override
+	{
+		if (const UCONodeMaterialModifyParameterPinData* PinDataOldPin = Cast<UCONodeMaterialModifyParameterPinData>(&Other))
+		{
+			LayerIndex = PinDataOldPin->LayerIndex;
+			ParameterName = PinDataOldPin->ParameterName;
+		}
+	}
 };
 
 
@@ -73,25 +69,19 @@ class UCONodeMaterialModifyTextureParamPinData : public UCONodeMaterialModifyPar
 	GENERATED_BODY()
 	
 public:
-	UCONodeMaterialModifyTextureParamPinData() 
-	{
-		VariableData.Type = UEdGraphSchema_CustomizableObject::PC_Texture;
-	}
 	
-
 	virtual void Copy(const UCustomizableObjectNodePinData& Other) override
 	{
+		Super::Copy(Other);
+		
 		if (const UCONodeMaterialModifyTextureParamPinData* PinDataOldPin = Cast<UCONodeMaterialModifyTextureParamPinData>(&Other))
 		{
-			VariableData = PinDataOldPin->VariableData;
-            
             UVLayoutMode = PinDataOldPin->UVLayoutMode;
             UVLayout = PinDataOldPin->UVLayout;
             ReferenceTexture = PinDataOldPin->ReferenceTexture;
 		}
 	}
 
-	
 	constexpr static int32 UV_LAYOUT_IGNORE = -1;
 	
 	UPROPERTY(EditAnywhere, Category = NoCategory)
@@ -115,12 +105,6 @@ UCLASS()
 class UCONodeMaterialModifyPassthroughTextureParamPinData : public UCONodeMaterialModifyParameterPinData
 {
 	GENERATED_BODY()
-	
-public:
-	UCONodeMaterialModifyPassthroughTextureParamPinData()
-	{
-		VariableData.Type = UEdGraphSchema_CustomizableObject::PC_Texture_Passthrough;
-	}
 };
 
 
@@ -131,12 +115,6 @@ UCLASS()
 class UCONodeMaterialModifyScalarParamPinData : public UCONodeMaterialModifyParameterPinData
 {
 	GENERATED_BODY()
-	
-public:
-	UCONodeMaterialModifyScalarParamPinData()
-	{
-		VariableData.Type = UEdGraphSchema_CustomizableObject::PC_Float;
-	}
 };
 
 
@@ -147,12 +125,6 @@ UCLASS()
 class UCONodeMaterialModifyColorParamPinData : public UCONodeMaterialModifyParameterPinData
 {
 	GENERATED_BODY()
-	
-public:
-	UCONodeMaterialModifyColorParamPinData()
-	{
-		VariableData.Type = UEdGraphSchema_CustomizableObject::PC_Color;
-	}
 };
 
 
@@ -177,6 +149,7 @@ public:
 	virtual bool CanPinBeRemoved(const UEdGraphPin& Pin) const override;
 	virtual UCustomizableObjectNodeRemapPins* CreateRemapPinsDefault() const override;
 	virtual TArray<FName> GetAllowedPinViewerCreationTypes() const override;
+	virtual TArray<FName> GetPinAllowedTypes(const UEdGraphPin& Pin) const override;
 	virtual bool IsExperimental() const override;
 
 	// Own Interface : Editable area of the pin name management
diff --git a/Engine/Plugins/Mutable/Source/CustomizableObjectEditor/Private/MuCOE/Nodes/CONodeSkeletalMeshBreak.cpp b/Engine/Plugins/Mutable/Source/CustomizableObjectEditor/Private/MuCOE/Nodes/CONodeSkeletalMeshBreak.cpp
--- a/Engine/Plugins/Mutable/Source/CustomizableObjectEditor/Private/MuCOE/Nodes/CONodeSkeletalMeshBreak.cpp
+++ b/Engine/Plugins/Mutable/Source/CustomizableObjectEditor/Private/MuCOE/Nodes/CONodeSkeletalMeshBreak.cpp
@@ -75,6 +75,9 @@ void UCONodeSkeletalMeshBreak::SetPinEditableName(const UEdGraphPin& Pin, const
 	UCustomizableObjectNodePinData* PinData = GetPinData(Pin);
 	if (UCONodeSkeletalMeshBreakPinData* CastedPinData = Cast<UCONodeSkeletalMeshBreakPinData>(PinData))
 	{
+		FScopedTransaction LocalTransaction(LOCTEXT("UCONodeSkeletalMeshBreak_SetNamePinEditableNameTransaction", "Change Pin Name"));
+		CastedPinData->Modify();
+		
 		CastedPinData->TargetMaterialSlotName = FName{Value.ToString()};
 		PinEditableNameChangedDelegate.Broadcast(Pin, Value);
 	}
diff --git a/Engine/Plugins/Mutable/Source/CustomizableObjectEditor/Private/MuCOE/Nodes/CustomizableObjectNode.cpp b/Engine/Plugins/Mutable/Source/CustomizableObjectEditor/Private/MuCOE/Nodes/CustomizableObjectNode.cpp
--- a/Engine/Plugins/Mutable/Source/CustomizableObjectEditor/Private/MuCOE/Nodes/CustomizableObjectNode.cpp
+++ b/Engine/Plugins/Mutable/Source/CustomizableObjectEditor/Private/MuCOE/Nodes/CustomizableObjectNode.cpp
@@ -669,7 +669,6 @@ int32 UCustomizableObjectNode::GetLOD() const
 void UCustomizableObjectNode::SetPinHidden(UEdGraphPin& Pin, bool bHidden)
 {
 	Pin.SafeSetHidden(bHidden && CanPinBeHidden(Pin));
-
 	GetGraph()->NotifyGraphChanged();
 }
 
@@ -808,6 +807,12 @@ TArray<FName> UCustomizableObjectNode::GetAllowedPinViewerCreationTypes() const
 }
 
 
+TArray<FName> UCustomizableObjectNode::GetPinAllowedTypes(const UEdGraphPin& Pin) const
+{
+	return {};
+}
+
+
 bool UCustomizableObjectNode::IsCompatibleWithGraph(const UEdGraph* TargetGraph) const
 {
 	if (const UCustomizableObjectGraph* CurrentGraph = Cast<UCustomizableObjectGraph>(TargetGraph))
diff --git a/Engine/Plugins/Mutable/Source/CustomizableObjectEditor/Private/MuCOE/PinViewer/SPinViewer.cpp b/Engine/Plugins/Mutable/Source/CustomizableObjectEditor/Private/MuCOE/PinViewer/SPinViewer.cpp
--- a/Engine/Plugins/Mutable/Source/CustomizableObjectEditor/Private/MuCOE/PinViewer/SPinViewer.cpp
+++ b/Engine/Plugins/Mutable/Source/CustomizableObjectEditor/Private/MuCOE/PinViewer/SPinViewer.cpp
@@ -37,19 +37,6 @@ FText SPinViewer::GetPinName(const UEdGraphPin& Pin)
 	{
 		PinComposedName =  Pin.PinName.ToString();
 	}
-
-	// Append Pin Editable Name
-	if (const UCustomizableObjectNode* OwningNode = Cast<UCustomizableObjectNode>(Pin.GetOwningNode()))
-	{
-		if (OwningNode->ShouldPinViewerShowPinEditableName(Pin))
-		{
-			const FString EditablePinName = OwningNode->GetPinEditableName(Pin).ToString();
-			if (!EditablePinName.IsEmpty())
-			{
-				PinComposedName = FString::Printf(TEXT("%s : %s"), *PinComposedName, *EditablePinName);
-			}
-		}
-	}
 	
 	return FText::FromString(PinComposedName);
 }
@@ -202,6 +189,8 @@ void SPinViewer::PinsRemapped(const TMap<UEdGraphPin*, UEdGraphPin*>& PinsToRema
 void SPinViewer::Construct(const FArguments& InArgs)
 {
 	Node = TStrongObjectPtr<UCustomizableObjectNode>{InArgs._Node};
+	check(Node)
+	
 	Node->RemapPinsDelegate.AddSP(this, &SPinViewer::PinsRemapped);
 	Node->NodeConnectionListChangedDelegate.AddSP(this, &SPinViewer::UpdateWidget);
 	Node->PinEditableNameChangedDelegate.AddSP(this, &SPinViewer::OnPinEditableNameChange);
diff --git a/Engine/Plugins/Mutable/Source/CustomizableObjectEditor/Private/MuCOE/PinViewer/SPinViewerListRow.cpp b/Engine/Plugins/Mutable/Source/CustomizableObjectEditor/Private/MuCOE/PinViewer/SPinViewerListRow.cpp
--- a/Engine/Plugins/Mutable/Source/CustomizableObjectEditor/Private/MuCOE/PinViewer/SPinViewerListRow.cpp
+++ b/Engine/Plugins/Mutable/Source/CustomizableObjectEditor/Private/MuCOE/PinViewer/SPinViewerListRow.cpp
@@ -9,6 +9,12 @@
 #include "Widgets/Input/SCheckBox.h"
 #include "Widgets/Views/SListView.h"
 #include "IDetailsView.h"
+#include "Widgets/Input/SComboBox.h"
+#include "Widgets/SWidget.h"
+#include "PropertyCustomizationHelpers.h"
+#include "DetailLayoutBuilder.h"
+#include "MuCOE/Widgets/SMutablePinTypeSelector.h"
+#include "Widgets/Input/SEditableTextBox.h"
 
 class STableViewBase;
 class SWidget;
@@ -20,7 +26,16 @@ void SPinViewerListRow::Construct(const FArguments& InArgs, const TSharedRef<STa
 {
 	PinViewer = InArgs._PinViewer;
 	PinReference = InArgs._PinReference;
-
+	
+	if (UCustomizableObjectNode* OwningNode = Cast<UCustomizableObjectNode>(PinReference.Get()->GetOwningNode()))
+	{
+		PinOwner = MakeWeakObjectPtr(OwningNode);
+	}
+	else
+	{
+		checkNoEntry()
+	}
+	
 	PinId = PinReference.Get()->PinId;
 	
 	SMutableExpandableTableRow<TSharedPtr<FEdGraphPinReference>>::Construct(FSuperRowType::FArguments(), OwnerTableView);
@@ -33,38 +48,100 @@ TSharedRef<SWidget> SPinViewerListRow::GenerateWidgetForColumn(const FName& Colu
 	
 	if (ColumnName == SPinViewer::COLUMN_NAME)
 	{
-		const TSharedRef<STextBlock> TextBlock = SNew(STextBlock)
-			.Text(SPinViewer::GetPinName(*Pin));
-		
+		const TSharedRef<STextBlock> PinNameTextBlock = SNew(STextBlock)
+		.Text(SPinViewer::GetPinName(*Pin));
+				
 		if (Pin->bOrphanedPin)
 		{
-			TextBlock->SetColorAndOpacity(FLinearColor::Red);
+			PinNameTextBlock->SetColorAndOpacity(FLinearColor::Red);
+		}
+		
+		if (TStrongObjectPtr<UCustomizableObjectNode> PinnedPinOwner = PinOwner.Pin())
+		{
+			if (PinnedPinOwner->CanRenamePin(*Pin))
+			{
+				return SNew(SVerticalBox)
+				+ SVerticalBox::Slot().AutoHeight().Padding(4).VAlign(VAlign_Center)
+					.HAlign(EHorizontalAlignment::HAlign_Left)
+				[
+					SNew(SHorizontalBox)
+
+					+ SHorizontalBox::Slot()
+					.AutoWidth()
+					.HAlign(EHorizontalAlignment::HAlign_Left)
+					.VAlign(VAlign_Center)
+					[
+						PinNameTextBlock
+					]
+
+					+ SHorizontalBox::Slot()
+					.HAlign(EHorizontalAlignment::HAlign_Left)
+					.VAlign(VAlign_Center)
+					.Padding(4, 0)
+					.AutoWidth()
+					[
+						SNew(SEditableTextBox)
+						.OnTextCommitted(this, &SPinViewerListRow::OnEditableNameCommited)
+						.Text(this, &SPinViewerListRow::GetEditableName)
+					]
+				];
+			}
+			else
+			{
+				return SNew(SVerticalBox)
+					+ SVerticalBox::Slot().AutoHeight().Padding(4).VAlign(VAlign_Center)
+					[
+						PinNameTextBlock
+					];
+			}
 		}
 		
-		return SNew(SVerticalBox)
-			+ SVerticalBox::Slot().AutoHeight().Padding(4).VAlign(VAlign_Center)
-			[
-				TextBlock
-			];
 	}
 	else if (ColumnName == SPinViewer::COLUMN_TYPE)
 	{
-		return SNew(SVerticalBox)
-			+ SVerticalBox::Slot().AutoHeight().Padding(4).VAlign(VAlign_Center).HAlign(HAlign_Center)
-			[
-				SNew(STextBlock).Text(UEdGraphSchema_CustomizableObject::GetPinCategoryFriendlyName(Pin->PinType.PinCategory))
-			];
+		if (TStrongObjectPtr<UCustomizableObjectNode> PinnedPinOwner = PinOwner.Pin())
+		{
+			if (PinnedPinOwner->GetPinAllowedTypes(*PinReference.Get()).Num() > 0)
+			{
+				return SNew(SVerticalBox)
+				
+				+ SVerticalBox::Slot()
+				.AutoHeight()
+				.Padding(4)
+				.VAlign(VAlign_Center)
+				.HAlign(HAlign_Left)
+				[
+					SNew(SMutablePinTypeSelector)
+					.PinReference(PinReference)
+				];
+			}
+			// If three is no more than one possible pin type for this pin just expose the type as a non-editable text
+			else
+			{
+				return SNew(SVerticalBox)
+					
+				+ SVerticalBox::Slot()
+				.AutoHeight().Padding(4)
+				.VAlign(VAlign_Center)
+				.HAlign(HAlign_Left)
+				[
+					SNew(STextBlock)
+					.Text(UEdGraphSchema_CustomizableObject::GetPinCategoryFriendlyName(GetInitiallySelectedPinType()))
+				];
+			}
+		}
 	}
 	else if (ColumnName == SPinViewer::COLUMN_VISIBILITY)
 	{
 		return SNew(SVerticalBox)
-			+ SVerticalBox::Slot().AutoHeight().Padding(4).VAlign(VAlign_Center).HAlign(HAlign_Center)
-			[
-				SNew(SCheckBox)
-				.OnCheckStateChanged(this, &SPinViewerListRow::OnPinVisibilityCheckStateChanged)
-				.IsChecked(this, &SPinViewerListRow::IsVisibilityChecked)
-				.IsEnabled(PinViewer->Node->CanPinBeHidden(*Pin))
-			];
+			
+		+ SVerticalBox::Slot().AutoHeight().Padding(4).VAlign(VAlign_Center).HAlign(HAlign_Center)
+		[
+			SNew(SCheckBox)
+			.OnCheckStateChanged(this, &SPinViewerListRow::OnPinVisibilityCheckStateChanged)
+			.IsChecked(this, &SPinViewerListRow::IsVisibilityChecked)
+			.IsEnabled(PinViewer->Node->CanPinBeHidden(*Pin))
+		];
 	}
 
 	return SNullWidget::NullWidget;
@@ -100,7 +177,7 @@ void SPinViewerListRow::SetAdditionalWidgetVisibility(const EVisibility InVisibi
 
 void SPinViewerListRow::OnPinVisibilityCheckStateChanged(ECheckBoxState NewRadioState)
 {
-	FScopedTransaction Transaction(LOCTEXT("ChangedPinVisiblilityTransaction", "Changed Pin Visiblility"));
+	FScopedTransaction Transaction(LOCTEXT("ChangedPinVisibilityTransaction", "Changed Pin Visibility"));
 	PinViewer->Node->Modify();
 	PinViewer->Node->SetPinHidden(*PinReference.Get(), NewRadioState == ECheckBoxState::Unchecked);
 }
@@ -112,4 +189,43 @@ ECheckBoxState SPinViewerListRow::IsVisibilityChecked() const
 }
 
 
+void SPinViewerListRow::OnEditableNameCommited(const FText& NewPinEditableName, ETextCommit::Type Arg) const
+{
+	if (TStrongObjectPtr<UCustomizableObjectNode> PinnedOwner = PinOwner.Pin())
+	{
+		if (const UEdGraphPin* Pin = PinReference.Get())
+		{
+			PinnedOwner->SetPinEditableName(*Pin, NewPinEditableName);
+		}
+	}
+}
+
+
+FText SPinViewerListRow::GetEditableName() const
+{
+	if (TStrongObjectPtr<UCustomizableObjectNode> PinnedOwner = PinOwner.Pin())
+	{
+		if (const UEdGraphPin* Pin = PinReference.Get())
+		{
+			return PinnedOwner->GetPinEditableName(*Pin);
+		}
+	}
+	
+	return {};
+}
+
+
+FName SPinViewerListRow::GetInitiallySelectedPinType() const
+{
+	if (UEdGraphPin* PinPointer = PinReference.Get())
+	{
+		return PinPointer->PinType.PinCategory;
+	}
+	else 
+	{
+		return NAME_Name;
+	}
+}
+
+
 #undef LOCTEXT_NAMESPACE
diff --git a/Engine/Plugins/Mutable/Source/CustomizableObjectEditor/Private/MuCOE/PinViewer/SPinViewerListRow.h b/Engine/Plugins/Mutable/Source/CustomizableObjectEditor/Private/MuCOE/PinViewer/SPinViewerListRow.h
--- a/Engine/Plugins/Mutable/Source/CustomizableObjectEditor/Private/MuCOE/PinViewer/SPinViewerListRow.h
+++ b/Engine/Plugins/Mutable/Source/CustomizableObjectEditor/Private/MuCOE/PinViewer/SPinViewerListRow.h
@@ -10,6 +10,7 @@
 class SPinViewer;
 class STableViewBase;
 class SWidget;
+class UCustomizableObjectNode;
 
 /** Custom row for a SListView that allows to add custom widgets as a row element. */
 class SPinViewerListRow : public SMutableExpandableTableRow<TSharedPtr<FEdGraphPinReference>>
@@ -35,9 +36,19 @@ public:
 	/** Returns the state of the check box in function of the visibility of the pin */
 	ECheckBoxState IsVisibilityChecked() const;
 
+	void OnEditableNameCommited(const FText& NewPinEditableName, ETextCommit::Type Arg) const;
+	FText GetEditableName() const;
+	
+
 private:
+	
+	FName GetInitiallySelectedPinType() const;
+	
 	SPinViewer* PinViewer = nullptr;
 	
+	/** Weak pointer to the Customizable Object node that contains this property */
+	TWeakObjectPtr<UCustomizableObjectNode> PinOwner;
+	
 	/** Pin representing this row. */
 	FEdGraphPinReference PinReference;
 	FGuid PinId; // Temporal performance fix. Avoids using the slow FEdPinReference hash function.
diff --git a/Engine/Plugins/Mutable/Source/CustomizableObjectEditor/Private/MuCOE/Widgets/SMutablePinTypeSelector.cpp b/Engine/Plugins/Mutable/Source/CustomizableObjectEditor/Private/MuCOE/Widgets/SMutablePinTypeSelector.cpp
new file mode 100644
--- /dev/null
+++ b/Engine/Plugins/Mutable/Source/CustomizableObjectEditor/Private/MuCOE/Widgets/SMutablePinTypeSelector.cpp
@@ -0,0 +1,112 @@
+// Copyright Epic Games, Inc. All Rights Reserved.
+
+#include "MuCOE/Widgets/SMutablePinTypeSelector.h"
+
+#include "DetailLayoutBuilder.h"
+#include "MuCOE/Nodes/CustomizableObjectNode.h"
+#include "MuCOE/EdGraphSchema_CustomizableObject.h"
+
+
+#define LOCTEXT_NAMESPACE "CustomizableObjectEditor"
+
+
+void SMutablePinTypeSelector::Construct(const FArguments& InArgs)
+{
+	PinReference = InArgs._PinReference;
+
+	UEdGraphPin* Pin = PinReference.Get();
+	check(Pin);
+	
+	if (UCustomizableObjectNode* OwningNode = Cast<UCustomizableObjectNode>(Pin->GetOwningNode()))
+	{
+		AllowedPinTypes = OwningNode->GetPinAllowedTypes(*Pin);
+		PinOwner = MakeWeakObjectPtr(OwningNode);
+	}
+	else
+	{
+		checkNoEntry();
+	}
+	
+	ChildSlot
+	[
+		SNew(SComboBox<FName>)
+		.OptionsSource(&AllowedPinTypes)
+		.InitiallySelectedItem(GetInitiallySelectedPinType())
+		.IsEnabled(this, &SMutablePinTypeSelector::ShouldVariableTypePropertyBeEnabled)
+		.OnGenerateWidget(this, &SMutablePinTypeSelector::OnGenerateVariableTypeRow)
+		.OnSelectionChanged(this, &SMutablePinTypeSelector::OnTypeDropdownSelectionChange)
+		.Content()
+		[
+			GenerateCurrentSelectedTypeWidget()
+		]
+	];
+}
+
+
+FName SMutablePinTypeSelector::GetInitiallySelectedPinType() const
+{
+	if (UEdGraphPin* PinPointer = PinReference.Get())
+	{
+		return PinPointer->PinType.PinCategory;
+	}
+	else 
+	{
+		return NAME_Name;
+	}
+}
+
+
+bool SMutablePinTypeSelector::ShouldVariableTypePropertyBeEnabled() const
+{
+	return PinReference.Get()->LinkedTo.Num() == 0;
+}
+
+
+TSharedRef<SWidget> SMutablePinTypeSelector::OnGenerateVariableTypeRow(FName Type)
+{
+	const FSlateBrush* IconBrush = FAppStyle::GetBrush(TEXT("Kismet.VariableList.TypeIcon"));
+
+	TSharedPtr<SHorizontalBox> RowWidget = SNew(SHorizontalBox)
+		+ SHorizontalBox::Slot()
+		.AutoWidth()
+		[
+			SNew(SImage)
+				.Image(IconBrush)
+				.ColorAndOpacity(UEdGraphSchema_CustomizableObject::GetPinTypeColor(Type))
+		]
+		+ SHorizontalBox::Slot()
+		.AutoWidth()
+		.Padding(7.5f, 0.0f, 0.0f, 0.0f)
+		[
+			SNew(STextBlock)
+				.Text(Type.IsValid() ? GetDefault<UEdGraphSchema_CustomizableObject>()->GetPinCategoryFriendlyName(Type) : FText::FromString("Invalid"))
+				.Font(IDetailLayoutBuilder::GetDetailFont())
+		];
+
+	return RowWidget.ToSharedRef();
+}
+
+
+void SMutablePinTypeSelector::OnTypeDropdownSelectionChange(FName NewTypeName, ESelectInfo::Type Arg) const
+{
+	if (UEdGraphPin* Pin = PinReference.Get())
+	{
+		Pin->PinType.PinCategory = NewTypeName;
+		
+		if (TStrongObjectPtr<UCustomizableObjectNode> Node = PinOwner.Pin())
+		{
+			Node->ReconstructNode();
+		}
+	}
+}
+
+
+TSharedRef<SWidget> SMutablePinTypeSelector::GenerateCurrentSelectedTypeWidget()
+{
+	const FName InitiallySelectedPinType = GetInitiallySelectedPinType();
+	check(AllowedPinTypes.Contains(InitiallySelectedPinType))
+	
+	return OnGenerateVariableTypeRow(InitiallySelectedPinType);
+}
+
+#undef LOCTEXT_NAMESPACE
diff --git a/Engine/Plugins/Mutable/Source/CustomizableObjectEditor/Private/MuCOE/Widgets/SMutablePinTypeSelector.h b/Engine/Plugins/Mutable/Source/CustomizableObjectEditor/Private/MuCOE/Widgets/SMutablePinTypeSelector.h
new file mode 100644
--- /dev/null
+++ b/Engine/Plugins/Mutable/Source/CustomizableObjectEditor/Private/MuCOE/Widgets/SMutablePinTypeSelector.h
@@ -0,0 +1,40 @@
+// Copyright Epic Games, Inc. All Rights Reserved.
+
+#pragma once
+
+#include "GraphEditAction.h"
+#include "EdGraph/EdGraphPin.h"
+#include "Widgets/Input/SComboBox.h"
+
+template <typename ItemType> class SListView;
+
+class UCustomizableObjectNode;
+struct FEdGraphPinReference;
+struct FGuid;
+
+
+class SMutablePinTypeSelector : public SCompoundWidget
+{
+public:
+	SLATE_BEGIN_ARGS(SMutablePinTypeSelector) {}
+	SLATE_ARGUMENT(FEdGraphPinReference, PinReference)
+	SLATE_END_ARGS()
+
+	void Construct(const FArguments& InArgs);
+	
+private:
+	
+	FName GetInitiallySelectedPinType() const;
+	bool ShouldVariableTypePropertyBeEnabled() const;
+	TSharedRef<SWidget> OnGenerateVariableTypeRow(FName Type);
+	void OnTypeDropdownSelectionChange(FName NewTypeName, ESelectInfo::Type Arg) const;
+	TSharedRef<SWidget> GenerateCurrentSelectedTypeWidget();
+
+	
+	FEdGraphPinReference PinReference;
+
+	TArray<FName> AllowedPinTypes;
+	
+	/** Weak pointer to the Customizable Object node that contains the pin */
+	TWeakObjectPtr<UCustomizableObjectNode> PinOwner;
+};
\ No newline at end of file
diff --git a/Engine/Plugins/Mutable/Source/CustomizableObjectEditor/Public/MuCOE/Nodes/CustomizableObjectNode.h b/Engine/Plugins/Mutable/Source/CustomizableObjectEditor/Public/MuCOE/Nodes/CustomizableObjectNode.h
--- a/Engine/Plugins/Mutable/Source/CustomizableObjectEditor/Public/MuCOE/Nodes/CustomizableObjectNode.h
+++ b/Engine/Plugins/Mutable/Source/CustomizableObjectEditor/Public/MuCOE/Nodes/CustomizableObjectNode.h
@@ -351,6 +351,13 @@ public:
 	// Returns an array with all the PC_Types that can be created from the pin viewer
 	UE_API virtual TArray<FName> GetAllowedPinViewerCreationTypes() const;
 
+	/**
+	 * Get the types the provided pin can changed to. 
+	 * @param Pin The pin you whant to know the types it can be changed to. 
+	 * @return A list of pin category names.
+	 */
+	UE_API virtual TArray<FName> GetPinAllowedTypes(const UEdGraphPin& Pin) const;
+
 	UE_API virtual bool IsCompatibleWithGraph(UEdGraph const* Graph) const override;
 	
 	UPROPERTY()
