[Mutable] Replace Material hash with FMaterialId.

#Synced-CL 49201640

diff --git a/Engine/Plugins/Mutable/Source/CustomizableObject/Internal/MuCO/CustomizableObjectSystemPrivate.h b/Engine/Plugins/Mutable/Source/CustomizableObject/Internal/MuCO/CustomizableObjectSystemPrivate.h
--- a/Engine/Plugins/Mutable/Source/CustomizableObject/Internal/MuCO/CustomizableObjectSystemPrivate.h
+++ b/Engine/Plugins/Mutable/Source/CustomizableObject/Internal/MuCO/CustomizableObjectSystemPrivate.h
@@ -500,9 +500,6 @@ struct FInstanceUpdateData
 };
 
 
-uint32 GetMaterialHash(const FInstanceUpdateData& UpdateData, const FInstanceUpdateData::FMaterial& Material, bool bPersistent);
-
-
 struct FComponentDescriptor
 {
 	TArray<UE::Mutable::Private::FMeshId> MeshDescriptors;
@@ -713,7 +710,7 @@ public:
 	TArray<UE::Mutable::Private::FImageId> CachedTextures;
 
 	/** Materials already build in this update. */
-	TMap<uint32, TStrongObjectPtr<UMaterialInterface>> MaterialUpdateCache; // TODO MATERIAL_PIN replace key with FMaterialId once MaterialPin completed.
+	TMap<UE::Mutable::Private::FMaterialId, TStrongObjectPtr<UMaterialInterface>> MaterialUpdateCache;
 
 	/** Textures already build in this update. */
 	TMap<FTextureCache::FId, TStrongObjectPtr<UTexture>> TextureUpdateCache;
diff --git a/Engine/Plugins/Mutable/Source/CustomizableObject/Private/MuCO/CustomizableObjectInstance.cpp b/Engine/Plugins/Mutable/Source/CustomizableObject/Private/MuCO/CustomizableObjectInstance.cpp
--- a/Engine/Plugins/Mutable/Source/CustomizableObject/Private/MuCO/CustomizableObjectInstance.cpp
+++ b/Engine/Plugins/Mutable/Source/CustomizableObject/Private/MuCO/CustomizableObjectInstance.cpp
@@ -5793,8 +5793,7 @@ FBuildMaterialOut BuildMaterial(const TSharedRef<FUpdateContextPrivate>& Context
 
 	const TSharedRef<const FInstanceUpdateData::FMaterial>& Material = Context->InstanceUpdateData.Materials[MaterialIndex];
 
-	const uint32 MaterialHash = GetMaterialHash(Context->InstanceUpdateData, Material.Get(), Context->bBake); // TODO MATERIAL_PIN Replace with FMaterialId once MaterialPin completed
-	TStrongObjectPtr<UMaterialInterface>* CachedMaterial = Context->MaterialUpdateCache.Find(MaterialHash);
+	TStrongObjectPtr<UMaterialInterface>* CachedMaterial = Context->MaterialUpdateCache.Find(Material->MaterialId);
 	if (CachedMaterial)
 	{
 		Out.Material = CachedMaterial->Get();
@@ -5851,7 +5850,7 @@ FBuildMaterialOut BuildMaterial(const TSharedRef<FUpdateContextPrivate>& Context
 				MIDName = MIDName.RightChop(MaterialInstanceConstantPrefix.Len());
 			}
 
-			MIDName = FString::Printf(TEXT("MID_%s_h%lu"), *MIDName, MaterialHash);
+			MIDName = FString::Printf(TEXT("MID_%s_h%lu"), *MIDName, UE::Mutable::Private::GetTypeHashPersistent(Material->MaterialId));
 
 			// Add a "tag" to better identify this asset during the bake
 			if (Context->bBake)
@@ -5954,7 +5953,7 @@ FBuildMaterialOut BuildMaterial(const TSharedRef<FUpdateContextPrivate>& Context
 		}
 	}
 
-	Context->MaterialUpdateCache.Add(MaterialHash, TStrongObjectPtr(MaterialInterface));
+	Context->MaterialUpdateCache.Add(Material->MaterialId, TStrongObjectPtr(MaterialInterface));
 
 	Out.Material = MaterialInterface;
 	
diff --git a/Engine/Plugins/Mutable/Source/CustomizableObject/Private/MuCO/CustomizableObjectSystem.cpp b/Engine/Plugins/Mutable/Source/CustomizableObject/Private/MuCO/CustomizableObjectSystem.cpp
--- a/Engine/Plugins/Mutable/Source/CustomizableObject/Private/MuCO/CustomizableObjectSystem.cpp
+++ b/Engine/Plugins/Mutable/Source/CustomizableObject/Private/MuCO/CustomizableObjectSystem.cpp
@@ -752,41 +752,6 @@ void FMutablePendingInstanceWork::AddIDRelease(UE::Mutable::Private::FInstance::
 }
 
 
-uint32 GetMaterialHash(const FInstanceUpdateData& UpdateData, const FInstanceUpdateData::FMaterial& Material, bool bPersistent)
-{
-	uint32 Hash = bPersistent ? UE::Mutable::Private::GetTypeHashPersistent(Material.MaterialId) : GetTypeHash(Material.MaterialId);
-	
-	for (int32 Index = Material.FirstScalar; Index < Material.FirstScalar + Material.ScalarCount; ++Index)
-	{
-		const FInstanceUpdateData::FScalar& Scalar = UpdateData.Scalars[Index];
-		
-		Hash = HashCombine(Hash, GetTypeHash(Scalar.Name));
-		Hash = HashCombine(Hash, GetTypeHash(Scalar.Scalar));
-		Hash = HashCombine(Hash, Scalar.MaterialLayer);
-	}
-
-	for (int32 Index = Material.FirstVector; Index < Material.FirstVector + Material.VectorCount; ++Index)
-	{
-		const FInstanceUpdateData::FVector& Vector = UpdateData.Vectors[Index];
-		
-		Hash = HashCombine(Hash, GetTypeHash(Vector.Name));
-		Hash = HashCombine(Hash, GetTypeHash(Vector.Vector));
-		Hash = HashCombine(Hash, Vector.MaterialLayer);
-	}
-
-	for (int32 Index = Material.FirstImage; Index < Material.FirstImage + Material.ImageCount; ++Index)
-	{
-		const TSharedRef<FInstanceUpdateData::FImage>& Image = UpdateData.Images[Index];
-		
-		Hash = HashCombine(Hash, GetTypeHash(Image->ParameterName));
-		Hash = HashCombine(Hash, GetTypeHash(Image->ImageID));
-		Hash = HashCombine(Hash, Image->MaterialLayer);
-	}
-
-	return Hash;
-}
-
-
 UCustomizableObjectSystem* UCustomizableObjectSystem::GetInstance()
 {
 	if (!UCustomizableObjectSystemPrivate::SSystem)
