[mutable] Fixed case where an error during the COI update makes the COI still eligible for being updated each time the ticker gets ticked.
- This only affects the "LOD" updates and not the standard "COI" updates. This will prevent the LOD update of a COI whose SkeletalMeshStatus has been resolved to Error.

#prefligh 693fe14e0d0fa241057ee086
#rnx
#Synced-CL 49271625

diff --git a/Engine/Plugins/Mutable/Source/CustomizableObject/Private/MuCO/CustomizableObjectInstance.cpp b/Engine/Plugins/Mutable/Source/CustomizableObject/Private/MuCO/CustomizableObjectInstance.cpp
--- a/Engine/Plugins/Mutable/Source/CustomizableObject/Private/MuCO/CustomizableObjectInstance.cpp
+++ b/Engine/Plugins/Mutable/Source/CustomizableObject/Private/MuCO/CustomizableObjectInstance.cpp
@@ -932,6 +932,11 @@ void UCustomizableInstancePrivate::TickUpdateCloseCustomizableObjects(UCustomiza
 		return;
 	}
 
+	if (SkeletalMeshStatus == ESkeletalMeshStatus::Error)
+	{
+		return;
+	}
+
 	const UCustomizableObjectSystemPrivate* SystemPrivate = UCustomizableObjectSystem::GetInstance()->GetPrivate();	
 
 	const EUpdateRequired UpdateRequired = SystemPrivate->IsUpdateRequired(Public, true, true, false);
