[Mutable] Fix crash with hacky Mutable streaming.

diff --git a/Engine/Plugins/Mutable/Source/CustomizableObject/Private/MuCO/UnrealConversionUtils.cpp b/Engine/Plugins/Mutable/Source/CustomizableObject/Private/MuCO/UnrealConversionUtils.cpp
--- a/Engine/Plugins/Mutable/Source/CustomizableObject/Private/MuCO/UnrealConversionUtils.cpp
+++ b/Engine/Plugins/Mutable/Source/CustomizableObject/Private/MuCO/UnrealConversionUtils.cpp
@@ -777,9 +777,13 @@ namespace UnrealConversionUtils
 			 FMemory::Memcpy(DestStaticVertexBuffer.PositionVertexBuffer.GetVertexData(), SrcStaticVertexBuffer.PositionVertexBuffer.GetVertexData(), NumVertices * DestStaticVertexBuffer.PositionVertexBuffer.GetStride());
 
 			 // Tangent and Texture coords buffers
-			 DestStaticVertexBuffer.StaticMeshVertexBuffer.SetUseFullPrecisionUVs(true);
-			 DestStaticVertexBuffer.StaticMeshVertexBuffer.SetUseHighPrecisionTangentBasis(false);
-			 DestStaticVertexBuffer.StaticMeshVertexBuffer.Init(NumVertices, NumTexCoords, bAllowCPUAccess);
+	     	 const bool bUseFullPrecisionUVs = SrcStaticVertexBuffer.StaticMeshVertexBuffer.GetUseFullPrecisionUVs();
+			 DestStaticVertexBuffer.StaticMeshVertexBuffer.SetUseFullPrecisionUVs(bUseFullPrecisionUVs);
+		 	
+	         const bool bUseHighPrecisionTangentBasis =  SrcStaticVertexBuffer.StaticMeshVertexBuffer.GetUseHighPrecisionTangentBasis();
+			 DestStaticVertexBuffer.StaticMeshVertexBuffer.SetUseHighPrecisionTangentBasis(bUseHighPrecisionTangentBasis);
+			 
+		 	 DestStaticVertexBuffer.StaticMeshVertexBuffer.Init(NumVertices, NumTexCoords, bAllowCPUAccess);
 			 FMemory::Memcpy(DestStaticVertexBuffer.StaticMeshVertexBuffer.GetTangentData(), SrcStaticVertexBuffer.StaticMeshVertexBuffer.GetTangentData(), DestStaticVertexBuffer.StaticMeshVertexBuffer.GetTangentSize());
 			 FMemory::Memcpy(DestStaticVertexBuffer.StaticMeshVertexBuffer.GetTexCoordData(), SrcStaticVertexBuffer.StaticMeshVertexBuffer.GetTexCoordData(), DestStaticVertexBuffer.StaticMeshVertexBuffer.GetTexCoordSize());
 
