[mutable] Fixed incomplete filtering of Pins at the time of compiling the MaterialModify node.
- The non-input pins were not being filtered, which caused the crash when trying to access data that the given node did not have.

#rnx
#Synced-CL 49273136

diff --git a/Engine/Plugins/Mutable/Source/CustomizableObjectEditor/Private/MuCOE/GenerateMutableSource/GenerateMutableSourceMaterial.cpp b/Engine/Plugins/Mutable/Source/CustomizableObjectEditor/Private/MuCOE/GenerateMutableSource/GenerateMutableSourceMaterial.cpp
--- a/Engine/Plugins/Mutable/Source/CustomizableObjectEditor/Private/MuCOE/GenerateMutableSource/GenerateMutableSourceMaterial.cpp
+++ b/Engine/Plugins/Mutable/Source/CustomizableObjectEditor/Private/MuCOE/GenerateMutableSource/GenerateMutableSourceMaterial.cpp
@@ -758,31 +758,36 @@ UE::Mutable::Private::Ptr<UE::Mutable::Private::NodeMaterial> GenerateMutableSou
 		// Tracking the index of the image. Needed for the FGeneratedImagePropertyKey
 		int32 ImageIndex = 0;
 	
-		for (const UEdGraphPin* InputPin : TypedNodeMaterialModify->GetAllNonOrphanPins())
+		for (const UEdGraphPin* NonOrphanPin : TypedNodeMaterialModify->GetAllNonOrphanPins())
 		{
-			if (InputPin == TypedNodeMaterialModify->MaterialPinRef.Get())
+			if (NonOrphanPin->Direction == EEdGraphPinDirection::EGPD_Output)
+			{
+				continue;
+			}
+			
+			if (NonOrphanPin == TypedNodeMaterialModify->MaterialPinRef.Get())
 			{
 				continue;
 			}
 
-			const FName ParameterName = TypedNodeMaterialModify->GetPinParameterName(*InputPin);
+			const FName ParameterName = TypedNodeMaterialModify->GetPinParameterName(*NonOrphanPin);
 			if (ParameterName.IsNone())
 			{
 				GenerationContext.Log(LOCTEXT("NoneParameterName", "'None' is not a valid Material Parameter name."), Node, EMessageSeverity::Warning);
 				continue;
 			}
 			
-			const int32 LayerIndex = TypedNodeMaterialModify->GetPinParameterLayerIndex(*InputPin);
+			const int32 LayerIndex = TypedNodeMaterialModify->GetPinParameterLayerIndex(*NonOrphanPin);
 			UE::Mutable::Private::FParameterKey ParameterKey = { ParameterName, LayerIndex };
 
-			const FName PinType = InputPin->PinType.PinCategory;
+			const FName PinType = NonOrphanPin->PinType.PinCategory;
 			if (PinType == UEdGraphSchema_CustomizableObject::PC_Texture)
 			{
 				UE::Mutable::Private::FImageParameterData ImageData;
 
-				if (const UEdGraphPin* FollowPin = FollowInputPin(*InputPin))
+				if (const UEdGraphPin* FollowPin = FollowInputPin(*NonOrphanPin))
 				{
-					TObjectPtr<UTexture2D> ReferenceTexture = GenerationContext.LoadObject(TypedNodeMaterialModify->GetTexturePinParameterReferenceTexture(*InputPin));
+					TObjectPtr<UTexture2D> ReferenceTexture = GenerationContext.LoadObject(TypedNodeMaterialModify->GetTexturePinParameterReferenceTexture(*NonOrphanPin));
 					if (!ReferenceTexture)
 					{
 						// Fallback if the user did not provide a reference texture
@@ -798,7 +803,7 @@ UE::Mutable::Private::Ptr<UE::Mutable::Private::NodeMaterial> GenerateMutableSou
 					// Fill Image Data
 					ImageData.ImageNode = ImageNode;
 					ImageData.ImagePropertyIndex = Props.ImagePropertiesIndex;
-					ImageData.LayoutIndex = TypedNodeMaterialModify->GetImagePinParameterUVIndex(*InputPin);
+					ImageData.LayoutIndex = TypedNodeMaterialModify->GetImagePinParameterUVIndex(*NonOrphanPin);
 				}
 
 				// We do not check Null images since it is possible to not have anything linked
@@ -811,7 +816,7 @@ UE::Mutable::Private::Ptr<UE::Mutable::Private::NodeMaterial> GenerateMutableSou
 			{
 				UE::Mutable::Private::FImageParameterData ImageData;
 
-				if (const UEdGraphPin* ConnectedPin = FollowInputPin(*InputPin))
+				if (const UEdGraphPin* ConnectedPin = FollowInputPin(*NonOrphanPin))
 				{
 					UE::Mutable::Private::FParameterKey ImageKey = { ParameterName, LayerIndex };
 					FGeneratedImageProperties Props = GenerateImageProperties(TypedNodeMaterialBase, nullptr, ImageKey, true, GenerationContext);
@@ -836,7 +841,7 @@ UE::Mutable::Private::Ptr<UE::Mutable::Private::NodeMaterial> GenerateMutableSou
 			{
 				UE::Mutable::Private::Ptr<UE::Mutable::Private::NodeColour> ColorNode;
 
-				if (const UEdGraphPin* FollowPin = FollowInputPin(*InputPin))
+				if (const UEdGraphPin* FollowPin = FollowInputPin(*NonOrphanPin))
 				{
 					ColorNode = GenerateMutableSourceColor(FollowPin, GenerationContext);
 				}
@@ -849,7 +854,7 @@ UE::Mutable::Private::Ptr<UE::Mutable::Private::NodeMaterial> GenerateMutableSou
 			{
 				UE::Mutable::Private::Ptr<UE::Mutable::Private::NodeScalar> ScalarNode;
 
-				if (const UEdGraphPin* FollowPin = FollowInputPin(*InputPin))
+				if (const UEdGraphPin* FollowPin = FollowInputPin(*NonOrphanPin))
 				{
 					ScalarNode = GenerateMutableSourceFloat(FollowPin, GenerationContext);
 				}
