[Mutable] Empty the arrays of generated textures and materials before adding new materials during the update.

#rnx
#Synced-CL 49281537

diff --git a/Engine/Plugins/Mutable/Source/CustomizableObject/Internal/MuCO/CustomizableObjectInstancePrivate.h b/Engine/Plugins/Mutable/Source/CustomizableObject/Internal/MuCO/CustomizableObjectInstancePrivate.h
--- a/Engine/Plugins/Mutable/Source/CustomizableObject/Internal/MuCO/CustomizableObjectInstancePrivate.h
+++ b/Engine/Plugins/Mutable/Source/CustomizableObject/Internal/MuCO/CustomizableObjectInstancePrivate.h
@@ -28,8 +28,6 @@ struct FMutableRefSkeletalMeshData;
 struct FMutableImageCacheKey;
 struct FStreamableManager;
 struct FStreamableHandle;
-struct FGeneratedMaterial;
-struct FGeneratedTexture;
 class UPhysicsAsset;
 class USkeleton;
 class UCustomizableObjectExtension;
@@ -271,7 +269,7 @@ public:
 	void ClearCOInstanceFlags(ECOInstanceFlags FlagsToClear) { InstanceFlagsPrivate = (ECOInstanceFlags)(InstanceFlagsPrivate & ~FlagsToClear); }
 	bool HasCOInstanceFlags(ECOInstanceFlags FlagsToCheck) const { return (InstanceFlagsPrivate & FlagsToCheck) != 0; }
 
-	UE_API void BuildMaterials(const TSharedRef<FUpdateContextPrivate>& Context, UCustomizableObjectInstance* Public);
+	UE_API void BuildResources(const TSharedRef<FUpdateContextPrivate>& Context, UCustomizableObjectInstance* Public);
 	
 	UE_API void TickUpdateCloseCustomizableObjects(UCustomizableObjectInstance& Public, FMutableInstanceUpdateMap& InOutRequestedUpdates);
 	UE_API void UpdateInstanceIfNotGenerated(UCustomizableObjectInstance& Public, FMutableInstanceUpdateMap& InOutRequestedUpdates);
diff --git a/Engine/Plugins/Mutable/Source/CustomizableObject/Internal/MuCO/CustomizableObjectSystemPrivate.h b/Engine/Plugins/Mutable/Source/CustomizableObject/Internal/MuCO/CustomizableObjectSystemPrivate.h
--- a/Engine/Plugins/Mutable/Source/CustomizableObject/Internal/MuCO/CustomizableObjectSystemPrivate.h
+++ b/Engine/Plugins/Mutable/Source/CustomizableObject/Internal/MuCO/CustomizableObjectSystemPrivate.h
@@ -704,7 +704,7 @@ public:
 	TArray<TStrongObjectPtr<const UObject>> Objects;
 
 	// Used during an update to prevent the pass-through textures loaded by LoadAdditionalAssetsAsync() from being unloaded by GC
-	// between AdditionalAssetsAsyncLoaded() and their setting into the generated materials in BuildMaterials()
+	// between AdditionalAssetsAsyncLoaded() and their setting into the generated materials in BuildResources()
 	TArray<TStrongObjectPtr<const UTexture>> LoadedPassThroughTexturesPendingSetMaterial;
 	
 	TArray<UE::Mutable::Private::FImageId> CachedTextures;
diff --git a/Engine/Plugins/Mutable/Source/CustomizableObject/Private/MuCO/CustomizableObjectInstance.cpp b/Engine/Plugins/Mutable/Source/CustomizableObject/Private/MuCO/CustomizableObjectInstance.cpp
--- a/Engine/Plugins/Mutable/Source/CustomizableObject/Private/MuCO/CustomizableObjectInstance.cpp
+++ b/Engine/Plugins/Mutable/Source/CustomizableObject/Private/MuCO/CustomizableObjectInstance.cpp
@@ -2117,7 +2117,7 @@ bool UCustomizableInstancePrivate::UpdateSkeletalMesh_PostBeginUpdate0(UCustomiz
 		}
 	}
 
-	// None of the current meshes requires a mesh update. Continue to BuildMaterials
+	// None of the current meshes requires a mesh update. Continue to BuildResources
 	if (!bUpdateMeshes)
 	{
 		return true;
@@ -5966,9 +5966,9 @@ FBuildMaterialOut BuildMaterial(const TSharedRef<FUpdateContextPrivate>& Context
 }
 
 
-void UCustomizableInstancePrivate::BuildMaterials(const TSharedRef<FUpdateContextPrivate>& Context, UCustomizableObjectInstance* Public)
+void UCustomizableInstancePrivate::BuildResources(const TSharedRef<FUpdateContextPrivate>& Context, UCustomizableObjectInstance* Public)
 {
-	MUTABLE_CPUPROFILER_SCOPE(UCustomizableInstancePrivate::BuildMaterials)
+	MUTABLE_CPUPROFILER_SCOPE(UCustomizableInstancePrivate::BuildResources)
 
 	UCustomizableObject* CustomizableObject = Public->GetCustomizableObject();
 
@@ -6001,6 +6001,9 @@ void UCustomizableInstancePrivate::BuildMaterials(const TSharedRef<FUpdateContex
 
 	Context->UpdateImageContext = UpdateImageContext;
 	
+	Materials.Empty();
+	Textures.Empty();
+
 	TArray<bool> RecreateRenderStateOnInstanceComponent;
 	RecreateRenderStateOnInstanceComponent.Init(false, Context->NumInstanceComponents);
 
@@ -6082,7 +6085,7 @@ void UCustomizableInstancePrivate::BuildMaterials(const TSharedRef<FUpdateContex
 		// SurfaceId per MaterialSlotIndex
 		TArray<int32> SurfaceIdToMaterialIndex;
 
-		MUTABLE_CPUPROFILER_SCOPE(BuildMaterials_LODLoop);
+		MUTABLE_CPUPROFILER_SCOPE(BuildResources_LODLoop);
 
 		const int32 FirstLOD = Context->bStreamMeshLODs ?
 			Context->FirstLODAvailable[ComponentName] :
@@ -6253,7 +6256,7 @@ void UCustomizableInstancePrivate::BuildMaterials(const TSharedRef<FUpdateContex
 	// TODO: MTBL-1697 Remove after merging ConvertResources and Callbacks.
 	if (RecreateRenderStateOnInstanceComponent.Find(true) != INDEX_NONE || NotifyUpdateOnInstanceComponent.Find(true) != INDEX_NONE)
 	{
-		MUTABLE_CPUPROFILER_SCOPE(BuildMaterials_RecreateRenderState);
+		MUTABLE_CPUPROFILER_SCOPE(BuildResources_RecreateRenderState);
 
 		for (TObjectIterator<UCustomizableObjectInstanceUsage> It; It; ++It)
 		{
@@ -6311,7 +6314,7 @@ void UCustomizableInstancePrivate::BuildMaterials(const TSharedRef<FUpdateContex
 	}
 
 	{
-		MUTABLE_CPUPROFILER_SCOPE(BuildMaterials_Exchange);
+		MUTABLE_CPUPROFILER_SCOPE(BuildResources_Exchange);
 
 		Context->Instance->GetPrivate()->LastMaterialsReuseCache = MoveTemp(Context->Instance->GetPrivate()->CurrentMaterialsReuseCache);
 		Context->Instance->GetPrivate()->LastTextureReuseCache = MoveTemp(Context->Instance->GetPrivate()->CurrentTextureReuseCache);
diff --git a/Engine/Plugins/Mutable/Source/CustomizableObject/Private/MuCO/CustomizableObjectSystem.cpp b/Engine/Plugins/Mutable/Source/CustomizableObject/Private/MuCO/CustomizableObjectSystem.cpp
--- a/Engine/Plugins/Mutable/Source/CustomizableObject/Private/MuCO/CustomizableObjectSystem.cpp
+++ b/Engine/Plugins/Mutable/Source/CustomizableObject/Private/MuCO/CustomizableObjectSystem.cpp
@@ -2818,7 +2818,7 @@ namespace impl
 					MUTABLE_CPUPROFILER_SCOPE(UpdateSkeletalMesh_PostBeginUpdate1);
 
 					// \TODO: Bring here
-					CustomizableInstancePrivateData->BuildMaterials(OperationData, CustomizableObjectInstance);
+					CustomizableInstancePrivateData->BuildResources(OperationData, CustomizableObjectInstance);
 				}
 
 				// This used to be CustomizableObjectInstance::UpdateSkeletalMesh_PostBeginUpdate2
