{Mutable] Fixed crash (check) due to invalid name.
+ Now when a parameter name of a Break Material node is None, the result is a nullpointer which is ignored in the core.
+ Also removed an early return to enable the node cache.

#rnx

#Synced-CL 49281961

diff --git a/Engine/Plugins/Mutable/Source/CustomizableObjectEditor/Private/MuCOE/GenerateMutableSource/GenerateMutableSourceColor.cpp b/Engine/Plugins/Mutable/Source/CustomizableObjectEditor/Private/MuCOE/GenerateMutableSource/GenerateMutableSourceColor.cpp
--- a/Engine/Plugins/Mutable/Source/CustomizableObjectEditor/Private/MuCOE/GenerateMutableSource/GenerateMutableSourceColor.cpp
+++ b/Engine/Plugins/Mutable/Source/CustomizableObjectEditor/Private/MuCOE/GenerateMutableSource/GenerateMutableSourceColor.cpp
@@ -389,7 +389,6 @@ UE::Mutable::Private::Ptr<UE::Mutable::Private::NodeColour> GenerateMutableSourc
 
 	else if (const UCONodeMaterialBreak* TypedNodeMaterialBreak = Cast<UCONodeMaterialBreak>(Node))
 	{
-		Result = new UE::Mutable::Private::NodeColourConstant();
 		check(TypedNodeMaterialBreak->MaterialPinRef.Get());
 
 		if (const UEdGraphPin* LinkedMaterial = FollowInputPin(*TypedNodeMaterialBreak->MaterialPinRef.Get()))
@@ -397,31 +396,32 @@ UE::Mutable::Private::Ptr<UE::Mutable::Private::NodeColour> GenerateMutableSourc
 			UE::Mutable::Private::Ptr<UE::Mutable::Private::NodeColourMaterialBreak> MaterialBreak = new UE::Mutable::Private::NodeColourMaterialBreak();
 
 			FName ParameterName = TypedNodeMaterialBreak->GetPinParameterName(*Pin);
-			if (ParameterName.IsNone())
+			if (!ParameterName.IsNone())
 			{
-				GenerationContext.Log(LOCTEXT("NoneParameterName", "'None' is not a valid Material Parameter name."), Node, EMessageSeverity::Warning);
-				return MaterialBreak;
-			}
-			
-			int32 LayerIndex = TypedNodeMaterialBreak->GetPinParameterLayerIndex(*Pin);
+				int32 LayerIndex = TypedNodeMaterialBreak->GetPinParameterLayerIndex(*Pin);
 
-			// Set the name of the parameter that represents the pin
-			MaterialBreak->ParameterName = ParameterName;
-			MaterialBreak->LayerIndex = LayerIndex;
+				// Set the name of the parameter that represents the pin
+				MaterialBreak->ParameterName = ParameterName;
+				MaterialBreak->LayerIndex = LayerIndex;
 
-			// Generate the material that is going to be broken into parameters
-			GenerationContext.MaterialBreakParameterStack.Add(FMaterialBreakParameter{ParameterName, Pin->PinType.PinCategory, LayerIndex });
-			UE::Mutable::Private::Ptr<UE::Mutable::Private::NodeMaterial> LinkedNode = GenerateMutableSourceMaterial(LinkedMaterial, GenerationContext,FGenerationMaterialOptions());
-			MaterialBreak->MaterialSource = LinkedNode;
-			GenerationContext.MaterialBreakParameterStack.Pop();
+				// Generate the material that is going to be broken into parameters
+				GenerationContext.MaterialBreakParameterStack.Add(FMaterialBreakParameter{ ParameterName, Pin->PinType.PinCategory, LayerIndex });
+				UE::Mutable::Private::Ptr<UE::Mutable::Private::NodeMaterial> LinkedNode = GenerateMutableSourceMaterial(LinkedMaterial, GenerationContext, FGenerationMaterialOptions());
+				MaterialBreak->MaterialSource = LinkedNode;
+				GenerationContext.MaterialBreakParameterStack.Pop();
 
-			if (LinkedNode)
-			{
-				Result = MaterialBreak;
+				if (LinkedNode)
+				{
+					Result = MaterialBreak;
+				}
+				else
+				{
+					GenerationContext.Log(LOCTEXT("MaterialBreakNodeError", "Something went wrong generation the Material Node linked to the Material Break Node."), Node);
+				}
 			}
 			else
 			{
-				GenerationContext.Log(LOCTEXT("MaterialBreakNodeError", "Something went wrong generation the Material Node linked to the Material Break Node."), Node);
+				GenerationContext.Log(LOCTEXT("NoneParameterName", "'None' is not a valid Material Parameter name."), Node, EMessageSeverity::Warning);
 			}
 		}
 		else
diff --git a/Engine/Plugins/Mutable/Source/CustomizableObjectEditor/Private/MuCOE/GenerateMutableSource/GenerateMutableSourceFloat.cpp b/Engine/Plugins/Mutable/Source/CustomizableObjectEditor/Private/MuCOE/GenerateMutableSource/GenerateMutableSourceFloat.cpp
--- a/Engine/Plugins/Mutable/Source/CustomizableObjectEditor/Private/MuCOE/GenerateMutableSource/GenerateMutableSourceFloat.cpp
+++ b/Engine/Plugins/Mutable/Source/CustomizableObjectEditor/Private/MuCOE/GenerateMutableSource/GenerateMutableSourceFloat.cpp
@@ -411,7 +411,6 @@ UE::Mutable::Private::Ptr<UE::Mutable::Private::NodeScalar> GenerateMutableSourc
 
 	else if (const UCONodeMaterialBreak* TypedNodeMaterialBreak = Cast<UCONodeMaterialBreak>(Node))
 	{
-		Result = new UE::Mutable::Private::NodeScalarConstant();
 		check(TypedNodeMaterialBreak->MaterialPinRef.Get());
 
 		if (const UEdGraphPin* LinkedMaterial = FollowInputPin(*TypedNodeMaterialBreak->MaterialPinRef.Get()))
@@ -419,31 +418,32 @@ UE::Mutable::Private::Ptr<UE::Mutable::Private::NodeScalar> GenerateMutableSourc
 			UE::Mutable::Private::Ptr<UE::Mutable::Private::NodeScalarMaterialBreak> MaterialBreak = new UE::Mutable::Private::NodeScalarMaterialBreak();
 
 			FName ParameterName = TypedNodeMaterialBreak->GetPinParameterName(*Pin);
-			if (ParameterName.IsNone())
+			if (!ParameterName.IsNone())
 			{
-				GenerationContext.Log(LOCTEXT("NoneParameterName", "'None' is not a valid Material Parameter name."), Node, EMessageSeverity::Warning);
-				return MaterialBreak;
-			}
-			
-			int32 LayerIndex = TypedNodeMaterialBreak->GetPinParameterLayerIndex(*Pin);
+				int32 LayerIndex = TypedNodeMaterialBreak->GetPinParameterLayerIndex(*Pin);
 
-			// Set the name of the parameter that represents the pin
-			MaterialBreak->ParameterName = ParameterName;
-			MaterialBreak->LayerIndex = LayerIndex;
+				// Set the name of the parameter that represents the pin
+				MaterialBreak->ParameterName = ParameterName;
+				MaterialBreak->LayerIndex = LayerIndex;
 
-			// Generate the material that is going to be broken into parameters
-			GenerationContext.MaterialBreakParameterStack.Add(FMaterialBreakParameter{ ParameterName, Pin->PinType.PinCategory, LayerIndex });
-			UE::Mutable::Private::Ptr<UE::Mutable::Private::NodeMaterial> LinkedNode = GenerateMutableSourceMaterial(LinkedMaterial, GenerationContext, FGenerationMaterialOptions());
-			MaterialBreak->MaterialSource = LinkedNode;
-			GenerationContext.MaterialBreakParameterStack.Pop();
+				// Generate the material that is going to be broken into parameters
+				GenerationContext.MaterialBreakParameterStack.Add(FMaterialBreakParameter{ ParameterName, Pin->PinType.PinCategory, LayerIndex });
+				UE::Mutable::Private::Ptr<UE::Mutable::Private::NodeMaterial> LinkedNode = GenerateMutableSourceMaterial(LinkedMaterial, GenerationContext, FGenerationMaterialOptions());
+				MaterialBreak->MaterialSource = LinkedNode;
+				GenerationContext.MaterialBreakParameterStack.Pop();
 
-			if (LinkedNode)
-			{
-				Result = MaterialBreak;
+				if (LinkedNode)
+				{
+					Result = MaterialBreak;
+				}
+				else
+				{
+					GenerationContext.Log(LOCTEXT("MaterialBreakNodeError", "Something went wrong generation the Material Node linked to the Material Break Node."), Node);
+				}
 			}
 			else
 			{
-				GenerationContext.Log(LOCTEXT("MaterialBreakNodeError", "Something went wrong generation the Material Node linked to the Material Break Node."), Node);
+				GenerationContext.Log(LOCTEXT("NoneParameterName", "'None' is not a valid Material Parameter name."), Node, EMessageSeverity::Warning);
 			}
 		}
 		else
diff --git a/Engine/Plugins/Mutable/Source/CustomizableObjectEditor/Private/MuCOE/GenerateMutableSource/GenerateMutableSourceImage.cpp b/Engine/Plugins/Mutable/Source/CustomizableObjectEditor/Private/MuCOE/GenerateMutableSource/GenerateMutableSourceImage.cpp
--- a/Engine/Plugins/Mutable/Source/CustomizableObjectEditor/Private/MuCOE/GenerateMutableSource/GenerateMutableSourceImage.cpp
+++ b/Engine/Plugins/Mutable/Source/CustomizableObjectEditor/Private/MuCOE/GenerateMutableSource/GenerateMutableSourceImage.cpp
@@ -1223,7 +1223,6 @@ UE::Mutable::Private::Ptr<UE::Mutable::Private::NodeImage> GenerateMutableSource
 
 	else if (const UCONodeMaterialBreak* TypedNodeMaterialBreak = Cast<UCONodeMaterialBreak>(Node))
 	{
-		Result = new UE::Mutable::Private::NodeImageConstant();
 		check(TypedNodeMaterialBreak->MaterialPinRef.Get());
 
 		if (const UEdGraphPin* LinkedMaterial = FollowInputPin(*TypedNodeMaterialBreak->MaterialPinRef.Get()))
@@ -1231,35 +1230,37 @@ UE::Mutable::Private::Ptr<UE::Mutable::Private::NodeImage> GenerateMutableSource
 			UE::Mutable::Private::Ptr<UE::Mutable::Private::NodeImageMaterialBreak> MaterialBreak = new UE::Mutable::Private::NodeImageMaterialBreak();
 
 			FName ParameterName = TypedNodeMaterialBreak->GetPinParameterName(*Pin);
-			if (ParameterName.IsNone())
+			if (!ParameterName.IsNone())
 			{
-				GenerationContext.Log(LOCTEXT("NoneParameterName", "'None' is not a valid Material Parameter name."), Node, EMessageSeverity::Warning);
-				return MaterialBreak;
-			}
-			
-			int32 LayerIndex = TypedNodeMaterialBreak->GetPinParameterLayerIndex(*Pin);
+				int32 LayerIndex = TypedNodeMaterialBreak->GetPinParameterLayerIndex(*Pin);
 
-			// Set the name of the parameter that represents the pin
-			MaterialBreak->ParameterName = ParameterName;
-			MaterialBreak->LayerIndex = LayerIndex;
+				// Set the name of the parameter that represents the pin
+				MaterialBreak->ParameterName = ParameterName;
+				MaterialBreak->LayerIndex = LayerIndex;
 
-			// Store the reference texture into the material generation options to generate all the textures with the correct size
-			FGenerationMaterialOptions GenerationOptions;
-			GenerationOptions.ReferenceTextureSize = ReferenceTextureSize;
+				// Store the reference texture into the material generation options to generate all the textures with the correct size
+				FGenerationMaterialOptions GenerationOptions;
+				GenerationOptions.ReferenceTextureSize = ReferenceTextureSize;
 
-			// Generate the material that is going to be broken into parameters
-			GenerationContext.MaterialBreakParameterStack.Add(FMaterialBreakParameter{ ParameterName, Pin->PinType.PinCategory, LayerIndex });
-			UE::Mutable::Private::Ptr<UE::Mutable::Private::NodeMaterial> LinkedNode = GenerateMutableSourceMaterial(LinkedMaterial, GenerationContext, GenerationOptions);
-			MaterialBreak->MaterialSource = LinkedNode;
-			GenerationContext.MaterialBreakParameterStack.Pop();
+				// Generate the material that is going to be broken into parameters
+				GenerationContext.MaterialBreakParameterStack.Add(FMaterialBreakParameter{ ParameterName, Pin->PinType.PinCategory, LayerIndex });
+				UE::Mutable::Private::Ptr<UE::Mutable::Private::NodeMaterial> LinkedNode = GenerateMutableSourceMaterial(LinkedMaterial, GenerationContext, GenerationOptions);
+				MaterialBreak->MaterialSource = LinkedNode;
+				GenerationContext.MaterialBreakParameterStack.Pop();
 
-			if (LinkedNode)
-			{
-				Result = MaterialBreak;
+				if (LinkedNode)
+				{
+					Result = MaterialBreak;
+				}
+				else
+				{
+					GenerationContext.Log(LOCTEXT("MaterialBreakNodeError", "Something went wrong generation the Material Node linked to the Material Break Node."), Node);
+				}
 			}
 			else
 			{
-				GenerationContext.Log(LOCTEXT("MaterialBreakNodeError", "Something went wrong generation the Material Node linked to the Material Break Node."), Node);
+				GenerationContext.Log(LOCTEXT("NoneParameterName", "'None' is not a valid Material Parameter name."), Node, EMessageSeverity::Warning);
+
 			}
 		}
 		else
