[Mutable] Fixed crash with passthrough textures that came from break material nodes.
+ Added a FMaterial cache to the CodeRunnerBegin. This allows the mutable core to load passthrough textures from FMaterials when needed.
+ Removed float and color values from material parameters. Now the values are taken from the UMaterial at runtime and they are not stored in the FMaterial.
+ Added some FParameterKey variables to all Break Material nodes.
+ Added support for passthrough textures to material parameters linked to break material nodes.
+ Added missing part of the codegenerator_material cache.

#rnx

#Synced-CL 49327150

diff --git a/Engine/Plugins/Mutable/Source/CustomizableObject/Internal/MuCO/CustomizableObjectPrivate.h b/Engine/Plugins/Mutable/Source/CustomizableObject/Internal/MuCO/CustomizableObjectPrivate.h
--- a/Engine/Plugins/Mutable/Source/CustomizableObject/Internal/MuCO/CustomizableObjectPrivate.h
+++ b/Engine/Plugins/Mutable/Source/CustomizableObject/Internal/MuCO/CustomizableObjectPrivate.h
@@ -1768,6 +1768,8 @@ enum class ECustomizableObjectVersions : int32
 
 	RemovedReuseTextures,
 
+	PassthroughTexturesForBreakNodes,
+
 	// -----<new versions can be added above this line>--------
 	LastCustomizableObjectVersion
 };
diff --git a/Engine/Plugins/Mutable/Source/CustomizableObjectEditor/Private/MuCOE/GenerateMutableSource/GenerateMutableSourceColor.cpp b/Engine/Plugins/Mutable/Source/CustomizableObjectEditor/Private/MuCOE/GenerateMutableSource/GenerateMutableSourceColor.cpp
--- a/Engine/Plugins/Mutable/Source/CustomizableObjectEditor/Private/MuCOE/GenerateMutableSource/GenerateMutableSourceColor.cpp
+++ b/Engine/Plugins/Mutable/Source/CustomizableObjectEditor/Private/MuCOE/GenerateMutableSource/GenerateMutableSourceColor.cpp
@@ -401,8 +401,7 @@ UE::Mutable::Private::Ptr<UE::Mutable::Private::NodeColour> GenerateMutableSourc
 				int32 LayerIndex = TypedNodeMaterialBreak->GetPinParameterLayerIndex(*Pin);
 
 				// Set the name of the parameter that represents the pin
-				MaterialBreak->ParameterName = ParameterName;
-				MaterialBreak->LayerIndex = LayerIndex;
+				MaterialBreak->ParameterKey = {ParameterName, LayerIndex};
 
 				// Generate the material that is going to be broken into parameters
 				GenerationContext.MaterialBreakParameterStack.Add(FMaterialBreakParameter{ ParameterName, Pin->PinType.PinCategory, LayerIndex });
diff --git a/Engine/Plugins/Mutable/Source/CustomizableObjectEditor/Private/MuCOE/GenerateMutableSource/GenerateMutableSourceFloat.cpp b/Engine/Plugins/Mutable/Source/CustomizableObjectEditor/Private/MuCOE/GenerateMutableSource/GenerateMutableSourceFloat.cpp
--- a/Engine/Plugins/Mutable/Source/CustomizableObjectEditor/Private/MuCOE/GenerateMutableSource/GenerateMutableSourceFloat.cpp
+++ b/Engine/Plugins/Mutable/Source/CustomizableObjectEditor/Private/MuCOE/GenerateMutableSource/GenerateMutableSourceFloat.cpp
@@ -423,8 +423,7 @@ UE::Mutable::Private::Ptr<UE::Mutable::Private::NodeScalar> GenerateMutableSourc
 				int32 LayerIndex = TypedNodeMaterialBreak->GetPinParameterLayerIndex(*Pin);
 
 				// Set the name of the parameter that represents the pin
-				MaterialBreak->ParameterName = ParameterName;
-				MaterialBreak->LayerIndex = LayerIndex;
+				MaterialBreak->ParameterKey = {ParameterName, LayerIndex};
 
 				// Generate the material that is going to be broken into parameters
 				GenerationContext.MaterialBreakParameterStack.Add(FMaterialBreakParameter{ ParameterName, Pin->PinType.PinCategory, LayerIndex });
diff --git a/Engine/Plugins/Mutable/Source/CustomizableObjectEditor/Private/MuCOE/GenerateMutableSource/GenerateMutableSourceImage.cpp b/Engine/Plugins/Mutable/Source/CustomizableObjectEditor/Private/MuCOE/GenerateMutableSource/GenerateMutableSourceImage.cpp
--- a/Engine/Plugins/Mutable/Source/CustomizableObjectEditor/Private/MuCOE/GenerateMutableSource/GenerateMutableSourceImage.cpp
+++ b/Engine/Plugins/Mutable/Source/CustomizableObjectEditor/Private/MuCOE/GenerateMutableSource/GenerateMutableSourceImage.cpp
@@ -1235,8 +1235,7 @@ UE::Mutable::Private::Ptr<UE::Mutable::Private::NodeImage> GenerateMutableSource
 				int32 LayerIndex = TypedNodeMaterialBreak->GetPinParameterLayerIndex(*Pin);
 
 				// Set the name of the parameter that represents the pin
-				MaterialBreak->ParameterName = ParameterName;
-				MaterialBreak->LayerIndex = LayerIndex;
+				MaterialBreak->ParameterKey = {ParameterName, LayerIndex};
 
 				// Store the reference texture into the material generation options to generate all the textures with the correct size
 				FGenerationMaterialOptions GenerationOptions;
diff --git a/Engine/Plugins/Mutable/Source/CustomizableObjectEditor/Private/MuCOE/GenerateMutableSource/GenerateMutableSourceMaterial.cpp b/Engine/Plugins/Mutable/Source/CustomizableObjectEditor/Private/MuCOE/GenerateMutableSource/GenerateMutableSourceMaterial.cpp
--- a/Engine/Plugins/Mutable/Source/CustomizableObjectEditor/Private/MuCOE/GenerateMutableSource/GenerateMutableSourceMaterial.cpp
+++ b/Engine/Plugins/Mutable/Source/CustomizableObjectEditor/Private/MuCOE/GenerateMutableSource/GenerateMutableSourceMaterial.cpp
@@ -665,32 +665,6 @@ UE::Mutable::Private::Ptr<UE::Mutable::Private::NodeMaterial> GenerateMutableSou
 					UE::Mutable::Private::Ptr<UE::Mutable::Private::NodeImageFormat> FormatNode = GenerateMutableImageFormat(ResizeNode, ReferenceTexture, GenerationContext, Node);
 					MaterialNode->ImageParametersToCompile.Add({ FName(ParameterInfo.Name), ParameterInfo.Index }, FormatNode);
 				}
-
-				else if (CurrentBreakParameter.ParameterType == UEdGraphSchema_CustomizableObject::PC_Texture_Passthrough)
-				{
-					//TODO(Max - GM): Add support for Passthrough texture parameters
-					////Find reference Texture from the material parameter
-					//UTexture* Texture = nullptr;
-					//bool bParameterFound = MaterialParameterNode->ReferenceValue->GetTextureParameterValue(GenerationContext.CurrentMaterialBreakParameter.ParameterName, Texture);
-
-					//UE::Mutable::Private::Ptr<UE::Mutable::Private::NodeImageConstant> ImageNode = new UE::Mutable::Private::NodeImageConstant();
-					//ImageNode->SetValue(GenerateImageConstant(Texture, GenerationContext, true));
-
-					//MaterialNode->ImageParameterNodes.Add({ FName(ParameterInfo.Name), ParameterInfo.Index });
-
-					FText Message = LOCTEXT("ParameterPassthroughUnsupported", "A Material Parameter node can't be break into passthrough textures right now.");
-					GenerationContext.Log(Message, Node, EMessageSeverity::Error);
-				}
-
-				else if (CurrentBreakParameter.ParameterType == UEdGraphSchema_CustomizableObject::PC_Color)
-				{
-					MaterialNode->ColorParametersToCompile.Add({ FName(ParameterInfo.Name), ParameterInfo.Index });
-				}
-
-				else if (CurrentBreakParameter.ParameterType == UEdGraphSchema_CustomizableObject::PC_Float)
-				{
-					MaterialNode->ScalarParametersToCompile.Add({ FName(ParameterInfo.Name), ParameterInfo.Index });
-				}
 			}
 		}
 
@@ -706,7 +680,7 @@ UE::Mutable::Private::Ptr<UE::Mutable::Private::NodeMaterial> GenerateMutableSou
 		{
 			if (const UEdGraphPin* ConnectedMeshPin = FollowInputPin(*TypedSkeletalMeshBreakNode->InputPassthroughMesh.Get()))
 			{
-				UE::Mutable::Private::Ptr<UE::Mutable::Private::NodeMaterialSkeletalMeshBreak> MaterialFromSlotNode = new UE::Mutable::Private::NodeMaterialSkeletalMeshBreak();;
+				UE::Mutable::Private::Ptr<UE::Mutable::Private::NodeMaterialSkeletalMeshBreak> MaterialFromSlotNode = new UE::Mutable::Private::NodeMaterialSkeletalMeshBreak();
 				MaterialFromSlotNode->Mesh = UE::Mutable::Private::GenerateMutableSourceSkeletalMesh(ConnectedMeshPin, GenerationContext);
 				MaterialFromSlotNode->SlotName = TargetSlotName;
 
@@ -747,47 +721,37 @@ UE::Mutable::Private::Ptr<UE::Mutable::Private::NodeMaterial> GenerateMutableSou
 		//-----------------------------------------------------------------
 		// TODOs: 
 		// - Add support for table node materials. This will be added when the current hack of the tables is removed. UE-314401
-		// - Add support for layout index once the pin viewer has all the needed properties. UE-317590
 		// - Add support for group projectors.
 		// - Remove the material edition part from the mesh section nodes. Talk about this with Gerard since we can lose the capacity to pass an image parameter of the reference material as mutable without linking any texture node.
-		// - Improve the detection of the reference texture. 
-		//		|_ Search it in the pin data
-		//		|_ Or instead find the first material constant (or parameter) node that contains the texture parameter and return it
-		//		|_ Or use the reference image of the group projector generated (if any). 
 
 		// Tracking the index of the image. Needed for the FGeneratedImagePropertyKey
 		int32 ImageIndex = 0;
 	
-		for (const UEdGraphPin* NonOrphanPin : TypedNodeMaterialModify->GetAllNonOrphanPins())
+		for (const UEdGraphPin* InputPin : TypedNodeMaterialModify->GetAllNonOrphanPins())
 		{
-			if (NonOrphanPin->Direction == EEdGraphPinDirection::EGPD_Output)
-			{
-				continue;
-			}
-			
-			if (NonOrphanPin == TypedNodeMaterialModify->MaterialPinRef.Get())
+			if (FEdGraphPinReference(InputPin) == TypedNodeMaterialModify->MaterialPinRef.Get() || InputPin->Direction == EGPD_Output)
 			{
 				continue;
 			}
 
-			const FName ParameterName = TypedNodeMaterialModify->GetPinParameterName(*NonOrphanPin);
+			const FName ParameterName = TypedNodeMaterialModify->GetPinParameterName(*InputPin);
 			if (ParameterName.IsNone())
 			{
 				GenerationContext.Log(LOCTEXT("NoneParameterName", "'None' is not a valid Material Parameter name."), Node, EMessageSeverity::Warning);
 				continue;
 			}
 			
-			const int32 LayerIndex = TypedNodeMaterialModify->GetPinParameterLayerIndex(*NonOrphanPin);
+			const int32 LayerIndex = TypedNodeMaterialModify->GetPinParameterLayerIndex(*InputPin);
 			UE::Mutable::Private::FParameterKey ParameterKey = { ParameterName, LayerIndex };
 
-			const FName PinType = NonOrphanPin->PinType.PinCategory;
+			const FName PinType = InputPin->PinType.PinCategory;
 			if (PinType == UEdGraphSchema_CustomizableObject::PC_Texture)
 			{
 				UE::Mutable::Private::FImageParameterData ImageData;
 
-				if (const UEdGraphPin* FollowPin = FollowInputPin(*NonOrphanPin))
+				if (const UEdGraphPin* FollowPin = FollowInputPin(*InputPin))
 				{
-					TObjectPtr<UTexture2D> ReferenceTexture = GenerationContext.LoadObject(TypedNodeMaterialModify->GetTexturePinParameterReferenceTexture(*NonOrphanPin));
+					TObjectPtr<UTexture2D> ReferenceTexture = GenerationContext.LoadObject(TypedNodeMaterialModify->GetTexturePinParameterReferenceTexture(*InputPin));
 					if (!ReferenceTexture)
 					{
 						// Fallback if the user did not provide a reference texture
@@ -803,7 +767,7 @@ UE::Mutable::Private::Ptr<UE::Mutable::Private::NodeMaterial> GenerateMutableSou
 					// Fill Image Data
 					ImageData.ImageNode = ImageNode;
 					ImageData.ImagePropertyIndex = Props.ImagePropertiesIndex;
-					ImageData.LayoutIndex = TypedNodeMaterialModify->GetImagePinParameterUVIndex(*NonOrphanPin);
+					ImageData.LayoutIndex = TypedNodeMaterialModify->GetImagePinParameterUVIndex(*InputPin);
 				}
 
 				// We do not check Null images since it is possible to not have anything linked
@@ -816,7 +780,7 @@ UE::Mutable::Private::Ptr<UE::Mutable::Private::NodeMaterial> GenerateMutableSou
 			{
 				UE::Mutable::Private::FImageParameterData ImageData;
 
-				if (const UEdGraphPin* ConnectedPin = FollowInputPin(*NonOrphanPin))
+				if (const UEdGraphPin* ConnectedPin = FollowInputPin(*InputPin))
 				{
 					UE::Mutable::Private::FParameterKey ImageKey = { ParameterName, LayerIndex };
 					FGeneratedImageProperties Props = GenerateImageProperties(TypedNodeMaterialBase, nullptr, ImageKey, true, GenerationContext);
@@ -841,7 +805,7 @@ UE::Mutable::Private::Ptr<UE::Mutable::Private::NodeMaterial> GenerateMutableSou
 			{
 				UE::Mutable::Private::Ptr<UE::Mutable::Private::NodeColour> ColorNode;
 
-				if (const UEdGraphPin* FollowPin = FollowInputPin(*NonOrphanPin))
+				if (const UEdGraphPin* FollowPin = FollowInputPin(*InputPin))
 				{
 					ColorNode = GenerateMutableSourceColor(FollowPin, GenerationContext);
 				}
@@ -854,7 +818,7 @@ UE::Mutable::Private::Ptr<UE::Mutable::Private::NodeMaterial> GenerateMutableSou
 			{
 				UE::Mutable::Private::Ptr<UE::Mutable::Private::NodeScalar> ScalarNode;
 
-				if (const UEdGraphPin* FollowPin = FollowInputPin(*NonOrphanPin))
+				if (const UEdGraphPin* FollowPin = FollowInputPin(*InputPin))
 				{
 					ScalarNode = GenerateMutableSourceFloat(FollowPin, GenerationContext);
 				}
diff --git a/Engine/Plugins/Mutable/Source/MutableRuntime/Internal/MuR/Operations.h b/Engine/Plugins/Mutable/Source/MutableRuntime/Internal/MuR/Operations.h
--- a/Engine/Plugins/Mutable/Source/MutableRuntime/Internal/MuR/Operations.h
+++ b/Engine/Plugins/Mutable/Source/MutableRuntime/Internal/MuR/Operations.h
@@ -503,8 +503,6 @@ namespace UE::Mutable::Private
 		{
 			PARAMETER Variable;
 
-			LIST_UINT32_ADDRESS ScalarParameterNames;
-			LIST_UINT32_ADDRESS ColorParameterNames;
 			LIST_UINT32_ADDRESS ImageParameterNames;
 			LIST_UINT32_ADDRESS ImageParameterAddress;
 		};
diff --git a/Engine/Plugins/Mutable/Source/MutableRuntime/Private/MuR/CodeRunner.cpp b/Engine/Plugins/Mutable/Source/MutableRuntime/Private/MuR/CodeRunner.cpp
--- a/Engine/Plugins/Mutable/Source/MutableRuntime/Private/MuR/CodeRunner.cpp
+++ b/Engine/Plugins/Mutable/Source/MutableRuntime/Private/MuR/CodeRunner.cpp
@@ -5605,12 +5605,8 @@ namespace UE::Mutable::Private
 				const FString& ParameterName = Program.ConstantStrings[Args.ParameterName];
 				const FParameterKey ParameterKey = { FName(ParameterName), Args.LayerIndex };
 				const FMaterial::FImageParameterData* ImageParameterData = Material->ImageParameters.Find(ParameterKey);
-				
-				if (!ImageParameterData)
-				{
-					StoreImage(Item, nullptr);
-				}
-				else
+
+				if (ImageParameterData)
 				{
 					const TVariant<OP::ADDRESS, TManagedPtr<const FImage>> ImageParameter = ImageParameterData->ImageParameter;
 
@@ -5627,6 +5623,40 @@ namespace UE::Mutable::Private
 						StoreImage(Item, ImageParameter.Get<TManagedPtr<const FImage>>());
 					}
 				}
+				else if (Material->PassthroughObject.IsResolved())
+				{
+					UTexture* ParameterValue = nullptr;
+
+					if (const UMaterialInterface* MaterialInterface = Material->PassthroughObject.Get())
+					{
+						FHashedMaterialParameterInfo ParameterInfo;
+						ParameterInfo.Name = FScriptName(FName(ParameterName));
+						ParameterInfo.Index = ParameterKey.LayerIndex;
+						ParameterInfo.Association = ParameterInfo.Index == INDEX_NONE ? GlobalParameter : LayerParameter;
+
+						// Find image using the parameter info
+						MaterialInterface->GetTextureParameterValue(ParameterInfo, ParameterValue);
+					}
+
+					if (ParameterValue)
+					{
+						TPassthroughObjectPtr<UTexture> ExternalObject(ParameterValue);
+						TManagedPtr<FImage> Result = FImage::CreateAsReference(ExternalObject, FImageDesc(), false);
+
+						StoreImage(Item, Result);
+					}
+					else
+					{
+						AbortRun();
+						UE_LOG(LogMutableCore, Warning, TEXT("Expected Passthrough Texture parameter [%s] not found in Break Material node."), *ParameterName);
+					}
+				}
+				else
+				{
+					check(false);
+					StoreImage(Item, nullptr);
+				}
+
 				break;
 			}
 			case 2:
@@ -6122,7 +6152,41 @@ namespace UE::Mutable::Private
 				const FParameterKey ParameterKey = { FName(ParameterName), Args.LayerIndex };
 				const float* Value = Material->ScalarParameters.Find(ParameterKey);
 
-				StoreScalar(Item, Value ? *Value : 0.0f);
+				if (Value)
+				{
+					StoreScalar(Item, *Value);
+				}
+				else if (Material->PassthroughObject.IsResolved())
+				{
+					bool bParameterFound = false;
+					float ParameterValue = 0.0f;
+
+					if (const UMaterialInterface* MaterialInterface = Material->PassthroughObject.Get())
+					{
+						FHashedMaterialParameterInfo ParameterInfo;
+						ParameterInfo.Name = FScriptName(FName(ParameterName));
+						ParameterInfo.Index = ParameterKey.LayerIndex;
+						ParameterInfo.Association = ParameterInfo.Index == INDEX_NONE ? GlobalParameter : LayerParameter;
+
+						// Find float using the parameter info
+						bParameterFound = MaterialInterface->GetScalarParameterValue(ParameterInfo, ParameterValue);
+					}
+
+					if (bParameterFound)
+					{
+						StoreScalar(Item, ParameterValue);
+					}
+					else
+					{
+						AbortRun();
+						UE_LOG(LogMutableCore, Warning, TEXT("Expected Scalar parameter [%s] not found in Break Material node."), *ParameterName);
+					}
+				}
+				else
+				{
+					check(false);
+					StoreScalar(Item, 0.0f);
+				}
 
 				break;
 			}
@@ -6437,7 +6501,41 @@ namespace UE::Mutable::Private
 				const FParameterKey ParameterKey = { FName(ParameterName), Args.LayerIndex };
 				const FVector4f* Value = Material->ColorParameters.Find(ParameterKey);
 
-				StoreColor(Item, Value ? *Value : FVector4f::Zero());
+				if (Value)
+				{
+					StoreColor(Item, *Value);
+				}
+				else if (Material->PassthroughObject.IsResolved())
+				{
+					bool bParameterFound = false;
+					FLinearColor ParameterValue;
+
+					if (const UMaterialInterface* MaterialInterface = Material->PassthroughObject.Get())
+					{
+						FHashedMaterialParameterInfo ParameterInfo;
+						ParameterInfo.Name = FScriptName(FName(ParameterName));
+						ParameterInfo.Index = ParameterKey.LayerIndex;
+						ParameterInfo.Association = ParameterInfo.Index == INDEX_NONE ? GlobalParameter : LayerParameter;
+
+						// Find float using the parameter info
+						bParameterFound = MaterialInterface->GetVectorParameterValue(ParameterInfo, ParameterValue);
+					}
+
+					if (bParameterFound)
+					{
+						StoreColor(Item, ParameterValue);
+					}
+					else
+					{
+						AbortRun();
+						UE_LOG(LogMutableCore, Warning, TEXT("Expected Color parameter [%s] not found in Break Material node."), *ParameterName);
+					}
+				}
+				else
+				{
+					check(false);
+					StoreColor(Item, FVector4f::Zero());
+				}
 
 				break;
 			}
@@ -6555,70 +6653,11 @@ namespace UE::Mutable::Private
 				//Images
 				const TArray<uint32>& ImageParameterNames = Program.ConstantUInt32Lists[Args.ImageParameterNames];
 				const TArray<uint32>& ImageParameterOperations = Program.ConstantUInt32Lists[Args.ImageParameterAddress];
-				// Scalars
-				const TArray<uint32>& ScalarParameterNames = Program.ConstantUInt32Lists[Args.ScalarParameterNames];
-				// Vectors
-				const TArray<uint32>& ColorParameterNames = Program.ConstantUInt32Lists[Args.ColorParameterNames];
 
 				// Get Parameter layers from program bytecode
 				const uint8* Data = Program.GetOpArgsPointer(Item.At);
 				Data += sizeof(OP::MaterialParameterArgs);
 
-				// Get parameter values from the material. Order maters since we store layer indexes directly to the program (Scalar, Colors, Textures)
-				// Scalars
-				for (int32 ParameterIndex = 0; ParameterIndex < ScalarParameterNames.Num(); ++ParameterIndex)
-				{
-					// Get Scalar Parameter Name
-					const FString& ParameterName = Program.ConstantStrings[ScalarParameterNames[ParameterIndex]];
-
-					// Get Scalar Parameter Layer
-					int32 LayerIndex;
-					FMemory::Memcpy(&LayerIndex, Data, sizeof(int32));
-					Data += sizeof(int32);
-
-					// Construct Scalar parameter info
-					FHashedMaterialParameterInfo ScalarParameterInfo;
-					ScalarParameterInfo.Name = FScriptName(FName(ParameterName));
-					ScalarParameterInfo.Index = LayerIndex;
-					ScalarParameterInfo.Association = ScalarParameterInfo.Index == INDEX_NONE ? GlobalParameter : LayerParameter;
-
-					// Find Texture using the parameter info
-					float ScalarValue;
-					bool bParameterFound = MaterialInterface->GetScalarParameterValue(ScalarParameterInfo, ScalarValue);
-
-					if (bParameterFound)
-					{
-						Material->ScalarParameters.Add({ FName(ScalarParameterInfo.Name), ScalarParameterInfo.Index }, ScalarValue);
-					}
-				}
-
-				// Colors
-				for (int32 ParameterIndex = 0; ParameterIndex < ColorParameterNames.Num(); ++ParameterIndex)
-				{
-					// Get Color Parameter Name
-					const FString& ParameterName = Program.ConstantStrings[ColorParameterNames[ParameterIndex]];
-
-					// Get Color Parameter Layer
-					int32 LayerIndex;
-					FMemory::Memcpy(&LayerIndex, Data, sizeof(int32));
-					Data += sizeof(int32);
-
-					// Construct Scalar parameter info
-					FHashedMaterialParameterInfo ColorParameterInfo;
-					ColorParameterInfo.Name = FScriptName(FName(ParameterName));
-					ColorParameterInfo.Index = LayerIndex;
-					ColorParameterInfo.Association = ColorParameterInfo.Index == INDEX_NONE ? GlobalParameter : LayerParameter;
-
-					// Find Texture using the parameter info
-					FLinearColor VectorValue;
-					bool bParameterFound = MaterialInterface->GetVectorParameterValue(ColorParameterInfo, VectorValue);
-
-					if (bParameterFound)
-					{
-						Material->ColorParameters.Add({FName(ColorParameterInfo.Name), ColorParameterInfo.Index}, VectorValue);
-					}
-				}
-
 				// Textures
 				for (int32 ParameterIndex = 0; ParameterIndex < ImageParameterNames.Num(); ++ParameterIndex)
 				{
diff --git a/Engine/Plugins/Mutable/Source/MutableRuntime/Private/MuR/CodeRunnerBegin.cpp b/Engine/Plugins/Mutable/Source/MutableRuntime/Private/MuR/CodeRunnerBegin.cpp
--- a/Engine/Plugins/Mutable/Source/MutableRuntime/Private/MuR/CodeRunnerBegin.cpp
+++ b/Engine/Plugins/Mutable/Source/MutableRuntime/Private/MuR/CodeRunnerBegin.cpp
@@ -230,6 +230,23 @@ namespace UE::Mutable::Private
 	}
 
 
+	void CodeRunnerBegin::StoreMaterial(const FCacheAddressInline& To, const TManagedPtr<const FMaterial>& Value)
+	{
+		if (!To.At)
+		{
+			return;
+		}
+		
+		Executed.Add(To);
+
+		FStackValue StackValue;
+		StackValue.Set<TManagedPtr<const FMaterial>>(Value);
+		Stack.Push(StackValue);
+
+		ResultsMaterial.Add(To, Value);
+	}
+
+
 	int32 CodeRunnerBegin::LoadInt(const FCacheAddressInline& To)
 	{
 		if (!To.At)
@@ -296,6 +313,17 @@ namespace UE::Mutable::Private
 	}
 
 
+	TManagedPtr<const FMaterial> CodeRunnerBegin::LoadMaterial(const FCacheAddressInline& To)
+	{
+		if (!To.At)
+		{
+			return nullptr;	
+		}
+		
+		return Stack.Pop().Get<TManagedPtr<const FMaterial>>();
+	}
+
+
 	TManagedPtr<const FInstance> CodeRunnerBegin::RunCode(const FScheduledOpInline& Root)
 	{
 		MUTABLE_CPUPROFILER_SCOPE(CodeRunnerBegin::RunCode);
@@ -433,13 +461,20 @@ namespace UE::Mutable::Private
 						
 						break;
 					}
+				case EDataType::Material:
+					{
+						FStackValue Value;
+						Value.Set<TManagedPtr<const FMaterial>>(ResultsMaterial[Item]);
+
+						Stack.Push(Value);
+						break;
+					}
 					
 				case EDataType::Color:		
 				case EDataType::Projector:	
 				case EDataType::Mesh:		
 				case EDataType::Image:		
 				case EDataType::Layout:		
-				case EDataType::Material:		
 				case EDataType::InstancedStruct:		
 				case EDataType::SkeletalMesh:		
 					break;
@@ -466,7 +501,6 @@ namespace UE::Mutable::Private
 			case EOpType::PR_PARAMETER:
 			case EOpType::SK_PARAMETER:
 			case EOpType::IM_PARAMETER:
-			case EOpType::MI_PARAMETER:
 			case EOpType::IM_PARAMETER_FROM_MATERIAL:
 			case EOpType::IS_PARAMETER:
 				{
@@ -475,6 +509,12 @@ namespace UE::Mutable::Private
 					StoreNone(Item);
 					break;
 				}
+
+			case EOpType::MI_PARAMETER:
+				{
+					StoreMaterial(Item,MakeManaged<const FMaterial>());
+					break;
+				}
 			
 			case EOpType::IN_CONDITIONAL:
 				check(Item.bEvaluate);
@@ -541,13 +581,16 @@ namespace UE::Mutable::Private
 								case EDataType::Scalar:		
 									StoreFloat(Item, LoadFloat(Item.CustomState));
 									break;
+							
+								case EDataType::Material:
+									StoreMaterial(Item, LoadMaterial(Item.CustomState));
+									break;
 
 								case EDataType::Color:		
 								case EDataType::Projector:	
 								case EDataType::Mesh:		
 								case EDataType::Image:		
 								case EDataType::Layout:		
-								case EDataType::Material:		
 									StoreNone(Item);
 									break;
 
@@ -635,13 +678,16 @@ namespace UE::Mutable::Private
 									case EDataType::Scalar:
 										StoreFloat(Item, 0.0f);
 										break;
+
+									case EDataType::Material:
+										StoreMaterial(Item, MakeManaged<const FMaterial>());
+										break;
 									
 									case EDataType::Color:			
 									case EDataType::Projector:		
 									case EDataType::Mesh:			
 									case EDataType::Image:			
 									case EDataType::Layout:			
-									case EDataType::Material:			
 										StoreNone(Item);
 										break;
 								
@@ -746,13 +792,16 @@ namespace UE::Mutable::Private
 								case EDataType::Scalar:
 									StoreFloat(Item, LoadFloat(Item.CustomState));
 									break;
-								
-								case EDataType::Color:		
-								case EDataType::Projector:	
-								case EDataType::Mesh:		
-								case EDataType::Image:		
-								case EDataType::Layout:		
-								case EDataType::Material:		
+
+								case EDataType::Material:
+									StoreMaterial(Item, LoadMaterial(Item.CustomState));
+									break;
+
+								case EDataType::Color:
+								case EDataType::Projector:
+								case EDataType::Mesh:
+								case EDataType::Image:
+								case EDataType::Layout:
 									StoreNone(Item);
 									break;
 
@@ -762,7 +811,6 @@ namespace UE::Mutable::Private
 							}
 							else
 							{
-								StoreNone(Item);
 							}
 
 							break;
@@ -868,6 +916,8 @@ namespace UE::Mutable::Private
 							{
         						FMaterialId MaterialId = FSystem::GetMaterialId(MaterialIdRegistry, Model, Params.Get(), Args.Material);
 								Instance->SetMaterialId(0, 0, 0, MaterialId);
+
+								LoadMaterial(Args.Material);
 							}
 							
 							StoreInstance( Item, Instance);
@@ -919,6 +969,9 @@ namespace UE::Mutable::Private
 							{
         						FMaterialId MaterialId = FSystem::GetMaterialId(MaterialIdRegistry, Model, Params.Get(), Args.Material);
 								Instance->SetOverlayMaterialId(0, MaterialId);
+
+								// Pop Material
+								LoadMaterial(Args.Material);
 							}
 							
 							StoreInstance( Item, Instance);
@@ -969,6 +1022,9 @@ namespace UE::Mutable::Private
 								FMaterialId MaterialId = FSystem::GetMaterialId(MaterialIdRegistry, Model, Params.Get(), Args.Material);
 								const FString& SlotName = Program.ConstantStrings[Args.SlotName];
 								Instance->AddOverrideMaterial(0, FName(SlotName), MaterialId);
+
+								// Pop Material
+								LoadMaterial(Args.Material);
 							}
 							
 							StoreInstance( Item, Instance);
@@ -2716,29 +2772,76 @@ namespace UE::Mutable::Private
 					OP::MaterialConstantArgs Args = Program.GetOpArgs<OP::MaterialConstantArgs>(Item.At);
 					PassthroughObjectLoader->Add<UMaterialInterface>(Args.ID);
 
-					StoreNone(Item);
+					TManagedPtr<const FMaterial> Material = Program.ConstantMaterials[Args.Value];
+					check(Material);
+
+					StoreMaterial(Item, Material);
 					break;
 				}
-		
-			case EOpType::SC_MATERIAL_BREAK:
-			{
-				OP::MaterialBreakArgs Args = Program.GetOpArgs<OP::MaterialBreakArgs>(Item.At);
-				PushOp(FScheduledOpInline(Args.Material, Item));
-
-				// Hack: we need to store any value here. TODO GMT: remove this.
-				StoreFloat(Item, 0.0f);
-				break;
-			}
 
+			case EOpType::SC_MATERIAL_BREAK:
 			case EOpType::IM_MATERIAL_BREAK:
 			case EOpType::CO_MATERIAL_BREAK:
 				{
-					check(!Item.bEvaluate);
-					
 					OP::MaterialBreakArgs Args = Program.GetOpArgs<OP::MaterialBreakArgs>(Item.At);
-					PushOp(FScheduledOpInline(Args.Material, Item));
 
-					StoreNone(Item);
+					switch (Item.Stage)
+					{
+						case 0:
+						{
+							// Push next stage
+							PushOp(FScheduledOpInline(Item, 1));
+
+							// Push Material Op
+							PushOp(FScheduledOpInline(Args.Material, Item));
+
+							break;
+						}
+						case 1:
+						{
+							TManagedPtr<const FMaterial> Material = LoadMaterial(Args.Material);
+
+							check(Args.ParameterName < (uint32)Program.ConstantStrings.Num());
+							const FString& ParameterName = Program.ConstantStrings[Args.ParameterName];
+							const FParameterKey ParameterKey = { FName(ParameterName), Args.LayerIndex };
+							const EOpType InputOpType = Program.GetOpType(Item.At);
+							const EDataType InputDataType = GetOpDataType(InputOpType);
+
+							switch (InputDataType)
+							{
+								case EDataType::Image:
+								{
+									const FMaterial::FImageParameterData* ImageParameterData = Material->ImageParameters.Find(ParameterKey);
+
+									if (ImageParameterData)
+									{
+										const TVariant<OP::ADDRESS, TManagedPtr<const FImage>> ImageParameter = ImageParameterData->ImageParameter;
+
+										if (ImageParameter.IsType<OP::ADDRESS>())
+										{
+											PushOp(FScheduledOpInline(ImageParameter.Get<OP::ADDRESS>(), Item));
+										}
+									}
+
+									StoreNone(Item);
+									break;
+								}
+
+								case EDataType::Scalar:
+								case EDataType::Color:
+								{
+									StoreNone(Item);
+									break;
+								}
+
+								default:
+									unimplemented();
+							}
+
+							break;
+						}
+					}
+
 					break;
 				}
 
@@ -2805,15 +2908,15 @@ namespace UE::Mutable::Private
 								case EDataType::Color:
 								case EDataType::Image:
 								case EDataType::Mesh:
-								case EDataType::Material:
 								case EDataType::InstancedStruct:
+								case EDataType::Scalar:
 									// Nothing
 									break;
-									
-								case EDataType::Scalar:
-									LoadFloat(OperantAt);
+
+								case EDataType::Material:
+									LoadMaterial(OperantAt);
 									break;
-									
+
 								default: 
 									unimplemented();
 								}
@@ -2835,8 +2938,9 @@ namespace UE::Mutable::Private
 					const OP::MaterialSkeletalMeshBreakArgs Args = Program.GetOpArgs<OP::MaterialSkeletalMeshBreakArgs>(Item.At);
 					PushOp(FScheduledOpInline(Args.SlotNameStringIndex, Item));
 					PushOp(FScheduledOpInline(Args.SkeletalMesh, Item));
-					StoreNone(Item);
 					
+					StoreMaterial(Item, MakeManaged<const FMaterial>());
+
 					break;
 				}
 			
@@ -2900,6 +3004,9 @@ namespace UE::Mutable::Private
 						//Material
 						Data += sizeof(OP::MaterialModifyArgs);
 
+						// Pop Material Value
+						TManagedPtr<const FMaterial> Material = LoadMaterial(Args.Material);
+
 						// Store parameter names and schedule the parameter operations
 						for (int32 ParameterIndex = 0; ParameterIndex < ParameterCount; ++ParameterIndex)
 						{
@@ -2932,7 +3039,8 @@ namespace UE::Mutable::Private
 							}
 						}
 
-						StoreNone(Item);
+						// Push material value
+						StoreMaterial(Item, Material);
 						break;
 					}
 				}
diff --git a/Engine/Plugins/Mutable/Source/MutableRuntime/Private/MuR/CodeRunnerBegin.h b/Engine/Plugins/Mutable/Source/MutableRuntime/Private/MuR/CodeRunnerBegin.h
--- a/Engine/Plugins/Mutable/Source/MutableRuntime/Private/MuR/CodeRunnerBegin.h
+++ b/Engine/Plugins/Mutable/Source/MutableRuntime/Private/MuR/CodeRunnerBegin.h
@@ -98,7 +98,7 @@ namespace UE::Mutable::Private
 		TSet<FCacheAddressInline> IndexOther;
 	};
 	
-	using FStackValue = TVariant<bool, int32, float, TManagedPtr<const String>, TManagedPtr<const FExtensionData>, TManagedPtr<const FInstance>>;
+	using FStackValue = TVariant<bool, int32, float, TManagedPtr<const String>, TManagedPtr<const FExtensionData>, TManagedPtr<const FInstance>, TManagedPtr<const FMaterial>> ;
 
 	class FStack : public TArray<FStackValue, TInlineAllocator<32>>
 	{
@@ -142,6 +142,8 @@ namespace UE::Mutable::Private
 
 		void StoreNone(const FCacheAddressInline& To);
 
+		void StoreMaterial(const FCacheAddressInline& To, const TManagedPtr<const FMaterial>& Value);
+
 		int32 LoadInt(const FCacheAddressInline& To);
 
 		float LoadFloat(const FCacheAddressInline& To);
@@ -153,6 +155,8 @@ namespace UE::Mutable::Private
 		TManagedPtr<const String> LoadString(const FCacheAddressInline& To);
 
 		TManagedPtr<const FExtensionData> LoadExtensionData(const FCacheAddressInline& To);
+
+		TManagedPtr<const FMaterial> LoadMaterial(const FCacheAddressInline& To);
 		
 		TSharedPtr<const FModel> Model;
 		TSharedPtr<const FParameters> Params;
@@ -173,6 +177,7 @@ namespace UE::Mutable::Private
 		TMap<FCacheAddressInline, TManagedPtr<const FInstance>> ResultsInstance;
 		TMap<FCacheAddressInline, TManagedPtr<const String>> ResultsString;
 		TMap<FCacheAddressInline, TManagedPtr<const FExtensionData>> ResultsExtensionData;
+		TMap<FCacheAddressInline, TManagedPtr<const FMaterial>> ResultsMaterial;
 
 		TSharedPtr<FProgramCache> ProgramCache;
 
diff --git a/Engine/Plugins/Mutable/Source/MutableTools/Internal/MuT/NodeColourMaterialBreak.h b/Engine/Plugins/Mutable/Source/MutableTools/Internal/MuT/NodeColourMaterialBreak.h
--- a/Engine/Plugins/Mutable/Source/MutableTools/Internal/MuT/NodeColourMaterialBreak.h
+++ b/Engine/Plugins/Mutable/Source/MutableTools/Internal/MuT/NodeColourMaterialBreak.h
@@ -6,6 +6,7 @@
 #include "MuT/NodeColour.h"
 #include "MuR/Ptr.h"
 #include "MuT/NodeMaterial.h"
+#include "MuR/Material.h"
 
 #define UE_API MUTABLETOOLS_API
 
@@ -17,8 +18,7 @@ namespace UE::Mutable::Private
 	public:
 
 		Ptr<NodeMaterial> MaterialSource;
-		FName ParameterName;
-		int32 LayerIndex;
+		FParameterKey ParameterKey;
 
 	public:
 
diff --git a/Engine/Plugins/Mutable/Source/MutableTools/Internal/MuT/NodeImageMaterialBreak.h b/Engine/Plugins/Mutable/Source/MutableTools/Internal/MuT/NodeImageMaterialBreak.h
--- a/Engine/Plugins/Mutable/Source/MutableTools/Internal/MuT/NodeImageMaterialBreak.h
+++ b/Engine/Plugins/Mutable/Source/MutableTools/Internal/MuT/NodeImageMaterialBreak.h
@@ -6,6 +6,7 @@
 #include "MuT/NodeImage.h"
 #include "MuR/Ptr.h"
 #include "MuT/NodeMaterial.h"
+#include "MuR/Material.h"
 
 #define UE_API MUTABLETOOLS_API
 
@@ -17,8 +18,7 @@ namespace UE::Mutable::Private
 	public:
 
 		Ptr<NodeMaterial> MaterialSource;
-		FName ParameterName;
-		int32 LayerIndex;
+		FParameterKey ParameterKey;
 
 	public:
 
diff --git a/Engine/Plugins/Mutable/Source/MutableTools/Internal/MuT/NodeMaterialParameter.h b/Engine/Plugins/Mutable/Source/MutableTools/Internal/MuT/NodeMaterialParameter.h
--- a/Engine/Plugins/Mutable/Source/MutableTools/Internal/MuT/NodeMaterialParameter.h
+++ b/Engine/Plugins/Mutable/Source/MutableTools/Internal/MuT/NodeMaterialParameter.h
@@ -31,8 +31,6 @@ namespace UE::Mutable::Private
 		/** Ranges for the parameter in case it is multidimensional. */
 		TArray<Ptr<NodeRange>> Ranges;
 
-		TArray<FParameterKey> ColorParametersToCompile;
-		TArray<FParameterKey> ScalarParametersToCompile;
 		TMap<FParameterKey, Ptr<NodeImage>> ImageParametersToCompile;
 
 	public:
diff --git a/Engine/Plugins/Mutable/Source/MutableTools/Internal/MuT/NodeScalarMaterialBreak.h b/Engine/Plugins/Mutable/Source/MutableTools/Internal/MuT/NodeScalarMaterialBreak.h
--- a/Engine/Plugins/Mutable/Source/MutableTools/Internal/MuT/NodeScalarMaterialBreak.h
+++ b/Engine/Plugins/Mutable/Source/MutableTools/Internal/MuT/NodeScalarMaterialBreak.h
@@ -6,6 +6,7 @@
 #include "MuT/NodeScalar.h"
 #include "MuR/Ptr.h"
 #include "MuT/NodeMaterial.h"
+#include "MuR/Material.h"
 
 #define UE_API MUTABLETOOLS_API
 
@@ -17,8 +18,7 @@ namespace UE::Mutable::Private
 	public:
 
 		Ptr<NodeMaterial> MaterialSource;
-		FName ParameterName;
-		int32 LayerIndex;
+		FParameterKey ParameterKey;
 
 	public:
 
diff --git a/Engine/Plugins/Mutable/Source/MutableTools/Private/MuT/ASTOpParameter.cpp b/Engine/Plugins/Mutable/Source/MutableTools/Private/MuT/ASTOpParameter.cpp
--- a/Engine/Plugins/Mutable/Source/MutableTools/Private/MuT/ASTOpParameter.cpp
+++ b/Engine/Plugins/Mutable/Source/MutableTools/Private/MuT/ASTOpParameter.cpp
@@ -41,8 +41,6 @@ namespace UE::Mutable::Private
 			return Type == Other->Type &&
 				Parameter == Other->Parameter &&
 				Ranges == Other->Ranges &&
-				ColorParameterKeys == Other->ColorParameterKeys &&
-				ScalarParameterKeys == Other->ScalarParameterKeys &&
 				ImageOperations == Other->ImageOperations;
 		}
 		return false;
@@ -59,9 +57,6 @@ namespace UE::Mutable::Private
 		{
 			n->Ranges.Emplace(n.get(), MapChild(c.rangeSize.child()), c.rangeName, c.rangeUID);
 		}
-
-		n->ColorParameterKeys = ColorParameterKeys;
-		n->ScalarParameterKeys = ScalarParameterKeys;
 		
 		for (const TPair<FParameterKey, ASTChild>& ImageOperation : ImageOperations)
 		{
@@ -116,23 +111,6 @@ namespace UE::Mutable::Private
 				FMemory::Memzero(Args);
 				Args.Variable = (OP::ADDRESS)LinkedParameterIndex;
 				
-				// Material Parameters
-				// Scalars
-				TArray<FString> ScalarParameterNames;
-				for (const FParameterKey& ParameterKey : ScalarParameterKeys)
-				{
-					ScalarParameterNames.Add(ParameterKey.ParameterName.ToString());
-				}
-				Args.ScalarParameterNames = Program.AddConstant(ScalarParameterNames);
-
-				// Colors
-				TArray<FString> ColorParameterNames;
-				for (const FParameterKey& ParameterKey : ColorParameterKeys)
-				{
-					ColorParameterNames.Add(ParameterKey.ParameterName.ToString());
-				}
-				Args.ColorParameterNames = Program.AddConstant(ColorParameterNames);
-
 				// Images
 				TArray<FString> ImageParameterNames;
 				TArray<uint32> ImageParameterOperations;
@@ -140,7 +118,7 @@ namespace UE::Mutable::Private
 				{
 					ImageParameterNames.Add(Element.Key.ParameterName.ToString());
 					ImageParameterOperations.Add(Element.Value->linkedAddress);
-				}
+		 		}
 
 				Args.ImageParameterNames = Program.AddConstant(ImageParameterNames);
 				Args.ImageParameterAddress = Program.AddConstant(ImageParameterOperations);
@@ -156,17 +134,6 @@ namespace UE::Mutable::Private
 
 				AppendCode(Program.ByteCode, Args);
 
-				// Store layer indexes
-				for (const FParameterKey& ParameterKey : ScalarParameterKeys)
-				{
-					AppendCode(Program.ByteCode, ParameterKey.LayerIndex);
-				}
-
-				for (const FParameterKey& ParameterKey : ColorParameterKeys)
-				{
-					AppendCode(Program.ByteCode, ParameterKey.LayerIndex);
-				}
-
 				for (const TPair<FParameterKey, ASTChild>& Element : ImageOperations)
 				{
 					AppendCode(Program.ByteCode, Element.Key.LayerIndex);
diff --git a/Engine/Plugins/Mutable/Source/MutableTools/Private/MuT/ASTOpParameter.h b/Engine/Plugins/Mutable/Source/MutableTools/Private/MuT/ASTOpParameter.h
--- a/Engine/Plugins/Mutable/Source/MutableTools/Private/MuT/ASTOpParameter.h
+++ b/Engine/Plugins/Mutable/Source/MutableTools/Private/MuT/ASTOpParameter.h
@@ -29,10 +29,6 @@ namespace UE::Mutable::Private
 		//** Ranges adding dimensions to this parameter. */
 		TArray<FRangeData> Ranges;
 
-		/** List of all the parameters that this material can be broken into. */
-		TArray<FParameterKey> ScalarParameterKeys;
-		TArray<FParameterKey> ColorParameterKeys;
-
 		/** Image Parameters are evaluated at runtime. */
 		TMap<FParameterKey, ASTChild> ImageOperations;
 
diff --git a/Engine/Plugins/Mutable/Source/MutableTools/Private/MuT/CodeGenerator.cpp b/Engine/Plugins/Mutable/Source/MutableTools/Private/MuT/CodeGenerator.cpp
--- a/Engine/Plugins/Mutable/Source/MutableTools/Private/MuT/CodeGenerator.cpp
+++ b/Engine/Plugins/Mutable/Source/MutableTools/Private/MuT/CodeGenerator.cpp
@@ -2562,9 +2562,7 @@ namespace UE::Mutable::Private
 		const void* ErrorContext)
 	{
 		Ptr<ASTOp> LastImageOp = BaseImageOp;
-
 		int32 CurrentLOD = Options.LODIndex;
-		check(CurrentLOD >= 0);
 
 		// Process patch image modifiers (from edit modifiers)
 		for (const FirstPassGenerator::FModifier& Modifier : Modifiers)
diff --git a/Engine/Plugins/Mutable/Source/MutableTools/Private/MuT/CodeGenerator_Color.cpp b/Engine/Plugins/Mutable/Source/MutableTools/Private/MuT/CodeGenerator_Color.cpp
--- a/Engine/Plugins/Mutable/Source/MutableTools/Private/MuT/CodeGenerator_Color.cpp
+++ b/Engine/Plugins/Mutable/Source/MutableTools/Private/MuT/CodeGenerator_Color.cpp
@@ -559,7 +559,7 @@ namespace UE::Mutable::Private
 		Ptr<ASTOpMaterialBreak> MaterialOp = new ASTOpMaterialBreak();
 
 		// Set the key that identifies the break node parameter
-		MaterialOp->ParameterKey = { Typed->ParameterName, Typed->LayerIndex };
+		MaterialOp->ParameterKey = Typed->ParameterKey;
 
 		// Generate Material Source
 		FMaterialGenerationOptions MaterialSourceOptions;
diff --git a/Engine/Plugins/Mutable/Source/MutableTools/Private/MuT/CodeGenerator_Image.cpp b/Engine/Plugins/Mutable/Source/MutableTools/Private/MuT/CodeGenerator_Image.cpp
--- a/Engine/Plugins/Mutable/Source/MutableTools/Private/MuT/CodeGenerator_Image.cpp
+++ b/Engine/Plugins/Mutable/Source/MutableTools/Private/MuT/CodeGenerator_Image.cpp
@@ -1969,7 +1969,7 @@ namespace UE::Mutable::Private
 		Ptr<ASTOpMaterialBreak> MaterialOp = new ASTOpMaterialBreak();
 
 		// Set the parameter key that this break node will generate
-		MaterialOp->ParameterKey = { Node->ParameterName, Node->LayerIndex };
+		MaterialOp->ParameterKey = Node->ParameterKey;
 
 		// Generate Material Source
 		FMaterialGenerationResult MaterialSourceResult;
diff --git a/Engine/Plugins/Mutable/Source/MutableTools/Private/MuT/CodeGenerator_Material.cpp b/Engine/Plugins/Mutable/Source/MutableTools/Private/MuT/CodeGenerator_Material.cpp
--- a/Engine/Plugins/Mutable/Source/MutableTools/Private/MuT/CodeGenerator_Material.cpp
+++ b/Engine/Plugins/Mutable/Source/MutableTools/Private/MuT/CodeGenerator_Material.cpp
@@ -106,6 +106,12 @@ namespace UE::Mutable::Private
 		default: 
 			check(false);
 		}
+
+		// Cache the result
+		{
+			UE::TUniqueLock Lock(GeneratedMaterials.Mutex);
+			GeneratedMaterials.Map.Add(Key, Result);
+		}
 	}
 
 
@@ -312,10 +318,6 @@ namespace UE::Mutable::Private
 			GenerateRange(RangeResult, InOptions, RangeNode);
 			MaterialOp->Ranges.Emplace(MaterialOp.get(), RangeResult.sizeOp, RangeResult.rangeName, RangeResult.rangeUID);
 		}
-		
-		// Copy Scalar and Color Names
-		MaterialOp->ColorParameterKeys = InNode->ColorParametersToCompile;
-		MaterialOp->ScalarParameterKeys = InNode->ScalarParametersToCompile;
 
 		// Generate Image Nodes
 		FImageGenerationResult ImageResult;
@@ -550,7 +552,6 @@ namespace UE::Mutable::Private
 				RectInCells.min = { 0,0 };
 				RectInCells.size = { FakeLayoutSize ,FakeLayoutSize };
 
-
 				imageAd = ApplyImageBlockModifiers(Modifiers, ImageOptions, imageAd, ImageData.Key, GridSize, LayoutBlockDesc, RectInCells, ErrorContext);
 
 				check(imageAd);
diff --git a/Engine/Plugins/Mutable/Source/MutableTools/Private/MuT/CodeGenerator_Scalar.cpp b/Engine/Plugins/Mutable/Source/MutableTools/Private/MuT/CodeGenerator_Scalar.cpp
--- a/Engine/Plugins/Mutable/Source/MutableTools/Private/MuT/CodeGenerator_Scalar.cpp
+++ b/Engine/Plugins/Mutable/Source/MutableTools/Private/MuT/CodeGenerator_Scalar.cpp
@@ -440,7 +440,7 @@ namespace UE::Mutable::Private
 		Ptr<ASTOpMaterialBreak> MaterialOp = new ASTOpMaterialBreak();
 
 		// Set the key that identifies the break node parameter
-		MaterialOp->ParameterKey = { Node->ParameterName, Node->LayerIndex };
+		MaterialOp->ParameterKey = Node->ParameterKey;
 
 		// Generate Material Source
 		FMaterialGenerationOptions MaterialSourceOptions;
