[Mutable] Rename Material nodes to SkeletalMeshSection.

#rnx
#Synced-CL 49271159

diff --git a/Engine/Plugins/Mutable/Config/DefaultMutable.ini b/Engine/Plugins/Mutable/Config/DefaultMutable.ini
--- a/Engine/Plugins/Mutable/Config/DefaultMutable.ini
+++ b/Engine/Plugins/Mutable/Config/DefaultMutable.ini
@@ -12,8 +12,8 @@
 +ClassRedirects=(OldName="/Script/CustomizableObjectEditor.CustomizableObjectNodeMeshClipDeform", NewName="/Script/CustomizableObjectEditor.CustomizableObjectNodeModifierClipDeform")
 +ClassRedirects=(OldName="/Script/CustomizableObjectEditor.CustomizableObjectNodeMeshClipWithMesh", NewName="/Script/CustomizableObjectEditor.CustomizableObjectNodeModifierClipWithMesh")
 +ClassRedirects=(OldName="/Script/CustomizableObjectEditor.CustomizableObjectNodeMeshClipMorph", NewName="/Script/CustomizableObjectEditor.CustomizableObjectNodeModifierClipMorph")
-+ClassRedirects=(OldName="/Script/CustomizableObjectEditor.CustomizableObjectNodeExtendMaterial", NewName="/Script/CustomizableObjectEditor.CustomizableObjectNodeModifierExtendMeshSection")
-+ClassRedirects=(OldName="/Script/CustomizableObjectEditor.CustomizableObjectNodeEditMaterial", NewName="/Script/CustomizableObjectEditor.CustomizableObjectNodeModifierEditMeshSection")
++ClassRedirects=(OldName="/Script/CustomizableObjectEditor.CustomizableObjectNodeExtendMaterial", NewName="/Script/CustomizableObjectEditor.CONodeModifierExtendSkeletalMeshSection")
++ClassRedirects=(OldName="/Script/CustomizableObjectEditor.CustomizableObjectNodeEditMaterial", NewName="/Script/CustomizableObjectEditor.CONodeModifierEditSkeletalMeshSection")
 +ClassRedirects=(OldName="/Script/CustomizableObjectEditor.CustomizableObjectNodeMorphMaterial", NewName="/Script/CustomizableObjectEditor.CustomizableObjectNodeModifierMorphMeshSection")
 +ClassRedirects=(OldName="/Script/CustomizableObjectEditor.CustomizableObjectNodeRemoveMesh", NewName="/Script/CustomizableObjectEditor.CustomizableObjectNodeModifierRemoveMesh")
 +ClassRedirects=(OldName="/Script/CustomizableObjectEditor.CustomizableObjectNodeRemoveMeshBlocks", NewName="/Script/CustomizableObjectEditor.CustomizableObjectNodeModifierRemoveMeshBlocks")
@@ -32,6 +32,16 @@
 ; Pin Data renames
 +ClassRedirects=(OldName="/Script/CustomizableObjectEditor.CONodeMaterialBreakPinData",NewName="/Script/CustomizableObjectEditor.CONodeMaterialBreakParameterPinData")
 +ClassRedirects=(OldName="/Script/CustomizableObjectEditor.CONodeMaterialModifyPinData",NewName="/Script/CustomizableObjectEditor.CONodeMaterialModifyParameterPinData")
++ClassRedirects=(OldName="/Script/CustomizableObjectEditor.CustomizableObjectNodeMaterial",NewName="/Script/CustomizableObjectEditor.CONodeSkeletalMeshSection")
++ClassRedirects=(OldName="/Script/CustomizableObjectEditor.CustomizableObjectNodeMaterialPinDataParameter",NewName="/Script/CustomizableObjectEditor.CONodeSkeletalMeshSectionPinDataParameter")
++ClassRedirects=(OldName="/Script/CustomizableObjectEditor.CustomizableObjectNodeMaterialBase",NewName="/Script/CustomizableObjectEditor.CONodeSkeletalMeshSectionBase")
++ClassRedirects=(OldName="/Script/CustomizableObjectEditor.CustomizableObjectNodeMaterialSwitch",NewName="/Script/CustomizableObjectEditor.CONodeSkeletalMeshSectionSwitch")
++ClassRedirects=(OldName="/Script/CustomizableObjectEditor.CustomizableObjectNodeModifierEditMeshSection",NewName="/Script/CustomizableObjectEditor.CONodeModifierEditSkeletalMeshSection")
++ClassRedirects=(OldName="/Script/CustomizableObjectEditor.CustomizableObjectNodeModifierExtendMeshSection",NewName="/Script/CustomizableObjectEditor.CONodeModifierExtendSkeletalMeshSection")
++ClassRedirects=(OldName="/Script/CustomizableObjectEditor.CustomizableObjectNodeModifierWithMaterial",NewName="/Script/CustomizableObjectEditor.CONodeModifierSkeletalMeshSection")
++ClassRedirects=(OldName="/Script/CustomizableObjectEditor.CustomizableObjectNodeModifierEditSkeletalMeshSection",NewName="/Script/CustomizableObjectEditor.CONodeModifierEditSkeletalMeshSection")
++ClassRedirects=(OldName="/Script/CustomizableObjectEditor.CustomizableObjectNodeCopyMaterial",NewName="/Script/CustomizableObjectEditor.CONodeCopySkeletalMeshSection")
++ClassRedirects=(OldName="/Script/CustomizableObjectEditor.CustomizableObjectNodeMaterialVariation",NewName="/Script/CustomizableObjectEditor.CONodeSkeletalMeshSectionVariation")
 
 ; Settings for the Mutable plugin
 [Features]
diff --git a/Engine/Plugins/Mutable/Source/CustomizableObjectEditor/Private/MuCOE/CustomizableObjectEditorModule.cpp b/Engine/Plugins/Mutable/Source/CustomizableObjectEditor/Private/MuCOE/CustomizableObjectEditorModule.cpp
--- a/Engine/Plugins/Mutable/Source/CustomizableObjectEditor/Private/MuCOE/CustomizableObjectEditorModule.cpp
+++ b/Engine/Plugins/Mutable/Source/CustomizableObjectEditor/Private/MuCOE/CustomizableObjectEditorModule.cpp
@@ -35,12 +35,12 @@
 #include "MuCOE/Nodes/CustomizableObjectNodeExternalPinDetails.h"
 #include "MuCOE/Nodes/CustomizableObjectNodeMacroInstance.h"
 #include "MuCOE/Nodes/CustomizableObjectNodeMacroInstanceDetails.h"
-#include "MuCOE/Nodes/CustomizableObjectNodeMaterial.h"
+#include "MuCOE/Nodes/CONodeSkeletalMeshSection.h"
 #include "MuCOE/Nodes/CustomizableObjectNodeMaterialParameter.h"
 #include "MuCOE/Nodes/CustomizableObjectNodeMeshSectionDetails.h"
-#include "MuCOE/Nodes/CustomizableObjectNodeModifierEditMeshSection.h"
+#include "MuCOE/Nodes/CONodeModifierEditSkeletalMeshSection.h"
 #include "MuCOE/Nodes/CustomizableObjectNodeModifierEditMeshSectionDetails.h"
-#include "MuCOE/Nodes/CustomizableObjectNodeModifierExtendMeshSection.h"
+#include "MuCOE/Nodes/CONodeModifierExtendSkeletalMeshSection.h"
 #include "MuCOE/Nodes/CustomizableObjectNodeModifierExtendMeshSectionDetails.h"
 #include "MuCOE/Nodes/CustomizableObjectNodeModifierMorphMeshSection.h"
 #include "MuCOE/Nodes/CustomizableObjectNodeModifierMorphMeshSectionDetails.h"
@@ -248,8 +248,8 @@ void FCustomizableObjectEditorModule::StartupModule()
 	// Property views
 	// Nodes
 	FPropertyEditorModule& PropertyModule = FModuleManager::LoadModuleChecked<FPropertyEditorModule>("PropertyEditor");
-	RegisterCustomDetails(PropertyModule, UCustomizableObjectNodeModifierEditMeshSection::StaticClass(), FOnGetDetailCustomizationInstance::CreateStatic(&FCustomizableObjectNodeModifierEditMeshSectionDetails::MakeInstance));
-	RegisterCustomDetails(PropertyModule, UCustomizableObjectNodeModifierExtendMeshSection::StaticClass(), FOnGetDetailCustomizationInstance::CreateStatic(&FCustomizableObjectNodeModifierExtendMeshSectionDetails::MakeInstance));
+	RegisterCustomDetails(PropertyModule, UCONodeModifierEditSkeletalMeshSection::StaticClass(), FOnGetDetailCustomizationInstance::CreateStatic(&FCustomizableObjectNodeModifierEditMeshSectionDetails::MakeInstance));
+	RegisterCustomDetails(PropertyModule, UCONodeModifierExtendSkeletalMeshSection::StaticClass(), FOnGetDetailCustomizationInstance::CreateStatic(&FCustomizableObjectNodeModifierExtendMeshSectionDetails::MakeInstance));
 	RegisterCustomDetails(PropertyModule, UCustomizableObjectNodeModifierRemoveMesh::StaticClass(), FOnGetDetailCustomizationInstance::CreateStatic(&FCustomizableObjectNodeModifierRemoveMeshDetails::MakeInstance));
 	RegisterCustomDetails(PropertyModule, UCustomizableObjectNodeModifierRemoveMeshBlocks::StaticClass(), FOnGetDetailCustomizationInstance::CreateStatic(&FCustomizableObjectNodeModifierRemoveMeshBlocksDetails::MakeInstance));
 	RegisterCustomDetails(PropertyModule, UCustomizableObjectNodeModifierMorphMeshSection::StaticClass(), FOnGetDetailCustomizationInstance::CreateStatic(&FCustomizableObjectNodeModifierMorphMeshSectionDetails::MakeInstance));
@@ -265,7 +265,7 @@ void FCustomizableObjectEditorModule::StartupModule()
 	RegisterCustomDetails(PropertyModule, UCustomizableObjectNodeProjectorConstant::StaticClass(), FOnGetDetailCustomizationInstance::CreateStatic(&FCustomizableObjectNodeProjectorParameterDetails::MakeInstance));
 	RegisterCustomDetails(PropertyModule, UCustomizableObjectNodeMeshMorph::StaticClass(), FOnGetDetailCustomizationInstance::CreateStatic(&FCustomizableObjectNodeMeshMorphDetails::MakeInstance));
 	RegisterCustomDetails(PropertyModule, UCustomizableObjectNodeExternalPin::StaticClass(), FOnGetDetailCustomizationInstance::CreateStatic(&FCustomizableObjectNodeExternalPinDetails::MakeInstance));
-	RegisterCustomDetails(PropertyModule, UCustomizableObjectNodeMaterial::StaticClass(), FOnGetDetailCustomizationInstance::CreateStatic(&FCustomizableObjectNodeMeshSectionDetails::MakeInstance));
+	RegisterCustomDetails(PropertyModule, UCONodeSkeletalMeshSection::StaticClass(), FOnGetDetailCustomizationInstance::CreateStatic(&FCustomizableObjectNodeMeshSectionDetails::MakeInstance));
 	RegisterCustomDetails(PropertyModule, UCustomizableObjectNodeSkeletalMesh::StaticClass(), FOnGetDetailCustomizationInstance::CreateStatic(&FCustomizableObjectNodeSkeletalMeshDetails::MakeInstance));
 	RegisterCustomDetails(PropertyModule, UCustomizableObjectNodeStaticMesh::StaticClass(), FOnGetDetailCustomizationInstance::CreateStatic(&FCustomizableObjectNodeDetails::MakeInstance));
 	RegisterCustomDetails(PropertyModule, UCustomizableObjectNodeSkeletalMeshParameter::StaticClass(), FOnGetDetailCustomizationInstance::CreateStatic(&FCustomizableObjectNodeSkeletalMeshParameterDetails::MakeInstance));
@@ -1048,7 +1048,7 @@ TMap<FName, FGuid> FCustomizableObjectEditorModule::GetParticipatingObjects(cons
 			AddObjectSoft(MaterialParamNode->ReferenceValue.ToSoftObjectPath());
 		}
 
-		else if (UCustomizableObjectNodeMaterial* MatNode = Cast<UCustomizableObjectNodeMaterial>(&Node))
+		else if (UCONodeSkeletalMeshSection* MatNode = Cast<UCONodeSkeletalMeshSection>(&Node))
 		{
 			AddObjectSoft(MatNode->GetMaterial());
 			int32 NumImages = MatNode->GetNumParameters(EMaterialParameterType::Texture);
diff --git a/Engine/Plugins/Mutable/Source/CustomizableObjectEditor/Private/MuCOE/CustomizableObjectEditor_Deprecated.h b/Engine/Plugins/Mutable/Source/CustomizableObjectEditor/Private/MuCOE/CustomizableObjectEditor_Deprecated.h
--- a/Engine/Plugins/Mutable/Source/CustomizableObjectEditor/Private/MuCOE/CustomizableObjectEditor_Deprecated.h
+++ b/Engine/Plugins/Mutable/Source/CustomizableObjectEditor/Private/MuCOE/CustomizableObjectEditor_Deprecated.h
@@ -14,7 +14,7 @@ struct FLODReductionSettings;
 
 // Place to hide all deprecated data structures. They are still needed for deserialization backwards compatibility.
 
-// UCustomizableObjectNodeMaterial
+// UCONodeSkeletalMeshSection
 // Deprecated, do not use!
 USTRUCT()
 struct FCustomizableObjectNodeMaterialImage
@@ -72,7 +72,7 @@ struct FCustomizableObjectNodeMaterialScalar
 // Deprecated, do not use!
 constexpr int32 UV_LAYOUT_DEFAULT = -2;
 
-// UCustomizableObjectNodeModifierEditMeshSection
+// UCONodeModifierEditSkeletalMeshSection
 // Deprecated, do not use!
 USTRUCT()
 struct FCustomizableObjectNodeEditMaterialImage
@@ -83,7 +83,7 @@ struct FCustomizableObjectNodeEditMaterialImage
 	FString Name;
 };
 
-// UCustomizableObjectNodeModifierExtendMeshSection
+// UCONodeModifierExtendSkeletalMeshSection
 // Deprecated, do not use!
 USTRUCT()
 struct FCustomizableObjectNodeExtendMaterialImage
diff --git a/Engine/Plugins/Mutable/Source/CustomizableObjectEditor/Private/MuCOE/EdGraphSchema_CustomizableObject.cpp b/Engine/Plugins/Mutable/Source/CustomizableObjectEditor/Private/MuCOE/EdGraphSchema_CustomizableObject.cpp
--- a/Engine/Plugins/Mutable/Source/CustomizableObjectEditor/Private/MuCOE/EdGraphSchema_CustomizableObject.cpp
+++ b/Engine/Plugins/Mutable/Source/CustomizableObjectEditor/Private/MuCOE/EdGraphSchema_CustomizableObject.cpp
@@ -37,7 +37,7 @@
 #include "MuCOE/Nodes/CustomizableObjectNodeColorSwitch.h"
 #include "MuCOE/Nodes/CustomizableObjectNodeColorVariation.h"
 #include "MuCOE/Nodes/CustomizableObjectNodeColorToSRGB.h"
-#include "MuCOE/Nodes/CustomizableObjectNodeCopyMaterial.h"
+#include "MuCOE/Nodes/CONodeCopySkeletalMeshSection.h"
 #include "MuCOE/Nodes/CustomizableObjectNodeCurve.h"
 #include "MuCOE/Nodes/CustomizableObjectNodeEnumParameter.h"
 #include "MuCOE/Nodes/CustomizableObjectNodeExposePin.h"
@@ -50,10 +50,10 @@
 #include "MuCOE/Nodes/CustomizableObjectNodeGroupProjectorParameter.h"
 #include "MuCOE/Nodes/CustomizableObjectNodeMacroInstance.h"
 #include "MuCOE/Nodes/CustomizableObjectNodeMaterialParameter.h"
-#include "MuCOE/Nodes/CustomizableObjectNodeMaterialSwitch.h"
-#include "MuCOE/Nodes/CustomizableObjectNodeMaterialVariation.h"
-#include "MuCOE/Nodes/CustomizableObjectNodeModifierExtendMeshSection.h"
-#include "MuCOE/Nodes/CustomizableObjectNodeModifierEditMeshSection.h"
+#include "MuCOE/Nodes/CONodeSkeletalMeshSectionSwitch.h"
+#include "MuCOE/Nodes/CONodeSkeletalMeshSectionVariation.h"
+#include "MuCOE/Nodes/CONodeModifierExtendSkeletalMeshSection.h"
+#include "MuCOE/Nodes/CONodeModifierEditSkeletalMeshSection.h"
 #include "MuCOE/Nodes/CustomizableObjectNodeModifierRemoveMesh.h"
 #include "MuCOE/Nodes/CustomizableObjectNodeModifierRemoveMeshBlocks.h"
 #include "MuCOE/Nodes/CustomizableObjectNodeModifierMorphMeshSection.h"
@@ -474,10 +474,10 @@ void UEdGraphSchema_CustomizableObject::GetGraphContextActions(FGraphContextMenu
 	{
 		UCustomizableObjectNode* ObjectTemplateNodes[]
 		{
-			ContextMenuBuilder.CreateTemplateNode<UCustomizableObjectNodeMaterial>(),
-			ContextMenuBuilder.CreateTemplateNode<UCustomizableObjectNodeMaterialVariation>(),
-			ContextMenuBuilder.CreateTemplateNode<UCustomizableObjectNodeMaterialSwitch>(),
-			ContextMenuBuilder.CreateTemplateNode<UCustomizableObjectNodeCopyMaterial>(),
+			ContextMenuBuilder.CreateTemplateNode<UCONodeSkeletalMeshSection>(),
+			ContextMenuBuilder.CreateTemplateNode<UCONodeSkeletalMeshSectionVariation>(),
+			ContextMenuBuilder.CreateTemplateNode<UCONodeSkeletalMeshSectionSwitch>(),
+			ContextMenuBuilder.CreateTemplateNode<UCONodeCopySkeletalMeshSection>(),
 			ContextMenuBuilder.CreateTemplateNode<UCustomizableObjectNodeTable>(),
 			ContextMenuBuilder.CreateTemplateNode<UCustomizableObjectNodeMacroInstance>(),
 			ContextMenuBuilder.CreateTemplateNode<UCustomizableObjectNodeStaticString>()
@@ -489,10 +489,10 @@ void UEdGraphSchema_CustomizableObject::GetGraphContextActions(FGraphContextMenu
 	{
 		UCustomizableObjectNode* ModifierTemplateNodes[]
 		{
-			ContextMenuBuilder.CreateTemplateNode<UCustomizableObjectNodeModifierExtendMeshSection>(),
+			ContextMenuBuilder.CreateTemplateNode<UCONodeModifierExtendSkeletalMeshSection>(),
 			ContextMenuBuilder.CreateTemplateNode<UCustomizableObjectNodeModifierRemoveMesh>(),
 			ContextMenuBuilder.CreateTemplateNode<UCustomizableObjectNodeModifierRemoveMeshBlocks>(),
-			ContextMenuBuilder.CreateTemplateNode<UCustomizableObjectNodeModifierEditMeshSection>(),
+			ContextMenuBuilder.CreateTemplateNode<UCONodeModifierEditSkeletalMeshSection>(),
 			ContextMenuBuilder.CreateTemplateNode<UCustomizableObjectNodeModifierMorphMeshSection>(),
 			ContextMenuBuilder.CreateTemplateNode<UCustomizableObjectNodeModifierClipMorph>(),
 			ContextMenuBuilder.CreateTemplateNode<UCustomizableObjectNodeModifierClipWithMesh>(),
@@ -1335,7 +1335,7 @@ void UEdGraphSchema_CustomizableObject::DroppedAssetsOnGraph(const TArray<FAsset
 		case ESpawnableObjectType::UMaterialInterface:
 			{
 				UMaterialInterface* Material = Cast<UMaterialInterface>(Object);
-				UCustomizableObjectNodeMaterial* Node = NewObject<UCustomizableObjectNodeMaterial>(Graph);
+				UCONodeSkeletalMeshSection* Node = NewObject<UCONodeSkeletalMeshSection>(Graph);
 				Node->SetMaterial(Material);
 				GraphNode = Node;
 				break;
diff --git a/Engine/Plugins/Mutable/Source/CustomizableObjectEditor/Private/MuCOE/GenerateMutableSource/GenerateMutableSource.cpp b/Engine/Plugins/Mutable/Source/CustomizableObjectEditor/Private/MuCOE/GenerateMutableSource/GenerateMutableSource.cpp
--- a/Engine/Plugins/Mutable/Source/CustomizableObjectEditor/Private/MuCOE/GenerateMutableSource/GenerateMutableSource.cpp
+++ b/Engine/Plugins/Mutable/Source/CustomizableObjectEditor/Private/MuCOE/GenerateMutableSource/GenerateMutableSource.cpp
@@ -18,7 +18,7 @@
 #include "MuCOE/GraphTraversal.h"
 #include "MuCOE/ICustomizableObjectExtensionNode.h"
 #include "MuCOE/Nodes/CustomizableObjectNodeMacroInstance.h"
-#include "MuCOE/Nodes/CustomizableObjectNodeMaterial.h"
+#include "MuCOE/Nodes/CONodeSkeletalMeshSection.h"
 #include "MuCOE/Nodes/CustomizableObjectNodeMeshInterface.h"
 #include "MuCOE/Nodes/CustomizableObjectNodeObjectGroup.h"
 #include "MuCOE/Nodes/CustomizableObjectNodeTable.h"
@@ -1595,7 +1595,7 @@ void PopulateReferenceSkeletalMeshesData(FMutableGraphGenerationContext& Generat
 }
 
 
-uint32 GetBaseTextureSize(const FMutableGraphGenerationContext& GenerationContext, const UCustomizableObjectNodeMaterialBase* Material, const UE::Mutable::Private::FParameterKey& ImageKey)
+uint32 GetBaseTextureSize(const FMutableGraphGenerationContext& GenerationContext, const UCONodeSkeletalMeshSectionBase* Material, const UE::Mutable::Private::FParameterKey& ImageKey)
 {
 	const FGeneratedImageProperties* ImageProperties = GenerationContext.ImageProperties.Find({ Material, ImageKey });
 	return ImageProperties ? ImageProperties->TextureSize : 0;
@@ -1944,7 +1944,7 @@ uint32 GetTypeHash(const FGeneratedMutableDataTableKey& Key)
 }
 
 
-FSharedSurfaceKey::FSharedSurfaceKey(const UCustomizableObjectNodeMaterialBase* NodeMaterial, const TArray<const UCustomizableObjectNodeMacroInstance*>& MacroContext)
+FSharedSurfaceKey::FSharedSurfaceKey(const UCONodeSkeletalMeshSectionBase* NodeMaterial, const TArray<const UCustomizableObjectNodeMacroInstance*>& MacroContext)
 {
 	check(NodeMaterial);
 
diff --git a/Engine/Plugins/Mutable/Source/CustomizableObjectEditor/Private/MuCOE/GenerateMutableSource/GenerateMutableSource.h b/Engine/Plugins/Mutable/Source/CustomizableObjectEditor/Private/MuCOE/GenerateMutableSource/GenerateMutableSource.h
--- a/Engine/Plugins/Mutable/Source/CustomizableObjectEditor/Private/MuCOE/GenerateMutableSource/GenerateMutableSource.h
+++ b/Engine/Plugins/Mutable/Source/CustomizableObjectEditor/Private/MuCOE/GenerateMutableSource/GenerateMutableSource.h
@@ -14,7 +14,7 @@
 #include "MuCOE/CustomizableObjectEditorLogger.h"
 #include "MuCOE/ExtensionDataCompilerInterface.h"
 #include "MuCOE/PassthroughObjectFactory.h"
-#include "MuCOE/Nodes/CustomizableObjectNodeMaterialBase.h"
+#include "MuCOE/Nodes/CONodeSkeletalMeshSectionBase.h"
 #include "MuCOE/Nodes/CustomizableObjectNodeObject.h"
 #include "MuCOE/Nodes/CustomizableObjectNodeTable.h"
 #include "MuT/NodeComponentNew.h"
@@ -43,7 +43,7 @@
 #include "Tasks/Task.h"
 
 class UCustomizableObjectNodeGroupProjectorParameter;
-class UCustomizableObjectNodeMaterialBase;
+class UCONodeSkeletalMeshSectionBase;
 class UCustomizableObjectNodeModifierClipWithMesh;
 struct FCustomizableObjectClothingAssetData;
 
@@ -53,7 +53,7 @@ class UAnimInstance;
 class UCompositeDataTable;
 class UCONodeMaterialConstant;
 class UCustomizableObjectNodeMacroInstance;
-class UCustomizableObjectNodeMaterial;
+class UCONodeSkeletalMeshSection;
 class UCustomizableObjectNodeMaterialParameter;
 class UCustomizableObjectNodeMeshMorph;
 class UCustomizableObjectNodeMeshMorphStackDefinition;
@@ -434,7 +434,7 @@ struct FGroupProjectorImageInfo
 {
 	UE::Mutable::Private::Ptr<UE::Mutable::Private::NodeImage> ImageNode;
 	UE::Mutable::Private::Ptr<UE::Mutable::Private::NodeImage> ImageResizeNode;
-	UCustomizableObjectNodeMaterialBase* TypedNodeMat;
+	UCONodeSkeletalMeshSectionBase* TypedNodeMat;
 	FString TextureName;
 	FString RealTextureName;
 	int32 UVLayout = 0;
@@ -444,7 +444,7 @@ struct FGroupProjectorImageInfo
 		UE::Mutable::Private::NodeImagePtr InImageNode,
 		const FString& InTextureName,
 		const FString& InRealTextureName,
-		UCustomizableObjectNodeMaterialBase* InTypedNodeMat,
+		UCONodeSkeletalMeshSectionBase* InTypedNodeMat,
 		int32 InUVLayout,
 		int32 LayerIndex)
 		: TypedNodeMat(InTypedNodeMat), TextureName(InTextureName), RealTextureName(InRealTextureName), UVLayout(InUVLayout), LayerIndex(LayerIndex)
@@ -676,7 +676,7 @@ uint32 GetTypeHash(const FGeneratedMutableDataTableKey& Key);
 
 struct FSharedSurfaceKey
 {
-	FSharedSurfaceKey(const UCustomizableObjectNodeMaterialBase* NodeMaterial, const TArray<const UCustomizableObjectNodeMacroInstance*>& CurrentMacroContext);
+	FSharedSurfaceKey(const UCONodeSkeletalMeshSectionBase* NodeMaterial, const TArray<const UCustomizableObjectNodeMacroInstance*>& CurrentMacroContext);
 
 	bool operator==(const FSharedSurfaceKey& o) const = default;
 
@@ -1101,7 +1101,7 @@ FString GenerateAnimationInstanceTag(const int32 AnimInstanceIndex, const FName&
 
 FString GenerateGameplayTag(const FString& GameplayTag);
 
-uint32 GetBaseTextureSize(const FMutableGraphGenerationContext& GenerationContext, const UCustomizableObjectNodeMaterialBase* Material, const UE::Mutable::Private::FParameterKey& ImageKey);
+uint32 GetBaseTextureSize(const FMutableGraphGenerationContext& GenerationContext, const UCONodeSkeletalMeshSectionBase* Material, const UE::Mutable::Private::FParameterKey& ImageKey);
 
 // Computes the LOD bias for a texture given the current mesh LOD and automatic LOD settings, the reference texture settings
 // and whether it's being built for a server or not
diff --git a/Engine/Plugins/Mutable/Source/CustomizableObjectEditor/Private/MuCOE/GenerateMutableSource/GenerateMutableSourceGroupProjector.cpp b/Engine/Plugins/Mutable/Source/CustomizableObjectEditor/Private/MuCOE/GenerateMutableSource/GenerateMutableSourceGroupProjector.cpp
--- a/Engine/Plugins/Mutable/Source/CustomizableObjectEditor/Private/MuCOE/GenerateMutableSource/GenerateMutableSourceGroupProjector.cpp
+++ b/Engine/Plugins/Mutable/Source/CustomizableObjectEditor/Private/MuCOE/GenerateMutableSource/GenerateMutableSourceGroupProjector.cpp
@@ -13,9 +13,9 @@
 #include "MuCOE/GraphTraversal.h"
 #include "MuCOE/GenerateMutableSource/GenerateMutableSourceProjector.h"
 #include "MuCOE/Nodes/CustomizableObjectNodeAnimationPose.h"
-#include "MuCOE/Nodes/CustomizableObjectNodeModifierExtendMeshSection.h"
+#include "MuCOE/Nodes/CONodeModifierExtendSkeletalMeshSection.h"
 #include "MuCOE/Nodes/CustomizableObjectNodeGroupProjectorParameter.h"
-#include "MuCOE/Nodes/CustomizableObjectNodeMaterial.h"
+#include "MuCOE/Nodes/CONodeSkeletalMeshSection.h"
 #include "MuCOE/Nodes/CustomizableObjectNodeObjectGroup.h"
 #include "MuT/NodeColourConstant.h"
 #include "MuT/NodeColourFromScalars.h"
@@ -36,7 +36,7 @@ class UPoseAsset;
 
 
 UE::Mutable::Private::NodeImagePtr GenerateMutableSourceGroupProjector(const int32 NodeLOD, const int32 ImageIndex, UE::Mutable::Private::NodeMeshPtr MeshNode, FMutableGraphGenerationContext& GenerationContext,
-	UCustomizableObjectNodeMaterialBase* TypedNodeMat, UCustomizableObjectNodeModifierExtendMeshSection* TypedNodeExt, bool& bShareProjectionTexturesBetweenLODs, bool& bIsGroupProjectorImage,
+	UCONodeSkeletalMeshSectionBase* TypedNodeMat, UCONodeModifierExtendSkeletalMeshSection* TypedNodeExt, bool& bShareProjectionTexturesBetweenLODs, bool& bIsGroupProjectorImage,
 	UTexture2D*& GroupProjectionReferenceTexture)
 {
 	check(static_cast<bool>(TypedNodeMat) != static_cast<bool>(TypedNodeExt)); // XOr. TypedNodeMat valid or TypedNodeExt valid. At least one valid.
diff --git a/Engine/Plugins/Mutable/Source/CustomizableObjectEditor/Private/MuCOE/GenerateMutableSource/GenerateMutableSourceGroupProjector.h b/Engine/Plugins/Mutable/Source/CustomizableObjectEditor/Private/MuCOE/GenerateMutableSource/GenerateMutableSourceGroupProjector.h
--- a/Engine/Plugins/Mutable/Source/CustomizableObjectEditor/Private/MuCOE/GenerateMutableSource/GenerateMutableSourceGroupProjector.h
+++ b/Engine/Plugins/Mutable/Source/CustomizableObjectEditor/Private/MuCOE/GenerateMutableSource/GenerateMutableSourceGroupProjector.h
@@ -7,10 +7,10 @@
 
 class UCustomizableObjectNodeGroupProjectorParameter;
 struct FGroupProjectorTempData;
-class UCustomizableObjectNodeModifierExtendMeshSection;
+class UCONodeModifierExtendSkeletalMeshSection;
 class FString;
-class UCustomizableObjectNodeMaterial;
-class UCustomizableObjectNodeMaterialBase;
+class UCONodeSkeletalMeshSection;
+class UCONodeSkeletalMeshSectionBase;
 class UCustomizableObjectNodeObjectGroup;
 class UEdGraphPin;
 class UTexture2D;
@@ -18,7 +18,7 @@ struct FMutableGraphGenerationContext;
 
 
 UE::Mutable::Private::NodeImagePtr GenerateMutableSourceGroupProjector(const int32 NodeLOD, const int32 ImageIndex, UE::Mutable::Private::NodeMeshPtr MeshNode, FMutableGraphGenerationContext& GenerationContext,
-	UCustomizableObjectNodeMaterialBase* TypedNodeMat, UCustomizableObjectNodeModifierExtendMeshSection* TypedNodeExt, bool& bShareProjectionTexturesBetweenLODs, bool& bIsGroupProjectorImage,
+	UCONodeSkeletalMeshSectionBase* TypedNodeMat, UCONodeModifierExtendSkeletalMeshSection* TypedNodeExt, bool& bShareProjectionTexturesBetweenLODs, bool& bIsGroupProjectorImage,
 	UTexture2D*& GroupProjectionReferenceTexture);
 
 
diff --git a/Engine/Plugins/Mutable/Source/CustomizableObjectEditor/Private/MuCOE/GenerateMutableSource/GenerateMutableSourceImage.cpp b/Engine/Plugins/Mutable/Source/CustomizableObjectEditor/Private/MuCOE/GenerateMutableSource/GenerateMutableSourceImage.cpp
--- a/Engine/Plugins/Mutable/Source/CustomizableObjectEditor/Private/MuCOE/GenerateMutableSource/GenerateMutableSourceImage.cpp
+++ b/Engine/Plugins/Mutable/Source/CustomizableObjectEditor/Private/MuCOE/GenerateMutableSource/GenerateMutableSourceImage.cpp
@@ -19,7 +19,7 @@
 #include "MuCOE/Nodes/CustomizableObjectNodeMacroInstance.h"
 #include "MuCOE/Nodes/CustomizableObjectNodeMeshInterface.h"
 #include "MuCOE/Nodes/CustomizableObjectNodeTable.h"
-#include "MuCOE/Nodes/CustomizableObjectNodeMaterial.h"
+#include "MuCOE/Nodes/CONodeSkeletalMeshSection.h"
 #include "MuCOE/Nodes/CustomizableObjectNodeTexture.h"
 #include "MuCOE/Nodes/CustomizableObjectNodePassThroughTexture.h"
 #include "MuCOE/Nodes/CustomizableObjectNodeTextureBinarise.h"
diff --git a/Engine/Plugins/Mutable/Source/CustomizableObjectEditor/Private/MuCOE/GenerateMutableSource/GenerateMutableSourceMaterial.cpp b/Engine/Plugins/Mutable/Source/CustomizableObjectEditor/Private/MuCOE/GenerateMutableSource/GenerateMutableSourceMaterial.cpp
--- a/Engine/Plugins/Mutable/Source/CustomizableObjectEditor/Private/MuCOE/GenerateMutableSource/GenerateMutableSourceMaterial.cpp
+++ b/Engine/Plugins/Mutable/Source/CustomizableObjectEditor/Private/MuCOE/GenerateMutableSource/GenerateMutableSourceMaterial.cpp
@@ -23,7 +23,7 @@
 #include "MuCOE/GenerateMutableSource/GenerateMutableSourceMacro.h"
 #include "MuCOE/GenerateMutableSource/GenerateMutableSourceTable.h"
 #include "MuCOE/GraphTraversal.h"
-#include "MuCOE/Nodes/CustomizableObjectNodeMaterial.h"
+#include "MuCOE/Nodes/CONodeSkeletalMeshSection.h"
 #include "MuCOE/Nodes/CustomizableObjectNodeMaterialParameter.h"
 #include "MuCOE/Nodes/CustomizableObjectNodeMacroInstance.h"
 #include "MuCOE/Nodes/CustomizableObjectNodeTable.h"
@@ -134,7 +134,7 @@ UE::Mutable::Private::Ptr<UE::Mutable::Private::NodeMaterial> GenerateMutableSou
 
 	UE::Mutable::Private::Ptr<UE::Mutable::Private::NodeMaterial> Result;
 	
-	if (UCustomizableObjectNodeMaterialBase* TypedNodeMaterialBase = Cast<UCustomizableObjectNodeMaterialBase>(Node))
+	if (UCONodeSkeletalMeshSectionBase* TypedNodeMaterialBase = Cast<UCONodeSkeletalMeshSectionBase>(Node))
 	{
 		check(MaterialOptions.SurfaceNode.get());
 
diff --git a/Engine/Plugins/Mutable/Source/CustomizableObjectEditor/Private/MuCOE/GenerateMutableSource/GenerateMutableSourceMesh.cpp b/Engine/Plugins/Mutable/Source/CustomizableObjectEditor/Private/MuCOE/GenerateMutableSource/GenerateMutableSourceMesh.cpp
--- a/Engine/Plugins/Mutable/Source/CustomizableObjectEditor/Private/MuCOE/GenerateMutableSource/GenerateMutableSourceMesh.cpp
+++ b/Engine/Plugins/Mutable/Source/CustomizableObjectEditor/Private/MuCOE/GenerateMutableSource/GenerateMutableSourceMesh.cpp
@@ -31,7 +31,7 @@
 #include "MuCOE/Nodes/CustomizableObjectNodeFloatConstant.h"
 #include "MuCOE/Nodes/CustomizableObjectNodeFloatParameter.h"
 #include "MuCOE/Nodes/CustomizableObjectNodeMacroInstance.h"
-#include "MuCOE/Nodes/CustomizableObjectNodeMaterial.h"
+#include "MuCOE/Nodes/CONodeSkeletalMeshSection.h"
 #include "MuCOE/Nodes/CustomizableObjectNodeMeshMorph.h"
 #include "MuCOE/Nodes/CustomizableObjectNodeMeshMorphStackApplication.h"
 #include "MuCOE/Nodes/CustomizableObjectNodeMeshMorphStackDefinition.h"
diff --git a/Engine/Plugins/Mutable/Source/CustomizableObjectEditor/Private/MuCOE/GenerateMutableSource/GenerateMutableSourceModifier.cpp b/Engine/Plugins/Mutable/Source/CustomizableObjectEditor/Private/MuCOE/GenerateMutableSource/GenerateMutableSourceModifier.cpp
--- a/Engine/Plugins/Mutable/Source/CustomizableObjectEditor/Private/MuCOE/GenerateMutableSource/GenerateMutableSourceModifier.cpp
+++ b/Engine/Plugins/Mutable/Source/CustomizableObjectEditor/Private/MuCOE/GenerateMutableSource/GenerateMutableSourceModifier.cpp
@@ -21,8 +21,8 @@
 #include "MuCOE/Nodes/CustomizableObjectNodeModifierClipMorph.h"
 #include "MuCOE/Nodes/CustomizableObjectNodeModifierClipWithMesh.h"
 #include "MuCOE/Nodes/CustomizableObjectNodeModifierClipWithUVMask.h"
-#include "MuCOE/Nodes/CustomizableObjectNodeModifierExtendMeshSection.h"
-#include "MuCOE/Nodes/CustomizableObjectNodeModifierEditMeshSection.h"
+#include "MuCOE/Nodes/CONodeModifierExtendSkeletalMeshSection.h"
+#include "MuCOE/Nodes/CONodeModifierEditSkeletalMeshSection.h"
 #include "MuCOE/Nodes/CustomizableObjectNodeModifierMorphMeshSection.h"
 #include "MuCOE/Nodes/CustomizableObjectNodeModifierRemoveMesh.h"
 #include "MuCOE/Nodes/CustomizableObjectNodeModifierRemoveMeshBlocks.h"
@@ -249,7 +249,7 @@ UE::Mutable::Private::Ptr<UE::Mutable::Private::NodeModifier> GenerateMutableSou
 		GenerationContext.MeshGenerationFlags.Pop();
 	}
 
-	else if (UCustomizableObjectNodeModifierExtendMeshSection* TypedNodeExt = Cast<UCustomizableObjectNodeModifierExtendMeshSection>(Node))
+	else if (UCONodeModifierExtendSkeletalMeshSection* TypedNodeExt = Cast<UCONodeModifierExtendSkeletalMeshSection>(Node))
 	{
 		const EMutableMeshConversionFlags ModifiersMeshFlags = EMutableMeshConversionFlags::None;
 		GenerationContext.MeshGenerationFlags.Push(ModifiersMeshFlags);
@@ -438,7 +438,7 @@ UE::Mutable::Private::Ptr<UE::Mutable::Private::NodeModifier> GenerateMutableSou
 		ClipNode->LayoutIndex = TypedNodeRemBlocks->ParentLayoutIndex;
 	}
 
-	else if (UCustomizableObjectNodeModifierEditMeshSection* TypedNodeEdit = Cast<UCustomizableObjectNodeModifierEditMeshSection>(Node))
+	else if (UCONodeModifierEditSkeletalMeshSection* TypedNodeEdit = Cast<UCONodeModifierEditSkeletalMeshSection>(Node))
 	{
 		const EMutableMeshConversionFlags ModifiersMeshFlags =
 			EMutableMeshConversionFlags::IgnoreSkinning |
diff --git a/Engine/Plugins/Mutable/Source/CustomizableObjectEditor/Private/MuCOE/GenerateMutableSource/GenerateMutableSourceSurface.cpp b/Engine/Plugins/Mutable/Source/CustomizableObjectEditor/Private/MuCOE/GenerateMutableSource/GenerateMutableSourceSurface.cpp
--- a/Engine/Plugins/Mutable/Source/CustomizableObjectEditor/Private/MuCOE/GenerateMutableSource/GenerateMutableSourceSurface.cpp
+++ b/Engine/Plugins/Mutable/Source/CustomizableObjectEditor/Private/MuCOE/GenerateMutableSource/GenerateMutableSourceSurface.cpp
@@ -27,14 +27,14 @@
 #include "MuCOE/GenerateMutableSource/GenerateMutableSourceMaterial.h"
 #include "MuCOE/GraphTraversal.h"
 #include "MuCOE/MutableUtils.h"
-#include "MuCOE/Nodes/CustomizableObjectNodeCopyMaterial.h"
+#include "MuCOE/Nodes/CONodeCopySkeletalMeshSection.h"
 #include "MuCOE/Nodes/CustomizableObjectNodeFloatConstant.h"
 #include "MuCOE/Nodes/CustomizableObjectNodeFloatParameter.h"
 #include "MuCOE/Nodes/CustomizableObjectNodeMacroInstance.h"
-#include "MuCOE/Nodes/CustomizableObjectNodeMaterialVariation.h"
-#include "MuCOE/Nodes/CustomizableObjectNodeMaterialSwitch.h"
-#include "MuCOE/Nodes/CustomizableObjectNodeModifierEditMeshSection.h"
-#include "MuCOE/Nodes/CustomizableObjectNodeModifierExtendMeshSection.h"
+#include "MuCOE/Nodes/CONodeSkeletalMeshSectionVariation.h"
+#include "MuCOE/Nodes/CONodeSkeletalMeshSectionSwitch.h"
+#include "MuCOE/Nodes/CONodeModifierEditSkeletalMeshSection.h"
+#include "MuCOE/Nodes/CONodeModifierExtendSkeletalMeshSection.h"
 #include "MuCOE/Nodes/CustomizableObjectNodeModifierMorphMeshSection.h"
 #include "MuCOE/Nodes/CustomizableObjectNodeModifierRemoveMesh.h"
 #include "MuCOE/Nodes/CustomizableObjectNodeModifierRemoveMeshBlocks.h"
@@ -90,7 +90,7 @@ UE::Mutable::Private::Ptr<UE::Mutable::Private::NodeSurface> GenerateMutableSour
 		}
 	}
 
-	if (UCustomizableObjectNodeMaterialBase* TypedNodeMat = Cast<UCustomizableObjectNodeMaterialBase>(Node))
+	if (UCONodeSkeletalMeshSectionBase* TypedNodeMat = Cast<UCONodeSkeletalMeshSectionBase>(Node))
 	{
 		if (GenerationContext.CurrentAutoLODStrategy == ECustomizableObjectAutomaticLODStrategy::AutomaticFromMesh
 			&& TypedNodeMat->GetMeshSectionMaxLOD() != INDEX_NONE && TypedNodeMat->GetMeshSectionMaxLOD() < GenerationContext.CurrentLOD)
@@ -226,7 +226,7 @@ UE::Mutable::Private::Ptr<UE::Mutable::Private::NodeSurface> GenerateMutableSour
 		SurfNode->Tags.AddUnique(TypedNodeMat->GetInternalTag());
 	}
 
-	else if (const UCustomizableObjectNodeMaterialVariation* TypedNodeVar = Cast<UCustomizableObjectNodeMaterialVariation>(Node))
+	else if (const UCONodeSkeletalMeshSectionVariation* TypedNodeVar = Cast<UCONodeSkeletalMeshSectionVariation>(Node))
 	{
 		UE::Mutable::Private::Ptr<UE::Mutable::Private::NodeSurfaceVariation> SurfNode = new UE::Mutable::Private::NodeSurfaceVariation();
 		Result = SurfNode;
@@ -282,7 +282,7 @@ UE::Mutable::Private::Ptr<UE::Mutable::Private::NodeSurface> GenerateMutableSour
 		}
 	}
 
-	else if (const UCustomizableObjectNodeMaterialSwitch* TypedNodeSwitch = Cast<UCustomizableObjectNodeMaterialSwitch>(Node))
+	else if (const UCONodeSkeletalMeshSectionSwitch* TypedNodeSwitch = Cast<UCONodeSkeletalMeshSectionSwitch>(Node))
 	{
 		// Using a lambda so control flow is easier to manage.
 		Result = [&]()
diff --git a/Engine/Plugins/Mutable/Source/CustomizableObjectEditor/Private/MuCOE/GraphTraversal.cpp b/Engine/Plugins/Mutable/Source/CustomizableObjectEditor/Private/MuCOE/GraphTraversal.cpp
--- a/Engine/Plugins/Mutable/Source/CustomizableObjectEditor/Private/MuCOE/GraphTraversal.cpp
+++ b/Engine/Plugins/Mutable/Source/CustomizableObjectEditor/Private/MuCOE/GraphTraversal.cpp
@@ -8,11 +8,11 @@
 #include "MuCOE/EdGraphSchema_CustomizableObject.h"
 #include "MuCOE/Nodes/CustomizableObjectNodeEnumParameter.h"
 #include "MuCOE/Nodes/CustomizableObjectNodeExposePin.h"
-#include "MuCOE/Nodes/CustomizableObjectNodeModifierExtendMeshSection.h"
+#include "MuCOE/Nodes/CONodeModifierExtendSkeletalMeshSection.h"
 #include "MuCOE/Nodes/CustomizableObjectNodeExternalPin.h"
-#include "MuCOE/Nodes/CustomizableObjectNodeMaterial.h"
+#include "MuCOE/Nodes/CONodeSkeletalMeshSection.h"
 #include "MuCOE/Nodes/CustomizableObjectNodeMaterialParameter.h"
-#include "MuCOE/Nodes/CustomizableObjectNodeMaterialVariation.h"
+#include "MuCOE/Nodes/CONodeSkeletalMeshSectionVariation.h"
 #include "MuCOE/Nodes/CustomizableObjectNodeMeshMorph.h"
 #include "MuCOE/Nodes/CustomizableObjectNodeMeshMorphStackApplication.h"
 #include "MuCOE/Nodes/CustomizableObjectNodeMeshMorphStackDefinition.h"
@@ -658,7 +658,7 @@ const UEdGraphPin* FindMeshBaseSource(const UEdGraphPin& Pin, const bool bOnlyLo
 		}
 	}
 
-	else if (const UCustomizableObjectNodeMaterialBase* TypedNodeMat = Cast<UCustomizableObjectNodeMaterialBase>(Node))
+	else if (const UCONodeSkeletalMeshSectionBase* TypedNodeMat = Cast<UCONodeSkeletalMeshSectionBase>(Node))
 	{
 		if (const UEdGraphPin* ConnectedPin = FollowInputPin(*TypedNodeMat->GetMeshPin()))
 		{
@@ -666,7 +666,7 @@ const UEdGraphPin* FindMeshBaseSource(const UEdGraphPin& Pin, const bool bOnlyLo
 		}
 	}
 
-	else if (const UCustomizableObjectNodeMaterialVariation* TypedNodeMatVar = Cast<UCustomizableObjectNodeMaterialVariation>(Node))
+	else if (const UCONodeSkeletalMeshSectionVariation* TypedNodeMatVar = Cast<UCONodeSkeletalMeshSectionVariation>(Node))
 	{
 		if (const UEdGraphPin* ConnectedPin = FollowInputPin(*TypedNodeMatVar->DefaultPin()))
 		{
@@ -674,7 +674,7 @@ const UEdGraphPin* FindMeshBaseSource(const UEdGraphPin& Pin, const bool bOnlyLo
 		}
 	}
 
-	else if (const UCustomizableObjectNodeModifierExtendMeshSection* TypedNodeExtend = Cast<UCustomizableObjectNodeModifierExtendMeshSection>(Node))
+	else if (const UCONodeModifierExtendSkeletalMeshSection* TypedNodeExtend = Cast<UCONodeModifierExtendSkeletalMeshSection>(Node))
 	{
 		if (const UEdGraphPin* ConnectedPin = FollowInputPin(*TypedNodeExtend->AddMeshPin()))
 		{
diff --git a/Engine/Plugins/Mutable/Source/CustomizableObjectEditor/Private/MuCOE/Nodes/CONodeCopySkeletalMeshSection.cpp b/Engine/Plugins/Mutable/Source/CustomizableObjectEditor/Private/MuCOE/Nodes/CONodeCopySkeletalMeshSection.cpp
new file mode 100644
--- /dev/null
+++ b/Engine/Plugins/Mutable/Source/CustomizableObjectEditor/Private/MuCOE/Nodes/CONodeCopySkeletalMeshSection.cpp
@@ -0,0 +1,447 @@
+// Copyright Epic Games, Inc. All Rights Reserved.
+
+#include "MuCOE/Nodes/CONodeCopySkeletalMeshSection.h"
+
+#include "MuCOE/EdGraphSchema_CustomizableObject.h"
+#include "MuCOE/GraphTraversal.h"
+#include "MuCOE/Nodes/CustomizableObjectNodeExposePin.h"
+#include "MuCOE/Nodes/CustomizableObjectNodeExternalPin.h"
+#include "MuCOE/Nodes/CustomizableObjectNodeObject.h"
+#include "MuCOE/Nodes/CustomizableObjectNodeSkeletalMesh.h"
+#include "MuCOE/Nodes/CONodeSkeletalMeshSection.h"
+#include "MuCO/CustomizableObjectCustomVersion.h"
+
+#include UE_INLINE_GENERATED_CPP_BY_NAME(CONodeCopySkeletalMeshSection)
+
+#define LOCTEXT_NAMESPACE "CustomizableObjectEditor"
+
+namespace CustomizableObjectNodeCopyMaterial
+{
+	/** Mesh input pin key */
+	const static FString MeshPinName = TEXT("Mesh_Input_Pin");
+
+	/** Material input pin key */
+	const static FString MeshSectionPinName = TEXT("MeshSection_Input_Pin");
+
+	/** Material output pin key*/
+	const static FString OutputPinName = TEXT("MeshSection_Output_Pin");
+}
+
+
+FText UCONodeCopySkeletalMeshSection::GetNodeTitle(ENodeTitleType::Type TitleType) const
+{
+	return LOCTEXT("Copy_MeshSection", "Copy Mesh Section");
+}
+
+
+void UCONodeCopySkeletalMeshSection::AllocateDefaultPins(UCustomizableObjectNodeRemapPins* RemapPins)
+{
+	const UEdGraphSchema_CustomizableObject* Schema = GetDefault<UEdGraphSchema_CustomizableObject>();
+
+	// Input pins
+	CustomCreatePin(EGPD_Input, Schema->PC_Mesh, FName(CustomizableObjectNodeCopyMaterial::MeshPinName), LOCTEXT("MeshPin", "Mesh"));
+	CustomCreatePin(EGPD_Input, Schema->PC_MeshSection, FName(CustomizableObjectNodeCopyMaterial::MeshSectionPinName), LOCTEXT("BaseMeshSectionPin", "Base Mesh Section"));
+
+	// Output pins
+	CustomCreatePin(EGPD_Output, Schema->PC_MeshSection, FName(CustomizableObjectNodeCopyMaterial::OutputPinName), LOCTEXT("MeshSectionPin", "Mesh Section"));
+}
+
+
+void UCONodeCopySkeletalMeshSection::BackwardsCompatibleFixup(int32 CustomizableObjectCustomVersion)
+{
+	Super::BackwardsCompatibleFixup(CustomizableObjectCustomVersion);
+
+	if (CustomizableObjectCustomVersion == FCustomizableObjectCustomVersion::MaterialPinsRename)
+	{
+		UEdGraphPin* BaseMaterialPin = FindPin(TEXT("Material_Input_Pin"), EEdGraphPinDirection::EGPD_Input);
+		if (BaseMaterialPin)
+		{
+			const FString PinFriendlyName = TEXT("Base Mesh Section");
+			BaseMaterialPin->PinName = FName(CustomizableObjectNodeCopyMaterial::MeshSectionPinName);
+			BaseMaterialPin->PinFriendlyName = FText::FromString(PinFriendlyName);
+		}
+
+		UEdGraphPin* MaterialPin = FindPin(TEXT("Material_Output_Pin"), EEdGraphPinDirection::EGPD_Output);
+		if (MaterialPin)
+		{
+			const FString PinFriendlyName = TEXT("Mesh Section");
+			MaterialPin->PinName = FName(CustomizableObjectNodeCopyMaterial::OutputPinName);
+			MaterialPin->PinFriendlyName = FText::FromString(PinFriendlyName);
+		}
+	}
+}
+
+
+UEdGraphPin* UCONodeCopySkeletalMeshSection::GetMeshPin() const
+{
+	return FindPin(CustomizableObjectNodeCopyMaterial::MeshPinName);
+}
+
+
+FPostImagePinModeChangedDelegate* UCONodeCopySkeletalMeshSection::GetPostImagePinModeChangedDelegate()
+{
+	if (UCONodeSkeletalMeshSection* NodeMaterial = GetMaterialNode())
+	{
+		return NodeMaterial->GetPostImagePinModeChangedDelegate();
+	}
+	else
+	{
+		return nullptr;
+	}
+}
+
+
+TArray<UCustomizableObjectLayout*> UCONodeCopySkeletalMeshSection::GetLayouts() const
+{
+	if (UCONodeSkeletalMeshSection* NodeMaterial = GetMaterialNode())
+	{
+		return NodeMaterial->GetLayouts();
+	}
+	else
+	{
+		return {};
+	}
+}
+
+
+UEdGraphPin* UCONodeCopySkeletalMeshSection::OutputPin() const
+{
+	UEdGraphPin* Pin = FindPin(CustomizableObjectNodeCopyMaterial::OutputPinName, EEdGraphPinDirection::EGPD_Output);
+
+	// Legacy name
+	if (!Pin)
+	{
+		Pin = FindPin(TEXT("Material_Output_Pin"), EEdGraphPinDirection::EGPD_Output);
+	}
+
+	return Pin;
+}
+
+
+bool UCONodeCopySkeletalMeshSection::RealMaterialDataHasChanged() const
+{
+	if (UCONodeSkeletalMeshSection* NodeMaterial = GetMaterialNode())
+	{
+		return NodeMaterial->RealMaterialDataHasChanged();
+	}
+	else
+	{
+		return false;
+	}
+}
+
+
+UEdGraphPin* UCONodeCopySkeletalMeshSection::GetEnableTagsPin() const
+{
+	if (UCONodeSkeletalMeshSection* NodeMaterial = GetMaterialNode())
+	{
+		return NodeMaterial->GetEnableTagsPin();
+	}
+	else
+	{
+		return nullptr;
+	}
+}
+
+
+UEdGraphPin* UCONodeCopySkeletalMeshSection::GetMeshSectionPin() const
+{
+	UEdGraphPin* Pin = FindPin(CustomizableObjectNodeCopyMaterial::MeshSectionPinName, EEdGraphPinDirection::EGPD_Input);
+
+	// Legacy name
+	if (!Pin)
+	{
+		Pin = FindPin(TEXT("Material_Input_Pin"), EEdGraphPinDirection::EGPD_Input);
+	}
+
+	return Pin;
+}
+
+
+UCustomizableObjectNodeSkeletalMesh* UCONodeCopySkeletalMeshSection::GetMeshNode() const
+{
+	UCustomizableObjectNodeSkeletalMesh* Result = nullptr;
+
+	UEdGraphPin* MaterialPin = GetMeshPin();
+	if (const UEdGraphPin* ConnectedPin = FollowInputPin(*MaterialPin))
+	{
+		const UEdGraphPin* SourceMeshPin = FindMeshBaseSource(*ConnectedPin, false);
+		if (SourceMeshPin)
+		{
+			Result = Cast<UCustomizableObjectNodeSkeletalMesh>(SourceMeshPin->GetOwningNode());
+		}
+	}
+
+	return Result;
+}
+
+
+UCONodeSkeletalMeshSection* UCONodeCopySkeletalMeshSection::GetMaterialNode() const
+{
+	UCONodeSkeletalMeshSection* Result = nullptr;
+
+	UEdGraphPin* MaterialPin = GetMeshSectionPin();
+	if (const UEdGraphPin* ConnectedPin = FollowInputPin(*MaterialPin))
+	{
+		return Cast<UCONodeSkeletalMeshSection>(ConnectedPin->GetOwningNode());
+	}
+
+	return Result;
+}
+
+
+bool UCONodeCopySkeletalMeshSection::CanConnect(const UEdGraphPin* InOwnedInputPin, const UEdGraphPin* InOutputPin, bool& bOutIsOtherNodeBlocklisted, bool& bOutArePinsCompatible) const
+{
+	if (!Super::CanConnect(InOwnedInputPin, InOutputPin, bOutIsOtherNodeBlocklisted, bOutIsOtherNodeBlocklisted))
+	{
+		return false;
+	}
+
+	if (InOwnedInputPin == GetMeshSectionPin())
+	{
+		const UEdGraphNode* OuputPinOwningNode = InOutputPin->GetOwningNode();
+		return (OuputPinOwningNode->IsA(UCONodeSkeletalMeshSection::StaticClass()) && !OuputPinOwningNode->IsA(UCONodeCopySkeletalMeshSection::StaticClass()))
+			|| OuputPinOwningNode->IsA(UCustomizableObjectNodeExternalPin::StaticClass());
+	}
+
+	return true;
+}
+
+
+bool UCONodeCopySkeletalMeshSection::ShouldBreakExistingConnections(const UEdGraphPin* InputPin, const UEdGraphPin* OutputPin) const
+{
+	return true;
+}
+
+
+bool UCONodeCopySkeletalMeshSection::IsNodeOutDatedAndNeedsRefresh()
+{
+	return false;
+}
+
+
+bool UCONodeCopySkeletalMeshSection::ProvidesCustomPinRelevancyTest() const
+{
+	return true;
+}
+
+
+bool UCONodeCopySkeletalMeshSection::IsPinRelevant(const UEdGraphPin* Pin) const
+{
+	const UEdGraphNode* Node = Pin->GetOwningNode();
+
+	if (Pin->Direction == EGPD_Output)
+	{
+		return (Node->IsA(UCONodeSkeletalMeshSection::StaticClass()) && !Node->IsA(UCONodeCopySkeletalMeshSection::StaticClass())) ||
+			(Node->IsA(UCustomizableObjectNodeExternalPin::StaticClass()) && Pin->PinType.PinCategory == UEdGraphSchema_CustomizableObject::PC_MeshSection) ||
+			Pin->PinType.PinCategory == UEdGraphSchema_CustomizableObject::PC_Mesh;
+	}
+	else
+	{
+		return Node->IsA(UCustomizableObjectNodeObject::StaticClass()) ||
+			(Node->IsA(UCustomizableObjectNodeExposePin::StaticClass()) && Pin->PinType.PinCategory == UEdGraphSchema_CustomizableObject::PC_MeshSection);
+			
+	}
+}
+
+
+UMaterialInterface* UCONodeCopySkeletalMeshSection::GetMaterial() const
+{
+	if (UCONodeSkeletalMeshSection* NodeMaterial = GetMaterialNode())
+	{
+		return NodeMaterial->GetMaterial();
+	}
+	else
+	{
+		return nullptr;
+	}
+}
+
+
+TArray<FString> UCONodeCopySkeletalMeshSection::GetEnableTags(TArray<const UCustomizableObjectNodeMacroInstance*>* MacroContext)
+{
+	if (UCONodeSkeletalMeshSection* NodeMaterial = GetMaterialNode())
+	{
+		return NodeMaterial->GetEnableTags(MacroContext);
+	}
+	
+	return TArray<FString>();
+}
+
+
+TArray<FString>* UCONodeCopySkeletalMeshSection::GetEnableTagsArray()
+{
+	if (UCONodeSkeletalMeshSection* NodeMaterial = GetMaterialNode())
+	{
+		return NodeMaterial->GetEnableTagsArray();
+	}
+
+	return nullptr;
+}
+
+
+UEdGraphPin* UCONodeCopySkeletalMeshSection::GetMaterialAssetPin() const
+{
+	if (UCONodeSkeletalMeshSection* NodeMaterial = GetMaterialNode())
+	{
+		return NodeMaterial->GetMaterialAssetPin();
+	}
+	else
+	{
+		return nullptr;
+	}
+}
+
+
+int32 UCONodeCopySkeletalMeshSection::GetNumParameters(EMaterialParameterType Type) const
+{
+	if (UCONodeSkeletalMeshSection* NodeMaterial = GetMaterialNode())
+	{
+		return NodeMaterial->GetNumParameters(Type);
+	}
+	else
+	{
+		return 0;
+	}
+}
+
+
+FNodeMaterialParameterId UCONodeCopySkeletalMeshSection::GetParameterId(EMaterialParameterType Type, int32 ParameterIndex) const
+{
+	UCONodeSkeletalMeshSection* NodeMaterial = GetMaterialNode();
+	check(NodeMaterial);
+
+	return NodeMaterial->GetParameterId(Type, ParameterIndex);
+}
+
+
+FName UCONodeCopySkeletalMeshSection::GetParameterName(EMaterialParameterType Type, int32 ParameterIndex) const
+{
+	UCONodeSkeletalMeshSection* NodeMaterial = GetMaterialNode();
+	check(NodeMaterial);
+	
+	return NodeMaterial->GetParameterName(Type, ParameterIndex);
+}
+
+
+int32 UCONodeCopySkeletalMeshSection::GetParameterLayerIndex(EMaterialParameterType Type, int32 ParameterIndex) const
+{
+	UCONodeSkeletalMeshSection* NodeMaterial = GetMaterialNode();
+	check(NodeMaterial);
+	
+	return NodeMaterial->GetParameterLayerIndex(Type, ParameterIndex);
+}
+
+
+FText UCONodeCopySkeletalMeshSection::GetParameterLayerName(EMaterialParameterType Type, int32 ParameterIndex) const
+{
+	UCONodeSkeletalMeshSection* NodeMaterial = GetMaterialNode();
+	check(NodeMaterial);
+	
+	return NodeMaterial->GetParameterLayerName(Type, ParameterIndex);
+}
+
+
+bool UCONodeCopySkeletalMeshSection::HasParameter(const FNodeMaterialParameterId& ParameterId) const
+{
+	if (UCONodeSkeletalMeshSection* NodeMaterial = GetMaterialNode())
+	{
+		return NodeMaterial->HasParameter(ParameterId);
+	}
+	else
+	{
+		return false;	
+	}
+}
+
+
+const UEdGraphPin* UCONodeCopySkeletalMeshSection::GetParameterPin(EMaterialParameterType Type, int32 ParameterIndex) const
+{
+	if (UCONodeSkeletalMeshSection* NodeMaterial = GetMaterialNode())
+	{
+		return NodeMaterial->GetParameterPin(Type, ParameterIndex);
+	}
+	else
+	{
+		return nullptr;	
+	}
+}
+
+
+UEdGraphPin* UCONodeCopySkeletalMeshSection::GetParameterPin(const FNodeMaterialParameterId& ParameterId) const
+{
+	if (UCONodeSkeletalMeshSection* NodeMaterial = GetMaterialNode())
+	{
+		return NodeMaterial->GetParameterPin(ParameterId);
+	}
+	else
+	{
+		return nullptr;	
+	}
+}
+
+
+bool UCONodeCopySkeletalMeshSection::IsImageMutableMode(int32 ImageIndex) const
+{
+	UCONodeSkeletalMeshSection* NodeMaterial = GetMaterialNode();
+	check(NodeMaterial);
+	
+	return NodeMaterial->IsImageMutableMode(ImageIndex);
+}
+
+
+bool UCONodeCopySkeletalMeshSection::IsImageMutableMode(const UEdGraphPin& Pin) const
+{
+	UCONodeSkeletalMeshSection* NodeMaterial = GetMaterialNode();
+	check(NodeMaterial);
+
+	return NodeMaterial->IsImageMutableMode(Pin);
+}
+
+
+UTexture2D* UCONodeCopySkeletalMeshSection::GetImageReferenceTexture(int32 ImageIndex) const
+{
+	if (UCONodeSkeletalMeshSection* NodeMaterial = GetMaterialNode())
+	{
+		return NodeMaterial->GetImageReferenceTexture(ImageIndex);
+	}
+	else
+	{
+		return nullptr;
+	}
+}
+
+
+UTexture2D* UCONodeCopySkeletalMeshSection::GetImageValue(int32 ImageIndex) const
+{
+	if (UCONodeSkeletalMeshSection* NodeMaterial = GetMaterialNode())
+	{
+		return NodeMaterial->GetImageValue(ImageIndex);
+	}
+	else
+	{
+		return nullptr;
+	}
+}
+
+
+int32 UCONodeCopySkeletalMeshSection::GetImageUVLayout(int32 ImageIndex) const
+{
+	if (UCONodeSkeletalMeshSection* NodeMaterial = GetMaterialNode())
+	{
+		return NodeMaterial->GetImageUVLayout(ImageIndex);
+	}
+	else
+	{
+		return UCustomizableObjectNodeMaterialPinDataImage::UV_LAYOUT_IGNORE;
+	}
+}
+
+
+FText UCONodeCopySkeletalMeshSection::GetTooltipText() const
+{
+	return LOCTEXT("CopyMaterial_Tooltip", "Copies a Customizable Object material.\nDuplicates all Material node input pins and properties except for the Mesh input pin.");
+}
+
+
+#undef LOCTEXT_NAMESPACE
diff --git a/Engine/Plugins/Mutable/Source/CustomizableObjectEditor/Private/MuCOE/Nodes/CONodeCopySkeletalMeshSection.h b/Engine/Plugins/Mutable/Source/CustomizableObjectEditor/Private/MuCOE/Nodes/CONodeCopySkeletalMeshSection.h
new file mode 100644
--- /dev/null
+++ b/Engine/Plugins/Mutable/Source/CustomizableObjectEditor/Private/MuCOE/Nodes/CONodeCopySkeletalMeshSection.h
@@ -0,0 +1,72 @@
+// Copyright Epic Games, Inc. All Rights Reserved.
+
+#pragma once
+
+#include "MuCOE/Nodes/CONodeSkeletalMeshSection.h"
+
+#include "CONodeCopySkeletalMeshSection.generated.h"
+
+#define UE_API CUSTOMIZABLEOBJECTEDITOR_API
+
+class UCustomizableObjectNodeRemapPins;
+class UCustomizableObjectNodeSkeletalMesh;
+class UEdGraphPin;
+class UObject;
+
+
+/** Generates a new Surface with the same connections as the Parent Material Surface, but with a different Mesh. */
+UCLASS(MinimalAPI)
+class UCONodeCopySkeletalMeshSection : public UCONodeSkeletalMeshSectionBase
+{
+public:
+	GENERATED_BODY()
+
+	// UEdGraphNode interface
+	UE_API virtual FText GetNodeTitle(ENodeTitleType::Type TitleType) const override;
+	UE_API virtual FText GetTooltipText() const override;
+
+	// UCustomizableObjectNode interface
+	UE_API virtual void AllocateDefaultPins(UCustomizableObjectNodeRemapPins* RemapPins) override;
+	UE_API virtual bool CanConnect(const UEdGraphPin* InOwnedInputPin, const UEdGraphPin* InOutputPin, bool& bOutIsOtherNodeBlocklisted, bool& bOutArePinsCompatible) const override;
+	UE_API virtual bool ShouldBreakExistingConnections(const UEdGraphPin* InputPin, const UEdGraphPin* OutputPin) const override;
+	UE_API virtual bool IsNodeOutDatedAndNeedsRefresh() override;
+	UE_API virtual bool ProvidesCustomPinRelevancyTest() const override;
+	UE_API virtual bool IsPinRelevant(const UEdGraphPin* Pin) const override;
+	UE_API virtual void BackwardsCompatibleFixup(int32 CustomizableObjectCustomVersion) override;
+	UE_API virtual TArray<FString> GetEnableTags(TArray<const UCustomizableObjectNodeMacroInstance*>* MacroContext = nullptr) override;
+	UE_API virtual TArray<FString>* GetEnableTagsArray() override;
+
+	// UCONodeSkeletalMeshSectionBase interface
+	UE_API virtual UMaterialInterface* GetMaterial() const override;
+	UE_API virtual UEdGraphPin* GetMaterialAssetPin() const override;
+	UE_API virtual int32 GetNumParameters(EMaterialParameterType Type) const override;
+	UE_API virtual FNodeMaterialParameterId GetParameterId(EMaterialParameterType Type, int32 ParameterIndex) const override;
+	UE_API virtual FName GetParameterName(EMaterialParameterType Type, int32 ParameterIndex) const override;
+	UE_API virtual int32 GetParameterLayerIndex(EMaterialParameterType Type, int32 ParameterIndex) const override;
+	UE_API virtual FText GetParameterLayerName(EMaterialParameterType Type, int32 ParameterIndex) const override;
+	UE_API virtual bool HasParameter(const FNodeMaterialParameterId& ParameterId) const override;
+	UE_API virtual const UEdGraphPin* GetParameterPin(EMaterialParameterType Type, int32 ParameterIndex) const override;
+	UE_API virtual UEdGraphPin* GetParameterPin(const FNodeMaterialParameterId& ParameterId) const override;
+	UE_API virtual bool IsImageMutableMode(int32 ImageIndex) const override;
+	UE_API virtual bool IsImageMutableMode(const UEdGraphPin& Pin) const override;
+	UE_API virtual UTexture2D* GetImageReferenceTexture(int32 ImageIndex) const override;
+	UE_API virtual UTexture2D* GetImageValue(int32 ImageIndex) const override;
+	UE_API virtual int32 GetImageUVLayout(int32 ImageIndex) const override;
+	UE_API virtual UCONodeSkeletalMeshSection* GetMaterialNode() const override;
+	UE_API virtual UEdGraphPin* GetMeshPin() const override;
+	UE_API virtual FPostImagePinModeChangedDelegate* GetPostImagePinModeChangedDelegate() override;
+	UE_API virtual TArray<UCustomizableObjectLayout*> GetLayouts() const override;
+	UE_API virtual UEdGraphPin* OutputPin() const override;
+	UE_API virtual bool RealMaterialDataHasChanged() const override;
+	UE_API virtual UEdGraphPin* GetEnableTagsPin() const override;
+	
+	// Own interface
+	UE_API UCustomizableObjectNodeSkeletalMesh* GetMeshNode() const;
+	
+protected:
+
+	UE_API UEdGraphPin* GetMeshSectionPin() const;
+
+};
+
+#undef UE_API
diff --git a/Engine/Plugins/Mutable/Source/CustomizableObjectEditor/Private/MuCOE/Nodes/CONodeMaterialModify.h b/Engine/Plugins/Mutable/Source/CustomizableObjectEditor/Private/MuCOE/Nodes/CONodeMaterialModify.h
--- a/Engine/Plugins/Mutable/Source/CustomizableObjectEditor/Private/MuCOE/Nodes/CONodeMaterialModify.h
+++ b/Engine/Plugins/Mutable/Source/CustomizableObjectEditor/Private/MuCOE/Nodes/CONodeMaterialModify.h
@@ -2,7 +2,7 @@
 
 #pragma once
 
-#include "CustomizableObjectNodeMaterial.h"
+#include "CONodeSkeletalMeshSection.h"
 #include "IDetailCustomization.h"
 #include "MuCOE/Nodes/CustomizableObjectNode.h"
 #include "MuCOE/RemapPins/CustomizableObjectNodeRemapPinsByName.h"
diff --git a/Engine/Plugins/Mutable/Source/CustomizableObjectEditor/Private/MuCOE/Nodes/CONodeModifierEditSkeletalMeshSection.cpp b/Engine/Plugins/Mutable/Source/CustomizableObjectEditor/Private/MuCOE/Nodes/CONodeModifierEditSkeletalMeshSection.cpp
new file mode 100644
--- /dev/null
+++ b/Engine/Plugins/Mutable/Source/CustomizableObjectEditor/Private/MuCOE/Nodes/CONodeModifierEditSkeletalMeshSection.cpp
@@ -0,0 +1,386 @@
+// Copyright Epic Games, Inc. All Rights Reserved.
+
+#include "MuCOE/Nodes/CONodeModifierEditSkeletalMeshSection.h"
+
+#include "MuCOE/Nodes/CONodeModifierExtendSkeletalMeshSection.h"
+#include "MuCOE/CustomizableObjectEditor_Deprecated.h"
+#include "MuCOE/CustomizableObjectLayout.h"
+#include "MuCOE/EdGraphSchema_CustomizableObject.h"
+#include "MuCOE/CustomizableObjectPin.h"
+#include "MuCO/CustomizableObjectCustomVersion.h"
+#include "MuCOE/Nodes/CustomizableObjectNodeSkeletalMesh.h"
+#include "MuCOE/Nodes/CustomizableObjectNodeTable.h"
+#include "MuCOE/Nodes/CustomizableObjectNodeLayoutBlocks.h"
+#include "MuCOE/GraphTraversal.h"
+#include "Materials/MaterialParameters.h"
+
+#include UE_INLINE_GENERATED_CPP_BY_NAME(CONodeModifierEditSkeletalMeshSection)
+
+class FCustomizableObjectNodeParentedMaterial;
+class UCustomizableObjectNodeRemapPins;
+
+#define LOCTEXT_NAMESPACE "CustomizableObjectEditor"
+
+
+void UCONodeModifierEditSkeletalMeshSection::AllocateDefaultPins(UCustomizableObjectNodeRemapPins* RemapPins)
+{
+	const UEdGraphSchema_CustomizableObject* Schema = GetDefault<UEdGraphSchema_CustomizableObject>();
+	
+	const int32 NumImages = GetNumParameters(EMaterialParameterType::Texture);
+	for (int32 ImageIndex = 0; ImageIndex < NumImages; ++ImageIndex)
+	{
+		{
+			UCustomizableObjectNodeEditMaterialPinEditImageData* PinEditImageData = NewObject<UCustomizableObjectNodeEditMaterialPinEditImageData>(this);
+			PinEditImageData->ImageParamId = GetParameterId(EMaterialParameterType::Texture, ImageIndex);
+			
+			const FName ImageName = GetParameterName(EMaterialParameterType::Texture, ImageIndex);
+			UEdGraphPin* PinImage = CustomCreatePin(EGPD_Input, Schema->PC_Texture, ImageName, FText::FromName(ImageName), PinEditImageData);
+			PinImage->bHidden = true;
+
+			PinsParameterMap.Add(PinEditImageData->ImageParamId, FEdGraphPinReference(PinImage));
+
+			FString PinMaskName = ImageName.ToString() + FString(TEXT(" Mask"));
+			UEdGraphPin* PinMask = CustomCreatePin(EGPD_Input, Schema->PC_Texture, *PinMaskName, FText::FromString(PinMaskName));
+			PinMask->bHidden = true;
+			
+			PinEditImageData->PinMask = FEdGraphPinReference(PinMask);
+		}
+	}
+
+	CustomCreatePin(EGPD_Output, Schema->PC_Modifier, TEXT("Modifier"), LOCTEXT("Modifier", "Modifier"));
+}
+
+
+const UEdGraphPin* UCONodeModifierEditSkeletalMeshSection::GetUsedImageMaskPin(const FNodeMaterialParameterId& ImageId) const
+{
+	if (const UEdGraphPin* Pin = GetUsedImagePin(ImageId))
+	{
+		UCustomizableObjectNodeEditMaterialPinEditImageData& PinData = GetPinData<UCustomizableObjectNodeEditMaterialPinEditImageData>(*Pin);
+		return PinData.PinMask.Get();
+	}
+
+	return nullptr;
+}
+
+
+bool UCONodeModifierEditSkeletalMeshSection::IsSingleOutputNode() const
+{
+	return true;
+}
+
+
+bool UCONodeModifierEditSkeletalMeshSection::CustomRemovePin(UEdGraphPin& Pin)
+{
+	for (TMap<FNodeMaterialParameterId, FEdGraphPinReference>::TIterator It = PinsParameterMap.CreateIterator(); It; ++It)
+	{
+		if (It.Value().Get() == &Pin) // We could improve performance if FEdGraphPinReference exposed the pin id.
+		{
+			It.RemoveCurrent();
+			break;
+		}
+	}
+
+	return Super::CustomRemovePin(Pin);
+}
+
+
+bool UCONodeModifierEditSkeletalMeshSection::HasPinViewer() const
+{
+	return true;
+}
+
+
+void UCONodeModifierEditSkeletalMeshSection::BackwardsCompatibleFixup(int32 CustomizableObjectCustomVersion)
+{
+	Super::BackwardsCompatibleFixup(CustomizableObjectCustomVersion);
+
+	const auto GetLayouts = [](UCONodeSkeletalMeshSectionBase& ParentMaterialNode, TArray<UCustomizableObjectLayout*>& OutLayouts)
+		{
+			UCONodeSkeletalMeshSection* NodeMaterial = ParentMaterialNode.GetMaterialNode();
+			if (!NodeMaterial)
+			{
+				return;
+			}
+
+			UEdGraphPin* MeshPin = NodeMaterial->GetMeshPin();
+			if (!MeshPin)
+			{
+				return;
+			}
+
+			const UEdGraphPin* ConnectedPin = FollowInputPin(*MeshPin);
+			if (!ConnectedPin)
+			{
+				return;
+			}
+
+			const UEdGraphPin* SourceMeshPin = FindMeshBaseSource(*ConnectedPin, false);
+			if (!SourceMeshPin)
+			{
+				return;
+			}
+
+			if (const UCustomizableObjectNodeSkeletalMesh* MeshNode = Cast<UCustomizableObjectNodeSkeletalMesh>(SourceMeshPin->GetOwningNode()))
+			{
+				UCustomizableObjectNodeSkeletalMeshPinDataMesh* MeshPinData = Cast<UCustomizableObjectNodeSkeletalMeshPinDataMesh>(MeshNode->GetPinData(*SourceMeshPin));
+				
+				if (!MeshPinData)
+				{
+					return;
+				}
+
+				// The custom version of SkeletalMesh nodes may be up to date if they're in a different CO. 
+				OutLayouts = MeshPinData->Layouts;
+
+				if (OutLayouts.IsEmpty())
+				{
+					// Pre FCustomizableObjectCustomVersion::RemoveNodeLayout code. Get Layouts from NodeLayoutBlocks
+					TArray<UEdGraphPin*> NonOrphanPins = MeshNode->GetAllNonOrphanPins();
+					for (const UEdGraphPin* Pin : NonOrphanPins)
+					{
+						if (const UCustomizableObjectNodeSkeletalMeshPinDataLayout* PinData = Cast<UCustomizableObjectNodeSkeletalMeshPinDataLayout>(MeshNode->GetPinData(*Pin)))
+						{
+							if (PinData->GetLODIndex() == MeshPinData->GetLODIndex() &&
+								PinData->GetSectionIndex() == MeshPinData->GetSectionIndex())
+							{
+								if (const UEdGraphPin* SourceLayoutConnectedPin = FollowInputPin(*Pin))
+								{
+									const UCustomizableObjectNodeLayoutBlocks* LayoutNode = Cast<UCustomizableObjectNodeLayoutBlocks>(SourceLayoutConnectedPin->GetOwningNode());
+									if (LayoutNode && LayoutNode->Layout)
+									{
+										OutLayouts.Add(LayoutNode->Layout);
+									}
+								}
+							}
+						}
+					}
+				}
+			}
+			else if (const UCustomizableObjectNodeTable* TableNode = Cast<UCustomizableObjectNodeTable>(SourceMeshPin->GetOwningNode()))
+			{
+				OutLayouts = TableNode->GetLayouts(SourceMeshPin);
+			}
+		};
+
+	// Convert deprecated node index list to the node id list.
+	if (CustomizableObjectCustomVersion == FCustomizableObjectCustomVersion::PostLoadToCustomVersion
+		&& BlockIds_DEPRECATED.Num() < Blocks_DEPRECATED.Num())
+	{
+		UCONodeSkeletalMeshSectionBase* ParentMaterialNode = GetCustomizableObjectExternalNode<UCONodeSkeletalMeshSectionBase>(ParentMaterialObject_DEPRECATED.Get(), ParentMaterialNodeId_DEPRECATED);
+
+		TArray<UCustomizableObjectLayout*> Layouts;
+		GetLayouts(*ParentMaterialNode, Layouts);
+
+		if (!Layouts.IsValidIndex(ParentLayoutIndex))
+		{
+			UE_LOG(LogMutable, Warning, TEXT("[%s] UCONodeModifierEditSkeletalMeshSection refers to an invalid texture layout index %d. Parent node has %d layouts."), 
+				*GetOutermost()->GetName(), ParentLayoutIndex, Layouts.Num());
+		}
+		else
+		{
+			UCustomizableObjectLayout* ParentLayout = Layouts[ParentLayoutIndex];
+
+			if (Cast<UCONodeSkeletalMeshSection>(ParentMaterialNode))
+			{
+				for (int32 IndexIndex = BlockIds_DEPRECATED.Num(); IndexIndex < Blocks_DEPRECATED.Num(); ++IndexIndex)
+				{
+					int32 BlockIndex = Blocks_DEPRECATED[IndexIndex];
+					if (ParentLayout->Blocks.IsValidIndex(BlockIndex) )
+					{
+						const FGuid Id = ParentLayout->Blocks[BlockIndex].Id;
+						if (Id.IsValid())
+						{
+							BlockIds_DEPRECATED.Add(Id);
+						}
+						else
+						{
+							UE_LOG(LogMutable, Warning, TEXT("[%s] UCONodeModifierEditSkeletalMeshSection refers to an valid layout block %d but that block doesn't have an id."),
+								*GetOutermost()->GetName(), BlockIndex );
+						}
+					}
+					else
+					{
+						UE_LOG(LogMutable, Warning, TEXT("[%s] UCONodeModifierEditSkeletalMeshSection refers to an invalid layout block index %d. Parent node has %d blocks."),
+							*GetOutermost()->GetName(), BlockIndex, ParentLayout->Blocks.Num());
+					}
+				}
+			}
+		}
+	}
+	
+	// Convert deprecated node id list to absolute rect list.
+	if (CustomizableObjectCustomVersion == FCustomizableObjectCustomVersion::UseUVRects)
+	{
+		// If we are here, it means this node was loaded from a version that didn't have it's own layout.
+		check(Layout->Blocks.IsEmpty());
+
+		UCONodeSkeletalMeshSectionBase* ParentMaterialNode = GetCustomizableObjectExternalNode<UCONodeSkeletalMeshSectionBase>(ParentMaterialObject_DEPRECATED.Get(), ParentMaterialNodeId_DEPRECATED);
+
+		TArray<UCustomizableObjectLayout*> ParentLayouts;
+		GetLayouts(*ParentMaterialNode, ParentLayouts);
+
+		if (!ParentLayouts.IsValidIndex(ParentLayoutIndex))
+		{
+			UE_LOG(LogMutable, Warning, TEXT("[%s] UCONodeModifierEditSkeletalMeshSection refers to an invalid texture layout index %d. Parent node has %d layouts."),
+				*GetOutermost()->GetName(), ParentLayoutIndex, ParentLayouts.Num());
+		}
+		else
+		{
+			UCustomizableObjectLayout* ParentLayout = ParentLayouts[ParentLayoutIndex];
+			FIntPoint GridSize = ParentLayout->GetGridSize();
+
+			Layout->SetGridSize(GridSize);
+
+			if (Cast<UCONodeSkeletalMeshSection>(ParentMaterialNode))
+			{
+				for (const FGuid& BlockId : BlockIds_DEPRECATED)
+				{
+					bool bSkipBlock = false;
+					for (const FCustomizableObjectLayoutBlock& ExistingBlock : Layout->Blocks)
+					{
+						if (ExistingBlock.Id == BlockId)
+						{
+							bSkipBlock = true;
+							UE_LOG(LogMutable, Log, TEXT("[%s] UCONodeModifierEditSkeletalMeshSection has a duplicated layout block id. One has been ignored during version upgrade."), *GetOutermost()->GetName());
+							break;
+						}
+					}
+
+					if (bSkipBlock)
+					{
+						continue;
+					}
+
+					bool bFoundInParent = false;
+					for (const FCustomizableObjectLayoutBlock& ParentBlock : ParentLayout->Blocks)
+					{
+						if (ParentBlock.Id == BlockId)
+						{
+							bFoundInParent = true;
+
+							FCustomizableObjectLayoutBlock NewBlock;
+							NewBlock = ParentBlock;
+
+							// Clear some unnecessary data.
+							NewBlock.bReduceBothAxes = false;
+							NewBlock.bReduceByTwo = false;
+							NewBlock.Priority = 0;
+
+							Layout->Blocks.Add(NewBlock);
+							break;
+						}
+					}
+
+					if (!bFoundInParent)
+					{
+						UE_LOG(LogMutable, Warning, TEXT("[%s] UCONodeModifierEditSkeletalMeshSection refers to and invalid layout block. It has been ignored during version upgrade."), *GetOutermost()->GetName());
+					}
+				}
+			}
+		}
+	}
+
+	if (CustomizableObjectCustomVersion == FCustomizableObjectCustomVersion::AutomaticNodeMaterial)
+	{
+		UCONodeSkeletalMeshSectionBase* ParentMaterial = GetCustomizableObjectExternalNode<UCONodeSkeletalMeshSectionBase>(ParentMaterialObject_DEPRECATED.Get(), ParentMaterialNodeId_DEPRECATED);
+
+		if (ParentMaterial)
+		{
+			for (const FCustomizableObjectNodeEditMaterialImage& Image : Images_DEPRECATED)
+			{
+				const UEdGraphPin* ImagePin = FindPin(Image.Name);
+				const UEdGraphPin* PinMask = FindPin(Image.Name + " Mask");
+				if (!ImagePin || !PinMask)
+				{
+					continue;
+				}
+				
+				UCustomizableObjectNodeEditMaterialPinEditImageData*  PinEditImageData = NewObject<UCustomizableObjectNodeEditMaterialPinEditImageData>(this);
+				PinEditImageData->ImageId_DEPRECATED = FGuid::NewGuid();
+				PinEditImageData->PinMask = PinMask;
+				
+				// Search for the Image Id the Edit pin was referring to.
+				const int32 NumImages = ParentMaterial->GetNumParameters(EMaterialParameterType::Texture);
+				for (int32 ImageIndex = 0; ImageIndex < NumImages; ++ImageIndex)
+				{
+					if (ParentMaterial->GetParameterName(EMaterialParameterType::Texture, ImageIndex).ToString() == Image.Name)
+					{
+						PinEditImageData->ImageId_DEPRECATED = ParentMaterial->GetParameterId(EMaterialParameterType::Texture, ImageIndex).ParameterId;
+						break;
+					}
+				}
+				
+				AddPinData(*ImagePin, *PinEditImageData);
+			}
+		}
+		
+		Images_DEPRECATED.Empty();
+	}
+
+	// Fill PinsParameter.
+	if (CustomizableObjectCustomVersion == FCustomizableObjectCustomVersion::AutomaticNodeMaterialPerformanceBug) // || CustomizableObjectCustomVersion < FCustomizableObjectCustomVersion::AutomaticNodeMaterialPerformance
+	{
+		for (const UEdGraphPin* Pin : GetAllNonOrphanPins())
+		{
+			if (const UCustomizableObjectNodeEditMaterialPinEditImageData* PinData = Cast<UCustomizableObjectNodeEditMaterialPinEditImageData>(GetPinData(*Pin)))
+			{
+				PinsParameter_DEPRECATED.Add(PinData->ImageId_DEPRECATED, FEdGraphPinReference(Pin));
+			}
+		}
+	}
+
+	if (CustomizableObjectCustomVersion == FCustomizableObjectCustomVersion::ConvertEditAndExtendToModifiers)
+	{
+		// Look for the parent material and set it as the modifier reference material 
+
+		PinsParameterMap = PinsParameterMap_DEPRECATED;
+		PinsParameterMap_DEPRECATED.Empty();
+
+		UCustomizableObjectNode* ParentNode = GetCustomizableObjectExternalNode<UCustomizableObjectNode>(ParentMaterialObject_DEPRECATED.Get(), ParentMaterialNodeId_DEPRECATED);
+
+		if (UCONodeSkeletalMeshSection* MaterialParentNode = Cast<UCONodeSkeletalMeshSection>(ParentNode))
+		{
+			ReferenceMaterial = MaterialParentNode->GetMaterial();
+		}
+		else if (UCONodeModifierExtendSkeletalMeshSection* ExtendParentNode = Cast<UCONodeModifierExtendSkeletalMeshSection>(ParentNode))
+		{
+			ReferenceMaterial = ExtendParentNode->ReferenceMaterial;
+		}
+		else
+		{
+			// Conversion failed?
+			ensure(false);
+			UE_LOG(LogMutable, Warning, TEXT("[%s] UCONodeModifierExtendSkeletalMeshSection version upgrade failed."), *GetOutermost()->GetName());
+		}
+
+		ReconstructNode();
+	}
+}
+
+
+FText UCONodeModifierEditSkeletalMeshSection::GetNodeTitle(ENodeTitleType::Type TitleType) const
+{
+	return LOCTEXT("Edit_MeshSection", "Edit Mesh Section");
+}
+
+
+FString UCONodeModifierEditSkeletalMeshSection::GetRefreshMessage() const
+{
+	return TEXT("Source material has changed, texture channels might have been added, removed or renamed. Please refresh the parent material node to reflect those changes.");
+}
+
+
+FText UCONodeModifierEditSkeletalMeshSection::GetTooltipText() const
+{
+	return LOCTEXT("Edit_Material_Tooltip", "Modify the texture parameters of an ancestor's material partially or completely.");
+}
+
+
+void UCONodeModifierEditSkeletalMeshSection::SetLayoutIndex(const int32 LayoutIndex)
+{
+	ParentLayoutIndex = LayoutIndex;
+}
+
+
+#undef LOCTEXT_NAMESPACE
+
diff --git a/Engine/Plugins/Mutable/Source/CustomizableObjectEditor/Private/MuCOE/Nodes/CONodeModifierEditSkeletalMeshSection.h b/Engine/Plugins/Mutable/Source/CustomizableObjectEditor/Private/MuCOE/Nodes/CONodeModifierEditSkeletalMeshSection.h
new file mode 100644
--- /dev/null
+++ b/Engine/Plugins/Mutable/Source/CustomizableObjectEditor/Private/MuCOE/Nodes/CONodeModifierEditSkeletalMeshSection.h
@@ -0,0 +1,90 @@
+// Copyright Epic Games, Inc. All Rights Reserved.
+
+#pragma once
+
+#include "MuCOE/Nodes/CustomizableObjectNodeModifierEditLayoutBlocks.h"
+#include "MuCOE/Nodes/CONodeSkeletalMeshSection.h"
+
+#include "CONodeModifierEditSkeletalMeshSection.generated.h"
+
+#define UE_API CUSTOMIZABLEOBJECTEDITOR_API
+
+namespace ENodeTitleType { enum Type : int; }
+struct FCustomizableObjectNodeEditMaterialImage;
+struct FNodeMaterialParameterId;
+
+class FArchive;
+class FCustomizableObjectNodeParentedMaterial;
+class UCustomizableObject;
+class UCustomizableObjectNodeRemapPins;
+class UObject;
+
+
+/** Additional data for the Parent Texture Parameter to edit pin. */
+UCLASS(MinimalAPI)
+class UCustomizableObjectNodeEditMaterialPinEditImageData : public UCustomizableObjectNodePinData
+{
+	GENERATED_BODY()
+
+public:
+
+	/** Parameter FGuid + Layer Index */
+	UPROPERTY()
+	FNodeMaterialParameterId ImageParamId;
+
+	UPROPERTY();
+	FEdGraphPinReference PinMask;
+
+	UPROPERTY()
+	FGuid ImageId_DEPRECATED;
+};
+
+
+UCLASS(MinimalAPI)
+class UCONodeModifierEditSkeletalMeshSection : public UCustomizableObjectNodeModifierEditLayoutBlocks
+{
+
+	GENERATED_BODY()
+
+public:
+	
+	// EdGraphNode interface
+	UE_API virtual FText GetNodeTitle(ENodeTitleType::Type TitleType) const override;
+	UE_API virtual FText GetTooltipText() const override;
+
+	// UCustomizableObjectNode interface
+	UE_API virtual void BackwardsCompatibleFixup(int32 CustomizableObjectCustomVersion) override;
+	UE_API virtual void AllocateDefaultPins(UCustomizableObjectNodeRemapPins* RemapPins) override;
+	UE_API virtual FString GetRefreshMessage() const override;
+	UE_API virtual bool IsSingleOutputNode() const override;
+	UE_API virtual bool CustomRemovePin(UEdGraphPin& Pin) override;
+	UE_API virtual bool HasPinViewer() const override;
+	
+	// Own interface
+	
+	/** */
+	UE_API virtual void SetLayoutIndex(const int32 LayoutIndex);
+
+	/** Returns the Image mask pin of the given Image that will be edited.
+	 *
+	 * @returns Always returns a valid pin if EditsImage(const FGuid&) returns true. */
+	UE_API const UEdGraphPin* GetUsedImageMaskPin(const FNodeMaterialParameterId& ImageId) const;
+
+private:
+
+	// Deprecated properties
+	UPROPERTY()
+	TMap<FNodeMaterialParameterId, FEdGraphPinReference> PinsParameterMap_DEPRECATED;
+
+	UPROPERTY()
+	TMap<FGuid, FEdGraphPinReference> PinsParameter_DEPRECATED;
+	
+	UPROPERTY()
+	TArray<FCustomizableObjectNodeEditMaterialImage> Images_DEPRECATED;
+	
+	// Old layout blocks to patch. Now in parent class.
+	UPROPERTY()
+	TArray<int32> Blocks_DEPRECATED;
+};
+
+#undef UE_API
diff --git a/Engine/Plugins/Mutable/Source/CustomizableObjectEditor/Private/MuCOE/Nodes/CONodeModifierExtendSkeletalMeshSection.cpp b/Engine/Plugins/Mutable/Source/CustomizableObjectEditor/Private/MuCOE/Nodes/CONodeModifierExtendSkeletalMeshSection.cpp
new file mode 100644
--- /dev/null
+++ b/Engine/Plugins/Mutable/Source/CustomizableObjectEditor/Private/MuCOE/Nodes/CONodeModifierExtendSkeletalMeshSection.cpp
@@ -0,0 +1,290 @@
+// Copyright Epic Games, Inc. All Rights Reserved.
+
+#include "MuCOE/Nodes/CONodeModifierExtendSkeletalMeshSection.h"
+
+#include "Materials/MaterialParameters.h"
+#include "MuCO/CustomizableObjectCustomVersion.h"
+#include "MuCOE/CustomizableObjectEditor_Deprecated.h"
+#include "MuCOE/CustomizableObjectGraph.h"
+#include "MuCOE/EdGraphSchema_CustomizableObject.h"
+#include "MuCOE/GraphTraversal.h"
+#include "MuCOE/Nodes/CONodeCopySkeletalMeshSection.h"
+#include "MuCOE/Nodes/CustomizableObjectNodeSkeletalMesh.h"
+#include "MuCOE/Nodes/CustomizableObjectNodeStaticString.h"
+#include "MuCOE/Nodes/CustomizableObjectNodeTable.h"
+#include "MuCOE/CustomizableObjectPin.h"
+#include "MuCOE/CustomizableObjectMacroLibrary/CustomizableObjectMacroLibrary.h"
+
+#include UE_INLINE_GENERATED_CPP_BY_NAME(CONodeModifierExtendSkeletalMeshSection)
+
+class UCustomizableObjectNode;
+class UCustomizableObjectNodeRemapPins;
+
+#define LOCTEXT_NAMESPACE "CustomizableObjectEditor"
+
+
+void UCONodeModifierExtendSkeletalMeshSection::BackwardsCompatibleFixup(int32 CustomizableObjectCustomVersion)
+{
+	Super::BackwardsCompatibleFixup(CustomizableObjectCustomVersion);
+	
+	if (CustomizableObjectCustomVersion == FCustomizableObjectCustomVersion::ExtendMaterialRemoveImages)
+	{
+		const UCONodeSkeletalMeshSectionBase* ParentMaterial = GetCustomizableObjectExternalNode<UCONodeSkeletalMeshSectionBase>(ParentMaterialObject_DEPRECATED.Get(), ParentMaterialNodeId_DEPRECATED);
+		if (ParentMaterial)
+		{
+			for (const FCustomizableObjectNodeExtendMaterialImage& Image : Images_DEPRECATED)
+			{
+				const UEdGraphPin* ImagePin = FindPin(Image.Name);
+				if (!ImagePin)
+				{
+					continue;
+				}
+
+				FGuid ImageId = FGuid::NewGuid();
+
+				// Search for the Image Id the Extend pin was referring to.
+				const int32 NumImages = ParentMaterial->GetNumParameters(EMaterialParameterType::Texture);
+				for (int32 ImageIndex = 0; ImageIndex < NumImages; ++ImageIndex)
+				{
+					if (ParentMaterial->GetParameterName(EMaterialParameterType::Texture, ImageIndex).ToString() == Image.Name)
+					{
+						ImageId = ParentMaterial->GetParameterId(EMaterialParameterType::Texture, ImageIndex).ParameterId;
+						break;
+					}
+				}
+
+				PinsParameter_DEPRECATED.Add(ImageId, FEdGraphPinReference(ImagePin));                
+			}
+		}
+
+		Images_DEPRECATED.Empty();
+	}
+
+	if (CustomizableObjectCustomVersion == FCustomizableObjectCustomVersion::ConvertEditAndExtendToModifiers)
+	{
+		// Look for the parent material and set it as the modifier reference material 
+		PinsParameterMap = PinsParameterMap_DEPRECATED;
+		PinsParameterMap_DEPRECATED.Empty();
+
+		UCustomizableObjectNode* ParentNode = GetCustomizableObjectExternalNode<UCustomizableObjectNode>(ParentMaterialObject_DEPRECATED.Get(), ParentMaterialNodeId_DEPRECATED);
+		if (ParentNode)
+		{
+			// Add an autogenerated tag to the legacy parent, so that it can be referenced in this modifier.
+			FString NewLegacyTag = MakeNodeAutoTag(ParentNode);
+			RequiredTags.Add(NewLegacyTag);
+
+			if (UCONodeSkeletalMeshSection* MaterialParentNode = Cast<UCONodeSkeletalMeshSection>(ParentNode))
+			{
+				ReferenceMaterial = MaterialParentNode->GetMaterial();
+				MaterialParentNode->Tags.AddUnique(NewLegacyTag);
+			}
+			else if (UCONodeModifierExtendSkeletalMeshSection* ExtendParentNode = Cast<UCONodeModifierExtendSkeletalMeshSection>(ParentNode))
+			{
+				ReferenceMaterial = ExtendParentNode->ReferenceMaterial;
+				ExtendParentNode->Tags.AddUnique(NewLegacyTag);
+			}
+			else
+			{
+				// Conversion failed?
+				ensure(false);
+				UE_LOG(LogMutable, Warning, TEXT("[%s] UCONodeModifierExtendSkeletalMeshSection version upgrade failed."), *GetOutermost()->GetName());
+			}
+
+			// If the tag was added to another CO, keep track of the addition to be able to repeat it in case
+			// that CO is not re-saved but this one is.
+			UCustomizableObject* ThisNodeObject = GraphTraversal::GetObject(*this);
+			if (ParentMaterialObject_DEPRECATED != ThisNodeObject)
+			{
+				FLegacyTag LegacyTag;
+				LegacyTag.ParentObject = ParentMaterialObject_DEPRECATED;
+				LegacyTag.ParentNode = ParentMaterialNodeId_DEPRECATED;
+				LegacyTag.Tag = NewLegacyTag;
+				LegacyBackportsRequiredTags.AddUnique(LegacyTag);
+			}
+		}
+		else
+		{
+			// No parent, no conversion.
+			UE_LOG(LogMutable, Log, TEXT("[%s] UCONodeModifierExtendSkeletalMeshSection has no parent. It will not be upgraded."), *GetOutermost()->GetName());
+		}
+
+		ReconstructNode();
+	}
+
+	if (CustomizableObjectCustomVersion == FCustomizableObjectCustomVersion::EnableMutableMacrosNewVersion)
+	{
+		if (!EnableTagsPinRef.Get())
+		{
+			UEdGraphPin* EnableTagsPin = CustomCreatePin(EGPD_Input, UEdGraphSchema_CustomizableObject::PC_String, TEXT("Enable Tags"), LOCTEXT("EnableTags", "Enable Tags"));
+			EnableTagsPin->PinType.ContainerType = EPinContainerType::Array;
+			EnableTagsPin->PinToolTip = "List of Tags that this node will Enable";
+			EnableTagsPinRef = EnableTagsPin;
+		}
+	}
+}
+
+
+void UCONodeModifierExtendSkeletalMeshSection::AllocateDefaultPins(UCustomizableObjectNodeRemapPins* RemapPins)
+{
+	const UEdGraphSchema_CustomizableObject* Schema = GetDefault<UEdGraphSchema_CustomizableObject>();
+
+	CustomCreatePin(EGPD_Input, Schema->PC_Mesh, TEXT("Add Mesh"), LOCTEXT("AddMesh", "Add Mesh"));
+
+	//Create Node Modifier Common Pins
+	Super::AllocateDefaultPins(RemapPins);
+
+	// Create enable tags pin
+	{
+		UEdGraphPin* EnableTagsPin = CustomCreatePin(EGPD_Input, Schema->PC_String, TEXT("Enable Tags"), LOCTEXT("EnableTags", "Enable Tags"));
+		EnableTagsPin->PinType.ContainerType = EPinContainerType::Array;
+		EnableTagsPin->PinToolTip = "List of Tags that this node will Enable";
+		EnableTagsPinRef = EnableTagsPin;
+	}
+
+	// Begin texture pins
+	{
+		const int32 NumImages = GetNumParameters(EMaterialParameterType::Texture);
+		for (int32 ImageIndex = 0; ImageIndex < NumImages; ++ImageIndex)
+		{
+			const FName ImageName = GetParameterName(EMaterialParameterType::Texture, ImageIndex);
+			UEdGraphPin* PinImage = CustomCreatePin(EGPD_Input, Schema->PC_Texture, ImageName, FText::FromName(ImageName));
+
+			const FNodeMaterialParameterId ImageId = GetParameterId(EMaterialParameterType::Texture, ImageIndex);
+				
+			PinsParameterMap.Add(ImageId, FEdGraphPinReference(PinImage));
+		}
+	}
+	// End texture pins
+
+	CustomCreatePin(EGPD_Output, Schema->PC_Modifier, TEXT("Modifier"), LOCTEXT("Modifier", "Modifier"));
+}
+
+
+FText UCONodeModifierExtendSkeletalMeshSection::GetNodeTitle(ENodeTitleType::Type TitleType) const
+{
+	return LOCTEXT("Extend_MeshSection", "Extend Mesh Section");
+}
+
+
+FString UCONodeModifierExtendSkeletalMeshSection::GetRefreshMessage() const
+{
+	return TEXT("The Reference Material has changed, texture channels might have been added, removed or renamed. Please refresh to reflect those changes.");
+}
+
+
+UEdGraphPin* UCONodeModifierExtendSkeletalMeshSection::AddMeshPin() const
+{
+	return FindPin(TEXT("Add Mesh"));
+}
+
+
+UEdGraphPin* UCONodeModifierExtendSkeletalMeshSection::GetEnableTagsPin() const
+{
+	return EnableTagsPinRef.Get();
+}
+
+
+TArray<UCustomizableObjectLayout*> UCONodeModifierExtendSkeletalMeshSection::GetLayouts() const
+{
+	TArray<UCustomizableObjectLayout*> Result;
+
+	if (UEdGraphPin* MeshPin = AddMeshPin())
+	{
+		if (const UEdGraphPin* ConnectedPin = FollowInputPin(*MeshPin))
+		{
+			if (const UEdGraphPin* SourceMeshPin = FindMeshBaseSource(*ConnectedPin, false))
+			{
+				if (const UCustomizableObjectNodeSkeletalMesh* MeshNode = Cast<UCustomizableObjectNodeSkeletalMesh>(SourceMeshPin->GetOwningNode()))
+				{
+					Result = MeshNode->GetLayouts(*SourceMeshPin);
+				}
+				else if (const UCustomizableObjectNodeTable* TableNode = Cast<UCustomizableObjectNodeTable>(SourceMeshPin->GetOwningNode()))
+				{
+					Result = TableNode->GetLayouts(SourceMeshPin);
+				}
+			}
+		}
+	}
+
+	return Result;
+}
+
+
+FText UCONodeModifierExtendSkeletalMeshSection::GetTooltipText() const
+{
+	return LOCTEXT("Extend_Material_Tooltip", "Extend a mesh section: add a new mesh fragment, and add its corresponding texture to the modified material texture parameters.");
+}
+
+
+bool UCONodeModifierExtendSkeletalMeshSection::IsSingleOutputNode() const
+{
+	return true;
+}
+
+
+bool UCONodeModifierExtendSkeletalMeshSection::CustomRemovePin(UEdGraphPin& Pin)
+{
+	for (TMap<FNodeMaterialParameterId, FEdGraphPinReference>::TIterator It = PinsParameterMap.CreateIterator(); It; ++It)
+	{
+		if (It.Value().Get() == &Pin) // We could improve performance if FEdGraphPinReference exposed the pin id.
+		{
+			It.RemoveCurrent();
+			break;
+		}
+	}
+
+	return Super::CustomRemovePin(Pin);
+}
+
+
+TArray<FString> UCONodeModifierExtendSkeletalMeshSection::GetEnableTags(TArray<const UCustomizableObjectNodeMacroInstance*>* MacroContext)
+{ 
+	const UEdGraphPin* TagsPin = GetEnableTagsPin();
+
+	if (!TagsPin)
+	{
+		return Tags;
+	}
+
+	const TArray<UEdGraphPin*> ConnectedPins = FollowInputPinArray(*TagsPin);
+
+	// Getting tags from linked pin
+	if (ConnectedPins.Num())
+	{
+		TArray<FString> OutTags;
+
+		for (const UEdGraphPin* StringPin : ConnectedPins)
+		{
+			const UEdGraphPin* SourceStringPin = GraphTraversal::FindIOPinSourceThroughMacroContext(*StringPin, MacroContext);
+
+			if (SourceStringPin)
+			{
+				if (const UCustomizableObjectNodeStaticString* StringNode = Cast<UCustomizableObjectNodeStaticString>(SourceStringPin->GetOwningNode()))
+				{
+					OutTags.AddUnique(StringNode->Value);
+				}
+			}
+			
+		}
+
+		return OutTags;
+	}
+	
+	return Tags;
+}
+
+
+TArray<FString>* UCONodeModifierExtendSkeletalMeshSection::GetEnableTagsArray()
+{
+	return &Tags;
+}
+
+
+FString UCONodeModifierExtendSkeletalMeshSection::GetInternalTagDisplayName()
+{
+	UMaterialInterface* Material = ReferenceMaterial;
+	return FString::Printf(TEXT("Extend Mesh [%s]"), Material ? *Material->GetName() : TEXT("no-material"));
+}
+
+
+#undef LOCTEXT_NAMESPACE
diff --git a/Engine/Plugins/Mutable/Source/CustomizableObjectEditor/Private/MuCOE/Nodes/CONodeModifierExtendSkeletalMeshSection.h b/Engine/Plugins/Mutable/Source/CustomizableObjectEditor/Private/MuCOE/Nodes/CONodeModifierExtendSkeletalMeshSection.h
new file mode 100644
--- /dev/null
+++ b/Engine/Plugins/Mutable/Source/CustomizableObjectEditor/Private/MuCOE/Nodes/CONodeModifierExtendSkeletalMeshSection.h
@@ -0,0 +1,85 @@
+// Copyright Epic Games, Inc. All Rights Reserved.
+
+#pragma once
+
+#include "MuCOE/Nodes/CONodeModifierSkeletalMeshSection.h"
+#include "MuCOE/CustomizableObjectLayout.h"
+
+#include "CONodeModifierExtendSkeletalMeshSection.generated.h"
+
+#define UE_API CUSTOMIZABLEOBJECTEDITOR_API
+
+namespace ENodeTitleType { enum Type : int; }
+
+class FArchive;
+class UCustomizableObject;
+class UCustomizableObjectNode;
+class UCustomizableObjectNodeRemapPins;
+class UEdGraphPin;
+class UObject;
+struct FCustomizableObjectNodeExtendMaterialImage;
+struct FEdGraphPinReference;
+
+
+UCLASS(MinimalAPI)
+class UCONodeModifierExtendSkeletalMeshSection : public UCONodeModifierSkeletalMeshSection
+{
+	GENERATED_BODY()
+
+public:
+
+	/** Tags enabled when the modifier is applied. Also used to decide what modifiers get applied to the data added by this node.  */
+	UPROPERTY(EditAnywhere, BlueprintReadWrite, Category = EnableTags)
+	TArray<FString> Tags;
+
+	/** First LOD where the modifier is applied to. */
+	UPROPERTY(EditAnywhere, Category = LOD, meta = (ClampMin = 0, ClampMax = 255))
+	uint32 FirstLOD = 0;
+
+public:
+	
+	// Begin EdGraphNode interface
+	UE_API virtual FText GetNodeTitle(ENodeTitleType::Type TitleType) const override;
+	UE_API virtual FText GetTooltipText() const override;
+
+	// UCustomizableObjectNode interface
+	UE_API virtual void BackwardsCompatibleFixup(int32 CustomizableObjectCustomVersion) override;
+	UE_API virtual void AllocateDefaultPins(UCustomizableObjectNodeRemapPins* RemapPins) override;
+	UE_API virtual FString GetRefreshMessage() const override;
+	UE_API virtual bool IsSingleOutputNode() const override;
+	UE_API virtual bool CustomRemovePin(UEdGraphPin& Pin) override;
+	UE_API virtual TArray<FString> GetEnableTags(TArray<const UCustomizableObjectNodeMacroInstance*>* MacroContext = nullptr) override;
+	UE_API virtual TArray<FString>* GetEnableTagsArray() override;
+	UE_API virtual FString GetInternalTagDisplayName() override;
+
+	// Own interface
+	UE_API UEdGraphPin* AddMeshPin() const;
+	UE_API UEdGraphPin* GetEnableTagsPin() const;
+
+	UE_API TArray<UCustomizableObjectLayout*> GetLayouts() const;
+
+private:
+
+	// Deprecated properties
+
+	UPROPERTY()
+	TObjectPtr<UCustomizableObject> ParentMaterialObject_DEPRECATED = nullptr;
+
+	UPROPERTY()
+	FGuid ParentMaterialNodeId_DEPRECATED;
+	
+
+	UPROPERTY()
+	TMap<FNodeMaterialParameterId, FEdGraphPinReference> PinsParameterMap_DEPRECATED;
+
+	UPROPERTY()
+	TMap<FGuid, FEdGraphPinReference> PinsParameter_DEPRECATED;
+
+	UPROPERTY()
+	TArray<FCustomizableObjectNodeExtendMaterialImage> Images_DEPRECATED;
+
+	UPROPERTY()
+	FEdGraphPinReference EnableTagsPinRef;
+};
+
+#undef UE_API
diff --git a/Engine/Plugins/Mutable/Source/CustomizableObjectEditor/Private/MuCOE/Nodes/CONodeModifierSkeletalMeshSection.cpp b/Engine/Plugins/Mutable/Source/CustomizableObjectEditor/Private/MuCOE/Nodes/CONodeModifierSkeletalMeshSection.cpp
new file mode 100644
--- /dev/null
+++ b/Engine/Plugins/Mutable/Source/CustomizableObjectEditor/Private/MuCOE/Nodes/CONodeModifierSkeletalMeshSection.cpp
@@ -0,0 +1,192 @@
+// Copyright Epic Games, Inc. All Rights Reserved.
+
+#include "MuCOE/Nodes/CONodeModifierSkeletalMeshSection.h"
+
+#include "MuCOE/Nodes/CONodeSkeletalMeshSection.h"
+#include "MuCOE/GraphTraversal.h"
+#include "MuCOE/CustomizableObjectPin.h"
+#include "MuCO/CustomizableObjectCustomVersion.h"
+#include "Materials/MaterialParameters.h"
+#include "MaterialCachedData.h"
+#include "Materials/MaterialExpressionTextureCoordinate.h"
+#include "Materials/MaterialExpressionTextureSample.h"
+#include "Materials/MaterialInstance.h"
+
+#include UE_INLINE_GENERATED_CPP_BY_NAME(CONodeModifierSkeletalMeshSection)
+
+
+#define LOCTEXT_NAMESPACE "CustomizableObjectEditor"
+
+
+
+void UCONodeModifierSkeletalMeshSection::PostEditChangeProperty(FPropertyChangedEvent& PropertyChangedEvent)
+{
+	Super::PostEditChangeProperty(PropertyChangedEvent);
+
+	if (const FProperty* PropertyThatChanged = PropertyChangedEvent.Property)
+	{
+		if (PropertyThatChanged->GetFName() == GET_MEMBER_NAME_CHECKED(UCONodeModifierSkeletalMeshSection, ReferenceMaterial))
+		{
+			UCustomizableObjectNode::ReconstructNode();
+		}
+	}
+}
+
+
+bool UCONodeModifierSkeletalMeshSection::UsesImage(const FNodeMaterialParameterId& ImageId) const
+{
+	if (const UEdGraphPin* Pin = GetUsedImagePin(ImageId))
+	{
+		return FollowInputPin(*Pin) != nullptr;
+	}
+	else
+	{
+		return false;
+	}
+}
+
+
+const UEdGraphPin* UCONodeModifierSkeletalMeshSection::GetUsedImagePin(const FNodeMaterialParameterId& ImageId) const
+{
+	if (const FEdGraphPinReference* PinReference = PinsParameterMap.Find(ImageId))
+	{
+		if (const UEdGraphPin& Pin = *PinReference->Get(); !IsPinOrphan(Pin)) // We always have a valid pin reference. If it is nullptr, it means that something went wrong.
+		{
+			return &Pin;
+		}
+	}
+
+	return nullptr;
+}
+
+
+bool UCONodeModifierSkeletalMeshSection::IsNodeOutDatedAndNeedsRefresh()
+{
+	const bool bOutdated = [&, this]()
+		{
+			for (const TTuple<FNodeMaterialParameterId, FEdGraphPinReference> Pair : PinsParameterMap)
+			{
+				const UEdGraphPin& Pin = *Pair.Value.Get();
+				if (!IsPinOrphan(Pin))
+				{
+					bool bHas = UCONodeSkeletalMeshSection::HasParameter(ReferenceMaterial,Pair.Key);
+					if (UsesImage(Pair.Key) && !bHas)
+					{
+						return true;
+					}
+				}
+			}
+
+			return false;
+		}();
+
+	// Remove previous compilation warnings
+	if (!bOutdated && bHasCompilerMessage)
+	{
+		RemoveWarnings();
+		GetGraph()->NotifyGraphChanged();
+	}
+
+	return bOutdated;
+}
+
+
+int32 UCONodeModifierSkeletalMeshSection::GetNumParameters(const EMaterialParameterType Type) const
+{
+	if (ReferenceMaterial)
+	{
+		return ReferenceMaterial->GetCachedExpressionData().GetParameterTypeEntry(Type).ParameterInfoSet.Num();
+	}
+	else
+	{
+		return 0;
+	}
+}
+
+
+FNodeMaterialParameterId UCONodeModifierSkeletalMeshSection::GetParameterId(const EMaterialParameterType Type, const int32 ParameterIndex) const
+{
+	const FMaterialCachedExpressionData& Data = ReferenceMaterial->GetCachedExpressionData();
+
+	if (Data.EditorOnlyData)
+	{
+		if (Data.EditorOnlyData->EditorEntries[(int32)Type].EditorInfo.Num() != 0)
+		{
+			const FGuid ParameterId = Data.EditorOnlyData->EditorEntries[(int32)Type].EditorInfo[ParameterIndex].ExpressionGuid;
+			const int32 LayerIndex = GetParameterLayerIndex(Type, ParameterIndex);
+
+			return { ParameterId, LayerIndex };
+		}
+	}
+
+	return FNodeMaterialParameterId();
+}
+
+
+FName UCONodeModifierSkeletalMeshSection::GetParameterName(const EMaterialParameterType Type, const int32 ParameterIndex) const
+{
+	check(ReferenceMaterial);
+
+	const FMaterialCachedParameterEntry& Entry = ReferenceMaterial->GetCachedExpressionData().GetParameterTypeEntry(Type);
+
+	for (TSet<FMaterialParameterInfo>::TConstIterator It(Entry.ParameterInfoSet); It; ++It)
+	{
+		const int32 IteratorIndex = It.GetId().AsInteger();
+
+		if (IteratorIndex == ParameterIndex)
+		{
+			return (*It).Name;
+		}
+	}
+
+	// The parameter should exist
+	check(false);
+
+	return FName();
+}
+
+
+int32 UCONodeModifierSkeletalMeshSection::GetParameterLayerIndex(const UMaterialInterface* InMaterial, const EMaterialParameterType Type, const int32 ParameterIndex)
+{
+	check(InMaterial);
+
+	const FMaterialCachedParameterEntry& Entry = InMaterial->GetCachedExpressionData().GetParameterTypeEntry(Type);
+
+	for (TSet<FMaterialParameterInfo>::TConstIterator It(Entry.ParameterInfoSet); It; ++It)
+	{
+		const int32 IteratorIndex = It.GetId().AsInteger();
+
+		if (IteratorIndex == ParameterIndex)
+		{
+			return (*It).Index;
+		}
+	}
+
+	// The parameter should exist
+	check(false);
+
+	return -1;
+}
+
+
+int32 UCONodeModifierSkeletalMeshSection::GetParameterLayerIndex(const EMaterialParameterType Type, const int32 ParameterIndex) const
+{
+	return GetParameterLayerIndex(ReferenceMaterial.Get(), Type, ParameterIndex);
+}
+
+
+FText UCONodeModifierSkeletalMeshSection::GetParameterLayerName(const EMaterialParameterType Type, const int32 ParameterIndex) const
+{
+	check(ReferenceMaterial)
+
+		int32 LayerIndex = GetParameterLayerIndex(Type, ParameterIndex);
+
+	FMaterialLayersFunctions LayersValue;
+	ReferenceMaterial->GetMaterialLayers(LayersValue);
+
+
+	return LayersValue.EditorOnly.LayerNames.IsValidIndex(LayerIndex) ? LayersValue.EditorOnly.LayerNames[LayerIndex] : FText();
+}
+
+
+#undef LOCTEXT_NAMESPACE
diff --git a/Engine/Plugins/Mutable/Source/CustomizableObjectEditor/Private/MuCOE/Nodes/CONodeModifierSkeletalMeshSection.h b/Engine/Plugins/Mutable/Source/CustomizableObjectEditor/Private/MuCOE/Nodes/CONodeModifierSkeletalMeshSection.h
new file mode 100644
--- /dev/null
+++ b/Engine/Plugins/Mutable/Source/CustomizableObjectEditor/Private/MuCOE/Nodes/CONodeModifierSkeletalMeshSection.h
@@ -0,0 +1,59 @@
+// Copyright Epic Games, Inc. All Rights Reserved.
+
+#pragma once
+
+#include "MuCOE/Nodes/CustomizableObjectNodeModifierBase.h"
+#include "MuCOE/Nodes/CONodeSkeletalMeshSectionBase.h"
+
+#include "CONodeModifierSkeletalMeshSection.generated.h"
+
+#define UE_API CUSTOMIZABLEOBJECTEDITOR_API
+
+class UObject;
+
+
+UCLASS(MinimalAPI, abstract)
+class UCONodeModifierSkeletalMeshSection : public UCustomizableObjectNodeModifierBase
+{
+
+	GENERATED_BODY()
+
+public:
+
+	/** Reference material that defines the structure of the data to be extended.
+	* The Sections modified with this modifier are supposed to have the same texture parameters, but don't need to have exactly
+	* the ReferenceMaterial set.
+	*/
+	UPROPERTY(EditAnywhere, Category = Parameters)
+	TObjectPtr<UMaterialInterface> ReferenceMaterial = nullptr;
+
+	/** Relates a Parameter id (and layer if is a layered material) to a Pin. Only used to improve performance. */
+	UPROPERTY()
+	TMap<FNodeMaterialParameterId, FEdGraphPinReference> PinsParameterMap;
+
+public:
+
+	//
+	UE_API virtual void PostEditChangeProperty(FPropertyChangedEvent& PropertyChangedEvent) override;
+
+	// UCustomizableObjectNode interface
+	UE_API virtual bool IsNodeOutDatedAndNeedsRefresh() override;
+
+	// Own interface
+
+	/** */
+	UE_API bool UsesImage(const FNodeMaterialParameterId& ImageId) const;
+
+	/** */
+	UE_API const UEdGraphPin* GetUsedImagePin(const FNodeMaterialParameterId& ImageId) const;
+
+	/** */
+	UE_API int32 GetNumParameters(const EMaterialParameterType Type) const;
+	UE_API FNodeMaterialParameterId GetParameterId(const EMaterialParameterType Type, const int32 ParameterIndex) const;
+	UE_API FName GetParameterName(const EMaterialParameterType Type, const int32 ParameterIndex) const;
+	static UE_API int32 GetParameterLayerIndex(const UMaterialInterface* InMaterial, const EMaterialParameterType Type, const int32 ParameterIndex);
+	UE_API int32 GetParameterLayerIndex(const EMaterialParameterType Type, const int32 ParameterIndex) const;
+	UE_API FText GetParameterLayerName(const EMaterialParameterType Type, const int32 ParameterIndex) const;
+};
+
+#undef UE_API
diff --git a/Engine/Plugins/Mutable/Source/CustomizableObjectEditor/Private/MuCOE/Nodes/CONodeSkeletalMeshSection.cpp b/Engine/Plugins/Mutable/Source/CustomizableObjectEditor/Private/MuCOE/Nodes/CONodeSkeletalMeshSection.cpp
new file mode 100644
--- /dev/null
+++ b/Engine/Plugins/Mutable/Source/CustomizableObjectEditor/Private/MuCOE/Nodes/CONodeSkeletalMeshSection.cpp
@@ -0,0 +1,1590 @@
+// Copyright Epic Games, Inc. All Rights Reserved.
+
+#include "MuCOE/Nodes/CONodeSkeletalMeshSection.h"
+
+#include "MaterialCachedData.h"
+#include "Materials/Material.h"
+#include "Materials/MaterialExpressionTextureCoordinate.h"
+#include "Materials/MaterialExpressionTextureSample.h"
+#include "Materials/MaterialInstance.h"
+#include "MuCOE/CustomizableObjectEditor_Deprecated.h"
+#include "MuCOE/EdGraphSchema_CustomizableObject.h"
+#include "MuCOE/GraphTraversal.h"
+#include "MuCOE/UnrealEditorPortabilityHelpers.h"
+#include "MuCOE/Nodes/CONodeCopySkeletalMeshSection.h"
+#include "MuCOE/Nodes/CustomizableObjectNodeSkeletalMesh.h"
+#include "MuCOE/Nodes/CustomizableObjectNodeStaticMesh.h"
+#include "MuCOE/Nodes/CustomizableObjectNodeStaticString.h"
+#include "MuCOE/Nodes/CustomizableObjectNodeTable.h"
+#include "MuCOE/Nodes/SCONodeSkeletalMeshSection.h"
+#include "ObjectEditorUtils.h"
+#include "PropertyCustomizationHelpers.h"
+#include "Modules/ModuleManager.h"
+#include "MuCO/CustomizableObjectCustomVersion.h"
+#include "MuCOE/CustomizableObjectEditorLogger.h"
+#include "MuCOE/CustomizableObjectEditorUtilities.h"
+#include "MuCOE/Nodes/CustomizableObjectNodeExternalPin.h"
+#include "MuCOE/CustomizableObjectMacroLibrary/CustomizableObjectMacroLibrary.h"
+
+#include UE_INLINE_GENERATED_CPP_BY_NAME(CONodeSkeletalMeshSection)
+
+class SGraphNode;
+class SWidget;
+class UCustomizableObjectLayout;
+class UObject;
+
+
+#define LOCTEXT_NAMESPACE "CustomizableObjectEditor"
+
+
+const TArray<EMaterialParameterType> UCONodeSkeletalMeshSection::ParameterTypes = {
+	EMaterialParameterType::Texture,
+	EMaterialParameterType::Vector,
+	EMaterialParameterType::Scalar };
+
+
+bool UCustomizableObjectNodeMaterialRemapPinsByName::Equal(const UCustomizableObjectNode& Node, const UEdGraphPin& OldPin, const UEdGraphPin& NewPin) const
+{
+	const UCONodeSkeletalMeshSectionPinDataParameter* PinDataOldPin = Cast<UCONodeSkeletalMeshSectionPinDataParameter>(Node.GetPinData(OldPin));
+	const UCONodeSkeletalMeshSectionPinDataParameter* PinDataNewPin = Cast<UCONodeSkeletalMeshSectionPinDataParameter>(Node.GetPinData(NewPin));
+	if (PinDataOldPin && PinDataNewPin)
+	{
+		return PinDataOldPin->MaterialParameterId == PinDataNewPin->MaterialParameterId && (!OldPin.LinkedTo.Num() || OldPin.PinType == NewPin.PinType); // Pin type must match only if it was connected
+	}
+	else
+	{
+		return Super::Equal(Node, OldPin, NewPin);
+	}
+}
+
+
+void UCustomizableObjectNodeMaterialRemapPinsByName::RemapPins(const UCustomizableObjectNode& Node, const TArray<UEdGraphPin*>& OldPins, const TArray<UEdGraphPin*>& NewPins, TMap<UEdGraphPin*, UEdGraphPin*>& PinsToRemap, TArray<UEdGraphPin*>& PinsToOrphan)
+{
+	for (UEdGraphPin* OldPin : OldPins)
+	{
+		bool bFound = false;
+
+		for (UEdGraphPin* NewPin : NewPins)
+		{
+			if (Equal(Node, *OldPin, *NewPin))
+			{
+				bFound = true;
+
+				if (UEdGraphPin** Result = PinsToRemap.Find(OldPin))
+				{
+					if ((*Result)->bOrphanedPin) // The node can have a deprecated and non-deprecated pin that should remap to the same new pin. Prioritize the non-deprecated
+					{
+						*Result = NewPin;
+					}
+				}
+				else
+				{
+					PinsToRemap.Add(OldPin, NewPin);
+				}
+			}
+		}
+
+		if (!bFound && (OldPin->LinkedTo.Num() || HasSavedPinData(Node, *OldPin)))
+		{
+			PinsToOrphan.Add(OldPin);
+		}
+	}
+}
+
+
+bool UCustomizableObjectNodeMaterialRemapPinsByName::HasSavedPinData(const UCustomizableObjectNode& Node, const UEdGraphPin &Pin) const
+{
+	if (const UCONodeSkeletalMeshSectionPinDataParameter* PinData = Cast<UCONodeSkeletalMeshSectionPinDataParameter>(Node.GetPinData(Pin)))
+	{
+		return !PinData->IsDefault();
+	}
+	else
+	{
+		return false;
+	}
+}
+
+
+bool UCustomizableObjectNodeMaterialPinDataImage::IsDefault() const
+{
+	const UCustomizableObjectNodeMaterialPinDataImage* Default = Cast<UCustomizableObjectNodeMaterialPinDataImage>(GetClass()->GetDefaultObject(false));
+
+	return PinMode == Default->PinMode &&
+		UVLayoutMode == Default->UVLayoutMode &&
+		ReferenceTexture == Default->ReferenceTexture;
+}
+
+
+EPinMode UCustomizableObjectNodeMaterialPinDataImage::GetPinMode() const
+{
+	return PinMode;
+}
+
+
+void UCustomizableObjectNodeMaterialPinDataImage::SetPinMode(const EPinMode InPinMode)
+{
+	FObjectEditorUtils::SetPropertyValue(this, GET_MEMBER_NAME_CHECKED(UCustomizableObjectNodeMaterialPinDataImage, PinMode), InPinMode);
+	NodeMaterial->GetPostImagePinModeChangedDelegate()->Broadcast();
+}
+
+
+void UCustomizableObjectNodeMaterialPinDataImage::Init(UCONodeSkeletalMeshSection& InNodeMaterial)
+{
+	NodeMaterial = &InNodeMaterial;
+}
+
+
+void UCustomizableObjectNodeMaterialPinDataImage::Copy(const UCustomizableObjectNodePinData& Other)
+{
+	if (const UCustomizableObjectNodeMaterialPinDataImage* PinDataOldPin = Cast<UCustomizableObjectNodeMaterialPinDataImage>(&Other))
+	{
+		PinMode = PinDataOldPin->PinMode;
+		UVLayoutMode = PinDataOldPin->UVLayoutMode;
+		UVLayout = PinDataOldPin->UVLayout;
+		ReferenceTexture = PinDataOldPin->ReferenceTexture;
+	}
+}
+
+
+FName UCONodeSkeletalMeshSection::GetPinName(const EMaterialParameterType Type, const int32 ParameterIndex) const
+{
+	const FString ParameterName = GetParameterName(Type, ParameterIndex).ToString();
+	return FName(GetParameterLayerIndex(Type, ParameterIndex) != INDEX_NONE ? GetParameterLayerName(Type, ParameterIndex).ToString() + " - " + ParameterName : ParameterName);
+}
+
+
+int32 UCONodeSkeletalMeshSection::GetExpressionTextureCoordinate(UMaterial* Material, const FGuid& ImageId)
+{
+	if (const UMaterialExpressionTextureSample* TextureSample = Material->FindExpressionByGUID<UMaterialExpressionTextureSample>(ImageId))
+	{
+		if (!TextureSample->Coordinates.Expression)
+		{
+			return TextureSample->ConstCoordinate;
+		}
+		else if (const UMaterialExpressionTextureCoordinate* TextureCoords = Cast<UMaterialExpressionTextureCoordinate>(TextureSample->Coordinates.Expression))
+		{
+			return TextureCoords->CoordinateIndex;
+		}
+	}
+
+	return -1;
+}
+
+
+FName UCONodeSkeletalMeshSection::NodePinModeToImagePinMode(const ENodePinMode NodePinMode)
+{
+	switch (NodePinMode)
+	{
+	case ENodePinMode::Mutable:
+		return UEdGraphSchema_CustomizableObject::PC_Texture;
+	case ENodePinMode::Passthrough:
+		return UEdGraphSchema_CustomizableObject::PC_Texture_Passthrough;
+	default:
+		unimplemented();
+		return {};
+	}
+}
+
+
+FName UCONodeSkeletalMeshSection::GetImagePinMode(const EPinMode PinMode) const
+{
+	switch (PinMode)
+	{
+	case EPinMode::Default:
+		return NodePinModeToImagePinMode(TextureParametersMode);
+	case EPinMode::Mutable:
+		return UEdGraphSchema_CustomizableObject::PC_Texture;
+	case EPinMode::Passthrough:
+		return UEdGraphSchema_CustomizableObject::PC_Texture_Passthrough;
+	default:
+		unimplemented();
+		return {};
+	}
+}
+
+
+FName UCONodeSkeletalMeshSection::GetImagePinMode(const UEdGraphPin& Pin) const
+{
+	return GetImagePinMode(GetPinData<UCustomizableObjectNodeMaterialPinDataImage>(Pin).GetPinMode());
+}
+
+
+int32 UCONodeSkeletalMeshSection::GetImageUVLayoutFromMaterial(const int32 ImageIndex) const
+{
+	const FNodeMaterialParameterId ImageId = GetParameterId(EMaterialParameterType::Texture, ImageIndex);
+
+	if (const int32 TextureCoordinate = GetExpressionTextureCoordinate(Material->GetMaterial(), ImageId.ParameterId);
+		TextureCoordinate >= 0)
+	{
+		return TextureCoordinate;
+	}
+
+	FMaterialLayersFunctions Layers;
+	Material->GetMaterialLayers(Layers);
+
+	TArray<TArray<TObjectPtr<UMaterialFunctionInterface>>*> MaterialFunctionInterfaces;
+	MaterialFunctionInterfaces.SetNumUninitialized(2);
+	MaterialFunctionInterfaces[0] = &Layers.Layers;
+	MaterialFunctionInterfaces[1] = &Layers.Blends;
+
+	for (const TArray<TObjectPtr<UMaterialFunctionInterface>>* MaterialFunctionInterface : MaterialFunctionInterfaces)
+	{
+		for (const TObjectPtr<UMaterialFunctionInterface>& Layer : *MaterialFunctionInterface)
+		{
+			if (const int32 TextureCoordinate = GetExpressionTextureCoordinate(Layer->GetPreviewMaterial()->GetMaterial(), ImageId.ParameterId); TextureCoordinate >= 0)
+			{
+				return TextureCoordinate;
+			}	
+		}
+	}
+
+	return -1;
+}
+
+
+bool UCONodeSkeletalMeshSectionPinDataParameter::IsDefault() const
+{
+	return true;
+}
+
+
+/** Translates a given EPinMode to FText. */
+FText EPinModeToText(const EPinMode PinMode)
+{
+	switch (PinMode)
+	{
+	case EPinMode::Default:
+		return LOCTEXT("EPinModeDefault", "Node Defined");
+
+	case EPinMode::Mutable:
+		return LOCTEXT("EPinModeMutable", "Mutable");
+
+	case EPinMode::Passthrough:
+		return LOCTEXT("EPinModePassthrough", "Passthrough");
+
+	default:
+		check(false); // Missing case.
+		return FText();
+	}
+}
+
+
+void UCustomizableObjectNodeMaterialPinDataImage::PostEditChangeProperty(FPropertyChangedEvent& PropertyChangedEvent)
+{
+	Super::PostEditChangeProperty(PropertyChangedEvent);
+
+	if (const FProperty* PropertyThatChanged = PropertyChangedEvent.Property)
+	{
+		if (PropertyThatChanged->GetFName() == GET_MEMBER_NAME_CHECKED(UCustomizableObjectNodeMaterialPinDataImage, PinMode))
+		{
+			NodeMaterial->UCustomizableObjectNode::ReconstructNode();
+		}
+	}
+}
+
+
+bool UCustomizableObjectNodeMaterialPinDataImage::CanEditChange(const FProperty* InProperty) const
+{
+	if (InProperty->GetFName() == GET_MEMBER_NAME_CHECKED(UCustomizableObjectNodeMaterialPinDataImage, PinMode))
+	{
+		return !NodeMaterial->GetPin(*this)->LinkedTo.Num();
+	}
+	
+	return Super::CanEditChange(InProperty);
+}
+
+
+void UCustomizableObjectNodeMaterialPinDataImage::BackwardsCompatibleFixup(int32 CustomizableObjectCustomVersion)
+{
+	Super::PostLoad();
+	
+	if (CustomizableObjectCustomVersion == FCustomizableObjectCustomVersion::NodeMaterialPinDataImageDetails)
+	{
+		if (UVLayout == UV_LAYOUT_IGNORE)
+		{
+			UVLayoutMode = EUVLayoutMode::Ignore;
+			UVLayout = 0;
+		}
+		else if (UVLayout == UV_LAYOUT_DEFAULT)
+		{
+			UVLayoutMode = EUVLayoutMode::FromMaterial;
+			UVLayout = 0;
+		}
+		else
+		{
+			UVLayoutMode = EUVLayoutMode::Index;
+		}
+	}
+}
+
+
+void UCONodeSkeletalMeshSection::AllocateDefaultPins(UCustomizableObjectNodeRemapPins* RemapPins)
+{
+	const UEdGraphSchema_CustomizableObject* Schema = GetDefault<UEdGraphSchema_CustomizableObject>();
+
+	{
+		const FString PinFriendlyName = TEXT("Mesh");
+		const FString PinName = PinFriendlyName + FString(TEXT("_Input_Pin"));
+		CustomCreatePin(EGPD_Input, Schema->PC_Mesh, FName(*PinName), FText::FromString(PinFriendlyName));
+	}
+	
+	{
+		const FString PinFriendlyName = UEdGraphSchema_CustomizableObject::GetPinCategoryFriendlyName(Schema->PC_Material).ToString();
+		const FString PinName = FString(TEXT("Material_Input_Pin"));
+		UEdGraphPin* TableMaterialPin = CustomCreatePin(EGPD_Input, Schema->PC_Material, FName(*PinName), FText::FromString(PinFriendlyName));
+		TableMaterialPin->PinToolTip = "Pin for a Material from a Table Node";
+	}
+
+	{
+		const FString PinFriendlyName = "Enable Tags";
+		const FString PinName = PinFriendlyName + FString(TEXT("_Input_Pin"));
+		UEdGraphPin* TagsPin = CustomCreatePin(EGPD_Input, Schema->PC_String, FName(PinName), FText::FromString(PinFriendlyName));
+		TagsPin->PinType.ContainerType = EPinContainerType::Array;
+		TagsPin->PinToolTip = "List of Tags that this node will Enable";
+		EnableTagsPinRef = TagsPin;
+	}
+
+	for (const EMaterialParameterType Type : ParameterTypes)
+	{
+		AllocateDefaultParameterPins(Type);
+	}
+
+	{
+		const FString PinFriendlyName = TEXT("Mesh Section");
+		const FString PinName = PinFriendlyName + FString(TEXT("_Output_Pin"));
+		CustomCreatePin(EGPD_Output, Schema->PC_MeshSection, FName(*PinName), FText::FromString(PinFriendlyName));
+	}
+}
+
+
+bool UCONodeSkeletalMeshSection::CanPinBeHidden(const UEdGraphPin& Pin) const
+{
+	return Super::CanPinBeHidden(Pin) &&
+		Pin.Direction == EGPD_Input && 
+		Pin.PinType.PinCategory != UEdGraphSchema_CustomizableObject::PC_Mesh;
+}
+
+
+bool UCONodeSkeletalMeshSection::HasPinViewer() const
+{
+	return true;
+}
+
+
+FString UCONodeSkeletalMeshSection::GetInternalTagDisplayName()
+{
+	return FString::Printf(TEXT("Mesh Section [%s]"), Material ? *Material->GetName() : TEXT("no-material"));
+}
+
+
+UCustomizableObjectNodeRemapPinsByName* UCONodeSkeletalMeshSection::CreateRemapPinsDefault() const
+{
+	return NewObject<UCustomizableObjectNodeMaterialRemapPinsByName>();
+}
+
+
+void UCONodeSkeletalMeshSection::BackwardsCompatibleFixup(int32 CustomizableObjectCustomVersion)
+{
+	Super::BackwardsCompatibleFixup(CustomizableObjectCustomVersion);
+	
+	if (CustomizableObjectCustomVersion == FCustomizableObjectCustomVersion::AutomaticNodeMaterial)
+	{
+		if (Material)
+		{
+			ConditionalPostLoadReference(*Material); // Make sure the Material has been fully loaded.
+		}
+		
+		for (const FCustomizableObjectNodeMaterialImage& Image : Images_DEPRECATED)
+		{
+			const FString OldPinName = (Image.LayerIndex == -1 ? Image.Name : Image.PinName) + FString(TEXT("_Input_Image"));
+			UEdGraphPin* OldPin = FindPin(OldPinName);
+			if (!OldPin) // If we can not find a pin it means that the data was corrupted (old Image array and pins where not synchronized).
+			{
+				continue;
+			}
+
+			UCustomizableObjectNodeMaterialPinDataImage* PinData = NewObject<UCustomizableObjectNodeMaterialPinDataImage>(this);
+			PinData->ParameterId_DEPRECATED = FGuid::NewGuid();
+			PinData->ReferenceTexture = Image.ReferenceTexture;
+
+			// Find referenced Material Parameter
+			const int32 NumParameters = GetNumParameters(EMaterialParameterType::Texture);
+			for (int32 ParameterIndex = 0; ParameterIndex < NumParameters; ++ParameterIndex)
+			{
+				if (GetParameterName(EMaterialParameterType::Texture, ParameterIndex).ToString() == Image.Name)
+				{
+					PinData->ParameterId_DEPRECATED = GetParameterId(EMaterialParameterType::Texture, ParameterIndex).ParameterId;
+
+					if (Image.UVLayout == -1)
+					{
+						PinData->UVLayout = Image.UVLayout;
+					}
+					else
+					{
+						const int32 UVLayout = GetImageUVLayoutFromMaterial(ParameterIndex);
+						if (UVLayout < 0) // Could not be deduced from the Material
+						{
+							PinData->UVLayout = Image.UVLayout;						
+						}
+						else if (UVLayout == Image.UVLayout)
+						{
+							PinData->UVLayout = UV_LAYOUT_DEFAULT;							
+						}
+						else
+						{
+							PinData->UVLayout = UVLayout;
+						}					
+					}
+					
+					break;
+				}
+			}
+			
+			AddPinData(*OldPin, *PinData);
+		}
+
+		for (const FCustomizableObjectNodeMaterialVector& Vector : VectorParams_DEPRECATED)
+		{
+			const FString OldPinName = (Vector.LayerIndex == -1 ? Vector.Name : Vector.PinName) + FString(TEXT("_Input_Vector"));
+			const UEdGraphPin* OldPin = FindPin(OldPinName);
+			if (!OldPin) // If we can not find a pin it means that the data was corrupted (old ScalarParams array and pins where not synchronized).
+			{
+				continue;
+			}
+			
+			UCustomizableObjectNodeMaterialPinDataVector* PinData = NewObject<UCustomizableObjectNodeMaterialPinDataVector>(this);
+			PinData->ParameterId_DEPRECATED = FGuid::NewGuid();
+			
+			// Find referenced Material Parameter
+			const int32 NumParameters = GetNumParameters(EMaterialParameterType::Vector);
+			for (int32 ParameterIndex = 0; ParameterIndex < NumParameters; ++ParameterIndex)
+			{
+				if (GetParameterName(EMaterialParameterType::Vector, ParameterIndex).ToString() == Vector.Name)
+				{
+					PinData->ParameterId_DEPRECATED = GetParameterId(EMaterialParameterType::Vector, ParameterIndex).ParameterId;
+					break;
+				}
+			}
+			
+			AddPinData(*OldPin, *PinData);
+		}
+
+		for (const FCustomizableObjectNodeMaterialScalar& Scalar : ScalarParams_DEPRECATED)
+		{
+			const FString OldPinName = (Scalar.LayerIndex == -1 ? Scalar.Name : Scalar.PinName) + FString(TEXT("_Input_Scalar"));
+			const UEdGraphPin* OldPin = FindPin(OldPinName);
+			if (!OldPin) // If we can not find a pin it means that the data was corrupted (old ScalarParams array and pins where not synchronized).
+			{
+				continue;
+			}
+			
+			UCustomizableObjectNodeMaterialPinDataScalar* PinData = NewObject<UCustomizableObjectNodeMaterialPinDataScalar>(this);
+			PinData->ParameterId_DEPRECATED = FGuid::NewGuid();
+
+			// Find referenced Material Parameter
+			const int32 NumParameters = GetNumParameters(EMaterialParameterType::Scalar);
+			for (int32 ParameterIndex = 0; ParameterIndex < NumParameters; ++ParameterIndex)
+			{
+				if (GetParameterName(EMaterialParameterType::Scalar, ParameterIndex).ToString() == Scalar.Name)
+				{
+					PinData->ParameterId_DEPRECATED = GetParameterId(EMaterialParameterType::Scalar, ParameterIndex).ParameterId;
+					break;
+				}
+			}
+			
+			AddPinData(*OldPin, *PinData);
+		}
+
+		// Check if there are still pins which where not present in the Images, ScalarParams and ScalarParams arrays.
+		for (const UEdGraphPin* Pin : GetAllNonOrphanPins())
+		{
+			if ((Pin->PinType.PinCategory == UEdGraphSchema_CustomizableObject::PC_Texture ||
+ 				 Pin->PinType.PinCategory == UEdGraphSchema_CustomizableObject::PC_Color ||
+				 Pin->PinType.PinCategory == UEdGraphSchema_CustomizableObject::PC_Float) &&
+				!GetPinData(*Pin))
+			{
+				UCONodeSkeletalMeshSectionPinDataParameter* PinData = [&](UObject* Outer) -> UCONodeSkeletalMeshSectionPinDataParameter*
+				{
+					if (Pin->PinType.PinCategory == UEdGraphSchema_CustomizableObject::PC_Texture)
+					{
+						return NewObject<UCustomizableObjectNodeMaterialPinDataImage>(Outer);	
+					}
+					else if (Pin->PinType.PinCategory == UEdGraphSchema_CustomizableObject::PC_Color)
+					{
+						return NewObject<UCustomizableObjectNodeMaterialPinDataVector>(Outer);
+					}
+					else if (Pin->PinType.PinCategory == UEdGraphSchema_CustomizableObject::PC_Float)
+					{
+						return NewObject<UCustomizableObjectNodeMaterialPinDataScalar>(Outer);	
+					}
+					else
+					{
+						check(false); // Parameter type not contemplated.
+						return nullptr;
+					}
+				}(this);
+				
+				PinData->ParameterId_DEPRECATED = FGuid::NewGuid();
+				
+				AddPinData(*Pin, *PinData);
+			}
+		}
+		
+		Images_DEPRECATED.Empty();
+		VectorParams_DEPRECATED.Empty();
+		ScalarParams_DEPRECATED.Empty();
+		
+		Super::ReconstructNode(); // Super required to avoid ambiguous call compilation error.
+	}
+
+	// Fill PinsParameter.
+	if (CustomizableObjectCustomVersion == FCustomizableObjectCustomVersion::AutomaticNodeMaterialPerformanceBug) // || CustomizableObjectCustomVersion < FCustomizableObjectCustomVersion::AutomaticNodeMaterialPerformance
+	{
+		for (const UEdGraphPin* Pin : GetAllNonOrphanPins())
+		{
+			if (const UCONodeSkeletalMeshSectionPinDataParameter* PinData = Cast<UCONodeSkeletalMeshSectionPinDataParameter>(GetPinData(*Pin)))
+			{
+				PinsParameter_DEPRECATED.Add(PinData->ParameterId_DEPRECATED, FEdGraphPinReference(Pin));
+			}
+		}
+	}
+
+	if (CustomizableObjectCustomVersion == FCustomizableObjectCustomVersion::AutomaticNodeMaterialUXImprovements)
+	{
+		TextureParametersMode = bDefaultPinModeMutable_DEPRECATED ?
+			ENodePinMode::Mutable :
+			ENodePinMode::Passthrough;
+	}
+	
+	if (CustomizableObjectCustomVersion == FCustomizableObjectCustomVersion::ExtendMaterialOnlyMutableModeParameters)
+	{
+		const uint32 NumTextureParameters = GetNumParameters(EMaterialParameterType::Texture);
+		for (uint32 ImageIndex = 0; ImageIndex < NumTextureParameters; ++ImageIndex)
+		{
+			if (const UEdGraphPin* Pin = GetParameterPin(EMaterialParameterType::Texture, ImageIndex))
+			{
+				UCustomizableObjectNodeMaterialPinDataImage* PinDataImage = Cast<UCustomizableObjectNodeMaterialPinDataImage>(GetPinData(*Pin));
+				PinDataImage->NodeMaterial = this;
+			}
+		}
+	}
+	
+	if (CustomizableObjectCustomVersion == FCustomizableObjectCustomVersion::ExtendMaterialOnlyMutableModeParametersBug)
+	{
+		for (const UEdGraphPin* Pin : GetAllPins())
+		{
+			if (UCustomizableObjectNodeMaterialPinDataImage* const PinDataImage = Cast<UCustomizableObjectNodeMaterialPinDataImage>(GetPinData(*Pin)))
+			{
+				PinDataImage->NodeMaterial = this;
+			}
+		}
+	}
+	
+	if (CustomizableObjectCustomVersion == FCustomizableObjectCustomVersion::NodeMaterialAddTablePin)
+	{
+		if (Material)
+		{
+			ConditionalPostLoadReference(*Material); // Make sure the Material has been fully loaded.
+		}
+		
+		Super::ReconstructNode();
+	}
+
+	if (CustomizableObjectCustomVersion == FCustomizableObjectCustomVersion::AddedTableMaterialSwitch)
+	{
+		if (Material)
+		{
+			ConditionalPostLoadReference(*Material); // Make sure the Material has been fully loaded.
+		}
+		
+		UMaterialInstance* DefaultPinValue = nullptr;
+
+		if (const UEdGraphPin* MaterialAssetPin = GetMaterialAssetPin())
+		{
+			if (const UEdGraphPin* ConnectedPin = FollowInputPin(*MaterialAssetPin))
+			{
+				if (const UCustomizableObjectNodeTable* TableNode = Cast<UCustomizableObjectNodeTable>(ConnectedPin->GetOwningNode()))
+				{
+					DefaultPinValue = TableNode->GetColumnDefaultAssetByType<UMaterialInstance>(ConnectedPin);
+				}
+			}
+		}
+
+		if (DefaultPinValue && DefaultPinValue->TextureParameterValues.Num())
+		{
+			const uint32 NumTextureParameters = GetNumParameters(EMaterialParameterType::Texture);
+			for (uint32 ImageIndex = 0; ImageIndex < NumTextureParameters; ++ImageIndex)
+			{
+				if (const UEdGraphPin* ImagePin = GetParameterPin(EMaterialParameterType::Texture, ImageIndex))
+				{
+					FGuid ParameterId = GetParameterId(EMaterialParameterType::Texture, ImageIndex).ParameterId;
+
+					TArray<FMaterialParameterInfo> TextureParameterInfo;
+					TArray<FGuid> TextureGuids;
+
+					// Getting parent's texture infos
+					DefaultPinValue->GetMaterial()->GetAllTextureParameterInfo(TextureParameterInfo, TextureGuids);
+
+					int32 TextureIndex = TextureGuids.Find(ParameterId);
+
+					if (TextureIndex == INDEX_NONE)
+					{
+						continue;
+					}
+
+					FName TextureName = TextureParameterInfo[TextureIndex].Name;
+
+					// Checking if the pin's texture has been modified in the material instance
+					for (const FTextureParameterValue& Texture : DefaultPinValue->TextureParameterValues)
+					{
+						if (TextureName == Texture.ParameterInfo.Name)
+						{
+							if (UCustomizableObjectNodeMaterialPinDataImage* PinDataImage = Cast<UCustomizableObjectNodeMaterialPinDataImage>(GetPinData(*ImagePin)))
+							{
+								PinDataImage->PinMode = EPinMode::Mutable;
+							}
+						}
+					}
+				}
+			}
+		}
+	}
+
+	if (CustomizableObjectCustomVersion == FCustomizableObjectCustomVersion::NewComponentOptions)
+	{
+		MeshComponentName_DEPRECATED = FName(FString::FromInt(MeshComponentIndex_DEPRECATED));
+	}
+
+	if (CustomizableObjectCustomVersion == FCustomizableObjectCustomVersion::NodeMaterialTypedImagePins)
+	{
+		const UEdGraphSchema_CustomizableObject* Schema = GetDefault<UEdGraphSchema_CustomizableObject>();
+		
+		for (UEdGraphPin* Pin : GetAllPins())
+		{
+			if (Pin->PinType.PinCategory != UEdGraphSchema_CustomizableObject::PC_Texture)
+			{
+				continue;
+			}
+			
+			if (Pin->LinkedTo.Num())
+			{
+				const UEdGraphPin* LinkedPin = Pin->LinkedTo[0];
+
+				UCustomizableObjectNodeMaterialPinDataImage& PinData = GetPinData<UCustomizableObjectNodeMaterialPinDataImage>(*Pin);
+
+				if (LinkedPin->PinType.PinCategory == UEdGraphSchema_CustomizableObject::PC_Texture)
+				{
+					Pin->PinType.PinCategory = UEdGraphSchema_CustomizableObject::PC_Texture;
+
+					if (PinData.PinMode == EPinMode::Default &&
+						TextureParametersMode == ENodePinMode::Passthrough)
+					{
+						PinData.PinMode = EPinMode::Mutable;
+					}
+					else if (PinData.PinMode == EPinMode::Passthrough)
+					{
+						PinData.PinMode = EPinMode::Mutable;
+					}
+				}
+				else if (LinkedPin->PinType.PinCategory == UEdGraphSchema_CustomizableObject::PC_Texture_Passthrough)
+				{
+					Pin->PinType.PinCategory = UEdGraphSchema_CustomizableObject::PC_Texture_Passthrough;
+
+					if (PinData.PinMode == EPinMode::Default &&
+						TextureParametersMode == ENodePinMode::Mutable)
+					{
+						PinData.PinMode = EPinMode::Passthrough;
+					}
+					else if (PinData.PinMode == EPinMode::Mutable)
+					{
+						PinData.PinMode = EPinMode::Passthrough;
+					}
+				}
+			}
+			else
+			{
+				if (!IsImageMutableMode(*Pin))
+				{
+					Pin->PinType.PinCategory = UEdGraphSchema_CustomizableObject::PC_Texture_Passthrough;
+				}				
+			}
+		}
+	}
+
+	if (CustomizableObjectCustomVersion == FCustomizableObjectCustomVersion::FixedMultilayerMaterialIds)
+	{
+		if (Material && Material->GetCachedExpressionData().bHasMaterialLayers)
+		{
+			ConditionalPostLoadReference(*Material); // Make sure the Material has been fully loaded.
+			
+			// Needed since we can not get the layer index of repeated parameters
+			Super::ReconstructNode();
+		}
+		else
+		{
+			for (TMap<FGuid, FEdGraphPinReference>::TIterator It = PinsParameter_DEPRECATED.CreateIterator(); It; ++It)
+			{
+				PinsParameterMap.Add({ It.Key(), -1 }, It.Value());
+
+				// Move pin data id info to the new struct
+				UEdGraphPin* GraphPin = It.Value().Get();
+				if (!GraphPin)
+				{
+					continue;
+				}
+
+				UCustomizableObjectNodePinData* GenericPinData = GetPinData(*GraphPin);
+				if (!GenericPinData)
+				{
+					continue;
+				}
+
+				if (UCONodeSkeletalMeshSectionPinDataParameter* PinData = Cast<UCONodeSkeletalMeshSectionPinDataParameter>(GenericPinData))
+				{
+					PinData->MaterialParameterId.LayerIndex = INDEX_NONE;
+					PinData->MaterialParameterId.ParameterId = PinData->ParameterId_DEPRECATED;
+				}
+			}
+		}
+	}
+
+	if (CustomizableObjectCustomVersion == FCustomizableObjectCustomVersion::MaterialPinsRename)
+	{
+		UEdGraphPin* MaterialPin = FindPin(TEXT("Material_Output_Pin"), EEdGraphPinDirection::EGPD_Output);
+		if (MaterialPin)
+		{
+			const FString PinFriendlyName = TEXT("Mesh Section");
+			const FString PinName = PinFriendlyName + FString(TEXT("_Output_Pin"));
+			MaterialPin->PinName = FName(PinName);
+			MaterialPin->PinFriendlyName = FText::FromString(PinFriendlyName);
+		}
+	}
+
+	if (CustomizableObjectCustomVersion == FCustomizableObjectCustomVersion::UpdatedNodesPinName)
+	{
+		const FString ExpectedPinFriendlyName = "Table Material";
+		const FString ExpectedPinName = ExpectedPinFriendlyName + FString(TEXT("_Input_Pin"));	
+		
+		if (UEdGraphPin* FoundPin = FindPin(ExpectedPinName,EGPD_Input))
+		{
+			const FString TargetPinFriendlyName = "Material";
+			const FString TargetPinName = TargetPinFriendlyName + FString(TEXT("_Input_Pin"));
+			
+			FoundPin->PinFriendlyName = FText::FromString(TargetPinFriendlyName);
+			FoundPin->PinName = FName(TargetPinName);
+		}
+	}
+
+	if (CustomizableObjectCustomVersion == FCustomizableObjectCustomVersion::EnableMutableMacrosNewVersion)
+	{
+		if (!EnableTagsPinRef.Get())
+		{
+			const FString PinFriendlyName = "Enable Tags";
+			const FString PinName = PinFriendlyName + FString(TEXT("_Input_Pin"));
+			UEdGraphPin* TagsPin = CustomCreatePin(EGPD_Input, UEdGraphSchema_CustomizableObject::PC_String, FName(PinName),  FText::FromString(PinFriendlyName));
+			TagsPin->PinType.ContainerType = EPinContainerType::Array;
+			TagsPin->PinToolTip = "List of Tags that this node will Enable";
+			EnableTagsPinRef = TagsPin;
+		}
+	}
+
+	if (CustomizableObjectCustomVersion == FCustomizableObjectCustomVersion::FixMaterialNodeMaterialPinIncorrectLocalization)
+	{
+		for (UEdGraphPin* Pin : GetAllPins())
+		{
+			if (Pin)
+			{
+				if (Pin->PinType.PinCategory == UEdGraphSchema_CustomizableObject::PC_Material && 
+					Pin->Direction == EGPD_Input && Pin->PinName.ToString().EndsWith(TEXT("_Input_Pin")) &&
+					Pin->PinName != FString(TEXT("Material_Input_Pin")))
+				{
+					FString PinName = FString(TEXT("Material_Input_Pin"));
+					Pin->PinName = FName(*PinName);
+				}
+			}
+		}
+	}
+}
+
+
+void UCONodeSkeletalMeshSection::Serialize(FArchive& Ar)
+{
+	Super::Serialize(Ar);
+
+	Ar.UsingCustomVersion(FCustomizableObjectCustomVersion::GUID);
+}
+
+
+void UCONodeSkeletalMeshSection::PostEditChangeProperty(FPropertyChangedEvent& PropertyChangedEvent)
+{
+	Super::PostEditChangeProperty(PropertyChangedEvent);
+
+	if (const FProperty* PropertyThatChanged = PropertyChangedEvent.Property)
+	{
+		if (PropertyThatChanged->GetFName() == GET_MEMBER_NAME_CHECKED(UCONodeSkeletalMeshSection, Material))
+		{
+			Super::ReconstructNode(); // Super required to avoid ambiguous call compilation error.
+		}
+		else if (PropertyThatChanged->GetFName() == GET_MEMBER_NAME_CHECKED(UCONodeSkeletalMeshSection, TextureParametersMode))
+		{
+			Super::ReconstructNode();
+			PostImagePinModeChangedDelegate.Broadcast();
+		}
+	}
+}
+
+
+FText UCONodeSkeletalMeshSection::GetNodeTitle(ENodeTitleType::Type TitleType) const
+{
+	if (TitleType == ENodeTitleType::ListView || !Material)
+	{
+		return LOCTEXT("Mesh Section", "Mesh Section");
+	}
+	else
+	{
+		FFormatNamedArguments Args;
+		Args.Add(TEXT("MeshSectionName"), FText::FromString(Material->GetName()));
+
+		return FText::Format(LOCTEXT("MeshSection_Title", "{MeshSectionName}\nMesh Section"), Args);
+	}
+}
+
+
+UEdGraphPin* UCONodeSkeletalMeshSection::OutputPin() const
+{
+	FString PinFriendlyName = TEXT("Mesh Section");
+	FString PinName = PinFriendlyName + FString(TEXT("_Output_Pin"));
+
+	UEdGraphPin* Pin = FindPin(PinName, EEdGraphPinDirection::EGPD_Output);
+	if (!Pin)
+	{
+		Pin = FindPin(TEXT("Mesh Section"), EEdGraphPinDirection::EGPD_Output);
+	}
+
+	// Legacy name
+	if (!Pin)
+	{
+		PinFriendlyName = TEXT("Material");
+		PinName = PinFriendlyName + FString(TEXT("_Output_Pin"));
+
+		Pin = FindPin(PinName, EEdGraphPinDirection::EGPD_Output);
+		if (!Pin)
+		{
+			Pin = FindPin(TEXT("Material"), EEdGraphPinDirection::EGPD_Output);
+		}		
+	}
+
+	return Pin;
+}
+
+
+UEdGraphPin* UCONodeSkeletalMeshSection::GetMeshPin() const
+{
+	FString PinFriendlyName = TEXT("Mesh");
+	FString PinName = PinFriendlyName + FString(TEXT("_Input_Pin"));
+
+	UEdGraphPin* Pin = FindPin(PinName);
+
+	if(Pin)
+	{
+		return Pin;
+	}
+	else
+	{
+		return FindPin(TEXT("Mesh"));
+	}
+}
+
+
+UEdGraphPin* UCONodeSkeletalMeshSection::GetMaterialAssetPin() const
+{
+	const FString PinFriendlyName = UEdGraphSchema_CustomizableObject::GetPinCategoryFriendlyName(UEdGraphSchema_CustomizableObject::PC_Material).ToString();
+	const FString PinName = PinFriendlyName + FString(TEXT("_Input_Pin"));
+
+	UEdGraphPin* Pin = FindPin(PinName);
+
+	if (Pin)
+	{
+		return Pin;
+	}
+	else
+	{
+		// Changing the locale can change the GetPinCategoryFriendlyName() result, leading to failing to find the MaterialAsset pin.
+		// This is a provisional fix while incorrect FriendlyName usages are removed in UE-294583
+		Pin = FindPin(PinFriendlyName);
+
+		if (Pin)
+		{
+			return Pin;
+		}
+		else
+		{
+			return FindPin(FString("Material_Input_Pin"));
+		}
+	}
+}
+
+
+UEdGraphPin* UCONodeSkeletalMeshSection::GetEnableTagsPin() const
+{
+	return EnableTagsPinRef.Get();
+}
+
+
+const UCONodeSkeletalMeshSection* UCONodeSkeletalMeshSection::GetMaterialNode() const
+{
+	return this;
+}
+
+
+bool UCONodeSkeletalMeshSection::IsImageMutableMode(const int32 ImageIndex) const
+{
+	if (const UEdGraphPin* Pin = GetParameterPin(EMaterialParameterType::Texture, ImageIndex))
+	{
+		return IsImageMutableMode(*Pin);
+	}
+	else
+	{
+		return NodePinModeToImagePinMode(TextureParametersMode) == UEdGraphSchema_CustomizableObject::PC_Texture;
+	}
+}
+
+
+bool UCONodeSkeletalMeshSection::IsImageMutableMode(const UEdGraphPin& Pin) const
+{
+	return GetImagePinMode(Pin) == UEdGraphSchema_CustomizableObject::PC_Texture; // The ImageMutableMode is stored in the PinData, not in the PinCategory
+}
+
+
+void UCONodeSkeletalMeshSection::SetImagePinMode(UEdGraphPin& Pin, const EPinMode PinMode) const
+{
+	UCustomizableObjectNodeMaterialPinDataImage& PinData = GetPinData<UCustomizableObjectNodeMaterialPinDataImage>(Pin);
+
+	Pin.PinType.PinCategory = GetImagePinMode(PinMode);	// Change the category so that it will remap correctly when reconstructing.
+	
+	PinData.SetPinMode(PinMode); // Will trigger a reconstruct.
+}
+
+
+UTexture2D* UCONodeSkeletalMeshSection::GetImageReferenceTexture(const int32 ImageIndex) const
+{
+	if (const UEdGraphPin* Pin = GetParameterPin(EMaterialParameterType::Texture, ImageIndex))
+	{
+		return GetPinData<UCustomizableObjectNodeMaterialPinDataImage>(*Pin).ReferenceTexture;
+	}
+	else
+	{
+		return nullptr;
+	}
+}
+
+
+UTexture2D* UCONodeSkeletalMeshSection::GetImageValue(const int32 ImageIndex) const
+{
+	FName TextureName = GetParameterName(EMaterialParameterType::Texture, ImageIndex);
+	int32 LayerIndex = GetParameterLayerIndex(EMaterialParameterType::Texture, ImageIndex);
+	
+	FHashedMaterialParameterInfo TextureInfo;
+	TextureInfo.Name = FScriptName(TextureName);
+	TextureInfo.Index = LayerIndex;
+	TextureInfo.Association = LayerIndex == INDEX_NONE ? GlobalParameter : LayerParameter; // TODO: Blend Parameters
+
+	UTexture* Texture;
+	Material->GetTextureParameterValue(TextureInfo, Texture);
+
+	return Cast<UTexture2D>(Texture);
+}
+
+
+TArray<UCustomizableObjectLayout*> UCONodeSkeletalMeshSection::GetLayouts() const
+{
+	TArray<UCustomizableObjectLayout*> Result;
+
+	if (UEdGraphPin* MeshPin = GetMeshPin())
+	{
+		if (const UEdGraphPin* ConnectedPin = FollowInputPin(*MeshPin))
+		{
+			if (const UEdGraphPin* SourceMeshPin = FindMeshBaseSource(*ConnectedPin, false))
+			{
+				if (const UCustomizableObjectNodeSkeletalMesh* MeshNode = Cast<UCustomizableObjectNodeSkeletalMesh>(SourceMeshPin->GetOwningNode()))
+				{
+					Result = MeshNode->GetLayouts(*SourceMeshPin);
+				}
+				else if (const UCustomizableObjectNodeTable* TableNode = Cast<UCustomizableObjectNodeTable>(SourceMeshPin->GetOwningNode()))
+				{
+					Result = TableNode->GetLayouts(SourceMeshPin);
+				}
+			}
+		}
+	}
+
+	return Result;
+}
+
+
+int32 UCONodeSkeletalMeshSection::GetImageUVLayout(const int32 ImageIndex) const
+{
+	if (const UEdGraphPin* Pin = GetParameterPin(EMaterialParameterType::Texture, ImageIndex))
+	{
+		const UCustomizableObjectNodeMaterialPinDataImage& PinData = GetPinData<UCustomizableObjectNodeMaterialPinDataImage>(*Pin);
+		switch(PinData.UVLayoutMode)
+		{
+		case EUVLayoutMode::FromMaterial:
+			break;
+		case EUVLayoutMode::Ignore:
+			return UCustomizableObjectNodeMaterialPinDataImage::UV_LAYOUT_IGNORE;
+		case EUVLayoutMode::Index:
+			return PinData.UVLayout;
+		default:
+			unimplemented();
+		}
+	}
+
+	const int32 UVIndex = GetImageUVLayoutFromMaterial(ImageIndex);
+	if (UVIndex == -1)
+	{
+		const FText ParamName = FText::FromName(GetParameterName(EMaterialParameterType::Texture, ImageIndex));
+		const FText Msg = FText::Format(LOCTEXT("UVLayoutMaterialError", "Could not deduce the UV Layout Index in [{0}]. Use [Index] UV Layout Mode or, in the UMaterial, remove any nodes connected to the [UVs] pin in [{0}] node ."), ParamName);
+		FCustomizableObjectEditorLogger::CreateLog(Msg)
+			.Severity(EMessageSeverity::Warning)
+			.Context(*this)
+			.Log();
+		
+		return 0;
+	}
+
+	return UVIndex;
+}
+
+
+int32 UCONodeSkeletalMeshSection::GetNumParameters(const EMaterialParameterType Type) const
+{
+	if (Material)
+	{
+		return Material->GetCachedExpressionData().GetParameterTypeEntry(Type).ParameterInfoSet.Num();
+	}
+	else
+	{
+		return 0;
+	}
+}
+
+
+FNodeMaterialParameterId UCONodeSkeletalMeshSection::GetParameterId(const EMaterialParameterType Type, const int32 ParameterIndex) const
+{
+	const FMaterialCachedExpressionData& Data = Material->GetCachedExpressionData();
+
+	if (Data.EditorOnlyData)
+	{
+		if (Data.EditorOnlyData->EditorEntries[(int32)Type].EditorInfo.Num() != 0)
+		{
+			const FGuid ParameterId = Data.EditorOnlyData->EditorEntries[(int32)Type].EditorInfo[ParameterIndex].ExpressionGuid;
+			const int32 LayerIndex = GetParameterLayerIndex(Type, ParameterIndex);
+
+			return { ParameterId, LayerIndex };
+		}
+	}
+	
+	return FNodeMaterialParameterId();
+}
+
+
+FName UCONodeSkeletalMeshSection::GetParameterName(const EMaterialParameterType Type, const int32 ParameterIndex) const
+{
+	check(Material);
+	
+	const FMaterialCachedParameterEntry& Entry = Material->GetCachedExpressionData().GetParameterTypeEntry(Type);
+
+	for (TSet<FMaterialParameterInfo>::TConstIterator It(Entry.ParameterInfoSet); It; ++It)
+	{
+		const int32 IteratorIndex = It.GetId().AsInteger();
+
+		if (IteratorIndex == ParameterIndex)
+		{
+			return (*It).Name;
+		}
+	}
+	
+	// The parameter should exist
+	check(false);
+
+	return FName();
+}
+
+
+int32 UCONodeSkeletalMeshSection::GetParameterLayerIndex(const UMaterialInterface* InMaterial, const EMaterialParameterType Type, const int32 ParameterIndex)
+{
+	check(InMaterial);
+
+	const FMaterialCachedParameterEntry& Entry = InMaterial->GetCachedExpressionData().GetParameterTypeEntry(Type);
+
+	for (TSet<FMaterialParameterInfo>::TConstIterator It(Entry.ParameterInfoSet); It; ++It)
+	{
+		const int32 IteratorIndex = It.GetId().AsInteger();
+
+		if (IteratorIndex == ParameterIndex)
+		{
+			return (*It).Index;
+		}
+	}
+
+	// The parameter should exist
+	check(false);
+
+	return -1;
+}
+
+
+int32 UCONodeSkeletalMeshSection::GetParameterLayerIndex(const EMaterialParameterType Type, const int32 ParameterIndex) const
+{
+	return GetParameterLayerIndex(Material.Get(), Type, ParameterIndex);
+}
+
+
+FText UCONodeSkeletalMeshSection::GetParameterLayerName(const EMaterialParameterType Type, const int32 ParameterIndex) const
+{
+	check(Material)
+
+	int32 LayerIndex = GetParameterLayerIndex(Type,ParameterIndex);
+
+	FMaterialLayersFunctions LayersValue;
+	Material->GetMaterialLayers(LayersValue);
+
+
+	return LayersValue.EditorOnly.LayerNames.IsValidIndex(LayerIndex) ? LayersValue.EditorOnly.LayerNames[LayerIndex] : FText();
+}
+
+
+bool UCONodeSkeletalMeshSection::HasParameter(const UMaterialInterface* InMaterial, const FNodeMaterialParameterId& ParameterId)
+{
+	if (!InMaterial)
+	{
+		return false;
+	}
+
+	for (const EMaterialParameterType Type : ParameterTypes)
+	{
+		const FMaterialCachedExpressionData& Data = InMaterial->GetCachedExpressionData();
+		const FMaterialCachedParameterEntry& Entry = Data.GetParameterTypeEntry(Type);
+
+		if (!Data.EditorOnlyData || Data.EditorOnlyData->EditorEntries[(int32)Type].EditorInfo.IsEmpty())
+		{
+			continue;
+		}
+
+		for (TSet<FMaterialParameterInfo>::TConstIterator It(Entry.ParameterInfoSet); It; ++It)
+		{
+			const int32 IteratorIndex = It.GetId().AsInteger();
+			
+			const FGuid& ParamGuid = Data.EditorOnlyData->EditorEntries[(int32)Type].EditorInfo[IteratorIndex].ExpressionGuid;
+			const int32 LayerIndex = GetParameterLayerIndex(InMaterial, Type, IteratorIndex);
+			const FNodeMaterialParameterId ParamId = { ParamGuid, LayerIndex };
+
+			if (ParamId == ParameterId)
+			{
+				return true;
+			}
+		}
+	}
+
+	return false;
+}
+
+
+bool UCONodeSkeletalMeshSection::HasParameter(const FNodeMaterialParameterId& ParameterId) const
+{
+	return HasParameter( Material.Get(), ParameterId );
+}
+
+
+UEdGraphPin* UCONodeSkeletalMeshSection::GetParameterPin(const EMaterialParameterType Type, const int32 ParameterIndex) const
+{
+	const FNodeMaterialParameterId ParameterId = GetParameterId(Type, ParameterIndex);
+	
+	return GetParameterPin(ParameterId);
+}
+
+
+UEdGraphPin* UCONodeSkeletalMeshSection::GetParameterPin(const FNodeMaterialParameterId& ParameterId) const
+{
+	if (const FEdGraphPinReference* Result = PinsParameterMap.Find(ParameterId))
+	{
+		return Result->Get();
+	}
+	else
+	{
+		return nullptr;
+	}
+}
+
+
+bool UCONodeSkeletalMeshSection::IsNodeOutDatedAndNeedsRefresh()
+{
+	const bool bOutdated = RealMaterialDataHasChanged();
+
+	// Remove previous compilation warnings
+	if (!bOutdated && bHasCompilerMessage)
+	{
+		RemoveWarnings();
+		GetGraph()->NotifyGraphChanged();
+	}
+
+	return bOutdated;
+}
+
+
+FString UCONodeSkeletalMeshSection::GetRefreshMessage() const
+{
+	return "Referenced material has changed, texture channels might have been added, removed or renamed. Please refresh the node material to reflect those changes.";
+}
+
+
+TSharedPtr<IDetailsView> UCONodeSkeletalMeshSection::CustomizePinDetails(const UEdGraphPin& Pin) const
+{
+	if (UCustomizableObjectNodeMaterialPinDataImage* PinData = Cast<UCustomizableObjectNodeMaterialPinDataImage>(GetPinData(Pin)))
+	{
+		FPropertyEditorModule& EditModule = FModuleManager::Get().GetModuleChecked<FPropertyEditorModule>("PropertyEditor");
+
+		FDetailsViewArgs DetailsViewArgs;
+		DetailsViewArgs.bAllowSearch = false;
+		DetailsViewArgs.bHideSelectionTip = true;
+		
+		const TSharedRef<IDetailsView> SettingsView = EditModule.CreateDetailView(DetailsViewArgs);
+		SettingsView->SetObject(PinData);
+		
+		return SettingsView;
+	}
+	else
+	{
+		return nullptr;
+	}
+}
+
+
+bool UCONodeSkeletalMeshSection::CustomRemovePin(UEdGraphPin& Pin)
+{
+	for (TMap<FNodeMaterialParameterId, FEdGraphPinReference>::TIterator Iterator = PinsParameterMap.CreateIterator(); Iterator; ++Iterator)
+	{
+		if (Iterator->Value.Get() == &Pin)
+		{
+			Iterator.RemoveCurrent();
+			break;
+		}
+	}
+
+	return Super::CustomRemovePin(Pin);
+}
+
+
+void UCONodeSkeletalMeshSection::SetMaterial(UMaterialInterface* InMaterial)
+{
+	Material = InMaterial;
+}
+
+
+UMaterialInterface* UCONodeSkeletalMeshSection::GetMaterial() const
+{
+	return Material;	
+}
+
+
+TArray<FString> UCONodeSkeletalMeshSection::GetEnableTags(TArray<const UCustomizableObjectNodeMacroInstance*>* MacroContext)
+{
+	// Getting tags from linked pin
+	const UEdGraphPin* EnableTagsPin = GetEnableTagsPin();
+
+	if (!EnableTagsPin)
+	{
+		return Tags;
+	}
+
+	const TArray<UEdGraphPin*> ConnectedPins = FollowInputPinArray(*EnableTagsPin);
+
+	if (ConnectedPins.Num())
+	{
+		TArray<FString> OutTags;
+
+		for (const UEdGraphPin* StringPin : ConnectedPins)
+		{
+			const UEdGraphPin* SourceStringPin = GraphTraversal::FindIOPinSourceThroughMacroContext(*StringPin, MacroContext);
+
+			if (SourceStringPin)
+			{
+				if (const UCustomizableObjectNodeStaticString* StringNode = Cast<UCustomizableObjectNodeStaticString>(SourceStringPin->GetOwningNode()))
+				{
+					OutTags.AddUnique(StringNode->Value);
+				}
+			}
+		}
+
+		return OutTags;
+	}
+
+	return Tags;
+}
+
+
+TArray<FString>* UCONodeSkeletalMeshSection::GetEnableTagsArray()
+{
+	return &Tags;
+}
+
+
+bool UCONodeSkeletalMeshSection::RealMaterialDataHasChanged() const
+{
+	for (const UEdGraphPin* Pin : GetAllNonOrphanPins())
+	{
+		if (const UCONodeSkeletalMeshSectionPinDataParameter* PinData = Cast<UCONodeSkeletalMeshSectionPinDataParameter>(GetPinData(*Pin)))
+		{
+			if (!HasParameter(PinData->MaterialParameterId) &&
+				(FollowInputPin(*Pin) || !PinData->IsDefault()))
+			{
+				return true;
+			}
+		}
+	}
+
+	return false;
+}
+
+
+FPostImagePinModeChangedDelegate* UCONodeSkeletalMeshSection::GetPostImagePinModeChangedDelegate()
+{
+	return &PostImagePinModeChangedDelegate;
+}
+
+
+bool UCONodeSkeletalMeshSection::IsPinRelevant(const UEdGraphPin* Pin) const
+{
+	const UEdGraphSchema_CustomizableObject* Schema = GetDefault<UEdGraphSchema_CustomizableObject>();
+
+	if (Pin->Direction == EEdGraphPinDirection::EGPD_Output)
+	{	
+		return Pin->PinType.PinCategory == Schema->PC_Mesh;
+	}
+	else if (Pin->Direction == EEdGraphPinDirection::EGPD_Input)
+	{	
+		return Pin->PinType.PinCategory == Schema->PC_MeshSection;
+	}
+	else
+	{
+		return false;
+	}
+}
+
+
+FText UCONodeSkeletalMeshSection::GetTooltipText() const
+{
+	return LOCTEXT("MeshSection_Tooltip", "Defines a Customizable Object mesh section.\nIt has a mesh section, a material assigned to it and the runtime modifiable inputs to the material asset parameters.");
+}
+
+
+TSharedPtr<SGraphNode> UCONodeSkeletalMeshSection::CreateVisualWidget()
+{
+	return SNew(SCONodeSkeletalMeshSection, this);
+}
+
+
+UCONodeSkeletalMeshSectionPinDataParameter* UCONodeSkeletalMeshSection::CreatePinData(const EMaterialParameterType Type, const int32 ParameterIndex)
+{
+	UCONodeSkeletalMeshSectionPinDataParameter* PinData = nullptr;
+
+	switch (Type)
+	{
+	case EMaterialParameterType::Texture:
+		{
+			UCustomizableObjectNodeMaterialPinDataImage* PinDataImage = NewObject<UCustomizableObjectNodeMaterialPinDataImage>(this);
+			PinDataImage->Init(*this);
+		
+			PinData = PinDataImage;
+			break;
+		}
+	case EMaterialParameterType::Vector:
+		{
+			PinData = NewObject<UCustomizableObjectNodeMaterialPinDataVector>(this);	
+			break;
+		}
+
+	case EMaterialParameterType::Scalar:
+		{
+			PinData = NewObject<UCustomizableObjectNodeMaterialPinDataScalar>(this);	
+			break;
+		}
+
+	default:
+		check(false); // Parameter type not contemplated.
+	}
+
+	PinData->MaterialParameterId = GetParameterId(Type, ParameterIndex);
+
+	return PinData;
+}
+
+void UCONodeSkeletalMeshSection::AllocateDefaultParameterPins(const EMaterialParameterType Type)
+{
+	const int32 NumParameters = GetNumParameters(Type);
+	for (int32 ParameterIndex = 0; ParameterIndex < NumParameters; ++ParameterIndex)
+	{
+		UCONodeSkeletalMeshSectionPinDataParameter* PinData = CreatePinData(Type, ParameterIndex);
+
+		const FName PinName = GetPinName(Type, ParameterIndex);
+
+		FName PinCategory;
+		switch(Type)
+		{
+		case EMaterialParameterType::Texture:
+			if (IsImageMutableMode(ParameterIndex)) // If a pin exists, we store the PinMode in its PinData.
+			{
+				PinCategory = UEdGraphSchema_CustomizableObject::PC_Texture;				
+			}
+			else
+			{
+				PinCategory = UEdGraphSchema_CustomizableObject::PC_Texture_Passthrough;
+			}
+			break;
+
+		case EMaterialParameterType::Vector:
+			PinCategory = UEdGraphSchema_CustomizableObject::PC_Color;
+			break;
+
+		case EMaterialParameterType::Scalar:
+			PinCategory = UEdGraphSchema_CustomizableObject::PC_Float;
+			break;
+
+		default:
+			check(false); // Type not contemplated
+		}		
+		
+		UEdGraphPin* Pin = CustomCreatePin(EGPD_Input, PinCategory, PinName, FText::FromName(PinName), PinData);
+		Pin->bHidden = true;
+
+		PinsParameterMap.Add(PinData->MaterialParameterId, FEdGraphPinReference(Pin));
+	}
+}
+
+
+void UCONodeSkeletalMeshSection::SetDefaultMaterial()
+{
+	if (const UEdGraphPin* MeshPin = GetMeshPin();
+		!Material && MeshPin)
+	{
+		if (const UEdGraphPin* LinkedMeshPin = FollowInputPin(*MeshPin))
+		{
+			UEdGraphNode* LinkedMeshNode = LinkedMeshPin->GetOwningNode();
+	
+			if (const UCustomizableObjectNodeSkeletalMesh* NodeSkeletalMesh = Cast<UCustomizableObjectNodeSkeletalMesh>(LinkedMeshNode))
+			{
+				Material = NodeSkeletalMesh->GetMaterialFor(LinkedMeshPin);
+				if (Material)
+				{
+					Super::ReconstructNode(); // Super required to avoid ambiguous call compilation error.
+				}
+			}
+			else if (const UCustomizableObjectNodeStaticMesh* NodeStaticMesh = Cast<UCustomizableObjectNodeStaticMesh>(LinkedMeshNode))
+			{
+				Material = NodeStaticMesh->GetMaterialFor(LinkedMeshPin);
+				if (Material)
+				{
+					Super::ReconstructNode();
+				}
+			}
+		}
+	}
+}
+
+
+void UCONodeSkeletalMeshSection::PinConnectionListChanged(UEdGraphPin* Pin)
+{
+	Super::PinConnectionListChanged(Pin);
+	
+	if (Pin == GetMeshPin())
+	{
+		if (LastMeshNodeConnected.IsValid())
+		{
+			LastMeshNodeConnected->PostEditChangePropertyDelegate.RemoveDynamic(this, &UCONodeSkeletalMeshSection::MeshPostEditChangeProperty);
+		}
+
+		if (const UEdGraphPin* ConnectedPin = FollowInputPin(*Pin))
+		{
+			UEdGraphNode* MeshNode = ConnectedPin->GetOwningNode();
+
+			if (MeshNode->IsA(UCustomizableObjectNodeStaticMesh::StaticClass()) || MeshNode->IsA(UCustomizableObjectNodeSkeletalMesh::StaticClass()))
+			{
+				SetDefaultMaterial();
+
+				LastMeshNodeConnected = Cast<UCustomizableObjectNode>(MeshNode);
+				LastMeshNodeConnected->PostEditChangePropertyDelegate.AddUniqueDynamic(this, &UCONodeSkeletalMeshSection::MeshPostEditChangeProperty);
+			}
+		}
+	}
+	else if (Cast<UCustomizableObjectNodeMaterialPinDataImage>(GetPinData(*Pin))) // Image pin
+	{
+		// If necessary, automatically change the Pin Mode. Connected pin can never change its type.
+		if (Pin->LinkedTo.Num())
+		{
+			if (UEdGraphPin* LinkedPin = Pin->LinkedTo[0])
+			{
+				EPinMode PinMode = EPinMode::Default;
+		
+				if (LinkedPin->PinType.PinCategory == UEdGraphSchema_CustomizableObject::PC_Texture)
+				{
+					PinMode = EPinMode::Mutable;
+				}
+				else if (LinkedPin->PinType.PinCategory == UEdGraphSchema_CustomizableObject::PC_Texture_Passthrough || LinkedPin->PinType.PinCategory == UEdGraphSchema_CustomizableObject::PC_Texture_PassthroughParameter)
+				{
+					PinMode = EPinMode::Passthrough;
+				}
+				else
+				{
+					unimplemented()
+				}
+		
+				SetImagePinMode(*Pin, PinMode);
+			}
+		}
+	}
+}
+
+
+void UCONodeSkeletalMeshSection::PostPasteNode()
+{
+	Super::PostPasteNode();
+	SetDefaultMaterial();
+}
+
+
+bool UCONodeSkeletalMeshSection::CanConnect(const UEdGraphPin* InOwnedInputPin, const UEdGraphPin* InOutputPin, bool& bOutIsOtherNodeBlocklisted, bool& bOutArePinsCompatible) const
+{
+	const UEdGraphSchema_CustomizableObject* Schema = GetDefault<UEdGraphSchema_CustomizableObject>();
+
+	if (InOwnedInputPin && InOutputPin)
+	{
+		if ((InOwnedInputPin->PinType.PinCategory == Schema->PC_Texture && (InOutputPin->PinType.PinCategory == Schema->PC_Texture_Passthrough || InOutputPin->PinType.PinCategory == Schema->PC_Texture_PassthroughParameter)) ||
+			(InOwnedInputPin->PinType.PinCategory == Schema->PC_Texture_Passthrough && InOutputPin->PinType.PinCategory == Schema->PC_Texture))
+		{
+			return true;
+		}
+	}
+
+	return Super::CanConnect(InOwnedInputPin, InOutputPin, bOutIsOtherNodeBlocklisted, bOutArePinsCompatible);
+}
+
+
+void UCONodeSkeletalMeshSection::MeshPostEditChangeProperty(FPostEditChangePropertyDelegateParameters& Parameters)
+{
+	if (const UEdGraphPin* MeshPin = FindPin(TEXT("Mesh")))
+	{
+		if (const UEdGraphPin* ConnectedPin = FollowInputPin(*MeshPin);
+			ConnectedPin && ConnectedPin->GetOwningNode() == Parameters.Node)
+		{
+			SetDefaultMaterial();
+		}
+		else if (UCustomizableObjectNode* MeshNode = Cast<UCustomizableObjectNode>(Parameters.Node))
+		{
+			MeshNode->PostEditChangePropertyDelegate.RemoveDynamic(this, &UCONodeSkeletalMeshSection::MeshPostEditChangeProperty);
+		}
+	}
+}
+
+
+#undef LOCTEXT_NAMESPACE
diff --git a/Engine/Plugins/Mutable/Source/CustomizableObjectEditor/Private/MuCOE/Nodes/CONodeSkeletalMeshSection.h b/Engine/Plugins/Mutable/Source/CustomizableObjectEditor/Private/MuCOE/Nodes/CONodeSkeletalMeshSection.h
new file mode 100644
--- /dev/null
+++ b/Engine/Plugins/Mutable/Source/CustomizableObjectEditor/Private/MuCOE/Nodes/CONodeSkeletalMeshSection.h
@@ -0,0 +1,332 @@
+// Copyright Epic Games, Inc. All Rights Reserved.
+
+#pragma once
+
+#include "MuCOE/Nodes/CONodeSkeletalMeshSectionBase.h"
+#include "MuCOE/RemapPins/CustomizableObjectNodeRemapPinsByName.h"
+
+#include "CONodeSkeletalMeshSection.generated.h"
+
+#define UE_API CUSTOMIZABLEOBJECTEDITOR_API
+
+enum class EMaterialParameterType : uint8;
+namespace ENodeTitleType { enum Type : int; }
+
+class FArchive;
+class SGraphNode;
+class SWidget;
+class UCustomizableObjectNodeRemapPins;
+class UEdGraphPin;
+class UMaterial;
+class UMaterialInterface;
+class UObject;
+class UTexture2D;
+struct FCustomizableObjectNodeMaterialImage;
+struct FCustomizableObjectNodeMaterialScalar;
+struct FCustomizableObjectNodeMaterialVector;
+struct FEdGraphPinReference;
+struct FFrame;
+struct FPropertyChangedEvent;
+
+
+/** Custom remap pins by name action.
+ *
+ * Remap pins by Texture Parameter Id. */
+UCLASS()
+class UCustomizableObjectNodeMaterialRemapPinsByName : public UCustomizableObjectNodeRemapPinsByName
+{
+	GENERATED_BODY()
+public:
+	virtual bool Equal(const UCustomizableObjectNode& Node, const UEdGraphPin& OldPin, const UEdGraphPin& NewPin) const override;
+
+	virtual void RemapPins(const UCustomizableObjectNode& Node, const TArray<UEdGraphPin*>& OldPins, const TArray<UEdGraphPin*>& NewPins, TMap<UEdGraphPin*, UEdGraphPin*>& PinsToRemap, TArray<UEdGraphPin*>& PinsToOrphan) override;
+
+	bool HasSavedPinData(const UCustomizableObjectNode& Node, const UEdGraphPin &Pin) const;
+};
+
+
+/** Base class for all Material Parameters. */
+UCLASS(MinimalAPI)
+class UCONodeSkeletalMeshSectionPinDataParameter : public UCustomizableObjectNodePinData
+{
+	GENERATED_BODY()
+
+public:
+
+	/** Parameter id + layer index */
+	UPROPERTY()
+	FNodeMaterialParameterId MaterialParameterId;
+
+	/** Texture Parameter Id. */
+	UPROPERTY()
+	FGuid ParameterId_DEPRECATED;
+
+	/** Returns true if all properties are in its default state. */
+	UE_API virtual bool IsDefault() const;
+};
+
+
+/** Node pin mode. All pins set to EPinMode::Default will use this this mode. */
+UENUM()
+enum class ENodePinMode
+{
+	Mutable UMETA(ToolTip = "All Material Texture FParameters go through Mutable."),
+	Passthrough UMETA(ToolTip = "All Material Texture FParameters are not modified by Mutable.")
+};
+
+
+/** Image pin, pin mode. */
+UENUM()
+enum class EPinMode
+{
+	Default UMETA(DisplayName = "Node Defined", ToolTip = "Use node's \"Default Texture Parameter Mode\"."),
+	Mutable UMETA(ToolTip = "The Material Texture FParameters goes through Mutable."),
+	Passthrough UMETA(ToolTip = "The Material Texture FParameters is not modified by Mutable.")
+};
+
+
+/** Image Pin, UV Layout Mode. */
+UENUM()
+enum class EUVLayoutMode
+{
+	/** Does not override the UV Index specified in the Material. */
+	FromMaterial,
+	/* Texture should not be transformed by any layout. These textures will not be reduced automatically for LODs. */
+	Ignore,
+	/** User specified UV Index. */
+	Index
+};
+
+
+/** Enum to FText. */
+FText EPinModeToText(EPinMode PinMode);
+
+UCLASS(MinimalAPI)
+class UCONodeSkeletalMeshSection : public UCONodeSkeletalMeshSectionBase
+{
+public:
+	GENERATED_BODY()
+	
+	// UObject interface.
+	UE_API virtual void Serialize(FArchive& Ar) override;
+	UE_API virtual void PostEditChangeProperty(FPropertyChangedEvent& PropertyChangedEvent) override;
+
+	// EdGraphNode interface
+	UE_API virtual FText GetNodeTitle(ENodeTitleType::Type TitleType) const override;
+	UE_API virtual FText GetTooltipText() const override;
+	UE_API virtual TSharedPtr<SGraphNode> CreateVisualWidget() override;
+	UE_API virtual void PinConnectionListChanged(UEdGraphPin* Pin) override;
+	UE_API virtual void PostPasteNode() override;
+	UE_API virtual bool CanConnect(const UEdGraphPin* InOwnedInputPin, const UEdGraphPin* InOutputPin, bool& bOutIsOtherNodeBlocklisted, bool& bOutArePinsCompatible) const override;
+
+	// UCustomizableObjectNode interface
+	UE_API virtual void BackwardsCompatibleFixup(int32 CustomizableObjectCustomVersion) override;
+	UE_API virtual void AllocateDefaultPins(UCustomizableObjectNodeRemapPins* RemapPins) override;
+	UE_API virtual bool CanPinBeHidden(const UEdGraphPin& Pin) const override;
+	UE_API virtual bool HasPinViewer() const override;
+	UE_API virtual UCustomizableObjectNodeRemapPinsByName* CreateRemapPinsDefault() const override;
+	virtual bool ProvidesCustomPinRelevancyTest() const override { return true; }
+	UE_API virtual bool IsPinRelevant(const UEdGraphPin* Pin) const override;
+	UE_API virtual bool CustomRemovePin(UEdGraphPin& Pin) override;
+	UE_API virtual bool IsNodeOutDatedAndNeedsRefresh() override;
+	UE_API virtual FString GetRefreshMessage() const override;
+	UE_API virtual TSharedPtr<IDetailsView> CustomizePinDetails(const UEdGraphPin& Pin) const override;
+	UE_API virtual TArray<FString> GetEnableTags(TArray<const UCustomizableObjectNodeMacroInstance*>* MacroContext = nullptr) override;
+	UE_API virtual TArray<FString>* GetEnableTagsArray() override;
+	UE_API virtual FString GetInternalTagDisplayName() override;
+	
+	// UCONodeSkeletalMeshSectionBase interface
+	UE_API virtual TArray<UCustomizableObjectLayout*> GetLayouts() const override;
+	UE_API virtual UEdGraphPin* OutputPin() const override;
+	UE_API virtual UMaterialInterface* GetMaterial() const override;
+	UE_API virtual int32 GetNumParameters(EMaterialParameterType Type) const override;
+	UE_API virtual FNodeMaterialParameterId GetParameterId(EMaterialParameterType Type, int32 ParameterIndex) const override;
+	UE_API virtual FName GetParameterName(EMaterialParameterType Type, int32 ParameterIndex) const override;
+	UE_API virtual int32 GetParameterLayerIndex(EMaterialParameterType Type, int32 ParameterIndex) const override;
+	UE_API virtual FText GetParameterLayerName(EMaterialParameterType Type, int32 ParameterIndex) const override;
+	UE_API virtual bool HasParameter(const FNodeMaterialParameterId& ParameterId) const override;
+	UE_API virtual UEdGraphPin* GetParameterPin(EMaterialParameterType Type, int32 ParameterIndex) const override;
+	UE_API virtual UEdGraphPin* GetParameterPin(const FNodeMaterialParameterId& ParameterId) const override;
+	UE_API virtual bool IsImageMutableMode(int32 ImageIndex) const override;
+	UE_API virtual bool IsImageMutableMode(const UEdGraphPin& Pin) const override;
+	UE_API virtual UTexture2D* GetImageReferenceTexture(int32 ImageIndex) const override;
+	UE_API virtual UTexture2D* GetImageValue(int32 ImageIndex) const override;
+	UE_API virtual int32 GetImageUVLayout(int32 ImageIndex) const override;
+	UE_API virtual UEdGraphPin* GetMeshPin() const override;
+	UE_API virtual UEdGraphPin* GetMaterialAssetPin() const override;
+	UE_API virtual UEdGraphPin* GetEnableTagsPin() const override;
+	UE_API virtual const UCONodeSkeletalMeshSection* GetMaterialNode() const override;
+	UE_API virtual bool RealMaterialDataHasChanged() const override;
+	UE_API virtual FPostImagePinModeChangedDelegate* GetPostImagePinModeChangedDelegate() override;
+	
+	// Own Interface
+
+	UE_API void SetMaterial(UMaterialInterface* InMaterial);
+
+	static UE_API bool HasParameter(const UMaterialInterface* InMaterial, const FNodeMaterialParameterId& ParameterId);
+	static UE_API int32 GetParameterLayerIndex(const UMaterialInterface* InMaterial, EMaterialParameterType Type, int32 ParameterIndex);
+
+	UPROPERTY(EditAnywhere, Category = Tags)
+	TArray<FString> Tags;
+
+private:
+	/** Set the Pin Mode of a Texture Parameter Pin. */
+	UE_API void SetImagePinMode(UEdGraphPin& Pin, EPinMode PinMode) const;
+	
+	/** Delegate called when a Texture Parameter Pin Mode changes. */
+	FPostImagePinModeChangedDelegate PostImagePinModeChangedDelegate;
+
+	UPROPERTY(EditAnywhere, Category=CustomizableObject)
+	TObjectPtr<UMaterialInterface> Material = nullptr;
+
+	UPROPERTY(EditAnywhere, Category=CustomizableObject, DisplayName = "Default Texture Parameter Mode", Meta = (ToolTip = "All Mateiral Texture FParameters set to \"Node Defined\" will use this mode."))
+	ENodePinMode TextureParametersMode = ENodePinMode::Passthrough;
+
+	UPROPERTY()
+	int32 MeshComponentIndex_DEPRECATED = 0;
+
+public:
+	/** Selects which Mesh component of the Instance this material belongs to */
+	UPROPERTY()
+	FName MeshComponentName_DEPRECATED;
+
+private:
+
+	/** Last static or skeletal mesh connected. Used to remove the callback once disconnected. */
+	TWeakObjectPtr<UCustomizableObjectNode> LastMeshNodeConnected;
+
+	/** List of material parameter types that are actually relevant to mutable. */
+	static UE_API const TArray<EMaterialParameterType> ParameterTypes;
+
+	/** Relates a Parameter id (key) (and layer if is a layered material) to a Pin (value). Only used to improve performance.
+	  * If a deprecated pin and a non-deprecated pin have the same Parameter id, this the non-deprecated one prevails. */
+	UPROPERTY()
+	TMap<FNodeMaterialParameterId, FEdGraphPinReference> PinsParameterMap;
+
+	UPROPERTY()
+	FEdGraphPinReference EnableTagsPinRef;
+	
+	/** Create the pin data of the given parameter type. */
+	UE_API UCONodeSkeletalMeshSectionPinDataParameter* CreatePinData(EMaterialParameterType Type, int32 ParameterIndex);
+
+	/** Allocate a pin for each parameter of the given type. */
+	UE_API void AllocateDefaultParameterPins(EMaterialParameterType Type);
+
+	/** Set the default Material from the connected static or skeletal mesh. */
+	UE_API void SetDefaultMaterial();
+
+	/** Connected NodeStaticMesh or NodeSkeletalMesh Mesh UPROPERTY changed callback function. Sets the default material. */
+	UFUNCTION()
+	UE_API void MeshPostEditChangeProperty(FPostEditChangePropertyDelegateParameters& Parameters);
+
+	/** Format pin name. */
+	UE_API FName GetPinName(EMaterialParameterType Type, int32 ParameterIndex) const;
+
+	/** Returns the texture coordinate of the given Material Expression. Returns -1 if not found. */
+	static UE_API int32 GetExpressionTextureCoordinate(UMaterial* Material, const FGuid &ImageId);
+
+	/** Return the Pin Category given the node NodePinMode. */
+	static UE_API FName NodePinModeToImagePinMode(ENodePinMode NodePinMode);
+
+	/** Return the Pin Category given a PinMode. */
+	UE_API FName GetImagePinMode(EPinMode PinMode) const;
+
+	/** Return the Pin Category given a Pin. */
+	UE_API FName GetImagePinMode(const UEdGraphPin& Pin) const;
+
+	/** Get the UV Layout Index defined in the Material. */
+	UE_API int32 GetImageUVLayoutFromMaterial(int32 ImageIndex) const;
+	
+	// Deprecated properties
+	/** Set all pins to Mutable mode. Even so, each pin can override its behaviour. */
+	UPROPERTY()
+	bool bDefaultPinModeMutable_DEPRECATED = false;
+	
+	UPROPERTY()
+	TArray<FCustomizableObjectNodeMaterialImage> Images_DEPRECATED;
+	
+	UPROPERTY()
+	TArray<FCustomizableObjectNodeMaterialVector> VectorParams_DEPRECATED;
+
+	UPROPERTY()
+	TArray<FCustomizableObjectNodeMaterialScalar> ScalarParams_DEPRECATED;
+
+	UPROPERTY()
+	TMap<FGuid, FEdGraphPinReference> PinsParameter_DEPRECATED;
+};
+
+
+/** Additional data for a Material Texture Parameter pin. */
+UCLASS(MinimalAPI)
+class UCustomizableObjectNodeMaterialPinDataImage : public UCONodeSkeletalMeshSectionPinDataParameter
+{
+	GENERATED_BODY()
+
+	friend void UCONodeSkeletalMeshSection::BackwardsCompatibleFixup(int32 CustomizableObjectCustomVersion);
+	
+public:
+	// UObject interface
+	UE_API virtual void PostEditChangeProperty(FPropertyChangedEvent& PropertyChangedEvent) override;
+	UE_API virtual bool CanEditChange(const FProperty* InProperty) const override;
+
+	// UCustomizableObjectNodePinData interface
+	UE_API virtual void BackwardsCompatibleFixup(int32 CustomizableObjectCustomVersion) override;
+
+	// NodePinDataParameter interface
+	/** Virtual function used to copy pin data when remapping pins. */
+	UE_API virtual void Copy(const UCustomizableObjectNodePinData& Other) override;
+	
+	// UCustomizableObjectNodeMaterialPinParameter interface
+	UE_API virtual bool IsDefault() const override;
+
+	// Own interface
+	/** Constructor parameters. Should always be called after a NewObject. */
+	UE_API void Init(UCONodeSkeletalMeshSection& InNodeMaterial);
+
+	UE_API EPinMode GetPinMode() const;
+
+	UE_API void SetPinMode(EPinMode InPinMode);
+	
+private:
+	/** Image pin mode. If is not default, overrides the defined node behaviour. */
+	UPROPERTY(EditAnywhere, Category = NoCategory, DisplayName = "Texture Parameter Mode")
+	EPinMode PinMode = EPinMode::Default;
+
+public:
+	constexpr static int32 UV_LAYOUT_IGNORE = -1;
+
+	UPROPERTY(EditAnywhere, Category = NoCategory)
+	EUVLayoutMode UVLayoutMode = EUVLayoutMode::FromMaterial;
+	
+	/** Index of the UV channel that will be used with this image. It is necessary to apply the proper layout transformations to it. */
+	UPROPERTY(EditAnywhere, Category = NoCategory, meta = (EditCondition = "UVLayoutMode == EUVLayoutMode::Index", EditConditionHides))
+	int32 UVLayout = -2;
+
+	/** Reference Texture used to decide the texture properties of the mutable-generated textures
+	* connected to this material. If null, it will try to be guessed at compile time from
+	* the graph. */
+	UPROPERTY(EditAnywhere, Category = NoCategory) // Required to be EditAnywhere for the selector to work.
+	TObjectPtr<UTexture2D> ReferenceTexture = nullptr;
+
+private:
+	UPROPERTY()
+	TObjectPtr<UCONodeSkeletalMeshSection> NodeMaterial = nullptr;
+};
+
+
+/** Additional data for a Material Vector Parameter pin. */
+UCLASS(MinimalAPI)
+class UCustomizableObjectNodeMaterialPinDataVector : public UCONodeSkeletalMeshSectionPinDataParameter
+{
+	GENERATED_BODY()
+};
+
+
+/** Additional data for a Material Float Parameter pin. */
+UCLASS(MinimalAPI)
+class UCustomizableObjectNodeMaterialPinDataScalar : public UCONodeSkeletalMeshSectionPinDataParameter
+{
+	GENERATED_BODY()
+};
+
+#undef UE_API
diff --git a/Engine/Plugins/Mutable/Source/CustomizableObjectEditor/Private/MuCOE/Nodes/CONodeSkeletalMeshSectionBase.cpp b/Engine/Plugins/Mutable/Source/CustomizableObjectEditor/Private/MuCOE/Nodes/CONodeSkeletalMeshSectionBase.cpp
new file mode 100644
--- /dev/null
+++ b/Engine/Plugins/Mutable/Source/CustomizableObjectEditor/Private/MuCOE/Nodes/CONodeSkeletalMeshSectionBase.cpp
@@ -0,0 +1,17 @@
+// Copyright Epic Games, Inc. All Rights Reserved.
+
+#include "MuCOE/Nodes/CONodeSkeletalMeshSectionBase.h"
+
+#include "MuCOE/EdGraphSchema_CustomizableObject.h"
+
+#include UE_INLINE_GENERATED_CPP_BY_NAME(CONodeSkeletalMeshSectionBase)
+
+int32 UCONodeSkeletalMeshSectionBase::GetMeshSectionMaxLOD() const
+{
+	return MaxLOD;
+}
+
+FLinearColor UCONodeSkeletalMeshSectionBase::GetNodeTitleColor() const
+{
+	return UEdGraphSchema_CustomizableObject::GetPinTypeColor(UEdGraphSchema_CustomizableObject::PC_MeshSection);
+}
diff --git a/Engine/Plugins/Mutable/Source/CustomizableObjectEditor/Private/MuCOE/Nodes/CONodeSkeletalMeshSectionBase.h b/Engine/Plugins/Mutable/Source/CustomizableObjectEditor/Private/MuCOE/Nodes/CONodeSkeletalMeshSectionBase.h
new file mode 100644
--- /dev/null
+++ b/Engine/Plugins/Mutable/Source/CustomizableObjectEditor/Private/MuCOE/Nodes/CONodeSkeletalMeshSectionBase.h
@@ -0,0 +1,168 @@
+// Copyright Epic Games, Inc. All Rights Reserved.
+
+#pragma once
+
+#include "MuCOE/Nodes/CustomizableObjectNode.h"
+
+#include "CONodeSkeletalMeshSectionBase.generated.h"
+
+#define UE_API CUSTOMIZABLEOBJECTEDITOR_API
+
+class UCustomizableObjectLayout;
+class UCONodeSkeletalMeshSection;
+class UEdGraphPin;
+class UObject;
+
+
+DECLARE_MULTICAST_DELEGATE(FPostImagePinModeChangedDelegate)
+
+/** This struct helps us to identify a material parameter using its id and layer index in case of multimaterials
+* When a multilayer material has the same material in multiple layer the parameter Id is not enough to identify a parameter
+* we also need the its layer index.
+*/
+USTRUCT()
+struct FNodeMaterialParameterId
+{
+	GENERATED_BODY()
+
+	UPROPERTY()
+	FGuid ParameterId;
+
+	UPROPERTY()
+	int32 LayerIndex = INDEX_NONE;
+
+	bool operator==(const FNodeMaterialParameterId& Other) const = default;
+};
+
+
+inline uint32 GetTypeHash(const FNodeMaterialParameterId& Key)
+{
+	uint32 Hash = GetTypeHash(Key.ParameterId);
+	Hash = HashCombine(Hash, GetTypeHash(Key.LayerIndex));
+
+	return Hash;
+}
+
+
+/** Equivalent to UE::Mutable::Private::NodeSurface but with limitations. Currently only nodes that generate a UE::Mutable::Private::NodeSurfaceNew inherit this (NodeMaterial and NodeCopyMaterial).
+ * Nodes that generate UE::Mutable::Private::NodeSurfaceEdit (NodeEditMaterial, NodeExtendMaterial), UE::Mutable::Private::NodeSurfaceSwitch (NodeSwitchMaterial)... are excluded.
+ *
+ * WARNING: All methods in this class should be PURE_VIRTUAL!
+ */
+UCLASS(MinimalAPI, abstract)
+class UCONodeSkeletalMeshSectionBase : public UCustomizableObjectNode
+{
+public:
+	GENERATED_BODY()
+	
+	// Own interface
+	virtual const UCONodeSkeletalMeshSection* GetMaterialNode() const PURE_VIRTUAL(UCustomizableObjectNodeMaterial::GetMaterialNode, return {}; );
+
+	UCONodeSkeletalMeshSection* GetMaterialNode()
+	{
+		const UCONodeSkeletalMeshSectionBase* ConstThis = this;
+		return const_cast<UCONodeSkeletalMeshSection*>(ConstThis->GetMaterialNode());
+	}
+
+	// Max LOD to propagate this section to, when using automatic LODs.
+	UE_API int32 GetMeshSectionMaxLOD() const;
+
+	virtual TArray<UCustomizableObjectLayout*> GetLayouts() const PURE_VIRTUAL(UCustomizableObjectNodeMaterial::GetLayouts, return {}; );
+	
+	virtual UMaterialInterface* GetMaterial() const PURE_VIRTUAL(UCustomizableObjectNodeMaterial::GetMaterial, return {}; );
+	
+	virtual UEdGraphPin* GetMeshPin() const PURE_VIRTUAL(UCustomizableObjectNodeMaterial::GetMeshPin, return {}; );
+
+	virtual UEdGraphPin* GetMaterialAssetPin() const PURE_VIRTUAL(UCustomizableObjectNodeMaterial::GetMaterialAssetPin, return {}; );
+
+	virtual UEdGraphPin* GetEnableTagsPin() const PURE_VIRTUAL(UCustomizableObjectNodeMaterial::GetMaterialAssetPin, return {}; );
+	
+	/** Returns the number of Material Parameters. */
+	virtual int32 GetNumParameters(EMaterialParameterType Type) const PURE_VIRTUAL(UCustomizableObjectNodeMaterial::GetNumParameters, return {}; );
+
+	/** Returns the Material Parameter id.
+	 *
+	 * @param ParameterIndex Have to be valid. */
+	virtual FNodeMaterialParameterId GetParameterId(EMaterialParameterType Type, int32 ParameterIndex) const PURE_VIRTUAL(UCustomizableObjectNodeMaterial::GetParameterId, return {}; );
+	
+	/** Returns the Material Parameter name.
+	 *
+	 * @param ParameterIndex Have to be valid. */
+	virtual FName GetParameterName(EMaterialParameterType Type, int32 ParameterIndex) const PURE_VIRTUAL(UCustomizableObjectNodeMaterial::GetParameterName, return {}; );
+	
+	/** Get the Material Parameter layer index.
+	 *
+	 * @param ParameterIndex Have to be valid.
+	 * @returns INDEX_NONE for global parameters. */
+	virtual int32 GetParameterLayerIndex(EMaterialParameterType Type, int32 ParameterIndex) const PURE_VIRTUAL(UCustomizableObjectNodeMaterial::GetParameterLayerIndex, return {}; );
+	
+	/** Get the Material Parameter layer name.
+	 *
+	 * @param ParameterIndex Have to be valid. */
+	virtual FText GetParameterLayerName(EMaterialParameterType Type, int32 ParameterIndex) const PURE_VIRTUAL(UCustomizableObjectNodeMaterial::GetParameterLayerName, return {}; );
+
+	/** Returns true if the Material contains the given Material Parameter. */
+	virtual bool HasParameter(const FNodeMaterialParameterId& ParameterId) const PURE_VIRTUAL(UCustomizableObjectNodeMaterial::HasParameter, return {}; );
+	
+	/** Get the pin for the given Material Parameter.
+	 * Not all parameters have pins.
+	 *
+	 * @param ParameterIndex Material Parameter id.
+	 * @return Can return nullptr. */
+	virtual const UEdGraphPin* GetParameterPin(EMaterialParameterType Type, int32 ParameterIndex) const PURE_VIRTUAL(UCustomizableObjectNodeMaterial::GetMaterialNode, return {}; );
+
+	/** Get the pin for the given Material Parameter.
+	 * Not all parameters have pins.
+	 *
+	 * @param ParameterId Material Parameter id.
+	 * @return Can return nullptr. */
+	virtual UEdGraphPin* GetParameterPin(const FNodeMaterialParameterId& ParameterId) const PURE_VIRTUAL(UCustomizableObjectNodeMaterial::GetParameterPin, return {}; );
+	
+	// --------------------
+	// IMAGES PARAMETERS
+	// --------------------
+
+	/** Returns true if the Material Texture Parameter goes through Mutable.
+	 *
+	 * @param ImageIndex Have to be valid. */
+	virtual bool IsImageMutableMode(int32 ImageIndex) const PURE_VIRTUAL(UCustomizableObjectNodeMaterial::IsImageMutableMode, return {}; );
+
+	/** Given an Image pin, returns true if the Material Texture Parameter goes through Mutable. */
+	virtual bool IsImageMutableMode(const UEdGraphPin& Pin) const PURE_VIRTUAL(UCustomizableObjectNodeMaterial::IsImageMutableMode, return {}; );
+	
+	/** Returns the reference texture assigned to a Material Texture Parameter.
+	 *
+	 * @param ImageIndex Have to be valid.
+	 * @return nullptr if it does not have one assigned. */
+	virtual UTexture2D* GetImageReferenceTexture(int32 ImageIndex) const PURE_VIRTUAL(UCustomizableObjectNodeMaterial::GetImageReferenceTexture, return {}; );
+
+	/** Returns the Texture set in the Material Texture Parameter.
+	 *
+	 * @param ImageIndex Have to be valid. */
+	virtual UTexture2D* GetImageValue(int32 ImageIndex) const PURE_VIRTUAL(UCustomizableObjectNodeMaterial::GetImageValue, return {}; );
+	
+	/** Get the Material Texture Parameter UV Index.
+	 *
+	 * @param ImageIndex Have to be valid.
+	 * @return Return -1 if the UV Index is set to Ignore. Return >= 0 for a valid UV Index.  */ 
+	virtual int32 GetImageUVLayout(int32 ImageIndex) const PURE_VIRTUAL(UCustomizableObjectNodeMaterial::GetImageUVLayout, return {}; );
+
+	// This method should be overridden in all derived classes
+	virtual UEdGraphPin* OutputPin() const PURE_VIRTUAL(UCustomizableObjectNodeMaterial::OutputPin, return {}; );
+	
+	virtual FPostImagePinModeChangedDelegate* GetPostImagePinModeChangedDelegate() PURE_VIRTUAL(UCustomizableObjectNodeMaterial::OutputPin, return {}; );
+	
+	/** Return true if a Material Parameter has changed on which we had a pin connected or data saved. */
+	virtual bool RealMaterialDataHasChanged() const PURE_VIRTUAL(UCustomizableObjectNodeMaterial::OutputPin, return {}; );
+
+	// UEdGraphNode interface
+	UE_API virtual FLinearColor GetNodeTitleColor() const override;
+
+private:
+	
+	/* Max LOD to propagate this section to, when using automatic LODs. */
+	UPROPERTY(EditAnywhere, Category = AdvancedSettings)
+	int32 MaxLOD = INDEX_NONE;
+};
+
+#undef UE_API
diff --git a/Engine/Plugins/Mutable/Source/CustomizableObjectEditor/Private/MuCOE/Nodes/CONodeSkeletalMeshSectionInterface.h b/Engine/Plugins/Mutable/Source/CustomizableObjectEditor/Private/MuCOE/Nodes/CONodeSkeletalMeshSectionInterface.h
new file mode 100644
--- /dev/null
+++ b/Engine/Plugins/Mutable/Source/CustomizableObjectEditor/Private/MuCOE/Nodes/CONodeSkeletalMeshSectionInterface.h
@@ -0,0 +1,31 @@
+// Copyright Epic Games, Inc. All Rights Reserved.
+
+#pragma once
+
+#include "MuCOE/Nodes/CustomizableObjectNode.h"
+
+#include "CONodeSkeletalMeshSectionInterface.generated.h"
+
+class UEdGraphPin;
+class UMaterialInterface;
+
+UINTERFACE(MinimalAPI)
+class UCONodeSkeletalMeshSectionInterface : public UInterface
+{
+	GENERATED_BODY()
+};
+
+
+class ICONodeSkeletalMeshSectionInterface
+{
+	GENERATED_BODY()
+
+public:
+	// Own interface
+
+	/** Returns the Unreal mesh (e.g., USkeletalMesh, UStaticMesh...). */
+	virtual TSoftObjectPtr<UMaterialInterface> GetMaterial() const PURE_VIRTUAL(UCustomizableObjectNodeMaterialInterface::GetMaterial, return {};);
+
+	/** Returns the output Mesh pin associated to the given LODIndex and SectionIndex. Override. */
+	virtual UEdGraphPin* GetMaterialPin() const PURE_VIRTUAL(UCustomizableObjectNodeMaterialInterface::GetMaterialPin, return {};);
+};
diff --git a/Engine/Plugins/Mutable/Source/CustomizableObjectEditor/Private/MuCOE/Nodes/CONodeSkeletalMeshSectionSwitch.cpp b/Engine/Plugins/Mutable/Source/CustomizableObjectEditor/Private/MuCOE/Nodes/CONodeSkeletalMeshSectionSwitch.cpp
new file mode 100644
--- /dev/null
+++ b/Engine/Plugins/Mutable/Source/CustomizableObjectEditor/Private/MuCOE/Nodes/CONodeSkeletalMeshSectionSwitch.cpp
@@ -0,0 +1,84 @@
+// Copyright Epic Games, Inc. All Rights Reserved.
+
+#include "MuCOE/Nodes/CONodeSkeletalMeshSectionSwitch.h"
+
+#include "MuCOE/EdGraphSchema_CustomizableObject.h"
+#include "MuCO/CustomizableObjectCustomVersion.h"
+
+#include UE_INLINE_GENERATED_CPP_BY_NAME(CONodeSkeletalMeshSectionSwitch)
+
+#define LOCTEXT_NAMESPACE "CustomizableObjectEditor"
+
+
+FName UCONodeSkeletalMeshSectionSwitch::GetCategory() const
+{
+	return UEdGraphSchema_CustomizableObject::PC_MeshSection;
+}
+
+void UCONodeSkeletalMeshSectionSwitch::BackwardsCompatibleFixup(int32 CustomizableObjectCustomVersion)
+{
+	Super::BackwardsCompatibleFixup(CustomizableObjectCustomVersion);
+
+	if (CustomizableObjectCustomVersion == FCustomizableObjectCustomVersion::FixMaterialPinsRename)
+	{
+		const UEdGraphSchema_CustomizableObject* Schema = GetDefault<UEdGraphSchema_CustomizableObject>();
+
+		// if there are pins that were manually fixed and re-saved after the FCustomizableObjectCustomVersion::MaterialPinsRename
+		// conversion, convert them to the new fixed version.
+
+		// In the new fix we only change the friendy name, but the actual name remains "material..."
+
+		// Output
+		UEdGraphPin* OldPin = FindPin(TEXT("Mesh Section"),EEdGraphPinDirection::EGPD_Output);
+		if (OldPin)
+		{
+			UEdGraphPin* NewPin = FindPin(TEXT("Material"), EEdGraphPinDirection::EGPD_Output);
+			if (!NewPin)
+			{
+				NewPin = CustomCreatePin(EGPD_Output, Schema->PC_MeshSection, TEXT("Material"), LOCTEXT("Material", "Material"));
+			}
+
+			FGuid PinId = NewPin->PinId;
+			NewPin->CopyPersistentDataFromOldPin(*OldPin);
+			NewPin->PinId = PinId;
+			NewPin->bHidden = OldPin->bHidden;
+			
+			CustomRemovePin(*OldPin);
+		}
+
+		// Inputs
+		int32 InputCount = -1;
+		for (UEdGraphPin* Pin : GetAllNonOrphanPins())
+		{
+			if (Pin->Direction == EEdGraphPinDirection::EGPD_Input)
+			{
+				InputCount++;
+			}
+		}
+		
+		for (int32 InputIndex=0; InputIndex < InputCount; ++InputIndex)
+		{
+			OldPin = FindPin(FString::Printf(TEXT("Mesh Section %d "),InputIndex), EEdGraphPinDirection::EGPD_Input);
+			if (OldPin)
+			{
+				FString NewName = FString::Printf(TEXT("Material %d "), InputIndex);
+				UEdGraphPin* NewPin = FindPin(NewName, EEdGraphPinDirection::EGPD_Input);
+				if (!NewPin)
+				{
+					NewPin = CustomCreatePin(EEdGraphPinDirection::EGPD_Input, Schema->PC_MeshSection, FName(*NewName), FText::FromString(NewName));
+				}
+
+				FGuid PinId = NewPin->PinId;
+				NewPin->CopyPersistentDataFromOldPin(*OldPin);
+				NewPin->PinId = PinId;
+				NewPin->bHidden = OldPin->bHidden;
+
+				CustomRemovePin(*OldPin);
+			}
+		}
+	}
+}
+
+
+#undef LOCTEXT_NAMESPACE
+
diff --git a/Engine/Plugins/Mutable/Source/CustomizableObjectEditor/Private/MuCOE/Nodes/CONodeSkeletalMeshSectionSwitch.h b/Engine/Plugins/Mutable/Source/CustomizableObjectEditor/Private/MuCOE/Nodes/CONodeSkeletalMeshSectionSwitch.h
new file mode 100644
--- /dev/null
+++ b/Engine/Plugins/Mutable/Source/CustomizableObjectEditor/Private/MuCOE/Nodes/CONodeSkeletalMeshSectionSwitch.h
@@ -0,0 +1,25 @@
+// Copyright Epic Games, Inc. All Rights Reserved.
+
+#pragma once
+
+#include "MuCOE/Nodes/CustomizableObjectNodeSwitchBase.h"
+
+#include "CONodeSkeletalMeshSectionSwitch.generated.h"
+
+#define UE_API CUSTOMIZABLEOBJECTEDITOR_API
+
+
+UCLASS(MinimalAPI)
+class UCONodeSkeletalMeshSectionSwitch : public UCustomizableObjectNodeSwitchBase
+{
+public:
+	GENERATED_BODY()
+
+	// UCustomizableObjectNode interface
+	UE_API virtual void BackwardsCompatibleFixup(int32 CustomizableObjectCustomVersion) override;
+
+	// UCustomizableObjectNodeSwitchBase interface
+	UE_API virtual FName GetCategory() const override;
+};
+
+#undef UE_API
diff --git a/Engine/Plugins/Mutable/Source/CustomizableObjectEditor/Private/MuCOE/Nodes/CONodeSkeletalMeshSectionVariation.cpp b/Engine/Plugins/Mutable/Source/CustomizableObjectEditor/Private/MuCOE/Nodes/CONodeSkeletalMeshSectionVariation.cpp
new file mode 100644
--- /dev/null
+++ b/Engine/Plugins/Mutable/Source/CustomizableObjectEditor/Private/MuCOE/Nodes/CONodeSkeletalMeshSectionVariation.cpp
@@ -0,0 +1,44 @@
+// Copyright Epic Games, Inc. All Rights Reserved.
+
+#include "MuCOE/Nodes/CONodeSkeletalMeshSectionVariation.h"
+
+#include "MuCO/CustomizableObjectCustomVersion.h"
+#include "MuCOE/CustomizableObjectEditor_Deprecated.h"
+#include "MuCOE/EdGraphSchema_CustomizableObject.h"
+
+#include UE_INLINE_GENERATED_CPP_BY_NAME(CONodeSkeletalMeshSectionVariation)
+
+
+void UCONodeSkeletalMeshSectionVariation::BackwardsCompatibleFixup(int32 CustomizableObjectCustomVersion)
+{
+	Super::BackwardsCompatibleFixup(CustomizableObjectCustomVersion);
+
+	if (CustomizableObjectCustomVersion == FCustomizableObjectCustomVersion::NodeVariationSerializationIssue)
+	{
+		for (const FCustomizableObjectMaterialVariation& OldVariation : Variations_DEPRECATED)
+		{
+			FCustomizableObjectVariation Variation;
+			Variation.Tag = OldVariation.Tag;
+			
+			VariationsData.Add(Variation);
+		}
+	}
+}
+
+
+FName UCONodeSkeletalMeshSectionVariation::GetCategory() const
+{
+	return UEdGraphSchema_CustomizableObject::PC_MeshSection;
+}
+
+
+bool UCONodeSkeletalMeshSectionVariation::IsInputPinArray() const
+{
+	return true;
+}
+
+
+bool UCONodeSkeletalMeshSectionVariation::IsSingleOutputNode() const
+{
+	return true;
+}
diff --git a/Engine/Plugins/Mutable/Source/CustomizableObjectEditor/Private/MuCOE/Nodes/CONodeSkeletalMeshSectionVariation.h b/Engine/Plugins/Mutable/Source/CustomizableObjectEditor/Private/MuCOE/Nodes/CONodeSkeletalMeshSectionVariation.h
new file mode 100644
--- /dev/null
+++ b/Engine/Plugins/Mutable/Source/CustomizableObjectEditor/Private/MuCOE/Nodes/CONodeSkeletalMeshSectionVariation.h
@@ -0,0 +1,46 @@
+// Copyright Epic Games, Inc. All Rights Reserved.
+
+#pragma once
+
+#include "MuCOE/Nodes/CustomizableObjectNodeVariation.h"
+
+#include "CONodeSkeletalMeshSectionVariation.generated.h"
+
+#define UE_API CUSTOMIZABLEOBJECTEDITOR_API
+
+struct FCustomizableObjectMaterialVariation;
+
+
+UENUM(BlueprintType)
+enum class ECustomizableObjectNodeMaterialVariationType : uint8
+{
+	Tag 		UMETA(DisplayName = "Tag"),
+	State 		UMETA(DisplayName = "State"),
+};
+
+
+UCLASS(MinimalAPI)
+class UCONodeSkeletalMeshSectionVariation : public UCustomizableObjectNodeVariation
+{
+public:
+	GENERATED_BODY()
+
+	UPROPERTY(EditAnywhere, Category = CustomizableObject)
+	ECustomizableObjectNodeMaterialVariationType Type = ECustomizableObjectNodeMaterialVariationType::Tag;
+
+private:
+	// Deprecated properties
+	UPROPERTY()
+	TArray<FCustomizableObjectMaterialVariation> Variations_DEPRECATED;
+
+public:
+	// UCustomizableObjectNode interface
+	UE_API virtual void BackwardsCompatibleFixup(int32 CustomizableObjectCustomVersion) override;
+	UE_API virtual bool IsSingleOutputNode() const override;
+	
+	// UCustomizableObjectNodeVariation interface
+	UE_API virtual FName GetCategory() const override;
+	UE_API virtual bool IsInputPinArray() const override;
+};
+
+#undef UE_API
diff --git a/Engine/Plugins/Mutable/Source/CustomizableObjectEditor/Private/MuCOE/Nodes/CustomizableObjectNodeCopyMaterial.cpp b/Engine/Plugins/Mutable/Source/CustomizableObjectEditor/Private/MuCOE/Nodes/CustomizableObjectNodeCopyMaterial.cpp
deleted file mode 100644
diff --git a/Engine/Plugins/Mutable/Source/CustomizableObjectEditor/Private/MuCOE/Nodes/CustomizableObjectNodeCopyMaterial.h b/Engine/Plugins/Mutable/Source/CustomizableObjectEditor/Private/MuCOE/Nodes/CustomizableObjectNodeCopyMaterial.h
deleted file mode 100644
diff --git a/Engine/Plugins/Mutable/Source/CustomizableObjectEditor/Private/MuCOE/Nodes/CustomizableObjectNodeMaterial.cpp b/Engine/Plugins/Mutable/Source/CustomizableObjectEditor/Private/MuCOE/Nodes/CustomizableObjectNodeMaterial.cpp
deleted file mode 100644
diff --git a/Engine/Plugins/Mutable/Source/CustomizableObjectEditor/Private/MuCOE/Nodes/CustomizableObjectNodeMaterial.h b/Engine/Plugins/Mutable/Source/CustomizableObjectEditor/Private/MuCOE/Nodes/CustomizableObjectNodeMaterial.h
deleted file mode 100644
diff --git a/Engine/Plugins/Mutable/Source/CustomizableObjectEditor/Private/MuCOE/Nodes/CustomizableObjectNodeMaterialBase.cpp b/Engine/Plugins/Mutable/Source/CustomizableObjectEditor/Private/MuCOE/Nodes/CustomizableObjectNodeMaterialBase.cpp
deleted file mode 100644
diff --git a/Engine/Plugins/Mutable/Source/CustomizableObjectEditor/Private/MuCOE/Nodes/CustomizableObjectNodeMaterialBase.h b/Engine/Plugins/Mutable/Source/CustomizableObjectEditor/Private/MuCOE/Nodes/CustomizableObjectNodeMaterialBase.h
deleted file mode 100644
diff --git a/Engine/Plugins/Mutable/Source/CustomizableObjectEditor/Private/MuCOE/Nodes/CustomizableObjectNodeMaterialInterface.h b/Engine/Plugins/Mutable/Source/CustomizableObjectEditor/Private/MuCOE/Nodes/CustomizableObjectNodeMaterialInterface.h
deleted file mode 100644
diff --git a/Engine/Plugins/Mutable/Source/CustomizableObjectEditor/Private/MuCOE/Nodes/CustomizableObjectNodeMaterialParameter.h b/Engine/Plugins/Mutable/Source/CustomizableObjectEditor/Private/MuCOE/Nodes/CustomizableObjectNodeMaterialParameter.h
--- a/Engine/Plugins/Mutable/Source/CustomizableObjectEditor/Private/MuCOE/Nodes/CustomizableObjectNodeMaterialParameter.h
+++ b/Engine/Plugins/Mutable/Source/CustomizableObjectEditor/Private/MuCOE/Nodes/CustomizableObjectNodeMaterialParameter.h
@@ -3,7 +3,7 @@
 #pragma once
 
 #include "MuCOE/Nodes/CustomizableObjectNodeParameter.h"
-#include "MuCOE/Nodes/CustomizableObjectNodeMaterialInterface.h"
+#include "MuCOE/Nodes/CONodeSkeletalMeshSectionInterface.h"
 
 #include "CustomizableObjectNodeMaterialParameter.generated.h"
 
@@ -15,7 +15,7 @@ class UMaterialInterface;
 
 
 UCLASS(MinimalAPI)
-class UCustomizableObjectNodeMaterialParameter : public UCustomizableObjectNodeParameter, public ICustomizableObjectNodeMaterialInterface
+class UCustomizableObjectNodeMaterialParameter : public UCustomizableObjectNodeParameter, public ICONodeSkeletalMeshSectionInterface
 {
 public:
 	GENERATED_BODY()
@@ -33,7 +33,7 @@ public:
 	// CustomizableObjectNodeParameter interface
 	UE_API virtual FName GetCategory() const override;
 
-	// ICustomizableObjectNodeMaterialInterface interface
+	// ICONodeSkeletalMeshSectionInterface interface
 	UE_API virtual TSoftObjectPtr<UMaterialInterface> GetMaterial() const override;
 	UE_API virtual UEdGraphPin* GetMaterialPin() const override;
 };
diff --git a/Engine/Plugins/Mutable/Source/CustomizableObjectEditor/Private/MuCOE/Nodes/CustomizableObjectNodeMaterialSwitch.cpp b/Engine/Plugins/Mutable/Source/CustomizableObjectEditor/Private/MuCOE/Nodes/CustomizableObjectNodeMaterialSwitch.cpp
deleted file mode 100644
diff --git a/Engine/Plugins/Mutable/Source/CustomizableObjectEditor/Private/MuCOE/Nodes/CustomizableObjectNodeMaterialSwitch.h b/Engine/Plugins/Mutable/Source/CustomizableObjectEditor/Private/MuCOE/Nodes/CustomizableObjectNodeMaterialSwitch.h
deleted file mode 100644
diff --git a/Engine/Plugins/Mutable/Source/CustomizableObjectEditor/Private/MuCOE/Nodes/CustomizableObjectNodeMaterialVariation.cpp b/Engine/Plugins/Mutable/Source/CustomizableObjectEditor/Private/MuCOE/Nodes/CustomizableObjectNodeMaterialVariation.cpp
deleted file mode 100644
diff --git a/Engine/Plugins/Mutable/Source/CustomizableObjectEditor/Private/MuCOE/Nodes/CustomizableObjectNodeMaterialVariation.h b/Engine/Plugins/Mutable/Source/CustomizableObjectEditor/Private/MuCOE/Nodes/CustomizableObjectNodeMaterialVariation.h
deleted file mode 100644
diff --git a/Engine/Plugins/Mutable/Source/CustomizableObjectEditor/Private/MuCOE/Nodes/CustomizableObjectNodeMeshSectionDetails.cpp b/Engine/Plugins/Mutable/Source/CustomizableObjectEditor/Private/MuCOE/Nodes/CustomizableObjectNodeMeshSectionDetails.cpp
--- a/Engine/Plugins/Mutable/Source/CustomizableObjectEditor/Private/MuCOE/Nodes/CustomizableObjectNodeMeshSectionDetails.cpp
+++ b/Engine/Plugins/Mutable/Source/CustomizableObjectEditor/Private/MuCOE/Nodes/CustomizableObjectNodeMeshSectionDetails.cpp
@@ -5,7 +5,7 @@
 #include "MuCOE/CustomizableObjectEditorUtilities.h"
 #include "MuCOE/CustomizableObjectEditorStyle.h"
 #include "MuCOE/GraphTraversal.h"
-#include "MuCOE/Nodes/CustomizableObjectNodeMaterial.h"
+#include "MuCOE/Nodes/CONodeSkeletalMeshSection.h"
 #include "MuCOE/Nodes/CustomizableObjectNodeSkeletalMesh.h"
 #include "MuCOE/UnrealEditorPortabilityHelpers.h"
 #include "PropertyCustomizationHelpers.h"
@@ -34,7 +34,7 @@ void FCustomizableObjectNodeMeshSectionDetails::CustomizeDetails( IDetailLayoutB
 	TSharedPtr<const IDetailsView> DetailsView = DetailBuilder.GetDetailsViewSharedPtr();
 	if (DetailsView && DetailsView->GetSelectedObjects().Num())
 	{
-		Node = Cast<UCustomizableObjectNodeMaterial>( DetailsView->GetSelectedObjects()[0].Get() );
+		Node = Cast<UCONodeSkeletalMeshSection>( DetailsView->GetSelectedObjects()[0].Get() );
 	}
 
 	if (!Node)
@@ -48,7 +48,7 @@ void FCustomizableObjectNodeMeshSectionDetails::CustomizeDetails( IDetailLayoutB
 
 	// Add the required tags widget
 	{
-		TagsPropertyHandle = DetailBuilder.GetProperty(GET_MEMBER_NAME_CHECKED(UCustomizableObjectNodeMaterial, Tags), UCustomizableObjectNodeMaterial::StaticClass());
+		TagsPropertyHandle = DetailBuilder.GetProperty(GET_MEMBER_NAME_CHECKED(UCONodeSkeletalMeshSection, Tags), UCONodeSkeletalMeshSection::StaticClass());
 		DetailBuilder.HideProperty(TagsPropertyHandle);
 
 		TagsPropertyHandle->SetOnPropertyValueChanged(FSimpleDelegate::CreateSP(this, &FCustomizableObjectNodeMeshSectionDetails::OnEnableTagsPropertyChanged));
diff --git a/Engine/Plugins/Mutable/Source/CustomizableObjectEditor/Private/MuCOE/Nodes/CustomizableObjectNodeMeshSectionDetails.h b/Engine/Plugins/Mutable/Source/CustomizableObjectEditor/Private/MuCOE/Nodes/CustomizableObjectNodeMeshSectionDetails.h
--- a/Engine/Plugins/Mutable/Source/CustomizableObjectEditor/Private/MuCOE/Nodes/CustomizableObjectNodeMeshSectionDetails.h
+++ b/Engine/Plugins/Mutable/Source/CustomizableObjectEditor/Private/MuCOE/Nodes/CustomizableObjectNodeMeshSectionDetails.h
@@ -45,7 +45,7 @@ private:
 
 private:
 
-	class UCustomizableObjectNodeMaterial* Node = nullptr;
+	class UCONodeSkeletalMeshSection* Node = nullptr;
 
 	/** */
 	TSharedPtr<IPropertyHandle> TagsPropertyHandle;
diff --git a/Engine/Plugins/Mutable/Source/CustomizableObjectEditor/Private/MuCOE/Nodes/CustomizableObjectNodeModifierBaseDetails.cpp b/Engine/Plugins/Mutable/Source/CustomizableObjectEditor/Private/MuCOE/Nodes/CustomizableObjectNodeModifierBaseDetails.cpp
--- a/Engine/Plugins/Mutable/Source/CustomizableObjectEditor/Private/MuCOE/Nodes/CustomizableObjectNodeModifierBaseDetails.cpp
+++ b/Engine/Plugins/Mutable/Source/CustomizableObjectEditor/Private/MuCOE/Nodes/CustomizableObjectNodeModifierBaseDetails.cpp
@@ -6,7 +6,7 @@
 #include "MuCOE/CustomizableObjectEditorStyle.h"
 #include "MuCOE/GraphTraversal.h"
 #include "MuCOE/Nodes/CustomizableObjectNodeModifierBase.h"
-#include "MuCOE/Nodes/CustomizableObjectNodeMaterial.h"
+#include "MuCOE/Nodes/CONodeSkeletalMeshSection.h"
 #include "MuCOE/Nodes/CustomizableObjectNodeSkeletalMesh.h"
 #include "MuCOE/UnrealEditorPortabilityHelpers.h"
 #include "PropertyCustomizationHelpers.h"
diff --git a/Engine/Plugins/Mutable/Source/CustomizableObjectEditor/Private/MuCOE/Nodes/CustomizableObjectNodeModifierClipWithMeshDetails.cpp b/Engine/Plugins/Mutable/Source/CustomizableObjectEditor/Private/MuCOE/Nodes/CustomizableObjectNodeModifierClipWithMeshDetails.cpp
--- a/Engine/Plugins/Mutable/Source/CustomizableObjectEditor/Private/MuCOE/Nodes/CustomizableObjectNodeModifierClipWithMeshDetails.cpp
+++ b/Engine/Plugins/Mutable/Source/CustomizableObjectEditor/Private/MuCOE/Nodes/CustomizableObjectNodeModifierClipWithMeshDetails.cpp
@@ -7,7 +7,7 @@
 #include "MuCOE/CustomizableObjectEditorUtilities.h"
 #include "MuCOE/GraphTraversal.h"
 #include "MuCOE/UnrealEditorPortabilityHelpers.h"
-#include "MuCOE/Nodes/CustomizableObjectNodeMaterial.h"
+#include "MuCOE/Nodes/CONodeSkeletalMeshSection.h"
 #include "MuCOE/Nodes/CustomizableObjectNodeModifierClipWithMesh.h"
 #include "MuCOE/Nodes/CustomizableObjectNodeStaticMesh.h"
 #include "PropertyCustomizationHelpers.h"
diff --git a/Engine/Plugins/Mutable/Source/CustomizableObjectEditor/Private/MuCOE/Nodes/CustomizableObjectNodeModifierEditMeshSection.cpp b/Engine/Plugins/Mutable/Source/CustomizableObjectEditor/Private/MuCOE/Nodes/CustomizableObjectNodeModifierEditMeshSection.cpp
deleted file mode 100644
diff --git a/Engine/Plugins/Mutable/Source/CustomizableObjectEditor/Private/MuCOE/Nodes/CustomizableObjectNodeModifierEditMeshSection.h b/Engine/Plugins/Mutable/Source/CustomizableObjectEditor/Private/MuCOE/Nodes/CustomizableObjectNodeModifierEditMeshSection.h
deleted file mode 100644
diff --git a/Engine/Plugins/Mutable/Source/CustomizableObjectEditor/Private/MuCOE/Nodes/CustomizableObjectNodeModifierEditMeshSectionBase.cpp b/Engine/Plugins/Mutable/Source/CustomizableObjectEditor/Private/MuCOE/Nodes/CustomizableObjectNodeModifierEditMeshSectionBase.cpp
--- a/Engine/Plugins/Mutable/Source/CustomizableObjectEditor/Private/MuCOE/Nodes/CustomizableObjectNodeModifierEditMeshSectionBase.cpp
+++ b/Engine/Plugins/Mutable/Source/CustomizableObjectEditor/Private/MuCOE/Nodes/CustomizableObjectNodeModifierEditMeshSectionBase.cpp
@@ -2,8 +2,8 @@
 
 #include "MuCOE/Nodes/CustomizableObjectNodeModifierEditMeshSectionBase.h"
 
-#include "MuCOE/Nodes/CustomizableObjectNodeMaterial.h"
-#include "MuCOE/Nodes/CustomizableObjectNodeModifierExtendMeshSection.h"
+#include "MuCOE/Nodes/CONodeSkeletalMeshSection.h"
+#include "MuCOE/Nodes/CONodeModifierExtendSkeletalMeshSection.h"
 #include "MuCOE/CustomizableObjectEditor_Deprecated.h"
 #include "MuCOE/CustomizableObjectLayout.h"
 #include "MuCOE/EdGraphSchema_CustomizableObject.h"
@@ -80,12 +80,12 @@ UCustomizableObjectLayout* UCustomizableObjectNodeModifierEditMeshSectionBase::G
 	{
 		TArray<UCustomizableObjectLayout*> AllLayouts;
 
-		if (UCustomizableObjectNodeMaterialBase* MaterialNode = Cast<UCustomizableObjectNodeMaterialBase>(Candidate))
+		if (UCONodeSkeletalMeshSectionBase* MaterialNode = Cast<UCONodeSkeletalMeshSectionBase>(Candidate))
 		{
 			AllLayouts = MaterialNode->GetLayouts();
 		}
 
-		else if (UCustomizableObjectNodeModifierExtendMeshSection* ExtendNode = Cast<UCustomizableObjectNodeModifierExtendMeshSection>(Candidate))
+		else if (UCONodeModifierExtendSkeletalMeshSection* ExtendNode = Cast<UCONodeModifierExtendSkeletalMeshSection>(Candidate))
 		{
 			AllLayouts = ExtendNode->GetLayouts();
 		}
diff --git a/Engine/Plugins/Mutable/Source/CustomizableObjectEditor/Private/MuCOE/Nodes/CustomizableObjectNodeModifierEditMeshSectionBase.h b/Engine/Plugins/Mutable/Source/CustomizableObjectEditor/Private/MuCOE/Nodes/CustomizableObjectNodeModifierEditMeshSectionBase.h
--- a/Engine/Plugins/Mutable/Source/CustomizableObjectEditor/Private/MuCOE/Nodes/CustomizableObjectNodeModifierEditMeshSectionBase.h
+++ b/Engine/Plugins/Mutable/Source/CustomizableObjectEditor/Private/MuCOE/Nodes/CustomizableObjectNodeModifierEditMeshSectionBase.h
@@ -2,7 +2,7 @@
 
 #pragma once
 
-#include "MuCOE/Nodes/CustomizableObjectNodeModifierWithMaterial.h"
+#include "MuCOE/Nodes/CONodeModifierSkeletalMeshSection.h"
 
 #include "CustomizableObjectNodeModifierEditMeshSectionBase.generated.h"
 
@@ -13,7 +13,7 @@ class UEdGraphPin;
 
 
 UCLASS(MinimalAPI, Abstract)
-class UCustomizableObjectNodeModifierEditMeshSectionBase : public UCustomizableObjectNodeModifierWithMaterial
+class UCustomizableObjectNodeModifierEditMeshSectionBase : public UCONodeModifierSkeletalMeshSection
 {
 	GENERATED_BODY()
 	
diff --git a/Engine/Plugins/Mutable/Source/CustomizableObjectEditor/Private/MuCOE/Nodes/CustomizableObjectNodeModifierEditMeshSectionDetails.cpp b/Engine/Plugins/Mutable/Source/CustomizableObjectEditor/Private/MuCOE/Nodes/CustomizableObjectNodeModifierEditMeshSectionDetails.cpp
--- a/Engine/Plugins/Mutable/Source/CustomizableObjectEditor/Private/MuCOE/Nodes/CustomizableObjectNodeModifierEditMeshSectionDetails.cpp
+++ b/Engine/Plugins/Mutable/Source/CustomizableObjectEditor/Private/MuCOE/Nodes/CustomizableObjectNodeModifierEditMeshSectionDetails.cpp
@@ -6,9 +6,9 @@
 #include "DetailWidgetRow.h"
 #include "IDetailGroup.h"
 #include "IDetailsView.h"
-#include "MuCOE/Nodes/CustomizableObjectNodeModifierEditMeshSection.h"
-#include "MuCOE/Nodes/CustomizableObjectNodeModifierExtendMeshSection.h"
-#include "MuCOE/Nodes/CustomizableObjectNodeMaterialBase.h"
+#include "MuCOE/Nodes/CONodeModifierEditSkeletalMeshSection.h"
+#include "MuCOE/Nodes/CONodeModifierExtendSkeletalMeshSection.h"
+#include "MuCOE/Nodes/CONodeSkeletalMeshSectionBase.h"
 #include "MuCOE/SCustomizableObjectLayoutEditor.h"
 #include "MuCOE/PinViewer/SPinViewer.h"
 #include "MuR/MutableTrace.h"
@@ -36,7 +36,7 @@ void FCustomizableObjectNodeModifierEditMeshSectionDetails::CustomizeDetails( ID
 	TSharedPtr<const IDetailsView> DetailsView = DetailBuilder.GetDetailsViewSharedPtr();
 	if (DetailsView && DetailsView->GetSelectedObjects().Num())
 	{
-		Node = Cast<UCustomizableObjectNodeModifierEditMeshSection>(DetailsView->GetSelectedObjects()[0].Get());
+		Node = Cast<UCONodeModifierEditSkeletalMeshSection>(DetailsView->GetSelectedObjects()[0].Get());
 	}
 
 	
diff --git a/Engine/Plugins/Mutable/Source/CustomizableObjectEditor/Private/MuCOE/Nodes/CustomizableObjectNodeModifierEditMeshSectionDetails.h b/Engine/Plugins/Mutable/Source/CustomizableObjectEditor/Private/MuCOE/Nodes/CustomizableObjectNodeModifierEditMeshSectionDetails.h
--- a/Engine/Plugins/Mutable/Source/CustomizableObjectEditor/Private/MuCOE/Nodes/CustomizableObjectNodeModifierEditMeshSectionDetails.h
+++ b/Engine/Plugins/Mutable/Source/CustomizableObjectEditor/Private/MuCOE/Nodes/CustomizableObjectNodeModifierEditMeshSectionDetails.h
@@ -9,7 +9,7 @@
 
 class FString;
 class IDetailLayoutBuilder;
-class UCustomizableObjectNodeModifierEditMeshSection;
+class UCONodeModifierEditSkeletalMeshSection;
 class SCustomizableObjectLayoutEditor;
 
 
@@ -26,7 +26,7 @@ public:
 	virtual void OnRequiredTagsPropertyChanged() override;
 
 private:
-	UCustomizableObjectNodeModifierEditMeshSection* Node;
+	UCONodeModifierEditSkeletalMeshSection* Node;
 
 	// Layout block editor widget
 	TSharedPtr<SCustomizableObjectLayoutEditor> LayoutBlocksEditor;
diff --git a/Engine/Plugins/Mutable/Source/CustomizableObjectEditor/Private/MuCOE/Nodes/CustomizableObjectNodeModifierExtendMeshSection.cpp b/Engine/Plugins/Mutable/Source/CustomizableObjectEditor/Private/MuCOE/Nodes/CustomizableObjectNodeModifierExtendMeshSection.cpp
deleted file mode 100644
diff --git a/Engine/Plugins/Mutable/Source/CustomizableObjectEditor/Private/MuCOE/Nodes/CustomizableObjectNodeModifierExtendMeshSection.h b/Engine/Plugins/Mutable/Source/CustomizableObjectEditor/Private/MuCOE/Nodes/CustomizableObjectNodeModifierExtendMeshSection.h
deleted file mode 100644
diff --git a/Engine/Plugins/Mutable/Source/CustomizableObjectEditor/Private/MuCOE/Nodes/CustomizableObjectNodeModifierExtendMeshSectionDetails.cpp b/Engine/Plugins/Mutable/Source/CustomizableObjectEditor/Private/MuCOE/Nodes/CustomizableObjectNodeModifierExtendMeshSectionDetails.cpp
--- a/Engine/Plugins/Mutable/Source/CustomizableObjectEditor/Private/MuCOE/Nodes/CustomizableObjectNodeModifierExtendMeshSectionDetails.cpp
+++ b/Engine/Plugins/Mutable/Source/CustomizableObjectEditor/Private/MuCOE/Nodes/CustomizableObjectNodeModifierExtendMeshSectionDetails.cpp
@@ -5,8 +5,8 @@
 #include "MuCOE/CustomizableObjectEditorUtilities.h"
 #include "MuCOE/CustomizableObjectEditorStyle.h"
 #include "MuCOE/GraphTraversal.h"
-#include "MuCOE/Nodes/CustomizableObjectNodeMaterial.h"
-#include "MuCOE/Nodes/CustomizableObjectNodeModifierExtendMeshSection.h"
+#include "MuCOE/Nodes/CONodeSkeletalMeshSection.h"
+#include "MuCOE/Nodes/CONodeModifierExtendSkeletalMeshSection.h"
 #include "MuCOE/Nodes/CustomizableObjectNodeSkeletalMesh.h"
 #include "MuCOE/UnrealEditorPortabilityHelpers.h"
 #include "PropertyCustomizationHelpers.h"
@@ -34,7 +34,7 @@ void FCustomizableObjectNodeModifierExtendMeshSectionDetails::CustomizeDetails(
 	TSharedPtr<const IDetailsView> DetailsView = DetailBuilder.GetDetailsViewSharedPtr();
 	if (DetailsView && DetailsView->GetSelectedObjects().Num())
 	{
-		Node = Cast<UCustomizableObjectNodeModifierExtendMeshSection>( DetailsView->GetSelectedObjects()[0].Get() );
+		Node = Cast<UCONodeModifierExtendSkeletalMeshSection>( DetailsView->GetSelectedObjects()[0].Get() );
 	}
 
 	if (!Node)
@@ -48,7 +48,7 @@ void FCustomizableObjectNodeModifierExtendMeshSectionDetails::CustomizeDetails(
 
 	// Add the required tags widget
 	{
-		EnableTagsPropertyHandle = DetailBuilder.GetProperty(GET_MEMBER_NAME_CHECKED(UCustomizableObjectNodeModifierExtendMeshSection, Tags), UCustomizableObjectNodeModifierExtendMeshSection::StaticClass());
+		EnableTagsPropertyHandle = DetailBuilder.GetProperty(GET_MEMBER_NAME_CHECKED(UCONodeModifierExtendSkeletalMeshSection, Tags), UCONodeModifierExtendSkeletalMeshSection::StaticClass());
 		DetailBuilder.HideProperty(EnableTagsPropertyHandle);
 
 		EnableTagsPropertyHandle->SetOnPropertyValueChanged(FSimpleDelegate::CreateSP(this, &FCustomizableObjectNodeModifierExtendMeshSectionDetails::OnRequiredTagsPropertyChanged));
diff --git a/Engine/Plugins/Mutable/Source/CustomizableObjectEditor/Private/MuCOE/Nodes/CustomizableObjectNodeModifierExtendMeshSectionDetails.h b/Engine/Plugins/Mutable/Source/CustomizableObjectEditor/Private/MuCOE/Nodes/CustomizableObjectNodeModifierExtendMeshSectionDetails.h
--- a/Engine/Plugins/Mutable/Source/CustomizableObjectEditor/Private/MuCOE/Nodes/CustomizableObjectNodeModifierExtendMeshSectionDetails.h
+++ b/Engine/Plugins/Mutable/Source/CustomizableObjectEditor/Private/MuCOE/Nodes/CustomizableObjectNodeModifierExtendMeshSectionDetails.h
@@ -38,7 +38,7 @@ private:
 
 private:
 
-	class UCustomizableObjectNodeModifierExtendMeshSection* Node = nullptr;
+	class UCONodeModifierExtendSkeletalMeshSection* Node = nullptr;
 
 	/** */
 	TSharedPtr<IPropertyHandle> EnableTagsPropertyHandle;
diff --git a/Engine/Plugins/Mutable/Source/CustomizableObjectEditor/Private/MuCOE/Nodes/CustomizableObjectNodeModifierMorphMeshSection.cpp b/Engine/Plugins/Mutable/Source/CustomizableObjectEditor/Private/MuCOE/Nodes/CustomizableObjectNodeModifierMorphMeshSection.cpp
--- a/Engine/Plugins/Mutable/Source/CustomizableObjectEditor/Private/MuCOE/Nodes/CustomizableObjectNodeModifierMorphMeshSection.cpp
+++ b/Engine/Plugins/Mutable/Source/CustomizableObjectEditor/Private/MuCOE/Nodes/CustomizableObjectNodeModifierMorphMeshSection.cpp
@@ -6,7 +6,7 @@
 #include "MuCO/CustomizableObjectCustomVersion.h"
 #include "MuCOE/EdGraphSchema_CustomizableObject.h"
 #include "MuCOE/GraphTraversal.h"
-#include "MuCOE/Nodes/CustomizableObjectNodeMaterial.h"
+#include "MuCOE/Nodes/CONodeSkeletalMeshSection.h"
 #include "MuCOE/Nodes/CustomizableObjectNodeSkeletalMesh.h"
 #include "MuCOE/CustomizableObjectMacroLibrary/CustomizableObjectMacroLibrary.h"
 
diff --git a/Engine/Plugins/Mutable/Source/CustomizableObjectEditor/Private/MuCOE/Nodes/CustomizableObjectNodeModifierMorphMeshSectionDetails.cpp b/Engine/Plugins/Mutable/Source/CustomizableObjectEditor/Private/MuCOE/Nodes/CustomizableObjectNodeModifierMorphMeshSectionDetails.cpp
--- a/Engine/Plugins/Mutable/Source/CustomizableObjectEditor/Private/MuCOE/Nodes/CustomizableObjectNodeModifierMorphMeshSectionDetails.cpp
+++ b/Engine/Plugins/Mutable/Source/CustomizableObjectEditor/Private/MuCOE/Nodes/CustomizableObjectNodeModifierMorphMeshSectionDetails.cpp
@@ -7,10 +7,10 @@
 #include "MuCOE/GraphTraversal.h"
 #include "MuCOE/Nodes/CustomizableObjectNodeObject.h"
 #include "MuCOE/Nodes/CustomizableObjectNodeObjectGroup.h"
-#include "MuCOE/Nodes/CustomizableObjectNodeMaterial.h"
+#include "MuCOE/Nodes/CONodeSkeletalMeshSection.h"
 #include "MuCOE/Nodes/CustomizableObjectNodeModifierBase.h"
 #include "MuCOE/Nodes/CustomizableObjectNodeModifierMorphMeshSection.h"
-#include "MuCOE/Nodes/CustomizableObjectNodeModifierExtendMeshSection.h"
+#include "MuCOE/Nodes/CONodeModifierExtendSkeletalMeshSection.h"
 #include "MuCOE/Nodes/CustomizableObjectNodeSkeletalMesh.h"
 #include "MuCOE/UnrealEditorPortabilityHelpers.h"
 #include "PropertyCustomizationHelpers.h"
@@ -48,7 +48,7 @@ void FCustomizableObjectNodeModifierMorphMeshSectionDetails::CustomizeDetails( I
 	}
 
 	// This property is not relevant for this node
-	DetailBuilder.HideProperty(GET_MEMBER_NAME_CHECKED(UCustomizableObjectNodeModifierWithMaterial, ReferenceMaterial), UCustomizableObjectNodeModifierWithMaterial::StaticClass());
+	DetailBuilder.HideProperty(GET_MEMBER_NAME_CHECKED(UCONodeModifierSkeletalMeshSection, ReferenceMaterial), UCONodeModifierSkeletalMeshSection::StaticClass());
 	
 	// Add a morph selection widget.
 	TSharedRef<IPropertyHandle> MorphTargetNameProperty = DetailBuilder.GetProperty("MorphTargetName");
@@ -211,7 +211,7 @@ TSharedPtr<SMutableSearchComboBox::FFilteredOption> FCustomizableObjectNodeModif
 		}
 
 		// Is it a relevant type that we want to show in the hierarchy?
-		if (UCustomizableObjectNodeMaterial* MeshSectionNode = Cast<UCustomizableObjectNodeMaterial>(InNode))
+		if (UCONodeSkeletalMeshSection* MeshSectionNode = Cast<UCONodeSkeletalMeshSection>(InNode))
 		{
 			Option = MakeShared<SMutableSearchComboBox::FFilteredOption>();
 			Option->Parent = ParentOption;
@@ -270,7 +270,7 @@ void FCustomizableObjectNodeModifierMorphMeshSectionDetails::AddMorphsFromNode(
 {
 	USkeletalMesh* SkeletalMesh = nullptr;
 
-	if (UCustomizableObjectNodeMaterialBase* MaterialNode = Cast<UCustomizableObjectNodeMaterialBase>(Candidate))
+	if (UCONodeSkeletalMeshSectionBase* MaterialNode = Cast<UCONodeSkeletalMeshSectionBase>(Candidate))
 	{
 		if (MaterialNode->OutputPin())
 		{
@@ -286,7 +286,7 @@ void FCustomizableObjectNodeModifierMorphMeshSectionDetails::AddMorphsFromNode(
 		}
 	}
 
-	else if (UCustomizableObjectNodeModifierExtendMeshSection* ExtendNode = Cast<UCustomizableObjectNodeModifierExtendMeshSection>(Candidate))
+	else if (UCONodeModifierExtendSkeletalMeshSection* ExtendNode = Cast<UCONodeModifierExtendSkeletalMeshSection>(Candidate))
 	{
 		if (ExtendNode->GetOutputPin())
 		{
diff --git a/Engine/Plugins/Mutable/Source/CustomizableObjectEditor/Private/MuCOE/Nodes/CustomizableObjectNodeModifierRemoveMeshBlocks.cpp b/Engine/Plugins/Mutable/Source/CustomizableObjectEditor/Private/MuCOE/Nodes/CustomizableObjectNodeModifierRemoveMeshBlocks.cpp
--- a/Engine/Plugins/Mutable/Source/CustomizableObjectEditor/Private/MuCOE/Nodes/CustomizableObjectNodeModifierRemoveMeshBlocks.cpp
+++ b/Engine/Plugins/Mutable/Source/CustomizableObjectEditor/Private/MuCOE/Nodes/CustomizableObjectNodeModifierRemoveMeshBlocks.cpp
@@ -2,11 +2,11 @@
 
 #include "MuCOE/Nodes/CustomizableObjectNodeModifierRemoveMeshBlocks.h"
 
-#include "MuCOE/Nodes/CustomizableObjectNodeModifierExtendMeshSection.h"
+#include "MuCOE/Nodes/CONodeModifierExtendSkeletalMeshSection.h"
 #include "MuCO/CustomizableObjectCustomVersion.h"
 #include "MuCOE/CustomizableObjectLayout.h"
 #include "MuCOE/EdGraphSchema_CustomizableObject.h"
-#include "MuCOE/Nodes/CustomizableObjectNodeMaterial.h"
+#include "MuCOE/Nodes/CONodeSkeletalMeshSection.h"
 #include "MuCOE/Nodes/CustomizableObjectNodeSkeletalMesh.h"
 #include "MuCOE/Nodes/CustomizableObjectNodeTable.h"
 #include "MuCOE/Nodes/CustomizableObjectNodeLayoutBlocks.h"
@@ -28,7 +28,7 @@ void UCustomizableObjectNodeModifierRemoveMeshBlocks::BackwardsCompatibleFixup(i
 	if (CustomizableObjectCustomVersion == FCustomizableObjectCustomVersion::PostLoadToCustomVersion
 		&& BlockIds_DEPRECATED.Num() < Blocks_DEPRECATED.Num())
 	{
-		UCustomizableObjectNodeMaterialBase* ParentMaterialNode = GetCustomizableObjectExternalNode<UCustomizableObjectNodeMaterialBase>(ParentMaterialObject_DEPRECATED.Get(), ParentMaterialNodeId_DEPRECATED);
+		UCONodeSkeletalMeshSectionBase* ParentMaterialNode = GetCustomizableObjectExternalNode<UCONodeSkeletalMeshSectionBase>(ParentMaterialObject_DEPRECATED.Get(), ParentMaterialNodeId_DEPRECATED);
 		if (ParentMaterialNode)
 		{
 			TArray<UCustomizableObjectLayout*> Layouts = ParentMaterialNode->GetLayouts();
@@ -38,7 +38,7 @@ void UCustomizableObjectNodeModifierRemoveMeshBlocks::BackwardsCompatibleFixup(i
 				UE_LOG(LogMutable, Warning, TEXT("[%s] UCustomizableObjectNodeModifierRemoveMeshBlocks refers to an invalid texture layout index %d. Parent node has %d layouts."),
 					*GetOutermost()->GetName(), ParentLayoutIndex, Layouts.Num());
 			}
-			else if (UCustomizableObjectNodeMaterial* ParentMaterial = Cast<UCustomizableObjectNodeMaterial>(ParentMaterialNode))
+			else if (UCONodeSkeletalMeshSection* ParentMaterial = Cast<UCONodeSkeletalMeshSection>(ParentMaterialNode))
 			{
 				UCustomizableObjectLayout* ParentLayout = Layouts[ParentLayoutIndex];
 
@@ -78,7 +78,7 @@ void UCustomizableObjectNodeModifierRemoveMeshBlocks::BackwardsCompatibleFixup(i
 		// If we are here, it means this node was loaded from a version that didn't have it's own layout.
 		check(Layout->Blocks.IsEmpty());
 
-		UCustomizableObjectNodeMaterialBase* ParentMaterialNode = GetCustomizableObjectExternalNode<UCustomizableObjectNodeMaterialBase>(ParentMaterialObject_DEPRECATED.Get(), ParentMaterialNodeId_DEPRECATED);
+		UCONodeSkeletalMeshSectionBase* ParentMaterialNode = GetCustomizableObjectExternalNode<UCONodeSkeletalMeshSectionBase>(ParentMaterialObject_DEPRECATED.Get(), ParentMaterialNodeId_DEPRECATED);
 
 		TArray<UCustomizableObjectLayout*> ParentLayouts;
 
@@ -92,7 +92,7 @@ void UCustomizableObjectNodeModifierRemoveMeshBlocks::BackwardsCompatibleFixup(i
 		{
 			[&]()
 				{
-					UCustomizableObjectNodeMaterial* NodeMaterial = ParentMaterialNode->GetMaterialNode();
+					UCONodeSkeletalMeshSection* NodeMaterial = ParentMaterialNode->GetMaterialNode();
 					if (!NodeMaterial)
 					{
 						return;
@@ -170,7 +170,7 @@ void UCustomizableObjectNodeModifierRemoveMeshBlocks::BackwardsCompatibleFixup(i
 
 			Layout->SetGridSize(GridSize);
 
-			if (Cast<UCustomizableObjectNodeMaterial>(ParentMaterialNode))
+			if (Cast<UCONodeSkeletalMeshSection>(ParentMaterialNode))
 			{
 				for (const FGuid& BlockId : BlockIds_DEPRECATED)
 				{
diff --git a/Engine/Plugins/Mutable/Source/CustomizableObjectEditor/Private/MuCOE/Nodes/CustomizableObjectNodeModifierRemoveMeshBlocksDetails.cpp b/Engine/Plugins/Mutable/Source/CustomizableObjectEditor/Private/MuCOE/Nodes/CustomizableObjectNodeModifierRemoveMeshBlocksDetails.cpp
--- a/Engine/Plugins/Mutable/Source/CustomizableObjectEditor/Private/MuCOE/Nodes/CustomizableObjectNodeModifierRemoveMeshBlocksDetails.cpp
+++ b/Engine/Plugins/Mutable/Source/CustomizableObjectEditor/Private/MuCOE/Nodes/CustomizableObjectNodeModifierRemoveMeshBlocksDetails.cpp
@@ -7,7 +7,7 @@
 #include "IDetailGroup.h"
 #include "IDetailsView.h"
 #include "MuCOE/Nodes/CustomizableObjectNodeModifierRemoveMeshBlocks.h"
-#include "MuCOE/Nodes/CustomizableObjectNodeMaterialBase.h"
+#include "MuCOE/Nodes/CONodeSkeletalMeshSectionBase.h"
 #include "MuCOE/SCustomizableObjectLayoutEditor.h"
 #include "Widgets/Layout/SBox.h"
 #include "Widgets/Text/STextBlock.h"
@@ -38,7 +38,7 @@ void FCustomizableObjectNodeModifierRemoveMeshBlocksDetails::CustomizeDetails(ID
 	}
 
 	// This property is not relevant for this node
-	DetailBuilder.HideProperty(GET_MEMBER_NAME_CHECKED(UCustomizableObjectNodeModifierWithMaterial, ReferenceMaterial), UCustomizableObjectNodeModifierWithMaterial::StaticClass());
+	DetailBuilder.HideProperty(GET_MEMBER_NAME_CHECKED(UCONodeModifierSkeletalMeshSection, ReferenceMaterial), UCONodeModifierSkeletalMeshSection::StaticClass());
 
 	IDetailCategoryBuilder& LayoutCategory = DetailBuilder.EditCategory("Layout Editor");
 	LayoutCategory.SetSortOrder(10000);
diff --git a/Engine/Plugins/Mutable/Source/CustomizableObjectEditor/Private/MuCOE/Nodes/CustomizableObjectNodeModifierRemoveMeshDetails.cpp b/Engine/Plugins/Mutable/Source/CustomizableObjectEditor/Private/MuCOE/Nodes/CustomizableObjectNodeModifierRemoveMeshDetails.cpp
--- a/Engine/Plugins/Mutable/Source/CustomizableObjectEditor/Private/MuCOE/Nodes/CustomizableObjectNodeModifierRemoveMeshDetails.cpp
+++ b/Engine/Plugins/Mutable/Source/CustomizableObjectEditor/Private/MuCOE/Nodes/CustomizableObjectNodeModifierRemoveMeshDetails.cpp
@@ -7,7 +7,7 @@
 #include "IDetailGroup.h"
 #include "IDetailsView.h"
 #include "MuCOE/Nodes/CustomizableObjectNodeModifierRemoveMesh.h"
-#include "MuCOE/Nodes/CustomizableObjectNodeMaterialBase.h"
+#include "MuCOE/Nodes/CONodeSkeletalMeshSectionBase.h"
 #include "Widgets/Layout/SBox.h"
 #include "Widgets/Text/STextBlock.h"
 #include "Widgets/Input/STextComboBox.h"
@@ -28,7 +28,7 @@ void FCustomizableObjectNodeModifierRemoveMeshDetails::CustomizeDetails(IDetailL
 	FCustomizableObjectNodeModifierBaseDetails::CustomizeDetails(DetailBuilder);
 
 	// This property is not relevant for this node
-	DetailBuilder.HideProperty(GET_MEMBER_NAME_CHECKED(UCustomizableObjectNodeModifierWithMaterial, ReferenceMaterial), UCustomizableObjectNodeModifierWithMaterial::StaticClass());
+	DetailBuilder.HideProperty(GET_MEMBER_NAME_CHECKED(UCONodeModifierSkeletalMeshSection, ReferenceMaterial), UCONodeModifierSkeletalMeshSection::StaticClass());
 
 }
 
diff --git a/Engine/Plugins/Mutable/Source/CustomizableObjectEditor/Private/MuCOE/Nodes/CustomizableObjectNodeModifierWithMaterial.cpp b/Engine/Plugins/Mutable/Source/CustomizableObjectEditor/Private/MuCOE/Nodes/CustomizableObjectNodeModifierWithMaterial.cpp
deleted file mode 100644
diff --git a/Engine/Plugins/Mutable/Source/CustomizableObjectEditor/Private/MuCOE/Nodes/CustomizableObjectNodeModifierWithMaterial.h b/Engine/Plugins/Mutable/Source/CustomizableObjectEditor/Private/MuCOE/Nodes/CustomizableObjectNodeModifierWithMaterial.h
deleted file mode 100644
diff --git a/Engine/Plugins/Mutable/Source/CustomizableObjectEditor/Private/MuCOE/Nodes/CustomizableObjectNodeObject.cpp b/Engine/Plugins/Mutable/Source/CustomizableObjectEditor/Private/MuCOE/Nodes/CustomizableObjectNodeObject.cpp
--- a/Engine/Plugins/Mutable/Source/CustomizableObjectEditor/Private/MuCOE/Nodes/CustomizableObjectNodeObject.cpp
+++ b/Engine/Plugins/Mutable/Source/CustomizableObjectEditor/Private/MuCOE/Nodes/CustomizableObjectNodeObject.cpp
@@ -11,15 +11,15 @@
 #include "MuCOE/EdGraphSchema_CustomizableObject.h"
 #include "MuCOE/GraphTraversal.h"
 #include "MuCOE/CustomizableObjectEditor.h"
-#include "MuCOE/Nodes/CustomizableObjectNodeMaterial.h"
+#include "MuCOE/Nodes/CONodeSkeletalMeshSection.h"
 #include "MuCOE/Nodes/CustomizableObjectNodeComponentMesh.h"
-#include "MuCOE/Nodes/CustomizableObjectNodeMaterialVariation.h"
+#include "MuCOE/Nodes/CONodeSkeletalMeshSectionVariation.h"
 #include "Misc/UObjectToken.h"
 #include "Logging/MessageLog.h"
 #include "MuCO/CustomizableObjectCustomVersion.h"
 #include "MuCOE/CustomizableObjectEditor_Deprecated.h"
 #include "MuCOE/Nodes/CustomizableObjectNodeComponentMeshAddTo.h"
-#include "MuCOE/Nodes/CustomizableObjectNodeMaterialSwitch.h"
+#include "MuCOE/Nodes/CONodeSkeletalMeshSectionSwitch.h"
 #include "Settings/EditorStyleSettings.h"
 
 #include UE_INLINE_GENERATED_CPP_BY_NAME(CustomizableObjectNodeObject)
@@ -340,13 +340,13 @@ void UCustomizableObjectNodeObject::BackwardsCompatibleFixup(int32 CustomizableO
 				UCustomizableObjectNodeComponentMeshAddTo* NodeComponentMeshAddTo = nullptr;
 				bool bFixNode = false;
 				
-				if (UCustomizableObjectNodeMaterial* NodeMaterial = Cast<UCustomizableObjectNodeMaterial>(LinkedNode))
+				if (UCONodeSkeletalMeshSection* NodeMaterial = Cast<UCONodeSkeletalMeshSection>(LinkedNode))
 				{
 					bFixNode = true;
 
 					NodeComponentMeshAddTo = CreateNodeComponent(NodeMaterial->MeshComponentName_DEPRECATED);
 				}
-				else if (UCustomizableObjectNodeMaterialSwitch* NodeMaterialSwitch = Cast<UCustomizableObjectNodeMaterialSwitch>(LinkedNode))
+				else if (UCONodeSkeletalMeshSectionSwitch* NodeMaterialSwitch = Cast<UCONodeSkeletalMeshSectionSwitch>(LinkedNode))
 				{
 					bFixNode = true;
 
@@ -363,9 +363,9 @@ void UCustomizableObjectNodeObject::BackwardsCompatibleFixup(int32 CustomizableO
 						{
 							if (UEdGraphPin* ConnectedPin = FollowInputPin(*NodeMaterialSwitch->SwitchPins[ElementIndex].Get()))
 							{
-								if (UCustomizableObjectNodeMaterialBase* FirstNodeMaterialBase = Cast<UCustomizableObjectNodeMaterialBase>(ConnectedPin->GetOwningNode()))
+								if (UCONodeSkeletalMeshSectionBase* FirstNodeMaterialBase = Cast<UCONodeSkeletalMeshSectionBase>(ConnectedPin->GetOwningNode()))
 								{
-									if (UCustomizableObjectNodeMaterial* FirstNodeMaterial = FirstNodeMaterialBase->GetMaterialNode())
+									if (UCONodeSkeletalMeshSection* FirstNodeMaterial = FirstNodeMaterialBase->GetMaterialNode())
 									{
 										if (bFirst)
 										{
@@ -387,7 +387,7 @@ void UCustomizableObjectNodeObject::BackwardsCompatibleFixup(int32 CustomizableO
 						NodeComponentMeshAddTo = CreateNodeComponent(ComponentName);
 					}();
 				}
-				else if (UCustomizableObjectNodeMaterialVariation* NodeMaterialVariation = Cast<UCustomizableObjectNodeMaterialVariation>(LinkedNode))
+				else if (UCONodeSkeletalMeshSectionVariation* NodeMaterialVariation = Cast<UCONodeSkeletalMeshSectionVariation>(LinkedNode))
 				{
 					bFixNode = true;
 
@@ -405,9 +405,9 @@ void UCustomizableObjectNodeObject::BackwardsCompatibleFixup(int32 CustomizableO
 						if (ConnectedDefaultPins.Num() > 0)
 						{
 							UEdGraphPin* FirstDefaultPin = ConnectedDefaultPins[0];
-							if (UCustomizableObjectNodeMaterialBase* FirstNodeMaterialBase = Cast<UCustomizableObjectNodeMaterialBase>(FirstDefaultPin->GetOwningNode()))
+							if (UCONodeSkeletalMeshSectionBase* FirstNodeMaterialBase = Cast<UCONodeSkeletalMeshSectionBase>(FirstDefaultPin->GetOwningNode()))
 							{
-								if (UCustomizableObjectNodeMaterial* FirstNodeMaterial = FirstNodeMaterialBase->GetMaterialNode())
+								if (UCONodeSkeletalMeshSection* FirstNodeMaterial = FirstNodeMaterialBase->GetMaterialNode())
 								{
 									bFirst = false;
 									ComponentName = FirstNodeMaterial->MeshComponentName_DEPRECATED;
@@ -424,9 +424,9 @@ void UCustomizableObjectNodeObject::BackwardsCompatibleFixup(int32 CustomizableO
 								if (ConnectedPins.Num() > 0)
 								{
 									UEdGraphPin* ConnectedPin = ConnectedPins[0];
-									if (UCustomizableObjectNodeMaterialBase* FirstNodeMaterialBase = Cast<UCustomizableObjectNodeMaterialBase>(ConnectedPin->GetOwningNode()))
+									if (UCONodeSkeletalMeshSectionBase* FirstNodeMaterialBase = Cast<UCONodeSkeletalMeshSectionBase>(ConnectedPin->GetOwningNode()))
 									{
-										if (UCustomizableObjectNodeMaterial* FirstNodeMaterial = FirstNodeMaterialBase->GetMaterialNode())
+										if (UCONodeSkeletalMeshSection* FirstNodeMaterial = FirstNodeMaterialBase->GetMaterialNode())
 										{
 											if (bFirst)
 											{
diff --git a/Engine/Plugins/Mutable/Source/CustomizableObjectEditor/Private/MuCOE/Nodes/CustomizableObjectNodeObject.h b/Engine/Plugins/Mutable/Source/CustomizableObjectEditor/Private/MuCOE/Nodes/CustomizableObjectNodeObject.h
--- a/Engine/Plugins/Mutable/Source/CustomizableObjectEditor/Private/MuCOE/Nodes/CustomizableObjectNodeObject.h
+++ b/Engine/Plugins/Mutable/Source/CustomizableObjectEditor/Private/MuCOE/Nodes/CustomizableObjectNodeObject.h
@@ -14,7 +14,7 @@
 namespace ENodeTitleType { enum Type : int; }
 
 class UCustomizableObject;
-class UCustomizableObjectNodeMaterialBase;
+class UCONodeSkeletalMeshSectionBase;
 class UCustomizableObjectNodeRemapPins;
 class UObject;
 struct FPropertyChangedEvent;
diff --git a/Engine/Plugins/Mutable/Source/CustomizableObjectEditor/Private/MuCOE/Nodes/CustomizableObjectNodeObjectDetails.h b/Engine/Plugins/Mutable/Source/CustomizableObjectEditor/Private/MuCOE/Nodes/CustomizableObjectNodeObjectDetails.h
--- a/Engine/Plugins/Mutable/Source/CustomizableObjectEditor/Private/MuCOE/Nodes/CustomizableObjectNodeObjectDetails.h
+++ b/Engine/Plugins/Mutable/Source/CustomizableObjectEditor/Private/MuCOE/Nodes/CustomizableObjectNodeObjectDetails.h
@@ -97,7 +97,7 @@ private:
 	TArray< TSharedPtr<FString> > GroupNodeComboOptions;
 
 	TArray< TSharedPtr<FString> > ParentComboOptions;
-	TArray< class UCustomizableObjectNodeMaterial* > ParentOptionNode;
+	TArray< class UCONodeSkeletalMeshSection* > ParentOptionNode;
 
 	void OnGroupNodeComboBoxSelectionChanged(TSharedPtr<FString> Selection, ESelectInfo::Type SelectInfo, TSharedRef<IPropertyHandle> ParentProperty);
 	
diff --git a/Engine/Plugins/Mutable/Source/CustomizableObjectEditor/Private/MuCOE/Nodes/CustomizableObjectNodeTable.cpp b/Engine/Plugins/Mutable/Source/CustomizableObjectEditor/Private/MuCOE/Nodes/CustomizableObjectNodeTable.cpp
--- a/Engine/Plugins/Mutable/Source/CustomizableObjectEditor/Private/MuCOE/Nodes/CustomizableObjectNodeTable.cpp
+++ b/Engine/Plugins/Mutable/Source/CustomizableObjectEditor/Private/MuCOE/Nodes/CustomizableObjectNodeTable.cpp
@@ -17,7 +17,7 @@
 #include "MuCOE/EdGraphSchema_CustomizableObject.h"
 #include "MuCOE/GraphTraversal.h"
 #include "MuCOE/MutableUtils.h"
-#include "MuCOE/Nodes/CustomizableObjectNodeMaterial.h"
+#include "MuCOE/Nodes/CONodeSkeletalMeshSection.h"
 #include "MuCOE/UnrealEditorPortabilityHelpers.h"
 #include "Rendering/SkeletalMeshLODModel.h"
 #include "Rendering/SkeletalMeshModel.h"
diff --git a/Engine/Plugins/Mutable/Source/CustomizableObjectEditor/Private/MuCOE/Nodes/SCONodeSkeletalMeshSection.cpp b/Engine/Plugins/Mutable/Source/CustomizableObjectEditor/Private/MuCOE/Nodes/SCONodeSkeletalMeshSection.cpp
new file mode 100644
--- /dev/null
+++ b/Engine/Plugins/Mutable/Source/CustomizableObjectEditor/Private/MuCOE/Nodes/SCONodeSkeletalMeshSection.cpp
@@ -0,0 +1,27 @@
+// Copyright Epic Games, Inc. All Rights Reserved.
+
+#include "MuCOE/Nodes/SCONodeSkeletalMeshSection.h"
+
+#include "MuCOE/EdGraphSchema_CustomizableObject.h"
+#include "MuCOE/Nodes/SCONodeSkeletalMeshSectionPinImage.h"
+#include "MuCOE/Nodes/CONodeSkeletalMeshSection.h"
+
+class SGraphPin;
+
+
+void SCONodeSkeletalMeshSection::Construct(const FArguments& InArgs, UCONodeSkeletalMeshSection* InGraphNode)
+{
+	SCustomizableObjectNode::Construct({}, InGraphNode);
+}
+
+
+TSharedPtr<SGraphPin> SCONodeSkeletalMeshSection::CreatePinWidget(UEdGraphPin* Pin) const
+{
+	if ((Pin->PinType.PinCategory == UEdGraphSchema_CustomizableObject::PC_Texture || Pin->PinType.PinCategory == UEdGraphSchema_CustomizableObject::PC_Texture_Passthrough) &&
+		Pin->Direction == EGPD_Input)
+	{
+		return SNew(SCONodeSkeletalMeshSectionPinImage, Pin);
+	}
+
+	return SGraphNode::CreatePinWidget(Pin);
+}
diff --git a/Engine/Plugins/Mutable/Source/CustomizableObjectEditor/Private/MuCOE/Nodes/SCONodeSkeletalMeshSection.h b/Engine/Plugins/Mutable/Source/CustomizableObjectEditor/Private/MuCOE/Nodes/SCONodeSkeletalMeshSection.h
new file mode 100644
--- /dev/null
+++ b/Engine/Plugins/Mutable/Source/CustomizableObjectEditor/Private/MuCOE/Nodes/SCONodeSkeletalMeshSection.h
@@ -0,0 +1,24 @@
+// Copyright Epic Games, Inc. All Rights Reserved.
+
+#pragma once
+
+#include "SCustomizableObjectNode.h"
+#include "SGraphNode.h"
+
+class UCONodeSkeletalMeshSection;
+class SGraphPin;
+class UEdGraphPin;
+
+
+/** Custom widget for the Material node. */
+class SCONodeSkeletalMeshSection : public SCustomizableObjectNode
+{
+public:
+	SLATE_BEGIN_ARGS(SCONodeSkeletalMeshSection) {}
+	SLATE_END_ARGS();
+
+	void Construct(const FArguments& InArgs, UCONodeSkeletalMeshSection* InGraphNode);
+
+	// SGraphNode interface
+	virtual TSharedPtr<SGraphPin> CreatePinWidget(UEdGraphPin* Pin) const override;
+};
diff --git a/Engine/Plugins/Mutable/Source/CustomizableObjectEditor/Private/MuCOE/Nodes/SCONodeSkeletalMeshSectionPinImage.cpp b/Engine/Plugins/Mutable/Source/CustomizableObjectEditor/Private/MuCOE/Nodes/SCONodeSkeletalMeshSectionPinImage.cpp
new file mode 100644
--- /dev/null
+++ b/Engine/Plugins/Mutable/Source/CustomizableObjectEditor/Private/MuCOE/Nodes/SCONodeSkeletalMeshSectionPinImage.cpp
@@ -0,0 +1,81 @@
+// Copyright Epic Games, Inc. All Rights Reserved.
+
+#include "MuCOE/Nodes/SCONodeSkeletalMeshSectionPinImage.h"
+
+#include "MuCOE/EdGraphSchema_CustomizableObject.h"
+#include "MuCOE/GraphTraversal.h"
+#include "MuCOE/Nodes/CONodeSkeletalMeshSection.h"
+#include "MuCOE/UnrealEditorPortabilityHelpers.h"
+#include "Widgets/Input/SEditableTextBox.h"
+#include "Widgets/SToolTip.h"
+
+class IToolTip;
+class SWidget;
+struct FSlateBrush;
+
+#define LOCTEXT_NAMESPACE "CustomizableObjectEditor"
+
+
+void SCONodeSkeletalMeshSectionPinImage::Construct(const FArguments& InArgs, UEdGraphPin* InGraphPinObj)
+{
+	SCustomizableObjectNodePin::Construct(SCustomizableObjectNodePin::FArguments(), InGraphPinObj);
+
+	// Override previously defined tool tip.
+	TSharedPtr<IToolTip> TooltipWidget = SNew(SToolTip)
+		.Text(this, &SCONodeSkeletalMeshSectionPinImage::GetPinTooltipText);
+
+	SetToolTip(TooltipWidget);
+}
+
+
+TSharedRef<SWidget>	SCONodeSkeletalMeshSectionPinImage::GetDefaultValueWidget()
+{
+	LabelAndValue->SetWrapSize(TNumericLimits<float>::Max()); // Remove warping.
+
+	return SNew(SEditableTextBox)
+		.Style(FAppStyle::Get(), "Graph.EditableTextBox")
+		.Text(this, &SCONodeSkeletalMeshSectionPinImage::GetDefaultValueText)
+		.SelectAllTextWhenFocused(false)
+		.Visibility(this, &SCONodeSkeletalMeshSectionPinImage::GetDefaultValueVisibility)
+		.IsReadOnly(true)
+		.ForegroundColor(FSlateColor::UseForeground());
+}
+
+
+FText SCONodeSkeletalMeshSectionPinImage::GetPinTooltipText() const
+{
+	if (GraphPinObj->bOrphanedPin)
+	{
+		return LOCTEXT("PinModeMutableOrpahan", "Pin not disapearing due to being connected or having a property modified.");
+	}
+	
+	if (GraphPinObj->PinType.PinCategory == UEdGraphSchema_CustomizableObject::PC_Texture)
+	{
+		return LOCTEXT("PinModeMutableTooltip", "Texture Parameter goes through Mutable.");
+	}
+	else
+	{
+		return LOCTEXT("PinModePassthroughTooltip", "Texture Parameter is ignored by Mutable.");
+	}
+}
+
+
+FText SCONodeSkeletalMeshSectionPinImage::GetDefaultValueText() const
+{
+	if (GraphPinObj->PinType.PinCategory == UEdGraphSchema_CustomizableObject::PC_Texture)
+	{
+		return LOCTEXT("PinModeMutable", "mutable");
+	}
+	else 
+	{
+		return LOCTEXT("PinModePassthrough", "passthrough");
+	}
+}
+
+
+EVisibility SCONodeSkeletalMeshSectionPinImage::GetDefaultValueVisibility() const
+{
+	return EVisibility::Visible;
+}
+
+#undef LOCTEXT_NAMESPACE
diff --git a/Engine/Plugins/Mutable/Source/CustomizableObjectEditor/Private/MuCOE/Nodes/SCONodeSkeletalMeshSectionPinImage.h b/Engine/Plugins/Mutable/Source/CustomizableObjectEditor/Private/MuCOE/Nodes/SCONodeSkeletalMeshSectionPinImage.h
new file mode 100644
--- /dev/null
+++ b/Engine/Plugins/Mutable/Source/CustomizableObjectEditor/Private/MuCOE/Nodes/SCONodeSkeletalMeshSectionPinImage.h
@@ -0,0 +1,34 @@
+// Copyright Epic Games, Inc. All Rights Reserved.
+
+#pragma once
+
+#include "SCustomizableObjectNodePin.h"
+#include "SGraphPin.h"
+
+class SWidget;
+class UEdGraphPin;
+struct FSlateBrush;
+
+
+/** Implements the "MUTABLE" and "PASSTHROUGH" text next to the pin name. */
+class SCONodeSkeletalMeshSectionPinImage : public SCustomizableObjectNodePin
+{
+public:
+
+	SLATE_BEGIN_ARGS(SCONodeSkeletalMeshSectionPinImage) {}
+	SLATE_END_ARGS()
+
+	void Construct(const FArguments& InArgs, UEdGraphPin* InGraphPinObj);
+
+	TSharedRef<SWidget>	GetDefaultValueWidget() override;
+
+private:
+	/** Return pin state text. */
+	FText GetDefaultValueText() const;
+
+	/** Return true if the pin state is visible. */
+	EVisibility GetDefaultValueVisibility() const;
+
+	/** Return pin tool tip. */
+	FText GetPinTooltipText() const;
+};
diff --git a/Engine/Plugins/Mutable/Source/CustomizableObjectEditor/Private/MuCOE/Nodes/SCustomizableObjectNodeMaterial.cpp b/Engine/Plugins/Mutable/Source/CustomizableObjectEditor/Private/MuCOE/Nodes/SCustomizableObjectNodeMaterial.cpp
deleted file mode 100644
diff --git a/Engine/Plugins/Mutable/Source/CustomizableObjectEditor/Private/MuCOE/Nodes/SCustomizableObjectNodeMaterial.h b/Engine/Plugins/Mutable/Source/CustomizableObjectEditor/Private/MuCOE/Nodes/SCustomizableObjectNodeMaterial.h
deleted file mode 100644
diff --git a/Engine/Plugins/Mutable/Source/CustomizableObjectEditor/Private/MuCOE/Nodes/SCustomizableObjectNodeMaterialPinImage.cpp b/Engine/Plugins/Mutable/Source/CustomizableObjectEditor/Private/MuCOE/Nodes/SCustomizableObjectNodeMaterialPinImage.cpp
deleted file mode 100644
diff --git a/Engine/Plugins/Mutable/Source/CustomizableObjectEditor/Private/MuCOE/Nodes/SCustomizableObjectNodeMaterialPinImage.h b/Engine/Plugins/Mutable/Source/CustomizableObjectEditor/Private/MuCOE/Nodes/SCustomizableObjectNodeMaterialPinImage.h
deleted file mode 100644
diff --git a/Engine/Plugins/Mutable/Source/CustomizableObjectEditor/Private/MuCOE/SCustomizableObjectEditorTagExplorer.cpp b/Engine/Plugins/Mutable/Source/CustomizableObjectEditor/Private/MuCOE/SCustomizableObjectEditorTagExplorer.cpp
--- a/Engine/Plugins/Mutable/Source/CustomizableObjectEditor/Private/MuCOE/SCustomizableObjectEditorTagExplorer.cpp
+++ b/Engine/Plugins/Mutable/Source/CustomizableObjectEditor/Private/MuCOE/SCustomizableObjectEditorTagExplorer.cpp
@@ -8,10 +8,10 @@
 #include "MuCOE/CustomizableObjectEditor.h"
 #include "MuCOE/CustomizableObjectGraph.h"
 #include "MuCOE/GraphTraversal.h"
-#include "MuCOE/Nodes/CustomizableObjectNodeModifierExtendMeshSection.h"
-#include "MuCOE/Nodes/CustomizableObjectNodeModifierEditMeshSection.h"
-#include "MuCOE/Nodes/CustomizableObjectNodeMaterial.h"
-#include "MuCOE/Nodes/CustomizableObjectNodeMaterialVariation.h"
+#include "MuCOE/Nodes/CONodeModifierExtendSkeletalMeshSection.h"
+#include "MuCOE/Nodes/CONodeModifierEditSkeletalMeshSection.h"
+#include "MuCOE/Nodes/CONodeSkeletalMeshSection.h"
+#include "MuCOE/Nodes/CONodeSkeletalMeshSectionVariation.h"
 #include "MuCOE/Nodes/CustomizableObjectNodeModifierClipDeform.h"
 #include "MuCOE/Nodes/CustomizableObjectNodeModifierClipMorph.h"
 #include "MuCOE/Nodes/CustomizableObjectNodeModifierClipWithMesh.h"
@@ -174,7 +174,7 @@ void SCustomizableObjectEditorTagExplorer::FillTagInformation(const UCustomizabl
 			// Gather tags from nodes that require them
 
 			// TODO: All other variations?
-			if (UCustomizableObjectNodeMaterialVariation* TypedNodeVariations = Cast<UCustomizableObjectNodeMaterialVariation>(Node))
+			if (UCONodeSkeletalMeshSectionVariation* TypedNodeVariations = Cast<UCONodeSkeletalMeshSectionVariation>(Node))
 			{
 				for (int32 i = 0; i < TypedNodeVariations->GetNumVariations(); ++i)
 				{
diff --git a/Engine/Plugins/Mutable/Source/CustomizableObjectEditor/Private/MuCOE/SMutableTagListWidget.cpp b/Engine/Plugins/Mutable/Source/CustomizableObjectEditor/Private/MuCOE/SMutableTagListWidget.cpp
--- a/Engine/Plugins/Mutable/Source/CustomizableObjectEditor/Private/MuCOE/SMutableTagListWidget.cpp
+++ b/Engine/Plugins/Mutable/Source/CustomizableObjectEditor/Private/MuCOE/SMutableTagListWidget.cpp
@@ -8,9 +8,9 @@
 #include "MuCOE/Nodes/CustomizableObjectNode.h"
 #include "MuCOE/Nodes/CustomizableObjectNodeObject.h"
 #include "MuCOE/Nodes/CustomizableObjectNodeObjectGroup.h"
-#include "MuCOE/Nodes/CustomizableObjectNodeMaterial.h"
+#include "MuCOE/Nodes/CONodeSkeletalMeshSection.h"
 #include "MuCOE/Nodes/CustomizableObjectNodeModifierBase.h"
-#include "MuCOE/Nodes/CustomizableObjectNodeModifierExtendMeshSection.h"
+#include "MuCOE/Nodes/CONodeModifierExtendSkeletalMeshSection.h"
 #include "MuCOE/Nodes/CustomizableObjectNodeVariation.h"
 #include "MuCOE/Nodes/CustomizableObjectNodeTextureVariation.h"
 #include "MuCOE/Nodes/CustomizableObjectNodeFloatVariation.h"
@@ -105,7 +105,7 @@ TSharedPtr<SMutableSearchComboBox::FFilteredOption> SMutableTagComboBox::AddNode
 		}
 
 		// Is it a relevant type that we want to show in the hierarchy?
-		if (UCustomizableObjectNodeMaterial* MeshSectionNode = Cast<UCustomizableObjectNodeMaterial>(InNode))
+		if (UCONodeSkeletalMeshSection* MeshSectionNode = Cast<UCONodeSkeletalMeshSection>(InNode))
 		{
 			Option = MakeShared<SMutableSearchComboBox::FFilteredOption>();
 			Option->Parent = ParentOption;
@@ -117,7 +117,7 @@ TSharedPtr<SMutableSearchComboBox::FFilteredOption> SMutableTagComboBox::AddNode
 			TagComboOptionsSource.Add(Option.ToSharedRef());
 		}
 
-		else if (UCustomizableObjectNodeModifierExtendMeshSection* ExtendMeshSectionNode = Cast<UCustomizableObjectNodeModifierExtendMeshSection>(InNode))
+		else if (UCONodeModifierExtendSkeletalMeshSection* ExtendMeshSectionNode = Cast<UCONodeModifierExtendSkeletalMeshSection>(InNode))
 		{
 			Option = MakeShared<SMutableSearchComboBox::FFilteredOption>();
 			Option->Parent = ParentOption;
