[Mutable] Remove Share Projector Between LODs optimization.

This optimization is already applied if two generated UTextures are equivalent.

#Synced-CL 49271159

diff --git a/Engine/Plugins/Mutable/Source/CustomizableObjectEditor/Private/MuCOE/GenerateMutableSource/GenerateMutableSource.h b/Engine/Plugins/Mutable/Source/CustomizableObjectEditor/Private/MuCOE/GenerateMutableSource/GenerateMutableSource.h
--- a/Engine/Plugins/Mutable/Source/CustomizableObjectEditor/Private/MuCOE/GenerateMutableSource/GenerateMutableSource.h
+++ b/Engine/Plugins/Mutable/Source/CustomizableObjectEditor/Private/MuCOE/GenerateMutableSource/GenerateMutableSource.h
@@ -430,34 +430,6 @@ struct FGroupNodeIdsTempData
 	}
 };
 
-struct FGroupProjectorImageInfo
-{
-	UE::Mutable::Private::Ptr<UE::Mutable::Private::NodeImage> ImageNode;
-	UE::Mutable::Private::Ptr<UE::Mutable::Private::NodeImage> ImageResizeNode;
-	UCONodeSkeletalMeshSectionBase* TypedNodeMat;
-	FString TextureName;
-	FString RealTextureName;
-	int32 UVLayout = 0;
-	int32 LayerIndex = -1;
-
-	FGroupProjectorImageInfo(
-		UE::Mutable::Private::NodeImagePtr InImageNode,
-		const FString& InTextureName,
-		const FString& InRealTextureName,
-		UCONodeSkeletalMeshSectionBase* InTypedNodeMat,
-		int32 InUVLayout,
-		int32 LayerIndex)
-		: TypedNodeMat(InTypedNodeMat), TextureName(InTextureName), RealTextureName(InRealTextureName), UVLayout(InUVLayout), LayerIndex(LayerIndex)
-	{
-		ImageNode = InImageNode;
-	}
-
-	static FString GenerateId(const UCustomizableObjectNode* TypedNodeMat, int32 ImageIndex)
-	{
-		return TypedNodeMat->GetOutermost()->GetPathName() + TypedNodeMat->NodeGuid.ToString() + FString("-") + FString::FromInt(ImageIndex);
-	}
-};
-
 
 /** Struct used to store info specific to each component during compilation */
 struct FMutableComponentInfo
@@ -917,9 +889,6 @@ struct FMutableGraphGenerationContext
 	// Used to aviod Nodes with duplicated ids
 	TMap<FGuid, TArray<const UObject*>> NodeIdsMap;
 	TMultiMap<const UCustomizableObject*, FGroupNodeIdsTempData> DuplicatedGroupNodeIds;
-
-	// For a given material node (the key is node package path + node uid + image index in node) stores images generated for the same node at a higher quality LOD to reuse that image node
-	TMap<FString, FGroupProjectorImageInfo> GroupProjectorLODCache;
 	
 	uint8 FromLOD = 0; // LOD to append to the CurrentLOD when using AutomaticLODs. 
 	uint8 CurrentLOD = 0;
diff --git a/Engine/Plugins/Mutable/Source/CustomizableObjectEditor/Private/MuCOE/GenerateMutableSource/GenerateMutableSourceGroupProjector.cpp b/Engine/Plugins/Mutable/Source/CustomizableObjectEditor/Private/MuCOE/GenerateMutableSource/GenerateMutableSourceGroupProjector.cpp
--- a/Engine/Plugins/Mutable/Source/CustomizableObjectEditor/Private/MuCOE/GenerateMutableSource/GenerateMutableSourceGroupProjector.cpp
+++ b/Engine/Plugins/Mutable/Source/CustomizableObjectEditor/Private/MuCOE/GenerateMutableSource/GenerateMutableSourceGroupProjector.cpp
@@ -36,7 +36,7 @@ class UPoseAsset;
 
 
 UE::Mutable::Private::NodeImagePtr GenerateMutableSourceGroupProjector(const int32 NodeLOD, const int32 ImageIndex, UE::Mutable::Private::NodeMeshPtr MeshNode, FMutableGraphGenerationContext& GenerationContext,
-	UCONodeSkeletalMeshSectionBase* TypedNodeMat, UCONodeModifierExtendSkeletalMeshSection* TypedNodeExt, bool& bShareProjectionTexturesBetweenLODs, bool& bIsGroupProjectorImage,
+	UCONodeSkeletalMeshSectionBase* TypedNodeMat, UCONodeModifierExtendSkeletalMeshSection* TypedNodeExt, bool& bIsGroupProjectorImage,
 	UTexture2D*& GroupProjectionReferenceTexture)
 {
 	check(static_cast<bool>(TypedNodeMat) != static_cast<bool>(TypedNodeExt)); // XOr. TypedNodeMat valid or TypedNodeExt valid. At least one valid.
@@ -69,8 +69,6 @@ UE::Mutable::Private::NodeImagePtr GenerateMutableSourceGroupProjector(const int
 				continue;
 			}
 
-			bShareProjectionTexturesBetweenLODs |= ProjectorTempData.CustomizableObjectNodeGroupProjectorParameter->bShareProjectionTexturesBetweenLODs;
-
 			if (!GroupProjectionReferenceTexture)
 			{
 				GroupProjectionReferenceTexture = ProjectorTempData.CustomizableObjectNodeGroupProjectorParameter->ReferenceTexture;
diff --git a/Engine/Plugins/Mutable/Source/CustomizableObjectEditor/Private/MuCOE/GenerateMutableSource/GenerateMutableSourceGroupProjector.h b/Engine/Plugins/Mutable/Source/CustomizableObjectEditor/Private/MuCOE/GenerateMutableSource/GenerateMutableSourceGroupProjector.h
--- a/Engine/Plugins/Mutable/Source/CustomizableObjectEditor/Private/MuCOE/GenerateMutableSource/GenerateMutableSourceGroupProjector.h
+++ b/Engine/Plugins/Mutable/Source/CustomizableObjectEditor/Private/MuCOE/GenerateMutableSource/GenerateMutableSourceGroupProjector.h
@@ -18,7 +18,7 @@ struct FMutableGraphGenerationContext;
 
 
 UE::Mutable::Private::NodeImagePtr GenerateMutableSourceGroupProjector(const int32 NodeLOD, const int32 ImageIndex, UE::Mutable::Private::NodeMeshPtr MeshNode, FMutableGraphGenerationContext& GenerationContext,
-	UCONodeSkeletalMeshSectionBase* TypedNodeMat, UCONodeModifierExtendSkeletalMeshSection* TypedNodeExt, bool& bShareProjectionTexturesBetweenLODs, bool& bIsGroupProjectorImage,
+	UCONodeSkeletalMeshSectionBase* TypedNodeMat, UCONodeModifierExtendSkeletalMeshSection* TypedNodeExt, bool& bIsGroupProjectorImage,
 	UTexture2D*& GroupProjectionReferenceTexture);
 
 
diff --git a/Engine/Plugins/Mutable/Source/CustomizableObjectEditor/Private/MuCOE/GenerateMutableSource/GenerateMutableSourceMaterial.cpp b/Engine/Plugins/Mutable/Source/CustomizableObjectEditor/Private/MuCOE/GenerateMutableSource/GenerateMutableSourceMaterial.cpp
--- a/Engine/Plugins/Mutable/Source/CustomizableObjectEditor/Private/MuCOE/GenerateMutableSource/GenerateMutableSourceMaterial.cpp
+++ b/Engine/Plugins/Mutable/Source/CustomizableObjectEditor/Private/MuCOE/GenerateMutableSource/GenerateMutableSourceMaterial.cpp
@@ -199,165 +199,137 @@ UE::Mutable::Private::Ptr<UE::Mutable::Private::NodeMaterial> GenerateMutableSou
 			}
 			else
 			{
-				const FString MaterialImageId = FGroupProjectorImageInfo::GenerateId(TypedNodeMaterialBase, ImageIndex);
-				const FGroupProjectorImageInfo* ProjectorInfo = GenerationContext.GroupProjectorLODCache.Find(MaterialImageId);
-
-				if (!ProjectorInfo) // No previous LOD of this material generated the image.
-				{
-					bool bIsGroupProjectorImage = false;
-					bool bShareProjectionTexturesBetweenLODs = false;
+				bool bIsGroupProjectorImage = false;
 					
-					UTexture2D* GroupProjectionReferenceTexture = nullptr;
+				UTexture2D* GroupProjectionReferenceTexture = nullptr;
 
-					UE::Mutable::Private::NodeImagePtr GroupProjectionImg = GenerateMutableSourceGroupProjector(GenerationContext.CurrentLOD, ImageIndex, MaterialOptions.SurfaceNode->Mesh, GenerationContext,
-						TypedNodeMaterialBase, nullptr, bShareProjectionTexturesBetweenLODs, bIsGroupProjectorImage,
-						GroupProjectionReferenceTexture);
+				UE::Mutable::Private::NodeImagePtr GroupProjectionImg = GenerateMutableSourceGroupProjector(GenerationContext.CurrentLOD, ImageIndex, MaterialOptions.SurfaceNode->Mesh, GenerationContext,
+					TypedNodeMaterialBase, nullptr, bIsGroupProjectorImage,
+					GroupProjectionReferenceTexture);
 
-					if (GroupProjectionImg.get() || TypedNodeMaterialBase->IsImageMutableMode(ImageIndex))
+				if (GroupProjectionImg.get() || TypedNodeMaterialBase->IsImageMutableMode(ImageIndex))
+				{
+					// Get the reference texture
+					UTexture2D* ReferenceTexture = nullptr;
 					{
-						// Get the reference texture
-						UTexture2D* ReferenceTexture = nullptr;
-						{
-							const FNodeMaterialParameterId ImageId = TypedNodeMaterialBase->GetParameterId(EMaterialParameterType::Texture, ImageIndex);
-							
-							//TODO(Max) UE-220247: Add support for multilayer materials
-							GenerationContext.CurrentMaterialTableParameter = ImageName.ToString();
-							GenerationContext.CurrentMaterialParameterId = ImageId.ParameterId.ToString();
+						const FNodeMaterialParameterId ImageId = TypedNodeMaterialBase->GetParameterId(EMaterialParameterType::Texture, ImageIndex);
+						
+						//TODO(Max) UE-220247: Add support for multilayer materials
+						GenerationContext.CurrentMaterialTableParameter = ImageName.ToString();
+						GenerationContext.CurrentMaterialParameterId = ImageId.ParameterId.ToString();
 
-							ReferenceTexture = GroupProjectionImg.get() ? GroupProjectionReferenceTexture : nullptr;
+						ReferenceTexture = GroupProjectionImg.get() ? GroupProjectionReferenceTexture : nullptr;
 
-							if (!ReferenceTexture)
-							{
-								ReferenceTexture = TypedNodeMaterialBase->GetImageReferenceTexture(ImageIndex);
-							}
+						if (!ReferenceTexture)
+						{
+							ReferenceTexture = TypedNodeMaterialBase->GetImageReferenceTexture(ImageIndex);
+						}
 
-							// In case of group projector, don't follow the pin to find the reference texture.
-							if (!GroupProjectionImg.get() && !ReferenceTexture && ImagePin)
+						// In case of group projector, don't follow the pin to find the reference texture.
+						if (!GroupProjectionImg.get() && !ReferenceTexture && ImagePin)
+						{
+							if (const UEdGraphPin* ConnectedPin = FollowInputPin(*ImagePin))
 							{
-								if (const UEdGraphPin* ConnectedPin = FollowInputPin(*ImagePin))
-								{
-									ReferenceTexture = FindReferenceImage(ConnectedPin, GenerationContext);
-								}
+								ReferenceTexture = FindReferenceImage(ConnectedPin, GenerationContext);
 							}
+						}
 
-							if (!ReferenceTexture &&
-								ConnectedMaterialAssetPin && Cast<UCustomizableObjectNodeTable>(ConnectedMaterialAssetPin->GetOwningNode()) != nullptr)
-							{
-								ReferenceTexture = FindReferenceImage(ConnectedMaterialAssetPin, GenerationContext);
-							}
+						if (!ReferenceTexture &&
+							ConnectedMaterialAssetPin && Cast<UCustomizableObjectNodeTable>(ConnectedMaterialAssetPin->GetOwningNode()) != nullptr)
+						{
+							ReferenceTexture = FindReferenceImage(ConnectedMaterialAssetPin, GenerationContext);
+						}
 
-							if (!ReferenceTexture)
-							{
-								ReferenceTexture = TypedNodeMaterialBase->GetImageValue(ImageIndex);
-							}
+						if (!ReferenceTexture)
+						{
+							ReferenceTexture = TypedNodeMaterialBase->GetImageValue(ImageIndex);
 						}
+					}
 
-						// Generate image key to identify it
-						UE::Mutable::Private::FParameterKey ImageKey = { ImageName, LayerIndex };
-						FGeneratedImageProperties Props = GenerateImageProperties(TypedNodeMaterialBase, ReferenceTexture, ImageKey, false, GenerationContext);
+					// Generate image key to identify it
+					UE::Mutable::Private::FParameterKey ImageKey = { ImageName, LayerIndex };
+					FGeneratedImageProperties Props = GenerateImageProperties(TypedNodeMaterialBase, ReferenceTexture, ImageKey, false, GenerationContext);
 
-						// Generate the texture nodes
-						UE::Mutable::Private::NodeImagePtr ImageNode = [&]()
+					// Generate the texture nodes
+					UE::Mutable::Private::NodeImagePtr ImageNode = [&]()
+						{
+							if (TypedNodeMaterialBase->IsImageMutableMode(ImageIndex))
 							{
-								if (TypedNodeMaterialBase->IsImageMutableMode(ImageIndex))
+								if (ImagePin)
 								{
-									if (ImagePin)
+									if (const UEdGraphPin* ConnectedPin = FollowInputPin(*ImagePin))
 									{
-										if (const UEdGraphPin* ConnectedPin = FollowInputPin(*ImagePin))
-										{
-											return GenerateMutableSourceImage(ConnectedPin, GenerationContext, Props.TextureSize);
-										}
+										return GenerateMutableSourceImage(ConnectedPin, GenerationContext, Props.TextureSize);
 									}
+								}
 
-									// If the table material pin is linked to a table node, get all the textures of the current material parameter (CurrentMaterialTableParameter) from the Material Instances of the specified data table column.
-									// Then Generate a mutable table column with all these textures.
-									if (ConnectedMaterialAssetPin && Cast<UCustomizableObjectNodeTable>(ConnectedMaterialAssetPin->GetOwningNode()) != nullptr)
-									{
-										return GenerateMutableSourceImage(ConnectedMaterialAssetPin, GenerationContext, Props.TextureSize);
-									}
+								// If the table material pin is linked to a table node, get all the textures of the current material parameter (CurrentMaterialTableParameter) from the Material Instances of the specified data table column.
+								// Then Generate a mutable table column with all these textures.
+								if (ConnectedMaterialAssetPin && Cast<UCustomizableObjectNodeTable>(ConnectedMaterialAssetPin->GetOwningNode()) != nullptr)
+								{
+									return GenerateMutableSourceImage(ConnectedMaterialAssetPin, GenerationContext, Props.TextureSize);
+								}
 
-									// Else
-									{
-										UTexture2D* Texture2D = TypedNodeMaterialBase->GetImageValue(ImageIndex);
+								// Else
+								{
+									UTexture2D* Texture2D = TypedNodeMaterialBase->GetImageValue(ImageIndex);
 
-										if (Texture2D)
-										{
-											UE::Mutable::Private::Ptr<UE::Mutable::Private::NodeImageConstant> ConstImageNode = new UE::Mutable::Private::NodeImageConstant();
-											ConstImageNode->Image = GenerateImageConstant(Texture2D, GenerationContext, false);
+									if (Texture2D)
+									{
+										UE::Mutable::Private::Ptr<UE::Mutable::Private::NodeImageConstant> ConstImageNode = new UE::Mutable::Private::NodeImageConstant();
+										ConstImageNode->Image = GenerateImageConstant(Texture2D, GenerationContext, false);
 
-											const uint32 MipsToSkip = ComputeLODBiasForTexture(GenerationContext, *Texture2D, nullptr, Props.TextureSize);
-											UE::Mutable::Private::Ptr<UE::Mutable::Private::NodeImage> Result = ResizeTextureByNumMips(ConstImageNode, MipsToSkip);
+										const uint32 MipsToSkip = ComputeLODBiasForTexture(GenerationContext, *Texture2D, nullptr, Props.TextureSize);
+										UE::Mutable::Private::Ptr<UE::Mutable::Private::NodeImage> Result = ResizeTextureByNumMips(ConstImageNode, MipsToSkip);
 
-											const UTextureLODSettings& TextureLODSettings = GenerationContext.CompilationContext->Options.TargetPlatform->GetTextureLODSettings();
-											const FTextureLODGroup& LODGroupInfo = TextureLODSettings.GetTextureLODGroup(Texture2D->LODGroup);
+										const UTextureLODSettings& TextureLODSettings = GenerationContext.CompilationContext->Options.TargetPlatform->GetTextureLODSettings();
+										const FTextureLODGroup& LODGroupInfo = TextureLODSettings.GetTextureLODGroup(Texture2D->LODGroup);
 
-											ConstImageNode->SourceDataDescriptor.OptionalMaxLODSize = LODGroupInfo.OptionalMaxLODSize;
-											ConstImageNode->SourceDataDescriptor.OptionalLODBias = LODGroupInfo.OptionalLODBias;
-											ConstImageNode->SourceDataDescriptor.NumNonOptionalLODs = GenerationContext.CompilationContext->Options.MinDiskMips;
+										ConstImageNode->SourceDataDescriptor.OptionalMaxLODSize = LODGroupInfo.OptionalMaxLODSize;
+										ConstImageNode->SourceDataDescriptor.OptionalLODBias = LODGroupInfo.OptionalLODBias;
+										ConstImageNode->SourceDataDescriptor.NumNonOptionalLODs = GenerationContext.CompilationContext->Options.MinDiskMips;
 
-											const FString TextureName = GetNameSafe(Texture2D).ToLower();
-											ConstImageNode->SourceDataDescriptor.SourceId = CityHash32(reinterpret_cast<const char*>(*TextureName), TextureName.Len() * sizeof(FString::ElementType));
+										const FString TextureName = GetNameSafe(Texture2D).ToLower();
+										ConstImageNode->SourceDataDescriptor.SourceId = CityHash32(reinterpret_cast<const char*>(*TextureName), TextureName.Len() * sizeof(FString::ElementType));
 
-											return Result;
-										}
-										else
-										{
-											return UE::Mutable::Private::NodeImagePtr();
-										}
+										return Result;
+									}
+									else
+									{
+										return UE::Mutable::Private::NodeImagePtr();
 									}
 								}
-								else
-								{
-									return UE::Mutable::Private::NodeImagePtr();
-								}
-							}();
-
-						if (GroupProjectionImg.get())
-						{
-							ImageNode = GroupProjectionImg;
-						}
-
-						// NOTE: Important step to generate a texture with the correct format and size
-						ImageNode = GenerateMipMapAndFormatNodes(Node, GenerationContext, ImageNode, ReferenceTexture);
-
-						// If we are generating an implicit component (with a passthrough mesh) we don't apply any layout.
-						bool bGeneratingImplicitComponent = GenerationContext.ComponentMeshOverride.get() != nullptr;
-
-						int32 UVLayout = -1;
-						if (!bGeneratingImplicitComponent)
-						{
-							UVLayout = TypedNodeMaterialBase->GetImageUVLayout(ImageIndex);
-						}
+							}
+							else
+							{
+								return UE::Mutable::Private::NodeImagePtr();
+							}
+						}();
 
-						// Generate image data
-						UE::Mutable::Private::FImageParameterData Image;
-						Image.ImageNode = ImageNode;
-						Image.ImagePropertyIndex = Props.ImagePropertiesIndex;
-						Image.LayoutIndex = UVLayout;
+					if (GroupProjectionImg.get())
+					{
+						ImageNode = GroupProjectionImg;
+					}
 
-						//Store the image inside the current generated material
-						BaseMaterial->ImageParameters.Add(ImageKey, Image);
+					// NOTE: Important step to generate a texture with the correct format and size
+					ImageNode = GenerateMipMapAndFormatNodes(Node, GenerationContext, ImageNode, ReferenceTexture);
 
-						if (bShareProjectionTexturesBetweenLODs && bIsGroupProjectorImage)
-						{
-							// Add to the GroupProjectorLODCache to potentially reuse this projection texture in higher LODs
-							ensure(GenerationContext.CurrentLOD == GenerationContext.FirstLODAvailable[GenerationContext.CurrentMeshComponent]);
+					// If we are generating an implicit component (with a passthrough mesh) we don't apply any layout.
+					bool bGeneratingImplicitComponent = GenerationContext.ComponentMeshOverride.get() != nullptr;
 
-							GenerationContext.GroupProjectorLODCache.Add(MaterialImageId,
-								FGroupProjectorImageInfo(ImageNode, ImageName.ToString(), ImageName.ToString(), TypedNodeMaterialBase, UVLayout, LayerIndex));
-						}
+					int32 UVLayout = -1;
+					if (!bGeneratingImplicitComponent)
+					{
+						UVLayout = TypedNodeMaterialBase->GetImageUVLayout(ImageIndex);
 					}
-				}
-				else
-				{
-					ensure(GenerationContext.CurrentLOD > GenerationContext.FirstLODAvailable[GenerationContext.CurrentMeshComponent]);
 
-					UE::Mutable::Private::FParameterKey ImageKey = { FName(ProjectorInfo->TextureName), ProjectorInfo->LayerIndex };
+					// Generate image data
 					UE::Mutable::Private::FImageParameterData Image;
+					Image.ImageNode = ImageNode;
+					Image.ImagePropertyIndex = Props.ImagePropertiesIndex;
+					Image.LayoutIndex = UVLayout;
 
-					Image.ImageNode = ProjectorInfo->ImageNode;
-					Image.LayoutIndex = ProjectorInfo->UVLayout;
-
+					//Store the image inside the current generated material
 					BaseMaterial->ImageParameters.Add(ImageKey, Image);
 				}
 			}
diff --git a/Engine/Plugins/Mutable/Source/CustomizableObjectEditor/Private/MuCOE/GenerateMutableSource/GenerateMutableSourceModifier.cpp b/Engine/Plugins/Mutable/Source/CustomizableObjectEditor/Private/MuCOE/GenerateMutableSource/GenerateMutableSourceModifier.cpp
--- a/Engine/Plugins/Mutable/Source/CustomizableObjectEditor/Private/MuCOE/GenerateMutableSource/GenerateMutableSourceModifier.cpp
+++ b/Engine/Plugins/Mutable/Source/CustomizableObjectEditor/Private/MuCOE/GenerateMutableSource/GenerateMutableSourceModifier.cpp
@@ -326,25 +326,11 @@ UE::Mutable::Private::Ptr<UE::Mutable::Private::NodeModifier> GenerateMutableSou
 
 				if (!ImageNode) // If
 				{
-					const FString MaterialImageId = FGroupProjectorImageInfo::GenerateId(TypedNodeExt, ImageIndex);
-					const FGroupProjectorImageInfo* ProjectorInfo = GenerationContext.GroupProjectorLODCache.Find(MaterialImageId);
-
-					if (ProjectorInfo)
-					{
-						ensure(LOD > GenerationContext.FirstLODAvailable[GenerationContext.CurrentMeshComponent]);
-						//check(ProjectorInfo->SurfNode->Images[ImageIndex].Image == ProjectorInfo->ImageNode); TODO(Max): Review this
-						ImageNode = ProjectorInfo->ImageNode;
-					}
-				}
-
-				if (!ImageNode) // Else if
-				{
-					bool bShareProjectionTexturesBetweenLODs = false;
 					bool bIsGroupProjectorImage = false;
 					UTexture2D* GroupProjectionReferenceTexture = nullptr;
 
 					ImageNode = GenerateMutableSourceGroupProjector(LOD, ImageIndex, AddMeshNode, GenerationContext,
-						nullptr, TypedNodeExt, bShareProjectionTexturesBetweenLODs, bIsGroupProjectorImage,
+						nullptr, TypedNodeExt, bIsGroupProjectorImage,
 						GroupProjectionReferenceTexture);
 				}
 
diff --git a/Engine/Plugins/Mutable/Source/CustomizableObjectEditor/Private/MuCOE/Nodes/CustomizableObjectNodeGroupProjectorParameter.h b/Engine/Plugins/Mutable/Source/CustomizableObjectEditor/Private/MuCOE/Nodes/CustomizableObjectNodeGroupProjectorParameter.h
--- a/Engine/Plugins/Mutable/Source/CustomizableObjectEditor/Private/MuCOE/Nodes/CustomizableObjectNodeGroupProjectorParameter.h
+++ b/Engine/Plugins/Mutable/Source/CustomizableObjectEditor/Private/MuCOE/Nodes/CustomizableObjectNodeGroupProjectorParameter.h
@@ -61,10 +61,7 @@ public:
 
 	UPROPERTY(EditAnywhere, Category = ProjectorGroup, Meta = (ToolTip = "Specifies at which LOD level the projection texture will not be used and possibly save memory. A negative value means they will never be dropped."))
 	int32 DropProjectionTextureAtLOD = -1;
-
-	UPROPERTY(EditAnywhere, Category = ProjectorGroup, Meta = (ToolTip = "If true, projection textures will be shared between LODs of the same object, and will save memory. Only use if all the LODs share the same UV layout."))
-	bool bShareProjectionTexturesBetweenLODs = false;
-
+	
 	UPROPERTY(EditAnywhere, Category = ProjectorGroup)
 	TArray<FGroupProjectorParameterImage> OptionTextures;
 
