[Mutable] Renamed missing Node Material types to Node Skeletal Mesh Section.

#Synced-CL 49271159

diff --git a/Engine/Plugins/Mutable/Config/DefaultMutable.ini b/Engine/Plugins/Mutable/Config/DefaultMutable.ini
--- a/Engine/Plugins/Mutable/Config/DefaultMutable.ini
+++ b/Engine/Plugins/Mutable/Config/DefaultMutable.ini
@@ -42,6 +42,9 @@
 +ClassRedirects=(OldName="/Script/CustomizableObjectEditor.CustomizableObjectNodeModifierEditSkeletalMeshSection",NewName="/Script/CustomizableObjectEditor.CONodeModifierEditSkeletalMeshSection")
 +ClassRedirects=(OldName="/Script/CustomizableObjectEditor.CustomizableObjectNodeCopyMaterial",NewName="/Script/CustomizableObjectEditor.CONodeCopySkeletalMeshSection")
 +ClassRedirects=(OldName="/Script/CustomizableObjectEditor.CustomizableObjectNodeMaterialVariation",NewName="/Script/CustomizableObjectEditor.CONodeSkeletalMeshSectionVariation")
++ClassRedirects=(OldName="/Script/CustomizableObjectEditor.CustomizableObjectNodeMaterialPinDataImage",NewName="/Script/CustomizableObjectEditor.CONodeSkeletalMeshSectionPinDataImage")
++ClassRedirects=(OldName="/Script/CustomizableObjectEditor.CustomizableObjectNodeMaterialPinDataVector",NewName="/Script/CustomizableObjectEditor.CONodeSkeletalMeshSectionPinDataVector")
++ClassRedirects=(OldName="/Script/CustomizableObjectEditor.CustomizableObjectNodeMaterialPinDataScalar",NewName="/Script/CustomizableObjectEditor.CONodeSkeletalMeshSectionPinDataScalar")
 
 ; Settings for the Mutable plugin
 [Features]
diff --git a/Engine/Plugins/Mutable/Source/CustomizableObjectEditor/Private/MuCOE/Nodes/CONodeCopySkeletalMeshSection.cpp b/Engine/Plugins/Mutable/Source/CustomizableObjectEditor/Private/MuCOE/Nodes/CONodeCopySkeletalMeshSection.cpp
--- a/Engine/Plugins/Mutable/Source/CustomizableObjectEditor/Private/MuCOE/Nodes/CONodeCopySkeletalMeshSection.cpp
+++ b/Engine/Plugins/Mutable/Source/CustomizableObjectEditor/Private/MuCOE/Nodes/CONodeCopySkeletalMeshSection.cpp
@@ -433,7 +433,7 @@ int32 UCONodeCopySkeletalMeshSection::GetImageUVLayout(int32 ImageIndex) const
 	}
 	else
 	{
-		return UCustomizableObjectNodeMaterialPinDataImage::UV_LAYOUT_IGNORE;
+		return UCONodeSkeletalMeshSectionPinDataImage::UV_LAYOUT_IGNORE;
 	}
 }
 
diff --git a/Engine/Plugins/Mutable/Source/CustomizableObjectEditor/Private/MuCOE/Nodes/CONodeMaterialModify.cpp b/Engine/Plugins/Mutable/Source/CustomizableObjectEditor/Private/MuCOE/Nodes/CONodeMaterialModify.cpp
--- a/Engine/Plugins/Mutable/Source/CustomizableObjectEditor/Private/MuCOE/Nodes/CONodeMaterialModify.cpp
+++ b/Engine/Plugins/Mutable/Source/CustomizableObjectEditor/Private/MuCOE/Nodes/CONodeMaterialModify.cpp
@@ -222,14 +222,14 @@ int32 UCONodeMaterialModify::GetImagePinParameterUVIndex(const UEdGraphPin& Pin)
 	switch(TextureParameterPnData.UVLayoutMode)
 	{
 	case EMaterialModifyUVLayoutMode::Ignore:
-		return UCustomizableObjectNodeMaterialPinDataImage::UV_LAYOUT_IGNORE;
+		return UCONodeSkeletalMeshSectionPinDataImage::UV_LAYOUT_IGNORE;
 	case EMaterialModifyUVLayoutMode::Index:
 		return TextureParameterPnData.UVLayout;
 	default:
 		unimplemented();
 	}
 		
-	return UCustomizableObjectNodeMaterialPinDataImage::UV_LAYOUT_IGNORE;
+	return UCONodeSkeletalMeshSectionPinDataImage::UV_LAYOUT_IGNORE;
 }
 
 
diff --git a/Engine/Plugins/Mutable/Source/CustomizableObjectEditor/Private/MuCOE/Nodes/CONodeSkeletalMeshSection.cpp b/Engine/Plugins/Mutable/Source/CustomizableObjectEditor/Private/MuCOE/Nodes/CONodeSkeletalMeshSection.cpp
--- a/Engine/Plugins/Mutable/Source/CustomizableObjectEditor/Private/MuCOE/Nodes/CONodeSkeletalMeshSection.cpp
+++ b/Engine/Plugins/Mutable/Source/CustomizableObjectEditor/Private/MuCOE/Nodes/CONodeSkeletalMeshSection.cpp
@@ -105,9 +105,9 @@ bool UCustomizableObjectNodeMaterialRemapPinsByName::HasSavedPinData(const UCust
 }
 
 
-bool UCustomizableObjectNodeMaterialPinDataImage::IsDefault() const
+bool UCONodeSkeletalMeshSectionPinDataImage::IsDefault() const
 {
-	const UCustomizableObjectNodeMaterialPinDataImage* Default = Cast<UCustomizableObjectNodeMaterialPinDataImage>(GetClass()->GetDefaultObject(false));
+	const UCONodeSkeletalMeshSectionPinDataImage* Default = Cast<UCONodeSkeletalMeshSectionPinDataImage>(GetClass()->GetDefaultObject(false));
 
 	return PinMode == Default->PinMode &&
 		UVLayoutMode == Default->UVLayoutMode &&
@@ -115,28 +115,28 @@ bool UCustomizableObjectNodeMaterialPinDataImage::IsDefault() const
 }
 
 
-EPinMode UCustomizableObjectNodeMaterialPinDataImage::GetPinMode() const
+EPinMode UCONodeSkeletalMeshSectionPinDataImage::GetPinMode() const
 {
 	return PinMode;
 }
 
 
-void UCustomizableObjectNodeMaterialPinDataImage::SetPinMode(const EPinMode InPinMode)
+void UCONodeSkeletalMeshSectionPinDataImage::SetPinMode(const EPinMode InPinMode)
 {
-	FObjectEditorUtils::SetPropertyValue(this, GET_MEMBER_NAME_CHECKED(UCustomizableObjectNodeMaterialPinDataImage, PinMode), InPinMode);
+	FObjectEditorUtils::SetPropertyValue(this, GET_MEMBER_NAME_CHECKED(UCONodeSkeletalMeshSectionPinDataImage, PinMode), InPinMode);
 	NodeMaterial->GetPostImagePinModeChangedDelegate()->Broadcast();
 }
 
 
-void UCustomizableObjectNodeMaterialPinDataImage::Init(UCONodeSkeletalMeshSection& InNodeMaterial)
+void UCONodeSkeletalMeshSectionPinDataImage::Init(UCONodeSkeletalMeshSection& InNodeMaterial)
 {
 	NodeMaterial = &InNodeMaterial;
 }
 
 
-void UCustomizableObjectNodeMaterialPinDataImage::Copy(const UCustomizableObjectNodePinData& Other)
+void UCONodeSkeletalMeshSectionPinDataImage::Copy(const UCustomizableObjectNodePinData& Other)
 {
-	if (const UCustomizableObjectNodeMaterialPinDataImage* PinDataOldPin = Cast<UCustomizableObjectNodeMaterialPinDataImage>(&Other))
+	if (const UCONodeSkeletalMeshSectionPinDataImage* PinDataOldPin = Cast<UCONodeSkeletalMeshSectionPinDataImage>(&Other))
 	{
 		PinMode = PinDataOldPin->PinMode;
 		UVLayoutMode = PinDataOldPin->UVLayoutMode;
@@ -205,7 +205,7 @@ FName UCONodeSkeletalMeshSection::GetImagePinMode(const EPinMode PinMode) const
 
 FName UCONodeSkeletalMeshSection::GetImagePinMode(const UEdGraphPin& Pin) const
 {
-	return GetImagePinMode(GetPinData<UCustomizableObjectNodeMaterialPinDataImage>(Pin).GetPinMode());
+	return GetImagePinMode(GetPinData<UCONodeSkeletalMeshSectionPinDataImage>(Pin).GetPinMode());
 }
 
 
@@ -269,13 +269,13 @@ FText EPinModeToText(const EPinMode PinMode)
 }
 
 
-void UCustomizableObjectNodeMaterialPinDataImage::PostEditChangeProperty(FPropertyChangedEvent& PropertyChangedEvent)
+void UCONodeSkeletalMeshSectionPinDataImage::PostEditChangeProperty(FPropertyChangedEvent& PropertyChangedEvent)
 {
 	Super::PostEditChangeProperty(PropertyChangedEvent);
 
 	if (const FProperty* PropertyThatChanged = PropertyChangedEvent.Property)
 	{
-		if (PropertyThatChanged->GetFName() == GET_MEMBER_NAME_CHECKED(UCustomizableObjectNodeMaterialPinDataImage, PinMode))
+		if (PropertyThatChanged->GetFName() == GET_MEMBER_NAME_CHECKED(UCONodeSkeletalMeshSectionPinDataImage, PinMode))
 		{
 			NodeMaterial->UCustomizableObjectNode::ReconstructNode();
 		}
@@ -283,9 +283,9 @@ void UCustomizableObjectNodeMaterialPinDataImage::PostEditChangeProperty(FProper
 }
 
 
-bool UCustomizableObjectNodeMaterialPinDataImage::CanEditChange(const FProperty* InProperty) const
+bool UCONodeSkeletalMeshSectionPinDataImage::CanEditChange(const FProperty* InProperty) const
 {
-	if (InProperty->GetFName() == GET_MEMBER_NAME_CHECKED(UCustomizableObjectNodeMaterialPinDataImage, PinMode))
+	if (InProperty->GetFName() == GET_MEMBER_NAME_CHECKED(UCONodeSkeletalMeshSectionPinDataImage, PinMode))
 	{
 		return !NodeMaterial->GetPin(*this)->LinkedTo.Num();
 	}
@@ -294,7 +294,7 @@ bool UCustomizableObjectNodeMaterialPinDataImage::CanEditChange(const FProperty*
 }
 
 
-void UCustomizableObjectNodeMaterialPinDataImage::BackwardsCompatibleFixup(int32 CustomizableObjectCustomVersion)
+void UCONodeSkeletalMeshSectionPinDataImage::BackwardsCompatibleFixup(int32 CustomizableObjectCustomVersion)
 {
 	Super::PostLoad();
 	
@@ -403,7 +403,7 @@ void UCONodeSkeletalMeshSection::BackwardsCompatibleFixup(int32 CustomizableObje
 				continue;
 			}
 
-			UCustomizableObjectNodeMaterialPinDataImage* PinData = NewObject<UCustomizableObjectNodeMaterialPinDataImage>(this);
+			UCONodeSkeletalMeshSectionPinDataImage* PinData = NewObject<UCONodeSkeletalMeshSectionPinDataImage>(this);
 			PinData->ParameterId_DEPRECATED = FGuid::NewGuid();
 			PinData->ReferenceTexture = Image.ReferenceTexture;
 
@@ -452,7 +452,7 @@ void UCONodeSkeletalMeshSection::BackwardsCompatibleFixup(int32 CustomizableObje
 				continue;
 			}
 			
-			UCustomizableObjectNodeMaterialPinDataVector* PinData = NewObject<UCustomizableObjectNodeMaterialPinDataVector>(this);
+			UCONodeSkeletalMeshSectionPinDataVector* PinData = NewObject<UCONodeSkeletalMeshSectionPinDataVector>(this);
 			PinData->ParameterId_DEPRECATED = FGuid::NewGuid();
 			
 			// Find referenced Material Parameter
@@ -478,7 +478,7 @@ void UCONodeSkeletalMeshSection::BackwardsCompatibleFixup(int32 CustomizableObje
 				continue;
 			}
 			
-			UCustomizableObjectNodeMaterialPinDataScalar* PinData = NewObject<UCustomizableObjectNodeMaterialPinDataScalar>(this);
+			UCONodeSkeletalMeshSectionPinDataScalar* PinData = NewObject<UCONodeSkeletalMeshSectionPinDataScalar>(this);
 			PinData->ParameterId_DEPRECATED = FGuid::NewGuid();
 
 			// Find referenced Material Parameter
@@ -507,15 +507,15 @@ void UCONodeSkeletalMeshSection::BackwardsCompatibleFixup(int32 CustomizableObje
 				{
 					if (Pin->PinType.PinCategory == UEdGraphSchema_CustomizableObject::PC_Texture)
 					{
-						return NewObject<UCustomizableObjectNodeMaterialPinDataImage>(Outer);	
+						return NewObject<UCONodeSkeletalMeshSectionPinDataImage>(Outer);	
 					}
 					else if (Pin->PinType.PinCategory == UEdGraphSchema_CustomizableObject::PC_Color)
 					{
-						return NewObject<UCustomizableObjectNodeMaterialPinDataVector>(Outer);
+						return NewObject<UCONodeSkeletalMeshSectionPinDataVector>(Outer);
 					}
 					else if (Pin->PinType.PinCategory == UEdGraphSchema_CustomizableObject::PC_Float)
 					{
-						return NewObject<UCustomizableObjectNodeMaterialPinDataScalar>(Outer);	
+						return NewObject<UCONodeSkeletalMeshSectionPinDataScalar>(Outer);	
 					}
 					else
 					{
@@ -563,7 +563,7 @@ void UCONodeSkeletalMeshSection::BackwardsCompatibleFixup(int32 CustomizableObje
 		{
 			if (const UEdGraphPin* Pin = GetParameterPin(EMaterialParameterType::Texture, ImageIndex))
 			{
-				UCustomizableObjectNodeMaterialPinDataImage* PinDataImage = Cast<UCustomizableObjectNodeMaterialPinDataImage>(GetPinData(*Pin));
+				UCONodeSkeletalMeshSectionPinDataImage* PinDataImage = Cast<UCONodeSkeletalMeshSectionPinDataImage>(GetPinData(*Pin));
 				PinDataImage->NodeMaterial = this;
 			}
 		}
@@ -573,7 +573,7 @@ void UCONodeSkeletalMeshSection::BackwardsCompatibleFixup(int32 CustomizableObje
 	{
 		for (const UEdGraphPin* Pin : GetAllPins())
 		{
-			if (UCustomizableObjectNodeMaterialPinDataImage* const PinDataImage = Cast<UCustomizableObjectNodeMaterialPinDataImage>(GetPinData(*Pin)))
+			if (UCONodeSkeletalMeshSectionPinDataImage* const PinDataImage = Cast<UCONodeSkeletalMeshSectionPinDataImage>(GetPinData(*Pin)))
 			{
 				PinDataImage->NodeMaterial = this;
 			}
@@ -639,7 +639,7 @@ void UCONodeSkeletalMeshSection::BackwardsCompatibleFixup(int32 CustomizableObje
 					{
 						if (TextureName == Texture.ParameterInfo.Name)
 						{
-							if (UCustomizableObjectNodeMaterialPinDataImage* PinDataImage = Cast<UCustomizableObjectNodeMaterialPinDataImage>(GetPinData(*ImagePin)))
+							if (UCONodeSkeletalMeshSectionPinDataImage* PinDataImage = Cast<UCONodeSkeletalMeshSectionPinDataImage>(GetPinData(*ImagePin)))
 							{
 								PinDataImage->PinMode = EPinMode::Mutable;
 							}
@@ -670,7 +670,7 @@ void UCONodeSkeletalMeshSection::BackwardsCompatibleFixup(int32 CustomizableObje
 			{
 				const UEdGraphPin* LinkedPin = Pin->LinkedTo[0];
 
-				UCustomizableObjectNodeMaterialPinDataImage& PinData = GetPinData<UCustomizableObjectNodeMaterialPinDataImage>(*Pin);
+				UCONodeSkeletalMeshSectionPinDataImage& PinData = GetPinData<UCONodeSkeletalMeshSectionPinDataImage>(*Pin);
 
 				if (LinkedPin->PinType.PinCategory == UEdGraphSchema_CustomizableObject::PC_Texture)
 				{
@@ -958,7 +958,7 @@ bool UCONodeSkeletalMeshSection::IsImageMutableMode(const UEdGraphPin& Pin) cons
 
 void UCONodeSkeletalMeshSection::SetImagePinMode(UEdGraphPin& Pin, const EPinMode PinMode) const
 {
-	UCustomizableObjectNodeMaterialPinDataImage& PinData = GetPinData<UCustomizableObjectNodeMaterialPinDataImage>(Pin);
+	UCONodeSkeletalMeshSectionPinDataImage& PinData = GetPinData<UCONodeSkeletalMeshSectionPinDataImage>(Pin);
 
 	Pin.PinType.PinCategory = GetImagePinMode(PinMode);	// Change the category so that it will remap correctly when reconstructing.
 	
@@ -970,7 +970,7 @@ UTexture2D* UCONodeSkeletalMeshSection::GetImageReferenceTexture(const int32 Ima
 {
 	if (const UEdGraphPin* Pin = GetParameterPin(EMaterialParameterType::Texture, ImageIndex))
 	{
-		return GetPinData<UCustomizableObjectNodeMaterialPinDataImage>(*Pin).ReferenceTexture;
+		return GetPinData<UCONodeSkeletalMeshSectionPinDataImage>(*Pin).ReferenceTexture;
 	}
 	else
 	{
@@ -1026,13 +1026,13 @@ int32 UCONodeSkeletalMeshSection::GetImageUVLayout(const int32 ImageIndex) const
 {
 	if (const UEdGraphPin* Pin = GetParameterPin(EMaterialParameterType::Texture, ImageIndex))
 	{
-		const UCustomizableObjectNodeMaterialPinDataImage& PinData = GetPinData<UCustomizableObjectNodeMaterialPinDataImage>(*Pin);
+		const UCONodeSkeletalMeshSectionPinDataImage& PinData = GetPinData<UCONodeSkeletalMeshSectionPinDataImage>(*Pin);
 		switch(PinData.UVLayoutMode)
 		{
 		case EUVLayoutMode::FromMaterial:
 			break;
 		case EUVLayoutMode::Ignore:
-			return UCustomizableObjectNodeMaterialPinDataImage::UV_LAYOUT_IGNORE;
+			return UCONodeSkeletalMeshSectionPinDataImage::UV_LAYOUT_IGNORE;
 		case EUVLayoutMode::Index:
 			return PinData.UVLayout;
 		default:
@@ -1241,7 +1241,7 @@ FString UCONodeSkeletalMeshSection::GetRefreshMessage() const
 
 TSharedPtr<IDetailsView> UCONodeSkeletalMeshSection::CustomizePinDetails(const UEdGraphPin& Pin) const
 {
-	if (UCustomizableObjectNodeMaterialPinDataImage* PinData = Cast<UCustomizableObjectNodeMaterialPinDataImage>(GetPinData(Pin)))
+	if (UCONodeSkeletalMeshSectionPinDataImage* PinData = Cast<UCONodeSkeletalMeshSectionPinDataImage>(GetPinData(Pin)))
 	{
 		FPropertyEditorModule& EditModule = FModuleManager::Get().GetModuleChecked<FPropertyEditorModule>("PropertyEditor");
 
@@ -1393,7 +1393,7 @@ UCONodeSkeletalMeshSectionPinDataParameter* UCONodeSkeletalMeshSection::CreatePi
 	{
 	case EMaterialParameterType::Texture:
 		{
-			UCustomizableObjectNodeMaterialPinDataImage* PinDataImage = NewObject<UCustomizableObjectNodeMaterialPinDataImage>(this);
+			UCONodeSkeletalMeshSectionPinDataImage* PinDataImage = NewObject<UCONodeSkeletalMeshSectionPinDataImage>(this);
 			PinDataImage->Init(*this);
 		
 			PinData = PinDataImage;
@@ -1401,13 +1401,13 @@ UCONodeSkeletalMeshSectionPinDataParameter* UCONodeSkeletalMeshSection::CreatePi
 		}
 	case EMaterialParameterType::Vector:
 		{
-			PinData = NewObject<UCustomizableObjectNodeMaterialPinDataVector>(this);	
+			PinData = NewObject<UCONodeSkeletalMeshSectionPinDataVector>(this);	
 			break;
 		}
 
 	case EMaterialParameterType::Scalar:
 		{
-			PinData = NewObject<UCustomizableObjectNodeMaterialPinDataScalar>(this);	
+			PinData = NewObject<UCONodeSkeletalMeshSectionPinDataScalar>(this);	
 			break;
 		}
 
@@ -1517,7 +1517,7 @@ void UCONodeSkeletalMeshSection::PinConnectionListChanged(UEdGraphPin* Pin)
 			}
 		}
 	}
-	else if (Cast<UCustomizableObjectNodeMaterialPinDataImage>(GetPinData(*Pin))) // Image pin
+	else if (Cast<UCONodeSkeletalMeshSectionPinDataImage>(GetPinData(*Pin))) // Image pin
 	{
 		// If necessary, automatically change the Pin Mode. Connected pin can never change its type.
 		if (Pin->LinkedTo.Num())
diff --git a/Engine/Plugins/Mutable/Source/CustomizableObjectEditor/Private/MuCOE/Nodes/CONodeSkeletalMeshSection.h b/Engine/Plugins/Mutable/Source/CustomizableObjectEditor/Private/MuCOE/Nodes/CONodeSkeletalMeshSection.h
--- a/Engine/Plugins/Mutable/Source/CustomizableObjectEditor/Private/MuCOE/Nodes/CONodeSkeletalMeshSection.h
+++ b/Engine/Plugins/Mutable/Source/CustomizableObjectEditor/Private/MuCOE/Nodes/CONodeSkeletalMeshSection.h
@@ -258,7 +258,7 @@ private:
 
 /** Additional data for a Material Texture Parameter pin. */
 UCLASS(MinimalAPI)
-class UCustomizableObjectNodeMaterialPinDataImage : public UCONodeSkeletalMeshSectionPinDataParameter
+class UCONodeSkeletalMeshSectionPinDataImage : public UCONodeSkeletalMeshSectionPinDataParameter
 {
 	GENERATED_BODY()
 
@@ -316,7 +316,7 @@ private:
 
 /** Additional data for a Material Vector Parameter pin. */
 UCLASS(MinimalAPI)
-class UCustomizableObjectNodeMaterialPinDataVector : public UCONodeSkeletalMeshSectionPinDataParameter
+class UCONodeSkeletalMeshSectionPinDataVector : public UCONodeSkeletalMeshSectionPinDataParameter
 {
 	GENERATED_BODY()
 };
@@ -324,7 +324,7 @@ class UCustomizableObjectNodeMaterialPinDataVector : public UCONodeSkeletalMeshS
 
 /** Additional data for a Material Float Parameter pin. */
 UCLASS(MinimalAPI)
-class UCustomizableObjectNodeMaterialPinDataScalar : public UCONodeSkeletalMeshSectionPinDataParameter
+class UCONodeSkeletalMeshSectionPinDataScalar : public UCONodeSkeletalMeshSectionPinDataParameter
 {
 	GENERATED_BODY()
 };
diff --git a/Engine/Plugins/Mutable/Source/CustomizableObjectEditor/Private/MuCOE/Nodes/CONodeSkeletalMeshSectionBase.h b/Engine/Plugins/Mutable/Source/CustomizableObjectEditor/Private/MuCOE/Nodes/CONodeSkeletalMeshSectionBase.h
--- a/Engine/Plugins/Mutable/Source/CustomizableObjectEditor/Private/MuCOE/Nodes/CONodeSkeletalMeshSectionBase.h
+++ b/Engine/Plugins/Mutable/Source/CustomizableObjectEditor/Private/MuCOE/Nodes/CONodeSkeletalMeshSectionBase.h
@@ -8,6 +8,7 @@
 
 #define UE_API CUSTOMIZABLEOBJECTEDITOR_API
 
+enum class EMaterialParameterType : uint8;
 class UCustomizableObjectLayout;
 class UCONodeSkeletalMeshSection;
 class UEdGraphPin;
@@ -56,7 +57,7 @@ public:
 	GENERATED_BODY()
 	
 	// Own interface
-	virtual const UCONodeSkeletalMeshSection* GetMaterialNode() const PURE_VIRTUAL(UCustomizableObjectNodeMaterial::GetMaterialNode, return {}; );
+	virtual const UCONodeSkeletalMeshSection* GetMaterialNode() const PURE_VIRTUAL(UCONodeSkeletalMeshSection::GetMaterialNode, return {}; );
 
 	UCONodeSkeletalMeshSection* GetMaterialNode()
 	{
@@ -67,56 +68,56 @@ public:
 	// Max LOD to propagate this section to, when using automatic LODs.
 	UE_API int32 GetMeshSectionMaxLOD() const;
 
-	virtual TArray<UCustomizableObjectLayout*> GetLayouts() const PURE_VIRTUAL(UCustomizableObjectNodeMaterial::GetLayouts, return {}; );
+	virtual TArray<UCustomizableObjectLayout*> GetLayouts() const PURE_VIRTUAL(UCONodeSkeletalMeshSection::GetLayouts, return {}; );
 	
-	virtual UMaterialInterface* GetMaterial() const PURE_VIRTUAL(UCustomizableObjectNodeMaterial::GetMaterial, return {}; );
+	virtual UMaterialInterface* GetMaterial() const PURE_VIRTUAL(UCONodeSkeletalMeshSection::GetMaterial, return {}; );
 	
-	virtual UEdGraphPin* GetMeshPin() const PURE_VIRTUAL(UCustomizableObjectNodeMaterial::GetMeshPin, return {}; );
+	virtual UEdGraphPin* GetMeshPin() const PURE_VIRTUAL(UCONodeSkeletalMeshSection::GetMeshPin, return {}; );
 
-	virtual UEdGraphPin* GetMaterialAssetPin() const PURE_VIRTUAL(UCustomizableObjectNodeMaterial::GetMaterialAssetPin, return {}; );
+	virtual UEdGraphPin* GetMaterialAssetPin() const PURE_VIRTUAL(UCONodeSkeletalMeshSection::GetMaterialAssetPin, return {}; );
 
-	virtual UEdGraphPin* GetEnableTagsPin() const PURE_VIRTUAL(UCustomizableObjectNodeMaterial::GetMaterialAssetPin, return {}; );
+	virtual UEdGraphPin* GetEnableTagsPin() const PURE_VIRTUAL(UCONodeSkeletalMeshSection::GetMaterialAssetPin, return {}; );
 	
 	/** Returns the number of Material Parameters. */
-	virtual int32 GetNumParameters(EMaterialParameterType Type) const PURE_VIRTUAL(UCustomizableObjectNodeMaterial::GetNumParameters, return {}; );
+	virtual int32 GetNumParameters(EMaterialParameterType Type) const PURE_VIRTUAL(UCONodeSkeletalMeshSection::GetNumParameters, return {}; );
 
 	/** Returns the Material Parameter id.
 	 *
 	 * @param ParameterIndex Have to be valid. */
-	virtual FNodeMaterialParameterId GetParameterId(EMaterialParameterType Type, int32 ParameterIndex) const PURE_VIRTUAL(UCustomizableObjectNodeMaterial::GetParameterId, return {}; );
+	virtual FNodeMaterialParameterId GetParameterId(EMaterialParameterType Type, int32 ParameterIndex) const PURE_VIRTUAL(UCONodeSkeletalMeshSection::GetParameterId, return {}; );
 	
 	/** Returns the Material Parameter name.
 	 *
 	 * @param ParameterIndex Have to be valid. */
-	virtual FName GetParameterName(EMaterialParameterType Type, int32 ParameterIndex) const PURE_VIRTUAL(UCustomizableObjectNodeMaterial::GetParameterName, return {}; );
+	virtual FName GetParameterName(EMaterialParameterType Type, int32 ParameterIndex) const PURE_VIRTUAL(UCONodeSkeletalMeshSection::GetParameterName, return {}; );
 	
 	/** Get the Material Parameter layer index.
 	 *
 	 * @param ParameterIndex Have to be valid.
 	 * @returns INDEX_NONE for global parameters. */
-	virtual int32 GetParameterLayerIndex(EMaterialParameterType Type, int32 ParameterIndex) const PURE_VIRTUAL(UCustomizableObjectNodeMaterial::GetParameterLayerIndex, return {}; );
+	virtual int32 GetParameterLayerIndex(EMaterialParameterType Type, int32 ParameterIndex) const PURE_VIRTUAL(UCONodeSkeletalMeshSection::GetParameterLayerIndex, return {}; );
 	
 	/** Get the Material Parameter layer name.
 	 *
 	 * @param ParameterIndex Have to be valid. */
-	virtual FText GetParameterLayerName(EMaterialParameterType Type, int32 ParameterIndex) const PURE_VIRTUAL(UCustomizableObjectNodeMaterial::GetParameterLayerName, return {}; );
+	virtual FText GetParameterLayerName(EMaterialParameterType Type, int32 ParameterIndex) const PURE_VIRTUAL(UCONodeSkeletalMeshSection::GetParameterLayerName, return {}; );
 
 	/** Returns true if the Material contains the given Material Parameter. */
-	virtual bool HasParameter(const FNodeMaterialParameterId& ParameterId) const PURE_VIRTUAL(UCustomizableObjectNodeMaterial::HasParameter, return {}; );
+	virtual bool HasParameter(const FNodeMaterialParameterId& ParameterId) const PURE_VIRTUAL(UCONodeSkeletalMeshSection::HasParameter, return {}; );
 	
 	/** Get the pin for the given Material Parameter.
 	 * Not all parameters have pins.
 	 *
 	 * @param ParameterIndex Material Parameter id.
 	 * @return Can return nullptr. */
-	virtual const UEdGraphPin* GetParameterPin(EMaterialParameterType Type, int32 ParameterIndex) const PURE_VIRTUAL(UCustomizableObjectNodeMaterial::GetMaterialNode, return {}; );
+	virtual const UEdGraphPin* GetParameterPin(EMaterialParameterType Type, int32 ParameterIndex) const PURE_VIRTUAL(UCONodeSkeletalMeshSection::GetMaterialNode, return {}; );
 
 	/** Get the pin for the given Material Parameter.
 	 * Not all parameters have pins.
 	 *
 	 * @param ParameterId Material Parameter id.
 	 * @return Can return nullptr. */
-	virtual UEdGraphPin* GetParameterPin(const FNodeMaterialParameterId& ParameterId) const PURE_VIRTUAL(UCustomizableObjectNodeMaterial::GetParameterPin, return {}; );
+	virtual UEdGraphPin* GetParameterPin(const FNodeMaterialParameterId& ParameterId) const PURE_VIRTUAL(UCONodeSkeletalMeshSection::GetParameterPin, return {}; );
 	
 	// --------------------
 	// IMAGES PARAMETERS
@@ -125,35 +126,35 @@ public:
 	/** Returns true if the Material Texture Parameter goes through Mutable.
 	 *
 	 * @param ImageIndex Have to be valid. */
-	virtual bool IsImageMutableMode(int32 ImageIndex) const PURE_VIRTUAL(UCustomizableObjectNodeMaterial::IsImageMutableMode, return {}; );
+	virtual bool IsImageMutableMode(int32 ImageIndex) const PURE_VIRTUAL(UCONodeSkeletalMeshSection::IsImageMutableMode, return {}; );
 
 	/** Given an Image pin, returns true if the Material Texture Parameter goes through Mutable. */
-	virtual bool IsImageMutableMode(const UEdGraphPin& Pin) const PURE_VIRTUAL(UCustomizableObjectNodeMaterial::IsImageMutableMode, return {}; );
+	virtual bool IsImageMutableMode(const UEdGraphPin& Pin) const PURE_VIRTUAL(UCONodeSkeletalMeshSection::IsImageMutableMode, return {}; );
 	
 	/** Returns the reference texture assigned to a Material Texture Parameter.
 	 *
 	 * @param ImageIndex Have to be valid.
 	 * @return nullptr if it does not have one assigned. */
-	virtual UTexture2D* GetImageReferenceTexture(int32 ImageIndex) const PURE_VIRTUAL(UCustomizableObjectNodeMaterial::GetImageReferenceTexture, return {}; );
+	virtual UTexture2D* GetImageReferenceTexture(int32 ImageIndex) const PURE_VIRTUAL(UCONodeSkeletalMeshSection::GetImageReferenceTexture, return {}; );
 
 	/** Returns the Texture set in the Material Texture Parameter.
 	 *
 	 * @param ImageIndex Have to be valid. */
-	virtual UTexture2D* GetImageValue(int32 ImageIndex) const PURE_VIRTUAL(UCustomizableObjectNodeMaterial::GetImageValue, return {}; );
+	virtual UTexture2D* GetImageValue(int32 ImageIndex) const PURE_VIRTUAL(UCONodeSkeletalMeshSection::GetImageValue, return {}; );
 	
 	/** Get the Material Texture Parameter UV Index.
 	 *
 	 * @param ImageIndex Have to be valid.
 	 * @return Return -1 if the UV Index is set to Ignore. Return >= 0 for a valid UV Index.  */ 
-	virtual int32 GetImageUVLayout(int32 ImageIndex) const PURE_VIRTUAL(UCustomizableObjectNodeMaterial::GetImageUVLayout, return {}; );
+	virtual int32 GetImageUVLayout(int32 ImageIndex) const PURE_VIRTUAL(UCONodeSkeletalMeshSection::GetImageUVLayout, return {}; );
 
 	// This method should be overridden in all derived classes
-	virtual UEdGraphPin* OutputPin() const PURE_VIRTUAL(UCustomizableObjectNodeMaterial::OutputPin, return {}; );
+	virtual UEdGraphPin* OutputPin() const PURE_VIRTUAL(UCONodeSkeletalMeshSection::OutputPin, return {}; );
 	
-	virtual FPostImagePinModeChangedDelegate* GetPostImagePinModeChangedDelegate() PURE_VIRTUAL(UCustomizableObjectNodeMaterial::OutputPin, return {}; );
+	virtual FPostImagePinModeChangedDelegate* GetPostImagePinModeChangedDelegate() PURE_VIRTUAL(UCONodeSkeletalMeshSection::OutputPin, return {}; );
 	
 	/** Return true if a Material Parameter has changed on which we had a pin connected or data saved. */
-	virtual bool RealMaterialDataHasChanged() const PURE_VIRTUAL(UCustomizableObjectNodeMaterial::OutputPin, return {}; );
+	virtual bool RealMaterialDataHasChanged() const PURE_VIRTUAL(UCONodeSkeletalMeshSection::OutputPin, return {}; );
 
 	// UEdGraphNode interface
 	UE_API virtual FLinearColor GetNodeTitleColor() const override;
