[mutable] Improved pointer handling at the COSystem level (GameThread)
- Improved handling of an improper execution of the COI LoadAdditionalAssetsAndData method.
- Made the execution of the FinishUpdateGlobal method in the case of having a UpdateResult::Error as update result provoke the automatic invalidation of the COI data.

#rnx
#Synced-CL 49331383

diff --git a/Engine/Plugins/Mutable/Source/CustomizableObject/Internal/MuCO/CustomizableObjectSystemPrivate.h b/Engine/Plugins/Mutable/Source/CustomizableObject/Internal/MuCO/CustomizableObjectSystemPrivate.h
--- a/Engine/Plugins/Mutable/Source/CustomizableObject/Internal/MuCO/CustomizableObjectSystemPrivate.h
+++ b/Engine/Plugins/Mutable/Source/CustomizableObject/Internal/MuCO/CustomizableObjectSystemPrivate.h
@@ -3,7 +3,6 @@
 #pragma once
 
 #include "LogBenchmarkUtil.h"
-#include "Containers/Queue.h"
 #include "MuCO/DescriptorHash.h"
 #include "MuCO/CustomizableObject.h"
 #include "MuCO/CustomizableObjectPrivate.h"
@@ -20,12 +19,8 @@
 #include "MuR/Image.h"
 #include "MuR/Material.h"
 #include "MuR/ResourceID.h"
-#include "UObject/GCObject.h"
-#include "WorldCollision.h"
-#include "Engine/StreamableManager.h"
 #include "Framework/Notifications/NotificationManager.h"
 #include "MuCO/FMutableTaskGraph.h"
-#include "AssetRegistry/AssetData.h"
 #include "ContentStreaming.h"
 #include "Animation/MorphTarget.h"
 #include "MuCO/MutableStreamableManager.h"
@@ -909,7 +904,7 @@ public:
 	/** Mutable TaskGraph system (Mutable Thread). */
 	FMutableTaskGraph MutableTaskGraph;
 
-	/** Last Mutable task from the previous update. The next update can not start until this task has has completed. */
+	/** Last Mutable task from the previous update. The next update can not start until this task has been completed. */
 	UE::Tasks::FTask LastUpdateMutableTask = UE::Tasks::MakeCompletedTask<void>();
 	
 #if WITH_EDITORONLY_DATA
diff --git a/Engine/Plugins/Mutable/Source/CustomizableObject/Private/MuCO/CustomizableObjectInstance.cpp b/Engine/Plugins/Mutable/Source/CustomizableObject/Private/MuCO/CustomizableObjectInstance.cpp
--- a/Engine/Plugins/Mutable/Source/CustomizableObject/Private/MuCO/CustomizableObjectInstance.cpp
+++ b/Engine/Plugins/Mutable/Source/CustomizableObject/Private/MuCO/CustomizableObjectInstance.cpp
@@ -3,7 +3,6 @@
 #include "MuCO/CustomizableObjectInstance.h"
 
 #include "Algo/Find.h"
-#include "Algo/MaxElement.h"
 #include "Animation/AnimBlueprintGeneratedClass.h"
 #include "Animation/AnimInstance.h"
 #include "Animation/Skeleton.h"
@@ -951,9 +950,7 @@ void UCustomizableInstancePrivate::TickUpdateCloseCustomizableObjects(UCustomiza
 		{
 			const EQueuePriorityType Priority = SystemPrivate->GetUpdatePriority(Public, false);
 
-			FMutableUpdateCandidate* UpdateCandidate = InOutRequestedUpdates.Find(&Public);
-
-			if (UpdateCandidate)
+			if (FMutableUpdateCandidate* UpdateCandidate = InOutRequestedUpdates.Find(&Public))
 			{
 				ensure(HasCOInstanceFlags(PendingLODsUpdate | PendingLODsDowngrade));
 
@@ -1950,10 +1947,7 @@ bool UCustomizableInstancePrivate::UpdateSkeletalMesh_PostBeginUpdate0(UCustomiz
 	if (!CustomizableObject)
 	{
 		UE_LOG(LogMutable, Warning, TEXT("Failed to generate SkeletalMesh for CO Instance %s. It does not have a CO."), *Instance->GetName());
-		
-		InvalidateGeneratedData();
 		Context->UpdateResult = EUpdateResult::Error;
-
 		return false;
 	}
 
@@ -1961,10 +1955,7 @@ bool UCustomizableInstancePrivate::UpdateSkeletalMesh_PostBeginUpdate0(UCustomiz
 	if (bHasInvalidMesh)
 	{
 		UE_LOG(LogMutable, Warning, TEXT("Failed to generate SkeletalMesh for CO Instance %s. CO [%s]"), *Instance->GetName(), *GetNameSafe(CustomizableObject));
-
-		InvalidateGeneratedData();
 		Context->UpdateResult = EUpdateResult::Error;
-
 		return false;
 	}
 	
@@ -2145,8 +2136,7 @@ bool UCustomizableInstancePrivate::UpdateSkeletalMesh_PostBeginUpdate0(UCustomiz
 		if (!ComponentData)
 		{
 			ensure(false);
-			
-			InvalidateGeneratedData();
+			Context->UpdateResult = EUpdateResult::Error;
 			return false;
 		}
 
@@ -2208,7 +2198,7 @@ bool UCustomizableInstancePrivate::UpdateSkeletalMesh_PostBeginUpdate0(UCustomiz
 
 		if (!ModelResources.ReferenceSkeletalMeshesData.IsValidIndex(ObjectComponentIndex.GetValue()))
 		{
-			InvalidateGeneratedData();
+			Context->UpdateResult = EUpdateResult::Error;
 			return false;
 		}
 
@@ -2269,7 +2259,7 @@ bool UCustomizableInstancePrivate::UpdateSkeletalMesh_PostBeginUpdate0(UCustomiz
 		const bool bBuildSkeletonDataSuccess = BuildSkeletonData(Context, *SkeletalMesh, FCustomizableObjectInstanceComponentIndex(InstanceComponentIndex) );
 		if (!bBuildSkeletonDataSuccess)
 		{
-			InvalidateGeneratedData();
+			Context->UpdateResult = EUpdateResult::Error;
 			return false;
 		}
 
@@ -2340,7 +2330,7 @@ bool UCustomizableInstancePrivate::UpdateSkeletalMesh_PostBeginUpdate0(UCustomiz
 		bool const bCopyRenderDataSuccess = BuildOrCopyRenderData(Context, SkeletalMesh, Instance, InstanceComponentIndex);
 		if (!bCopyRenderDataSuccess)
 		{
-			InvalidateGeneratedData();
+			Context->UpdateResult = EUpdateResult::Error;
 			return false;
 		}
 
@@ -4934,323 +4924,330 @@ UE::Tasks::FTask LoadAdditionalAssetsAndData(const TSharedRef<FUpdateContextPriv
 	const bool bUpdateMeshes = DoComponentsNeedUpdate(Context, bHasInvalidMesh);
 
 	// Load assets coming from SubMeshes of the newly generated Mesh
-	if (Context->InstanceUpdateData.LODs.Num())
+	for (int32 InstanceComponentIndex = 0; InstanceComponentIndex < Context->InstanceUpdateData.Components.Num(); ++InstanceComponentIndex)
 	{
-		for (int32 InstanceComponentIndex = 0; InstanceComponentIndex < Context->InstanceUpdateData.Components.Num(); ++InstanceComponentIndex)
-		{
-			const TSharedRef<const FInstanceUpdateData::FComponent>& Component = Components[InstanceComponentIndex];
-			FCustomizableInstanceComponentData* ComponentData = InstancePrivate->GetComponentData(Component->Id);
-			
-			if (Component->SkeletalMesh->PassthroughObject)
-			{
-				// Pass-through components don't have a Reference Mesh so don't access it
-				continue;
-			}
+		const TSharedRef<const FInstanceUpdateData::FComponent>& Component = Components[InstanceComponentIndex];
+		FCustomizableInstanceComponentData* ComponentData = InstancePrivate->GetComponentData(Component->Id);
+		
+		if (!ComponentData)
+        {
+            UE_LOG(LogMutable, Error, TEXT("ComponentData is null: ComponentIDIndex : '%i', COIName :'%s'."), Component->Id.GetValue(), *Instance->GetName());
+			Context->UpdateResult = EUpdateResult::Error;
 			
-			const FCustomizableObjectComponentIndex ObjectComponentIndex = Component->Id;
-			const FMutableRefSkeletalMeshData& RefSkeletalMeshData = ModelResources.ReferenceSkeletalMeshesData[ObjectComponentIndex.GetValue()];
+			return UE::Tasks::FTask{};
+        }
+        
+		
+		if (Component->SkeletalMesh->PassthroughObject)
+		{
+			// Pass-through components don't have a Reference Mesh so don't access it
+			continue;
+		}
+		
+		const FCustomizableObjectComponentIndex ObjectComponentIndex = Component->Id;
+		const FMutableRefSkeletalMeshData& RefSkeletalMeshData = ModelResources.ReferenceSkeletalMeshesData[ObjectComponentIndex.GetValue()];
 
-			for (int32 AssetUserDataIndex : RefSkeletalMeshData.AssetUserDataIndices)
-			{
+		for (int32 AssetUserDataIndex : RefSkeletalMeshData.AssetUserDataIndices)
+		{
 #if !WITH_EDITOR
-				Context->StreamedResourceIndex.AddUnique(AssetUserDataIndex); // Used to hold/release streamed resources in non-editor builds.
+			Context->StreamedResourceIndex.AddUnique(AssetUserDataIndex); // Used to hold/release streamed resources in non-editor builds.
 #endif
-				ComponentData->StreamedResourceIndex.AddUnique(AssetUserDataIndex);
-			}
+			ComponentData->StreamedResourceIndex.AddUnique(AssetUserDataIndex);
+		}
 
-			const FName ComponentName = Context->ComponentNames[Component->Id.GetValue()];
+		const FName ComponentName = Context->ComponentNames[Component->Id.GetValue()];
 
-			if (bUpdateMeshes)
+		if (bUpdateMeshes)
+		{
+			// Morphs
 			{
-				// Morphs
+				// Data
+				for (int32 LODIndex = Context->GetFirstRequestedLOD()[ComponentName]; LODIndex < Component->LODCount; ++LODIndex)
 				{
-					// Data
-					for (int32 LODIndex = Context->GetFirstRequestedLOD()[ComponentName]; LODIndex < Component->LODCount; ++LODIndex)
-					{
-						const TSharedRef<FInstanceUpdateData::FLOD>& LOD = Context->InstanceUpdateData.LODs[Component->FirstLOD + LODIndex];
-
-						UE::Mutable::Private::TManagedPtr<const UE::Mutable::Private::FMesh> MutableMesh = LOD->Mesh;
-						if (!MutableMesh)
-						{
-							continue;
-						}
+					const TSharedRef<FInstanceUpdateData::FLOD>& LOD = Context->InstanceUpdateData.LODs[Component->FirstLOD + LODIndex];
 
-						LoadMorphTargetsData(StreamRequest, MutableMesh.ToStrongRef(), Context->InstanceUpdateData.RealTimeMorphTargetMeshData);
+					UE::Mutable::Private::TManagedPtr<const UE::Mutable::Private::FMesh> MutableMesh = LOD->Mesh;
+					if (!MutableMesh)
+					{
+						continue;
 					}
 
-					// Metadata
-					const int32 FirstLOD = Context->bStreamMeshLODs ?
-						Context->FirstResidentLOD[ComponentName] :
-						Context->GetFirstRequestedLOD()[ComponentName];
-				
-					for (int32 LODIndex = FirstLOD; LODIndex < Component->LODCount; ++LODIndex)
-					{
-						const TSharedRef<FInstanceUpdateData::FLOD>& LOD = Context->InstanceUpdateData.LODs[Component->FirstLOD + LODIndex];
+					LoadMorphTargetsData(StreamRequest, MutableMesh.ToStrongRef(), Context->InstanceUpdateData.RealTimeMorphTargetMeshData);
+				}
 
-						UE::Mutable::Private::TManagedPtr<const UE::Mutable::Private::FMesh> MutableMesh = LOD->Mesh;
-						if (!MutableMesh)
-						{
-							continue;
-						}
+				// Metadata
+				const int32 FirstLOD = Context->bStreamMeshLODs ?
+					Context->FirstResidentLOD[ComponentName] :
+					Context->GetFirstRequestedLOD()[ComponentName];
+			
+				for (int32 LODIndex = FirstLOD; LODIndex < Component->LODCount; ++LODIndex)
+				{
+					const TSharedRef<FInstanceUpdateData::FLOD>& LOD = Context->InstanceUpdateData.LODs[Component->FirstLOD + LODIndex];
 
-						LoadMorphTargetsMetadata(StreamRequest, MutableMesh.ToStrongRef(), Context->InstanceUpdateData.RealTimeMorphTargetMeshData);
+					UE::Mutable::Private::TManagedPtr<const UE::Mutable::Private::FMesh> MutableMesh = LOD->Mesh;
+					if (!MutableMesh)
+					{
+						continue;
 					}
+
+					LoadMorphTargetsMetadata(StreamRequest, MutableMesh.ToStrongRef(), Context->InstanceUpdateData.RealTimeMorphTargetMeshData);
 				}
+			}
 
-				// Cloth
+			// Cloth
+			{
+				// Data
+				// From FirstResidentLOD instead of FirstRequestedLOD since clothing we generate all LODs, even the hacky streaming copied ones.
+				for (int32 LODIndex = Context->FirstResidentLOD[ComponentName]; LODIndex < Component->LODCount; ++LODIndex)
 				{
-					// Data
-					// From FirstResidentLOD instead of FirstRequestedLOD since clothing we generate all LODs, even the hacky streaming copied ones.
-					for (int32 LODIndex = Context->FirstResidentLOD[ComponentName]; LODIndex < Component->LODCount; ++LODIndex)
-					{
-						const TSharedRef<FInstanceUpdateData::FLOD>& LOD = Context->InstanceUpdateData.LODs[Component->FirstLOD + LODIndex];
+					const TSharedRef<FInstanceUpdateData::FLOD>& LOD = Context->InstanceUpdateData.LODs[Component->FirstLOD + LODIndex];
 
-						UE::Mutable::Private::TManagedPtr<const UE::Mutable::Private::FMesh> MutableMesh = LOD->Mesh;
-						if (!MutableMesh)
-						{
-							continue;
-						}
+					UE::Mutable::Private::TManagedPtr<const UE::Mutable::Private::FMesh> MutableMesh = LOD->Mesh;
+					if (!MutableMesh)
+					{
+						continue;
+					}
 
-						if (bUpdateMeshes)
-						{
-							LoadClothing(StreamRequest, MutableMesh.ToStrongRef(), Context->InstanceUpdateData.ClothingMeshData);
-						}
+					if (bUpdateMeshes)
+					{
+						LoadClothing(StreamRequest, MutableMesh.ToStrongRef(), Context->InstanceUpdateData.ClothingMeshData);
 					}
+				}
+
+				// Metadata
+				for (int32 LODIndex = Context->FirstLODAvailable[ComponentName]; LODIndex < Component->LODCount; ++LODIndex)
+				{
+					const TSharedRef<FInstanceUpdateData::FLOD>& LOD = Context->InstanceUpdateData.LODs[Component->FirstLOD + LODIndex];
 
-					// Metadata
-					for (int32 LODIndex = Context->FirstLODAvailable[ComponentName]; LODIndex < Component->LODCount; ++LODIndex)
+					UE::Mutable::Private::TManagedPtr<const UE::Mutable::Private::FMesh> MutableMesh = LOD->Mesh;
+					if (!MutableMesh)
 					{
-						const TSharedRef<FInstanceUpdateData::FLOD>& LOD = Context->InstanceUpdateData.LODs[Component->FirstLOD + LODIndex];
+						continue;
+					}
+				
+					const TArray<uint64>& StreamedResources = MutableMesh->GetStreamedResources();
 
-						UE::Mutable::Private::TManagedPtr<const UE::Mutable::Private::FMesh> MutableMesh = LOD->Mesh;
-						if (!MutableMesh)
-						{
-							continue;
-						}
-					
-						const TArray<uint64>& StreamedResources = MutableMesh->GetStreamedResources();
+					for (uint64 ResourceId : StreamedResources)
+					{
+						FCustomizableObjectStreameableResourceId TypedResourceId = BitCast<FCustomizableObjectStreameableResourceId>(ResourceId);
 
-						for (uint64 ResourceId : StreamedResources)
+						if (TypedResourceId.Type == (uint8)FCustomizableObjectStreameableResourceId::EType::Clothing)
 						{
-							FCustomizableObjectStreameableResourceId TypedResourceId = BitCast<FCustomizableObjectStreameableResourceId>(ResourceId);
+							check(TypedResourceId.Id != 0 && TypedResourceId.Id <= TNumericLimits<uint32>::Max());
 
-							if (TypedResourceId.Type == (uint8)FCustomizableObjectStreameableResourceId::EType::Clothing)
+							const TMap<uint32, FClothingStreamable>& ClothingStreamables = ModelStreamableBulkData->ClothingStreamables;
+							if (const FClothingStreamable* ClothingStreamable = ClothingStreamables.Find(TypedResourceId.Id))
 							{
-								check(TypedResourceId.Id != 0 && TypedResourceId.Id <= TNumericLimits<uint32>::Max());
-
-								const TMap<uint32, FClothingStreamable>& ClothingStreamables = ModelStreamableBulkData->ClothingStreamables;
-								if (const FClothingStreamable* ClothingStreamable = ClothingStreamables.Find(TypedResourceId.Id))
+								FClothingMeshData& ReadDestData = Context->InstanceUpdateData.ClothingMeshData.FindOrAdd(TypedResourceId.Id);
+								ReadDestData.ClothingAssetIndex = ClothingStreamable->ClothingAssetIndex;
+								ReadDestData.ClothingAssetLOD = ClothingStreamable->ClothingAssetLOD;
+ 							
+								// TODO: Add async loading of ClothingAsset Data. This could be loaded as an streamead resource similar to and the asset user data.
+								int32 ClothingAssetIndex = ClothingStreamable->ClothingAssetIndex; 
+								int32 PhysicsAssetIndex = ClothingStreamable->PhysicsAssetIndex;					
+								const TSoftObjectPtr<UPhysicsAsset>& PhysicsAsset = ModelResources.PhysicsAssets.IsValidIndex(PhysicsAssetIndex) 
+										? ModelResources.PhysicsAssets[PhysicsAssetIndex] 
+										: nullptr;
+
+								// The entry should always be in the map
+								if (!PhysicsAsset.IsNull())
 								{
-									FClothingMeshData& ReadDestData = Context->InstanceUpdateData.ClothingMeshData.FindOrAdd(TypedResourceId.Id);
-									ReadDestData.ClothingAssetIndex = ClothingStreamable->ClothingAssetIndex;
-									ReadDestData.ClothingAssetLOD = ClothingStreamable->ClothingAssetLOD;
- 								
-									// TODO: Add async loading of ClothingAsset Data. This could be loaded as an streamead resource similar to and the asset user data.
-									int32 ClothingAssetIndex = ClothingStreamable->ClothingAssetIndex; 
-									int32 PhysicsAssetIndex = ClothingStreamable->PhysicsAssetIndex;					
-									const TSoftObjectPtr<UPhysicsAsset>& PhysicsAsset = ModelResources.PhysicsAssets.IsValidIndex(PhysicsAssetIndex) 
-											? ModelResources.PhysicsAssets[PhysicsAssetIndex] 
-											: nullptr;
-
-									// The entry should always be in the map
-									if (!PhysicsAsset.IsNull())
+									if (PhysicsAsset.Get())
 									{
-										if (PhysicsAsset.Get())
+										if (InstancePrivate->ClothingPhysicsAssets.IsValidIndex(ClothingAssetIndex))
 										{
-											if (InstancePrivate->ClothingPhysicsAssets.IsValidIndex(ClothingAssetIndex))
-											{
-												InstancePrivate->ClothingPhysicsAssets[ClothingAssetIndex] = PhysicsAsset.Get();
-											}
-										}
-										else
-										{
-											ComponentData->ClothingPhysicsAssetsToStream.Emplace(ClothingAssetIndex, PhysicsAssetIndex);
-											AssetsToStream.AddUnique(PhysicsAsset.ToSoftObjectPath());
+											InstancePrivate->ClothingPhysicsAssets[ClothingAssetIndex] = PhysicsAsset.Get();
 										}
 									}
+									else
+									{
+										ComponentData->ClothingPhysicsAssetsToStream.Emplace(ClothingAssetIndex, PhysicsAssetIndex);
+										AssetsToStream.AddUnique(PhysicsAsset.ToSoftObjectPath());
+									}
 								}
-								else
-								{
-									UE_LOG(LogMutable, Error, TEXT("Invalid streamed clothing data block [%llu] found."), TypedResourceId.Id);
-								}
+							}
+							else
+							{
+								UE_LOG(LogMutable, Error, TEXT("Invalid streamed clothing data block [%llu] found."), TypedResourceId.Id);
 							}
 						}
 					}
 				}
 			}
-			
-			for (int32 LODIndex = Context->FirstLODAvailable[ComponentName]; LODIndex < Component->LODCount; ++LODIndex)
+		}
+		
+		for (int32 LODIndex = Context->FirstLODAvailable[ComponentName]; LODIndex < Component->LODCount; ++LODIndex)
+		{
+			const TSharedRef<FInstanceUpdateData::FLOD>& LOD = Context->InstanceUpdateData.LODs[Component->FirstLOD + LODIndex];
+
+			UE::Mutable::Private::TManagedPtr<const UE::Mutable::Private::FMesh> MutableMesh = LOD->Mesh;
+			if (!MutableMesh)
 			{
-				const TSharedRef<FInstanceUpdateData::FLOD>& LOD = Context->InstanceUpdateData.LODs[Component->FirstLOD + LODIndex];
+				continue;
+			}
+			
+			const TArray<uint64>& StreamedResources = MutableMesh->GetStreamedResources();
 
-				UE::Mutable::Private::TManagedPtr<const UE::Mutable::Private::FMesh> MutableMesh = LOD->Mesh;
-				if (!MutableMesh)
-				{
-					continue;
-				}
-				
-				const TArray<uint64>& StreamedResources = MutableMesh->GetStreamedResources();
+			for (uint64 ResourceId : StreamedResources)
+			{
+				FCustomizableObjectStreameableResourceId TypedResourceId = BitCast<FCustomizableObjectStreameableResourceId>(ResourceId);
 
-				for (uint64 ResourceId : StreamedResources)
+				if (TypedResourceId.Type == (uint8)FCustomizableObjectStreameableResourceId::EType::AssetUserData)
 				{
-					FCustomizableObjectStreameableResourceId TypedResourceId = BitCast<FCustomizableObjectStreameableResourceId>(ResourceId);
-
-					if (TypedResourceId.Type == (uint8)FCustomizableObjectStreameableResourceId::EType::AssetUserData)
-					{
-						const uint32 ResourceIndex = TypedResourceId.Id;
+					const uint32 ResourceIndex = TypedResourceId.Id;
 #if !WITH_EDITOR
-						Context->StreamedResourceIndex.AddUnique(ResourceIndex); // Used to hold/release streamed resources in non-editor builds.
+					Context->StreamedResourceIndex.AddUnique(ResourceIndex); // Used to hold/release streamed resources in non-editor builds.
 #endif
-						ComponentData->StreamedResourceIndex.AddUnique(ResourceIndex);
-					}
+					ComponentData->StreamedResourceIndex.AddUnique(ResourceIndex);
 				}
+			}
 
-				// Add AssetUserData from Mesh Parameters
-				for (const TStrongObjectPtr<USkeletalMesh>& StrongMesh : MutableMesh->SkeletalMeshes)
+			// Add AssetUserData from Mesh Parameters
+			for (const TStrongObjectPtr<USkeletalMesh>& StrongMesh : MutableMesh->SkeletalMeshes)
+			{
+				if (USkeletalMesh* SkeletalMesh = StrongMesh.Get())
 				{
-					if (USkeletalMesh* SkeletalMesh = StrongMesh.Get())
+					if (const TArray<UAssetUserData*>* AssetUserDataArray = SkeletalMesh->GetAssetUserDataArray())
 					{
-						if (const TArray<UAssetUserData*>* AssetUserDataArray = SkeletalMesh->GetAssetUserDataArray())
+						for (UAssetUserData* AssetUserData : *AssetUserDataArray)
 						{
-							for (UAssetUserData* AssetUserData : *AssetUserDataArray)
+							// No need to load since the SKM is a StrongObjectPtr
+							if (AssetUserData)
 							{
-								// No need to load since the SKM is a StrongObjectPtr
-								if (AssetUserData)
-								{
-									ComponentData->AssetUserDataArray.Add(AssetUserData);
-								}
+								ComponentData->AssetUserDataArray.Add(AssetUserData);
 							}
 						}
 					}
 				}
-				
-				for (int32 TagIndex = 0; TagIndex < MutableMesh->GetTagCount(); ++TagIndex)
+			}
+			
+			for (int32 TagIndex = 0; TagIndex < MutableMesh->GetTagCount(); ++TagIndex)
+			{
+				FString Tag = MutableMesh->GetTag(TagIndex);
+				if (Tag.RemoveFromStart("__PA:"))
 				{
-					FString Tag = MutableMesh->GetTag(TagIndex);
-					if (Tag.RemoveFromStart("__PA:"))
-					{
-						const int32 AssetIndex = FCString::Atoi(*Tag);
-						const TSoftObjectPtr<UPhysicsAsset>& PhysicsAsset = ModelResources.PhysicsAssets.IsValidIndex(AssetIndex) ? ModelResources.PhysicsAssets[AssetIndex] : nullptr;
+					const int32 AssetIndex = FCString::Atoi(*Tag);
+					const TSoftObjectPtr<UPhysicsAsset>& PhysicsAsset = ModelResources.PhysicsAssets.IsValidIndex(AssetIndex) ? ModelResources.PhysicsAssets[AssetIndex] : nullptr;
 
-						if (!PhysicsAsset.IsNull())
+					if (!PhysicsAsset.IsNull())
+					{
+						if (UPhysicsAsset* LoadedPhysicsAssetPtr = PhysicsAsset.Get())
 						{
-							if (UPhysicsAsset* LoadedPhysicsAssetPtr = PhysicsAsset.Get())
-							{
-								ComponentData->PhysicsAssets.PhysicsAssetsToMerge.Add(LoadedPhysicsAssetPtr);
-							}
-							else
-							{
-								ComponentData->PhysicsAssets.PhysicsAssetToLoad.Add(AssetIndex);
-								AssetsToStream.AddUnique(PhysicsAsset.ToSoftObjectPath());
-							}
+							ComponentData->PhysicsAssets.PhysicsAssetsToMerge.Add(LoadedPhysicsAssetPtr);
+						}
+						else
+						{
+							ComponentData->PhysicsAssets.PhysicsAssetToLoad.Add(AssetIndex);
+							AssetsToStream.AddUnique(PhysicsAsset.ToSoftObjectPath());
 						}
 					}
-					
-					if (Tag.RemoveFromStart("__AnimBP:"))
-					{
-						FString SlotIndexString, AnimBpIndexString;
+				}
+				
+				if (Tag.RemoveFromStart("__AnimBP:"))
+				{
+					FString SlotIndexString, AnimBpIndexString;
 
-						if (Tag.Split(TEXT("_Slot_"), &SlotIndexString, &AnimBpIndexString))
+					if (Tag.Split(TEXT("_Slot_"), &SlotIndexString, &AnimBpIndexString))
+					{
+						if (SlotIndexString.IsEmpty() || AnimBpIndexString.IsEmpty())
 						{
-							if (SlotIndexString.IsEmpty() || AnimBpIndexString.IsEmpty())
-							{
-								continue;
-							}
+							continue;
+						}
 
-							const int32 AnimBpIndex = FCString::Atoi(*AnimBpIndexString);
-							if (!ModelResources.AnimBPs.IsValidIndex(AnimBpIndex))
-							{
-								continue;
-							}
+						const int32 AnimBpIndex = FCString::Atoi(*AnimBpIndexString);
+						if (!ModelResources.AnimBPs.IsValidIndex(AnimBpIndex))
+						{
+							continue;
+						}
 
-							FName SlotIndex = *SlotIndexString;
+						FName SlotIndex = *SlotIndexString;
 
-							const TSoftClassPtr<UAnimInstance>& AnimBPAsset = ModelResources.AnimBPs[AnimBpIndex];
+						const TSoftClassPtr<UAnimInstance>& AnimBPAsset = ModelResources.AnimBPs[AnimBpIndex];
 
-							if (!AnimBPAsset.IsNull())
+						if (!AnimBPAsset.IsNull())
+						{
+							const TSoftClassPtr<UAnimInstance>* FoundAnimBpSlot = ComponentData->AnimSlotToBP.Find(SlotIndex);
+							bool bIsSameAnimBp = FoundAnimBpSlot && AnimBPAsset == *FoundAnimBpSlot;
+							if (!FoundAnimBpSlot)
 							{
-								const TSoftClassPtr<UAnimInstance>* FoundAnimBpSlot = ComponentData->AnimSlotToBP.Find(SlotIndex);
-								bool bIsSameAnimBp = FoundAnimBpSlot && AnimBPAsset == *FoundAnimBpSlot;
-								if (!FoundAnimBpSlot)
-								{
-									ComponentData->AnimSlotToBP.Add(SlotIndex, AnimBPAsset);
+								ComponentData->AnimSlotToBP.Add(SlotIndex, AnimBPAsset);
 
-									if (AnimBPAsset.Get())
-									{
-										InstancePrivate->GatheredAnimBPs.Add(AnimBPAsset.Get());
-									}
-									else
-									{
-										AssetsToStream.AddUnique(AnimBPAsset.ToSoftObjectPath());
-									}
+								if (AnimBPAsset.Get())
+								{
+									InstancePrivate->GatheredAnimBPs.Add(AnimBPAsset.Get());
 								}
-								else if (!bIsSameAnimBp)
+								else
 								{
-									// Two submeshes should not have the same animation slot index
-									Context->UpdateResult = EUpdateResult::Warning;
+									AssetsToStream.AddUnique(AnimBPAsset.ToSoftObjectPath());
+								}
+							}
+							else if (!bIsSameAnimBp)
+							{
+								// Two submeshes should not have the same animation slot index
+								Context->UpdateResult = EUpdateResult::Warning;
 
-									FString WarningMessage = FString::Printf(TEXT("Two submeshes have the same anim slot index [%s] in a Mutable Instance."), *SlotIndex.ToString());
-									UE_LOG(LogMutable, Warning, TEXT("%s"), *WarningMessage);
+								FString WarningMessage = FString::Printf(TEXT("Two submeshes have the same anim slot index [%s] in a Mutable Instance."), *SlotIndex.ToString());
+								UE_LOG(LogMutable, Warning, TEXT("%s"), *WarningMessage);
 #if WITH_EDITOR
-									FMessageLog MessageLog("Mutable");
-									MessageLog.Notify(FText::FromString(WarningMessage), EMessageSeverity::Warning, true);
+								FMessageLog MessageLog("Mutable");
+								MessageLog.Notify(FText::FromString(WarningMessage), EMessageSeverity::Warning, true);
 #endif
-								}
 							}
 						}
 					}
-					else if (Tag.RemoveFromStart("__AnimBPTag:"))
-					{
-						InstancePrivate->AnimBPGameplayTags.AddTag(FGameplayTag::RequestGameplayTag(*Tag));
-					}
-#if WITH_EDITORONLY_DATA
-					else if (Tag.RemoveFromStart("__MeshPath:"))
-					{
-						ComponentData->MeshPartPaths.Add(Tag);
-					}
-#endif
 				}
-
-				const int32 AdditionalPhysicsNum = MutableMesh->AdditionalPhysicsBodies.Num();
-				for (int32 I = 0; I < AdditionalPhysicsNum; ++I)
+				else if (Tag.RemoveFromStart("__AnimBPTag:"))
 				{
-					const int32 ExternalId = MutableMesh->AdditionalPhysicsBodies[I]->CustomId;
-					
-					ComponentData->PhysicsAssets.AdditionalPhysicsAssetsToLoad.Add(ExternalId);
-					AssetsToStream.Add(ModelResources.AnimBpOverridePhysiscAssetsInfo[ExternalId].SourceAsset.ToSoftObjectPath());
+					InstancePrivate->AnimBPGameplayTags.AddTag(FGameplayTag::RequestGameplayTag(*Tag));
 				}
+#if WITH_EDITORONLY_DATA
+				else if (Tag.RemoveFromStart("__MeshPath:"))
+				{
+					ComponentData->MeshPartPaths.Add(Tag);
+				}
+#endif
 			}
 
-
-			for (int32 ResourceIndex : ComponentData->StreamedResourceIndex)
+			const int32 AdditionalPhysicsNum = MutableMesh->AdditionalPhysicsBodies.Num();
+			for (int32 I = 0; I < AdditionalPhysicsNum; ++I)
 			{
+				const int32 ExternalId = MutableMesh->AdditionalPhysicsBodies[I]->CustomId;
+				
+				ComponentData->PhysicsAssets.AdditionalPhysicsAssetsToLoad.Add(ExternalId);
+				AssetsToStream.Add(ModelResources.AnimBpOverridePhysiscAssetsInfo[ExternalId].SourceAsset.ToSoftObjectPath());
+			}
+		}
+
+
+		for (int32 ResourceIndex : ComponentData->StreamedResourceIndex)
+		{
 #if WITH_EDITOR
-				if (!ModelResources.StreamedResourceDataEditor.IsValidIndex(ResourceIndex))
-				{
-					UE_LOG(LogMutable, Error, TEXT("Invalid streamed resource index. Max Index [%d]. Resource Index [%d]."), ModelResources.StreamedResourceDataEditor.Num(), ResourceIndex);
-					continue;
-				}
+			if (!ModelResources.StreamedResourceDataEditor.IsValidIndex(ResourceIndex))
+			{
+				UE_LOG(LogMutable, Error, TEXT("Invalid streamed resource index. Max Index [%d]. Resource Index [%d]."), ModelResources.StreamedResourceDataEditor.Num(), ResourceIndex);
+				continue;
+			}
 
-				if (const FCustomizableObjectAssetUserData* AUDResource = ModelResources.StreamedResourceDataEditor[ResourceIndex].Data.GetPtr<FCustomizableObjectAssetUserData>())
-				{
-					AssetsToStream.AddUnique(AUDResource->AssetUserDataEditor.ToSoftObjectPath());
-				}
+			if (const FCustomizableObjectAssetUserData* AUDResource = ModelResources.StreamedResourceDataEditor[ResourceIndex].Data.GetPtr<FCustomizableObjectAssetUserData>())
+			{
+				AssetsToStream.AddUnique(AUDResource->AssetUserDataEditor.ToSoftObjectPath());
+			}
 #else
-				if (!ModelResources.StreamedResourceData.IsValidIndex(ResourceIndex))
-				{
-					UE_LOG(LogMutable, Error, TEXT("Invalid streamed resource index. Max Index [%d]. Resource Index [%d]."), ModelResources.StreamedResourceData.Num(), ResourceIndex);
-					continue;
-				}
+			if (!ModelResources.StreamedResourceData.IsValidIndex(ResourceIndex))
+			{
+				UE_LOG(LogMutable, Error, TEXT("Invalid streamed resource index. Max Index [%d]. Resource Index [%d]."), ModelResources.StreamedResourceData.Num(), ResourceIndex);
+				continue;
+			}
 
-				const FCustomizableObjectStreamedResourceData& StreamedResource = ModelResources.StreamedResourceData[ResourceIndex];
-				if (!StreamedResource.IsLoaded())
-				{
-					AssetsToStream.AddUnique(StreamedResource.GetPath().ToSoftObjectPath());
-				}
-#endif
+			const FCustomizableObjectStreamedResourceData& StreamedResource = ModelResources.StreamedResourceData[ResourceIndex];
+			if (!StreamedResource.IsLoaded())
+			{
+				AssetsToStream.AddUnique(StreamedResource.GetPath().ToSoftObjectPath());
 			}
+#endif
 		}
 	}
+	
 
 	Context->PassthroughObjectLoader->Load(ModelResources.PassthroughObjects, AssetsToStream);
 	
diff --git a/Engine/Plugins/Mutable/Source/CustomizableObject/Private/MuCO/CustomizableObjectSystem.cpp b/Engine/Plugins/Mutable/Source/CustomizableObject/Private/MuCO/CustomizableObjectSystem.cpp
--- a/Engine/Plugins/Mutable/Source/CustomizableObject/Private/MuCO/CustomizableObjectSystem.cpp
+++ b/Engine/Plugins/Mutable/Source/CustomizableObject/Private/MuCO/CustomizableObjectSystem.cpp
@@ -37,7 +37,6 @@
 #include "MuCO/ICustomizableObjectEditorModule.h"
 #include "MuT/UnrealPixelFormatOverride.h"
 #include "MuR/Parameters.h"
-#include "MuR/Types.h"
 #else
 #include "Engine/Engine.h"
 #endif
@@ -1002,9 +1001,8 @@ void FinishUpdateGlobal(const TSharedRef<FUpdateContextPrivate>& Context)
 	check(IsInGameThread())
 
 	UCustomizableObjectInstance* Instance = Context->Instance.Get();
-
-	UCustomizableObject* Object = Instance ? Instance->GetCustomizableObject() : nullptr;
-	UModelResources* ModelResources = Object ? Object->GetPrivate()->GetModelResources() : nullptr;
+	UCustomizableObject* Object = IsValid(Instance) ? Instance->GetCustomizableObject() : nullptr;
+	UModelResources* ModelResources = IsValid(Object) ? Object->GetPrivate()->GetModelResources() : nullptr;
 
 	UCustomizableObjectSystem* System = UCustomizableObjectSystem::GetInstance();
 	UCustomizableObjectSystemPrivate* SystemPrivate = System ? System->GetPrivate() : nullptr;
@@ -1016,7 +1014,7 @@ void FinishUpdateGlobal(const TSharedRef<FUpdateContextPrivate>& Context)
 		SystemPrivate->CurrentMutableOperation = nullptr;
 	}
 	
-	if (Instance)
+	if (IsValid(Instance))
 	{
 		UCustomizableInstancePrivate* PrivateInstance = Instance->GetPrivate();
 		
@@ -1044,7 +1042,16 @@ void FinishUpdateGlobal(const TSharedRef<FUpdateContextPrivate>& Context)
 
 		case EUpdateResult::Error: 
 		case EUpdateResult::Error16BitBoneIndex:
-			PrivateInstance->SkeletalMeshStatus = ESkeletalMeshStatus::Error;
+			{
+				UCustomizableInstancePrivate* InstancePrivate = Instance->GetPrivate();
+				check(InstancePrivate)
+				if (InstancePrivate)
+				{
+					InstancePrivate->InvalidateGeneratedData();
+				}
+				
+				PrivateInstance->SkeletalMeshStatus = ESkeletalMeshStatus::Error;
+			}
 			break;
 			
 		case EUpdateResult::ErrorReplaced:
@@ -1979,7 +1986,7 @@ namespace impl
 	{
 		MUTABLE_CPUPROFILER_SCOPE(Subtask_Mutable_PrepareSkeletonData);
 
-		int32 NumInstanceComponents = OperationData->InstanceUpdateData.Components.Num();
+		const int32 NumInstanceComponents = OperationData->InstanceUpdateData.Components.Num();
 		OperationData->InstanceUpdateData.SkeletonsPerInstanceComponent.SetNum(NumInstanceComponents);
 
 		for (int32 InstanceComponentIndex = 0; InstanceComponentIndex< NumInstanceComponents; ++InstanceComponentIndex)
@@ -2735,30 +2742,30 @@ namespace impl
 		check(IsInGameThread());
 
 		UCustomizableObjectSystem* System = UCustomizableObjectSystem::GetInstance();
-		if (!System || !System->IsValidLowLevel() || System->HasAnyFlags(RF_BeginDestroyed))
+		if (!IsValid(System))
 		{
 			Context->UpdateResult = EUpdateResult::Error;
 			FinishUpdateGlobal(Context);
 			return;
 		}
-
-		UCustomizableObjectInstance* Instance = Context->Instance.Get();
-		if (!Instance)
+		
+		UCustomizableObjectInstance* CustomizableObjectInstance = Context->Instance.Get();
+		if (!IsValid(CustomizableObjectInstance))
 		{
 			Context->UpdateResult = EUpdateResult::Error;
 			FinishUpdateGlobal(Context);
 			return;
 		}
 
-		UCustomizableObject* Object = Instance->GetCustomizableObject();
-		if (!Object)
+		UCustomizableObject* CustomizableObject = CustomizableObjectInstance->GetCustomizableObject();
+		if (!IsValid(CustomizableObject))
 		{
 			Context->UpdateResult = EUpdateResult::Error;
 			FinishUpdateGlobal(Context);
 			return;
 		}
-
-		Instance->GetPrivate()->RelevantParameters = MoveTemp(Context->RelevantParameters);
+		
+		CustomizableObjectInstance->GetPrivate()->RelevantParameters = MoveTemp(Context->RelevantParameters);
 
 		// Actual work
 		// TODO MTBL-391: Review This hotfix
@@ -2777,7 +2784,7 @@ namespace impl
 		check(IsInGameThread());
 
 		UCustomizableObjectSystem* System = UCustomizableObjectSystem::GetInstance();
-		if (!System || !System->IsValidLowLevel() || System->HasAnyFlags(RF_BeginDestroyed))
+		if (!IsValid(System))
 		{
 			OperationData->UpdateResult = EUpdateResult::Error;
 			FinishUpdateGlobal(OperationData);
@@ -2801,8 +2808,7 @@ namespace impl
 		UCustomizableObjectInstance* CustomizableObjectInstance = OperationData->Instance.Get();
 
 		// Actual work
-		// TODO: Review checks.
-		const bool bInstanceInvalid = !CustomizableObjectInstance || !CustomizableObjectInstance->IsValidLowLevel();
+		const bool bInstanceInvalid = !IsValid(CustomizableObjectInstance);
 		if (!bInstanceInvalid)
 		{
 			UCustomizableInstancePrivate* CustomizableInstancePrivateData = CustomizableObjectInstance->GetPrivate();
@@ -2833,6 +2839,12 @@ namespace impl
 			}
 		} // if (!bInstanceValid)
 
+		if (OperationData->UpdateResult == EUpdateResult::Error)
+		{
+			FinishUpdateGlobal(OperationData);
+			return;
+		}
+		
 		if (FLogBenchmarkUtil::IsBenchmarkingReportingEnabled())
 		{
 			// Memory used in the context of this the update of mesh
@@ -2885,7 +2897,7 @@ namespace impl
 		check(IsInGameThread());
 
 		UCustomizableObjectSystem* System = UCustomizableObjectSystem::GetInstance();
-		if (!System)
+		if (!IsValid(System))
 		{
 			OperationData->UpdateResult = EUpdateResult::Error;
 			FinishUpdateGlobal(OperationData);
@@ -2895,7 +2907,7 @@ namespace impl
 		UCustomizableObjectSystemPrivate* SystemPrivate = System->GetPrivate();
 
 		UCustomizableObject* CustomizableObject = OperationData->Object.Get(); 
-		if (!CustomizableObject)
+		if (!IsValid(CustomizableObject))
 		{
 			OperationData->UpdateResult = EUpdateResult::Error;
 			FinishUpdateGlobal(OperationData);
@@ -2903,7 +2915,7 @@ namespace impl
 		}
 		
 		UCustomizableObjectInstance* ObjectInstance = OperationData->Instance.Get();
-		if (!ObjectInstance)
+		if (!IsValid(ObjectInstance))
 		{
 			OperationData->UpdateResult = EUpdateResult::Error;
 			FinishUpdateGlobal(OperationData);
@@ -2979,6 +2991,12 @@ namespace impl
 		// Next Task: Load Unreal Assets
 		//-------------------------------------------------------------
 		UE::Tasks::FTask Game_LoadUnrealAssets = LoadAdditionalAssetsAndData(OperationData);
+		if (!Game_LoadUnrealAssets.IsValid())
+		{
+			OperationData->UpdateResult = EUpdateResult::Error;
+			FinishUpdateGlobal(OperationData);
+			return;
+		}
 
 		// Next-next Task: Convert Resources
 		//-------------------------------------------------------------
@@ -3037,7 +3055,7 @@ namespace impl
 		
 		Operation->StartUpdateTime = FPlatformTime::Seconds();
 
-		if (!System)
+		if (!IsValid(System))
 		{
 			Operation->UpdateResult = EUpdateResult::Error;
 			FinishUpdateGlobal(Operation);
@@ -3047,19 +3065,17 @@ namespace impl
 		UCustomizableObjectSystemPrivate* SystemPrivate = System->GetPrivate();
 		
 		UCustomizableObject* Object = Operation->Object.Get();
-		if (!Object)
+		if (!IsValid(Object))
 		{
 			Operation->UpdateResult = EUpdateResult::Error;
 			FinishUpdateGlobal(Operation);
 			return;
 		}
-
-		UCustomizableObjectPrivate* ObjectPrivate = Object->GetPrivate();
-
+		
 		TRACE_CPUPROFILER_EVENT_SCOPE_TEXT(*Object->GetName()); // So that we know what we are currently updating.
 
 		UCustomizableObjectInstance* Instance = Operation->Instance.Get();
-		if (!Instance) // Only start if it hasn't been already destroyed (i.e. GC after finish PIE)
+		if (!IsValid(Instance)) // Only start if it hasn't been already destroyed (i.e. GC after finish PIE)
 		{
 			Operation->UpdateResult = EUpdateResult::Error;
 			FinishUpdateGlobal(Operation);
@@ -3085,6 +3101,8 @@ namespace impl
 			return;
 		}
 		
+		UCustomizableObjectPrivate* ObjectPrivate = Object->GetPrivate();
+		
 		// If the object is locked (for instance, compiling) we skip any instance update.
 		if (ObjectPrivate->bLocked)
 		{
@@ -3493,7 +3511,7 @@ int32 UCustomizableObjectSystem::TickInternal(const bool bBlocking)
 	return 0;
 #else // !UE_SERVER
 
-	if (!Private)
+	if (!IsValid(this) || !IsValid(Private))
 	{
 		return 0;
 	}
@@ -3508,7 +3526,7 @@ int32 UCustomizableObjectSystem::TickInternal(const bool bBlocking)
 		return 0;
 	}
 
-	if (GWorld)
+	if (IsValid(GWorld))
 	{
 		const EWorldType::Type WorldType = GWorld->WorldType;
 
@@ -3525,7 +3543,7 @@ int32 UCustomizableObjectSystem::TickInternal(const bool bBlocking)
 	for ( TArray<TObjectPtr<UCustomizableObject>>::TIterator It = GetPrivate()->ObjectsPendingLoad.CreateIterator(); It; ++It )
 	{
 		UCustomizableObject* CO = (*It);
-		if (!CO)
+		if (!IsValid(CO))
 		{
 			It.RemoveCurrent();
 			continue;
@@ -3538,9 +3556,9 @@ int32 UCustomizableObjectSystem::TickInternal(const bool bBlocking)
 			TSet<UCustomizableObject*> RelatedCustomizableObjects;
 			Module->GetRelatedObjects(CO, RelatedCustomizableObjects);
 
-			for (const UObject* CustObject : RelatedCustomizableObjects)
+			for (const UObject* CustomizableObject : RelatedCustomizableObjects)
 			{
-				if (CustObject && !CustObject->HasAnyFlags(EObjectFlags::RF_LoadCompleted))
+				if (CustomizableObject && !CustomizableObject->HasAnyFlags(EObjectFlags::RF_LoadCompleted))
 				{
 					bReady = false;
 					break;
@@ -3708,13 +3726,10 @@ int32 UCustomizableObjectSystem::TickInternal(const bool bBlocking)
 			// Look for a higher priority LOD update
 			for (TPair<const UCustomizableObjectInstance*, FMutableUpdateCandidate>& LODUpdateTuple : RequestedLODUpdates)
 			{
-				const UCustomizableObjectInstance* Instance = LODUpdateTuple.Key;
-
-				if (Instance)
+				if (const UCustomizableObjectInstance* Instance = LODUpdateTuple.Key)
 				{
 					if (UCustomizableObject* Object = Instance->GetCustomizableObject())
 					{
-
 						FMutableUpdateCandidate& LODUpdateCandidate = LODUpdateTuple.Value;
 						ensure(LODUpdateCandidate.HasBeenIssued());
 
@@ -3763,9 +3778,7 @@ int32 UCustomizableObjectSystem::TickInternal(const bool bBlocking)
 				check(PendingInstance);
 				
 				// Maybe there's a LODUpdate that has the same instance, merge both updates as an optimization
-				FMutableUpdateCandidate* LODUpdateWithSameInstance = RequestedLODUpdates.Find(PendingInstance);
-				
-				if (LODUpdateWithSameInstance)
+				if (FMutableUpdateCandidate* LODUpdateWithSameInstance = RequestedLODUpdates.Find(PendingInstance))
 				{
 					LODUpdateWithSameInstance->ApplyLODUpdateParamsToInstance(PendingInstanceUpdateFound->Context.Get());
 				}
