[Mutable] Fix possible crash when aquiring a dependecy set from the data cache.

#rnx
#Synced-CL 49285955

diff --git a/Engine/Plugins/Mutable/Source/MutableRuntime/Internal/MuR/SystemPrivate.h b/Engine/Plugins/Mutable/Source/MutableRuntime/Internal/MuR/SystemPrivate.h
--- a/Engine/Plugins/Mutable/Source/MutableRuntime/Internal/MuR/SystemPrivate.h
+++ b/Engine/Plugins/Mutable/Source/MutableRuntime/Internal/MuR/SystemPrivate.h
@@ -1231,7 +1231,8 @@ namespace UE::Mutable::Private
 		{
 			TArray<TManagedPtr<const FImage>, TInlineAllocator<16>> LocalImages;
 			TArray<TManagedPtr<const FMesh> , TInlineAllocator<16>> LocalMeshes;
-			TArray<FOpExecutionData*        , TInlineAllocator<16>> LocalData;
+			TArray<FOpExecutionData*        , TInlineAllocator<16>> LocalDataImages;
+			TArray<FOpExecutionData*        , TInlineAllocator<16>> LocalDataMeshes;
 
 			if (AddressSet.IsEmpty())
 			{
@@ -1271,7 +1272,7 @@ namespace UE::Mutable::Private
 							{
 								return false;
 							}
-							LocalData.Add(Data);
+							LocalDataImages.Add(Data);
 							break;
 						}
 						case EDataType::Mesh:
@@ -1280,7 +1281,7 @@ namespace UE::Mutable::Private
 							{
 								return false;
 							}
-							LocalData.Add(Data);
+							LocalDataMeshes.Add(Data);
 							break;
 						}
 
@@ -1298,29 +1299,20 @@ namespace UE::Mutable::Private
 
 			// We have been successful acquiring all addresses in the set. 
 			// Commit resource acquisition.
-			for (int32 Index = LocalData.Num() - 1; Index >= 0; --Index)
+			for (int32 Index = LocalDataImages.Num() - 1; Index >= 0; --Index)
 			{
-				FOpExecutionData* Data = LocalData[Index];
+				FOpExecutionData* Data = LocalDataImages[Index];
+				check(Data->DataType == EDataType::Image);
 
-				switch (Data->DataType)
-				{
-					case EDataType::Image:
-					{
-						StrongImageResults[Data->DataTypeIndex] = LocalImages[Index];
-						break;
-					}
-					case EDataType::Mesh:
-					{
-						StrongMeshResults[Data->DataTypeIndex] = LocalMeshes[Index];
-						break;
-					}
+				StrongImageResults[Data->DataTypeIndex] = LocalImages[Index];
+			}
 
-					default:
-					{
-						check(false);
-						break;
-					}
-				}	
+			for (int32 Index = LocalDataMeshes.Num() - 1; Index >= 0; --Index)
+			{
+				FOpExecutionData* Data = LocalDataMeshes[Index];
+				check(Data->DataType == EDataType::Mesh);
+
+				StrongMeshResults[Data->DataTypeIndex] = LocalMeshes[Index];
 			}
 
 			return true;
