[mutable] Fix crash when converting constants that have as a child IM_REFERENCE.
- During the generation of the constant, the subgraph was being marked as constant when in fact it was not, as it contained an image parameter.

#Synced-CL 49556441

diff --git a/Engine/Plugins/Mutable/Source/MutableTools/Private/MuT/CodeOptimiser.cpp b/Engine/Plugins/Mutable/Source/MutableTools/Private/MuT/CodeOptimiser.cpp
--- a/Engine/Plugins/Mutable/Source/MutableTools/Private/MuT/CodeOptimiser.cpp
+++ b/Engine/Plugins/Mutable/Source/MutableTools/Private/MuT/CodeOptimiser.cpp
@@ -322,6 +322,18 @@ namespace UE::Mutable::Private
 					bool bIsConstantSubgraph = true;
 					switch (SubgraphType)
 					{
+					case EOpType::IM_REFERENCE:
+					case EOpType::ME_REFERENCE:
+						if (Pass < 2)
+						{
+							bIsConstantSubgraph = false;
+						}
+						else
+						{
+							const ASTOpReferenceResource* Typed = static_cast<const ASTOpReferenceResource*>(SubgraphRoot.get());
+							bIsConstantSubgraph = Typed->bForceLoad;
+						}
+						break;
 					case EOpType::BO_PARAMETER:
 					case EOpType::NU_PARAMETER:
 					case EOpType::SC_PARAMETER:
