[Mutable] Remove unused Temp caches.

#Synced-CL 49557694

diff --git a/Engine/Plugins/Mutable/Source/MutableRuntime/Internal/MuR/SystemPrivate.h b/Engine/Plugins/Mutable/Source/MutableRuntime/Internal/MuR/SystemPrivate.h
--- a/Engine/Plugins/Mutable/Source/MutableRuntime/Internal/MuR/SystemPrivate.h
+++ b/Engine/Plugins/Mutable/Source/MutableRuntime/Internal/MuR/SystemPrivate.h
@@ -1504,16 +1504,6 @@ namespace UE::Mutable::Private
     	/** Only valid in initial generation. We can not load UObjects while streaming. */
     	TSharedPtr<FPassthroughObjectLoader> CurrentPassthroughObjectLoader;
     	
-		/** Resources that have been used in the past, but haven't been deallocated because they still fitted the memory budget and they could be reused. */
-		TMemoryTrackedArray<TSharedPtr<FImage>> PooledImages;
-
-		/** List of intermediate resources that are not soterd anywhere yet. They are still locally referenced by code. */
-		TMemoryTrackedArray<TSharedPtr<const FImage>> TempImages;
-		TMemoryTrackedArray<TSharedPtr<const FMesh>> TempMeshes;
-
-		/** List of resources that are currently in any cache position, and the number of positions they are in. */
-		TMemoryTrackedMap<TSharedPtr<const FResource>, int32> CacheResources;
-
 		/** Given a mutable model, find or create its rom cache. */
 		FModelCacheEntry* FindModelCache(const FModel*);
 		FModelCacheEntry& FindOrAddModelCache(const TSharedPtr<const FModel>&);
@@ -1929,12 +1919,6 @@ namespace UE::Mutable::Private
 		{
 			CurrentInstanceCache->CheckHitCountsCleared();
 
-			// If this check fails it means some operation is not correctly handling resource management and didn't release
-			// a resource it created.
-			// This should be reported and reviewed, but it is not fatal. Some unnecessary memory may be used temporarily.
-			ensure(TempImages.Num() == 0);
-			ensure(TempMeshes.Num() == 0);
-
 #if !(UE_BUILD_SHIPPING || UE_BUILD_TEST)
 			// If this check fails it means have not set up correctly all the paths to debug threading for resource management.
 			check(DebugRunnerThreadID != FThread::InvalidThreadId);
diff --git a/Engine/Plugins/Mutable/Source/MutableRuntime/Private/MuR/CodeRunnerTasked.cpp b/Engine/Plugins/Mutable/Source/MutableRuntime/Private/MuR/CodeRunnerTasked.cpp
--- a/Engine/Plugins/Mutable/Source/MutableRuntime/Private/MuR/CodeRunnerTasked.cpp
+++ b/Engine/Plugins/Mutable/Source/MutableRuntime/Private/MuR/CodeRunnerTasked.cpp
@@ -2297,7 +2297,6 @@ namespace UE::Mutable::Private
 					FInputMemoryStream Stream(RomLoadOp->m_streamBuffer.GetData(), RomLoadOp->m_streamBuffer.Num());
 					FInputArchive Arch(&Stream);
 
-					// TODO: Try to reuse buffer from PooledImages.
 					TManagedPtr<FImage> Value = FImage::StaticUnserialise(Arch);
 
 					Program.SetImageRomValue(RomIndex, Value);
diff --git a/Engine/Plugins/Mutable/Source/MutableRuntime/Private/MuR/System.cpp b/Engine/Plugins/Mutable/Source/MutableRuntime/Private/MuR/System.cpp
--- a/Engine/Plugins/Mutable/Source/MutableRuntime/Private/MuR/System.cpp
+++ b/Engine/Plugins/Mutable/Source/MutableRuntime/Private/MuR/System.cpp
@@ -178,12 +178,7 @@ namespace UE::Mutable::Private
 				}
 			}
 		}
-
-		m_pD->WorkingMemoryManager.PooledImages.Empty();
-		m_pD->WorkingMemoryManager.CacheResources.Empty();
-		check(m_pD->WorkingMemoryManager.TempImages.IsEmpty());
-		check(m_pD->WorkingMemoryManager.TempMeshes.IsEmpty());
-
+    	
 		m_pD->UpdateStats();
 	}
 
@@ -1018,12 +1013,6 @@ namespace UE::Mutable::Private
 		// Reduce the cache until it fits the limit.
 		m_pD->WorkingMemoryManager.EnsureBudgetBelow(0);
 
-		// If we don't constrain the memory budget, free the pooled images or they may pile up.
-		if (m_pD->WorkingMemoryManager.BudgetBytes == 0)
-		{
-			m_pD->WorkingMemoryManager.PooledImages.Empty();
-		}
-
 		m_pD->UpdateStats();
 	}
 
