[Mutable] Better FUpdateContextPrivate memory layout.

#Synced-CL 49557694

diff --git a/Engine/Plugins/Mutable/Source/CustomizableObject/Internal/MuCO/CustomizableObjectSystemPrivate.h b/Engine/Plugins/Mutable/Source/CustomizableObject/Internal/MuCO/CustomizableObjectSystemPrivate.h
--- a/Engine/Plugins/Mutable/Source/CustomizableObject/Internal/MuCO/CustomizableObjectSystemPrivate.h
+++ b/Engine/Plugins/Mutable/Source/CustomizableObject/Internal/MuCO/CustomizableObjectSystemPrivate.h
@@ -528,22 +528,22 @@
 	const FDescriptorHash& GetCapturedDescriptorHash() const;
 
 	const FCustomizableObjectInstanceDescriptor&& MoveCommittedDescriptor();
-	
-	EQueuePriorityType PriorityType = EQueuePriorityType::Low;
-	
-	FInstanceUpdateDelegate UpdateCallback;
-	FInstanceUpdateNativeDelegate UpdateNativeCallback;
 
-	/** Weak reference to the instance we are operating on.
-	 *It is weak because we don't want to lock it in case it becomes irrelevant in the game while operations are pending and it needs to be destroyed. */
-	TWeakObjectPtr<UCustomizableObjectInstance> Instance;
-
-	/** Customizable Object we are operating on. It can be destroyed between Game Thread tasks. */
-	TWeakObjectPtr<UCustomizableObject> Object;
-
 	/** Return the object component index associated with a component in this instance. */
 	FCustomizableObjectComponentIndex GetObjectComponentIndex(FCustomizableObjectInstanceComponentIndex InstanceComponentIndex);
 	
+	const FInstanceUpdateData::FComponent* GetComponentUpdateData(FCustomizableObjectInstanceComponentIndex InstanceComponentIndex);
+
+	void InitComponentDescriptors(int32 Size);
+
+	const TArray<FComponentDescriptor>& GetComponentDescriptors() const;
+	
+	FComponentDescriptor* GetComponentDescriptor(FCustomizableObjectComponentIndex Index);
+	
+	UE::Mutable::Private::FImageOperator::FImagePixelFormatFunc PixelFormatOverride;
+
+	FInstanceUpdateData InstanceUpdateData;
+
 private:
 	/** Descriptor which the update will be performed on. */
 	FCustomizableObjectInstanceDescriptor CapturedDescriptor;
@@ -552,62 +552,61 @@
 	FDescriptorHash CapturedDescriptorHash;
 
 public:
+	TArray<uint32> RelevantParameters;
+	
+	/** Hard references to objects. Avoids GC to collect them. */
+	TArray<TStrongObjectPtr<const UObject>> Objects;
+
+	// Used during an update to prevent the pass-through textures loaded by LoadAdditionalAssetsAsync() from being unloaded by GC
+	// between AdditionalAssetsAsyncLoaded() and their setting into the generated materials in BuildResources()
+	TArray<TStrongObjectPtr<const UTexture>> LoadedPassThroughTexturesPendingSetMaterial;
+	
+	TArray<UE::Mutable::Private::FImageId> CachedTextures;
+	
+	/** Used to know if the updated instances' meshes are different from the previous ones. 
+	  * The index of the array is the instance component's index.
+	  * @return true if the mesh is new or new to this instance (e.g. mesh cached by another instance). */
+	TArray<bool> MeshChangedPerInstanceComponent;
+
+	/** List of component names. Index is the ObjectComponentIndex. */
+	TArray<FName> ComponentNames;
+
+	/** Index of the resource in the StreamedResourceData array of the Model Resources. */
+	TArray<int32> StreamedResourceIndex;
+
+	/** Index of the resource in the ExtensionStreamedResourceData array of the Model Resources. */
+	TArray<int32> ExtensionStreamedResourceIndex;
+	
+private:
+	TArray<FComponentDescriptor> ComponentDescriptors;
+
+public:
+	/** Materials already build in this update. */
+	TMap<UE::Mutable::Private::FMaterialId, TStrongObjectPtr<UMaterialInterface>> MaterialUpdateCache;
+
+	/** Textures already build in this update. */
+	TMap<FTextureCache::FId, TStrongObjectPtr<UTexture>> TextureUpdateCache;
+	
+	TMap<FName, uint8> NumLODsAvailable;
+
+	/** Copy of UModelResources::ComponentFirstLODAvailable. First compiled LOD per component for the running platform. Constant. */
+	TMap<FName, uint8> FirstLODAvailable;
+	
+	TMap<FName, uint8> FirstResidentLOD;
+
+	TMap<UE::Mutable::Private::FImageId, TUniquePtr<FTexturePlatformData>> ImageToPlatformDataMap;
+	
+	/** If a InstanceUsage is in this set it means that its AttachParent has been modified (USkeletalMesh changed, UMaterial changed...). */
+	TSet<TWeakObjectPtr<UCustomizableObjectInstanceUsage>> AttachedParentUpdated;
+
+	FInstanceUpdateDelegate UpdateCallback;
+	FInstanceUpdateNativeDelegate UpdateNativeCallback;
+
+	UE::Mutable::Private::TManagedPtr<const UE::Mutable::Private::FInstance> MutableInstance;
+
 	/** Instance parameters at the time of the operation request. */
 	TSharedPtr<UE::Mutable::Private::FParameters> Parameters; 
 	TSharedPtr<UE::Mutable::Private::FSystem> MutableSystem;
-
-	bool bOnlyUpdateIfNotGenerated = false;
-	bool bIgnoreCloseDist = false;
-	bool bForceHighPriority = false;
-	
-	FInstanceUpdateData InstanceUpdateData;
-	TArray<uint32> RelevantParameters;
-
-	const FInstanceUpdateData::FComponent* GetComponentUpdateData(FCustomizableObjectInstanceComponentIndex InstanceComponentIndex);
-
-	/** This option comes from the operation request */
-	bool bNeverStream = false;
-	
-	/** When this option is enabled it will reuse the Mutable core instance and its temp data between updates.  */
-	bool bLiveUpdateMode = false;
-
-	bool bUseResueTextureCache = false;
-	
-	bool bUseMeshCache = false;
-	
-	bool bUseCaches = true;
-	
-	/** Whether the mesh to generate should support Mesh LOD streaming or not. */
-	bool bStreamMeshLODs = false;
-
-	/** true if the Update has blocked Low Priority Tasks from launching. */
-	bool bLowPriorityTasksBlocked = false;
-
-	/** The Context has been successfully created. */
-	bool bValid = false;
-
-	/** Do not optimize the update. */
-	bool bForce = false;
-
-	/** Update due to a bake. */
-	bool bBake = false;
-
-	/** User config. */
-	bool bProgressiveMipStreamingEnabled = false;
-	
-	/** User config. */
-	bool bIsOnlyGenerateRequestedLODsEnabled = false;
-	
-	bool bIsUpdateAborted = false;
-		
-	bool bKeepOwnershipOfGeneratedResources = false;
-
-	/** This option comes from the operation request. It is used to reduce the number of mipmaps that mutable must generate for images.  */
-	int32 MipsToSkip = 0;
-
-	UE::Mutable::Private::FInstance::FID InstanceID = 0; // Redundant
-	UE::Mutable::Private::TManagedPtr<const UE::Mutable::Private::FInstance> MutableInstance;
-
 	TSharedPtr<UE::Mutable::Private::FModel> Model;
 	TSharedPtr<UE::Mutable::Private::FMeshIdRegistry> MeshIdRegistry;
 	TSharedPtr<UE::Mutable::Private::FImageIdRegistry> ImageIdRegistry;
@@ -620,55 +619,16 @@
 
 	TSharedPtr<FMutableUpdateImageContext> UpdateImageContext;
 	
-	// Number of possible components in the entire CO
-	// TODO: Redundant because it is in the CO
-	uint8 NumObjectComponents = 0;
+	/** Manage loading Passthrough Objects discovered on the StartUpdate. */
+	TSharedPtr<UE::Mutable::Private::FPassthroughObjectLoader> PassthroughObjectLoader;
+	
+	/** Weak reference to the instance we are operating on.
+	 *It is weak because we don't want to lock it in case it becomes irrelevant in the game while operations are pending and it needs to be destroyed. */
+	TWeakObjectPtr<UCustomizableObjectInstance> Instance;
 
-	// Number of components in the instance being generated
-	// TODO: Redundant to store it here because it is implicit in the size of some arrays or in the MutableInstance while it is valid.
-	uint8 NumInstanceComponents = 0;
+	/** Customizable Object we are operating on. It can be destroyed between Game Thread tasks. */
+	TWeakObjectPtr<UCustomizableObject> Object;
 
-	/** List of component names. Index is the ObjectComponentIndex. */
-	TArray<FName> ComponentNames;
-
-	/** Index of the resource in the StreamedResourceData array of the Model Resources. */
-	TArray<int32> StreamedResourceIndex;
-
-	/** Index of the resource in the ExtensionStreamedResourceData array of the Model Resources. */
-	TArray<int32> ExtensionStreamedResourceIndex;
-	
-	TMap<FName, uint8> NumLODsAvailable;
-
-	/** Copy of UModelResources::ComponentFirstLODAvailable. First compiled LOD per component for the running platform. Constant. */
-	TMap<FName, uint8> FirstLODAvailable;
-	
-	TMap<FName, uint8> FirstResidentLOD;
-
-	TMap<UE::Mutable::Private::FImageId, TUniquePtr<FTexturePlatformData>> ImageToPlatformDataMap;
-
-	EUpdateResult UpdateResult = EUpdateResult::Success;
-
-	UE::Mutable::Private::FImageOperator::FImagePixelFormatFunc PixelFormatOverride;
-
-private:
-	TArray<FComponentDescriptor> ComponentDescriptors;
-
-public:
-	void InitComponentDescriptors(int32 Size);
-	const TArray<FComponentDescriptor>& GetComponentDescriptors() const;
-	FComponentDescriptor* GetComponentDescriptor(FCustomizableObjectComponentIndex Index);
-
-	/** Used to know if the updated instances' meshes are different from the previous ones. 
-	  * The index of the array is the instance component's index.
-	  * @return true if the mesh is new or new to this instance (e.g. mesh cached by another instance). */
-	TArray<bool> MeshChangedPerInstanceComponent;
-	
-	bool UpdateStarted = false;
-	bool bLevelBegunPlay = false;
-
-	/** true if the update has been optimized (skips all Tasks and calls FinishUpdateGlobal directly on the Enqueue). */
-	bool bOptimizedUpdate = false;
-	
 	// Update stats
 	double StartQueueTime = 0.0;
 	double QueueTime = 0.0;
@@ -686,29 +646,77 @@
 
 	// Update Memory stats
 	int64 UpdateStartBytes = 0;
+
 	int64 UpdateEndPeakBytes = 0;
+	
 	int64 UpdateEndRealPeakBytes = 0;
+	
+	UE::Mutable::Private::FInstance::FID InstanceID = 0; // Redundant
+	
+	/** This option comes from the operation request. It is used to reduce the number of mipmaps that mutable must generate for images.  */
+	int32 MipsToSkip = 0;
+	
+	// Number of possible components in the entire CO
+	// TODO: Redundant because it is in the CO
+	uint8 NumObjectComponents = 0;
 
-	/** If a InstanceUsage is in this set it means that its AttachParent has been modified (USkeletalMesh changed, UMaterial changed...). */
-	TSet<TWeakObjectPtr<UCustomizableObjectInstanceUsage>> AttachedParentUpdated;
+	// Number of components in the instance being generated
+	// TODO: Redundant to store it here because it is implicit in the size of some arrays or in the MutableInstance while it is valid.
+	uint8 NumInstanceComponents = 0;
+	
+	EQueuePriorityType PriorityType = EQueuePriorityType::Low;
 
-	/** Manage loading Passthrough Objects discovered on the StartUpdate. */
-	TSharedPtr<UE::Mutable::Private::FPassthroughObjectLoader> PassthroughObjectLoader;
+	EUpdateResult UpdateResult = EUpdateResult::Success;
 	
-	/** Hard references to objects. Avoids GC to collect them. */
-	TArray<TStrongObjectPtr<const UObject>> Objects;
+	/** This option comes from the operation request */
+	uint8 bNeverStream : 1 = false;
+	
+	uint8 bOnlyUpdateIfNotGenerated : 1 = false;
 
-	// Used during an update to prevent the pass-through textures loaded by LoadAdditionalAssetsAsync() from being unloaded by GC
-	// between AdditionalAssetsAsyncLoaded() and their setting into the generated materials in BuildResources()
-	TArray<TStrongObjectPtr<const UTexture>> LoadedPassThroughTexturesPendingSetMaterial;
+	uint8 bIgnoreCloseDist : 1 = false;
 	
-	TArray<UE::Mutable::Private::FImageId> CachedTextures;
+	uint8 bForceHighPriority : 1 = false;
+	
+	/** When this option is enabled it will reuse the Mutable core instance and its temp data between updates.  */
+	uint8 bLiveUpdateMode : 1 = false;
 
-	/** Materials already build in this update. */
-	TMap<UE::Mutable::Private::FMaterialId, TStrongObjectPtr<UMaterialInterface>> MaterialUpdateCache;
+	uint8 bUseResueTextureCache : 1 = false;
+	
+	uint8 bUseMeshCache : 1 = false;
+	
+	uint8 bUseCaches : 1 = true;
+	
+	/** Whether the mesh to generate should support Mesh LOD streaming or not. */
+	uint8 bStreamMeshLODs : 1 = false;
 
-	/** Textures already build in this update. */
-	TMap<FTextureCache::FId, TStrongObjectPtr<UTexture>> TextureUpdateCache;
+	/** true if the Update has blocked Low Priority Tasks from launching. */
+	uint8 bLowPriorityTasksBlocked : 1 = false;
+
+	/** The Context has been successfully created. */
+	uint8 bValid : 1 = false;
+
+	/** Do not optimize the update. */
+	uint8 bForce : 1= false;
+
+	/** Update due to a bake. */
+	uint8 bBake : 1 = false;
+
+	/** User config. */
+	uint8 bProgressiveMipStreamingEnabled : 1 = false;
+	
+	/** User config. */
+	uint8 bIsOnlyGenerateRequestedLODsEnabled : 1 = false;
+	
+	uint8 bIsUpdateAborted : 1 = false;
+		
+	uint8 bKeepOwnershipOfGeneratedResources : 1 = false;
+	
+	uint8 UpdateStarted : 1 = false;
+
+	uint8 bLevelBegunPlay : 1 = false;
+
+	/** true if the update has been optimized (skips all Tasks and calls FinishUpdateGlobal directly on the Enqueue). */
+	uint8 bOptimizedUpdate : 1 = false;
 };
 
 
@@ -854,8 +862,6 @@
 
 	// Init the async release of a Mutable Core Instance ID and all the temp resources associated with it
 	void InitInstanceIDRelease(UE::Mutable::Private::FInstance::FID IDToRelease);
-
-	void GetMipStreamingConfig(const UCustomizableObjectInstance& Instance, bool bTextureStreamingEnabled, bool& bOutNeverStream, int32& OutMipsToSkip) const;
 	
 	bool IsReplaceDiscardedWithReferenceMeshEnabled() const;
 	void SetReplaceDiscardedWithReferenceMeshEnabled(bool bIsEnabled);
diff --git a/Engine/Plugins/Mutable/Source/CustomizableObject/Private/MuCO/CustomizableObjectSystem.cpp b/Engine/Plugins/Mutable/Source/CustomizableObject/Private/MuCO/CustomizableObjectSystem.cpp
--- a/Engine/Plugins/Mutable/Source/CustomizableObject/Private/MuCO/CustomizableObjectSystem.cpp
+++ b/Engine/Plugins/Mutable/Source/CustomizableObject/Private/MuCO/CustomizableObjectSystem.cpp
@@ -1409,48 +1409,6 @@ void UpdateSkeletalMesh(const TSharedRef<FUpdateContextPrivate>& Context)
 }
 
 
-void UCustomizableObjectSystemPrivate::GetMipStreamingConfig(const UCustomizableObjectInstance& Instance, bool bTextureStreamingEnabled, bool& bOutNeverStream, int32& OutMipsToSkip) const
-{
-	bOutNeverStream = false;
-
-	// From user-controlled per-state flag?
-	const FString CurrentState = Instance.GetCurrentState();
-
-	if (const UModelResources* ModelResources = Instance.GetCustomizableObject()->GetPrivate()->GetModelResources())
-	{
-		if (const FMutableStateData* State = ModelResources->StateUIDataMap.Find(CurrentState))
-		{
-			bOutNeverStream = State->bDisableTextureStreaming;
-		}
-
-#if WITH_EDITORONLY_DATA
-		// Was streaming disabled at object-compilation time? 
-		if (ModelResources->bIsTextureStreamingDisabled)
-		{
-			bOutNeverStream = true;
-		}
-#endif
-	}
-
-	OutMipsToSkip = 0; // 0 means generate all mips
-
-	// Streaming disabled from platform settings or from platform CustomizableObjectSystem properties?
-#if PLATFORM_SUPPORTS_TEXTURE_STREAMING
-	if (!IStreamingManager::Get().IsTextureStreamingEnabled() || !bTextureStreamingEnabled)
-	{
-		bOutNeverStream = true;
-	}
-#else
-	bOutNeverStream = true;
-#endif
-	
-	if (!bOutNeverStream)
-	{
-		OutMipsToSkip = 255; // This means skip all possible mips until only UTexture::GetStaticMinTextureResidentMipCount() are left
-	}
-}
-
-
 bool UCustomizableObjectSystemPrivate::IsReplaceDiscardedWithReferenceMeshEnabled() const
 {
 	return bReplaceDiscardedWithReferenceMesh;
@@ -2860,7 +2818,7 @@ namespace impl
 		const  TSharedPtr<UE::Mutable::Private::FSystem> MutableSystem = CustomizableObjectSystemPrivate->MutableSystem;
 		CustomizableObjectSystemPrivate->LastUpdateMutableTask = CustomizableObjectSystemPrivate->MutableTaskGraph.AddMutableThreadTask(
 			TEXT("Task_Mutable_ReleaseInstance"),
-			[InstanceID = OperationData->InstanceID, MutableSystem, bLiveUpdateMode = OperationData->bLiveUpdateMode]()
+			[InstanceID = OperationData->InstanceID, MutableSystem, bLiveUpdateMode = static_cast<bool>(OperationData->bLiveUpdateMode)]()
 			{
 				Task_Mutable_ReleaseInstance(InstanceID, MutableSystem, bLiveUpdateMode);
 			});
@@ -3146,18 +3104,34 @@ namespace impl
 		const int32 State = InstancePrivate->GetState();
 		const FString& StateName = ObjectPrivate->GetStateName(State);
 		const FMutableStateData* StateData = ModelResources.StateUIDataMap.Find(StateName);
-
+		
+		// Streaming disabled from platform settings or from platform CustomizableObjectSystem properties?
+#if PLATFORM_SUPPORTS_TEXTURE_STREAMING
+		if (!IStreamingManager::Get().IsTextureStreamingEnabled() ||
+			!Operation->bProgressiveMipStreamingEnabled
+#if WITH_EDITORONLY_DATA
+			|| ModelResources.bIsTextureStreamingDisabled // Was streaming disabled at object-compilation time? 
+#endif
+			)
+		{
+			Operation->bNeverStream = true;
+		}
+		else
+		{
+			Operation->bNeverStream = StateData ? StateData->bDisableTextureStreaming : false;
+		}
+#else
+		Operation->bNeverStream = true;
+#endif
+	
+		Operation->MipsToSkip = !Operation->bNeverStream ? 255 : 0; // // 0 means generate all mips. 255 means skip all possible mips until only UTexture::GetStaticMinTextureResidentMipCount() are left
+		
 		Operation->bLiveUpdateMode = false;
 
 		if (SystemPrivate->EnableMutableLiveUpdate)
 		{
 			Operation->bLiveUpdateMode = StateData ? StateData->bLiveUpdateMode : false;
 		}
-
-		Operation->bNeverStream = false;
-		Operation->MipsToSkip = 0;
-
-		SystemPrivate->GetMipStreamingConfig(*Instance, Operation->bProgressiveMipStreamingEnabled, Operation->bNeverStream, Operation->MipsToSkip);
 		
 		if (Operation->bLiveUpdateMode && (!Operation->bNeverStream || Operation->MipsToSkip > 0))
 		{
