[Mutable] Remove Core Pimpl.

- Fixed wrong constantness of FProgram due to Pimpl.

#Synced-CL 49561974

diff --git a/Engine/Plugins/Mutable/Source/CustomizableObject/Private/MuCO/CustomizableObject.cpp b/Engine/Plugins/Mutable/Source/CustomizableObject/Private/MuCO/CustomizableObject.cpp
--- a/Engine/Plugins/Mutable/Source/CustomizableObject/Private/MuCO/CustomizableObject.cpp
+++ b/Engine/Plugins/Mutable/Source/CustomizableObject/Private/MuCO/CustomizableObject.cpp
@@ -27,7 +27,6 @@
 #include "MuCO/MutableProjectorTypeUtils.h"
 #include "MuCO/UnrealMutableModelDiskStreamer.h"
 #include "MuR/Model.h"
-#include "MuR/ModelPrivate.h"
 #include "MuR/Operations.h"
 #include "PhysicsEngine/PhysicsAsset.h"
 #include "Serialization/MemoryReader.h"
@@ -910,7 +909,7 @@ void UCustomizableObjectPrivate::SaveEmbeddedData(FArchive& Ar)
 	{	
 		if (Ar.IsCooking())
 		{
-			Model->GetPrivate()->Program.RomsCompileData.Empty();
+			Model->Program.RomsCompileData.Empty();
 		}
 		UnrealMutableOutputStream Stream(Ar);
 		UE::Mutable::Private::FOutputArchive Arch(&Stream);
@@ -2231,7 +2230,7 @@ void UCustomizableObjectPrivate::SetModel(const TSharedPtr<UE::Mutable::Private:
 #endif
 		check(LocalModelResources);
 		
-		UE::Mutable::Private::FProgram& Program = MutableModel->GetPrivate()->Program;
+		UE::Mutable::Private::FProgram& Program = MutableModel->Program;
 		Program.ExternalOperationProvider = MakeShared<UE::Mutable::Private::FExternalOperationProvider>(*LocalModelResources);
 
 		for (UE::Mutable::Private::FParameterDesc& Parameter : Program.Parameters)
@@ -3355,7 +3354,7 @@ void UCustomizableObjectPrivate::LogMemory()
 	const UE::Mutable::Private::FModel* MuModel = GetModel().Get();
 	if (MuModel)
 	{
-		const UE::Mutable::Private::FProgram& Program = MuModel->GetPrivate()->Program;
+		const UE::Mutable::Private::FProgram& Program = MuModel->Program;
 		int32 ByteCode = Program.ByteCode.GetAllocatedSize();
 		int32 OpAddress = Program.OpAddress.GetAllocatedSize();
 		int32 Roms = Program.Roms.GetAllocatedSize();
diff --git a/Engine/Plugins/Mutable/Source/CustomizableObject/Private/MuCO/CustomizableObjectSystem.cpp b/Engine/Plugins/Mutable/Source/CustomizableObject/Private/MuCO/CustomizableObjectSystem.cpp
--- a/Engine/Plugins/Mutable/Source/CustomizableObject/Private/MuCO/CustomizableObjectSystem.cpp
+++ b/Engine/Plugins/Mutable/Source/CustomizableObject/Private/MuCO/CustomizableObjectSystem.cpp
@@ -26,7 +26,7 @@
 #include "SkinnedAssetCompiler.h"
 #include "Components/SkeletalMeshComponent.h"
 #include "MuCO/CustomizableObjectSystemPrivate.h"
-#include "MuR/ParametersPrivate.h"
+#include "MuR/Parameters.h"
 
 #if WITH_EDITOR
 #include "Editor.h"
@@ -3366,10 +3366,7 @@ namespace impl
 		// Note : We are accessing the internal parameters (no the ones on the COI) as those do change in a later stage and may not yet have the SKM parameters in them
 		// at this point.
 		{
-			UE::Mutable::Private::FParameters::Private* ParametersPrivate = Operation->Parameters->GetPrivate();
-			check(ParametersPrivate);
-		
-			const int32 ParameterCount = ParametersPrivate->Values.Num();
+			const int32 ParameterCount = Operation->Parameters->Values.Num();
 			for (int32 ParameterIndex = 0; ParameterIndex < ParameterCount; ++ParameterIndex)
 			{
 				auto AddToPrerequisites = [&Prerequisites](USkeletalMesh* SkeletalMesh) -> void
@@ -3388,10 +3385,10 @@ namespace impl
 					}
 				};
 
-				if (!ParametersPrivate->HasMultipleValues(ParameterIndex))
+				if (!Operation->Parameters->HasMultipleValues(ParameterIndex))
 				{
 					// Add the non ranged value
-					UE::Mutable::Private::FParameterValue& ParameterValue = ParametersPrivate->Values[ParameterIndex];
+					UE::Mutable::Private::FParameterValue& ParameterValue = Operation->Parameters->Values[ParameterIndex];
 					if (TStrongObjectPtr<USkeletalMesh>* SkeletalMeshParameterValue = ParameterValue.TryGet<UE::Mutable::Private::FParamSkeletalMeshType>())
 					{
 						AddToPrerequisites(SkeletalMeshParameterValue->Get());
@@ -3399,7 +3396,7 @@ namespace impl
 				}
 				else
 				{
-					const TMap<TArray<int32>, UE::Mutable::Private::FParameterValue>& MultiValueParameterData = ParametersPrivate->MultiValues[ParameterIndex];
+					const TMap<TArray<int32>, UE::Mutable::Private::FParameterValue>& MultiValueParameterData = Operation->Parameters->MultiValues[ParameterIndex];
 
 					TArray<UE::Mutable::Private::FParameterValue> MultiDimensionalParameterValues;
 					MultiValueParameterData.GenerateValueArray(MultiDimensionalParameterValues);
diff --git a/Engine/Plugins/Mutable/Source/CustomizableObjectEditor/Private/MuCOE/CustomizableObjectCompiler.cpp b/Engine/Plugins/Mutable/Source/CustomizableObjectEditor/Private/MuCOE/CustomizableObjectCompiler.cpp
--- a/Engine/Plugins/Mutable/Source/CustomizableObjectEditor/Private/MuCOE/CustomizableObjectCompiler.cpp
+++ b/Engine/Plugins/Mutable/Source/CustomizableObjectEditor/Private/MuCOE/CustomizableObjectCompiler.cpp
@@ -30,7 +30,6 @@
 #include "MuCO/CustomizableObjectSystemPrivate.h"
 #include "MuCOE/CompileRequest.h"
 #include "MuR/Model.h"
-#include "MuR/ModelPrivate.h"
 #include "MuCOE/CustomizableObjectVersionBridge.h"
 
 #include "DerivedDataCache.h"
@@ -1133,7 +1132,7 @@ void FCustomizableObjectCompiler::CompleteRequest(ECompilationStatePrivate State
 
 		if (PlatformData->Model)
 		{
-			PlatformData->Model->GetPrivate()->UnloadRoms();
+			PlatformData->Model->UnloadRoms();
 		}
 
 		if (!CurrentOptions.bIsCooking)
diff --git a/Engine/Plugins/Mutable/Source/CustomizableObjectEditor/Private/MuCOE/SMutableCodeViewer.cpp b/Engine/Plugins/Mutable/Source/CustomizableObjectEditor/Private/MuCOE/SMutableCodeViewer.cpp
--- a/Engine/Plugins/Mutable/Source/CustomizableObjectEditor/Private/MuCOE/SMutableCodeViewer.cpp
+++ b/Engine/Plugins/Mutable/Source/CustomizableObjectEditor/Private/MuCOE/SMutableCodeViewer.cpp
@@ -28,7 +28,7 @@
 #include "MuCO/UnrealMutableImageProvider.h"
 #include "MuT/ErrorLog.h"
 #include "MuT/TypeInfo.h"
-#include "MuR/SystemPrivate.h"
+#include "MuR/System.h"
 #include "Widgets/SNullWidget.h"
 #include "Widgets/Colors/SColorBlock.h"
 #include "Widgets/Input/SButton.h"
@@ -633,7 +633,7 @@ void SMutableCodeViewer::Construct(const FArguments& InArgs, UCustomizableObject
 					[
 						// Generate a new Constants panel to show the data stored on the current mutable program
 						SAssignNew(ConstantsWidget, SMutableConstantsWidget,
-							&(MutableModel->GetPrivate()->Program),
+							&(MutableModel->Program),
 							SharedThis(this))
 					]
 				]
@@ -832,7 +832,7 @@ void SMutableCodeViewer::CacheOperationsMatchingStringPattern()
 		// to start searching
 		SearchPayload.SetupRootBatch(RootNodeAddresses);
 
-		const UE::Mutable::Private::FProgram& Program = MutableModel->GetPrivate()->Program;
+		const UE::Mutable::Private::FProgram& Program = MutableModel->Program;
 		GetOperationsMatchingStringPattern(SearchString,bIsSearchStringRegularExpression,SearchPayload, Program);
 	
 		// Dump the located resources array onto the navigation array
@@ -1057,7 +1057,7 @@ void SMutableCodeViewer::OnSelectedOperationTypeFromTree()
 	const TSharedPtr<FMutableCodeTreeElement> ReferenceOperationElement = TreeView->GetSelectedItems()[0];
 	
 	const UE::Mutable::Private::EOpType OperationType =
-		MutableModel->GetPrivate()->Program.GetOpType(ReferenceOperationElement->MutableOperation);
+		MutableModel->Program.GetOpType(ReferenceOperationElement->MutableOperation);
 
 	// Find the operation type directly in our array of operation elements (from the drop down)
 	const TSharedPtr<const FMutableOperationElement>* RepresentativeElement = FoundModelOperationTypeElements.
@@ -1101,7 +1101,7 @@ void SMutableCodeViewer::CacheAddressesOfOperationsOfType()
 	SearchPayload.SetupRootBatch(RootNodeAddresses);
 	
 	// Main update procedure run for the targeted state and the targeted parameter values
-	const UE::Mutable::Private::FProgram& Program = MutableModel->GetPrivate()->Program;
+	const UE::Mutable::Private::FProgram& Program = MutableModel->Program;
 	GetOperationsOfType(OperationTypeToSearch,SearchPayload, Program);
 	
 	if (!SearchPayload.FoundElements.IsEmpty())
@@ -1170,7 +1170,7 @@ void SMutableCodeViewer::CacheOperationTypesPresentOnModel()
 	// there are
 	{
 		// Get the types and the amount of instances of each unique operation on the mutable model
-		const UE::Mutable::Private::FProgram& Program = MutableModel->GetPrivate()->Program;
+		const UE::Mutable::Private::FProgram& Program = MutableModel->Program;
 		const uint32 ProgramAddressesCount = Program.OpAddress.Num();
 		
 		// Ensure first operation type is NONE since we are skipping it due to it having to have that type
@@ -1345,9 +1345,7 @@ void SMutableCodeViewer::GenerateAllTreeElements()
 	// 	be ignored due to how we handle them when expanding and contracting elements (OnExpansionChanged)
 
 	// Generate all root nodes
-	const UE::Mutable::Private::FModel::Private* ModelPrivate = MutableModel->GetPrivate();
-	check(ModelPrivate);
-	const UE::Mutable::Private::FProgram& Program = ModelPrivate->Program;
+	const UE::Mutable::Private::FProgram& Program = MutableModel->Program;
 	const uint32 StateCount = Program.States.Num();
 	for ( uint32 StateIndex = 0; StateIndex < StateCount; ++StateIndex )
 	{
@@ -1941,7 +1939,7 @@ void SMutableCodeViewer::GetChildrenForInfo(TSharedPtr<FMutableCodeTreeElement>
 	}
 
 	check(MutableModel);
-	const UE::Mutable::Private::FProgram& Program = MutableModel->GetPrivate()->Program;
+	const UE::Mutable::Private::FProgram& Program = MutableModel->Program;
 
 	UE::Mutable::Private::OP::ADDRESS ParentAddress = InInfo->MutableOperation;
 
@@ -2008,7 +2006,7 @@ void SMutableCodeViewer::OnExpansionChanged(TSharedPtr<FMutableCodeTreeElement>
 void SMutableCodeViewer::GetVisibleChildren(TSharedPtr<FMutableCodeTreeElement> InInfo, TSet<TSharedPtr<FMutableCodeTreeElement>>& OutChildren)
 {
 	check(MutableModel);
-	const UE::Mutable::Private::FProgram& MutableProgram = MutableModel->GetPrivate()->Program;
+	const UE::Mutable::Private::FProgram& MutableProgram = MutableModel->Program;
 	
 	TArray<TSharedPtr<FMutableCodeTreeElement>> ToSearchForChildren;
 	ToSearchForChildren.Add(InInfo);
@@ -2086,7 +2084,7 @@ void SMutableCodeViewer::OnSelectionChanged(TSharedPtr<FMutableCodeTreeElement>
 	bIsPreviewPendingUpdate = true;
 
 	SelectedOperationAddress = SelectedNodes[0]->MutableOperation;
-	const UE::Mutable::Private::EOpType OperationType = MutableModel->GetPrivate()->Program.GetOpType(SelectedOperationAddress);
+	const UE::Mutable::Private::EOpType OperationType = MutableModel->Program.GetOpType(SelectedOperationAddress);
 	const UE::Mutable::Private::EDataType OperationDataType = UE::Mutable::Private::GetOpDataType(OperationType);
 	
 	switch (OperationDataType)
@@ -2359,7 +2357,7 @@ void SMutableCodeViewer::TreeExpandElements(TArray<TSharedPtr<FMutableCodeTreeEl
 		// Filter the elements being expanded if the user has defined a desired EDataType
 		if (FilteringDataType != UE::Mutable::Private::EDataType::None)
 		{
-			const UE::Mutable::Private::EOpType OperationType = Item->MutableModel->GetPrivate()->Program.GetOpType(OperationAddress);
+			const UE::Mutable::Private::EOpType OperationType = Item->MutableModel->Program.GetOpType(OperationAddress);
 			const UE::Mutable::Private::EDataType OperationDataType = UE::Mutable::Private::GetOpDataType(OperationType);
 
 			// If it is not of the desired type then ignore it and continue to the next pending element
@@ -2500,11 +2498,10 @@ void SMutableCodeViewer::CacheRootNodeAddresses()
 
 	TArray<UE::Mutable::Private::OP::ADDRESS> FoundRootNodeAddresses;
 	
-	const UE::Mutable::Private::FModel::Private* ModelPrivate = MutableModel->GetPrivate();
-	const int32 StateCount = ModelPrivate->Program.States.Num();
+	const int32 StateCount = MutableModel->Program.States.Num();
 	for ( int32 StateIndex=0; StateIndex < StateCount; ++StateIndex )
 	{
-		const UE::Mutable::Private::FProgram::FState& State = ModelPrivate->Program.States[StateIndex];
+		const UE::Mutable::Private::FProgram::FState& State = MutableModel->Program.States[StateIndex];
 		FoundRootNodeAddresses.Add(State.Root);
 	}
 
@@ -2531,7 +2528,7 @@ void SMutableCodeViewer::CacheRootNodeAddresses()
 	SearchPayload.SetupRootBatch(RootNodeAddresses);
 	
 	// Main update procedure run for the targeted state and the targeted parameter values
-	const UE::Mutable::Private::FProgram& Program = MutableModel->GetPrivate()->Program; 
+	const UE::Mutable::Private::FProgram& Program = MutableModel->Program; 
 	GetOperationsReferencingConstantResource(ConstantDataType,IndexOnConstantsArray,SearchPayload, Program);
 	
 	// At this point we did get all the addresses of operations that do involve the usage of our resource
@@ -2891,20 +2888,20 @@ void SMutableCodeViewer::Tick(const FGeometry& AllottedGeometry, const double In
 
 	bIsPreviewPendingUpdate = false;
 
-	const UE::Mutable::Private::EOpType OperationType = MutableModel->GetPrivate()->Program.GetOpType(SelectedOperationAddress);
+	const UE::Mutable::Private::EOpType OperationType = MutableModel->Program.GetOpType(SelectedOperationAddress);
 	const UE::Mutable::Private::EDataType OperationDataType = UE::Mutable::Private::GetOpDataType(OperationType);
 
 	UE::Mutable::Private::FSettings Settings;
 	const TSharedPtr<UE::Mutable::Private::FSystem> System = MakeShared<UE::Mutable::Private::FSystem>(Settings);
 
-	System->GetPrivate()->BeginBuild(MutableModel, Object->GetPrivate()->MeshIdRegistry, Object->GetPrivate()->ImageIdRegistry, Object->GetPrivate()->MaterialIdRegistry, Object->GetPrivate()->SkeletalMeshIdRegistry);
+	System->BeginBuild(MutableModel, Object->GetPrivate()->MeshIdRegistry, Object->GetPrivate()->ImageIdRegistry, Object->GetPrivate()->MaterialIdRegistry, Object->GetPrivate()->SkeletalMeshIdRegistry);
 	
 	switch (OperationDataType)
 	{
 	case UE::Mutable::Private::EDataType::Layout:
 	{
 		check(PreviewLayoutViewer);
-		UE::Mutable::Private::TManagedPtr<const UE::Mutable::Private::FLayout> MutableLayout = System->GetPrivate()->BuildLayout(ExternalResourceProvider, MutableModel, PreviewParameters.Get(), SelectedOperationAddress);
+		UE::Mutable::Private::TManagedPtr<const UE::Mutable::Private::FLayout> MutableLayout = System->BuildLayout(ExternalResourceProvider, MutableModel, PreviewParameters.Get(), SelectedOperationAddress);
 		PreviewLayoutViewer->SetLayout(MutableLayout);
 		break;
 	}
@@ -2912,7 +2909,7 @@ void SMutableCodeViewer::Tick(const FGeometry& AllottedGeometry, const double In
 	case UE::Mutable::Private::EDataType::Image:
 	{
 		check(PreviewImageViewer);
-		UE::Mutable::Private::TManagedPtr<const UE::Mutable::Private::FImage> MutableImage = System->GetPrivate()->BuildImage(ExternalResourceProvider, MutableModel, PreviewParameters.Get(), SelectedOperationAddress, MipsToSkip);
+		UE::Mutable::Private::TManagedPtr<const UE::Mutable::Private::FImage> MutableImage = System->BuildImage(ExternalResourceProvider, MutableModel, PreviewParameters.Get(), SelectedOperationAddress, MipsToSkip);
 		PreviewImageViewer->SetImage(MutableImage, 0);
 		break;
 	}
@@ -2920,7 +2917,7 @@ void SMutableCodeViewer::Tick(const FGeometry& AllottedGeometry, const double In
 	case UE::Mutable::Private::EDataType::Mesh:
 	{
 		check(PreviewMeshViewer);
-		UE::Mutable::Private::TManagedPtr<const UE::Mutable::Private::FMesh> MutableMesh = System->GetPrivate()->BuildMesh(ExternalResourceProvider, MutableModel, PreviewParameters.Get(), SelectedOperationAddress, UE::Mutable::Private::EMeshContentFlags::AllFlags);
+		UE::Mutable::Private::TManagedPtr<const UE::Mutable::Private::FMesh> MutableMesh = System->BuildMesh(ExternalResourceProvider, MutableModel, PreviewParameters.Get(), SelectedOperationAddress, UE::Mutable::Private::EMeshContentFlags::AllFlags);
 		PreviewMeshViewer->SetMesh(MutableMesh);
 		break;
 	}
@@ -2928,7 +2925,7 @@ void SMutableCodeViewer::Tick(const FGeometry& AllottedGeometry, const double In
 	case UE::Mutable::Private::EDataType::Instance:
 	{
 		check(PreviewInstanceViewer);
-		UE::Mutable::Private::TManagedPtr<const UE::Mutable::Private::FInstance> MutableInstance = System->GetPrivate()->BuildInstance(ExternalResourceProvider, MutableModel, PreviewParameters.Get(), SelectedOperationAddress);
+		UE::Mutable::Private::TManagedPtr<const UE::Mutable::Private::FInstance> MutableInstance = System->BuildInstance(ExternalResourceProvider, MutableModel, PreviewParameters.Get(), SelectedOperationAddress);
 		PreviewInstanceViewer->SetInstance(MutableInstance, MutableModel, PreviewParameters, *System);
 		break;
 	}
@@ -2936,7 +2933,7 @@ void SMutableCodeViewer::Tick(const FGeometry& AllottedGeometry, const double In
 	case UE::Mutable::Private::EDataType::Bool:
 	{
 		check(PreviewBoolViewer);
-		const bool MutableBool = System->GetPrivate()->BuildBool(ExternalResourceProvider, MutableModel, PreviewParameters.Get(), SelectedOperationAddress);
+		const bool MutableBool = System->BuildBool(ExternalResourceProvider, MutableModel, PreviewParameters.Get(), SelectedOperationAddress);
 		PreviewBoolViewer->SetBool(MutableBool);
 		break;
 	}
@@ -2944,7 +2941,7 @@ void SMutableCodeViewer::Tick(const FGeometry& AllottedGeometry, const double In
 	case UE::Mutable::Private::EDataType::Int:
 	{
 		check(PreviewIntViewer);
-		const int32 MutableInt = System->GetPrivate()->BuildInt(ExternalResourceProvider, MutableModel, PreviewParameters.Get(), SelectedOperationAddress);
+		const int32 MutableInt = System->BuildInt(ExternalResourceProvider, MutableModel, PreviewParameters.Get(), SelectedOperationAddress);
 		PreviewIntViewer->SetInt(MutableInt);
 		break;
 	}
@@ -2952,7 +2949,7 @@ void SMutableCodeViewer::Tick(const FGeometry& AllottedGeometry, const double In
 	case UE::Mutable::Private::EDataType::Scalar:
 	{
 		check(PreviewScalarViewer);
-		const float MutableScalar = System->GetPrivate()->BuildScalar(ExternalResourceProvider, MutableModel, PreviewParameters.Get(), SelectedOperationAddress);
+		const float MutableScalar = System->BuildScalar(ExternalResourceProvider, MutableModel, PreviewParameters.Get(), SelectedOperationAddress);
 		PreviewScalarViewer->SetScalar(MutableScalar);
 		break;
 	}
@@ -2960,7 +2957,7 @@ void SMutableCodeViewer::Tick(const FGeometry& AllottedGeometry, const double In
 	case UE::Mutable::Private::EDataType::String:
 	{
 		check(PreviewStringViewer);
-		UE::Mutable::Private::TManagedPtr<const UE::Mutable::Private::String> MutableString = System->GetPrivate()->BuildString(ExternalResourceProvider, MutableModel, PreviewParameters.Get(), SelectedOperationAddress);
+		UE::Mutable::Private::TManagedPtr<const UE::Mutable::Private::String> MutableString = System->BuildString(ExternalResourceProvider, MutableModel, PreviewParameters.Get(), SelectedOperationAddress);
 		const FText MutableText = FText::FromString(FString(MutableString->GetValue()));
 		PreviewStringViewer->SetString(MutableText);
 		break;
@@ -2969,7 +2966,7 @@ void SMutableCodeViewer::Tick(const FGeometry& AllottedGeometry, const double In
 	case UE::Mutable::Private::EDataType::Color:
 	{
 		check(PreviewColorViewer);
-		FVector4f Color = System->GetPrivate()->BuildColour(ExternalResourceProvider, MutableModel, PreviewParameters.Get(), SelectedOperationAddress);
+		FVector4f Color = System->BuildColour(ExternalResourceProvider, MutableModel, PreviewParameters.Get(), SelectedOperationAddress);
 		PreviewColorViewer->SetColor(Color);
 		break;
 	}
@@ -2977,7 +2974,7 @@ void SMutableCodeViewer::Tick(const FGeometry& AllottedGeometry, const double In
 	case UE::Mutable::Private::EDataType::Projector:
 	{
 		check (PreviewProjectorViewer);
-		UE::Mutable::Private::FProjector Projector = System->GetPrivate()->BuildProjector(ExternalResourceProvider, MutableModel, PreviewParameters.Get(), SelectedOperationAddress);
+		UE::Mutable::Private::FProjector Projector = System->BuildProjector(ExternalResourceProvider, MutableModel, PreviewParameters.Get(), SelectedOperationAddress);
 		PreviewProjectorViewer->SetProjector(Projector);
 		break;
 	}
@@ -2999,7 +2996,7 @@ void SMutableCodeViewer::Tick(const FGeometry& AllottedGeometry, const double In
 		break;
 	}
 
-	System->GetPrivate()->EndBuild();
+	System->EndBuild();
 }
 
 
@@ -3008,9 +3005,9 @@ UE::Mutable::Private::TManagedPtr<const UE::Mutable::Private::FImage> SMutableCo
 	UE::Mutable::Private::FSettings Settings;
 	const TSharedPtr<UE::Mutable::Private::FSystem> System = MakeShared<UE::Mutable::Private::FSystem>(Settings);
 
-	System->GetPrivate()->BeginBuild(MutableModel, Object->GetPrivate()->MeshIdRegistry, Object->GetPrivate()->ImageIdRegistry, Object->GetPrivate()->MaterialIdRegistry, Object->GetPrivate()->SkeletalMeshIdRegistry);
-	UE::Mutable::Private::TManagedPtr<const UE::Mutable::Private::FImage> MutableImage = System->GetPrivate()->BuildImage(ExternalResourceProvider, MutableModel, PreviewParameters.Get(), TargetOpAddress, 0);
-	System->GetPrivate()->EndBuild();
+	System->BeginBuild(MutableModel, Object->GetPrivate()->MeshIdRegistry, Object->GetPrivate()->ImageIdRegistry, Object->GetPrivate()->MaterialIdRegistry, Object->GetPrivate()->SkeletalMeshIdRegistry);
+	UE::Mutable::Private::TManagedPtr<const UE::Mutable::Private::FImage> MutableImage = System->BuildImage(ExternalResourceProvider, MutableModel, PreviewParameters.Get(), TargetOpAddress, 0);
+	System->EndBuild();
 	
 	return MutableImage;
 }
@@ -3236,7 +3233,7 @@ void FMutableCodeTreeElement::SetElementCurrentState(const int32& InStateIndex)
 
 	// Check for an out of bounds value
 	check(MutableModel);
-	UE::Mutable::Private::FProgram& MutableProgram = MutableModel->GetPrivate()->Program;
+	UE::Mutable::Private::FProgram& MutableProgram = MutableModel->Program;
 	check(InStateIndex >= 0 && InStateIndex < MutableProgram.States.Num());
 
 	CurrentMutableStateIndex = InStateIndex;
@@ -3265,7 +3262,7 @@ void FMutableCodeTreeElement::SetElementCurrentState(const int32& InStateIndex)
 
 void FMutableCodeTreeElement::GenerateLabelText()
 {
-	const UE::Mutable::Private::FProgram& Program = MutableModel->GetPrivate()->Program;
+	const UE::Mutable::Private::FProgram& Program = MutableModel->Program;
 	const UE::Mutable::Private::EOpType OperationType = Program.GetOpType(MutableOperation);
 	FString OpName = UE::Mutable::Private::s_opNames[static_cast<int32>(OperationType)];
 	OpName.TrimEndInline();
diff --git a/Engine/Plugins/Mutable/Source/CustomizableObjectEditor/Private/MuCOE/SMutableCodeViewer.h b/Engine/Plugins/Mutable/Source/CustomizableObjectEditor/Private/MuCOE/SMutableCodeViewer.h
--- a/Engine/Plugins/Mutable/Source/CustomizableObjectEditor/Private/MuCOE/SMutableCodeViewer.h
+++ b/Engine/Plugins/Mutable/Source/CustomizableObjectEditor/Private/MuCOE/SMutableCodeViewer.h
@@ -2,7 +2,7 @@
 
 #pragma once
 
-#include "MuR/ModelPrivate.h"
+#include "MuR/Model.h"
 #include "MuT/TypeInfo.h"
 #include "UObject/GCObject.h"
 #include "Widgets/SCompoundWidget.h"
diff --git a/Engine/Plugins/Mutable/Source/CustomizableObjectEditor/Private/MuCOE/SMutableConstantsWidget.h b/Engine/Plugins/Mutable/Source/CustomizableObjectEditor/Private/MuCOE/SMutableConstantsWidget.h
--- a/Engine/Plugins/Mutable/Source/CustomizableObjectEditor/Private/MuCOE/SMutableConstantsWidget.h
+++ b/Engine/Plugins/Mutable/Source/CustomizableObjectEditor/Private/MuCOE/SMutableConstantsWidget.h
@@ -5,7 +5,7 @@
 #include "MuR/Image.h"
 #include "MuR/Mesh.h"
 #include "MuR/Operations.h"
-#include "MuR/ModelPrivate.h"
+#include "MuR/Model.h"
 #include "Widgets/Views/SHeaderRow.h"
 #include "Widgets/Views/STileView.h"
 
diff --git a/Engine/Plugins/Mutable/Source/CustomizableObjectEditor/Private/MuCOE/SMutableCurveViewer.h b/Engine/Plugins/Mutable/Source/CustomizableObjectEditor/Private/MuCOE/SMutableCurveViewer.h
--- a/Engine/Plugins/Mutable/Source/CustomizableObjectEditor/Private/MuCOE/SMutableCurveViewer.h
+++ b/Engine/Plugins/Mutable/Source/CustomizableObjectEditor/Private/MuCOE/SMutableCurveViewer.h
@@ -4,7 +4,7 @@
 
 #include "CurveModel.h"
 #include "Curves/RichCurve.h"
-#include "MuR/ParametersPrivate.h"
+#include "MuR/Parameters.h"
 #include "Widgets/SCompoundWidget.h"
 
 enum class ECurvePointType : uint8;
diff --git a/Engine/Plugins/Mutable/Source/CustomizableObjectEditor/Private/MuCOE/SMutableInstanceViewer.cpp b/Engine/Plugins/Mutable/Source/CustomizableObjectEditor/Private/MuCOE/SMutableInstanceViewer.cpp
--- a/Engine/Plugins/Mutable/Source/CustomizableObjectEditor/Private/MuCOE/SMutableInstanceViewer.cpp
+++ b/Engine/Plugins/Mutable/Source/CustomizableObjectEditor/Private/MuCOE/SMutableInstanceViewer.cpp
@@ -11,7 +11,7 @@
 #include "PropertyCustomizationHelpers.h"
 #include "Widgets/Layout/SScrollBox.h"
 #include "Widgets/Views/SListView.h"
-#include "MuR/SystemPrivate.h"
+#include "MuR/System.h"
 
 class ITableRow;
 class STableViewBase;
@@ -285,7 +285,7 @@ FString ToString(const UE::Mutable::Private::TPassthroughObjectPtr<UObject>& Pas
 
 
 
-void SMutableInstanceViewer::SetInstance(const UE::Mutable::Private::TManagedPtr<const UE::Mutable::Private::FInstance>& InInstance, const TSharedPtr<const UE::Mutable::Private::FModel>& Model, const TSharedPtr<const UE::Mutable::Private::FParameters>& Parameters, const UE::Mutable::Private::FSystem& System)
+void SMutableInstanceViewer::SetInstance(const UE::Mutable::Private::TManagedPtr<const UE::Mutable::Private::FInstance>& InInstance, const TSharedPtr<const UE::Mutable::Private::FModel>& Model, const TSharedPtr<const UE::Mutable::Private::FParameters>& Parameters, UE::Mutable::Private::FSystem& System)
 {
 	if (InInstance != MutableInstance)
 	{
@@ -321,7 +321,7 @@ void SMutableInstanceViewer::SetInstance(const UE::Mutable::Private::TManagedPtr
 					if (MeshId)
 					{
 						UE::Mutable::Private::OP::ADDRESS MeshRootAddress = MeshId.GetKey()->Address;
-						UE::Mutable::Private::TManagedPtr<const UE::Mutable::Private::FMesh> MutableMesh = System.GetPrivate()->BuildMesh(
+						UE::Mutable::Private::TManagedPtr<const UE::Mutable::Private::FMesh> MutableMesh = System.BuildMesh(
 								ExternalResourceProvider,
 								Model, 
 								Parameters.Get(), 
diff --git a/Engine/Plugins/Mutable/Source/CustomizableObjectEditor/Private/MuCOE/SMutableInstanceViewer.h b/Engine/Plugins/Mutable/Source/CustomizableObjectEditor/Private/MuCOE/SMutableInstanceViewer.h
--- a/Engine/Plugins/Mutable/Source/CustomizableObjectEditor/Private/MuCOE/SMutableInstanceViewer.h
+++ b/Engine/Plugins/Mutable/Source/CustomizableObjectEditor/Private/MuCOE/SMutableInstanceViewer.h
@@ -62,7 +62,7 @@ public:
 	void Construct(const FArguments& InArgs, const TSharedPtr<FUnrealMutableResourceProvider>& InExternalResourceProvider);
 
 	/** Set the Mutable Instance to be used for this widget */
-	void SetInstance(const UE::Mutable::Private::TManagedPtr<const UE::Mutable::Private::FInstance>&, const TSharedPtr<const UE::Mutable::Private::FModel>&, const TSharedPtr<const UE::Mutable::Private::FParameters>&, const UE::Mutable::Private::FSystem&);
+	void SetInstance(const UE::Mutable::Private::TManagedPtr<const UE::Mutable::Private::FInstance>&, const TSharedPtr<const UE::Mutable::Private::FModel>&, const TSharedPtr<const UE::Mutable::Private::FParameters>&, UE::Mutable::Private::FSystem&);
 	
 	/** */
 	TSharedRef<SWidget> GenerateLODsListView(const TSharedPtr<TArray<TSharedPtr<FMutableInstanceViewerLODElement>>>& InBufferChannelElements);
diff --git a/Engine/Plugins/Mutable/Source/CustomizableObjectEditor/Private/MuCOE/SMutableProjectorViewer.h b/Engine/Plugins/Mutable/Source/CustomizableObjectEditor/Private/MuCOE/SMutableProjectorViewer.h
--- a/Engine/Plugins/Mutable/Source/CustomizableObjectEditor/Private/MuCOE/SMutableProjectorViewer.h
+++ b/Engine/Plugins/Mutable/Source/CustomizableObjectEditor/Private/MuCOE/SMutableProjectorViewer.h
@@ -2,7 +2,7 @@
 
 #pragma once
 
-#include "MuR/ParametersPrivate.h"
+#include "MuR/Parameters.h"
 #include "Widgets/SCompoundWidget.h"
 
 namespace UE::Mutable::Private { enum class EProjectorType : uint32; }
diff --git a/Engine/Plugins/Mutable/Source/MutableRuntime/Internal/MuR/CodeVisitor.h b/Engine/Plugins/Mutable/Source/MutableRuntime/Internal/MuR/CodeVisitor.h
--- a/Engine/Plugins/Mutable/Source/MutableRuntime/Internal/MuR/CodeVisitor.h
+++ b/Engine/Plugins/Mutable/Source/MutableRuntime/Internal/MuR/CodeVisitor.h
@@ -6,12 +6,11 @@
 #include "HAL/UnrealMemory.h"
 #include "Misc/AssertionMacros.h"
 #include "MuR/Model.h"
-#include "MuR/ModelPrivate.h"
 #include "MuR/Operations.h"
-#include "MuR/ParametersPrivate.h"
+#include "MuR/Parameters.h"
 #include "MuR/Ptr.h"
 #include "MuR/System.h"
-#include "MuR/SystemPrivate.h"
+#include "MuR/System.h"
 
 namespace UE::Mutable::Private
 {
@@ -94,7 +93,7 @@ namespace UE::Mutable::Private
         }
 
 
-        void Traverse( OP::ADDRESS Root, FProgram& Program )
+        void Traverse( OP::ADDRESS Root, const FProgram& Program )
         {
             Pending.Reserve( Program.OpAddress.Num() );
 
@@ -103,7 +102,7 @@ namespace UE::Mutable::Private
             Recurse( Program );
          }
 
-        void FullTraverse( FProgram& Program )
+        void FullTraverse( const FProgram& Program )
         {
             // Visit all the state roots
             for ( int32 StateIndex=0; StateIndex<Program.States.Num(); ++StateIndex)
@@ -119,7 +118,7 @@ namespace UE::Mutable::Private
         /** Do the actual work by overriding this in the derived classes.
          * Return true if the traverse has to continue with the children of the given address.
 		 */
-        virtual bool Visit( OP::ADDRESS, FProgram& ) = 0;
+        virtual bool Visit(OP::ADDRESS, const FProgram&) = 0;
 
         //! Operations to be processed
         struct FPending
@@ -156,7 +155,7 @@ namespace UE::Mutable::Private
 		TArray<TArray<int32>> Visited;
 
         //! Process all the pending operations and visit all children if necessary
-        void Recurse( FProgram& Program )
+        void Recurse(const FProgram& Program)
         {
 			Visited.Empty();
 			Visited.SetNum(Program.OpAddress.Num());
@@ -219,7 +218,7 @@ namespace UE::Mutable::Private
 
         DiscreteCoveredCodeVisitorBase
             (
-                FSystem::Private* InSystem,
+                FSystem* InSystem,
                 const TSharedPtr<const FModel>& InModel,
                 const TSharedPtr<const FParameters>& InParams,
                 unsigned InLodMask,
@@ -240,12 +239,12 @@ namespace UE::Mutable::Private
         {
             PARENT::SetDefaultState( STATE() );
 
-            PARENT::Traverse( at, Model->GetPrivate()->Program );
+            PARENT::Traverse( at, Model->Program );
         }
 
     protected:
 
-        virtual bool Visit( OP::ADDRESS Address, FProgram& Program )
+        virtual bool Visit(OP::ADDRESS Address, const FProgram& Program)
         {
             bool bRecurse = true;
 
@@ -447,7 +446,7 @@ namespace UE::Mutable::Private
 
 
     protected:
-        FSystem::Private* System = nullptr;
+        FSystem* System = nullptr;
 		TSharedPtr<const FModel> Model;
         const FParameters* Params = nullptr;
         uint32 LODMask = 0;
@@ -477,7 +476,7 @@ namespace UE::Mutable::Private
 
         UniqueDiscreteCoveredCodeVisitor
             (
-                FSystem::Private* InSystem,
+                FSystem* InSystem,
 				const TSharedPtr<const FModel>& InModel,
                 const TSharedPtr<const FParameters>& InParams,
                 uint32 InLodMask
diff --git a/Engine/Plugins/Mutable/Source/MutableRuntime/Internal/MuR/MeshPrivate.h b/Engine/Plugins/Mutable/Source/MutableRuntime/Internal/MuR/MeshPrivate.h
--- a/Engine/Plugins/Mutable/Source/MutableRuntime/Internal/MuR/MeshPrivate.h
+++ b/Engine/Plugins/Mutable/Source/MutableRuntime/Internal/MuR/MeshPrivate.h
@@ -6,7 +6,7 @@
 
 #include "MuR/Layout.h"
 #include "MuR/Skeleton.h"
-#include "MuR/SerialisationPrivate.h"
+#include "MuR/Serialisation.h"
 #include "MuR/ConvertData.h"
 
 #include "CoreFwd.h"
diff --git a/Engine/Plugins/Mutable/Source/MutableRuntime/Internal/MuR/Model.h b/Engine/Plugins/Mutable/Source/MutableRuntime/Internal/MuR/Model.h
--- a/Engine/Plugins/Mutable/Source/MutableRuntime/Internal/MuR/Model.h
+++ b/Engine/Plugins/Mutable/Source/MutableRuntime/Internal/MuR/Model.h
@@ -3,16 +3,27 @@
 #pragma once
 
 #include "MuR/Parameters.h"
-#include "MuR/Ptr.h"
+#include "MuR/Model.h"
+#include "MuR/System.h"
 #include "MuR/RefCounted.h"
-#include "Templates/Function.h"
+#include "MuR/Serialisation.h"
+#include "MuR/Operations.h"
+#include "MuR/ExtensionData.h"
+#include "MuR/Image.h"
+#include "MuR/Mesh.h"
+#include "MuR/MutableRuntimeModule.h"
+
 #include "Templates/SharedPointer.h"
+#include "ExternalOperationProvider.h"
 
 class UMaterialInterface;
 
 
 #define UE_API MUTABLERUNTIME_API
 
+#define MUTABLE_MAX_RUNTIME_PARAMETERS_PER_STATE	65
+#define MUTABLE_GROW_BORDER_VALUE					2
+
 
 namespace UE::Mutable::Private
 {
@@ -22,6 +33,1140 @@ namespace UE::Mutable::Private
 	class FSystem;
     class FModelParametersGenerator;
 
+	
+	/** Used to debug and log. */
+	constexpr bool DebugRom = false;
+	constexpr bool DebugRomAll = false;
+	constexpr int32 DebugRomIndex = 44;
+	constexpr int32 DebugImageIndex = 9;
+
+	struct FConstantResourceIndex
+	{
+		uint32 Index : 31;
+		/** This may mean that the resource needs to be looked up on a different array. */
+		uint32 Streamable : 1;
+	};
+	static_assert(sizeof(FConstantResourceIndex) == 4);
+	MUTABLE_DEFINE_POD_SERIALISABLE(FConstantResourceIndex);
+	//MUTABLE_DEFINE_POD_VECTOR_SERIALISABLE(FConstantResourceIndex);
+
+	/** This is encoded with minimal bits. Make sure to review all uses if extended. */
+	enum class ERomDataType : uint32
+	{
+		Image = 0,
+		Mesh  = 1
+	};
+
+    /** Data stored for a rom even if it is not loaded.
+	* This struct is size-sensitive since there may be many roms and it is always loaded in memory when a CO is.
+	*/
+    struct FRomDataRuntime
+    {
+		/** Size of the rom */
+		uint32 Size : 30;
+
+		/** Index of the resource in its type-specific array. See ERomDataType. */
+		uint32 ResourceType : 1;
+
+		/** Properties of the rom data. */
+		uint32 IsHighRes : 1;
+
+		// TODO: Store the offset here and delete the FModelStreamableBlock map?
+		/** Offset in file */
+		// uint32 Offset = 0;
+    };
+
+	/** Not critical to keep this size, but it is memory-usage sensitive. */
+	static_assert(sizeof(FRomDataRuntime) == 4);
+	MUTABLE_DEFINE_POD_SERIALISABLE(FRomDataRuntime);
+
+	struct FRomDataCompile
+	{
+		/** ID used to identify the origin of this data and used for grouping. */
+		uint32 SourceId;
+	};
+	MUTABLE_DEFINE_POD_SERIALISABLE(FRomDataCompile);
+
+    //!
+    template<typename DATA>
+    inline void AppendCode(TArray<uint8>& Code, const DATA& Data )
+    {
+        int32 Pos = Code.Num();
+        Code.SetNum( Pos+sizeof(DATA), EAllowShrinking::No);
+		FMemory::Memcpy (&Code[Pos], &Data, sizeof(DATA));
+    }
+
+	//!
+	struct FImageLODRange
+	{
+		int32 FirstIndex    = 0;
+		uint16 ImageSizeX   = 0;
+		uint16 ImageSizeY   = 0;
+		uint8 LODCount      = 0;
+		uint8 NumLODsInTail = 0; 
+		uint8 Flags         = 0;
+		EImageFormat ImageFormat = EImageFormat::None;
+	};
+	MUTABLE_DEFINE_POD_SERIALISABLE(FImageLODRange);
+
+	struct FMeshContentRange
+	{
+		static constexpr uint32 FirstIndexMaxBits   = 24;
+		static constexpr uint32 ContentFlagsMaxBits = 32 - FirstIndexMaxBits;
+		static constexpr uint32 FirstIndexBitMask   = (1 << FirstIndexMaxBits) - 1; 
+
+		static_assert(FirstIndexMaxBits < 32);
+		static_assert(ContentFlagsMaxBits >= sizeof(EMeshContentFlags)*8);
+
+		// NOTE: Bitfields layout can be implementation defined, here we want consistency across all
+		// compilers so that the struct can be POD serializable.
+		uint32 FirstIndex_ContentFlags = 0; // Low bits are FirstIndex, high bits are ContentFlags.
+		uint32 MeshIDPrefix = 0;
+
+		FORCEINLINE EMeshContentFlags GetContentFlags() const
+		{
+			return static_cast<EMeshContentFlags>(
+					(FirstIndex_ContentFlags >> FirstIndexMaxBits) &
+					((1 << ContentFlagsMaxBits) - 1));
+		}
+
+		FORCEINLINE uint32 GetFirstIndex() const
+		{
+			return FirstIndex_ContentFlags & FirstIndexBitMask;
+		}
+
+		FORCEINLINE void SetContentFlags(EMeshContentFlags ContentFlags)
+		{
+			check(uint32(ContentFlags) < (1 << ContentFlagsMaxBits));
+
+			FirstIndex_ContentFlags = 
+					(FirstIndex_ContentFlags & FirstIndexBitMask) | 
+					((uint32(ContentFlags) << FirstIndexMaxBits));
+		}
+
+		FORCEINLINE void SetFirstIndex(uint32 FirstIndex) 
+		{
+			check(FirstIndex < ((1 << FirstIndexMaxBits)));
+
+			FirstIndex_ContentFlags = 
+					(FirstIndex_ContentFlags & ~FirstIndexBitMask) | 
+					(FirstIndex & FirstIndexBitMask);
+		}
+	};
+	static_assert(sizeof(FMeshContentRange) == sizeof(uint32)*2);
+	MUTABLE_DEFINE_POD_SERIALISABLE(FMeshContentRange);
+
+	struct FExtensionDataConstant
+	{
+		// This should always be valid, but if the state is Unloaded it won't be usable.
+		//
+		// Avoid storing references to this Data in Memory while the state is Unloaded.
+		TManagedPtr<const FExtensionData> Data;
+
+		enum class ELoadState : uint8
+		{
+			Invalid,
+			Unloaded,
+			FailedToLoad,
+			CurrentlyLoaded,
+			AlwaysLoaded
+		};
+
+		void Serialise(FOutputArchive& Arch) const
+		{
+			Arch << Data;
+		}
+		
+		void Unserialise(FInputArchive& Arch)
+		{
+			Arch >> Data;
+
+			check(Data);
+			check(Data->Origin == FExtensionData::EOrigin::ConstantAlwaysLoaded || Data->Origin == FExtensionData::EOrigin::ConstantStreamed);
+		}
+	};
+
+    //!
+    struct FProgram
+    {
+        FProgram()
+        {
+            // Add the null instruction at address 0.
+            // TODO: Will do it in the linker
+            AppendCode( ByteCode, EOpType::NONE );
+            OpAddress.Add(0);
+        }
+
+        struct FState
+        {
+            /** Name of the state */
+            FString Name;
+
+            /** First instruction of the full build of an instance in this state */
+            OP::ADDRESS Root = 0;
+
+            /** 
+			 * List of parameters index (to FProgram::Parameters) of the runtime parameters of
+             * this state.
+			 */
+			TArray<int> m_runtimeParameters;
+
+            /** List of instructions that need to be cached to efficiently update this state */
+			TArray<OP::ADDRESS> m_updateCache;
+
+            /** 
+			 * List of root instructions for the dynamic resources that depend on the runtime
+			 * parameters of this state, with a mask of relevant runtime parameters.
+             * The mask has a bit on for every runtime parameter in the m_runtimeParameters array.
+			 * The uint64 is linked to MUTABLE_MAX_RUNTIME_PARAMETERS_PER_STATE
+			 */
+			TArray< TPair<OP::ADDRESS,uint64> > m_dynamicResources;
+
+            /** */
+            inline void Serialise( FOutputArchive& arch ) const
+            {
+                arch << Name;
+                arch << Root;
+                arch << m_runtimeParameters;
+                arch << m_updateCache;
+                arch << m_dynamicResources;
+            }
+
+            /** */
+            inline void Unserialise( FInputArchive& arch )
+            {
+                arch >> Name;
+                arch >> Root;
+                arch >> m_runtimeParameters;
+                arch >> m_updateCache;
+                arch >> m_dynamicResources;
+            }
+
+            /** Returns the mask of parameters (from the runtime parameter list of this state) including the parameters that 
+			* are relevant for the dynamic resource at the given address.
+			*/
+            uint64 IsDynamic( OP::ADDRESS at ) const
+            {
+                uint64 res = 0;
+
+                for ( int32 i=0; i<m_dynamicResources.Num(); ++i )
+                {
+                    if ( m_dynamicResources[i].Key==at )
+                    {
+                        res = m_dynamicResources[i].Value;
+						break;
+                    }
+                }
+
+                return res;
+            }
+
+            /** */
+            bool IsUpdateCache( OP::ADDRESS at ) const
+            {
+                bool res = false;
+
+                for (int32 i=0; !res && i<m_updateCache.Num(); ++i )
+                {
+                    if ( m_updateCache[i]==at )
+                    {
+                        res = true;
+                    }
+                }
+
+                return res;
+            }
+
+            /** */ 
+            void AddUpdateCache( OP::ADDRESS at )
+            {
+                if ( !IsUpdateCache(at) )
+                {
+                    m_updateCache.Add( at );
+                }
+            }
+        };
+
+        /** Location in the ByteCode of the beginning of each operation */
+		TArray<uint32> OpAddress;
+
+        /** Byte-coded representation of the program, using variable-sized op data. */
+		TArray<uint8> ByteCode;
+
+        /** */
+		TArray<FState> States;
+
+		/** Data for every rom required in-game. */
+		TArray<FRomDataRuntime> Roms;
+
+		/** Data for every rom required at compile-time. It is empty in cooked data. */
+		TArray<FRomDataCompile> RomsCompileData;
+	
+		/** Loaded roms worth tracking (only images and meshes for now). Stores the rom's data type.*/
+		mutable TSparseArray<uint8> LoadedMemTrackedRoms;
+
+		/** Constant image mip data is split in 2 sets: ConstantImageLODsPermanent constains data that is always loaded. 
+		* Index with FConstantResourceIndex::Index, when Streamable is 0.
+		*/
+		TArray<TManagedPtr<const FImage>> ConstantImageLODsPermanent;
+
+		/** Constant image mip data is split in 2 sets: ConstantImageLODsStreamed constains data that is streamed in and out. 
+		* Index with FConstantResourceIndex::Index, when Streamable is 1.
+		* This part is empty for an unused Model, and shouldn't be serialised.
+		*/
+		mutable TMap<uint32, TManagedPtr<const FImage>> ConstantImageLODsStreamed;
+
+		/** Constant image mip chain indices: ranges in this array are defined in FImageLODRange and the indices here refer to ConstantImageLODs. */
+		TArray<FConstantResourceIndex> ConstantImageLODIndices;
+
+		/** Constant image data. */
+		TArray<FImageLODRange> ConstantImages;
+
+		/** Constant mesh content indices: ranges in this array are defined in FMeshContentRange and the indices here refer to ConstantMeshes. */
+		TArray<FConstantResourceIndex> ConstantMeshContentIndices;
+		
+		/** Constant mesh data */
+		TArray<FMeshContentRange> ConstantMeshes;
+
+		/** Constant mesh data is split in 2 sets: ConstantMeshesPermanent constains data that is always loaded.
+		* Index with FConstantResourceIndex::Index, when Streamable is 0.
+		*/
+		TArray<TManagedPtr<const FMesh>> ConstantMeshesPermanent;
+
+		/** Constant mesh data is split in 2 sets: ConstantMeshesStreamed constains data that is streamed in and out.
+		* Index with FConstantResourceIndex::Index, when Streamable is 1.
+		* This part is empty for an unused Model, and shouldn't be serialised.
+		*/
+		mutable TMap<uint32, TManagedPtr<const FMesh>> ConstantMeshesStreamed;
+
+		/** Constant FExtensionData */
+		TArray<FExtensionDataConstant> ConstantExtensionData;
+
+        /** Constant string data */
+		TArray<FString> ConstantStrings;
+
+		/** */
+		TArray<TArray<uint32>> ConstantUInt32Lists;
+
+		/** */
+		TArray<TArray<uint64>> ConstantUInt64Lists;
+
+        /** Constant layout data */
+		TArray<TManagedPtr<const FLayout>> ConstantLayouts;
+
+        /** Constant projectors */
+		TArray<FProjector> ConstantProjectors;
+
+        /** Constant matrices, usually used for transforms */
+		TArray<FMatrix44f> ConstantMatrices;
+
+		/** Constant shapes */
+		TArray<FShape> ConstantShapes;
+
+        /** Constant curves */
+		TArray<FRichCurve> ConstantCurves;
+
+        /** Constant skeletons */
+		TArray<TManagedPtr<const FSkeleton>> ConstantSkeletons;
+
+		/** Constant Physics Bodies */
+		TArray<TManagedPtr<const FPhysicsBody>> ConstantPhysicsBodies;
+
+        /** FParameters of the model. */
+        /** The value stored here is the default value. */
+		TArray<FParameterDesc> Parameters;
+
+        /** Ranges for iteration of the model operations. */
+		TArray<FRangeDesc> Ranges;
+
+        /** 
+		 * List of parameter lists. These are used in several places, like storing the
+         * pregenerated list of parameters influencing a resource.
+         * The parameter lists are sorted. 
+		 */
+		TArray<TArray<uint16>> ParameterLists;
+
+    	/** Relevant parameters for the given instruction. */
+	   	TMap<OP::ADDRESS, int32> RelevantParameterList;
+
+
+		/** Constant Material Data */
+		TArray<TManagedPtr<const FMaterial>> ConstantMaterials;
+
+    	TSharedPtr<const IExternalOperationProvider> ExternalOperationProvider;
+    	
+#if WITH_EDITOR
+		/** 
+		 * State of the program. True unless the streamed resources were destroyed,
+		 * which could happen in the editor after recompiling the CO. 
+		 */
+		bool bIsValid = true;
+#endif
+        /** */
+        void Serialise(FOutputArchive& Arch) const
+        {
+            Arch << OpAddress;
+            Arch << ByteCode;
+            Arch << States;
+			Arch << Roms;
+			Arch << RomsCompileData;
+			Arch << ConstantImageLODsPermanent;
+			Arch << ConstantImageLODIndices;
+			Arch << ConstantImages;
+			Arch << ConstantMeshesPermanent;
+			Arch << ConstantMeshContentIndices;
+			Arch << ConstantMeshes;
+			Arch << ConstantExtensionData;
+			Arch << ConstantStrings;
+			Arch << ConstantUInt32Lists;
+			Arch << ConstantUInt64Lists;
+            Arch << ConstantLayouts;
+            Arch << ConstantProjectors;
+			Arch << ConstantMatrices;
+			Arch << ConstantShapes;
+            Arch << ConstantCurves;
+            Arch << ConstantSkeletons;
+            Arch << Parameters;
+            Arch << Ranges;
+            Arch << ParameterLists;
+            Arch << RelevantParameterList;
+			Arch << ConstantMaterials;
+        }
+
+        /** */
+        void Unserialise(FInputArchive& Arch)
+        {
+            Arch >> OpAddress;
+            Arch >> ByteCode;
+            Arch >> States;
+			Arch >> Roms;
+			Arch >> RomsCompileData;
+			Arch >> ConstantImageLODsPermanent;
+			Arch >> ConstantImageLODIndices;
+			Arch >> ConstantImages;
+			Arch >> ConstantMeshesPermanent;
+			Arch >> ConstantMeshContentIndices;
+			Arch >> ConstantMeshes;
+			Arch >> ConstantExtensionData;
+			Arch >> ConstantStrings;
+			Arch >> ConstantUInt32Lists;
+			Arch >> ConstantUInt64Lists;
+            Arch >> ConstantLayouts;
+            Arch >> ConstantProjectors;
+			Arch >> ConstantMatrices;
+			Arch >> ConstantShapes;
+            Arch >> ConstantCurves;
+            Arch >> ConstantSkeletons;
+            Arch >> Parameters;
+            Arch >> Ranges;
+            Arch >> ParameterLists;
+            Arch >> RelevantParameterList;
+            Arch >> ConstantMaterials;
+        }
+
+        /** Debug method that logs the top used instruction types. */
+        void LogHistogram() const;
+
+		/** Return true if the given ROM is loaded. */
+		FORCEINLINE bool IsRomLoaded(int32 RomIndex) const
+		{
+			switch (Roms[RomIndex].ResourceType)
+			{
+				case uint32(ERomDataType::Image):
+					return ConstantImageLODsStreamed.Contains(RomIndex);
+				case uint32(ERomDataType::Mesh):
+					return ConstantMeshesStreamed.Contains(RomIndex);
+				//TODO(Max): Make materials loadable and unloadable from a rom
+				default:
+					check(false);
+					break;
+			}
+
+			return false;
+		}
+
+		/** Unload a rom resource. 
+		* If a pointer to a size is passed in OutDataSize, the value will be set with the size of the unloaded rom.
+		*/
+		FORCEINLINE void UnloadRom(int32 RomIndex, int32* OutDataSize=nullptr) const
+		{
+			if (DebugRom && (DebugRomAll||RomIndex == DebugRomIndex))
+				UE_LOG(LogMutableCore, Log, TEXT("Unloading rom %d."), RomIndex);
+
+			if (LoadedMemTrackedRoms.IsValidIndex(RomIndex))
+			{
+				LoadedMemTrackedRoms.RemoveAt(RomIndex);
+			}
+
+			switch (Roms[RomIndex].ResourceType)
+			{
+			case uint32(ERomDataType::Image):
+			{
+				TManagedPtr<const FImage> Data;
+				ConstantImageLODsStreamed.RemoveAndCopyValue(RomIndex, Data);
+				if (OutDataSize && Data)
+				{
+					*OutDataSize = Data->GetDataSize();
+				}
+				break;
+			}
+			case uint32(ERomDataType::Mesh):
+			{
+				TManagedPtr<const FMesh> Data;
+				ConstantMeshesStreamed.RemoveAndCopyValue(RomIndex, Data);
+				if (OutDataSize && Data)
+				{
+					*OutDataSize = Data->GetDataSize();
+				}
+				break;
+			}
+			//TODO(Max): Make materials loadable and unloadable from a rom
+			default:
+				check(false);
+				break;
+			}
+		}
+
+		FORCEINLINE SSIZE_T GetRomBytes(int32 RomIndex) const
+		{
+			SSIZE_T Result = 0;
+			switch (Roms[RomIndex].ResourceType)
+			{
+			case uint32(ERomDataType::Image):
+			{
+				const TManagedPtr<const FImage>* Data = ConstantImageLODsStreamed.Find(RomIndex);
+				if (Data)
+				{
+					Result = (*Data)->GetDataSize();
+				}
+				break;
+			}
+			case uint32(ERomDataType::Mesh):
+			{
+				const TManagedPtr<const FMesh>* Data = ConstantMeshesStreamed.Find(RomIndex);
+				if (Data)
+				{
+					Result = (*Data)->GetDataSize();
+				}
+				break;
+			}
+			//TODO(Max): Make materials loadable and unloadable from a rom
+			default:
+				check(false);
+				break;
+			}
+
+			return Result;
+		}
+
+		FORCEINLINE void SetMeshRomValue(uint32 RomIndex, const TManagedPtr<FMesh>& Value) const
+		{
+			check(Roms[RomIndex].ResourceType == uint32(ERomDataType::Mesh));
+			check(!ConstantMeshesStreamed.Contains(RomIndex));
+
+			LoadedMemTrackedRoms.EmplaceAt(RomIndex, (uint8)Roms[RomIndex].ResourceType);
+			ConstantMeshesStreamed.Add(RomIndex, Value);
+		}
+
+		FORCEINLINE void SetImageRomValue(uint32 RomIndex, const TManagedPtr<FImage>& Value) const
+		{
+			check(Roms[RomIndex].ResourceType == uint32(ERomDataType::Image));
+			check(!ConstantImageLODsStreamed.Contains(RomIndex));
+
+			LoadedMemTrackedRoms.EmplaceAt(RomIndex, (uint8)Roms[RomIndex].ResourceType);
+			ConstantImageLODsStreamed.Add(RomIndex, Value);
+		}
+
+		OP::ADDRESS AddConstant(TManagedPtr<const FExtensionData> Data)
+		{
+			// Ensure unique
+			for (int32 Index = 0; Index < ConstantExtensionData.Num(); Index++)
+			{
+				const FExtensionData* Candidate = ConstantExtensionData[Index].Data.Get();
+				if (*Candidate == *Data)
+				{
+					return Index;
+				}
+			}
+
+			FExtensionDataConstant& NewConstant = ConstantExtensionData.AddDefaulted_GetRef();
+			NewConstant.Data = Data;
+
+			return ConstantExtensionData.Num() - 1;
+		}
+
+		OP::ADDRESS AddConstant( TManagedPtr<const FLayout> pLayout )
+        {
+            // Ensure unique
+            for (int32 i=0; i<ConstantLayouts.Num(); ++i)
+            {
+                if (ConstantLayouts[i]==pLayout)
+                {
+                    return (OP::ADDRESS)i;
+                }
+            }
+
+            OP::ADDRESS index = OP::ADDRESS( ConstantLayouts.Num() );
+            ConstantLayouts.Add( pLayout );
+            return index;
+        }
+
+        OP::ADDRESS AddConstant(const TManagedPtr<const FSkeleton>& Skeleton )
+        {
+            // Ensure unique
+            for ( int32 i=0; i<ConstantSkeletons.Num(); ++i)
+            {
+                if ( ConstantSkeletons[i]== Skeleton
+                     ||
+                     *ConstantSkeletons[i]==*Skeleton
+                     )
+                {
+                    return (OP::ADDRESS)i;
+                }
+            }
+
+            OP::ADDRESS index = OP::ADDRESS( ConstantSkeletons.Num() );
+            ConstantSkeletons.Add(Skeleton);
+            return index;
+        }
+
+        OP::ADDRESS AddConstant(const TManagedPtr<const FPhysicsBody>& PhysicsBody)
+        {
+            // Ensure unique
+            for (int32 i = 0; i < ConstantPhysicsBodies.Num(); ++i)
+            {
+                if (ConstantPhysicsBodies[i] == PhysicsBody || *ConstantPhysicsBodies[i] == *PhysicsBody )
+                {
+                    return (OP::ADDRESS)i;
+                }
+            }
+
+            OP::ADDRESS index = OP::ADDRESS(ConstantPhysicsBodies.Num());
+            ConstantPhysicsBodies.Add(PhysicsBody);
+            return index;
+        }
+
+        OP::CONSTANT_STRING_INDEX AddConstant(const FString& Str)
+        {     
+			return (OP::CONSTANT_STRING_INDEX)ConstantStrings.AddUnique(Str);
+        }
+		
+		OP::ADDRESS AddConstant(const TArray<uint32>& UInt32List)
+		{
+			return (OP::ADDRESS)ConstantUInt32Lists.AddUnique(UInt32List);
+		}
+
+		OP::ADDRESS AddConstant(const TArray<uint64>& UInt64List)
+		{
+			return (OP::ADDRESS)ConstantUInt64Lists.AddUnique(UInt64List);
+		}
+
+		OP::ADDRESS AddConstant(const TArray<FString>& StringList)
+		{
+			const int32 NumStrings = StringList.Num();
+
+			TArray<uint32> StringAddrs;
+			StringAddrs.SetNum(NumStrings);
+
+			for (int32 StringIndex = 0; StringIndex < NumStrings; ++StringIndex)
+			{
+				StringAddrs[StringIndex] = AddConstant(StringList[StringIndex]);
+			}
+
+			return (OP::ADDRESS)AddConstant(StringAddrs);
+		}
+
+        OP::ADDRESS AddConstant( const FMatrix44f& m )
+        {
+            // Ensure unique
+            for (int32 i=0; i<ConstantMatrices.Num(); ++i)
+            {
+                if (ConstantMatrices[i]==m)
+                {
+                    return (OP::ADDRESS)i;
+                }
+            }
+
+            OP::ADDRESS index = OP::ADDRESS( ConstantMatrices.Num() );
+            ConstantMatrices.Add( m );
+            return index;
+        }
+
+        OP::ADDRESS AddConstant( const FShape& m )
+        {
+            // Ensure unique
+            for (int32 i=0; i<ConstantShapes.Num(); ++i)
+            {
+                if (ConstantShapes[i]==m)
+                {
+                    return (OP::ADDRESS)i;
+                }
+            }
+
+            OP::ADDRESS index = OP::ADDRESS( ConstantShapes.Num() );
+            ConstantShapes.Add( m );
+            return index;
+        }
+
+        OP::ADDRESS AddConstant( const FProjector& m )
+        {
+            // Ensure unique
+            for (int32 i=0; i<ConstantProjectors.Num(); ++i)
+            {
+                if (ConstantProjectors[i]==m)
+                {
+                    return (OP::ADDRESS)i;
+                }
+            }
+
+            OP::ADDRESS index = OP::ADDRESS( ConstantProjectors.Num() );
+            ConstantProjectors.Add( m );
+            return index;
+        }
+
+        OP::ADDRESS AddConstant( const FRichCurve& m )
+        {
+            OP::ADDRESS index = OP::ADDRESS( ConstantCurves.AddUnique( m ) );
+            return index;
+        }
+
+		OP::ADDRESS AddConstant(const TManagedPtr<const FMaterial>& Material)
+		{
+			for (int32 Index=0; Index<ConstantMaterials.Num(); ++Index)
+			{
+				if (ConstantMaterials[Index] == Material || *ConstantMaterials[Index] == *Material)
+				{
+					return static_cast<OP::ADDRESS>(Index);
+				}
+			}
+
+			OP::ADDRESS Index = static_cast<OP::ADDRESS>(ConstantMaterials.Num());
+			ConstantMaterials.Add(Material);
+			return Index;
+		}
+
+		inline TManagedPtr<const FImage> GetImageLOD(FConstantResourceIndex Index) const
+		{
+			TManagedPtr<const FImage> Mip;
+			if (!Index.Streamable)
+			{
+				if (ConstantImageLODsPermanent.IsValidIndex(Index.Index))
+				{
+					Mip = ConstantImageLODsPermanent[Index.Index];
+				}
+			}
+			else
+			{
+				const TManagedPtr<const FImage>* Found = ConstantImageLODsStreamed.Find(Index.Index);
+				if (Found)
+				{
+					Mip = *Found;
+				}
+			}
+
+			return Mip;
+		}
+
+        /** Get a constant image, assuming at least some mips are loaded. The image constant will be composed with lodaded mips if necessary. */
+		template <typename CreateImageFunc>
+        void GetConstant(uint32 ConstantIndex, TManagedPtr<const FImage>& Result, int32 MipsToSkip, const CreateImageFunc& CreateImage) const
+        {
+			const int32 NumLODs = ConstantImages[ConstantIndex].LODCount;
+
+			check(NumLODs > 0);
+			const int32 ReallySkippedLODs = FMath::Min(NumLODs - 1, MipsToSkip);
+			const int32 FirstLODIndexIndex = ConstantImages[ConstantIndex].FirstIndex;
+			const int32 FirstTailLOD = ConstantImages[ConstantIndex].LODCount - ConstantImages[ConstantIndex].NumLODsInTail;
+
+			// Compose the result image
+			{
+				MUTABLE_CPUPROFILER_SCOPE(ComposeConstantImage);
+
+				TManagedPtr<FImage> ResultImage = nullptr;
+
+				int32 LOD = ReallySkippedLODs;
+
+				// Find first availabe LOD.
+				for (; LOD <= FirstTailLOD; ++LOD)
+				{
+					MUTABLE_CPUPROFILER_SCOPE(Search_FisrtValid);
+					TManagedPtr<const FImage> FirstValidLODImage = GetImageLOD(ConstantImageLODIndices[FirstLODIndexIndex + LOD]);
+					
+					if (FirstValidLODImage)
+					{
+						break;
+					}
+				}
+			
+				if (LOD < NumLODs)
+				{						
+					int32 SizeX = ConstantImages[ConstantIndex].ImageSizeX; 
+					int32 SizeY = ConstantImages[ConstantIndex].ImageSizeY; 
+
+					for (int32 I = 0; I < LOD; ++I)
+					{
+						SizeX = FMath::DivideAndRoundUp(SizeX, 2);
+						SizeY = FMath::DivideAndRoundUp(SizeY, 2);
+					}
+					 
+					EImageFormat Format = ConstantImages[ConstantIndex].ImageFormat; 
+
+					ResultImage = CreateImage(SizeX, SizeY, NumLODs - LOD, Format, EInitializationType::NotInitialized);
+					ResultImage->Flags = ConstantImages[ConstantIndex].Flags;
+				}
+	
+				check(ResultImage);
+
+				// Some non-block pixel formats require manual resize.
+				const bool bFormatNeedsDataResize = ResultImage->DataStorage.IsEmpty();
+				
+				int32 DestLOD = 0;
+				// Process hi-res separeted mips.
+				for (; LOD < FirstTailLOD; ++LOD, ++DestLOD)
+				{
+					MUTABLE_CPUPROFILER_SCOPE(NonTailLODs);
+					TManagedPtr<const FImage> Image = GetImageLOD(ConstantImageLODIndices[FirstLODIndexIndex + LOD]);
+
+					if (bFormatNeedsDataResize)
+					{
+						ResultImage->DataStorage.ResizeLOD(DestLOD, Image->DataStorage.GetLOD(0).Num());
+					}
+
+					TArrayView<uint8> ResultLODView = ResultImage->DataStorage.GetLOD(DestLOD);
+					TArrayView<const uint8> ImageLODView = Image->DataStorage.GetLOD(0);
+
+					check(ResultLODView.Num() == ImageLODView.Num())
+					FMemory::Memcpy(ResultLODView.GetData(), ImageLODView.GetData(), ResultLODView.Num());
+				}
+
+				// Process compacted tail.
+				// NOTE: The rom tail data storage can have multiple image buffers.
+				if (LOD < NumLODs)
+				{
+					MUTABLE_CPUPROFILER_SCOPE(TailLODs);
+					TManagedPtr<const FImage> TailLODsImage = GetImageLOD(ConstantImageLODIndices[FirstLODIndexIndex + FirstTailLOD]);
+
+					if (!ensureAlwaysMsgf(TailLODsImage, TEXT("Missing tail image rom, setting image to black.")))
+					{
+						ResultImage->InitToBlack();
+						Result = ResultImage;
+						return;
+					}
+
+					check(TailLODsImage->DataStorage.GetNumLODs() >= NumLODs - LOD);		
+	
+					int32 FirstTailImageTailLOD = FMath::Min(
+							TailLODsImage->DataStorage.ComputeFirstCompactedTailLOD(),
+							TailLODsImage->DataStorage.GetNumLODs());
+
+					int32 TailImageLOD = FMath::Max(0, ReallySkippedLODs - FirstTailLOD);
+					for (; TailImageLOD < FirstTailImageTailLOD; ++TailImageLOD, ++LOD, ++DestLOD)
+					{
+						MUTABLE_CPUPROFILER_SCOPE(TailLODCopy);
+						if (bFormatNeedsDataResize)
+						{
+							ResultImage->DataStorage.ResizeLOD(DestLOD, TailLODsImage->DataStorage.GetLOD(TailImageLOD).Num());
+						}
+
+						TArrayView<uint8> ResultLODView = ResultImage->DataStorage.GetLOD(DestLOD);
+						TArrayView<const uint8> ImageLODView = TailLODsImage->DataStorage.GetLOD(TailImageLOD);
+
+						check(ResultLODView.Num() == ImageLODView.Num());
+						FMemory::Memcpy(ResultLODView.GetData(), ImageLODView.GetData(), ResultLODView.Num());
+					}	
+
+					if (LOD < NumLODs)
+					{
+						MUTABLE_CPUPROFILER_SCOPE(TailTailCopy);
+						check(DestLOD >= ResultImage->DataStorage.ComputeFirstCompactedTailLOD());
+						check(TailImageLOD >= TailLODsImage->DataStorage.ComputeFirstCompactedTailLOD());
+
+						TArrayView<const uint8> FirstTailLODView = TailLODsImage->DataStorage.GetLOD(TailImageLOD);
+						TArrayView<const uint8> LastTailLODView  = TailLODsImage->DataStorage.GetLOD(TailLODsImage->DataStorage.GetNumLODs() - 1);
+					
+						// The tail is stored compacted in memory.
+						int64 NumBytesInTail = LastTailLODView.GetData() + LastTailLODView.Num() - FirstTailLODView.GetData();
+						check(NumBytesInTail >= 0);
+						
+						if (bFormatNeedsDataResize)
+						{
+							const int32 NumLODsToResize = ResultImage->DataStorage.GetNumLODs() - DestLOD;
+							check(NumLODsToResize == TailLODsImage->DataStorage.GetNumLODs() - TailImageLOD);
+
+							ResultImage->DataStorage.GetInternalArray(DestLOD).Reserve(NumBytesInTail);
+
+							for (int32 LODOffset = 0; LODOffset < NumLODsToResize; ++LODOffset)
+							{
+								int32 TailLODSizeInBytes = TailLODsImage->DataStorage.GetLOD(LODOffset + TailImageLOD).Num();
+								ResultImage->DataStorage.ResizeLOD(LODOffset + DestLOD, TailLODSizeInBytes);
+							}
+						}
+
+						TArrayView<uint8> ResultLODView = ResultImage->DataStorage.GetLOD(DestLOD);
+
+						check((int64)ResultImage->DataStorage.GetInternalArray(DestLOD).Num() >= NumBytesInTail);
+						FMemory::Memcpy(ResultLODView.GetData(), FirstTailLODView.GetData(), NumBytesInTail);
+					}
+				}
+	
+				Result = ResultImage;
+			}
+		}		
+
+		template<class CreateMeshFunc>
+        void GetConstant(
+				uint32 MeshConstantIndex, 
+				int32 SkeletonConstantIndex, 
+				TManagedPtr<const FMesh>& OutMesh, 
+				EMeshContentFlags FilterContentFlags, 
+				CreateMeshFunc&& CreateMesh) const
+        {
+			MUTABLE_CPUPROFILER_SCOPE(GetConstant_Mesh)
+
+			FMeshContentRange MeshContentRange = ConstantMeshes[MeshConstantIndex];
+
+			TManagedPtr<const FMesh> EmptyMesh = nullptr;
+			auto GetMeshAtResourceIndex = [&](FConstantResourceIndex ResourceIndex) -> TManagedPtr<const FMesh>
+			{
+				MUTABLE_CPUPROFILER_SCOPE(GetConstant_Mesh_GetMesh)
+
+				if (!ResourceIndex.Streamable)
+				{
+					if (ConstantMeshesPermanent.IsValidIndex(ResourceIndex.Index))
+					{
+						return ConstantMeshesPermanent[ResourceIndex.Index];
+					}
+
+					if (!EmptyMesh)
+					{
+						EmptyMesh = MakeManaged<FMesh>();
+					}
+					
+					return EmptyMesh;
+				}
+				else
+				{
+					const TManagedPtr<const FMesh>* Found = ConstantMeshesStreamed.Find(ResourceIndex.Index);
+					if (Found)
+					{
+						return *Found;
+					}
+
+					if (!EmptyMesh)
+					{
+						EmptyMesh = MakeManaged<FMesh>();
+					}
+
+					return EmptyMesh;
+				}
+			};
+
+			TManagedPtr<const FMesh> GeometryMesh = nullptr;
+			TManagedPtr<const FMesh> PoseMesh     = nullptr;
+			TManagedPtr<const FMesh> PhysicsMesh  = nullptr;
+			TManagedPtr<const FMesh> MetaDataMesh = nullptr;
+
+			int32 MeshRomCurrentIndex = MeshContentRange.GetFirstIndex();
+			if (EnumHasAnyFlags(FilterContentFlags, EMeshContentFlags::GeometryData) &&
+				EnumHasAnyFlags(MeshContentRange.GetContentFlags(), EMeshContentFlags::GeometryData))
+			{
+				const FConstantResourceIndex ResourceIndex = ConstantMeshContentIndices[MeshRomCurrentIndex];
+				GeometryMesh = GetMeshAtResourceIndex(ResourceIndex);
+			}
+			MeshRomCurrentIndex += (int32)EnumHasAnyFlags(MeshContentRange.GetContentFlags(), EMeshContentFlags::GeometryData);
+
+			if (EnumHasAnyFlags(FilterContentFlags, EMeshContentFlags::PoseData) &&
+				EnumHasAnyFlags(MeshContentRange.GetContentFlags(), EMeshContentFlags::PoseData))
+			{
+				const FConstantResourceIndex ResourceIndex = ConstantMeshContentIndices[MeshRomCurrentIndex];
+				PoseMesh = GetMeshAtResourceIndex(ResourceIndex);
+			}
+			MeshRomCurrentIndex += (int32)EnumHasAnyFlags(MeshContentRange.GetContentFlags(), EMeshContentFlags::PoseData);
+
+			if (EnumHasAnyFlags(FilterContentFlags, EMeshContentFlags::PhysicsData) &&
+				EnumHasAnyFlags(MeshContentRange.GetContentFlags(), EMeshContentFlags::PhysicsData))
+			{
+				const FConstantResourceIndex ResourceIndex = ConstantMeshContentIndices[MeshRomCurrentIndex];
+				PhysicsMesh = GetMeshAtResourceIndex(ResourceIndex);
+			}
+			MeshRomCurrentIndex += (int32)EnumHasAnyFlags(MeshContentRange.GetContentFlags(), EMeshContentFlags::PhysicsData);
+
+			if (EnumHasAnyFlags(FilterContentFlags, EMeshContentFlags::MetaData) &&
+				EnumHasAnyFlags(MeshContentRange.GetContentFlags(), EMeshContentFlags::MetaData))
+			{
+				const FConstantResourceIndex ResourceIndex = ConstantMeshContentIndices[MeshRomCurrentIndex];
+				MetaDataMesh = GetMeshAtResourceIndex(ResourceIndex);
+			}
+			MeshRomCurrentIndex += (int32)EnumHasAnyFlags(MeshContentRange.GetContentFlags(), EMeshContentFlags::MetaData);
+
+			check(MeshRomCurrentIndex - MeshContentRange.GetFirstIndex() == FMath::CountBits((uint64)MeshContentRange.GetContentFlags()));
+
+			uint32 MeshBudgetReserve = 0;
+			MeshBudgetReserve += GeometryMesh ? GeometryMesh->GetDataSize() : 0;
+			MeshBudgetReserve += PoseMesh     ? PoseMesh->GetDataSize()     : 0;
+			MeshBudgetReserve += PhysicsMesh  ? PhysicsMesh->GetDataSize()  : 0;
+			MeshBudgetReserve += MetaDataMesh ? MetaDataMesh->GetDataSize() : 0;
+
+			TManagedPtr<FMesh> Result = nullptr;
+			if (GeometryMesh)
+			{
+				MUTABLE_CPUPROFILER_SCOPE(GetConstant_Mesh_Geoemtry)
+
+				Result = CreateMesh(MeshBudgetReserve);
+				Result->CopyFrom(*GeometryMesh);
+			}
+
+			if (PoseMesh)
+			{
+				MUTABLE_CPUPROFILER_SCOPE(GetConstant_Mesh_Pose)
+
+				if (!Result)
+				{
+					Result = CreateMesh(MeshBudgetReserve);
+					Result->CopyFrom(*PoseMesh);
+				}
+				else
+				{
+					Result->BonePoses = PoseMesh->BonePoses; 
+					Result->BoneMap   = PoseMesh->BoneMap;
+
+					for (const TPair<EMeshBufferType, FMeshBufferSet>& AdditionalBuffer : PoseMesh->AdditionalBuffers)
+					{
+						const bool bIsPoseBufferType = 	
+								AdditionalBuffer.Key == EMeshBufferType::SkeletonDeformBinding;
+
+						if (bIsPoseBufferType)
+						{
+							Result->AdditionalBuffers.Emplace(AdditionalBuffer);
+						}
+					}	
+				}
+			}
+
+			if (PhysicsMesh)
+			{
+				MUTABLE_CPUPROFILER_SCOPE(GetConstant_Mesh_Physics)
+				
+				if (!Result)
+				{
+					Result = CreateMesh(MeshBudgetReserve);
+					Result->CopyFrom(*PhysicsMesh);
+				}
+				else
+				{
+					Result->PhysicsBody = PhysicsMesh->PhysicsBody;
+
+					for (const TPair<EMeshBufferType, FMeshBufferSet>& AdditionalBuffer : PhysicsMesh->AdditionalBuffers)
+					{
+						const bool bIsPhysicsBufferType = 
+								AdditionalBuffer.Key == EMeshBufferType::PhysicsBodyDeformBinding   ||
+								AdditionalBuffer.Key == EMeshBufferType::PhysicsBodyDeformSelection ||
+								AdditionalBuffer.Key == EMeshBufferType::PhysicsBodyDeformOffsets;
+
+						if (bIsPhysicsBufferType)
+						{
+							Result->AdditionalBuffers.Emplace(AdditionalBuffer);
+						}
+					}
+				}
+			}
+
+			if (MetaDataMesh)
+			{
+				MUTABLE_CPUPROFILER_SCOPE(GetConstant_Mesh_Metadata)
+
+				if (!Result)
+				{
+					Result = CreateMesh(MeshBudgetReserve);
+					Result->CopyFrom(*MetaDataMesh);
+				}
+				else
+				{
+					// Only in case the geoemtry has been filtered, add the meta data descriptors.
+					if (EnumHasAnyFlags(MeshContentRange.GetContentFlags(), EMeshContentFlags::GeometryData) &&
+						!EnumHasAnyFlags(FilterContentFlags, EMeshContentFlags::GeometryData))
+					{
+						check(MetaDataMesh->VertexBuffers.IsDescriptor());
+						check(MetaDataMesh->IndexBuffers.IsDescriptor());
+
+						Result->VertexBuffers = MetaDataMesh->VertexBuffers;
+						Result->IndexBuffers = MetaDataMesh->IndexBuffers;
+						Result->Surfaces = MetaDataMesh->Surfaces;
+					}
+
+					Result->Tags = MetaDataMesh->Tags; 
+					Result->SkeletonIDs = MetaDataMesh->SkeletonIDs; 
+					Result->StreamedResources = MetaDataMesh->StreamedResources; 
+				}
+			}
+
+			Result->MeshIDPrefix = MeshContentRange.MeshIDPrefix;
+
+			check(ConstantSkeletons.Num() > SkeletonConstantIndex);
+			if (ConstantSkeletons.IsValidIndex(SkeletonConstantIndex))
+			{
+				Result->Skeleton = ConstantSkeletons[SkeletonConstantIndex];
+			}
+
+			OutMesh = Result;
+		}
+
+		void GetExtensionDataConstant(int32 ConstantIndex, TManagedPtr<const FExtensionData>& Result) const
+		{
+			const FExtensionDataConstant& Constant = ConstantExtensionData[ConstantIndex];
+
+			check(Constant.Data);
+
+			Result = Constant.Data;
+		}
+
+        FORCEINLINE EOpType GetOpType( OP::ADDRESS at ) const
+        {
+            if (at>=OP::ADDRESS(OpAddress.Num())) return EOpType::NONE;
+
+            EOpType result;
+            uint64 byteCodeAddress = OpAddress[at];
+			FMemory::Memcpy( &result, &ByteCode[byteCodeAddress], sizeof(EOpType) );
+            check( result<EOpType::COUNT);
+            return result;
+        }
+
+		template<typename ARGS>
+		inline const ARGS GetOpArgs(OP::ADDRESS at) const
+		{
+			ARGS result;
+			uint64 byteCodeAddress = OpAddress[at];
+			byteCodeAddress += sizeof(EOpType);
+			FMemory::Memcpy(&result, &ByteCode[byteCodeAddress], sizeof(ARGS));
+			return result;
+		}
+
+		template<typename ARGS>
+		inline void SetOpArgs(OP::ADDRESS at, const ARGS& Args)
+		{
+			uint64 byteCodeAddress = OpAddress[at];
+			byteCodeAddress += sizeof(EOpType);
+			FMemory::Memcpy(&ByteCode[byteCodeAddress], &Args, sizeof(ARGS));
+		}
+
+        inline const uint8* GetOpArgsPointer( OP::ADDRESS at ) const
+        {
+            uint64 byteCodeAddress = OpAddress[at];
+            byteCodeAddress += sizeof(EOpType);
+            const uint8* pData = (const uint8*)&ByteCode[byteCodeAddress];
+            return pData;
+        }
+
+        inline uint8* GetOpArgsPointer( OP::ADDRESS at )
+        {
+            uint64 byteCodeAddress = OpAddress[at];
+            byteCodeAddress += sizeof(EOpType);
+            uint8* pData = (uint8*)&ByteCode[byteCodeAddress];
+            return pData;
+        }
+    };
+
+	
     /** A Model represents a customisable object with any number of parameters.
     * When values are given to the parameters, specific Instances can be built, which hold the
     * built application-usable data.
@@ -29,17 +1174,18 @@ namespace UE::Mutable::Private
     class FModel
 	{
 	public:
+		static UE_API void Serialise(const FModel*, FOutputArchive&);
+		static UE_API TSharedPtr<FModel> StaticUnserialise(FInputArchive&);
+    	
+    	void Serialise(FOutputArchive& arch) const
+    	{
+    		arch << Program;
+    	}
 
-		//-----------------------------------------------------------------------------------------
-		// Life cycle
-		//-----------------------------------------------------------------------------------------
-		UE_API FModel();
-
-		//! Don't call directly. Manage with a TSharedPtr.
-		UE_API ~FModel();
-
-		static UE_API void Serialise( const FModel*, FOutputArchive& );
-		static UE_API TSharedPtr<FModel> StaticUnserialise( FInputArchive& );
+    	void Unserialise(FInputArchive& arch)
+    	{
+    		arch >> Program;
+    	}
 
 		/** Special serialise operation that serialises the data in separate "files". An object
          * with the ModelStreamer interface is responsible of storing this data and providing
@@ -149,18 +1295,10 @@ namespace UE::Mutable::Private
 		 * returns -1.
 		 */
     	UE_API int32 GetConstantImageRomId(int32 ConstantImageIndex, int32 LODIndex) const;
-
-		//-----------------------------------------------------------------------------------------
-		// Interface pattern
-		//-----------------------------------------------------------------------------------------
-		class Private;
-
-		UE_API Private* GetPrivate() const;
-
-	private:
-
-		Private* m_pD;
-
+    	
+    	UE_API void UnloadRoms();
+    	
+    	FProgram Program;
 	};
 
 }
diff --git a/Engine/Plugins/Mutable/Source/MutableRuntime/Internal/MuR/MutableString.h b/Engine/Plugins/Mutable/Source/MutableRuntime/Internal/MuR/MutableString.h
--- a/Engine/Plugins/Mutable/Source/MutableRuntime/Internal/MuR/MutableString.h
+++ b/Engine/Plugins/Mutable/Source/MutableRuntime/Internal/MuR/MutableString.h
@@ -6,7 +6,7 @@
 #include "MuR/Ptr.h"
 #include "MuR/RefCounted.h"
 #include "MuR/Serialisation.h"
-#include "MuR/SerialisationPrivate.h"
+#include "MuR/Serialisation.h"
 
 #define UE_API MUTABLERUNTIME_API
 
diff --git a/Engine/Plugins/Mutable/Source/MutableRuntime/Internal/MuR/Parameters.h b/Engine/Plugins/Mutable/Source/MutableRuntime/Internal/MuR/Parameters.h
--- a/Engine/Plugins/Mutable/Source/MutableRuntime/Internal/MuR/Parameters.h
+++ b/Engine/Plugins/Mutable/Source/MutableRuntime/Internal/MuR/Parameters.h
@@ -2,24 +2,27 @@
 
 #pragma once
 
-#include "HAL/Platform.h"
+#include "MuR/Serialisation.h"
 #include "MuR/ManagedPointer.h"
 #include "MuR/RefCounted.h"
 #include "MuR/Types.h"
+#include "MuR/TVariant.h"
+#include "HAL/Platform.h"
 #include "Templates/SharedPointer.h"
 #include "VectorTypes.h"
+#include "Misc/TVariant.h"
+#include "Engine/Texture.h"
+#include "Engine/SkeletalMesh.h"
+#include "Materials/MaterialInterface.h"
 
 struct FInstancedStruct;
-class UTexture;
-class USkeletalMesh;
-class UMaterialInterface;
+
 
 #define UE_API MUTABLERUNTIME_API
 
 
 namespace UE::Mutable::Private
 {
-
 	class FInputArchive;
 	class FOutputArchive;
 	class FParameters;
@@ -124,8 +127,6 @@ namespace UE::Mutable::Private
         //! \param index Index of the parameter from 0 to GetRangeCount()-1
         UE_API int32 GetPosition( int32 Index) const;
 
-    public:
-
         friend class FParameters;
 
 		//! Run-time data
@@ -136,10 +137,378 @@ namespace UE::Mutable::Private
 
 		//! Position in the several dimension of the range, as defined in Parameters
 		TArray<int32> Values;
+    };
+
+
+	template<class T>
+    class TIndirectObject
+    {
+        TUniquePtr<T> StoragePtr;
 
+    public:
+        template<typename... TArgs>
+        TIndirectObject(TArgs&&... Args) : StoragePtr(MakeUnique<T>(Forward<TArgs>(Args)...)) 
+        {
+        }
+        
+        TIndirectObject(const TIndirectObject<T>& Other) : StoragePtr(MakeUnique<T>()) 
+        { 
+            Get() = Other; 
+        }
+
+		TIndirectObject(TIndirectObject<T>&& Other) : StoragePtr(MoveTemp(Other.StoragePtr)) 
+        {
+        }
+
+        TIndirectObject(const T& Object) : StoragePtr(MakeUnique<T>()) 
+        { 
+            Get() = Object; 
+        }
+        
+        TIndirectObject(T&& Object) : StoragePtr(MakeUnique<T>()) 
+        { 
+            Get() = MoveTemp(Object); 
+        }
+
+        TIndirectObject& operator=(TIndirectObject<T>&&) = default;
+
+        TIndirectObject& operator=(const TIndirectObject<T>& Other) 
+        { 
+            Get() = Other; 
+			return *this;
+        }
+
+        TIndirectObject& operator=(const T& Object) 
+        { 
+            Get() = Object; 
+            return *this;
+        }
+
+        TIndirectObject& operator=(const T&& Object) 
+        { 
+            Get() = MoveTemp(Object); 
+            return *this;
+        }
+
+        const T& Get() const 
+        { 
+            check(StoragePtr); 
+            return *StoragePtr; 
+        }
+
+        T& Get() 
+        { 
+            check(StoragePtr); 
+            return *StoragePtr; 
+        }
+
+        operator T&() 
+        { 
+            return Get(); 
+        }
+
+        operator const T&() const 
+        { 
+            return Get(); 
+        }
+
+        T* operator &() 
+        { 
+            return StoragePtr.Get(); 
+        }
+
+        const T* operator &() const 
+        { 
+            return StoragePtr.Get(); 
+        }
+
+        T& operator *() 
+        { 
+            return Get(); 
+        }
+
+        const T& operator *() const 
+        { 
+            return Get(); 
+        }
+
+        bool operator==(const TIndirectObject<T>& Other) const 
+        { 
+            return *StoragePtr == Other; 
+        }
+
+        bool operator==(const T& Object) const 
+        { 
+            return *StoragePtr == Object; 
+        }
+
+		//!
+		void Serialise(FOutputArchive& Arch) const
+		{
+			Arch << Get();
+		}
+
+		//!
+		void Unserialise(FInputArchive& Arch)
+		{
+			Arch >> Get();
+		} 
     };
 
 
+    MUTABLE_DEFINE_ENUM_SERIALISABLE(EParameterType)
+    MUTABLE_DEFINE_ENUM_SERIALISABLE(EProjectorType)
+
+
+    /** Description of a projector to project an image on a mesh. */
+    struct FProjector
+    {
+        EProjectorType type = EProjectorType::Planar;
+        FVector3f position = {0,0,0};
+		FVector3f direction = {0,0,0};
+		FVector3f up = {0,0,0};
+		FVector3f scale = {0,0,0};
+        float projectionAngle = 0.0f;
+
+        //!
+        inline void GetDirectionSideUp(FVector3f& OutDirection, FVector3f& OutSide, FVector3f& OutUp) const
+        {
+			OutDirection = direction;
+            OutUp = up;
+            OutSide = FVector3f::CrossProduct( up, direction );
+			OutSide.Normalize();
+        }
+
+
+        //!
+        void Serialise( FOutputArchive& arch ) const
+        {
+            arch << type;
+            arch << position;
+            arch << direction;
+            arch << up;
+            arch << scale;
+            arch << projectionAngle;
+        }
+
+        //!
+        void Unserialise( FInputArchive& arch )
+        {
+            arch >> type;
+            arch >> position;
+            arch >> direction;
+            arch >> up;
+            arch >> scale;
+            arch >> projectionAngle;
+        }
+
+        bool operator==( const FProjector& o ) const
+        {
+            return     type==o.type
+                    && position==o.position
+                    && up==o.up
+                    && direction==o.direction
+                    && scale==o.scale
+                    && projectionAngle==o.projectionAngle;
+        }
+
+    };
+
+	//---------------------------------------------------------------------------------------------
+	//! Information about a generic shape in space.
+	//---------------------------------------------------------------------------------------------
+	struct FShape
+	{
+		// Transform
+		FVector3f position = FVector3f(0,0,0);
+		FVector3f up = FVector3f(0, 0, 0);
+		FVector3f side = FVector3f(0, 0, 0);
+		
+		FVector3f size = FVector3f(0, 0, 0);
+
+		// 
+		enum class Type : uint8
+		{
+			None = 0,
+			Ellipse,
+			AABox
+		};
+        uint8 type = 0;
+
+		//!
+		void Serialise(FOutputArchive& arch) const
+		{
+			arch << position;
+			arch << up;
+			arch << side;
+			arch << size;
+			arch << type;
+		}
+
+		//!
+		void Unserialise(FInputArchive& arch)
+		{
+			arch >> position;
+			arch >> up;
+			arch >> side;
+			arch >> size;
+			arch >> type;
+		}
+
+        bool operator==( const FShape& o ) const
+        {
+            return type==o.type
+                    && position==o.position
+                    && up==o.up
+                    && side==o.side
+                    && size==o.size;
+        }
+	};
+
+
+	using FParamBoolType = bool;
+	using FParamIntType = int32;
+	using FParamFloatType = float;
+	using FParamColorType = FVector4f;
+	using FParamProjectorType = TIndirectObject<FProjector>;
+	using FParamTextureType = TStrongObjectPtr<UTexture>;
+	using FParamSkeletalMeshType = TStrongObjectPtr<USkeletalMesh>;
+	using FParamStringType = TIndirectObject<FString>;
+	using FParamMatrixType = TIndirectObject<FMatrix44f>;
+	using FParamMaterialType = TStrongObjectPtr<UMaterialInterface>;
+	using FParamInstancedStructType = TManagedPtr<const FInstancedStruct>; // TSharedPtr always valid and FInstancedStruct always initialized to its type.
+	
+	using FParameterValue = TVariant<
+		FParamBoolType,
+		FParamIntType,
+		FParamFloatType,
+		FParamColorType,
+		FParamProjectorType,
+		FParamTextureType,
+		FParamSkeletalMeshType,
+		FParamStringType,
+		FParamMatrixType,
+		FParamMaterialType,
+		FParamInstancedStructType>;
+
+    // static_assert to track ParameterValue size changes. It is ok to change if needed.
+    static_assert(sizeof(FParameterValue) == 8*4, "ParameterValue size has changed.");
+
+    struct FParameterDesc
+    {
+        FString Name;
+
+        //! Unique id (provided externally, so no actual guarantee that it is unique.)
+		FGuid UID;
+
+        EParameterType Type = EParameterType::None;
+
+        FParameterValue DefaultValue;
+
+        //! Ranges, if the parameter is multi-dimensional. The indices refer to the Model's program
+        //! vector of range descriptors.
+		TArray<uint32> Ranges;
+
+        //! Possible values of the parameter in case of being an integer, and its names
+        struct FIntValueDesc
+        {
+            int16 Value;
+			FString Name;
+
+            //!
+            bool operator==( const FIntValueDesc& Other ) const
+            {
+                return Value== Other.Value &&
+						Name== Other.Name;
+            }
+
+            //!
+            void Serialise( FOutputArchive& Arch ) const
+            {
+				Arch << Value;
+				Arch << Name;
+            }
+
+            //!
+            void Unserialise( FInputArchive& Arch)
+            {
+				Arch >> Value;
+				Arch >> Name;
+            }
+        };
+
+        //! For integer parameters, this contains the description of the possible values.
+        //! If empty, the integer may have any value.
+		TArray<FIntValueDesc> PossibleValues;
+
+        //!
+        bool operator==( const FParameterDesc& other ) const
+        {
+            return Name == other.Name && UID == other.UID && Type == other.Type &&
+                   DefaultValue == other.DefaultValue &&
+                   Ranges == other.Ranges &&
+                   PossibleValues == other.PossibleValues;
+        }
+
+        //!
+        void Serialise( FOutputArchive& arch ) const
+        {
+			arch << Name;
+			arch << UID;
+            arch << Type;
+            arch << DefaultValue;
+            arch << Ranges;
+            arch << PossibleValues;
+        }
+
+        //!
+        void Unserialise( FInputArchive& arch )
+        {
+			arch >> Name;
+			arch >> UID;
+            arch >> Type;
+			arch >> DefaultValue;
+            arch >> Ranges;
+			arch >> PossibleValues;
+        }
+    };
+
+
+    struct FRangeDesc
+    {
+		FString Name;
+		FString UID;
+
+		/** Parameter that controls the size of this range, if any. */
+		int32 DimensionParameter = -1;
+
+        //!
+        bool operator==( const FRangeDesc& other ) const
+        {
+            return Name==other.Name
+				&&
+				UID == other.UID 
+				&&
+				DimensionParameter == other.DimensionParameter;
+        }
+
+        //!
+        void Serialise( FOutputArchive& Arch ) const
+        {
+			Arch << Name;
+			Arch << UID;
+			Arch << DimensionParameter;
+        }
+
+        //!
+        void Unserialise( FInputArchive& Arch)
+        {
+			Arch >> Name;
+			Arch >> UID;
+			Arch >> DimensionParameter;
+        }
+    };
+	
     /** This class represents the parameters of a model including description, type and
     * value, and additional resources associated to the parameter.
     *
@@ -154,20 +523,17 @@ namespace UE::Mutable::Private
 		//-----------------------------------------------------------------------------------------
 		// Life cycle
 		//-----------------------------------------------------------------------------------------
-
-		/** Instances of this class must be obtained from a Model instance. */
-		UE_API FParameters();
-
+		
         //! Serialisation
         static UE_API void Serialise( const FParameters*, FOutputArchive& );
         static UE_API TSharedPtr<FParameters> StaticUnserialise( FInputArchive& );
 
+		void Serialise(FOutputArchive& Arch) const;
+		void Unserialise(FInputArchive& Arch);
+		
 		//! Deep clone this object.
 		UE_API TSharedPtr<FParameters> Clone() const;
-
-		/** */
-		UE_API ~FParameters();
-
+		
 		//-----------------------------------------------------------------------------------------
 		// Own interface
 		//-----------------------------------------------------------------------------------------
@@ -388,17 +754,22 @@ namespace UE::Mutable::Private
         //! Utility method to compare the values of a specific parameter with the values of another
         //! FParameters object. It returns false if type or values are different.
         UE_API bool HasSameValue( int32 ThisParamIndex, const TSharedPtr<const FParameters>& Other, int32 OtherParamIndex ) const;
+		
+		UE_API FProjector GetProjectorValue( int32 index, const FRangeIndex* ) const;
 
-		//-----------------------------------------------------------------------------------------
-		// Interface pattern
-		//-----------------------------------------------------------------------------------------
-		class Private;
-
-		UE_API Private* GetPrivate() const;
-
-	private:
-
-		Private* m_pD;
+		/** Return true if the parameter has any multi-dimensional values set. This is independent to if the model
+		* accepts multi-dimensional parameters for this particular parameter. */
+		UE_API bool HasMultipleValues(int32 ParamIndex) const;
+		
+		//! Run-time data
+		TSharedPtr<const FModel> Model;
+ 
+		//! Values for the parameters if they are not multidimensional.
+		TArray<FParameterValue> Values;
+		
+		//! If the parameter is multidemensional, the values are stored here.
+		//! The key of the map is the vector of values stored in a FRangeIndex
+		TArray< TMap< TArray<int32>, FParameterValue > > MultiValues;
 
 	};
 
diff --git a/Engine/Plugins/Mutable/Source/MutableRuntime/Internal/MuR/Serialisation.h b/Engine/Plugins/Mutable/Source/MutableRuntime/Internal/MuR/Serialisation.h
--- a/Engine/Plugins/Mutable/Source/MutableRuntime/Internal/MuR/Serialisation.h
+++ b/Engine/Plugins/Mutable/Source/MutableRuntime/Internal/MuR/Serialisation.h
@@ -2,21 +2,23 @@
 
 #pragma once
 
-#include "HAL/Platform.h"
 #include "MuR/ManagedPointer.h"
 #include "MuR/Ptr.h"
 #include "MuR/RefCounted.h"
 #include "MuR/MutableMath.h"
+#include "MuR/Serialisation.h"
 
-#include "Curves/RichCurve.h"
+#include "Misc/TVariant.h"
 #include "Math/Vector.h"
+#include "Math/Quat.h"
 #include "Math/IntVector.h"
 #include "Math/Vector4.h"
+#include "HAL/Platform.h"
+#include "Curves/RichCurve.h"
+#include "Containers/StaticArray.h"
 #include "Containers/Array.h"
 #include "Containers/Map.h"
-#include "Containers/ContainerAllocationPolicies.h"
 #include "StructUtils/InstancedStruct.h"
-#include "Templates/SharedPointer.h"
 #include "UObject/StrongObjectPtr.h"
 
 class UTexture;
@@ -391,6 +393,382 @@ namespace UE::Mutable::Private
 
 	MUTABLERUNTIME_API void operator<<(FOutputArchive& Arch, const TManagedPtr<const FInstancedStruct>& V);
 	MUTABLERUNTIME_API void operator>>(FInputArchive& Arch, TManagedPtr<const FInstancedStruct>& V);
+	
+	
+    //!
+    class FInputMemoryStream : public FInputStream
+    {
+    public:
+
+        //-----------------------------------------------------------------------------------------
+        // Life cycle
+        //-----------------------------------------------------------------------------------------
+
+        //! Create the stream using an external buffer.
+        //! The buffer will not be owned by this object, so it cannot be deallocated while this
+        //! objects is in use.
+        UE_API FInputMemoryStream( const void* pBuffer, uint64 size );
+
+
+        //-----------------------------------------------------------------------------------------
+        // FInputStream interface
+        //-----------------------------------------------------------------------------------------
+        UE_API void Read( void* pData, uint64 size ) override;
+
+
+    private:
+
+		const void* Buffer = nullptr;
+		uint64 Size = 0;
+		uint64 Pos = 0;
+
+    };
+
+
+#define MUTABLE_IMPLEMENT_POD_SERIALISABLE(Type)				     \
+    void DLLEXPORT operator<<(FOutputArchive& Arch, const Type& T)   \
+    {																 \
+        Arch.Stream->Write(&T, sizeof(Type));						 \
+    }																 \
+                                                                     \
+    void DLLEXPORT operator>>(FInputArchive& Arch, Type& T)		     \
+    {																 \
+        Arch.Stream->Read(&T, sizeof(Type));						 \
+    }																 \
+		
+
+#define MUTABLE_IMPLEMENT_POD_VECTOR_SERIALISABLE(Type)                        \
+    template<class Alloc>                                                      \
+	void DLLEXPORT operator<<(FOutputArchive& Arch, const TArray<Type, Alloc>& V)         \
+	{                                                                          \
+		uint32 Num = uint32(V.Num());                                          \
+		Arch << Num;                                                           \
+		if (Num)                                                               \
+		{                                                                      \
+			Arch.Stream->Write(V.GetData(), Num * sizeof(Type));			   \
+		}                                                                      \
+	}                                                                          \
+                                                                               \
+    template<class Alloc>                                                      \
+	void DLLEXPORT operator>>(FInputArchive& Arch, TArray<Type, Alloc>& V)     \
+	{                                                                          \
+		uint32 Num;                                                            \
+		Arch >> Num;                                                           \
+		V.SetNum(Num);                                                         \
+		if (Num)                                                               \
+		{                                                                      \
+			Arch.Stream->Read(V.GetData(), Num * sizeof(Type));					   \
+		}                                                                      \
+	}                                                                          \
+
+	/** TVariant custom serialize. Based on the default serialization. */
+	template <typename... Ts>
+	void operator<<(FOutputArchive& Ar, const TVariant<Ts...>& Variant)
+	{
+		const uint8 Index = static_cast<uint8>(Variant.GetIndex());
+		Ar << Index;
+		
+		Visit([&Ar](auto& StoredValue)
+		{
+			Ar << StoredValue;
+		}, Variant);
+	}
+
+
+	template <typename T, typename VariantType>
+	struct TVariantLoadFromInputArchiveCaller
+	{
+		/** Default construct the type and load it from the FArchive */
+		static void Load(FInputArchive& Ar, VariantType& OutVariant)
+		{
+			OutVariant.template Emplace<T>();
+			Ar >> OutVariant.template Get<T>();
+		}
+	};
+
+	
+	template <typename... Ts>
+	struct TVariantLoadFromInputArchiveLookup
+	{
+		using VariantType = TVariant<Ts...>;
+		static_assert((std::is_default_constructible<Ts>::value && ...), "Each type in TVariant template parameter pack must be default constructible in order to use FArchive serialization");
+
+		/** Load the type at the specified index from the FArchive and emplace it into the TVariant */
+		static void Load(SIZE_T TypeIndex, FInputArchive& Ar, VariantType& OutVariant)
+		{
+			static constexpr void(*Loaders[])(FInputArchive&, VariantType&) = { &TVariantLoadFromInputArchiveCaller<Ts, VariantType>::Load... };
+			check(TypeIndex < UE_ARRAY_COUNT(Loaders));
+			Loaders[TypeIndex](Ar, OutVariant);
+		}
+	};
+
+	
+	template <typename... Ts>
+	void operator>>(FInputArchive& Ar, TVariant<Ts...>& Variant)
+	{
+		uint8 Index;
+		Ar >> Index;
+		check(Index < sizeof...(Ts));
+
+		TVariantLoadFromInputArchiveLookup<Ts...>::Load(static_cast<SIZE_T>(Index), Ar, Variant);
+	}
+
+#define MUTABLE_IMPLEMENT_ENUM_SERIALISABLE(Type)						\
+        void DLLEXPORT operator<<(FOutputArchive& Arch, const Type& T)   \
+		{																\
+            uint32 V = (uint32)T;                                       \
+            Arch.Stream->Write(&V, sizeof(uint32));					\
+		}																\
+																		\
+        void DLLEXPORT operator>>(FInputArchive& Arch, Type& T)   		\
+		{																\
+            uint32 V;													\
+            Arch.Stream->Read(&V, sizeof(uint32));					\
+			T = (Type)V;												\
+		}																\
+
+    template<typename T0, typename T1>
+    inline void operator<<(FOutputArchive& Arch, const std::pair<T0, T1>& V)
+    {
+        Arch << V.first;
+        Arch << V.second;
+    }
+
+    template<typename T0, typename T1>
+    inline void operator>>(FInputArchive& Arch, std::pair<T0, T1>& V)
+    {
+        Arch >> V.first;
+        Arch >> V.second;
+    }
+	
+	template<typename T, uint32 Size, uint32 Align> 
+    void operator<<(FOutputArchive& Arch, const TStaticArray<T, Size, Align>& V)
+	{
+		for (int32 I = 0; I < Size; ++I)
+		{
+			Arch << V[I];
+		}
+	}
+
+	template<typename T, uint32 Size, uint32 Align> 
+    void operator>>(FInputArchive& Arch, TStaticArray<T, Size, Align>& V)
+	{
+		for (uint32 I = 0; I < Size; ++I)
+		{
+			Arch >> V[I];
+		}
+	}
+	
+
+	//---------------------------------------------------------------------------------------------
+	template< typename K, typename T >
+	inline void operator<< (FOutputArchive& Arch, const TMap<K, T>& V)
+	{
+		Arch << (uint32)V.Num();
+		for (const TPair<K, T>& Element : V)
+		{
+			Arch << Element.Key;
+			Arch << Element.Value;
+		}
+	}
+
+	template< typename K, typename T >
+	inline void operator>> (FInputArchive& Arch, TMap<K, T>& V)
+	{
+		uint32 Num;
+		Arch >> Num;
+
+		for (uint32 Index = 0; Index < Num; ++Index)
+		{
+			K Key;
+			T Element;
+			Arch >> Key;
+			Arch >> Element;
+
+			V.Emplace(MoveTemp(Key), MoveTemp(Element));
+		}
+	}
+
+	// Unreal POD Serializables
+	MUTABLE_IMPLEMENT_POD_VECTOR_SERIALISABLE(float);
+	MUTABLE_IMPLEMENT_POD_VECTOR_SERIALISABLE(double);
+	MUTABLE_IMPLEMENT_POD_VECTOR_SERIALISABLE(uint8);
+	MUTABLE_IMPLEMENT_POD_VECTOR_SERIALISABLE(uint16);
+	MUTABLE_IMPLEMENT_POD_VECTOR_SERIALISABLE(uint32);
+	MUTABLE_IMPLEMENT_POD_VECTOR_SERIALISABLE(uint64);
+	MUTABLE_IMPLEMENT_POD_VECTOR_SERIALISABLE(int8);
+	MUTABLE_IMPLEMENT_POD_VECTOR_SERIALISABLE(int16);
+	MUTABLE_IMPLEMENT_POD_VECTOR_SERIALISABLE(int32);
+	MUTABLE_IMPLEMENT_POD_VECTOR_SERIALISABLE(int64);
+	MUTABLE_IMPLEMENT_POD_VECTOR_SERIALISABLE(TCHAR);
+
+	MUTABLE_IMPLEMENT_POD_VECTOR_SERIALISABLE(FIntVector2);
+	MUTABLE_IMPLEMENT_POD_VECTOR_SERIALISABLE(FUintVector2);
+	MUTABLE_IMPLEMENT_POD_VECTOR_SERIALISABLE(UE::Math::TIntVector2<uint16>);
+	MUTABLE_IMPLEMENT_POD_VECTOR_SERIALISABLE(UE::Math::TIntVector2<int16>);
+	MUTABLE_IMPLEMENT_POD_VECTOR_SERIALISABLE(FVector2f);
+	MUTABLE_IMPLEMENT_POD_VECTOR_SERIALISABLE(FVector4f);
+	MUTABLE_IMPLEMENT_POD_VECTOR_SERIALISABLE(FMatrix44f);
+	MUTABLE_IMPLEMENT_POD_VECTOR_SERIALISABLE(FRichCurveKey);
+
+
+	template < 
+		typename T 
+	>
+	inline void operator<< (FOutputArchive& Arch, const TManagedPtr<const T>& Ptr)
+	{
+		if (!Ptr)
+		{
+			Arch << (int32)-1;
+		}
+		else
+		{
+			int32* it = Arch.History.Find(Ptr.Get());
+
+			if (!it)
+			{
+				int32 Id = Arch.History.Num();
+				Arch.History.Add(Ptr.Get(), Id);
+				Arch << Id;
+				T::Serialise(Ptr.Get(), Arch);
+			}
+			else
+			{
+				Arch << *it;
+			}
+		}
+	}
+
+	template< 
+		typename T
+	>
+	inline void operator>> (FInputArchive& Arch, TManagedPtr<const T>& Ptr)
+	{
+		int32 Id;
+		Arch >> Id;
+
+		if (Id == -1)
+		{
+			Ptr.Reset();
+		}
+		else
+		{
+			if (Id < Arch.History.Num())
+			{
+				Ptr = StaticCastManagedPtr<T>(Arch.History[Id]);
+
+				// If the pointer was null it means the position in history is used, but not set yet
+				// option 1: we have a smart pointer loop which is very bad.
+				// option 2: the resource in this Ptr is also pointed by a Proxy that has absorbed it
+				//			 and this reference should also be a proxy instead of a pointer.
+				check(Ptr);
+			}
+			else
+			{
+				// Ids come in order, but they may have been absorbed outside in some serialisations
+				// like proxies.
+				//check( Id == Arch.>History.Num() );
+				Arch.History.SetNum(Id + 1, EAllowShrinking::No);
+
+				TManagedPtr<T> Temp = T::StaticUnserialise(Arch);
+				Ptr = Temp;
+				Arch.History[Id] = Temp;
+			}
+		}
+	}
+
+	//---------------------------------------------------------------------------------------------
+	template< typename T >
+	inline void operator<< ( FOutputArchive& arch, const Ptr<T>& p )
+	{
+		operator<<( arch, (const Ptr<const T>&) p );
+	}
+
+	template< typename T >
+	inline void operator>> ( FInputArchive& arch, Ptr<T>& p )
+	{
+		operator>>( arch, (Ptr<const T>&) p );
+	}
+
+
+	//---------------------------------------------------------------------------------------------
+	//---------------------------------------------------------------------------------------------
+	//---------------------------------------------------------------------------------------------
+	template<typename T0, typename T1> inline void operator<<(FOutputArchive& arch, const TPair<T0,T1>& v)
+	{
+		arch << v.Key;
+		arch << v.Value;
+	}
+
+	template<typename T0, typename T1> inline void operator>>(FInputArchive& arch, TPair<T0, T1>& v)
+	{
+		arch >> v.Key;
+		arch >> v.Value;
+	}
+
+
+	//---------------------------------------------------------------------------------------------
+	// TODO: As POD?
+	template< typename T >
+	inline void operator<< (FOutputArchive& arch, const UE::Math::TQuat<T>& v)
+	{		
+		arch << v.X;
+		arch << v.Y;
+		arch << v.Z;
+		arch << v.W;
+	}
+
+	template< typename T >
+	inline void operator>> (FInputArchive& arch, UE::Math::TQuat<T>& v)
+	{
+		arch >> v.X;
+		arch >> v.Y;
+		arch >> v.Z;
+		arch >> v.W;
+	}
+
+
+	//---------------------------------------------------------------------------------------------
+	// TODO: As POD?
+	template< typename T >
+	inline void operator<< (FOutputArchive& arch, const UE::Math::TVector<T>& v)
+	{
+		arch << v.X;
+		arch << v.Y;
+		arch << v.Z;
+	}
+
+	template< typename T >
+	inline void operator>> (FInputArchive& arch, UE::Math::TVector<T>& v)
+	{
+		arch >> v.X;
+		arch >> v.Y;
+		arch >> v.Z;
+	}
+
+
+	//---------------------------------------------------------------------------------------------
+	template< typename T >
+	inline void operator<< (FOutputArchive& arch, const UE::Math::TTransform<T>& v)
+	{
+		arch << v.GetRotation();
+		arch << v.GetTranslation();
+		arch << v.GetScale3D();
+	}
+
+	template< typename T >
+	inline void operator>> (FInputArchive& arch, UE::Math::TTransform<T>& v)
+	{
+		UE::Math::TQuat<T> Rot;
+		UE::Math::TVector<T> Trans;
+		UE::Math::TVector<T> Scale;
+
+		arch >> Rot;
+		arch >> Trans;
+		arch >> Scale;
+
+		v.SetComponents(Rot, Trans, Scale);
+	}
 }
 
 
