[Mutable] Remove Core Pimpl.

diff --git a/Engine/Plugins/Mutable/Source/MutableRuntime/Internal/MuR/System.h b/Engine/Plugins/Mutable/Source/MutableRuntime/Internal/MuR/System.h
--- a/Engine/Plugins/Mutable/Source/MutableRuntime/Internal/MuR/System.h
+++ b/Engine/Plugins/Mutable/Source/MutableRuntime/Internal/MuR/System.h
@@ -2,18 +2,33 @@
 
 #pragma once
 
-#include "SkeletalMesh.h"
 #include "MuR/Image.h"
 #include "MuR/Mesh.h"
 #include "MuR/Instance.h"
 #include "MuR/RefCounted.h"
 #include "MuR/Settings.h"
 #include "MuR/Types.h"
+#include "MuR/SystemTypes.h"
+#include "MuR/MemoryCounters.h"
+#include "MuR/Operations.h"
+#include "MuR/MutableString.h"
+#include "MuR/Model.h"
+#include "MuR/Material.h"
+#include "MuR/Parameters.h"
+#include "MuR/MutableTrace.h"
+#include "MuR/ManagedPointer.h"
+
+#include "SkeletalMesh.h"
 #include "Tasks/Task.h"
 #include "Templates/SharedPointer.h"
 #include "Templates/Tuple.h"
 #include "HAL/Platform.h"
 
+#if !(UE_BUILD_SHIPPING || UE_BUILD_TEST)
+#include "HAL/Thread.h"
+#include "HAL/PlatformTLS.h"
+#endif
+
 #include "System.generated.h"
 
 class UTexture;
@@ -55,7 +70,7 @@ namespace UE::Mutable::Private
 	class FParameters;
     class FMesh;
 	class FExtensionDataStreamer;
-
+	
 
 	/** */
 	enum class EExecutionStrategy : uint8
@@ -91,17 +106,1905 @@ namespace UE::Mutable::Private
 		virtual TTuple<UE::Tasks::FTask, TFunction<void()>> GetMeshAsync(USkeletalMesh* SkeleltalMesh, int32 InLODIndex, int32 InSectionIndex, TFunction<void(UE::Mutable::Private::TManagedPtr<FMesh>)>& ResultCallback) = 0;
     };
 
+	
+		// Call the tick of the LLM system (we do this to simulate a frame since the LLM system is not entirelly designed to run over a program)
+	inline void UpdateLLMStats()
+	{
+		// This code will only be compiled (and ran) if the global definition to enable LLM tracking is set to 1 for the host program
+		// Ex : 			GlobalDefinitions.Add("LLM_ENABLED_IN_CONFIG=1");
+#if ENABLE_LOW_LEVEL_MEM_TRACKER && IS_PROGRAM
+		FLowLevelMemTracker& MemTracker = FLowLevelMemTracker::Get();
+		if (MemTracker.IsEnabled())
+		{
+			MemTracker.UpdateStatsPerFrame();
+		}
+#endif
+	}
+
+	constexpr uint64 AllParametersMask = TNumericLimits<uint64>::Max();
+ 
+	/** ExecutinIndex stores the location inside all ranges of the execution of a specific
+	* operation. The first integer on each pair is the dimension/range index in the program
+	* array of ranges, and the second integer is the value inside that range.
+	* The vector order is undefined.
+	*/
+	class ExecutionIndex : public  TArray<TPair<int32, int32>>
+	{
+	public:
+		//! Set or add a value to the index
+		void SetFromModelRangeIndex(uint16 FRangeIndex, int32 RangeValue)
+		{
+			SizeType Index = IndexOfByPredicate([=](const ElementType& v) { return v.Key >= FRangeIndex; });
+			if (Index != INDEX_NONE && (*this)[Index].Key == FRangeIndex)
+			{
+				// Update
+				(*this)[Index].Value = RangeValue;
+			}
+			else
+			{
+				// Add new
+				Push(ElementType(FRangeIndex, RangeValue));
+			}
+		}
+
+		//! Get the value of the index from the range index in the model.
+		int32 GetFromModelRangeIndex(int32 ModelRangeIndex) const
+		{
+			for (const TPair<int32, int32>& e : *this)
+			{
+				if (e.Key == ModelRangeIndex)
+				{
+					return e.Value;
+				}
+			}
+			return 0;
+		}
+	};
+
+	//! Container that stores data per executable code operation (indexed by address and execution
+	//! index).
+	template<class DATA>
+	class CodeContainer
+	{
+		using MemoryCounter = MemoryCounters::FInternalMemoryCounter;
+		using ArrayDataContainerType = TArray<DATA, FDefaultMemoryTrackingAllocator<MemoryCounter>>;
+		using MapDataContainerType = TMap<FCacheAddress, DATA, FDefaultMemoryTrackingSetAllocator<MemoryCounter>>;
+
+	public:
+		void resize(size_t s)
+		{
+			m_index0.SetNumZeroed(s);
+		}
+
+		uint32 size_code() const
+		{
+			return uint32(m_index0.Num());
+		}
+
+		void clear()
+		{
+			m_index0.Empty();
+			m_otherIndex.Empty();
+		}
+
+		inline void erase(const FCacheAddress& at)
+		{
+			if (at.ExecutionIndex == 0 && at.ExecutionOptions == 0 && at.Type == FScheduledOp::EType::Full)
+			{
+				m_index0[at.At] = nullptr;
+			}
+			else
+			{
+				m_otherIndex.Remove(at);
+			}
+		}
+
+		inline DATA get(const FCacheAddress& at) const
+		{
+			if (at.ExecutionIndex == 0 && at.ExecutionOptions == 0 && at.Type == FScheduledOp::EType::Full)
+			{
+				if (at.At < uint32(m_index0.Num()))
+				{
+					return m_index0[at.At];
+				}
+				else
+				{
+					return 0;
+				}
+			}
+			else
+			{
+				const DATA* it = m_otherIndex.Find(at);
+				if (it)
+				{
+					return *it;
+				}
+			}
+			return 0;
+		}
+
+		inline DATA* get_ptr(const FCacheAddress& at)
+		{
+			if (at.ExecutionIndex == 0 && at.ExecutionOptions == 0 && at.Type == FScheduledOp::EType::Full)
+			{
+				if (at.At < uint32(m_index0.Num()))
+				{
+					return &m_index0[at.At];
+				}
+				else
+				{
+					return nullptr;
+				}
+			}
+			else
+			{
+				DATA* it = m_otherIndex.Find(at);
+				if (it)
+				{
+					return it;
+				}
+			}
+			return nullptr;
+		}
+
+		inline const DATA* get_ptr(const FCacheAddress& at) const
+		{
+			if (at.ExecutionIndex == 0 && at.ExecutionOptions == 0 && at.Type == FScheduledOp::EType::Full)
+			{
+				if (at.At < uint32(m_index0.Num()))
+				{
+					return &m_index0[at.At];
+				}
+				else
+				{
+					return nullptr;
+				}
+			}
+			else
+			{
+				const DATA* it = m_otherIndex.Find(at);
+				if (it)
+				{
+					return it;
+				}
+			}
+			return nullptr;
+		}
+
+		inline DATA& operator[](const FCacheAddress& at)
+		{
+			if (at.ExecutionIndex == 0 && at.ExecutionOptions == 0 && at.Type == FScheduledOp::EType::Full)
+			{
+				return m_index0[at.At];
+			}
+			else
+			{
+				return m_otherIndex.FindOrAdd(at);
+			}
+		}
+
+		inline const DATA& operator[](const FCacheAddress& at) const
+		{
+			if (at.ExecutionIndex == 0 && at.ExecutionOptions == 0 && at.Type == FScheduledOp::EType::Full)
+			{
+				return m_index0[at.At];
+			}
+			else
+			{
+				return m_otherIndex[at];
+			}
+		}
+
+		struct iterator
+		{
+		private:
+			friend class CodeContainer<DATA>;
+
+			const CodeContainer<DATA>* container;
+			typename CodeContainer<DATA>::ArrayDataContainerType::TIterator it0;
+			typename CodeContainer<DATA>::MapDataContainerType::TIterator it1;
+
+			iterator(CodeContainer<DATA>* InContainer)
+				: container(InContainer)
+				, it0(InContainer->m_index0.CreateIterator())
+				, it1(InContainer->m_otherIndex.CreateIterator())
+			{
+			}
+
+		public:
+
+			inline void operator++(int)
+			{
+				if (it0)
+				{
+					++it0;
+				}
+				else
+				{
+					++it1;
+				}
+			}
+
+			iterator operator++()
+			{
+				if (it0)
+				{
+					++it0;
+				}
+				else
+				{
+					++it1;
+				}
+				return *this;
+			}
+
+			inline bool operator!=(const iterator& o) const
+			{
+				return it0 != o.it0 || it1 != o.it1;
+			}
+
+			inline DATA& operator*()
+			{
+				if (it0)
+				{
+					return *it0;
+				}
+				else
+				{
+					return it1->Value;
+				}
+			}
+
+			inline FCacheAddress get_address() const
+			{
+				if (it0)
+				{
+					return { uint32_t(it0.GetIndex()), 0, 0};
+				}
+				else
+				{
+					return it1->Key;
+				}
+			}
+
+			inline bool IsValid() const
+			{
+				return bool(it0) || bool(it1);
+			}
+		};
+
+		inline iterator begin()
+		{
+			iterator it(this);
+			return it;
+		}
+
+		inline int32 GetAllocatedSize() const
+		{
+			return m_index0.GetAllocatedSize()
+				+ m_otherIndex.GetAllocatedSize();
+		}
+	private:
+		// For index 0
+		ArrayDataContainerType m_index0;
+
+		// For index>0
+		MapDataContainerType m_otherIndex;
+	};
+
+
+	/** Interface for storage of data while Mutable code is being executed. */
+	class FProgramCache
+	{
+	public:
+		TMemoryTrackedArray<ExecutionIndex, TInlineAllocator<4, AllocType>> m_usedRangeIndices;
+
+		/** Runtime data for each program op. */
+		struct FOpExecutionData
+		{
+			uint16 OpHitCount;
+
+			/** Enabled if the op descriptor has been calculated. */
+			uint8 IsDescCacheValid : 1;
+			uint8 IsValueValid : 1;
+			uint8 IsCacheLocked : 1;
+
+			/** The operation DATA_TYPE. */
+			EDataType DataType;
+
+			/** The position in the type-specific array where the result data is stored. 
+			* 0 means index not valid.
+			* For small types like bool, int and float, the actual value is the index, instead of having their own array.
+			*/
+			union
+			{
+				int32 DataTypeIndex;
+				float ScalarResult;
+			};
+		};
+
+		/** Cached resources while the program is executing. 
+		* first value of the pair:
+		* 0 : value not valid (not set).
+		* 1 : valid, not worth freeing for memory
+		* 2 : valid, worth freeing
+		*/
+		CodeContainer<FOpExecutionData> OpExecutionData;
+
+
+		/** */
+
+		TMemoryTrackedArray<TManagedWeakPtr<const FLayout>> LayoutResults;
+		TMemoryTrackedArray<TManagedWeakPtr<const FInstance>> InstanceResults;
+		TMemoryTrackedArray<TManagedWeakPtr<const String>> StringResults;
+		TMemoryTrackedArray<TManagedWeakPtr<const FExtensionData>> ExtensionDataResults;
+		TMemoryTrackedArray<TManagedWeakPtr<const FInstancedStruct>> InstancedStructResults;
+		TMemoryTrackedArray<TManagedWeakPtr<const FMaterial>> MaterialResults;
+		TMemoryTrackedArray<TManagedWeakPtr<const FSkeletalMesh>> SkeletalMeshResults;
+
+		TMemoryTrackedArray<TManagedWeakPtr<const FImage>> ImageResults;
+		TMemoryTrackedArray<TManagedWeakPtr<const FMesh>> MeshResults;
+		
+		/** Storage for strong references. Size of Strong Results is the same as the weak. */
+		TMemoryTrackedArray<TManagedPtr<const FImage>> StrongImageResults;
+		TMemoryTrackedArray<TManagedPtr<const FMesh>> StrongMeshResults;
+
+		//TMemoryTrackedArray<TManagedPtr<const FResource>> ResourceResults;
+		//TMemoryTrackedArray<TManagedPtr<const FResource>> WeakResourceResults;
+
+		TMemoryTrackedArray<FVector4f> ColorResults;
+		TMemoryTrackedArray<FProjector> ProjectorResults;
+		TMemoryTrackedArray<FMatrix44f> MatrixResults;
+
+		/** */
+		inline const ExecutionIndex& GetRangeIndex(uint32_t i)
+		{
+			// Make sure we have the default element.
+			if (m_usedRangeIndices.IsEmpty())
+			{
+				m_usedRangeIndices.Push(ExecutionIndex());
+			}
+
+			check(i < uint32_t(m_usedRangeIndices.Num()));
+			return m_usedRangeIndices[i];
+		}
+
+		//!
+		inline uint32 GetRangeIndexIndex(const ExecutionIndex& rangeIndex)
+		{
+			if (rangeIndex.IsEmpty())
+			{
+				return 0;
+			}
+
+			// Make sure we have the default element.
+			if (m_usedRangeIndices.IsEmpty())
+			{
+				m_usedRangeIndices.Push(ExecutionIndex());
+			}
+
+			// Look for or add the new element.
+			int32 ElemIndex = m_usedRangeIndices.Find(rangeIndex);
+			if (ElemIndex != INDEX_NONE)
+			{
+				return ElemIndex;
+			}
+
+			m_usedRangeIndices.Push(rangeIndex);
+			return uint32_t(m_usedRangeIndices.Num()) - 1;
+		}
+
+		void Init(uint32 Size)
+		{
+			// This clear would prevent live update cache reusal
+			//OpExecutionData.clear();
+
+			OpExecutionData.resize(Size);
+
+			if (ColorResults.IsEmpty())
+			{
+				// Insert dafault/null values
+
+				LayoutResults.Emplace();
+				InstanceResults.Emplace();
+				StringResults.Emplace();
+				ExtensionDataResults.Emplace();
+				MaterialResults.Emplace();
+				SkeletalMeshResults.Emplace();
+				InstancedStructResults.Add({});
+
+				ImageResults.Emplace();
+				MeshResults.Emplace();
+				
+				StrongImageResults.Emplace();
+				StrongMeshResults.Emplace();
+	
+				ColorResults.Emplace();
+				MatrixResults.Emplace();
+				ProjectorResults.Emplace();
+			}
+
+			for (TManagedPtr<const FImage>& StrongImagePtr : StrongImageResults)
+			{
+				StrongImagePtr = nullptr;
+			}
+
+			for (TManagedPtr<const FMesh>& StrongMeshPtr : StrongMeshResults)
+			{
+				StrongMeshPtr = nullptr;
+			}
+
+			StrongImageResults[0] = MakeManagedNull<const FImage>();
+			StrongMeshResults[0] = MakeManagedNull<const FMesh>();
+
+		}
+
+		bool IsValid(const FCacheAddress& Address) const
+		{
+			if (Address.At == 0 || Address.At>=OpExecutionData.size_code())
+			{
+				return false;
+			}
+
+			const FOpExecutionData* Data = OpExecutionData.get_ptr(Address);
+			if (!Data)
+			{
+				return false;
+			}
+
+			// Is it a desc data query?
+			switch (Address.Type)
+			{
+			case FScheduledOp::EType::Full:
+			{
+				// Resource validity cannot be checked with this method. See 
+				// AquireCachedResult and AquireCachedResultSet.
+				//check(Data->DataType != EDataType::Mesh && Data->DataType != EDataType::Image);
+				return Data->IsValueValid;
+			}
+			case FScheduledOp::EType::ImageDesc:
+			{
+				return Data->IsDescCacheValid;
+			}
+			default:
+				unimplemented();
+				return false;
+			}
+		}
+
+		/** */
+		void CheckHitCountsCleared()
+		{
+#if !(UE_BUILD_SHIPPING || UE_BUILD_TEST)
+			//MUTABLE_CPUPROFILER_SCOPE(CheckHitCountsCleared);
+
+			//int32 IncorrectCount = 0;
+			//CodeContainer<int>::iterator it = m_opHitCount.begin();
+			//for (; it.IsValid(); ++it)
+			//{
+			//	int32 Count = *it;
+			//	if (Count>0 && Count < UE_MUTABLE_CACHE_COUNT_LIMIT)
+			//	{
+			//		// We don't manage the hitcounts of small types that don't use much memory
+			//		if (m_resources[it.get_address()].Key==2)
+			//		{
+			//			// Op hitcount should have reached 0, otherwise, it means we requested an operation but never read 
+			//			// the result.
+			//			++IncorrectCount;
+			//		}
+			//	}
+			//}
+
+			//if (IncorrectCount > 0)
+			//{
+			//	UE_LOG(LogMutableCore, Log, TEXT("The op-hit-count didn't hit 0 for %5d operations. This may mean that too much memory is cached."), IncorrectCount);
+			//	//check(false);
+			//}
+#endif
+		}
+
+
+		void Clear()
+		{
+			MUTABLE_CPUPROFILER_SCOPE(ProgramCacheClear);
+
+			uint32 CodeSize = OpExecutionData.size_code();
+			OpExecutionData.clear();
+			OpExecutionData.resize(CodeSize);
+		}
+
+
+		void ClearDescCache()
+		{
+			MUTABLE_CPUPROFILER_SCOPE(ProgramDescCacheClear);
+
+			CodeContainer<FProgramCache::FOpExecutionData>::iterator it = OpExecutionData.begin();
+			for (; it.IsValid(); ++it)
+			{
+				(*it).IsDescCacheValid = false;
+			}
+		}
+
+
+		bool GetBool(FCacheAddress Address)
+		{
+			if (!Address.At) 
+			{
+				return false;
+			}
+
+			const FOpExecutionData* Data = OpExecutionData.get_ptr(Address);
+			if (!Data)
+			{
+				return false;
+			}
+
+			check(Data->DataType == EDataType::Bool);
+			return Data->DataTypeIndex != 0;
+		}
+
+		float GetScalar(FCacheAddress Address)
+		{
+			if (!Address.At)
+			{
+				return 0.0f;
+			}
+
+			const FOpExecutionData* Data = OpExecutionData.get_ptr(Address);
+			if (!Data)
+			{
+				return 0.0f;
+			}
+
+			check(Data->DataType == EDataType::Scalar);
+			return Data->ScalarResult;
+		}
+
+		int32 GetInt(FCacheAddress Address)
+		{
+			if (!Address.At)
+			{
+				return 0;
+			}
+
+
+			const FOpExecutionData* Data = OpExecutionData.get_ptr(Address);
+			if (!Data) 
+			{
+				return 0;
+			}
+
+			check(Data->DataType == EDataType::Int);
+			return Data->DataTypeIndex;
+		}
+
+		FVector4f GetColour(FCacheAddress Address)
+		{
+			if (!Address.At) 
+			{
+				return FVector4f();
+			}
+
+			const FOpExecutionData* Data = OpExecutionData.get_ptr(Address);
+			if (!Data)
+			{
+				return FVector4f();
+			}
+
+			check(Data->DataType == EDataType::Color);
+			return ColorResults[Data->DataTypeIndex];
+		}
+
+		FMatrix44f GetMatrix(FCacheAddress Address)
+		{
+			if (!Address.At)
+			{
+				return FMatrix44f::Identity;
+			}
+
+			const FOpExecutionData* Data = OpExecutionData.get_ptr(Address);
+			if (!Data) 
+			{
+				return FMatrix44f::Identity;
+			}
+
+			check(Data->DataType == EDataType::Matrix);
+			return MatrixResults[Data->DataTypeIndex];
+		}
+
+		TManagedPtr<const FMaterial> GetMaterial(FCacheAddress at)
+		{
+			if (!at.At)
+			{
+				return MakeManaged<FMaterial>();
+			}
+			const FOpExecutionData* Data = OpExecutionData.get_ptr(at);
+			if (!Data)
+			{
+				return MakeManaged<FMaterial>();
+			}
+
+			check(Data->DataType == EDataType::Material);
+			TManagedPtr<const FMaterial> FoundMaterial = MaterialResults[Data->DataTypeIndex].Pin();
+			check(FoundMaterial);
+
+			if (!FoundMaterial)
+			{
+				return MakeManaged<FMaterial>();
+			}
+			
+			return FoundMaterial;
+		}
+
+		FProjector GetProjector(FCacheAddress at)
+		{
+			if (!at.At) return FProjector();
+			const FOpExecutionData* Data = OpExecutionData.get_ptr(at);
+			if (!Data) return FProjector();
+
+			check(Data->DataType == EDataType::Projector);
+			return ProjectorResults[Data->DataTypeIndex];
+		}
+
+
+		TManagedPtr<const FInstance> GetInstance(FCacheAddress Address)
+		{
+			if (!Address.At)
+			{
+				return nullptr;
+			}
+
+			FOpExecutionData* Data = OpExecutionData.get_ptr(Address);
+			if (!Data)
+			{
+				return nullptr;
+			}
+
+			check(Data->DataType == EDataType::Instance);
+			TManagedPtr<const FInstance> Result = InstanceResults[Data->DataTypeIndex].Pin();
+
+			// We need to decrease the hit-count even if the result is null.
+			check(Data->OpHitCount > 0);
+
+			--Data->OpHitCount;
+			if (Data->OpHitCount == 0 && !Data->IsCacheLocked)
+			{
+				Data->IsValueValid = false;
+				InstanceResults[Data->DataTypeIndex] = nullptr;
+			}
+
+			return Result;
+		}
+
+
+		TManagedPtr<const FImage> GetImage(FCacheAddress Address)
+		{
+			if (!Address.At || Address.At >= OpExecutionData.size_code())
+			{
+				return nullptr;
+			}
+			
+			FOpExecutionData* Data = OpExecutionData.get_ptr(Address);
+			if (!Data) 
+			{
+				return nullptr;
+			}
+
+			check(Data->DataType == EDataType::Image);
+			TManagedPtr<const FImage> Result = StrongImageResults[Data->DataTypeIndex];
+
+			// We need to decrease the hit-count even if the result is null.
+			check(Data->OpHitCount > 0);
+
+			--Data->OpHitCount;
+			if (Data->OpHitCount == 0)
+			{
+				// If this is the last cached reference, now Result is the only reference, which is given 
+				// to the calling code. 
+				if (!Data->IsCacheLocked)
+				{
+					ImageResults[Data->DataTypeIndex] = nullptr;
+				}
+
+				StrongImageResults[Data->DataTypeIndex] = nullptr;
+			}
+
+			return Result;
+		}
+
+		TManagedPtr<const FMesh> GetMesh(FCacheAddress Address)
+		{
+			if (!Address.At || Address.At >= OpExecutionData.size_code()) 
+			{
+				return nullptr;
+			}
+
+			FOpExecutionData* Data = OpExecutionData.get_ptr(Address);
+			if (!Data) 
+			{
+				return nullptr;
+			}
+
+			check(Data->DataType == EDataType::Mesh);
+			TManagedPtr<const FMesh> Result = StrongMeshResults[Data->DataTypeIndex];
+
+			// We need to decrease the hit-count even if the result is null.
+			check(Data->OpHitCount > 0);
+
+			--Data->OpHitCount;
+			if (Data->OpHitCount == 0 && !Data->IsCacheLocked)
+			{
+				// If this is the last cached reference, now Result is the only reference, which is given 
+				// to the calling code. 
+				if (!Data->IsCacheLocked)
+				{
+					MeshResults[Data->DataTypeIndex] = nullptr;
+				}
+
+				StrongMeshResults[Data->DataTypeIndex] = nullptr;
+			}
+
+			return Result;
+		}
+
+		TManagedPtr<const FSkeletalMesh> GetSkeletalMesh(FCacheAddress Address)
+		{
+			if (!Address.At)
+			{
+				return nullptr;
+			}
+
+			if (Address.At >= OpExecutionData.size_code())
+			{
+				return nullptr;
+			}
+
+			FOpExecutionData* Data = OpExecutionData.get_ptr(Address);
+			if (!Data) 
+			{
+				return nullptr;
+			}
+
+			check(Data->DataType == EDataType::SkeletalMesh);
+			TManagedPtr<const FSkeletalMesh> Result = SkeletalMeshResults[Data->DataTypeIndex].Pin();
+
+			// We need to decrease the hit-count even if the result is null.
+			check(Data->OpHitCount > 0);
+
+			--Data->OpHitCount;
+			if (Data->OpHitCount == 0 && !Data->IsCacheLocked)
+			{
+				Data->IsValueValid = false;
+				SkeletalMeshResults[Data->DataTypeIndex] = nullptr;
+			}
+
+			return Result;
+		}
+		
+		TManagedPtr<const FLayout> GetLayout(FCacheAddress Address)
+		{
+			if (!Address.At)
+			{
+				return nullptr;
+			}
+
+			FOpExecutionData* Data = OpExecutionData.get_ptr(Address);
+			if (!Data)
+			{
+				return nullptr;
+			}
+
+			check(Data->DataType == EDataType::Layout);
+			return LayoutResults[Data->DataTypeIndex].Pin();
+		}
+
+
+		TManagedPtr<const String> GetString(FCacheAddress Address)
+		{
+			if (!Address.At)
+			{
+				return nullptr;
+			}
+			FOpExecutionData* Data = OpExecutionData.get_ptr(Address);
+			if (!Data) 
+			{
+				return nullptr;
+			}
+
+			check(Data->DataType == EDataType::String);
+			return StringResults[Data->DataTypeIndex].Pin();
+		}
+
+
+		TManagedPtr<const FExtensionData> GetExtensionData(FCacheAddress Address)
+		{
+			if (!Address.At) 
+			{
+				return nullptr;
+			}
+
+			FOpExecutionData* Data = OpExecutionData.get_ptr(Address);
+			if (!Data)
+			{
+				return nullptr;
+			}
+
+			check(Data->DataType == EDataType::ExtensionData);
+			return ExtensionDataResults[Data->DataTypeIndex].Pin();
+		}
+		
+		TManagedPtr<const FInstancedStruct> GetInstancedStruct(FCacheAddress Address)
+		{
+			if (!Address.At) 
+			{
+				return nullptr;
+			}
+
+			FOpExecutionData* Data = OpExecutionData.get_ptr(Address);
+			if (!Data) 
+			{
+				return nullptr;
+			}
+
+			check(Data->DataType == EDataType::InstancedStruct);
+			TManagedPtr<const FInstancedStruct> Result = InstancedStructResults[Data->DataTypeIndex].Pin();
+
+			check(Result);
+			check(Result->IsValid());
+
+			return Result;
+		}
+		
+		
+		void SetValidDesc(FCacheAddress at)
+		{
+			check(at.Type == FScheduledOp::EType::ImageDesc);
+			check(at.At < OpExecutionData.size_code());
+			FOpExecutionData& Data = OpExecutionData[at];
+			Data.IsDescCacheValid = true;
+		}
+
+		void SetBool(FCacheAddress at, bool v)
+		{
+			check(at.At < OpExecutionData.size_code());
+			FOpExecutionData* Data = OpExecutionData.get_ptr(at);
+			check(Data->DataType == EDataType::Bool || Data->DataType == EDataType::None);
+			Data->DataType = EDataType::Bool;
+			Data->DataTypeIndex = v;
+			Data->IsValueValid = true;
+		}
+
+		void SetInt(FCacheAddress at, int32 v)
+		{
+			check(at.At < OpExecutionData.size_code());
+			FOpExecutionData* Data = OpExecutionData.get_ptr(at);
+			check(Data->DataType == EDataType::Int || Data->DataType == EDataType::None);
+			Data->DataType = EDataType::Int;
+			Data->DataTypeIndex = v;
+			Data->IsValueValid = true;
+		}
+
+		void SetScalar(FCacheAddress Address, float v)
+		{
+			check(Address.At < OpExecutionData.size_code());
+			FOpExecutionData* Data = OpExecutionData.get_ptr(Address);
+			check(Data->DataType == EDataType::Scalar || Data->DataType == EDataType::None);
+			Data->DataType = EDataType::Scalar;
+			Data->ScalarResult = v;
+			Data->IsValueValid = true;
+		}
+
+		void SetColour(FCacheAddress Address, const FVector4f& v)
+		{
+			check(Address.At < OpExecutionData.size_code());
+			FOpExecutionData* Data = OpExecutionData.get_ptr(Address);
+			check(Data->DataType == EDataType::Color || Data->DataType == EDataType::None);
+			Data->DataType = EDataType::Color;
+			Data->IsValueValid = true;
+
+			if (!Data->DataTypeIndex)
+			{
+				Data->DataTypeIndex = ColorResults.Num();
+				ColorResults.Add(v);
+			}
+			else
+			{
+				ColorResults[Data->DataTypeIndex] = v;
+			}
+			check(Data->DataTypeIndex != 0);
+		}
+
+		void SetMatrix(FCacheAddress Address, const FMatrix44f& v)
+		{
+			check(Address.At < OpExecutionData.size_code());
+			FOpExecutionData* Data = OpExecutionData.get_ptr(Address);
+			check(Data->DataType == EDataType::Matrix || Data->DataType == EDataType::None);
+			Data->DataType = EDataType::Matrix;
+			Data->IsValueValid = true;
+
+			if (!Data->DataTypeIndex)
+			{
+				Data->DataTypeIndex = MatrixResults.Num();
+				MatrixResults.Add(v);
+			}
+			else
+			{
+				MatrixResults[Data->DataTypeIndex] = v;
+			}
+			check(Data->DataTypeIndex != 0);
+		}
+
+		void SetMaterial(FCacheAddress Address, const TManagedPtr<const FMaterial>& Value)
+		{
+			check(Value);		// All material pointers provided must be valid
+			check(Address.At < OpExecutionData.size_code());
+			FOpExecutionData* Data = OpExecutionData.get_ptr(Address);
+			check(Data->DataType == EDataType::Material || Data->DataType == EDataType::None);
+			Data->DataType = EDataType::Material;
+			Data->IsValueValid = true;
+
+			if (!Data->DataTypeIndex)
+			{
+				Data->DataTypeIndex = MaterialResults.Num();
+				MaterialResults.Add(Value.ToWeakPtr());
+			}
+			else
+			{
+				MaterialResults[Data->DataTypeIndex] = Value.ToWeakPtr();
+			}
+			check(Data->DataTypeIndex != 0);
+		}
+		
+		void SetProjector(FCacheAddress Address, const FProjector& v)
+		{
+			check(Address.At < OpExecutionData.size_code());
+			FOpExecutionData* Data = OpExecutionData.get_ptr(Address);
+			check(Data->DataType == EDataType::Projector || Data->DataType == EDataType::None);
+			Data->DataType = EDataType::Projector;
+			Data->IsValueValid = true;
+
+			if (!Data->DataTypeIndex)
+			{
+				Data->DataTypeIndex = ProjectorResults.Num();
+				ProjectorResults.Add(v);
+			}
+			else
+			{
+				ProjectorResults[Data->DataTypeIndex] = v;
+			}
+			check(Data->DataTypeIndex != 0);
+		}
+
+		void SetInstance(FCacheAddress Address, const TManagedPtr<const FInstance>& Value)
+		{
+			check(Address.At < OpExecutionData.size_code());
+			FOpExecutionData* Data = OpExecutionData.get_ptr(Address);
+			check(Data->DataType == EDataType::Instance || Data->DataType == EDataType::None);
+			Data->DataType = EDataType::Instance;
+			Data->IsValueValid = true;
+
+			if (!Data->DataTypeIndex)
+			{
+				Data->DataTypeIndex = InstanceResults.Num();
+				InstanceResults.Add(Value.ToWeakPtr());
+			}
+			else
+			{
+				InstanceResults[Data->DataTypeIndex] = Value.ToWeakPtr();
+			}
+
+			check(Data->DataTypeIndex != 0);
+		}
+
+		void SetExtensionData(FCacheAddress Address, const TManagedPtr<const FExtensionData>& Value)
+		{
+			check(Address.At < OpExecutionData.size_code());
+			FOpExecutionData* Data = OpExecutionData.get_ptr(Address);
+			check(Data->DataType == EDataType::ExtensionData || Data->DataType == EDataType::None);
+			Data->DataType = EDataType::ExtensionData;
+			Data->IsValueValid = true;
+
+			if (!Data->DataTypeIndex)
+			{
+				Data->DataTypeIndex = ExtensionDataResults.Num();
+				ExtensionDataResults.Add(Value.ToWeakPtr());
+			}
+			else
+			{
+				ExtensionDataResults[Data->DataTypeIndex] = Value.ToWeakPtr();
+			}
+			check(Data->DataTypeIndex != 0);
+		}
+
+		void SetInstancedStruct(FCacheAddress Address, const TManagedPtr<const FInstancedStruct>& Value)
+		{
+			check(Address.At < OpExecutionData.size_code());
+			FOpExecutionData* Data = OpExecutionData.get_ptr(Address);
+			check(Data->DataType == EDataType::InstancedStruct || Data->DataType == EDataType::None);
+			Data->DataType = EDataType::InstancedStruct;
+			Data->IsValueValid = true;
+
+			if (!Data->DataTypeIndex)
+			{
+				Data->DataTypeIndex = InstancedStructResults.Num();
+				InstancedStructResults.Add(Value.ToWeakPtr());
+			}
+			else
+			{
+				InstancedStructResults[Data->DataTypeIndex] = Value.ToWeakPtr();
+			}
+			check(Data->DataTypeIndex != 0);
+		}
+		
+		void SetImage(FCacheAddress Address, const TManagedPtr<const FImage>& Value)
+		{
+			check(Address.At < OpExecutionData.size_code());
+			FOpExecutionData* Data = OpExecutionData.get_ptr(Address);
+			check(Data->DataType == EDataType::Image || Data->DataType == EDataType::None);
+			Data->DataType = EDataType::Image;
+			Data->IsValueValid = true;
+
+			const TManagedPtr<const FImage>* ValueToStorePtr;
+			if (Value.IsValid())
+			{
+				ValueToStorePtr = &Value;
+			}
+			else
+			{
+				// Set to special null reference that is not valid but can be checked if has been set.
+				ValueToStorePtr = &StrongImageResults[0];  
+			}
+
+			check(ValueToStorePtr->IsSet());
+
+			if (!Data->DataTypeIndex)
+			{
+				Data->DataTypeIndex = ImageResults.Num();
+				ImageResults.Emplace(ValueToStorePtr->ToWeakPtr());
+				StrongImageResults.Emplace(*ValueToStorePtr);
+			}
+			else
+			{
+				ImageResults[Data->DataTypeIndex] = ValueToStorePtr->ToWeakPtr();
+				StrongImageResults[Data->DataTypeIndex] = *ValueToStorePtr;
+			}
+			check(Data->DataTypeIndex != 0);
+
+			UE::Mutable::Private::UpdateLLMStats();
+		}
+
+		void SetMesh(FCacheAddress Address, const TManagedPtr<const FMesh>& Value)
+		{
+			check(Address.At < OpExecutionData.size_code());
+
+			FOpExecutionData* Data = OpExecutionData.get_ptr(Address);
+
+			check(Data->DataType == EDataType::Mesh || Data->DataType == EDataType::None);
+			Data->DataType = EDataType::Mesh;
+			Data->IsValueValid = true;
 
+			const TManagedPtr<const FMesh>* ValueToStorePtr;
+			if (Value.IsValid())
+			{
+				ValueToStorePtr = &Value;
+			}
+			else
+			{
+				// Set to special null reference that is not valid but can be check if is Set.
+				ValueToStorePtr = &StrongMeshResults[0];   
+			}
+
+			check(ValueToStorePtr->IsSet());
+
+			if (!Data->DataTypeIndex)
+			{
+				Data->DataTypeIndex = MeshResults.Num();
+				MeshResults.Emplace(ValueToStorePtr->ToWeakPtr());
+				StrongMeshResults.Emplace(*ValueToStorePtr);
+			}
+			else
+			{
+				MeshResults[Data->DataTypeIndex] = ValueToStorePtr->ToWeakPtr();
+				StrongMeshResults[Data->DataTypeIndex] = *ValueToStorePtr;
+			}
+			check(Data->DataTypeIndex != 0);
+
+			UE::Mutable::Private::UpdateLLMStats();
+		}
+
+		void SetSkeletalMesh(FCacheAddress Address, const TManagedPtr<const FSkeletalMesh>& Value)
+		{
+			check(Address.At < OpExecutionData.size_code());
+			FOpExecutionData* Data = OpExecutionData.get_ptr(Address);
+			check(Data->DataType == EDataType::SkeletalMesh || Data->DataType == EDataType::None);
+			Data->DataType = EDataType::SkeletalMesh;
+			Data->IsValueValid = true;
+
+			if (!Data->DataTypeIndex)
+			{
+				Data->DataTypeIndex = SkeletalMeshResults.Num();
+				SkeletalMeshResults.Add(Value);
+			}
+			else
+			{
+				SkeletalMeshResults[Data->DataTypeIndex] = Value;
+			}
+			check(Data->DataTypeIndex != 0);
+
+			UE::Mutable::Private::UpdateLLMStats();
+		}
+
+		void SetLayout(FCacheAddress Address, const TManagedPtr<const FLayout>& Value)
+		{
+			check(Address.At < OpExecutionData.size_code());
+			FOpExecutionData* Data = OpExecutionData.get_ptr(Address);
+			check(Data->DataType == EDataType::Layout || Data->DataType == EDataType::None);
+			Data->DataType = EDataType::Layout;
+			Data->IsValueValid = true;
+
+			if (!Data->DataTypeIndex)
+			{
+				Data->DataTypeIndex = LayoutResults.Num();
+				LayoutResults.Add(Value);
+			}
+			else
+			{
+				LayoutResults[Data->DataTypeIndex] = Value;
+			}
+			check(Data->DataTypeIndex != 0);
+
+			UE::Mutable::Private::UpdateLLMStats();
+		}
+
+		void SetString(FCacheAddress Address, const TManagedPtr<const String>& Value)
+		{
+			check(Address.At < OpExecutionData.size_code());
+			FOpExecutionData* Data = OpExecutionData.get_ptr(Address);
+			check(Data->DataType == EDataType::String || Data->DataType == EDataType::None);
+			Data->DataType = EDataType::String;
+			Data->IsValueValid = true;
+
+			if (!Data->DataTypeIndex)
+			{
+				Data->DataTypeIndex = StringResults.Num();
+				StringResults.Add(Value.ToWeakPtr());
+			}
+			else
+			{
+				StringResults[Data->DataTypeIndex] = Value.ToWeakPtr();
+			}
+			check(Data->DataTypeIndex!=0);
+		}
+
+
+		inline void IncreaseHitCount(FCacheAddress Address)
+		{
+			// Don't count hits for instruction 0, which is always null. It is usually already
+			// check that At is not 0, and then it is not requested, generating a stray non-zero count
+			// at its position.
+			if (Address.At)
+			{
+				check(Address.At < OpExecutionData.size_code());
+				FOpExecutionData& Data = OpExecutionData[Address];
+				Data.OpHitCount++;
+			}
+		}
+
+		inline void SetForceCached(OP::ADDRESS Address)
+		{
+			// \TODO: It only locks Address,0,0
+			if (Address)
+			{
+				check(Address < OpExecutionData.size_code());
+				FOpExecutionData& Data = OpExecutionData[FCacheAddress(Address, 0, 0)];
+				Data.IsCacheLocked = true;
+			}
+		}
+
+		bool TryAquireCachedResultSet(TConstArrayView<FCacheAddress> AddressSet)
+		{
+			TArray<TManagedPtr<const FImage>, TInlineAllocator<16>> LocalImages;
+			TArray<TManagedPtr<const FMesh> , TInlineAllocator<16>> LocalMeshes;
+			TArray<FOpExecutionData*        , TInlineAllocator<16>> LocalData;
+
+			if (AddressSet.IsEmpty())
+			{
+				return true;
+			}
+
+			for (const FCacheAddress& Address : AddressSet)
+			{
+				if (!Address.At)
+				{
+					continue;
+				}
+
+				check(Address.At < OpExecutionData.size_code());
+
+				FOpExecutionData* Data = OpExecutionData.get_ptr(Address);
+			
+				if (!Data)
+				{
+					return false;
+				}
+			
+				if (Address.Type == FScheduledOp::EType::ImageDesc)
+				{
+					if (!Data->IsDescCacheValid)
+					{
+						return false;
+					}
+				}
+				else
+				{
+					switch (Data->DataType)
+					{
+						case EDataType::Image:
+						{
+							if (!LocalImages.Add_GetRef(ImageResults[Data->DataTypeIndex].Pin()).IsSet())
+							{
+								return false;
+							}
+							LocalData.Add(Data);
+							break;
+						}
+						case EDataType::Mesh:
+						{
+							if (!LocalMeshes.Add_GetRef(MeshResults[Data->DataTypeIndex].Pin()).IsSet())
+							{
+								return false;
+							}
+							LocalData.Add(Data);
+							break;
+						}
+
+						default:
+						{
+							if (!Data->IsValueValid)
+							{
+								return false;
+							}
+							break;
+						}
+					}
+				}
+			}
+
+			// We have been successful acquiring all addresses in the set. 
+			// Commit resource acquisition.
+			for (int32 Index = LocalData.Num() - 1; Index >= 0; --Index)
+			{
+				FOpExecutionData* Data = LocalData[Index];
+
+				switch (Data->DataType)
+				{
+					case EDataType::Image:
+					{
+						StrongImageResults[Data->DataTypeIndex] = LocalImages[Index];
+						break;
+					}
+					case EDataType::Mesh:
+					{
+						StrongMeshResults[Data->DataTypeIndex] = LocalMeshes[Index];
+						break;
+					}
+
+					default:
+					{
+						check(false);
+						break;
+					}
+				}	
+			}
+
+			return true;
+		}
+
+		bool AquireCachedResult(FCacheAddress Address)
+		{
+			if (!Address.At)
+			{
+				return false;
+			}
+
+			check(Address.At < OpExecutionData.size_code());
+
+			FOpExecutionData* Data = OpExecutionData.get_ptr(Address);
+		
+			if (!Data)
+			{
+				return false;
+			}
+
+			if (Address.Type == FScheduledOp::EType::ImageDesc)
+			{
+				return Data->IsDescCacheValid;
+			}
+	
+			switch (Data->DataType)
+			{
+				case EDataType::Image:
+				{
+					if (StrongImageResults[Data->DataTypeIndex].IsSet())
+					{
+						return true;
+					}
+					else
+					{
+						TManagedPtr<const FImage> PinnedPtr = ImageResults[Data->DataTypeIndex].Pin();
+						StrongImageResults[Data->DataTypeIndex] = PinnedPtr;
+						return PinnedPtr.IsSet();
+					}
+				}
+				case EDataType::Mesh:
+				{
+					if (StrongMeshResults[Data->DataTypeIndex].IsSet())
+					{
+						return true;
+					}
+					else
+					{
+
+						TManagedPtr<const FMesh> PinnedPtr = MeshResults[Data->DataTypeIndex].Pin();
+						StrongMeshResults[Data->DataTypeIndex] = PinnedPtr;
+						return PinnedPtr.IsSet();
+					}
+				}
+
+				default:
+				{
+					return Data->IsValueValid;
+				}
+			}
+		}
+	};
+
+
+	inline bool operator==(const FCacheAddress& a, const FCacheAddress& b)
+	{
+		return a.At == b.At
+			&&
+			a.ExecutionIndex == b.ExecutionIndex
+			&&
+			a.ExecutionOptions == b.ExecutionOptions
+			&&
+			a.Type == b.Type;
+	}
+
+	inline bool operator<(const FCacheAddress& a, const FCacheAddress& b)
+	{
+		if (a.At < b.At) return true;
+		if (a.At > b.At) return false;
+		if (a.ExecutionIndex < b.ExecutionIndex) return true;
+		if (a.ExecutionIndex > b.ExecutionIndex) return false;
+		if (a.ExecutionOptions < b.ExecutionOptions) return true;
+		if (a.ExecutionOptions > b.ExecutionOptions) return false;
+		return a.Type < b.Type;
+	}
+
+	/** Data for an instance that is currently being processed in the mutable system. This means it is
+	* between a BeginUpdate and EndUpdate, or during an "atomic" operation (like generate a single resource).
+	*/
+	struct FLiveInstance
+	{
+		FInstance::FID InstanceID;
+		int32 State = 0;
+		uint32 LODMask = 0;
+		TManagedPtr<const FInstance> Instance;
+		TSharedPtr<const FModel> Model;
+
+		TSharedPtr<const FParameters> OldParameters;
+		
+		/** Mask of the parameters that have changed since the last update.
+		* Every bit represents a state parameter.
+		*/
+		uint64 UpdatedParameters = 0;
+
+		/** Cached data for the generation of this instance. */
+		TSharedPtr<FProgramCache> Cache;
+
+		TSharedPtr<FExternalResourceProvider> ExternalResourceProvider;
+		TSharedPtr<FPassthroughObjectLoader> PassthroughObjectLoader;
+		
+		~FLiveInstance()
+		{
+			// Manually done to trace mem deallocations
+			MUTABLE_CPUPROFILER_SCOPE(LiveInstanceDestructor);
+			{
+				MUTABLE_CPUPROFILER_SCOPE(LiveInstanceDestructor_DeleteCache);
+				Cache = nullptr;
+			}
+
+			{	
+				MUTABLE_CPUPROFILER_SCOPE(LiveInstanceDestructor_DeleteOldParameters);
+				OldParameters = nullptr;
+			}
+			{
+				MUTABLE_CPUPROFILER_SCOPE(LiveInstanceDestructor_DeleteInstance);
+				Instance = nullptr;
+			}
+			{
+				MUTABLE_CPUPROFILER_SCOPE(LiveInstanceDestructor_DeleteModel);
+				Model = nullptr;
+			}
+		}
+	};
+
+
+    /** Struct to manage all the memory allocated for resources used during mutable operation. */
+    struct FWorkingMemoryManager
+    {	
+		using MemoryCounter = MemoryCounters::FInternalMemoryCounter;
+
+		template<class Type>
+		using TMemoryTrackedArray = TArray<Type, FDefaultMemoryTrackingAllocator<MemoryCounter>>;
+
+		template<class KeyType, class ValueType>
+		using TMemoryTrackedMap = TMap<KeyType, ValueType, FDefaultMemoryTrackingSetAllocator<MemoryCounter>>;
+
+		/** Cached traking for streamed model data for one model. */
+        struct FModelCacheEntry
+        {
+            /** Model who's data is being tracked. */
+            TWeakPtr<const FModel> Model;
+
+            /** For each model rom, the last time its streamed data was used. */
+            TMemoryTrackedArray<TPair<uint64, uint64>> RomWeights;
+
+			/** Count of pending operations for every rom index. */
+			TMemoryTrackedArray<uint16> PendingOpsPerRom;
+        };
+
+		FLiveInstance CurrentLiveInstance;
+
+		/** Maximum working memory that mutable should be using. */
+		int64 BudgetBytes = 0;
+
+		/** Maximum excess memory reached suring the current operation. */
+		int64 BudgetExcessBytes = 0;
+    	
+		/** This value is used to track the order of loading of roms. */
+        uint64 RomTick = 0;
+
+		/** Control info for the per-model cache of streamed data. */
+        TMemoryTrackedArray<FModelCacheEntry> CachePerModel;
+
+		/** Data for each mutable instance that is being updated. */
+		TMemoryTrackedArray<FLiveInstance> LiveInstances;
+
+		/** Temporary reference to the memory of the current instance being updated. Only valid during a mutable "atomic" operation, like a BeginUpdate or a GetImage. */
+		TSharedPtr<FProgramCache> CurrentInstanceCache;
+
+    	TSharedPtr<FMeshIdRegistry> CurrentMeshIdRegistry;
+    	
+		TSharedPtr<FImageIdRegistry> CurrentImageIdRegistry;
+
+    	TSharedPtr<FMaterialIdRegistry> CurrentMaterialIdRegistry;
+
+    	TSharedPtr<FSkeletalMeshIdRegistry> CurrentSkeletalMeshIdRegistry;
+
+    	/** Only valid in initial generation. We can not load UObjects while streaming. */
+    	TSharedPtr<FPassthroughObjectLoader> CurrentPassthroughObjectLoader;
+    	
+		/** Resources that have been used in the past, but haven't been deallocated because they still fitted the memory budget and they could be reused. */
+		TMemoryTrackedArray<TSharedPtr<FImage>> PooledImages;
+
+		/** List of intermediate resources that are not soterd anywhere yet. They are still locally referenced by code. */
+		TMemoryTrackedArray<TSharedPtr<const FImage>> TempImages;
+		TMemoryTrackedArray<TSharedPtr<const FMesh>> TempMeshes;
+
+		/** List of resources that are currently in any cache position, and the number of positions they are in. */
+		TMemoryTrackedMap<TSharedPtr<const FResource>, int32> CacheResources;
+
+		/** Given a mutable model, find or create its rom cache. */
+		FModelCacheEntry* FindModelCache(const FModel*);
+		FModelCacheEntry& FindOrAddModelCache(const TSharedPtr<const FModel>&);
+
+        /** Make sure the working memory is below the internal budget, even counting with the passed additional memory. 
+		* An optional function can be passed to "block" the unload of certain roms of data.
+		* Return true if it succeeded, false otherwise.
+		*/
+		bool EnsureBudgetBelow(uint64 AdditionalMemory);
+		
+		/** Return true if the memory budget is 90% full. */
+		bool IsMemoryBudgetFull() const;
+
+		/** Calculate the current usage of memory as used to calculate the budget. */
+		int64 GetCurrentMemoryBytes() const;
+
+        /** Register that a specific rom has been requested and update the heuristics to keep it in memory. */
+        void MarkRomUsed( int32 RomIndex, const TSharedPtr<const FModel>& );
+
+		/** */
+		[[nodiscard]] TManagedPtr<FImage> CreateImage(uint32 SizeX, uint32 SizeY, uint32 Lods, EImageFormat Format, EInitializationType Init)
+		{
+			CheckRunnerThread();
+
+			uint32 DataSize = FImage::CalculateDataSize(SizeX, SizeY, Lods, Format);
+
+			//TODO: 
+
+			// Make room in the budget
+			EnsureBudgetBelow(DataSize);
+
+			// Create it
+			TManagedPtr<FImage> Result = MakeManaged<FImage>(SizeX, SizeY, Lods, Format, Init);
+
+			return Result;
+		}
+
+		/** Ref will be nulled and relesed in any case. */
+		[[nodiscard]] TManagedPtr<FImage> CloneOrTakeOver(TManagedPtr<const FImage>&& Resource)
+		{
+			CheckRunnerThread();
+	
+			TManagedPtr<FImage> Result;
+			if (!Resource.IsUniqueReference())
+			{
+				uint32 DataSize = Resource->GetDataSize();
+				EnsureBudgetBelow(DataSize);
+
+				Result = Resource->Clone();
+			}
+			else
+			{
+				Result = ConstCastManagedPtr<FImage>(Resource);
+			}
+
+			Resource = nullptr;
+			return Result;
+		}
+
+
+		[[nodiscard]] TManagedPtr<FMesh> CreateMesh(int32 BudgetReserveSize)
+		{
+			CheckRunnerThread();
+
+			EnsureBudgetBelow(BudgetReserveSize);
+
+			TManagedPtr<FMesh> Result = MakeManaged<FMesh>();
+
+			return Result;
+		}
+
+		[[nodiscard]] TManagedPtr<FMesh> CloneOrTakeOver(TManagedPtr<const FMesh>&& Resource)
+		{
+			CheckRunnerThread();
+
+			TManagedPtr<FMesh> Result;
+			if (!Resource.IsUniqueReference())
+			{
+				const int32 ResourceDataSize = Resource->GetDataSize();
+				Result = CreateMesh(ResourceDataSize); 
+				Result->CopyFrom(*Resource);
+			}
+			else
+			{
+				Result = ConstCastManagedPtr<FMesh>(Resource);
+			}
+
+			Resource = nullptr;
+			return Result;
+		}
+
+		/** */
+		[[nodiscard]] TManagedPtr<const FMesh> LoadMesh(const FCacheAddress& From)
+		{
+			TManagedPtr<const FMesh> Result = CurrentInstanceCache->GetMesh(From);
+			return Result;
+		}
+
+    	/** */
+    	[[nodiscard]] TManagedPtr<const FSkeletalMesh> LoadSkeletalMesh(const FCacheAddress& From)
+		{
+			TManagedPtr<const FSkeletalMesh> Result = CurrentInstanceCache->GetSkeletalMesh(From);
+			return Result;
+		}
+    	
+
+		/** */
+		[[nodiscard]] TManagedPtr<const FImage> LoadImage(const FCacheAddress& From)
+		{
+			return CurrentInstanceCache->GetImage(From);
+		}
+
+
+    	/** */
+    	[[nodiscard]] TManagedPtr<const FMaterial> LoadMaterial(const FCacheAddress& From)
+		{
+			TManagedPtr<const FMaterial> Result = CurrentInstanceCache->GetMaterial(From);
+			return Result;
+		}
+
+		/** */
+		void StoreImage(const FCacheAddress& To, const TManagedPtr<const FImage>& Resource)
+		{
+			CurrentInstanceCache->SetImage(To, Resource);
+		}
+
+		void StoreMesh(const FCacheAddress& To, const TManagedPtr<const FMesh>& Resource)
+		{
+			//Resource->CheckIntegrity();
+			CurrentInstanceCache->SetMesh(To, Resource);
+		}
+
+    	void StoreMaterial(const FCacheAddress& To, const TManagedPtr<const FMaterial>& Resource)
+		{
+			CurrentInstanceCache->SetMaterial(To, Resource);
+		}
+
+		int32 GetRomBytes() const;
+
+		/** Calculate the amount of bytes in data cached in the level 0 and 1 cache in all live instances. */
+		int32 GetCacheBytes() const
+		{
+			int32 Result = 0;
+			TSet<TManagedPtr<const FResource>> Cache0Unique;
+			TSet<TManagedPtr<const FResource>> Cache1Unique;
+
+			for (const FLiveInstance& Instance : LiveInstances)
+			{
+			CodeContainer<FProgramCache::FOpExecutionData>::iterator it = Instance.Cache->OpExecutionData.begin();
+			for (; it.IsValid(); ++it)
+			{
+				if (!(*it).DataTypeIndex)
+				{
+					continue;
+				}
+
+				TSet<TManagedPtr<const FResource>>* TargetSet = &Cache0Unique;
+				if ((*it).IsCacheLocked)
+				{
+					TargetSet = &Cache1Unique;
+				}
+
+				switch ((*it).DataType)
+				{
+				case EDataType::Image:
+				{
+					TManagedPtr<const FImage> Value = Instance.Cache->ImageResults[(*it).DataTypeIndex].Pin();
+					if (Value)
+					{
+						TargetSet->Add(MoveTemp(Value));
+					}
+					break;
+				}
+				case EDataType::Mesh:
+				{
+					TManagedPtr<const FMesh> Value = Instance.Cache->MeshResults[(*it).DataTypeIndex].Pin();
+					if (Value)
+					{
+						TargetSet->Add(MoveTemp(Value));
+					}
+					break;
+				}
+				default:
+					break;
+				}
+			}
+		}
+
+			// Merge
+			Cache0Unique.Append(Cache1Unique);
+
+			// Count
+			for (const TManagedPtr<const FResource>& Value : Cache0Unique)
+			{
+				Result += Value->GetDataSize();
+			}
+
+			return Result;
+		}
+
+		/** 
+		 * Remove all intermediate data (big and small) from the memory except for the one that has been explicitely
+		 * marked as state cache.
+		 */
+		void ClearCacheLayer0()
+		{
+			check(CurrentInstanceCache);
+
+			MUTABLE_CPUPROFILER_SCOPE(ClearLayer0);
+
+			CodeContainer<FProgramCache::FOpExecutionData>::iterator it = CurrentInstanceCache->OpExecutionData.begin();
+			for (; it.IsValid(); ++it)
+			{
+				FProgramCache::FOpExecutionData& Data = *it;
+				if (!Data.DataTypeIndex || Data.IsCacheLocked)
+				{
+					continue;
+				}
+
+				switch (Data.DataType)
+				{
+				case EDataType::Image:
+				{
+					CurrentInstanceCache->ImageResults[Data.DataTypeIndex] = nullptr;
+					CurrentInstanceCache->StrongImageResults[Data.DataTypeIndex] = nullptr;
+					break;
+				}
+
+				case EDataType::Mesh:
+				{
+					CurrentInstanceCache->MeshResults[Data.DataTypeIndex] = nullptr;
+					CurrentInstanceCache->StrongMeshResults[Data.DataTypeIndex] = nullptr;
+					break;
+				}
+
+				case EDataType::SkeletalMesh:
+				{
+					CurrentInstanceCache->SkeletalMeshResults[Data.DataTypeIndex] = nullptr;
+					break;
+				}
+
+				case EDataType::Layout:
+				{
+					CurrentInstanceCache->LayoutResults[Data.DataTypeIndex] = nullptr;
+					break;
+				}
+
+				case EDataType::Instance:
+				{
+					CurrentInstanceCache->InstanceResults[Data.DataTypeIndex] = nullptr;
+					break;
+				}
+
+				default:
+					break;
+				}
+
+				Data.OpHitCount = 0;
+				Data.IsValueValid = false;
+			}
+		}
+
+
+		/** Remove all intermediate data (big and small) from the memory including the one that has been explicitely
+		* marked as state cache.
+		*/
+		void ClearCacheLayer1()
+		{
+			MUTABLE_CPUPROFILER_SCOPE(ClearLayer1);
+
+			CodeContainer<FProgramCache::FOpExecutionData>::iterator it = CurrentInstanceCache->OpExecutionData.begin();
+			for (; it.IsValid(); ++it)
+			{
+				FProgramCache::FOpExecutionData& Data = *it;
+				
+				Data.OpHitCount = 0;
+				Data.IsValueValid = false;
+			}
+
+			for (TManagedWeakPtr<const FImage>& Value : CurrentInstanceCache->ImageResults)
+			{
+				Value = nullptr;
+			}
+
+			for (TManagedPtr<const FImage>& Value : CurrentInstanceCache->StrongImageResults)
+			{
+				Value = nullptr;
+			}
+
+
+			for (TManagedWeakPtr<const FSkeletalMesh>& Value : CurrentInstanceCache->SkeletalMeshResults)
+			{
+				Value = nullptr;
+			}
+
+			for (TManagedPtr<const FMesh>& Value : CurrentInstanceCache->StrongMeshResults)
+			{
+				Value = nullptr;
+			}
+
+
+			for (TManagedWeakPtr<const FLayout>& Value : CurrentInstanceCache->LayoutResults)
+			{
+				Value = nullptr;
+			}
+
+			for (TManagedWeakPtr<const FInstance>& Value : CurrentInstanceCache->InstanceResults)
+			{
+				Value = nullptr;
+			}
+		}
+
+
+		void LogWorkingMemory(const class CodeRunner* CurrentRunner) const;
+
+
+#if !(UE_BUILD_SHIPPING || UE_BUILD_TEST)
+		static constexpr uint64 InvalidRunnerId = 0;
+		std::atomic<uint64> DebugCurrentRunnerId {1};
+		
+		/** Temp variable with the ID of the thread running code, for debugging. */
+		uint32 DebugRunnerThreadID = FThread::InvalidThreadId;
+		uint64 DebugRunnerID = InvalidRunnerId;
+#endif
+
+		/** This is a development-only check to make sure calls to resource management happen in the correct thread. */
+		FORCEINLINE void BeginRunnerThread()
+		{
+#if !(UE_BUILD_SHIPPING || UE_BUILD_TEST)
+			// If this check fails it means have not set up correctly all the paths to debug threading for resource management.
+			check(DebugRunnerThreadID == FThread::InvalidThreadId);
+			check(DebugRunnerID == InvalidRunnerId)
+
+			DebugRunnerThreadID = FPlatformTLS::GetCurrentThreadId();
+			DebugRunnerID = ++DebugCurrentRunnerId;
+#endif
+		}
+
+		/** This is a development-only check to make sure calls to resource management happen in the correct thread. */
+		FORCEINLINE void ResetRunnerThread()
+		{
+#if !(UE_BUILD_SHIPPING || UE_BUILD_TEST)
+			// If this check fails it means have not set up correctly all the paths to debug threading for resource management.
+			check(DebugRunnerThreadID == FThread::InvalidThreadId);
+			check(DebugRunnerID != InvalidRunnerId)
+
+			DebugRunnerThreadID = FPlatformTLS::GetCurrentThreadId();
+#endif
+		}
+
+		/** This is a development-only check to make sure calls to resource management happen in the correct thread. */
+		FORCEINLINE void InvalidateRunnerThread()
+		{
+#if !(UE_BUILD_SHIPPING || UE_BUILD_TEST)
+			// If this check fails it means have not set up correctly all the paths to debug threading for resource management.
+			check(DebugRunnerThreadID != FThread::InvalidThreadId);
+			check(DebugRunnerID != InvalidRunnerId)
+
+			DebugRunnerThreadID = FThread::InvalidThreadId;
+#endif
+		}
+
+		/** This is a development-only check to make sure calls to resource management happen in the correct thread. */
+		FORCEINLINE void CheckRunnerThread()
+		{
+#if !(UE_BUILD_SHIPPING || UE_BUILD_TEST)
+			// If this check fails it means have not set up correctly all the paths to debug threading for resource management.
+			check(DebugRunnerThreadID != FThread::InvalidThreadId);
+			check(DebugRunnerID != InvalidRunnerId);
+
+			// If this check fails it means we are doing resource management from a thread that is not the CodeRunner::RunCode
+			// thread, and this is not allowed.
+			check(DebugRunnerThreadID == FPlatformTLS::GetCurrentThreadId());
+			check(DebugRunnerID == DebugCurrentRunnerId);
+#endif
+		}
+
+
+		/** This is a development-only check to make sure calls to resource management happen in the correct thread. */
+		FORCEINLINE void EndRunnerThread()
+		{
+			CurrentInstanceCache->CheckHitCountsCleared();
+
+			// If this check fails it means some operation is not correctly handling resource management and didn't release
+			// a resource it created.
+			// This should be reported and reviewed, but it is not fatal. Some unnecessary memory may be used temporarily.
+			ensure(TempImages.Num() == 0);
+			ensure(TempMeshes.Num() == 0);
+
+#if !(UE_BUILD_SHIPPING || UE_BUILD_TEST)
+			// If this check fails it means have not set up correctly all the paths to debug threading for resource management.
+			check(DebugRunnerThreadID != FThread::InvalidThreadId);
+			check(DebugRunnerID != InvalidRunnerId);
+
+			DebugRunnerThreadID = FThread::InvalidThreadId;
+			DebugRunnerID = InvalidRunnerId;
+#endif
+		}
+
+	};
+	
+	
     /** Main system class to load models and build instances. */
 	class FSystem
 	{
     public:
-
         //! This constant can be used in place of the lodMask in methods like BeginUpdate
         static constexpr uint32 AllLODs = 0xffffffff;
 
-    public:
-
 		//! Constructor of a system object to build data.
         //! \param Settings Optional class with the settings to use in this system. The default
         //! value configures a production-ready system.
@@ -229,30 +2132,68 @@ namespace UE::Mutable::Private
 
 		UE_API static FSkeletalMeshId GetSkeletalMeshId(const TSharedPtr<FSkeletalMeshIdRegistry>& Registry, const TSharedPtr<const FModel>& Model, const FParameters* Params, OP::ADDRESS RootAt);
 		
-		//-----------------------------------------------------------------------------------------
-		// Interface pattern
-		//-----------------------------------------------------------------------------------------
-		class Private;
-
-		UE_API Private* GetPrivate() const;
-
-    public:
-
         // Prevent copy, move and assignment.
         FSystem( const FSystem& ) = delete;
 		FSystem& operator=( const FSystem& ) = delete;
         FSystem( FSystem&& ) = delete;
         FSystem& operator=( FSystem&& ) = delete;
+		
+		/** This method can be used to internally prepare for code execution. */
+		MUTABLERUNTIME_API void BeginBuild(const TSharedPtr<const FModel>&,
+			const TSharedPtr<FMeshIdRegistry>&,
+			const TSharedPtr<FImageIdRegistry>&,
+			const TSharedPtr<FMaterialIdRegistry>&,
+			const TSharedPtr<FSkeletalMeshIdRegistry>&);
+    	
+		MUTABLERUNTIME_API void EndBuild();
 
-		UE_API ~FSystem();
+		MUTABLERUNTIME_API bool BuildBool(const TSharedPtr<FExternalResourceProvider>& ExternalResourceProvider, const TSharedPtr<const FModel>&, const FParameters*, OP::ADDRESS) ;
+		MUTABLERUNTIME_API int32 BuildInt(const TSharedPtr<FExternalResourceProvider>& ExternalResourceProvider, const TSharedPtr<const FModel>&, const FParameters*, OP::ADDRESS) ;
+		MUTABLERUNTIME_API float BuildScalar(const TSharedPtr<FExternalResourceProvider>& ExternalResourceProvider, const TSharedPtr<const FModel>&, const FParameters*, OP::ADDRESS) ;
+		MUTABLERUNTIME_API FVector4f BuildColour(const TSharedPtr<FExternalResourceProvider>& ExternalResourceProvider, const TSharedPtr<const FModel>&, const FParameters*, OP::ADDRESS) ;
+		MUTABLERUNTIME_API TManagedPtr<const String> BuildString(const TSharedPtr<FExternalResourceProvider>& ExternalResourceProvider, const TSharedPtr<const FModel>&, const FParameters*, OP::ADDRESS) ;
+		MUTABLERUNTIME_API TManagedPtr<const FImage> BuildImage(const TSharedPtr<FExternalResourceProvider>& ExternalResourceProvider, const TSharedPtr<const FModel>&, const FParameters*, OP::ADDRESS, int32 MipsToSkip);
+		MUTABLERUNTIME_API TManagedPtr<const FMesh> BuildMesh(const TSharedPtr<FExternalResourceProvider>& ExternalResourceProvider, const TSharedPtr<const FModel>&, const FParameters*, OP::ADDRESS, EMeshContentFlags MeshContentFilter);
+		MUTABLERUNTIME_API TManagedPtr<const FInstance> BuildInstance(const TSharedPtr<FExternalResourceProvider>& ExternalResourceProvider, const TSharedPtr<const FModel>&, const FParameters*, OP::ADDRESS);
+		MUTABLERUNTIME_API TManagedPtr<const FLayout> BuildLayout(const TSharedPtr<FExternalResourceProvider>& ExternalResourceProvider, const TSharedPtr<const FModel>&, const FParameters*, OP::ADDRESS) ;
+    	MUTABLERUNTIME_API FProjector BuildProjector(const TSharedPtr<FExternalResourceProvider>& ExternalResourceProvider, const TSharedPtr<const FModel>&, const FParameters*, OP::ADDRESS) ;
+    	
+        FSettings Settings;
 
-	private:
+        //! Data streaming interface, if any.
+		TSharedPtr<FModelReader> StreamInterface;
+    	
+		/** Counter used to generate unique IDs for every new instance created in the system. */
+        FInstance::FID LastInstanceID = 0;
 
-		Private* m_pD;
+		/** This flag is turned on when a streaming error or similar happens. Results are not usable.
+		* This should only happen in-editor.
+		*/
+		bool bUnrecoverableError = false;
 
-	};
+		/** If this is set, it will be tried first instead of the internal formatting function. */
+		FImageOperator::FImagePixelFormatFunc ImagePixelFormatOverride;
+
+		FWorkingMemoryManager WorkingMemoryManager;
 
+		/** The pointer returned by this function is only valid for the duration of the current mutable operation. */
+		inline FLiveInstance* FindLiveInstance(FInstance::FID);
 
+        bool CheckUpdatedParameters( const FLiveInstance*, const TSharedPtr<const FParameters>&, uint64& OutUpdatedParameters );
+		
+		void RunCode(
+			const TSharedPtr<FExternalResourceProvider>& ExternalResourceProvider,
+			const TSharedPtr<const FModel>&,
+			const FParameters*,
+			OP::ADDRESS,
+			uint32 LODs = AllLODs,
+			uint8 ExecutionOptions = 0);
+
+		void PrepareCache(const FModel*, int32 State);
+
+		//! Update some mutable core unreal stats.
+		void UpdateStats();
+	};
 }
 
 #undef UE_API
diff --git a/Engine/Plugins/Mutable/Source/MutableRuntime/Internal/MuR/SystemTypes.h b/Engine/Plugins/Mutable/Source/MutableRuntime/Internal/MuR/SystemTypes.h
new file mode 100644
--- /dev/null
+++ b/Engine/Plugins/Mutable/Source/MutableRuntime/Internal/MuR/SystemTypes.h
@@ -0,0 +1,130 @@
+// Copyright Epic Games, Inc. All Rights Reserved.
+
+#pragma once
+
+#include "MuR/Operations.h"
+
+#ifndef MUTABLE_DEBUG_CODERUNNER_TASK_SCHEDULE_CALLSTACK
+	#define MUTABLE_DEBUG_CODERUNNER_TASK_SCHEDULE_CALLSTACK 0	
+#endif
+
+namespace UE::Mutable::Private
+{
+	/** This structure stores the data about an ongoing mutable operation that needs to be executed. */
+	struct FScheduledOp
+	{
+		inline FScheduledOp() = default;
+
+		inline FScheduledOp(OP::ADDRESS InAt, const FScheduledOp& InOpTemplate, uint8 InStage = 0, uint32 InCustomState = 0)
+			: At(InAt)
+			, CustomState(InCustomState)
+			, ExecutionIndex(InOpTemplate.ExecutionIndex)
+			, ExecutionOptions(InOpTemplate.ExecutionOptions)
+			, Stage(InStage)
+			, Type(InOpTemplate.Type)
+		{
+			check(InStage < ((1 << 6) - 1));
+		}
+
+		static inline FScheduledOp FromOpAndOptions(OP::ADDRESS InAt, const FScheduledOp& InOpTemplate, uint8 InExecutionOptions)
+		{
+			FScheduledOp Result;
+
+			Result.At = InAt;
+			Result.ExecutionOptions = InExecutionOptions;
+			Result.ExecutionIndex = InOpTemplate.ExecutionIndex;
+			Result.Stage = 0;
+			Result.CustomState = InOpTemplate.CustomState;
+			Result.Type = InOpTemplate.Type;
+
+			return Result;
+		}
+
+		//! Address of the operation
+		OP::ADDRESS At = 0;
+
+		//! Additional custom state data that the operation can store. This is usually used to pass information
+		//! between execution stages of an operation.
+		uint32 CustomState = 0;
+
+		//! Index of the operation execution: This is used for iteration of different ranges.
+		//! It is an index into the CodeRunner::GetMemory()::m_rangeIndex vector.
+		//! executionIndex 0 is always used for empty ExecutionIndex, which is the most common
+		//! one.
+		uint16 ExecutionIndex = 0;
+
+		//! Additional execution options. Set externally to this op, it usually alters the result.
+		//! For example, this is used to keep track of the mipmaps to skip in image operations.
+		uint8 ExecutionOptions = 0;
+
+		//! Internal stage of the operation.
+		//! Stage 0 is usually scheduling of children, and 1 is execution. Some instructions
+		//! may have more steges to schedule children that are optional for execution, etc.
+		uint8 Stage      : 6 = 0;
+
+		uint8 bForceLazy : 1 = 0;
+
+		//! Type of calculation we are requesting for this operation.
+		enum class EType : uint8
+		{
+			//! Execute the operation to calculate the full result
+			Full,
+
+			//! Execute the operation to obtain the descriptor of an image or mesh.
+			ImageDesc,
+		};
+		EType Type       : 1 = EType::Full;
+
+#if MUTABLE_DEBUG_CODERUNNER_TASK_SCHEDULE_CALLSTACK
+		static constexpr uint64 CallstackMaxDepth = 16;
+		uint64 StackDepth = 0; 
+		uint64 ScheduleCallstack[CallstackMaxDepth];
+#endif
+	};
+	
+	inline uint32 GetTypeHash(const FScheduledOp& Op)
+	{
+		return HashCombineFast(::GetTypeHash(Op.At), HashCombineFast((uint32)Op.Stage), (uint32)Op.ExecutionIndex);
+	}
+
+	/** A cache address is the operation plus the context of execution(iteration indices, etc...). */
+	struct FCacheAddress
+	{
+		/** The meaning of all these fields is the same than the FScheduledOp struct. */
+		OP::ADDRESS At = 0;
+		uint16 ExecutionIndex = 0;
+		uint8 ExecutionOptions = 0;
+		FScheduledOp::EType Type = FScheduledOp::EType::Full;
+
+		FCacheAddress() 
+		{
+		}
+
+		FCacheAddress(OP::ADDRESS InAt, uint16 InExecutionIndex, uint8 InExecutionOptions, FScheduledOp::EType InType = FScheduledOp::EType::Full)
+			: At(InAt), ExecutionIndex(InExecutionIndex), ExecutionOptions(InExecutionOptions), Type(InType)
+		{
+		}
+
+		FCacheAddress(OP::ADDRESS InAt, const FScheduledOp& Item)
+			: At(InAt), ExecutionIndex(Item.ExecutionIndex), ExecutionOptions(Item.ExecutionOptions), Type(Item.Type)
+		{
+		}
+		
+		FCacheAddress(const FScheduledOp& Item)
+			: At(Item.At), ExecutionIndex(Item.ExecutionIndex), ExecutionOptions(Item.ExecutionOptions), Type(Item.Type)
+		{
+		}
+	};
+
+	inline uint32 GetTypeHash(const FCacheAddress& Address)
+	{
+		uint32 Hash = Address.At;
+
+		Hash = HashCombineFast(Hash, static_cast<uint32>(Address.ExecutionIndex));
+		Hash = HashCombineFast(Hash, static_cast<uint32>(Address.ExecutionOptions));
+		Hash = HashCombineFast(Hash, static_cast<uint32>(Address.Type));
+
+		return Hash;
+	}
+
+} // namespace UE::Mutable::Private
diff --git a/Engine/Plugins/Mutable/Source/MutableRuntime/Private/MuR/CodeRunner.cpp b/Engine/Plugins/Mutable/Source/MutableRuntime/Private/MuR/CodeRunner.cpp
--- a/Engine/Plugins/Mutable/Source/MutableRuntime/Private/MuR/CodeRunner.cpp
+++ b/Engine/Plugins/Mutable/Source/MutableRuntime/Private/MuR/CodeRunner.cpp
@@ -17,7 +17,6 @@
 #include "MuR/Mesh.h"
 #include "MuR/MeshBufferSet.h"
 #include "MuR/Model.h"
-#include "MuR/ModelPrivate.h"
 #include "MuR/MutableMath.h"
 #include "MuR/MutableString.h"
 #include "MuR/MutableTrace.h"
@@ -56,10 +55,10 @@
 #include "MuR/OpMeshOptimizeSkinning.h"
 #include "MuR/Operations.h"
 #include "MuR/Parameters.h"
-#include "MuR/ParametersPrivate.h"
+#include "MuR/Parameters.h"
 #include "MuR/PhysicsBody.h"
 #include "MuR/Skeleton.h"
-#include "MuR/SystemPrivate.h"
+#include "MuR/System.h"
 #include "MuR/External/FloatAdapter.h"
 #include "MuR/External/MaterialAdapter.h"
 #include "MuR/External/TextureAdapter.h"
@@ -153,7 +152,7 @@ namespace UE::Mutable::Private
 
 	TSharedRef<CodeRunner> CodeRunner::Create(
 		const FSettings& InSettings,
-		FSystem::Private* InSystem,
+		FSystem* InSystem,
 		EExecutionStrategy InExecutionStrategy,
 		const TSharedPtr<const FModel>& InModel,
 		const FParameters* InParams,
@@ -171,7 +170,7 @@ namespace UE::Mutable::Private
 
     CodeRunner::CodeRunner(FPrivateToken PrivateToken, 
 		const FSettings& InSettings,
-		FSystem::Private* InSystem,
+		FSystem* InSystem,
 		EExecutionStrategy InExecutionStrategy,
 		const TSharedPtr<const FModel>& InModel,
 		const FParameters* InParams,
@@ -190,7 +189,7 @@ namespace UE::Mutable::Private
 		, ExternalResourceProvider(ExternalResourceProvider)
 	{
 		MUTABLE_CPUPROFILER_SCOPE(CodeRunner_Create)
-		const FProgram& Program = Model->GetPrivate()->Program;
+		const FProgram& Program = Model->Program;
 		ScheduledStagePerOp.resize(Program.OpAddress.Num());
 
 		
@@ -301,7 +300,7 @@ namespace UE::Mutable::Private
     {
         MUTABLE_CPUPROFILER_SCOPE(RunCode_Conditional);
 
-		const FProgram& Program = Model->GetPrivate()->Program;
+		const FProgram& Program = Model->Program;
 
 		EOpType type = Program.GetOpType(Item.At);
 		OP::ConditionalArgs Args = Program.GetOpArgs<OP::ConditionalArgs>(Item.At);
@@ -374,7 +373,7 @@ namespace UE::Mutable::Private
 	//---------------------------------------------------------------------------------------------
 	void CodeRunner::RunCode_Switch(const FScheduledOp& Item, const FModel* InModel )
 	{
-		const FProgram& Program = Model->GetPrivate()->Program;
+		const FProgram& Program = Model->Program;
 
 		EOpType Type = Program.GetOpType(Item.At);
 
@@ -522,7 +521,7 @@ namespace UE::Mutable::Private
     {
         MUTABLE_CPUPROFILER_SCOPE(RunCode_Instance);
 
-		const FProgram& Program = Model->GetPrivate()->Program;
+		const FProgram& Program = Model->Program;
 
 		EOpType type = Program.GetOpType(Item.At);
         switch (type)
@@ -834,7 +833,7 @@ namespace UE::Mutable::Private
 			return;
 		}
 
-		const FProgram& Program = Model->GetPrivate()->Program;
+		const FProgram& Program = Model->Program;
 
 		EOpType type = Program.GetOpType(Item.At);
         switch (type)
@@ -1073,7 +1072,7 @@ namespace UE::Mutable::Private
     {
 		MUTABLE_CPUPROFILER_SCOPE(RunCode_Constant);
 
-		const FProgram& Program = Model->GetPrivate()->Program;
+		const FProgram& Program = Model->Program;
 
 		EOpType type = Program.GetOpType(Item.At);
         switch (type)
@@ -1169,7 +1168,7 @@ namespace UE::Mutable::Private
 
 		check((EMeshContentFlags)Item.ExecutionOptions != EMeshContentFlags::None); 
 
-		const FProgram& Program = Model->GetPrivate()->Program;
+		const FProgram& Program = Model->Program;
 
 		EOpType type = Program.GetOpType(Item.At);
 
@@ -2071,7 +2070,7 @@ namespace UE::Mutable::Private
 				{
 					TManagedPtr<const FMesh> Source = LoadMesh(FCacheAddress(Args.Source, Item));
 					
-					check(Args.MorphShape < (uint32)InModel->GetPrivate()->Program.ConstantShapes.Num());
+					check(Args.MorphShape < (uint32)InModel->Program.ConstantShapes.Num());
 
 					// Should be an ellipse
 					const FShape& MorphShape = Program.ConstantShapes[Args.MorphShape];
@@ -2083,7 +2082,7 @@ namespace UE::Mutable::Private
 
 					if (Args.VertexSelectionType == EClipVertexSelectionType::Shape)
 					{
-						check(Args.VertexSelectionShapeOrBone < (uint32)InModel->GetPrivate()->Program.ConstantShapes.Num());
+						check(Args.VertexSelectionShapeOrBone < (uint32)InModel->Program.ConstantShapes.Num());
 
 						// Should be None or an axis aligned box
 						const FShape& SelectionShape = Program.ConstantShapes[Args.VertexSelectionShapeOrBone];
@@ -3326,7 +3325,7 @@ namespace UE::Mutable::Private
 
 		FImageOperator ImOp = MakeImageOperator(this);
 
-		const FProgram& Program = Model->GetPrivate()->Program;
+		const FProgram& Program = Model->Program;
 
 		EOpType type = Program.GetOpType(Item.At);
 		switch (type)
@@ -3432,7 +3431,7 @@ namespace UE::Mutable::Private
 
                     // We support both integers and scalars here, which is not common.
                     // \todo: review if this is necessary or we can enforce it at compile time.
-                    EDataType RangeSizeType = GetOpDataType( InModel->GetPrivate()->Program.GetOpType(Args.rangeSize) );
+                    EDataType RangeSizeType = GetOpDataType( InModel->Program.GetOpType(Args.rangeSize) );
                     if (RangeSizeType == EDataType::Int)
                     {
 						Iterations = LoadInt(RangeAddress);
@@ -5738,7 +5737,7 @@ namespace UE::Mutable::Private
             return nullptr;
         }
 
-		const FProgram& Program = InModel.GetPrivate()->Program;
+		const FProgram& Program = InModel.Program;
 		const FParameterDesc& paramDesc = Program.Parameters[ ParameterIndex ];
         for( int32 rangeIndexInParam = 0;
              rangeIndexInParam<paramDesc.Ranges.Num();
@@ -5763,7 +5762,7 @@ namespace UE::Mutable::Private
     {
         MUTABLE_CPUPROFILER_SCOPE(RunCode_Bool);
 
-		const FProgram& Program = Model->GetPrivate()->Program;
+		const FProgram& Program = Model->Program;
 		EOpType type = Program.GetOpType(Item.At);
         switch (type)
         {
@@ -5976,7 +5975,7 @@ namespace UE::Mutable::Private
     {
         MUTABLE_CPUPROFILER_SCOPE(RunCode_Int);
 
-		const FProgram& Program = Model->GetPrivate()->Program;
+		const FProgram& Program = Model->Program;
 
 		EOpType type = Program.GetOpType(Item.At);
         switch (type)
@@ -6031,7 +6030,7 @@ namespace UE::Mutable::Private
     {
         MUTABLE_CPUPROFILER_SCOPE(RunCode_Scalar);
 
-		const FProgram& Program = Model->GetPrivate()->Program;
+		const FProgram& Program = Model->Program;
 
 		EOpType type = Program.GetOpType(Item.At);
         switch (type)
@@ -6209,7 +6208,7 @@ namespace UE::Mutable::Private
     {
         MUTABLE_CPUPROFILER_SCOPE(RunCode_String );
 
-		const FProgram& Program = Model->GetPrivate()->Program;
+		const FProgram& Program = Model->Program;
 
 		EOpType type = Program.GetOpType( Item.At );
         switch ( type )
@@ -6218,7 +6217,7 @@ namespace UE::Mutable::Private
         case EOpType::ST_CONSTANT:
         {
 			OP::ResourceConstantArgs Args = Program.GetOpArgs<OP::ResourceConstantArgs>( Item.At );
-            check( Args.value < (uint32)InModel->GetPrivate()->Program.ConstantStrings.Num() );
+            check( Args.value < (uint32)InModel->Program.ConstantStrings.Num() );
 
             const FString& Result = Program.ConstantStrings[Args.value];
             StoreString( Item, MakeManaged<String>(Result) );
@@ -6248,7 +6247,7 @@ namespace UE::Mutable::Private
     {
 		MUTABLE_CPUPROFILER_SCOPE(RunCode_Colour);
 
-		const FProgram& Program = Model->GetPrivate()->Program;
+		const FProgram& Program = Model->Program;
 
 		EOpType type = Program.GetOpType(Item.At);
 
@@ -6559,7 +6558,7 @@ namespace UE::Mutable::Private
     {
         MUTABLE_CPUPROFILER_SCOPE(RunCode_Projector);
 
-		const FProgram& Program = Model->GetPrivate()->Program;
+		const FProgram& Program = Model->Program;
 
 		EOpType type = Program.GetOpType(Item.At);
         switch (type)
@@ -6577,7 +6576,7 @@ namespace UE::Mutable::Private
         {
 			OP::ParameterArgs Args = Program.GetOpArgs<OP::ParameterArgs>(Item.At);
 			TSharedPtr<FRangeIndex> Index = BuildCurrentOpRangeIndex( Item, Args.variable );
-            FProjector Result = pParams->GetPrivate()->GetProjectorValue(Args.variable,Index.Get());
+            FProjector Result = pParams->GetProjectorValue(Args.variable,Index.Get());
 
             // The type cannot be changed, take it from the default value
             const FProjector& def = Program.Parameters[Args.variable].DefaultValue.Get<FParamProjectorType>();
@@ -6598,7 +6597,7 @@ namespace UE::Mutable::Private
     {
 		MUTABLE_CPUPROFILER_SCOPE(RunCode_Transform);
 
-		const FProgram& Program = Model->GetPrivate()->Program;
+		const FProgram& Program = Model->Program;
 
 		EOpType type = Program.GetOpType(Item.At);
 
@@ -6628,7 +6627,7 @@ namespace UE::Mutable::Private
 	{
 		MUTABLE_CPUPROFILER_SCOPE(RunCode_Material);
 
-		const FProgram& Program = Model->GetPrivate()->Program;
+		const FProgram& Program = Model->Program;
 
 		EOpType type = Program.GetOpType(Item.At);
 
@@ -6937,7 +6936,7 @@ namespace UE::Mutable::Private
     {
         //MUTABLE_CPUPROFILER_SCOPE(RunCode_Layout);
 
-		const FProgram& Program = Model->GetPrivate()->Program;
+		const FProgram& Program = Model->Program;
 
 		EOpType type = Program.GetOpType(Item.At);
         switch (type)
@@ -6946,7 +6945,7 @@ namespace UE::Mutable::Private
         case EOpType::LA_CONSTANT:
         {
 			OP::ResourceConstantArgs Args = Program.GetOpArgs<OP::ResourceConstantArgs>(Item.At);
-            check( Args.value < (uint32)InModel->GetPrivate()->Program.ConstantLayouts.Num() );
+            check( Args.value < (uint32)InModel->Program.ConstantLayouts.Num() );
 
             TManagedPtr<const FLayout> Result = Program.ConstantLayouts[ Args.value ];
             StoreLayout(Item, Result);
@@ -7110,7 +7109,7 @@ namespace UE::Mutable::Private
 		//UE_LOG( LogMutableCore, Log, TEXT("Running :%5d , %d "), Item.At, Item.Stage );
 		check( Item.Type == FScheduledOp::EType::Full );
 		
-		const FProgram& Program = InModel->GetPrivate()->Program;
+		const FProgram& Program = InModel->Program;
 
 		EOpType type = Program.GetOpType(Item.At);
 		//UE_LOG(LogMutableCore, Log, TEXT("Running :%5d , %d, of type %d "), Item.At, Item.Stage, type);
@@ -7459,7 +7458,7 @@ namespace UE::Mutable::Private
 		MUTABLE_CPUPROFILER_SCOPE(RunCodeImageDesc);
 		check(Item.Type == FScheduledOp::EType::ImageDesc);
 
-		const FProgram& Program = Model->GetPrivate()->Program;
+		const FProgram& Program = Model->Program;
 
 		EOpType Type = Program.GetOpType(Item.At);
 		switch (Type)
diff --git a/Engine/Plugins/Mutable/Source/MutableRuntime/Private/MuR/CodeRunner.h b/Engine/Plugins/Mutable/Source/MutableRuntime/Private/MuR/CodeRunner.h
--- a/Engine/Plugins/Mutable/Source/MutableRuntime/Private/MuR/CodeRunner.h
+++ b/Engine/Plugins/Mutable/Source/MutableRuntime/Private/MuR/CodeRunner.h
@@ -14,7 +14,7 @@
 #include "MuR/Serialisation.h"
 #include "MuR/Settings.h"
 #include "MuR/System.h"
-#include "MuR/SystemPrivate.h"
+#include "MuR/System.h"
 #include "MuR/Types.h"
 #include "Tasks/Task.h"
 #include "Templates/RefCounting.h"
@@ -46,7 +46,7 @@ namespace  UE::Mutable::Private
 	public:
 		static TSharedRef<CodeRunner> Create(
 				const FSettings&, 
-				FSystem::Private*, 
+				FSystem*, 
 				EExecutionStrategy,
 				const TSharedPtr<const FModel>&, 
 				const FParameters* Params,
@@ -61,7 +61,7 @@ namespace  UE::Mutable::Private
 		// bounded.
 		explicit CodeRunner(FPrivateToken, 
 			const FSettings&, 
-			FSystem::Private*, 
+			FSystem*, 
 			EExecutionStrategy,
 			const TSharedPtr<const FModel>&, 
 			const FParameters* Params,
@@ -589,7 +589,7 @@ namespace  UE::Mutable::Private
 		*/
 		bool bUnrecoverableError = false;
 
-		FSystem::Private* System = nullptr;
+		FSystem* System = nullptr;
 		TSharedPtr<const FModel> Model = nullptr;
 		const FParameters* Params = nullptr;
 		uint32 LODMask = 0;
diff --git a/Engine/Plugins/Mutable/Source/MutableRuntime/Private/MuR/CodeRunnerBegin.cpp b/Engine/Plugins/Mutable/Source/MutableRuntime/Private/MuR/CodeRunnerBegin.cpp
--- a/Engine/Plugins/Mutable/Source/MutableRuntime/Private/MuR/CodeRunnerBegin.cpp
+++ b/Engine/Plugins/Mutable/Source/MutableRuntime/Private/MuR/CodeRunnerBegin.cpp
@@ -4,8 +4,9 @@
 
 #include "CodeRunner.h"
 #include "MuR/MutableTrace.h"
+#include "MuR/Model.h"
 #include "MuR/Operations.h"
-#include "MuR/SystemPrivate.h"
+#include "MuR/System.h"
 #include "MuR/External/Operation.h"
 
 
@@ -24,9 +25,9 @@ namespace UE::Mutable::Private
 			uint32 InLodMask) :
 		Model(InModel),
 		Params(InParams),
-		Program(InModel->GetPrivate()->Program),
+		Program(InModel->Program),
 		LodMask(InLodMask),
-		Executed(InModel->GetPrivate()->Program.OpAddress.Num()),
+		Executed(InModel->Program.OpAddress.Num()),
 		ProgramCache(InProgramCache),
 		MeshIdRegistry(InMeshIdRegistry),
 		ImageIdRegistry(InImageIdRegistry),
@@ -1769,7 +1770,7 @@ namespace UE::Mutable::Private
 					{
 						OP::ResourceConstantArgs Args = Program.GetOpArgs<OP::ResourceConstantArgs>(Item.At);
 				
-						check(Args.value < static_cast<uint32>(Model->GetPrivate()->Program.ConstantStrings.Num()));
+						check(Args.value < static_cast<uint32>(Model->Program.ConstantStrings.Num()));
 
 						const FString& Result = Program.ConstantStrings[Args.value];
 
@@ -2378,7 +2379,7 @@ namespace UE::Mutable::Private
 							int32 NumIterations = 0;
 							if (Args.rangeSize)
 							{
-								const EDataType RangeSizeType = GetOpDataType(Model->GetPrivate()->Program.GetOpType(Args.rangeSize) );
+								const EDataType RangeSizeType = GetOpDataType(Model->Program.GetOpType(Args.rangeSize) );
 								if (RangeSizeType == EDataType::Int)
 								{
 									NumIterations = LoadInt(Args.rangeSize);
diff --git a/Engine/Plugins/Mutable/Source/MutableRuntime/Private/MuR/CodeRunnerBegin.h b/Engine/Plugins/Mutable/Source/MutableRuntime/Private/MuR/CodeRunnerBegin.h
--- a/Engine/Plugins/Mutable/Source/MutableRuntime/Private/MuR/CodeRunnerBegin.h
+++ b/Engine/Plugins/Mutable/Source/MutableRuntime/Private/MuR/CodeRunnerBegin.h
@@ -2,7 +2,7 @@
 
 #pragma once
 
-#include "MuR/SystemPrivate.h"
+#include "MuR/System.h"
 
 
 namespace UE::Mutable::Private
@@ -161,7 +161,7 @@ namespace UE::Mutable::Private
 		TSharedPtr<const FModel> Model;
 		TSharedPtr<const FParameters> Params;
 
-		FProgram& Program;
+		const FProgram& Program;
 		
 		uint32 LodMask;
 
diff --git a/Engine/Plugins/Mutable/Source/MutableRuntime/Private/MuR/CodeRunnerTasked.cpp b/Engine/Plugins/Mutable/Source/MutableRuntime/Private/MuR/CodeRunnerTasked.cpp
--- a/Engine/Plugins/Mutable/Source/MutableRuntime/Private/MuR/CodeRunnerTasked.cpp
+++ b/Engine/Plugins/Mutable/Source/MutableRuntime/Private/MuR/CodeRunnerTasked.cpp
@@ -18,7 +18,6 @@
 #include "MuR/Layout.h"
 #include "MuR/Mesh.h"
 #include "MuR/Model.h"
-#include "MuR/ModelPrivate.h"
 #include "MuR/MutableMath.h"
 #include "MuR/MutableTrace.h"
 #include "MuR/OpImageBlend.h"
@@ -30,9 +29,9 @@
 #include "MuR/Ptr.h"
 #include "MuR/RefCounted.h"
 #include "MuR/Serialisation.h"
-#include "MuR/SerialisationPrivate.h"
+#include "MuR/Serialisation.h"
 #include "MuR/Settings.h"
-#include "MuR/SystemPrivate.h"
+#include "MuR/System.h"
 #include "MuR/MutableRuntimeModule.h"
 #include "Stats/Stats.h"
 #include "Templates/RefCounting.h"
@@ -322,7 +321,7 @@ namespace UE::Mutable::Private
 				if (ProfileContext)
 				{
 					++ProfileContext->NumRunOps;
-					++ProfileContext->RunOpsPerType[int32(Model->GetPrivate()->Program.GetOpType(Item.At))];
+					++ProfileContext->RunOpsPerType[int32(Model->Program.GetOpType(Item.At))];
 				}
 			}
 
@@ -1970,7 +1969,7 @@ namespace UE::Mutable::Private
 
 		bOutFailed = false;
 
-		const FProgram& Program = Runner->Model->GetPrivate()->Program;
+		const FProgram& Program = Runner->Model->Program;
 		FWorkingMemoryManager::FModelCacheEntry* ModelCache = Runner->System->WorkingMemoryManager.FindModelCache(Runner->Model.Get());
 		
 		TArray<UE::Tasks::FTask, TInlineAllocator<4>> ReadCompleteEvents;
@@ -2098,7 +2097,7 @@ namespace UE::Mutable::Private
 	}	
 
 	//---------------------------------------------------------------------------------------------
-		bool CodeRunner::FLoadMeshRomTask::Complete(CodeRunner* Runner)
+	bool CodeRunner::FLoadMeshRomTask::Complete(CodeRunner* Runner)
 	{
 		// This runs in the Runner thread
 		MUTABLE_CPUPROFILER_SCOPE(FLoadMeshRomTask_Complete);
@@ -2108,7 +2107,7 @@ namespace UE::Mutable::Private
 			return false;
 		}
 
-		FProgram& Program = Runner->Model->GetPrivate()->Program;
+		const FProgram& Program = Runner->Model->Program;
 		FWorkingMemoryManager::FModelCacheEntry* ModelCache = Runner->System->WorkingMemoryManager.FindModelCache(Runner->Model.Get());
 		
 		bool bSomeMissingData = false;
@@ -2181,7 +2180,7 @@ namespace UE::Mutable::Private
 		MUTABLE_CPUPROFILER_SCOPE(FLoadImageRomsTask_Prepare);
 		bOutFailed = false;
 
-		FProgram& Program = Runner->Model->GetPrivate()->Program;
+		const FProgram& Program = Runner->Model->Program;
 
 		FWorkingMemoryManager::FModelCacheEntry* ModelCache = Runner->System->WorkingMemoryManager.FindModelCache(Runner->Model.Get());
 
@@ -2277,7 +2276,7 @@ namespace UE::Mutable::Private
 			return false;
 		}
 
-		FProgram& Program = Runner->Model->GetPrivate()->Program;
+		const FProgram& Program = Runner->Model->Program;
 		
 		FWorkingMemoryManager::FModelCacheEntry* ModelCache = Runner->System->WorkingMemoryManager.FindModelCache(Runner->Model.Get());
 
@@ -2497,7 +2496,7 @@ namespace UE::Mutable::Private
 	{
 		TSharedPtr<FIssuedTask> Issued;
 
-		FProgram& Program = Model->GetPrivate()->Program;
+		const FProgram& Program = Model->Program;
 
 		EOpType type = Program.GetOpType(Item.At);
 
@@ -2614,7 +2613,7 @@ namespace UE::Mutable::Private
 
 		case EOpType::IM_REFERENCE:
 		{
-			OP::ResourceReferenceArgs Args = Model->GetPrivate()->Program.GetOpArgs<OP::ResourceReferenceArgs>(Item.At);
+			OP::ResourceReferenceArgs Args = Model->Program.GetOpArgs<OP::ResourceReferenceArgs>(Item.At);
 
 			// We only convert references to images if indicated in the operation.
 			if (Args.ForceLoad)
diff --git a/Engine/Plugins/Mutable/Source/MutableRuntime/Private/MuR/Model.cpp b/Engine/Plugins/Mutable/Source/MutableRuntime/Private/MuR/Model.cpp
--- a/Engine/Plugins/Mutable/Source/MutableRuntime/Private/MuR/Model.cpp
+++ b/Engine/Plugins/Mutable/Source/MutableRuntime/Private/MuR/Model.cpp
@@ -10,12 +10,11 @@
 #include "Misc/AssertionMacros.h"
 #include "MuR/Image.h"
 #include "MuR/Mesh.h"
-#include "MuR/ModelPrivate.h"
+#include "MuR/Model.h"
 #include "MuR/Operations.h"
-#include "MuR/ParametersPrivate.h"
+#include "MuR/Parameters.h"
+#include "MuR/Serialisation.h"
 #include "MuR/Serialisation.h"
-#include "MuR/SerialisationPrivate.h"
-#include "MuR/System.h"
 #include "MuR/Types.h"
 #include "Templates/Tuple.h"
 
@@ -31,19 +30,8 @@ namespace UE::Mutable::Private
 	MUTABLE_IMPLEMENT_POD_SERIALISABLE(FImageLODRange);
 	MUTABLE_IMPLEMENT_POD_SERIALISABLE(FMeshContentRange);
 	MUTABLE_IMPLEMENT_POD_SERIALISABLE(FConstantResourceIndex);
-	//MUTABLE_IMPLEMENT_POD_VECTOR_SERIALISABLE(FConstantResourceIndex);
 
 	
-    //---------------------------------------------------------------------------------------------
-    void FProgram::Check()
-    {
-    #ifdef MUTABLE_DEBUG
-		// Insert debug checks here.
-    #endif
-    }
-
-
-    //---------------------------------------------------------------------------------------------
     void FProgram::LogHistogram() const
     {
 #if 0
@@ -78,47 +66,23 @@ namespace UE::Mutable::Private
     }
 
 	
-	//---------------------------------------------------------------------------------------------
-    void FModel::Private::UnloadRoms()
+    void FModel::UnloadRoms()
     {
 	    for (int32 RomIndex = 0; RomIndex < Program.Roms.Num(); ++RomIndex)
 	    {
 		    Program.UnloadRom(RomIndex);
 	    }
     }
-	
-
-    //---------------------------------------------------------------------------------------------
-    //---------------------------------------------------------------------------------------------
-    //---------------------------------------------------------------------------------------------
-    FModel::FModel()
-    {
-        m_pD = new Private();
-    }
-
-
-    //---------------------------------------------------------------------------------------------
-    FModel::~FModel()
-    {
-		MUTABLE_CPUPROFILER_SCOPE(ModelDestructor);
-
-        check( m_pD );
-        delete m_pD;
-        m_pD = 0;
-    }
 
 
-    //---------------------------------------------------------------------------------------------
     void FModel::Serialise( const FModel* p, FOutputArchive& arch )
     {
 		LLM_SCOPE_BYNAME(TEXT("MutableRuntime"));
 
-        arch << *p->m_pD;
+        arch << *p;
     }
 
-    //---------------------------------------------------------------------------------------------
-    //---------------------------------------------------------------------------------------------
-    //---------------------------------------------------------------------------------------------
+   
     class FOutputModelStream : public FOutputStream
     {
     public:
@@ -141,12 +105,11 @@ namespace UE::Mutable::Private
     };
 
 
-    //---------------------------------------------------------------------------------------------
     void FModel::Serialise( FModel* p, FModelWriter& streamer, bool bDropData )
     {
 		LLM_SCOPE_BYNAME(TEXT("MutableRuntime"));
 
-		UE::Mutable::Private::FProgram& Program = p->m_pD->Program;
+		UE::Mutable::Private::FProgram& Program = p->Program;
 
 		FOutputMemoryStream MemStream(16 * 1024 * 1024);
 
@@ -214,7 +177,7 @@ namespace UE::Mutable::Private
 			FOutputModelStream stream(&streamer);
 			FOutputArchive arch(&stream);
 
-			arch << *p->m_pD;
+			arch << *p;
 
 			streamer.CloseWriteFile();
 		}
@@ -223,19 +186,20 @@ namespace UE::Mutable::Private
 
 	bool FModel::HasExternalData() const
     {
-		return m_pD->Program.Roms.Num() > 0;
+		return Program.Roms.Num() > 0;
     }
 
+	
 #if WITH_EDITOR
 	bool FModel::IsValid() const
 	{
-		return m_pD->Program.bIsValid;
+		return Program.bIsValid;
 	}
 
 
 	void FModel::Invalidate()
     {
-		m_pD->Program.bIsValid = false;
+		Program.bIsValid = false;
     }
 #endif
 
@@ -244,8 +208,8 @@ namespace UE::Mutable::Private
     {
 		LLM_SCOPE_BYNAME(TEXT("MutableRuntime"));
 
-		m_pD->Program.ConstantImageLODsStreamed.Empty(0);
-		m_pD->Program.ConstantMeshesStreamed.Empty(0);
+		Program.ConstantImageLODsStreamed.Empty(0);
+		Program.ConstantMeshesStreamed.Empty(0);
 	}
 
 
@@ -254,78 +218,72 @@ namespace UE::Mutable::Private
 		LLM_SCOPE_BYNAME(TEXT("MutableRuntime"));
 
 		TSharedPtr<FModel> pResult = MakeShared<FModel>();
-        arch >> *pResult->m_pD;
+        arch >> *pResult;
         return pResult;
     }
-
-
-	FModel::Private* FModel::GetPrivate() const
-    {
-        return m_pD;
-    }
-
-
+	
+	
     bool FModel::GetBoolDefaultValue(int32 Index) const
     {
-    	check(m_pD->Program.Parameters.IsValidIndex(Index));
-		check(m_pD->Program.Parameters[Index].Type == EParameterType::Bool);
+    	check(Program.Parameters.IsValidIndex(Index));
+		check(Program.Parameters[Index].Type == EParameterType::Bool);
 
         // Early out in case of invalid parameters
-        if (!m_pD->Program.Parameters.IsValidIndex(Index) ||
-            m_pD->Program.Parameters[Index].Type != EParameterType::Bool)
+        if (!Program.Parameters.IsValidIndex(Index) ||
+            Program.Parameters[Index].Type != EParameterType::Bool)
         {
             return false;
         }
 		
-        return m_pD->Program.Parameters[Index].DefaultValue.Get<FParamBoolType>();
+        return Program.Parameters[Index].DefaultValue.Get<FParamBoolType>();
     }
 
 
 	int32 FModel::GetIntDefaultValue(int32 Index) const
 	{
-		check(m_pD->Program.Parameters.IsValidIndex(Index));
-		check(m_pD->Program.Parameters[Index].Type == EParameterType::Int);
+		check(Program.Parameters.IsValidIndex(Index));
+		check(Program.Parameters[Index].Type == EParameterType::Int);
 
         // Early out in case of invalid parameters
-        if (!m_pD->Program.Parameters.IsValidIndex(Index) ||
-            m_pD->Program.Parameters[Index].Type != EParameterType::Int)
+        if (!Program.Parameters.IsValidIndex(Index) ||
+            Program.Parameters[Index].Type != EParameterType::Int)
         {
             return 0;
         }
 		
-        return m_pD->Program.Parameters[Index].DefaultValue.Get<FParamIntType>();
+        return Program.Parameters[Index].DefaultValue.Get<FParamIntType>();
 	}
 
 
 	float FModel::GetFloatDefaultValue(int32 Index) const
 	{
-    	check(m_pD->Program.Parameters.IsValidIndex(Index));
-		check(m_pD->Program.Parameters[Index].Type == EParameterType::Float);
+    	check(Program.Parameters.IsValidIndex(Index));
+		check(Program.Parameters[Index].Type == EParameterType::Float);
 
         // Early out in case of invalid parameters
-        if (!m_pD->Program.Parameters.IsValidIndex(Index) ||
-            m_pD->Program.Parameters[Index].Type != EParameterType::Float)
+        if (!Program.Parameters.IsValidIndex(Index) ||
+            Program.Parameters[Index].Type != EParameterType::Float)
         {
             return 0.0f;
         }
 		
-        return m_pD->Program.Parameters[Index].DefaultValue.Get<FParamFloatType>();
+        return Program.Parameters[Index].DefaultValue.Get<FParamFloatType>();
 	}
 
 
 	void FModel::GetColourDefaultValue(int32 Index, FVector4f& OutValue) const
     {
-    	check(m_pD->Program.Parameters.IsValidIndex(Index));
-		check(m_pD->Program.Parameters[Index].Type == EParameterType::Color);
+    	check(Program.Parameters.IsValidIndex(Index));
+		check(Program.Parameters[Index].Type == EParameterType::Color);
 
         // Early out in case of invalid parameters
-        if (!m_pD->Program.Parameters.IsValidIndex(Index) ||
-            m_pD->Program.Parameters[Index].Type != EParameterType::Color)
+        if (!Program.Parameters.IsValidIndex(Index) ||
+            Program.Parameters[Index].Type != EParameterType::Color)
         {
             return;
         }
 
-        FParamColorType& Color = m_pD->Program.Parameters[Index].DefaultValue.Get<FParamColorType>();
+        const FParamColorType& Color = Program.Parameters[Index].DefaultValue.Get<FParamColorType>();
 		
 		OutValue = Color;
     }
@@ -333,34 +291,34 @@ namespace UE::Mutable::Private
 
 	FMatrix44f FModel::GetMatrixDefaultValue(int32 Index) const
 	{
-    	check(m_pD->Program.Parameters.IsValidIndex(Index));
-    	check(m_pD->Program.Parameters[Index].Type == EParameterType::Matrix);
+    	check(Program.Parameters.IsValidIndex(Index));
+    	check(Program.Parameters[Index].Type == EParameterType::Matrix);
 
     	// Early out in case of invalid parameters
-    	if (!m_pD->Program.Parameters.IsValidIndex(Index) ||
-			m_pD->Program.Parameters[Index].Type != EParameterType::Matrix)
+    	if (!Program.Parameters.IsValidIndex(Index) ||
+			Program.Parameters[Index].Type != EParameterType::Matrix)
     	{
     		return FMatrix44f::Identity;
     	}
 
-    	return m_pD->Program.Parameters[Index].DefaultValue.Get<FParamMatrixType>();
+    	return Program.Parameters[Index].DefaultValue.Get<FParamMatrixType>();
 	}
 
 
 	void FModel::GetProjectorDefaultValue(int32 Index, EProjectorType* OutProjectionType, FVector3f* OutPos,
 	                                     FVector3f* OutDir, FVector3f* OutUp, FVector3f* OutScale, float* OutProjectionAngle) const
 	{
-    	check(m_pD->Program.Parameters.IsValidIndex(Index));
-		check(m_pD->Program.Parameters[Index].Type == EParameterType::Projector);
+    	check(Program.Parameters.IsValidIndex(Index));
+		check(Program.Parameters[Index].Type == EParameterType::Projector);
 
         // Early out in case of invalid parameters
-        if (!m_pD->Program.Parameters.IsValidIndex(Index) ||
-            m_pD->Program.Parameters[Index].Type != EParameterType::Projector)
+        if (!Program.Parameters.IsValidIndex(Index) ||
+            Program.Parameters[Index].Type != EParameterType::Projector)
         {
             return;
         }
 
-        const FProjector& Projector = m_pD->Program.Parameters[Index].DefaultValue.Get<FParamProjectorType>();
+        const FProjector& Projector = Program.Parameters[Index].DefaultValue.Get<FParamProjectorType>();
         if (OutProjectionType) *OutProjectionType = Projector.type;
     	if (OutPos) *OutPos = Projector.position;
 		if (OutDir) *OutDir = Projector.direction;
@@ -372,34 +330,38 @@ namespace UE::Mutable::Private
 	
     int32 FModel::GetRomCount() const
     {
-    	return m_pD->Program.Roms.Num();
+    	return Program.Roms.Num();
     }
 
+	
 #if WITH_EDITOR
 	uint32 FModel::GetRomSourceId(int32 Index) const
 	{
-		return m_pD->Program.RomsCompileData[Index].SourceId;
+		return Program.RomsCompileData[Index].SourceId;
 	}
 #endif
 
+	
     uint32 FModel::GetRomSize(int32 Index) const
     {
-    	return m_pD->Program.Roms[Index].Size;
+    	return Program.Roms[Index].Size;
     }
 
+	
 	bool FModel::IsMeshData(int32 Index) const
 	{
-		return m_pD->Program.Roms[Index].ResourceType==uint32(ERomDataType::Mesh);
+		return Program.Roms[Index].ResourceType==uint32(ERomDataType::Mesh);
 	}
 
+	
 	bool FModel::IsRomHighRes(int32 Index) const
 	{
-		return m_pD->Program.Roms[Index].IsHighRes==1;
+		return Program.Roms[Index].IsHighRes==1;
 	}
 
+	
     int32 FModel::GetConstantImageRomId(int32 ConstantImageIndex, int32 LODIndex) const
     {
-        const FProgram& Program = m_pD->Program;
         check(Program.ConstantImages.IsValidIndex(ConstantImageIndex));
         
         FImageLODRange LODRange = Program.ConstantImages[ConstantImageIndex];
@@ -425,13 +387,13 @@ namespace UE::Mutable::Private
 
         TSharedPtr<FParameters> pRes = MakeShared<FParameters>();
 
-        pRes->GetPrivate()->Model = Model;
+        pRes->Model = Model;
 
-		const FProgram& Program = Model->GetPrivate()->Program;
-        pRes->GetPrivate()->Values.SetNum(Program.Parameters.Num());
+		const FProgram& Program = Model->Program;
+        pRes->Values.SetNum(Program.Parameters.Num());
         for ( int32 p=0; p< Program.Parameters.Num(); ++p )
         {
-            pRes->GetPrivate()->Values[p] = Program.Parameters[p].DefaultValue;
+            pRes->Values[p] = Program.Parameters[p].DefaultValue;
         }
 
         // Copy old values
@@ -439,7 +401,7 @@ namespace UE::Mutable::Private
         {
             for ( int32 p=0; p<OldParameters->GetCount(); ++p )
             {
-                int32 thisP = pRes->GetPrivate()->Find( OldParameters->GetName(p) );
+                int32 thisP = pRes->Find( OldParameters->GetName(p) );
 
                 if ( thisP>=0 )
                 {
@@ -471,7 +433,7 @@ namespace UE::Mutable::Private
                         {
 //							float m[16];
 //							OldParameters->GetProjectorValue( p, m );
-                            pRes->GetPrivate()->Values[thisP].Set<FParamProjectorType>(OldParameters->GetPrivate()->Values[p].Get<FParamProjectorType>());
+                            pRes->Values[thisP].Set<FParamProjectorType>(OldParameters->Values[p].Get<FParamProjectorType>());
                             break;
                         }
 
@@ -504,33 +466,30 @@ namespace UE::Mutable::Private
     }
 
 
-    //---------------------------------------------------------------------------------------------
     bool FModel::IsParameterMultidimensional(const int32 ParamIndex) const
     {
-		if (m_pD->Program.Parameters.IsValidIndex(ParamIndex))
+		if (Program.Parameters.IsValidIndex(ParamIndex))
 		{
-			return m_pD->Program.Parameters[ParamIndex].Ranges.Num() > 0;
+			return Program.Parameters[ParamIndex].Ranges.Num() > 0;
 		}
 
 		return false;
     }
 	
 
-    //---------------------------------------------------------------------------------------------
     int32 FModel::GetStateCount() const
     {
-        return (int32)m_pD->Program.States.Num();
+        return (int32)Program.States.Num();
     }
 
 
-    //---------------------------------------------------------------------------------------------
     const FString& FModel::GetStateName( int32 index ) const
     {
         const char* strRes = 0;
 
-        if ( index>=0 && index<(int32)m_pD->Program.States.Num() )
+        if ( index>=0 && index<(int32)Program.States.Num() )
         {
-            return m_pD->Program.States[index].Name;
+            return Program.States[index].Name;
         }
 
 		static FString None;
@@ -538,14 +497,13 @@ namespace UE::Mutable::Private
     }
 
 
-    //---------------------------------------------------------------------------------------------
     int32 FModel::FindState( const FString& Name ) const
     {
         int32 res = -1;
 
-        for ( int32 i=0; res<0 && i<(int32)m_pD->Program.States.Num(); ++i )
+        for ( int32 i=0; res<0 && i<(int32)Program.States.Num(); ++i )
         {
-            if ( m_pD->Program.States[i].Name == Name )
+            if ( Program.States[i].Name == Name )
             {
                 res = i;
             }
@@ -555,28 +513,26 @@ namespace UE::Mutable::Private
     }
 
 
-    //---------------------------------------------------------------------------------------------
     int32 FModel::GetStateParameterCount( int32 stateIndex ) const
     {
         int32 res = -1;
 
-        if ( stateIndex>=0 && stateIndex<(int32)m_pD->Program.States.Num() )
+        if ( stateIndex>=0 && stateIndex<(int32)Program.States.Num() )
         {
-            res = (int32)m_pD->Program.States[stateIndex].m_runtimeParameters.Num();
+            res = (int32)Program.States[stateIndex].m_runtimeParameters.Num();
         }
 
         return res;
     }
 
 
-    //---------------------------------------------------------------------------------------------
     int32 FModel::GetStateParameterIndex( int32 stateIndex, int32 paramIndex ) const
     {
         int32 res = -1;
 
-        if ( stateIndex>=0 && stateIndex<(int32)m_pD->Program.States.Num() )
+        if ( stateIndex>=0 && stateIndex<(int32)Program.States.Num() )
         {
-            const FProgram::FState& state = m_pD->Program.States[stateIndex];
+            const FProgram::FState& state = Program.States[stateIndex];
             if ( paramIndex>=0 && paramIndex<(int32)state.m_runtimeParameters.Num() )
             {
                 res = (int32)state.m_runtimeParameters[paramIndex];
@@ -585,6 +541,5 @@ namespace UE::Mutable::Private
 
         return res;
     }
-
 }
 
diff --git a/Engine/Plugins/Mutable/Source/MutableRuntime/Private/MuR/OpImageMipmap.cpp b/Engine/Plugins/Mutable/Source/MutableRuntime/Private/MuR/OpImageMipmap.cpp
--- a/Engine/Plugins/Mutable/Source/MutableRuntime/Private/MuR/OpImageMipmap.cpp
+++ b/Engine/Plugins/Mutable/Source/MutableRuntime/Private/MuR/OpImageMipmap.cpp
@@ -2,7 +2,7 @@
 
 #include "MuR/ImagePrivate.h"
 #include "MuR/ImageFormatUtils.h"
-#include "MuR/SystemPrivate.h"
+#include "MuR/System.h"
 #include "MuR/ParallelExecutionUtils.h"
 #include "HAL/UnrealMemory.h"
 
diff --git a/Engine/Plugins/Mutable/Source/MutableRuntime/Private/MuR/OpImageProject.cpp b/Engine/Plugins/Mutable/Source/MutableRuntime/Private/MuR/OpImageProject.cpp
--- a/Engine/Plugins/Mutable/Source/MutableRuntime/Private/MuR/OpImageProject.cpp
+++ b/Engine/Plugins/Mutable/Source/MutableRuntime/Private/MuR/OpImageProject.cpp
@@ -11,7 +11,6 @@
 #include "MuR/MutableTrace.h"
 #include "MuR/OpMeshClipWithMesh.h"
 #include "MuR/Parameters.h"
-#include "MuR/ParametersPrivate.h"
 #include "MuR/Ptr.h"
 #include "MuR/Raster.h"
 
