[Mutable] Remove Core Pimpl.

diff --git a/Engine/Plugins/Mutable/Source/MutableRuntime/Private/MuR/Parameters.cpp b/Engine/Plugins/Mutable/Source/MutableRuntime/Private/MuR/Parameters.cpp
--- a/Engine/Plugins/Mutable/Source/MutableRuntime/Private/MuR/Parameters.cpp
+++ b/Engine/Plugins/Mutable/Source/MutableRuntime/Private/MuR/Parameters.cpp
@@ -8,11 +8,10 @@
 #include "HAL/UnrealMemory.h"
 #include "Misc/AssertionMacros.h"
 #include "MuR/Model.h"
-#include "MuR/ModelPrivate.h"
 #include "MuR/Material.h"
-#include "MuR/ParametersPrivate.h"
+#include "MuR/Parameters.h"
+#include "MuR/Serialisation.h"
 #include "MuR/Serialisation.h"
-#include "MuR/SerialisationPrivate.h"
 #include "MuR/System.h"
 
 
@@ -20,25 +19,11 @@ namespace UE::Mutable::Private
 {
     MUTABLE_IMPLEMENT_ENUM_SERIALISABLE(EParameterType)              
     MUTABLE_IMPLEMENT_ENUM_SERIALISABLE(EProjectorType)
-
 	
-	FParameters::FParameters()
-	{
-		m_pD = new Private();
-	}
-
-
-	FParameters::~FParameters()
-	{
-        check( m_pD );
-		delete m_pD;
-		m_pD = 0;
-	}
-
 
 	void FParameters::Serialise( const FParameters* p, FOutputArchive& arch )
 	{
-		arch << *p->m_pD;
+		arch << *p;
     }
 
 
@@ -47,14 +32,22 @@ namespace UE::Mutable::Private
 		LLM_SCOPE_BYNAME(TEXT("MutableRuntime"));
 
 		TSharedPtr<FParameters> pResult = MakeShared<FParameters>();
-		arch >> *pResult->m_pD;
+		arch >> *pResult;
 		return pResult;
 	}
 
+	
+	void FParameters::Serialise(FOutputArchive& Arch) const
+	{
+    	Arch << Values;
+    	Arch << MultiValues;
+	}
 
-	FParameters::Private* FParameters::GetPrivate() const
+	
+	void FParameters::Unserialise(FInputArchive& Arch)
 	{
-		return m_pD;
+		Arch >> Values;
+		Arch >> MultiValues;
 	}
 
 
@@ -64,9 +57,9 @@ namespace UE::Mutable::Private
 
 		TSharedPtr<FParameters> pRes = MakeShared<FParameters>();
 
-		pRes->m_pD->Model = m_pD->Model;
-		pRes->m_pD->Values = m_pD->Values;
-		pRes->m_pD->MultiValues = m_pD->MultiValues;
+		pRes->Model = Model;
+		pRes->Values = Values;
+		pRes->MultiValues = MultiValues;
 
 		return pRes;
 	}
@@ -74,13 +67,13 @@ namespace UE::Mutable::Private
 
 	int32 FParameters::GetCount() const
 	{
-		return (int32)m_pD->Values.Num();
+		return (int32)Values.Num();
 	}
 
 
 	const FString& FParameters::GetName( int32 Index ) const
 	{
-		const FProgram& Program = m_pD->Model->GetPrivate()->Program;
+		const FProgram& Program = Model->Program;
 		check( Index>=0 && Index<(int32)Program.Parameters.Num() );
 
 		return Program.Parameters[Index].Name;
@@ -89,22 +82,36 @@ namespace UE::Mutable::Private
 
 	const FGuid& FParameters::GetUid( int32 Index ) const
 	{
-		const FProgram& Program = m_pD->Model->GetPrivate()->Program;
+		const FProgram& Program = Model->Program;
 		check( Index>=0 && Index<Program.Parameters.Num() );
 
 		return Program.Parameters[Index].UID;
 	}
 
 
-	int32 FParameters::Find( const FString& strName ) const
+	int32 FParameters::Find( const FString& Name ) const
 	{
-		return m_pD->Find( strName );
+    	const FProgram& Program = Model->Program;
+
+    	int32 result = -1;
+
+    	for( int32 i=0; result<0 && i<(int32)Program.Parameters.Num(); ++i )
+    	{
+    		const FParameterDesc& p = Program.Parameters[i];
+
+    		if (p.Name == Name)
+    		{
+    			result = i;
+    		}
+    	}
+
+    	return result;
 	}
 
 
 	EParameterType FParameters::GetType( int32 Index ) const
 	{
-		const FProgram& Program = m_pD->Model->GetPrivate()->Program;
+		const FProgram& Program = Model->Program;
 		check( Index>=0 && Index<Program.Parameters.Num() );
 
 		return Program.Parameters[Index].Type;
@@ -115,7 +122,7 @@ namespace UE::Mutable::Private
     {
 		LLM_SCOPE_BYNAME(TEXT("MutableRuntime"));
 
-        const FProgram& Program = m_pD->Model->GetPrivate()->Program;
+        const FProgram& Program = Model->Program;
         check( ParamIndex>=0 && ParamIndex<Program.Parameters.Num() );
 		TSharedPtr<FRangeIndex> Range;
 
@@ -137,28 +144,28 @@ namespace UE::Mutable::Private
 
     int32 FParameters::GetValueCount( int32 ParamIndex ) const
     {
-        if (ParamIndex<0 || ParamIndex >= int32(m_pD->MultiValues.Num()) )
+        if (ParamIndex<0 || ParamIndex >= int32(MultiValues.Num()) )
         {
             return 0;
         }
 
-        return int32( m_pD->MultiValues[ParamIndex].Num() );
+        return int32( MultiValues[ParamIndex].Num() );
     }
 
 
     TSharedPtr<FRangeIndex> FParameters::GetValueIndex( int32 ParamIndex, int32 ValueIndex ) const
     {
-        if (ParamIndex<0 || ParamIndex >= m_pD->MultiValues.Num() )
+        if (ParamIndex<0 || ParamIndex >= MultiValues.Num() )
         {
             return nullptr;
         }
 
-        if (ValueIndex <0 || ValueIndex >= int32(m_pD->MultiValues[ParamIndex].Num()) )
+        if (ValueIndex <0 || ValueIndex >= int32(MultiValues[ParamIndex].Num()) )
         {
             return nullptr;
         }
 
-		TMap< TArray<int32>, FParameterValue >::TIterator it = m_pD->MultiValues[ParamIndex].CreateIterator();
+		TMap<TArray<int32>, FParameterValue>::TConstIterator it = MultiValues[ParamIndex].CreateConstIterator();
         for ( int32 i=0; i< ValueIndex; ++i )
         {
             ++it;
@@ -173,24 +180,24 @@ namespace UE::Mutable::Private
 
     void FParameters::ClearAllValues( int32 ParamIndex )
     {
-        if (ParamIndex<0 || ParamIndex >= int32(m_pD->MultiValues.Num()) )
+        if (ParamIndex<0 || ParamIndex >= int32(MultiValues.Num()) )
         {
             return;
         }
 
-        m_pD->MultiValues[ParamIndex].Empty();
+        MultiValues[ParamIndex].Empty();
     }
 
 
     bool FParameters::GetBoolValue( int32 Index, const FRangeIndex* Pos ) const
 	{
-		check( Index>=0 && Index<(int32)m_pD->Values.Num() );
+		check( Index>=0 && Index<(int32)Values.Num() );
         check( GetType(Index)==EParameterType::Bool );
 
         // Early out in case of invalid parameters
         if ( Index < 0
              ||
-             Index >= (int32)m_pD->Values.Num()
+             Index >= (int32)Values.Num()
              ||
              GetType(Index) != EParameterType::Bool )
         {
@@ -201,15 +208,15 @@ namespace UE::Mutable::Private
         if (!Pos)
         {
             // Return the single value
-            return m_pD->Values[Index].Get<FParamBoolType>();
+            return Values[Index].Get<FParamBoolType>();
         }
 
         // Multivalue case
         check( Pos->Parameter==Index );
 
-        if ( Index<int32(m_pD->MultiValues.Num()))
+        if ( Index<int32(MultiValues.Num()))
         {
-            const TMap< TArray<int32>, FParameterValue >& m = m_pD->MultiValues[Index];
+            const TMap< TArray<int32>, FParameterValue >& m = MultiValues[Index];
 			const FParameterValue* it = m.Find(Pos->Values);
             if (it)
             {
@@ -218,20 +225,20 @@ namespace UE::Mutable::Private
         }
 
         // Multivalue parameter, but no multivalue set. Return single value.
-        return m_pD->Values[Index].Get<FParamBoolType>();
+        return Values[Index].Get<FParamBoolType>();
 	}
 
 
     void FParameters::SetBoolValue( int32 Index, bool value,
                                    const FRangeIndex* Pos )
 	{
-		check( Index>=0 && Index<(int32)m_pD->Values.Num() );
+		check( Index>=0 && Index<(int32)Values.Num() );
 		check( GetType(Index)== EParameterType::Bool );
 
         // Early out in case of invalid parameters
         if ( Index < 0
              ||
-             Index >= (int32)m_pD->Values.Num()
+             Index >= (int32)Values.Num()
              ||
              GetType(Index) != EParameterType::Bool )
         {
@@ -242,12 +249,12 @@ namespace UE::Mutable::Private
         if (!Pos)
         {
             // Clear multivalue, if set.
-            if (Index<int32(m_pD->MultiValues.Num()))
+            if (Index<int32(MultiValues.Num()))
             {
-                m_pD->MultiValues[Index].Empty();
+                MultiValues[Index].Empty();
             }
 
-            m_pD->Values[Index].Set<FParamBoolType>(value);
+            Values[Index].Set<FParamBoolType>(value);
         }
 
         // Multivalue case
@@ -255,12 +262,12 @@ namespace UE::Mutable::Private
         {
             check( Pos->Parameter==Index );
 
-            if ( Index>=int32(m_pD->MultiValues.Num()))
+            if ( Index>=int32(MultiValues.Num()))
             {
-                m_pD->MultiValues.SetNum(Index+1);
+                MultiValues.SetNum(Index+1);
             }
 
-			TMap< TArray<int32>, FParameterValue >& m = m_pD->MultiValues[Index];
+			TMap< TArray<int32>, FParameterValue >& m = MultiValues[Index];
 			FParameterValue& it = m.FindOrAdd(Pos->Values);
             it.Set<FParamBoolType>(value);
         }
@@ -269,7 +276,7 @@ namespace UE::Mutable::Private
 
 	int32 FParameters::GetIntPossibleValueCount( int32 ParamIndex ) const
 	{
-		const FProgram& Program = m_pD->Model->GetPrivate()->Program;
+		const FProgram& Program = Model->Program;
 		check( ParamIndex>=0 && ParamIndex<(int32)Program.Parameters.Num() );
 		return (int32)Program.Parameters[ParamIndex].PossibleValues.Num();
 	}
@@ -277,7 +284,7 @@ namespace UE::Mutable::Private
 
 	int32 FParameters::GetIntPossibleValue( int32 ParamIndex, int32 valueIndex ) const
 	{
-		const FProgram& Program = m_pD->Model->GetPrivate()->Program;
+		const FProgram& Program = Model->Program;
 		check( ParamIndex>=0
 				&& ParamIndex<(int32)Program.Parameters.Num() );
 		check( valueIndex>=0
@@ -289,7 +296,7 @@ namespace UE::Mutable::Private
 
 	int32 FParameters::GetIntValueIndex(int32 ParamIndex, const FString& ValueName) const
 	{
-		const FProgram& Program = m_pD->Model->GetPrivate()->Program;
+		const FProgram& Program = Model->Program;
 		check(ParamIndex >= 0
 			&& ParamIndex < (int32)Program.Parameters.Num());
 
@@ -308,7 +315,7 @@ namespace UE::Mutable::Private
 
 	int32 FParameters::GetIntValueIndex(int32 ParamIndex, int32 Value) const
 	{
-		const FProgram& Program = m_pD->Model->GetPrivate()->Program;
+		const FProgram& Program = Model->Program;
 		check(ParamIndex >= 0
 			&& ParamIndex < (int32)Program.Parameters.Num());
 
@@ -325,7 +332,7 @@ namespace UE::Mutable::Private
 
 	const FString& FParameters::GetIntPossibleValueName( int32 ParamIndex, int32 valueIndex ) const
 	{
-		const FProgram& Program = m_pD->Model->GetPrivate()->Program;
+		const FProgram& Program = Model->Program;
 		check( ParamIndex>=0
 				&& ParamIndex<(int32)Program.Parameters.Num() );
 		check( valueIndex>=0
@@ -337,13 +344,13 @@ namespace UE::Mutable::Private
 
     int32 FParameters::GetIntValue( int32 Index, const FRangeIndex* Pos ) const
 	{
-		check( Index>=0 && Index<(int32)m_pD->Values.Num() );
+		check( Index>=0 && Index<(int32)Values.Num() );
 		check( GetType(Index)== EParameterType::Int );
 
         // Early out in case of invalid parameters
         if ( Index < 0
              ||
-             Index >= (int32)m_pD->Values.Num()
+             Index >= (int32)Values.Num()
              ||
              GetType(Index) != EParameterType::Int )
         {
@@ -354,15 +361,15 @@ namespace UE::Mutable::Private
         if (!Pos)
         {
             // Return the single value
-            return m_pD->Values[Index].Get<FParamIntType>();
+            return Values[Index].Get<FParamIntType>();
         }
 
         // Multivalue case
         check( Pos->Parameter==Index );
 
-        if ( Index<int32(m_pD->MultiValues.Num()))
+        if ( Index<int32(MultiValues.Num()))
         {
-            const TMap< TArray<int32>, FParameterValue >& m = m_pD->MultiValues[Index];
+            const TMap< TArray<int32>, FParameterValue >& m = MultiValues[Index];
 			const FParameterValue* it = m.Find(Pos->Values);
             if (it)
             {
@@ -371,20 +378,20 @@ namespace UE::Mutable::Private
         }
 
         // Multivalue parameter, but no multivalue set. Return single value.
-        return m_pD->Values[Index].Get<FParamIntType>();
+        return Values[Index].Get<FParamIntType>();
     }
 
 
     void FParameters::SetIntValue( int32 Index, int32 value,
                                   const FRangeIndex* Pos )
 	{
-		check( Index>=0 && Index<(int32)m_pD->Values.Num() );
+		check( Index>=0 && Index<(int32)Values.Num() );
 		check( GetType(Index)== EParameterType::Int );
 
         // Early out in case of invalid parameters
         if ( Index < 0
              ||
-             Index >= (int32)m_pD->Values.Num()
+             Index >= (int32)Values.Num()
              ||
              GetType(Index) != EParameterType::Int )
         {
@@ -395,12 +402,12 @@ namespace UE::Mutable::Private
         if (!Pos)
         {
             // Clear multivalue, if set.
-            if (Index<int32(m_pD->MultiValues.Num()))
+            if (Index<int32(MultiValues.Num()))
             {
-                m_pD->MultiValues[Index].Empty();
+                MultiValues[Index].Empty();
             }
 
-            m_pD->Values[Index].Set<FParamIntType>(value);
+            Values[Index].Set<FParamIntType>(value);
         }
 
         // Multivalue case
@@ -408,12 +415,12 @@ namespace UE::Mutable::Private
         {
             check( Pos->Parameter==Index );
 
-            if ( Index>=int32(m_pD->MultiValues.Num()))
+            if ( Index>=int32(MultiValues.Num()))
             {
-                m_pD->MultiValues.SetNum(Index+1);
+                MultiValues.SetNum(Index+1);
             }
 
-			TMap< TArray<int32>, FParameterValue >& m = m_pD->MultiValues[Index];
+			TMap< TArray<int32>, FParameterValue >& m = MultiValues[Index];
 			FParameterValue& it = m.FindOrAdd(Pos->Values);
             it.Set<FParamIntType>(value);
         }
@@ -423,13 +430,13 @@ namespace UE::Mutable::Private
     float FParameters::GetFloatValue( int32 Index,
                                      const FRangeIndex* Pos ) const
 	{
-		check( Index>=0 && Index<(int32)m_pD->Values.Num() );
+		check( Index>=0 && Index<(int32)Values.Num() );
 		check( GetType(Index)== EParameterType::Float );
 
         // Early out in case of invalid parameters
         if ( Index < 0
              ||
-             Index >= (int32)m_pD->Values.Num()
+             Index >= (int32)Values.Num()
              ||
              GetType(Index) != EParameterType::Float )
         {
@@ -440,15 +447,15 @@ namespace UE::Mutable::Private
         if (!Pos)
         {
             // Return the single value
-            return m_pD->Values[Index].Get<FParamFloatType>();
+            return Values[Index].Get<FParamFloatType>();
         }
 
         // Multivalue case
         check( Pos->Parameter==Index );
 
-        if ( Index<int32(m_pD->MultiValues.Num()))
+        if ( Index<int32(MultiValues.Num()))
         {
-            const TMap< TArray<int32>, FParameterValue >& m = m_pD->MultiValues[Index];
+            const TMap< TArray<int32>, FParameterValue >& m = MultiValues[Index];
 			const FParameterValue* it = m.Find(Pos->Values);
             if (it)
             {
@@ -457,20 +464,20 @@ namespace UE::Mutable::Private
         }
 
         // Multivalue parameter, but no multivalue set. Return single value.
-        return m_pD->Values[Index].Get<FParamFloatType>();
+        return Values[Index].Get<FParamFloatType>();
     }
 
 
     void FParameters::SetFloatValue( int32 Index, float value,
                                     const FRangeIndex* Pos )
 	{
-		check( Index>=0 && Index<(int32)m_pD->Values.Num() );
+		check( Index>=0 && Index<(int32)Values.Num() );
 		check( GetType(Index)== EParameterType::Float );
 
         // Early out in case of invalid parameters
         if ( Index < 0
              ||
-             Index >= (int32)m_pD->Values.Num()
+             Index >= (int32)Values.Num()
              ||
              GetType(Index) != EParameterType::Float )
         {
@@ -481,12 +488,12 @@ namespace UE::Mutable::Private
         if (!Pos)
         {
             // Clear multivalue, if set.
-            if (Index<int32(m_pD->MultiValues.Num()))
+            if (Index<int32(MultiValues.Num()))
             {
-                m_pD->MultiValues[Index].Empty();
+                MultiValues[Index].Empty();
             }
 
-            m_pD->Values[Index].Set<FParamFloatType>(value);
+            Values[Index].Set<FParamFloatType>(value);
         }
 
         // Multivalue case
@@ -494,12 +501,12 @@ namespace UE::Mutable::Private
         {
             check( Pos->Parameter==Index );
 
-            if ( Index>=int32(m_pD->MultiValues.Num()))
+            if ( Index>=int32(MultiValues.Num()))
             {
-                m_pD->MultiValues.SetNum(Index+1);
+                MultiValues.SetNum(Index+1);
             }
 
-			TMap< TArray<int32>, FParameterValue >& m = m_pD->MultiValues[Index];
+			TMap< TArray<int32>, FParameterValue >& m = MultiValues[Index];
 			FParameterValue& it = m.FindOrAdd(Pos->Values);
             it.Set<FParamFloatType>(value);
         }
@@ -508,13 +515,13 @@ namespace UE::Mutable::Private
 
     void FParameters::GetColourValue( int32 Index, FVector4f& OutValue, const FRangeIndex* Pos ) const
     {
-        check( Index>=0 && Index<(int32)m_pD->Values.Num() );
+        check( Index>=0 && Index<(int32)Values.Num() );
         check( GetType(Index)== EParameterType::Color );
 
         // Early out in case of invalid parameters
         if ( Index < 0
              ||
-             Index >= (int32)m_pD->Values.Num()
+             Index >= (int32)Values.Num()
              ||
              GetType(Index) != EParameterType::Color )
         {
@@ -524,16 +531,16 @@ namespace UE::Mutable::Private
         // Single value case
         if (!Pos)
         {
-			OutValue = m_pD->Values[Index].Get<FParamColorType>();
+			OutValue = Values[Index].Get<FParamColorType>();
             return;
         }
 
         // Multivalue case
         check( Pos->Parameter==Index );
 
-        if ( Index<int32(m_pD->MultiValues.Num()))
+        if ( Index<int32(MultiValues.Num()))
         {
-            const TMap< TArray<int32>, FParameterValue >& m = m_pD->MultiValues[Index];
+            const TMap< TArray<int32>, FParameterValue >& m = MultiValues[Index];
 			const FParameterValue* it = m.Find(Pos->Values);
             if (it)
             {
@@ -543,20 +550,20 @@ namespace UE::Mutable::Private
         }
 
         // Multivalue parameter, but no multivalue set. Return single value.
-		OutValue = m_pD->Values[Index].Get<FParamColorType>();
+		OutValue = Values[Index].Get<FParamColorType>();
         return;
     }
 
 
     void FParameters::SetColourValue( int32 Index, FVector4f InValue, const FRangeIndex* Pos )
     {
-        check( Index>=0 && Index<(int32)m_pD->Values.Num() );
+        check( Index>=0 && Index<(int32)Values.Num() );
         check( GetType(Index)== EParameterType::Color );
 
         // Early out in case of invalid parameters
         if ( Index < 0
              ||
-             Index >= (int32)m_pD->Values.Num()
+             Index >= (int32)Values.Num()
              ||
              GetType(Index) != EParameterType::Color )
         {
@@ -567,12 +574,12 @@ namespace UE::Mutable::Private
         if (!Pos)
         {
             // Clear multivalue, if set.
-            if (Index<int32(m_pD->MultiValues.Num()))
+            if (Index<int32(MultiValues.Num()))
             {
-                m_pD->MultiValues[Index].Empty();
+                MultiValues[Index].Empty();
             }
 
-        	m_pD->Values[Index].Set<FParamColorType>(InValue);
+        	Values[Index].Set<FParamColorType>(InValue);
         }
 
         // Multivalue case
@@ -580,12 +587,12 @@ namespace UE::Mutable::Private
         {
             check( Pos->Parameter==Index );
 
-            if ( Index>=int32(m_pD->MultiValues.Num()))
+            if ( Index>=int32(MultiValues.Num()))
             {
-                m_pD->MultiValues.SetNum(Index+1);
+                MultiValues.SetNum(Index+1);
             }
 
-			TMap< TArray<int32>, FParameterValue >& m = m_pD->MultiValues[Index];
+			TMap< TArray<int32>, FParameterValue >& m = MultiValues[Index];
 			FParameterValue& it = m.FindOrAdd(Pos->Values);
 
         	it.Set<FParamColorType>(InValue);
@@ -595,13 +602,13 @@ namespace UE::Mutable::Private
 
 	USkeletalMesh* FParameters::GetSkeletalMeshValue(int32 Index, const FRangeIndex* Pos) const
 	{
-		check(Index >= 0 && Index < (int32)m_pD->Values.Num());
+		check(Index >= 0 && Index < (int32)Values.Num());
 		check(GetType(Index) == EParameterType::SkeletalMesh);
 
 		// Early out in case of invalid parameters
 		if (Index < 0
 			||
-			Index >= (int32)m_pD->Values.Num()
+			Index >= (int32)Values.Num()
 			||
 			GetType(Index) != EParameterType::SkeletalMesh)
 		{
@@ -611,15 +618,15 @@ namespace UE::Mutable::Private
 		// Single value case
 		if (!Pos)
 		{
-			return m_pD->Values[Index].Get<FParamSkeletalMeshType>().Get();
+			return Values[Index].Get<FParamSkeletalMeshType>().Get();
 		}
 
 		// Multivalue case
 		check(Pos->Parameter == Index);
 
-		if (Index < int32(m_pD->MultiValues.Num()))
+		if (Index < int32(MultiValues.Num()))
 		{
-			const TMap< TArray<int32>, FParameterValue >& m = m_pD->MultiValues[Index];
+			const TMap< TArray<int32>, FParameterValue >& m = MultiValues[Index];
 			const FParameterValue* it = m.Find(Pos->Values);
 			if (it)
 			{
@@ -628,19 +635,19 @@ namespace UE::Mutable::Private
 		}
 
 		// Multivalue parameter, but no multivalue set. Return single value.
-		return m_pD->Values[Index].Get<FParamSkeletalMeshType>().Get();
+		return Values[Index].Get<FParamSkeletalMeshType>().Get();
 	}
 
 
 	void FParameters::SetSkeletalMeshValue(int32 Index, USkeletalMesh* Value, const FRangeIndex* Pos)
 	{
-		check(Index >= 0 && Index < (int32)m_pD->Values.Num());
+		check(Index >= 0 && Index < (int32)Values.Num());
 		check(GetType(Index) == EParameterType::SkeletalMesh);
 
 		// Single value case
 		if (!Pos)
 		{
-			m_pD->Values[Index].Set<FParamSkeletalMeshType>(TStrongObjectPtr(Value));
+			Values[Index].Set<FParamSkeletalMeshType>(TStrongObjectPtr(Value));
 		}
 
 		// Multivalue case
@@ -648,12 +655,12 @@ namespace UE::Mutable::Private
 		{
 			check(Pos->Parameter == Index);
 
-			if (Index >= int32(m_pD->MultiValues.Num()))
+			if (Index >= int32(MultiValues.Num()))
 			{
-				m_pD->MultiValues.SetNum(Index + 1);
+				MultiValues.SetNum(Index + 1);
 			}
 
-			TMap< TArray<int32>, FParameterValue >& m = m_pD->MultiValues[Index];
+			TMap< TArray<int32>, FParameterValue >& m = MultiValues[Index];
 			FParameterValue& it = m.FindOrAdd(Pos->Values);
 			it.Set<FParamSkeletalMeshType>(TStrongObjectPtr(Value));
 		}
@@ -662,13 +669,13 @@ namespace UE::Mutable::Private
 
 	UTexture* FParameters::GetTextureValue( int32 Index, const FRangeIndex* Pos ) const
     {
-        check( Index >= 0 && Index < (int32)m_pD->Values.Num() );
+        check( Index >= 0 && Index < (int32)Values.Num() );
         check( GetType( Index ) == EParameterType::Texture );
 
 		// Early out in case of invalid parameters
         if ( Index < 0
              ||
-             Index >= (int32)m_pD->Values.Num()
+             Index >= (int32)Values.Num()
              ||
              GetType(Index) != EParameterType::Texture )
         {
@@ -678,15 +685,15 @@ namespace UE::Mutable::Private
 		// Single value case
 		if (!Pos)
 		{
-			return m_pD->Values[Index].Get<FParamTextureType>().Get();
+			return Values[Index].Get<FParamTextureType>().Get();
 		}
 
 		// Multivalue case
 		check(Pos->Parameter == Index);
 
-		if (Index < int32(m_pD->MultiValues.Num()))
+		if (Index < int32(MultiValues.Num()))
 		{
-			const TMap< TArray<int32>, FParameterValue >& m = m_pD->MultiValues[Index];
+			const TMap< TArray<int32>, FParameterValue >& m = MultiValues[Index];
 			const FParameterValue* it = m.Find(Pos->Values);
 			if (it)
 			{
@@ -695,19 +702,19 @@ namespace UE::Mutable::Private
 		}
 
 		// Multivalue parameter, but no multivalue set. Return single value.
-		return m_pD->Values[Index].Get<FParamTextureType>().Get();
+		return Values[Index].Get<FParamTextureType>().Get();
     }
 
 
     void FParameters::SetTextureValue( int32 Index, UTexture* Value, const FRangeIndex* Pos )
     {
-        check( Index >= 0 && Index < (int32)m_pD->Values.Num() );
+        check( Index >= 0 && Index < (int32)Values.Num() );
         check( GetType( Index ) == EParameterType::Texture );
 
 		// Single value case
 		if (!Pos)
 		{
-			m_pD->Values[Index].Set<FParamTextureType>(TStrongObjectPtr(Value));
+			Values[Index].Set<FParamTextureType>(TStrongObjectPtr(Value));
 		}
 
 		// Multivalue case
@@ -715,12 +722,12 @@ namespace UE::Mutable::Private
 		{
 			check(Pos->Parameter == Index);
 
-			if (Index >= int32(m_pD->MultiValues.Num()))
+			if (Index >= int32(MultiValues.Num()))
 			{
-				m_pD->MultiValues.SetNum(Index + 1);
+				MultiValues.SetNum(Index + 1);
 			}
 
-			TMap< TArray<int32>, FParameterValue >& m = m_pD->MultiValues[Index];
+			TMap< TArray<int32>, FParameterValue >& m = MultiValues[Index];
 			FParameterValue& it = m.FindOrAdd(Pos->Values);
 			it.Set<FParamTextureType>(TStrongObjectPtr(Value));
 		}
@@ -729,11 +736,11 @@ namespace UE::Mutable::Private
 
     void FParameters::GetStringValue( int32 Index, FString& OutValue, const FRangeIndex* Pos ) const
     {
-        check( Index >= 0 && Index < (int32)m_pD->Values.Num() );
+        check( Index >= 0 && Index < (int32)Values.Num() );
         check( GetType( Index ) == EParameterType::String );
 
         // Early out in case of invalid parameters
-        if ( Index < 0 || Index >= (int32)m_pD->Values.Num() ||
+        if ( Index < 0 || Index >= (int32)Values.Num() ||
              GetType( Index ) != EParameterType::String )
         {
             OutValue = TEXT("");
@@ -744,16 +751,16 @@ namespace UE::Mutable::Private
         if ( !Pos )
         {
             // Return the single value
-			OutValue = m_pD->Values[Index].Get<FParamStringType>();
+			OutValue = Values[Index].Get<FParamStringType>();
             return;
         }
 
         // Multivalue case
         check( Pos->Parameter == Index );
 
-        if ( Index < int32( m_pD->MultiValues.Num() ) )
+        if ( Index < int32( MultiValues.Num() ) )
         {
-            const TMap< TArray<int32>, FParameterValue >& m = m_pD->MultiValues[Index];
+            const TMap< TArray<int32>, FParameterValue >& m = MultiValues[Index];
 			const FParameterValue* it = m.Find( Pos->Values );
             if ( it )
             {
@@ -763,7 +770,7 @@ namespace UE::Mutable::Private
         }
 
         // Multivalue parameter, but no multivalue set. Return single value.
-		OutValue = m_pD->Values[Index].Get<FParamStringType>();
+		OutValue = Values[Index].Get<FParamStringType>();
         return;
     }
 
@@ -772,11 +779,11 @@ namespace UE::Mutable::Private
     {
 		LLM_SCOPE_BYNAME(TEXT("MutableRuntime"));
 
-        check( Index >= 0 && Index < (int32)m_pD->Values.Num() );
+        check( Index >= 0 && Index < (int32)Values.Num() );
         check( GetType( Index ) == EParameterType::Float );
 
         // Early out in case of invalid parameters
-        if ( Index < 0 || Index >= (int32)m_pD->Values.Num() ||
+        if ( Index < 0 || Index >= (int32)Values.Num() ||
              GetType( Index ) != EParameterType::Float )
         {
             return;
@@ -786,12 +793,12 @@ namespace UE::Mutable::Private
         if ( !Pos )
         {
             // Clear multivalue, if set.
-            if ( Index < int32( m_pD->MultiValues.Num() ) )
+            if ( Index < int32( MultiValues.Num() ) )
             {
-                m_pD->MultiValues[Index].Empty();
+                MultiValues[Index].Empty();
             }
 
-        	m_pD->Values[Index].Set<FParamStringType>(Value);
+        	Values[Index].Set<FParamStringType>(Value);
         }
 
         // Multivalue case
@@ -799,12 +806,12 @@ namespace UE::Mutable::Private
         {
             check( Pos->Parameter == Index );
 
-            if ( Index >= int32( m_pD->MultiValues.Num() ) )
+            if ( Index >= int32( MultiValues.Num() ) )
             {
-                m_pD->MultiValues.SetNum( Index + 1 );
+                MultiValues.SetNum( Index + 1 );
             }
 
-			TMap< TArray<int32>, FParameterValue >& m = m_pD->MultiValues[Index];
+			TMap< TArray<int32>, FParameterValue >& m = MultiValues[Index];
 			FParameterValue& it = m.FindOrAdd( Pos->Values );
         	it.Set<FParamStringType>(Value);
         }
@@ -812,11 +819,11 @@ namespace UE::Mutable::Private
 
     void FParameters::GetMatrixValue(int32 Index, FMatrix44f& OutValue, const FRangeIndex* Pos) const
     {
-    	check( Index >= 0 && Index < (int32)m_pD->Values.Num() );
+    	check( Index >= 0 && Index < (int32)Values.Num() );
     	check( GetType( Index ) == EParameterType::Matrix );
 
     	// Early out in case of invalid parameters
-    	if ( Index < 0 || Index >= (int32)m_pD->Values.Num() ||
+    	if ( Index < 0 || Index >= (int32)Values.Num() ||
 			 GetType( Index ) != EParameterType::Matrix )
     	{
     		OutValue = FMatrix44f::Identity;
@@ -827,16 +834,16 @@ namespace UE::Mutable::Private
     	if ( !Pos )
     	{
     		// Return the single value
-    		OutValue = m_pD->Values[Index].Get<FParamMatrixType>();
+    		OutValue = Values[Index].Get<FParamMatrixType>();
     		return;
     	}
 
     	// Multivalue case
     	check( Pos->Parameter == Index );
 
-    	if ( Index < int32( m_pD->MultiValues.Num() ) )
+    	if ( Index < int32( MultiValues.Num() ) )
     	{
-    		const TMap< TArray<int32>, FParameterValue >& m = m_pD->MultiValues[Index];
+    		const TMap< TArray<int32>, FParameterValue >& m = MultiValues[Index];
     		const FParameterValue* it = m.Find( Pos->Values );
     		if ( it )
     		{
@@ -846,18 +853,18 @@ namespace UE::Mutable::Private
     	}
 
     	// Multivalue parameter, but no multivalue set. Return single value.
-    	OutValue = m_pD->Values[Index].Get<FParamMatrixType>();
+    	OutValue = Values[Index].Get<FParamMatrixType>();
     }
 
     void FParameters::SetMatrixValue(int32 Index, const FMatrix44f& Value, const FRangeIndex* Pos)
     {
     	LLM_SCOPE_BYNAME(TEXT("MutableRuntime"));
 
-    	check( Index >= 0 && Index < (int32)m_pD->Values.Num() );
+    	check( Index >= 0 && Index < (int32)Values.Num() );
     	check( GetType( Index ) == EParameterType::Matrix );
 
     	// Early out in case of invalid parameters
-    	if ( Index < 0 || Index >= (int32)m_pD->Values.Num() ||
+    	if ( Index < 0 || Index >= (int32)Values.Num() ||
 			 GetType( Index ) != EParameterType::Matrix )
     	{
     		return;
@@ -867,12 +874,12 @@ namespace UE::Mutable::Private
     	if ( !Pos )
     	{
     		// Clear multivalue, if set.
-    		if ( Index < int32( m_pD->MultiValues.Num() ) )
+    		if ( Index < int32( MultiValues.Num() ) )
     		{
-    			m_pD->MultiValues[Index].Empty();
+    			MultiValues[Index].Empty();
     		}
 
-    		m_pD->Values[Index].Set<FParamMatrixType>(Value);
+    		Values[Index].Set<FParamMatrixType>(Value);
     	}
 
     	// Multivalue case
@@ -880,12 +887,12 @@ namespace UE::Mutable::Private
     	{
     		check( Pos->Parameter == Index );
 
-    		if ( Index >= int32( m_pD->MultiValues.Num() ) )
+    		if ( Index >= int32( MultiValues.Num() ) )
     		{
-    			m_pD->MultiValues.SetNum( Index + 1 );
+    			MultiValues.SetNum( Index + 1 );
     		}
 
-    		TMap< TArray<int32>, FParameterValue >& m = m_pD->MultiValues[Index];
+    		TMap< TArray<int32>, FParameterValue >& m = MultiValues[Index];
     		FParameterValue& it = m.FindOrAdd( Pos->Values );
     		it.Set<FParamMatrixType>(Value);
     	}
@@ -894,11 +901,11 @@ namespace UE::Mutable::Private
 
 	UMaterialInterface* FParameters::GetMaterialValue(int32 Index, const FRangeIndex* RangeIndex) const
 	{
-		check(Index >= 0 && Index < (int32)m_pD->Values.Num());
+		check(Index >= 0 && Index < (int32)Values.Num());
 		check(GetType(Index) == EParameterType::Material);
 
 		// Early out in case of invalid parameters
-		if (Index < 0 || Index >= (int32)m_pD->Values.Num() ||
+		if (Index < 0 || Index >= (int32)Values.Num() ||
 			GetType(Index) != EParameterType::Material)
 		{
 			return nullptr;
@@ -908,15 +915,15 @@ namespace UE::Mutable::Private
 		if (!RangeIndex)
 		{
 			// Return the single value
-			return m_pD->Values[Index].Get<FParamMaterialType>().Get();
+			return Values[Index].Get<FParamMaterialType>().Get();
 		}
 
 		// Multivalue case
 		check(RangeIndex->Parameter == Index);
 
-		if (Index < int32(m_pD->MultiValues.Num()))
+		if (Index < int32(MultiValues.Num()))
 		{
-			const TMap< TArray<int32>, FParameterValue >& m = m_pD->MultiValues[Index];
+			const TMap< TArray<int32>, FParameterValue >& m = MultiValues[Index];
 			const FParameterValue* it = m.Find(RangeIndex->Values);
 			if (it)
 			{
@@ -925,7 +932,7 @@ namespace UE::Mutable::Private
 		}
 
 		// Multivalue parameter, but no multivalue set. Return single value.
-		return m_pD->Values[Index].Get<FParamMaterialType>().Get();
+		return Values[Index].Get<FParamMaterialType>().Get();
 	}
 
 
@@ -933,11 +940,11 @@ namespace UE::Mutable::Private
 	{
 		LLM_SCOPE_BYNAME(TEXT("MutableRuntime"));
 
-		check(Index >= 0 && Index < (int32)m_pD->Values.Num());
+		check(Index >= 0 && Index < (int32)Values.Num());
 		check(GetType(Index) == EParameterType::Material);
 
 		// Early out in case of invalid parameters
-		if (Index < 0 || Index >= (int32)m_pD->Values.Num() ||
+		if (Index < 0 || Index >= (int32)Values.Num() ||
 			GetType(Index) != EParameterType::Material)
 		{
 			return;
@@ -947,12 +954,12 @@ namespace UE::Mutable::Private
 		if (!RangeIndex)
 		{
 			// Clear multivalue, if set.
-			if (Index < int32(m_pD->MultiValues.Num()))
+			if (Index < int32(MultiValues.Num()))
 			{
-				m_pD->MultiValues[Index].Empty();
+				MultiValues[Index].Empty();
 			}
 
-			m_pD->Values[Index].Set<FParamMaterialType>(TStrongObjectPtr(Material));
+			Values[Index].Set<FParamMaterialType>(TStrongObjectPtr(Material));
 		}
 
 		// Multivalue case
@@ -960,12 +967,12 @@ namespace UE::Mutable::Private
 		{
 			check(RangeIndex->Parameter == Index);
 
-			if (Index >= int32(m_pD->MultiValues.Num()))
+			if (Index >= int32(MultiValues.Num()))
 			{
-				m_pD->MultiValues.SetNum(Index + 1);
+				MultiValues.SetNum(Index + 1);
 			}
 
-			TMap< TArray<int32>, FParameterValue >& m = m_pD->MultiValues[Index];
+			TMap< TArray<int32>, FParameterValue >& m = MultiValues[Index];
 			FParameterValue& it = m.FindOrAdd(RangeIndex->Values);
 			it.Set<FParamMaterialType>(TStrongObjectPtr(Material));
 		}
@@ -974,11 +981,11 @@ namespace UE::Mutable::Private
 
 	TManagedPtr<const FInstancedStruct> FParameters::GetInstancedStructValue(int32 Index, const FRangeIndex* RangeIndex) const
 	{
-    	check(Index >= 0 && Index < (int32)m_pD->Values.Num());
+    	check(Index >= 0 && Index < (int32)Values.Num());
     	check(GetType(Index) == EParameterType::InstancedStruct);
 
     	// Early out in case of invalid parameters
-    	if (Index < 0 || Index >= (int32)m_pD->Values.Num() ||
+    	if (Index < 0 || Index >= (int32)Values.Num() ||
 			GetType(Index) != EParameterType::InstancedStruct)
     	{
     		return {};
@@ -988,15 +995,15 @@ namespace UE::Mutable::Private
     	if (!RangeIndex)
     	{
     		// Return the single value
-    		return m_pD->Values[Index].Get<FParamInstancedStructType>();
+    		return Values[Index].Get<FParamInstancedStructType>();
     	}
 
     	// Multivalue case
     	check(RangeIndex->Parameter == Index);
 
-    	if (Index < int32(m_pD->MultiValues.Num()))
+    	if (Index < int32(MultiValues.Num()))
     	{
-    		const TMap< TArray<int32>, FParameterValue >& m = m_pD->MultiValues[Index];
+    		const TMap< TArray<int32>, FParameterValue >& m = MultiValues[Index];
     		const FParameterValue* it = m.Find(RangeIndex->Values);
     		if (it)
     		{
@@ -1005,7 +1012,7 @@ namespace UE::Mutable::Private
     	}
 
     	// Multivalue parameter, but no multivalue set. Return single value.
-    	return m_pD->Values[Index].Get<FParamInstancedStructType>();
+    	return Values[Index].Get<FParamInstancedStructType>();
 	}
 
 
@@ -1013,11 +1020,11 @@ namespace UE::Mutable::Private
 	{
     	LLM_SCOPE_BYNAME(TEXT("MutableRuntime"));
 
-    	check(Index >= 0 && Index < (int32)m_pD->Values.Num());
+    	check(Index >= 0 && Index < (int32)Values.Num());
     	check(GetType(Index) == EParameterType::InstancedStruct);
 
     	// Early out in case of invalid parameters
-    	if (Index < 0 || Index >= (int32)m_pD->Values.Num() ||
+    	if (Index < 0 || Index >= (int32)Values.Num() ||
 			GetType(Index) != EParameterType::InstancedStruct)
     	{
     		return;
@@ -1027,12 +1034,12 @@ namespace UE::Mutable::Private
     	if (!RangeIndex)
     	{
     		// Clear multivalue, if set.
-    		if (Index < int32(m_pD->MultiValues.Num()))
+    		if (Index < int32(MultiValues.Num()))
     		{
-    			m_pD->MultiValues[Index].Empty();
+    			MultiValues[Index].Empty();
     		}
 
-    		m_pD->Values[Index].Set<FParamInstancedStructType>(Value);
+    		Values[Index].Set<FParamInstancedStructType>(Value);
     	}
 
     	// Multivalue case
@@ -1040,21 +1047,21 @@ namespace UE::Mutable::Private
     	{
     		check(RangeIndex->Parameter == Index);
 
-    		if (Index >= int32(m_pD->MultiValues.Num()))
+    		if (Index >= int32(MultiValues.Num()))
     		{
-    			m_pD->MultiValues.SetNum(Index + 1);
+    			MultiValues.SetNum(Index + 1);
     		}
 
-    		TMap< TArray<int32>, FParameterValue >& m = m_pD->MultiValues[Index];
+    		TMap< TArray<int32>, FParameterValue >& m = MultiValues[Index];
     		FParameterValue& it = m.FindOrAdd(RangeIndex->Values);
     		it.Set<FParamInstancedStructType>(Value);
     	}
 	}
 
 
-	FProjector FParameters::Private::GetProjectorValue( int32 Index, const FRangeIndex* Pos ) const
+	FProjector FParameters::GetProjectorValue( int32 Index, const FRangeIndex* Pos ) const
     {
-        const FProgram& Program = Model->GetPrivate()->Program;
+        const FProgram& Program = Model->Program;
 
         // Early out in case of invalid parameters
         if ( Index < 0
@@ -1104,20 +1111,20 @@ namespace UE::Mutable::Private
 		float* OutProjectionAngle,
 		const FRangeIndex* Pos ) const
 	{
-		check(ParameterIndex >=0 && ParameterIndex <m_pD->Values.Num() );
+		check(ParameterIndex >=0 && ParameterIndex <Values.Num() );
 		check( GetType(ParameterIndex)== EParameterType::Projector );
 
         // Early out in case of invalid parameters
         if (ParameterIndex < 0
              ||
-			ParameterIndex >= m_pD->Values.Num()
+			ParameterIndex >= Values.Num()
              ||
              GetType(ParameterIndex) != EParameterType::Projector )
         {
             return;
         }
 
-        FProjector result = m_pD->GetProjectorValue(ParameterIndex, Pos );
+        FProjector result = GetProjectorValue(ParameterIndex, Pos );
 
         // Copy results
         if (OutType) *OutType = result.type;
@@ -1139,13 +1146,13 @@ namespace UE::Mutable::Private
         float projectionAngle,
         const FRangeIndex* RangePosition)
 	{
-		check(ParameterIndex >=0 && ParameterIndex <m_pD->Values.Num() );
+		check(ParameterIndex >=0 && ParameterIndex <Values.Num() );
 		check( GetType(ParameterIndex)== EParameterType::Projector );
 
         // Early out in case of invalid parameters
         if (ParameterIndex < 0
              ||
-			ParameterIndex >= m_pD->Values.Num()
+			ParameterIndex >= Values.Num()
              ||
              GetType(ParameterIndex) != EParameterType::Projector )
         {
@@ -1155,9 +1162,9 @@ namespace UE::Mutable::Private
 	    // FParameters cannot change the projector type anymore
 		FProjector Value;
         Value.type = EProjectorType::Count;
-        if (m_pD->Model)
+        if (Model)
         {
-            const FProgram& Program = m_pD->Model->GetPrivate()->Program;
+            const FProgram& Program = Model->Program;
             const FProjector& Projector = Program.Parameters[ParameterIndex].DefaultValue.Get<FParamProjectorType>();
             Value.type = Projector.type;
         }
@@ -1173,12 +1180,12 @@ namespace UE::Mutable::Private
         if (!RangePosition)
         {
             // Clear multivalue, if set.
-            if (ParameterIndex<m_pD->MultiValues.Num())
+            if (ParameterIndex<MultiValues.Num())
             {
-                m_pD->MultiValues[ParameterIndex].Empty();
+                MultiValues[ParameterIndex].Empty();
             }
 
-            m_pD->Values[ParameterIndex].Set<FParamProjectorType>(Value);
+            Values[ParameterIndex].Set<FParamProjectorType>(Value);
         }
 
         // Multivalue case
@@ -1186,12 +1193,12 @@ namespace UE::Mutable::Private
         {
             check(RangePosition->Parameter== ParameterIndex);
 
-            if (ParameterIndex >=m_pD->MultiValues.Num())
+            if (ParameterIndex >=MultiValues.Num())
             {
-                m_pD->MultiValues.SetNum(ParameterIndex+1);
+                MultiValues.SetNum(ParameterIndex+1);
             }
 
-			TMap< TArray<int32>, FParameterValue >& m = m_pD->MultiValues[ParameterIndex];
+			TMap< TArray<int32>, FParameterValue >& m = MultiValues[ParameterIndex];
 			FParameterValue& it = m.FindOrAdd(RangePosition->Values);
             it.Set<FParamProjectorType>(Value);
         }
@@ -1207,23 +1214,23 @@ namespace UE::Mutable::Private
             return false;
         }
 
-        if ( !(m_pD->Values[thisParamIndex]==other->m_pD->Values[otherParamIndex] ) )
+        if ( !(Values[thisParamIndex]==other->Values[otherParamIndex] ) )
         {
             return false;
         }
 
         size_t thisNumMultiValues = 0;
-        bool thisHasMultiValues = int32(m_pD->MultiValues.Num()) > thisParamIndex;
+        bool thisHasMultiValues = int32(MultiValues.Num()) > thisParamIndex;
         if (thisHasMultiValues)
         {
-            thisNumMultiValues = m_pD->MultiValues[thisParamIndex].Num();
+            thisNumMultiValues = MultiValues[thisParamIndex].Num();
         }
 
         size_t otherNumMultiValues = 0;
-        bool otherHasMultiValues = int32(other->m_pD->MultiValues.Num()) > otherParamIndex;
+        bool otherHasMultiValues = int32(other->MultiValues.Num()) > otherParamIndex;
         if (otherHasMultiValues)
         {
-            otherNumMultiValues = other->m_pD->MultiValues[otherParamIndex].Num();
+            otherNumMultiValues = other->MultiValues[otherParamIndex].Num();
         }
 
         if ( thisNumMultiValues != otherNumMultiValues )
@@ -1235,39 +1242,27 @@ namespace UE::Mutable::Private
              &&
              otherHasMultiValues
              &&
-             !(m_pD->MultiValues[thisParamIndex]==other->m_pD->MultiValues[otherParamIndex] ))
+             !(MultiValues[thisParamIndex]==other->MultiValues[otherParamIndex] ))
         {
             return false;
         }
 
         return true;
     }
+	
+	
+	bool FParameters::HasMultipleValues(int32 ParamIndex) const
+    {
+    	if (ParamIndex >= MultiValues.Num())
+    	{
+    		return false;
+    	}
 
+    	return MultiValues[ParamIndex].Num()>0;
+    }
+	
 
-	int32 FParameters::Private::Find( const FString& Name ) const
-	{
-		const FProgram& Program = Model->GetPrivate()->Program;
-
-		int32 result = -1;
-
-		for( int32 i=0; result<0 && i<(int32)Program.Parameters.Num(); ++i )
-		{
-			const FParameterDesc& p = Program.Parameters[i];
-
-			if (p.Name == Name)
-			{
-				result = i;
-			}
-		}
-
-		return result;
-	}
-
-
-    //---------------------------------------------------------------------------------------------
-    //---------------------------------------------------------------------------------------------
-    //---------------------------------------------------------------------------------------------
-     int32 FRangeIndex::GetRangeCount() const
+    int32 FRangeIndex::GetRangeCount() const
     {
         return int32( Values.Num() );
     }
@@ -1276,24 +1271,24 @@ namespace UE::Mutable::Private
     const FString& FRangeIndex::GetRangeName( int32 Index ) const
     {
         check( Index >= 0 && Index < GetRangeCount() );
-        check( Index < int32(Parameters->GetPrivate()->Values.Num()) );
-        TSharedPtr<const FModel> Model = Parameters->GetPrivate()->Model;
+        check( Index < int32(Parameters->Values.Num()) );
+        TSharedPtr<const FModel> Model = Parameters->Model;
         check(Model);
-        check( Index < int32(Model->GetPrivate()->Program.Ranges.Num()) );
+        check( Index < int32(Model->Program.Ranges.Num()) );
 
-        return Model->GetPrivate()->Program.Ranges[Index].Name;
+        return Model->Program.Ranges[Index].Name;
     }
 
 
     const FString& FRangeIndex::GetRangeUid( int32 Index ) const
     {
         check( Index >= 0 && Index < GetRangeCount() );
-        check( Index < int32(Parameters->GetPrivate()->Values.Num()) );
-		TSharedPtr<const FModel> Model = Parameters->GetPrivate()->Model;
+        check( Index < int32(Parameters->Values.Num()) );
+		TSharedPtr<const FModel> Model = Parameters->Model;
 		check(Model);
-        check( Index < int32(Model->GetPrivate()->Program.Ranges.Num()) );
+        check( Index < int32(Model->Program.Ranges.Num()) );
 
-        return Model->GetPrivate()->Program.Ranges[Index].UID;
+        return Model->Program.Ranges[Index].UID;
     }
 
 
@@ -1316,6 +1311,5 @@ namespace UE::Mutable::Private
         }
         return 0;
     }
-
 }
 
