[Mutable] Remove Core Pimpl.

diff --git a/Engine/Plugins/Mutable/Source/MutableRuntime/Private/MuR/System.cpp b/Engine/Plugins/Mutable/Source/MutableRuntime/Private/MuR/System.cpp
--- a/Engine/Plugins/Mutable/Source/MutableRuntime/Private/MuR/System.cpp
+++ b/Engine/Plugins/Mutable/Source/MutableRuntime/Private/MuR/System.cpp
@@ -19,23 +19,23 @@
 #include "MuR/CodeVisitor.h"
 #include "MuR/Mesh.h"
 #include "MuR/Model.h"
-#include "MuR/ModelPrivate.h"
 #include "MuR/MutableMath.h"
 #include "MuR/MutableString.h"
 #include "MuR/MutableTrace.h"
 #include "MuR/Operations.h"
 #include "MuR/Parameters.h"
-#include "MuR/ParametersPrivate.h"
+#include "MuR/Parameters.h"
 #include "MuR/Platform.h"
 #include "MuR/Serialisation.h"
-#include "MuR/SystemPrivate.h"
+#include "MuR/System.h"
 #include "MuR/MutableRuntimeModule.h"
 #include "Templates/SharedPointer.h"
 #include "Templates/Tuple.h"
 #include "Trace/Detail/Channel.h"
 #include "ProfilingDebugging/CountersTrace.h"
 #include "PackedNormal.h"
-#include "MuR/SystemPrivate.h"
+#include "Engine/SkeletalMesh.h"
+#include "MuR/System.h"
 
 #include UE_INLINE_GENERATED_CPP_BY_NAME(System)
 
@@ -83,94 +83,49 @@ namespace UE::Mutable::Private
 	TRACE_DECLARE_INT_COUNTER(MutableRuntime_MemTotal,		TEXT("MutableRuntime/MemTotal"));
 	TRACE_DECLARE_INT_COUNTER(MutableRuntime_MemBudget,		TEXT("MutableRuntime/MemBudget"));
 
-
-    //---------------------------------------------------------------------------------------------
+   
     FSystem::FSystem(const FSettings& InSettings)
     {
-		LLM_SCOPE_BYNAME(TEXT("MutableRuntime"));
-
-        // Choose the implementation
-        m_pD = new FSystem::Private(InSettings);
-    }
-
-
-    //---------------------------------------------------------------------------------------------
-    FSystem::~FSystem()
-    {
-		LLM_SCOPE_BYNAME(TEXT("MutableRuntime"));
-		MUTABLE_CPUPROFILER_SCOPE(SystemDestructor);
-
-        check( m_pD );
+    	LLM_SCOPE_BYNAME(TEXT("MutableRuntime"));
 
-        delete m_pD;
-        m_pD = nullptr;
-    }
-
-
-    //---------------------------------------------------------------------------------------------
-    FSystem::Private* FSystem::GetPrivate() const
-    {
-        return m_pD;
-    }
-
-
-    //---------------------------------------------------------------------------------------------
-    FSystem::Private::Private(const FSettings& InSettings)
-    {
-		LLM_SCOPE_BYNAME(TEXT("MutableRuntime"));
-
-        Settings = InSettings;
+    	Settings = InSettings;
 	
-		WorkingMemoryManager.BudgetBytes = Settings.WorkingMemoryBytes;
-
-		UpdateStats();
-	}
-
-
-    //---------------------------------------------------------------------------------------------
-    FSystem::Private::~Private()
-    {
-		LLM_SCOPE_BYNAME(TEXT("MutableRuntime"));
-		MUTABLE_CPUPROFILER_SCOPE(SystemPrivateDestructor);
-
-		// Make it explicit to try to capture metrics
-        StreamInterface = nullptr;
-	}
+    	WorkingMemoryManager.BudgetBytes = Settings.WorkingMemoryBytes;
 
+    	UpdateStats();
+    }
 
-    //---------------------------------------------------------------------------------------------
-    void FSystem::SetStreamingInterface(const TSharedPtr<FModelReader>& InInterface )
+   
+    void FSystem::SetStreamingInterface(const TSharedPtr<FModelReader>& InInterface)
     {
 		LLM_SCOPE_BYNAME(TEXT("MutableRuntime"));
        
-        m_pD->StreamInterface = InInterface;
+        StreamInterface = InInterface;
     }
 
 
-	//---------------------------------------------------------------------------------------------
 	void FSystem::SetWorkingMemoryBytes(uint64 InBytes)
 	{
 		LLM_SCOPE_BYNAME(TEXT("MutableRuntime"));
 		MUTABLE_CPUPROFILER_SCOPE(SetWorkingMemoryBytes);
 
-		m_pD->WorkingMemoryManager.BudgetBytes = InBytes;
-		m_pD->WorkingMemoryManager.EnsureBudgetBelow(0);
+		WorkingMemoryManager.BudgetBytes = InBytes;
+		WorkingMemoryManager.EnsureBudgetBelow(0);
 
-		m_pD->UpdateStats();
+		UpdateStats();
 	}
 
 	
-	//---------------------------------------------------------------------------------------------
 	void FSystem::ClearWorkingMemory()
     {
 		LLM_SCOPE_BYNAME(TEXT("MutableRuntime"));
 
 		// rom caches
-		for (FWorkingMemoryManager::FModelCacheEntry& ModelCache : m_pD->WorkingMemoryManager.CachePerModel)
+		for (FWorkingMemoryManager::FModelCacheEntry& ModelCache : WorkingMemoryManager.CachePerModel)
 		{
 			if (const TSharedPtr<const FModel> CacheModel = ModelCache.Model.Pin())
 			{
-				FProgram& Program = CacheModel->GetPrivate()->Program;
+				const FProgram& Program = CacheModel->Program;
 
 				for (int32 RomIndex = 0; RomIndex < Program.Roms.Num(); ++RomIndex)
 				{
@@ -179,20 +134,18 @@ namespace UE::Mutable::Private
 			}
 		}
     	
-		m_pD->UpdateStats();
+		UpdateStats();
 	}
 
 
-	//---------------------------------------------------------------------------------------------
 	void FSystem::SetImagePixelConversionOverride(const FImageOperator::FImagePixelFormatFunc& FormatFunc)
 	{
 		LLM_SCOPE_BYNAME(TEXT("MutableRuntime"));
 
-		m_pD->ImagePixelFormatOverride = FormatFunc;
+		ImagePixelFormatOverride = FormatFunc;
 	}
 
-
-    //---------------------------------------------------------------------------------------------
+   
     FInstance::FID FSystem::NewInstance( const TSharedPtr<const FModel>& InModel,
     	const TSharedPtr<FExternalResourceProvider>& ExternalResourceProvider,
     	const TSharedPtr<FPassthroughObjectLoader>& PassthroughObjectLoader)
@@ -201,16 +154,16 @@ namespace UE::Mutable::Private
 		MUTABLE_CPUPROFILER_SCOPE(NewInstance);
 
 		FLiveInstance InstanceData;
-		InstanceData.InstanceID = ++m_pD->LastInstanceID;
+		InstanceData.InstanceID = ++LastInstanceID;
 		InstanceData.Instance = nullptr;
 		InstanceData.Model = InModel;
 		InstanceData.State = -1;
 		InstanceData.Cache = MakeShared<FProgramCache>();
     	InstanceData.ExternalResourceProvider = ExternalResourceProvider;
     	InstanceData.PassthroughObjectLoader = PassthroughObjectLoader;
-		m_pD->WorkingMemoryManager.LiveInstances.Add(InstanceData);
+		WorkingMemoryManager.LiveInstances.Add(InstanceData);
 
-		TRACE_COUNTER_SET(MutableRuntime_LiveInstances, m_pD->WorkingMemoryManager.LiveInstances.Num());
+		TRACE_COUNTER_SET(MutableRuntime_LiveInstances, WorkingMemoryManager.LiveInstances.Num());
 
 		return InstanceData.InstanceID;
 	}
@@ -220,15 +173,14 @@ namespace UE::Mutable::Private
     	const TSharedPtr<FExternalResourceProvider>& ExternalResourceProvider,
     	const TSharedPtr<FPassthroughObjectLoader>& PassthroughObjectLoader)
     {
-    	FLiveInstance* LiveInstance = m_pD->FindLiveInstance(InstanceId);
+    	FLiveInstance* LiveInstance = FindLiveInstance(InstanceId);
     	check(LiveInstance);
 
     	LiveInstance->ExternalResourceProvider = ExternalResourceProvider;
     	LiveInstance->PassthroughObjectLoader = PassthroughObjectLoader;
     }
 
-
-    //---------------------------------------------------------------------------------------------
+   
     TManagedPtr<const FInstance> FSystem::BeginUpdate_GameThread(FInstance::FID InInstanceID,
                                      const TSharedPtr<const FParameters>& InParams,
                                      const TSharedPtr<FMeshIdRegistry>& InMeshIdRegistry,
@@ -241,11 +193,11 @@ namespace UE::Mutable::Private
 		LLM_SCOPE_BYNAME(TEXT("MutableRuntime"));
 		MUTABLE_CPUPROFILER_SCOPE(BeginUpdate_GameThread);
 
-    	FLiveInstance* LiveInstance = m_pD->FindLiveInstance(InInstanceID);
+    	FLiveInstance* LiveInstance = FindLiveInstance(InInstanceID);
     	check(LiveInstance);
 		check(InParams);
 
-    	FScheduledOpInline Root(LiveInstance->Model->GetPrivate()->Program.States[InStateIndex].Root, 0, 0, true);
+    	FScheduledOpInline Root(LiveInstance->Model->Program.States[InStateIndex].Root, 0, 0, true);
     	CodeRunnerBegin Runner(
     		LiveInstance->Model,
     		InParams,
@@ -288,22 +240,22 @@ namespace UE::Mutable::Private
 			return nullptr;
 		}
 
-		FLiveInstance* LiveInstance = m_pD->FindLiveInstance(InInstanceID);
+		FLiveInstance* LiveInstance = FindLiveInstance(InInstanceID);
 		if (!LiveInstance)
 		{
 			UE_LOG(LogMutableCore, Error, TEXT("Invalid instance id in mutable update."));
 			return nullptr;
 		}
 
-		m_pD->WorkingMemoryManager.CurrentInstanceCache = LiveInstance->Cache;
-    	m_pD->WorkingMemoryManager.CurrentPassthroughObjectLoader = LiveInstance->PassthroughObjectLoader;
+		WorkingMemoryManager.CurrentInstanceCache = LiveInstance->Cache;
+    	WorkingMemoryManager.CurrentPassthroughObjectLoader = LiveInstance->PassthroughObjectLoader;
 
-    	m_pD->WorkingMemoryManager.CurrentMeshIdRegistry = InMeshIdRegistry;
-    	m_pD->WorkingMemoryManager.CurrentImageIdRegistry = InImageIdRegistry;
-    	m_pD->WorkingMemoryManager.CurrentMaterialIdRegistry = InMaterialIdRegistry;
-    	m_pD->WorkingMemoryManager.CurrentSkeletalMeshIdRegistry = InSkeletalMeshIdRegistry;
+    	WorkingMemoryManager.CurrentMeshIdRegistry = InMeshIdRegistry;
+    	WorkingMemoryManager.CurrentImageIdRegistry = InImageIdRegistry;
+    	WorkingMemoryManager.CurrentMaterialIdRegistry = InMaterialIdRegistry;
+    	WorkingMemoryManager.CurrentSkeletalMeshIdRegistry = InSkeletalMeshIdRegistry;
 
-		FProgram& Program = LiveInstance->Model->GetPrivate()->Program;
+		const FProgram& Program = LiveInstance->Model->Program;
 
 		bool bValidState = InStateIndex >= 0 && InStateIndex < (int)Program.States.Num();
 		if (!bValidState)
@@ -322,28 +274,28 @@ namespace UE::Mutable::Private
 		// If we changed parameters that are not in this state, we need to rebuild all.
 		if (!bFullBuild)
 		{
-			bFullBuild = m_pD->CheckUpdatedParameters(LiveInstance, InParams, LiveInstance->UpdatedParameters);
+			bFullBuild = CheckUpdatedParameters(LiveInstance, InParams, LiveInstance->UpdatedParameters);
 		}
 
 		// Remove cached data
-		m_pD->WorkingMemoryManager.ClearCacheLayer0();
+		WorkingMemoryManager.ClearCacheLayer0();
 		if (bFullBuild)
 		{
-			m_pD->WorkingMemoryManager.ClearCacheLayer1();
+			WorkingMemoryManager.ClearCacheLayer1();
 		}
 
-		m_pD->WorkingMemoryManager.BeginRunnerThread();
+		WorkingMemoryManager.BeginRunnerThread();
 
-		OP::ADDRESS RootAt = LiveInstance->Model->GetPrivate()->Program.States[InStateIndex].Root;
+		OP::ADDRESS RootAt = LiveInstance->Model->Program.States[InStateIndex].Root;
 
 		// Prepare instance cache
-		m_pD->PrepareCache(LiveInstance->Model.Get(), InStateIndex);
+		PrepareCache(LiveInstance->Model.Get(), InStateIndex);
 		LiveInstance->OldParameters = InParams->Clone();
 
 		// Ensure the model cache has been created
-		m_pD->WorkingMemoryManager.FindOrAddModelCache(LiveInstance->Model);
+		WorkingMemoryManager.FindOrAddModelCache(LiveInstance->Model);
 
-		m_pD->RunCode(LiveInstance->ExternalResourceProvider, LiveInstance->Model, InParams.Get(), RootAt, InLodMask);
+		RunCode(LiveInstance->ExternalResourceProvider, LiveInstance->Model, InParams.Get(), RootAt, InLodMask);
 
 		TManagedPtr<const FInstance> Result = LiveInstance->Cache->GetInstance(FCacheAddress(RootAt, 0, 0));
 
@@ -358,15 +310,15 @@ namespace UE::Mutable::Private
 			Result = MakeManaged<FInstance>();
 		}
 
-		m_pD->WorkingMemoryManager.EndRunnerThread();
+		WorkingMemoryManager.EndRunnerThread();
 
-		m_pD->WorkingMemoryManager.CurrentInstanceCache = nullptr;
-    	m_pD->WorkingMemoryManager.CurrentPassthroughObjectLoader = nullptr;
+		WorkingMemoryManager.CurrentInstanceCache = nullptr;
+    	WorkingMemoryManager.CurrentPassthroughObjectLoader = nullptr;
 
-    	m_pD->WorkingMemoryManager.CurrentMeshIdRegistry = nullptr;
-    	m_pD->WorkingMemoryManager.CurrentImageIdRegistry = nullptr;
-    	m_pD->WorkingMemoryManager.CurrentMaterialIdRegistry = nullptr;
-    	m_pD->WorkingMemoryManager.CurrentSkeletalMeshIdRegistry = nullptr;
+    	WorkingMemoryManager.CurrentMeshIdRegistry = nullptr;
+    	WorkingMemoryManager.CurrentImageIdRegistry = nullptr;
+    	WorkingMemoryManager.CurrentMaterialIdRegistry = nullptr;
+    	WorkingMemoryManager.CurrentSkeletalMeshIdRegistry = nullptr;
 
 		return Result;
 	}
@@ -374,7 +326,7 @@ namespace UE::Mutable::Private
 
     TManagedRef<FImage> BlackImage()
     {
-    	return MakeManaged<UE::Mutable::Private::FImage>(16, 16, 1, EImageFormat::RGBA_UByte, EInitializationType::Black);
+    	return MakeManaged<FImage>(16, 16, 1, EImageFormat::RGBA_UByte, EInitializationType::Black);
     }
 	
 
@@ -389,21 +341,21 @@ namespace UE::Mutable::Private
     	}
     	
 		// Find the live instance
-		FLiveInstance* LiveInstance = m_pD->FindLiveInstance(instanceID);
+		FLiveInstance* LiveInstance = FindLiveInstance(instanceID);
 		check(LiveInstance);
-		m_pD->WorkingMemoryManager.CurrentInstanceCache = LiveInstance->Cache;
-		m_pD->WorkingMemoryManager.CurrentPassthroughObjectLoader = LiveInstance->PassthroughObjectLoader;
+		WorkingMemoryManager.CurrentInstanceCache = LiveInstance->Cache;
+		WorkingMemoryManager.CurrentPassthroughObjectLoader = LiveInstance->PassthroughObjectLoader;
 
 		OP::ADDRESS RootAddress = ImageId.GetKey()->Address;
-		TManagedPtr<const FImage> Result = m_pD->BuildImage(LiveInstance->ExternalResourceProvider, LiveInstance->Model, LiveInstance->OldParameters.Get(), RootAddress, MipsToSkip);
+		TManagedPtr<const FImage> Result = BuildImage(LiveInstance->ExternalResourceProvider, LiveInstance->Model, LiveInstance->OldParameters.Get(), RootAddress, MipsToSkip);
 
 		if (!Result)
 		{
 			Result = BlackImage();
 		}
 
-		m_pD->WorkingMemoryManager.CurrentInstanceCache = nullptr;
-    	m_pD->WorkingMemoryManager.CurrentPassthroughObjectLoader = nullptr;
+		WorkingMemoryManager.CurrentInstanceCache = nullptr;
+    	WorkingMemoryManager.CurrentPassthroughObjectLoader = nullptr;
     	
 		return Result;
 	}
@@ -428,10 +380,10 @@ namespace UE::Mutable::Private
     	OP::ADDRESS RootAddress = ImageId.GetKey()->Address;
 
 		// Find the live instance
-		FLiveInstance* LiveInstance = m_pD->FindLiveInstance(instanceID);
+		FLiveInstance* LiveInstance = FindLiveInstance(instanceID);
 		check(LiveInstance);
 
-    	EOpType OpType = LiveInstance->Model->GetPrivate()->Program.GetOpType(RootAddress);
+    	EOpType OpType = LiveInstance->Model->Program.GetOpType(RootAddress);
     	if (GetOpDataType(OpType) != EDataType::Image)
     	{
     		return UE::Tasks::Launch(TEXT("FSystem::GetImageInvalid"),
@@ -451,15 +403,15 @@ namespace UE::Mutable::Private
     	TSharedRef<FContext> Context = MakeShared<FContext>();
 
     	UE::Tasks::FTask RunnerTask = UE::Tasks::Launch(TEXT("FSystem::GetImage"),
-		[Context, SystemPrivate = m_pD, RootAddress, LiveInstance, MipsToSkip]()
+		[Context, System = this, RootAddress, LiveInstance, MipsToSkip]()
 		{
-			SystemPrivate->WorkingMemoryManager.CurrentInstanceCache = LiveInstance->Cache;
-			SystemPrivate->WorkingMemoryManager.CurrentPassthroughObjectLoader = LiveInstance->PassthroughObjectLoader;
-			SystemPrivate->WorkingMemoryManager.BeginRunnerThread();
+			System->WorkingMemoryManager.CurrentInstanceCache = LiveInstance->Cache;
+			System->WorkingMemoryManager.CurrentPassthroughObjectLoader = LiveInstance->PassthroughObjectLoader;
+			System->WorkingMemoryManager.BeginRunnerThread();
 
 			Context->Runner = CodeRunner::Create(
-				SystemPrivate->Settings,
-				SystemPrivate,
+				System->Settings,
+				System,
 				EExecutionStrategy::MinimizeMemory,
 				LiveInstance->Model,
 				LiveInstance->OldParameters.Get(),
@@ -476,14 +428,14 @@ namespace UE::Mutable::Private
 		UE::Tasks::ETaskPriority::Inherit);
     	
 		return UE::Tasks::Launch(TEXT("FSystem::GetImageResultTask"),
-		[Context, SystemPrivate = m_pD, RootAddress, MipsToSkip]() -> TManagedPtr<const FImage>
+		[Context, System = this, RootAddress, MipsToSkip]() -> TManagedPtr<const FImage>
 		{
 			TManagedPtr<const FImage> Result;
 
-			SystemPrivate->bUnrecoverableError = Context->Runner->bUnrecoverableError;
+			System->bUnrecoverableError = Context->Runner->bUnrecoverableError;
 			if (!Context->Runner->bUnrecoverableError)
 			{
-				Result = SystemPrivate->WorkingMemoryManager.LoadImage(FCacheAddress(RootAddress, 0, MipsToSkip));
+				Result = System->WorkingMemoryManager.LoadImage(FCacheAddress(RootAddress, 0, MipsToSkip));
 			}
 
 			if (!Result)
@@ -491,9 +443,9 @@ namespace UE::Mutable::Private
 				Result = MakeManaged<FImage>(16, 16, 1, EImageFormat::RGBA_UByte, EInitializationType::Black);
 			}
 
-			SystemPrivate->WorkingMemoryManager.EndRunnerThread();
-			SystemPrivate->WorkingMemoryManager.CurrentInstanceCache = nullptr;
-			SystemPrivate->WorkingMemoryManager.CurrentPassthroughObjectLoader = nullptr;
+			System->WorkingMemoryManager.EndRunnerThread();
+			System->WorkingMemoryManager.CurrentInstanceCache = nullptr;
+			System->WorkingMemoryManager.CurrentPassthroughObjectLoader = nullptr;
 			
 			return Result;
 		},
@@ -524,33 +476,33 @@ namespace UE::Mutable::Private
 		FExtendedImageDesc Result;
 
 		// Find the live instance
-		FLiveInstance* LiveInstance = m_pD->FindLiveInstance(instanceID);
+		FLiveInstance* LiveInstance = FindLiveInstance(instanceID);
 		check(LiveInstance);
-		m_pD->WorkingMemoryManager.CurrentInstanceCache = LiveInstance->Cache;
-		m_pD->WorkingMemoryManager.CurrentPassthroughObjectLoader = LiveInstance->PassthroughObjectLoader;
+		WorkingMemoryManager.CurrentInstanceCache = LiveInstance->Cache;
+		WorkingMemoryManager.CurrentPassthroughObjectLoader = LiveInstance->PassthroughObjectLoader;
 
 		OP::ADDRESS RootAddress = ImageId.GetKey()->Address;
 
 		const UE::Mutable::Private::FModel* Model = LiveInstance->Model.Get();
-		const UE::Mutable::Private::FProgram& Program = Model->GetPrivate()->Program;
+		const UE::Mutable::Private::FProgram& Program = Model->Program;
 
 		// TODO: It should be possible to reuse this data if cleared in the correct places only, together with HeapImageDesc.
 		int32 VarValue = CVarClearImageDescCache.GetValueOnAnyThread();
 		if (VarValue != 0)
 		{
-			m_pD->WorkingMemoryManager.CurrentInstanceCache->ClearDescCache();
+			WorkingMemoryManager.CurrentInstanceCache->ClearDescCache();
 		}
 
 		UE::Mutable::Private::EOpType OpType = Program.GetOpType(RootAddress);
 		if (GetOpDataType(OpType) == EDataType::Image)
 		{
 			// GetImageDesc may call normal execution paths where meshes are computed.
-			m_pD->WorkingMemoryManager.BeginRunnerThread();
+			WorkingMemoryManager.BeginRunnerThread();
 					
 			int8 executionOptions = 0;
 			TSharedRef<CodeRunner> Runner = CodeRunner::Create(
-					m_pD->Settings,
-					m_pD,
+					Settings,
+					this,
 					EExecutionStrategy::MinimizeMemory,
 					LiveInstance->Model,
 					LiveInstance->OldParameters.Get(),
@@ -566,11 +518,11 @@ namespace UE::Mutable::Private
 
 			Result = Runner->GetImageDescResult(RootAddress);
 
-			m_pD->WorkingMemoryManager.EndRunnerThread();
+			WorkingMemoryManager.EndRunnerThread();
 		}
 
-		m_pD->WorkingMemoryManager.CurrentInstanceCache = nullptr;
-		m_pD->WorkingMemoryManager.CurrentPassthroughObjectLoader = nullptr;
+		WorkingMemoryManager.CurrentInstanceCache = nullptr;
+		WorkingMemoryManager.CurrentPassthroughObjectLoader = nullptr;
 
 		return Result;
 	}
@@ -595,10 +547,10 @@ namespace UE::Mutable::Private
 		
 		OP::ADDRESS RootAddress = ImageId.GetKey()->Address;
 		
-		FLiveInstance* LiveInstance = m_pD->FindLiveInstance(InstanceID);
+		FLiveInstance* LiveInstance = FindLiveInstance(InstanceID);
 		check(LiveInstance);
 
-		EOpType OpType = LiveInstance->Model->GetPrivate()->Program.GetOpType(RootAddress);
+		EOpType OpType = LiveInstance->Model->Program.GetOpType(RootAddress);
 		if (GetOpDataType(OpType) != EDataType::Image)
 		{
 			return UE::Tasks::Launch(TEXT("FSystem::GetImageDescInvalid"),
@@ -618,25 +570,25 @@ namespace UE::Mutable::Private
 		TSharedRef<FContext> Context = MakeShared<FContext>();
 		
 		UE::Tasks::FTask RunnerTask = UE::Tasks::Launch(TEXT("FSystem::GetImageDesc"),
-		[Context, SystemPrivate = m_pD, LiveInstance, RootAddress]()
+		[Context, System = this, LiveInstance, RootAddress]()
 		{
-			SystemPrivate->WorkingMemoryManager.CurrentInstanceCache = LiveInstance->Cache;
-			SystemPrivate->WorkingMemoryManager.CurrentPassthroughObjectLoader = LiveInstance->PassthroughObjectLoader;
+			System->WorkingMemoryManager.CurrentInstanceCache = LiveInstance->Cache;
+			System->WorkingMemoryManager.CurrentPassthroughObjectLoader = LiveInstance->PassthroughObjectLoader;
 				
 			// TODO: It should be possible to reuse this data if cleared in the correct places only, together with HeapImageDesc.
 			int32 VarValue = CVarClearImageDescCache.GetValueOnAnyThread();
 			if (VarValue != 0)
 			{
-				SystemPrivate->WorkingMemoryManager.CurrentInstanceCache->ClearDescCache();
+				System->WorkingMemoryManager.CurrentInstanceCache->ClearDescCache();
 			}
 
-			SystemPrivate->WorkingMemoryManager.BeginRunnerThread();
+			System->WorkingMemoryManager.BeginRunnerThread();
 
 			// GetImageDesc may call normal execution paths where meshes are computed.
 			int8 ExecutionOptions = 0;
 			Context->Runner = CodeRunner::Create(
-				SystemPrivate->Settings,
-				SystemPrivate,
+				System->Settings,
+				System,
 				EExecutionStrategy::MinimizeMemory,
 				LiveInstance->Model,
 				LiveInstance->OldParameters.Get(),
@@ -654,13 +606,13 @@ namespace UE::Mutable::Private
 		UE::Tasks::ETaskPriority::Inherit);
 		
 		return UE::Tasks::Launch(TEXT("FSystem::GetImageDescResultTask"),
-		[Context, SystemPrivate = m_pD, RootAddress]() -> FExtendedImageDesc
+		[Context, System = this, RootAddress]() -> FExtendedImageDesc
 		{
 			FExtendedImageDesc Result = Context->Runner->GetImageDescResult(RootAddress);
 
-			SystemPrivate->WorkingMemoryManager.EndRunnerThread();
-			SystemPrivate->WorkingMemoryManager.CurrentInstanceCache = nullptr;
-			SystemPrivate->WorkingMemoryManager.CurrentPassthroughObjectLoader = nullptr;
+			System->WorkingMemoryManager.EndRunnerThread();
+			System->WorkingMemoryManager.CurrentInstanceCache = nullptr;
+			System->WorkingMemoryManager.CurrentPassthroughObjectLoader = nullptr;
 
 			return Result;
 		},
@@ -682,13 +634,13 @@ namespace UE::Mutable::Private
 		}
 
 		// Find the live instance
-		FLiveInstance* LiveInstance = m_pD->FindLiveInstance(instanceID);
+		FLiveInstance* LiveInstance = FindLiveInstance(instanceID);
 		check(LiveInstance);
-		m_pD->WorkingMemoryManager.CurrentInstanceCache = LiveInstance->Cache;
-		m_pD->WorkingMemoryManager.CurrentPassthroughObjectLoader = LiveInstance->PassthroughObjectLoader;
+		WorkingMemoryManager.CurrentInstanceCache = LiveInstance->Cache;
+		WorkingMemoryManager.CurrentPassthroughObjectLoader = LiveInstance->PassthroughObjectLoader;
 
 		OP::ADDRESS RootAddress = MeshId.GetKey()->Address;
-		TManagedPtr<const FMesh> Result = m_pD->BuildMesh(
+		TManagedPtr<const FMesh> Result = BuildMesh(
 			LiveInstance->ExternalResourceProvider,
 			LiveInstance->Model,
 			LiveInstance->OldParameters.Get(),
@@ -701,8 +653,8 @@ namespace UE::Mutable::Private
 			Result = MakeManaged<FMesh>();
 		}
 
-		m_pD->WorkingMemoryManager.CurrentInstanceCache = nullptr;
-		m_pD->WorkingMemoryManager.CurrentPassthroughObjectLoader = nullptr;
+		WorkingMemoryManager.CurrentInstanceCache = nullptr;
+		WorkingMemoryManager.CurrentPassthroughObjectLoader = nullptr;
 		
 		return Result;
 	}
@@ -728,10 +680,10 @@ namespace UE::Mutable::Private
 		OP::ADDRESS RootAddress = MeshId.GetKey()->Address;
 
 		// Find the live instance
-		FLiveInstance* LiveInstance = m_pD->FindLiveInstance(instanceID);
+		FLiveInstance* LiveInstance = FindLiveInstance(instanceID);
 		check(LiveInstance);
 		
-		EOpType OpType = LiveInstance->Model->GetPrivate()->Program.GetOpType(RootAddress);
+		EOpType OpType = LiveInstance->Model->Program.GetOpType(RootAddress);
 		if (GetOpDataType(OpType) != EDataType::Mesh)
 		{
 			return UE::Tasks::Launch(TEXT("FSystem::GetMeshInvalid"),
@@ -751,16 +703,16 @@ namespace UE::Mutable::Private
 		TSharedRef<FContext> Context = MakeShared<FContext>();
 		
 		UE::Tasks::FTask RunnerTask = UE::Tasks::Launch(TEXT("FSystem::GetMesh"),
-		[Context, SystemPrivate = m_pD, RootAddress, LiveInstance, MeshContentFilter]()
+		[Context, System = this, RootAddress, LiveInstance, MeshContentFilter]()
 		{
-			SystemPrivate->WorkingMemoryManager.CurrentInstanceCache = LiveInstance->Cache;
-			SystemPrivate->WorkingMemoryManager.CurrentPassthroughObjectLoader = LiveInstance->PassthroughObjectLoader;
-			SystemPrivate->WorkingMemoryManager.BeginRunnerThread();
+			System->WorkingMemoryManager.CurrentInstanceCache = LiveInstance->Cache;
+			System->WorkingMemoryManager.CurrentPassthroughObjectLoader = LiveInstance->PassthroughObjectLoader;
+			System->WorkingMemoryManager.BeginRunnerThread();
 			
 			uint8 ExecutionOptions = static_cast<uint8>(MeshContentFilter);
 			Context->Runner = CodeRunner::Create(
-					SystemPrivate->Settings,
-					SystemPrivate, 
+					System->Settings,
+					System, 
 					EExecutionStrategy::MinimizeMemory, 
 					LiveInstance->Model, 
 					LiveInstance->OldParameters.Get(), 
@@ -778,19 +730,19 @@ namespace UE::Mutable::Private
 		UE::Tasks::ETaskPriority::Inherit);
 		
 		return UE::Tasks::Launch(TEXT("FSystem::GetMeshResultTask"),
-		[Context, SystemPrivate = m_pD, RootAddress, MeshContentFilter]() -> TManagedPtr<const FMesh>
+		[Context, System = this, RootAddress, MeshContentFilter]() -> TManagedPtr<const FMesh>
 		{
 			TManagedPtr<const FMesh> Result;
 
-			SystemPrivate->bUnrecoverableError = Context->Runner->bUnrecoverableError;
+			System->bUnrecoverableError = Context->Runner->bUnrecoverableError;
 			if (!Context->Runner->bUnrecoverableError)
 			{
-				Result = SystemPrivate->WorkingMemoryManager.LoadMesh(FCacheAddress(RootAddress, 0, static_cast<uint8>(MeshContentFilter)));
+				Result = System->WorkingMemoryManager.LoadMesh(FCacheAddress(RootAddress, 0, static_cast<uint8>(MeshContentFilter)));
 			}
 			
-			SystemPrivate->WorkingMemoryManager.EndRunnerThread();
-			SystemPrivate->WorkingMemoryManager.CurrentInstanceCache = nullptr;
-			SystemPrivate->WorkingMemoryManager.CurrentPassthroughObjectLoader = nullptr;
+			System->WorkingMemoryManager.EndRunnerThread();
+			System->WorkingMemoryManager.CurrentInstanceCache = nullptr;
+			System->WorkingMemoryManager.CurrentPassthroughObjectLoader = nullptr;
 			
 			return Result;
 		},
@@ -820,10 +772,10 @@ namespace UE::Mutable::Private
 		OP::ADDRESS RootAddress = SkeletalMeshId.GetKey()->Address;
 		
 		// Find the live instance
-		FLiveInstance* LiveInstance = m_pD->FindLiveInstance(InstanceID);
+		FLiveInstance* LiveInstance = FindLiveInstance(InstanceID);
 		check(LiveInstance);
 		
-		UE::Mutable::Private::EOpType OpType = LiveInstance->Model->GetPrivate()->Program.GetOpType(RootAddress);
+		UE::Mutable::Private::EOpType OpType = LiveInstance->Model->Program.GetOpType(RootAddress);
 		if (GetOpDataType(OpType) != EDataType::SkeletalMesh)
 		{
 			return UE::Tasks::Launch(TEXT("FSystem::GetSkeletalMeshInvalid"),
@@ -843,15 +795,15 @@ namespace UE::Mutable::Private
 		TSharedRef<FContext> Context = MakeShared<FContext>();
 		
 		UE::Tasks::FTask RunnerTask = UE::Tasks::Launch(TEXT("FSystem::GetSkeletalMesh"),
-		[Context, SystemPrivate = m_pD, RootAddress, LiveInstance, SkeletalMeshContentFilter]()
+		[Context, System = this, RootAddress, LiveInstance, SkeletalMeshContentFilter]()
 		{
-			SystemPrivate->WorkingMemoryManager.CurrentInstanceCache = LiveInstance->Cache;
-			SystemPrivate->WorkingMemoryManager.CurrentPassthroughObjectLoader = LiveInstance->PassthroughObjectLoader;
-			SystemPrivate->WorkingMemoryManager.BeginRunnerThread();
+			System->WorkingMemoryManager.CurrentInstanceCache = LiveInstance->Cache;
+			System->WorkingMemoryManager.CurrentPassthroughObjectLoader = LiveInstance->PassthroughObjectLoader;
+			System->WorkingMemoryManager.BeginRunnerThread();
 
 			Context->Runner = CodeRunner::Create(
-					SystemPrivate->Settings,
-					SystemPrivate, 
+					System->Settings,
+					System, 
 					EExecutionStrategy::MinimizeMemory, 
 					LiveInstance->Model, 
 					LiveInstance->OldParameters.Get(), 
@@ -869,19 +821,19 @@ namespace UE::Mutable::Private
 		UE::Tasks::ETaskPriority::Inherit);
 		
 		return UE::Tasks::Launch(TEXT("FSystem::GetSkeletalMeshResultTask"),
-		[Context, SystemPrivate = m_pD, RootAddress, SkeletalMeshContentFilter]() -> TManagedPtr<const FSkeletalMesh>
+		[Context, System = this, RootAddress, SkeletalMeshContentFilter]() -> TManagedPtr<const FSkeletalMesh>
 		{
 			TManagedPtr<const FSkeletalMesh> Result;
 
-			SystemPrivate->bUnrecoverableError = Context->Runner->bUnrecoverableError;
+			System->bUnrecoverableError = Context->Runner->bUnrecoverableError;
 			if (!Context->Runner->bUnrecoverableError)
 			{
-				Result = SystemPrivate->WorkingMemoryManager.LoadSkeletalMesh(FCacheAddress(RootAddress, 0, static_cast<uint8>(SkeletalMeshContentFilter)));
+				Result = System->WorkingMemoryManager.LoadSkeletalMesh(FCacheAddress(RootAddress, 0, static_cast<uint8>(SkeletalMeshContentFilter)));
 			}
 			
-			SystemPrivate->WorkingMemoryManager.EndRunnerThread();
-			SystemPrivate->WorkingMemoryManager.CurrentInstanceCache = nullptr;
-			SystemPrivate->WorkingMemoryManager.CurrentPassthroughObjectLoader = nullptr;
+			System->WorkingMemoryManager.EndRunnerThread();
+			System->WorkingMemoryManager.CurrentInstanceCache = nullptr;
+			System->WorkingMemoryManager.CurrentPassthroughObjectLoader = nullptr;
 			
 			return Result;
 		},
@@ -911,10 +863,10 @@ namespace UE::Mutable::Private
 		OP::ADDRESS RootAddress = MaterialId.GetKey()->Address;
 	
 		// Find the live instance
-		FLiveInstance* LiveInstance = m_pD->FindLiveInstance(InstanceID);
+		FLiveInstance* LiveInstance = FindLiveInstance(InstanceID);
 		check(LiveInstance);
 		
-		EOpType OpType = LiveInstance->Model->GetPrivate()->Program.GetOpType(RootAddress);
+		EOpType OpType = LiveInstance->Model->Program.GetOpType(RootAddress);
 		if (GetOpDataType(OpType) != EDataType::Material)
 		{
 			return UE::Tasks::Launch(TEXT("FSystem::GetMaterialInvalid"),
@@ -934,15 +886,15 @@ namespace UE::Mutable::Private
 		TSharedRef<FContext> Context = MakeShared<FContext>();
 		
 		UE::Tasks::FTask RunnerTask = UE::Tasks::Launch(TEXT("FSystem::GetMaterial"),
-		[Context, SystemPrivate = m_pD, RootAddress, LiveInstance]()
+		[Context, System = this, RootAddress, LiveInstance]()
 		{
-			SystemPrivate->WorkingMemoryManager.CurrentInstanceCache = LiveInstance->Cache;
-			SystemPrivate->WorkingMemoryManager.CurrentPassthroughObjectLoader = LiveInstance->PassthroughObjectLoader;
-			SystemPrivate->WorkingMemoryManager.BeginRunnerThread();
+			System->WorkingMemoryManager.CurrentInstanceCache = LiveInstance->Cache;
+			System->WorkingMemoryManager.CurrentPassthroughObjectLoader = LiveInstance->PassthroughObjectLoader;
+			System->WorkingMemoryManager.BeginRunnerThread();
 				
 			Context->Runner = CodeRunner::Create(
-					SystemPrivate->Settings,
-					SystemPrivate, 
+					System->Settings,
+					System, 
 					EExecutionStrategy::MinimizeMemory, 
 					LiveInstance->Model, 
 					LiveInstance->OldParameters.Get(), 
@@ -960,19 +912,19 @@ namespace UE::Mutable::Private
 		UE::Tasks::ETaskPriority::Inherit);
 		
 		return UE::Tasks::Launch(TEXT("FSystem::GetMaterialResultTask"),
-		[Context, SystemPrivate = m_pD, RootAddress]() -> TManagedPtr<const FMaterial>
+		[Context, System = this, RootAddress]() -> TManagedPtr<const FMaterial>
 		{
 			TManagedPtr<const FMaterial> Result;
 
-			SystemPrivate->bUnrecoverableError = Context->Runner->bUnrecoverableError;
+			System->bUnrecoverableError = Context->Runner->bUnrecoverableError;
 			if (!Context->Runner->bUnrecoverableError)
 			{
-				Result = SystemPrivate->WorkingMemoryManager.LoadMaterial(FCacheAddress(RootAddress, 0, 0));
+				Result = System->WorkingMemoryManager.LoadMaterial(FCacheAddress(RootAddress, 0, 0));
 			}
 			
-			SystemPrivate->WorkingMemoryManager.EndRunnerThread();
-			SystemPrivate->WorkingMemoryManager.CurrentInstanceCache = nullptr;
-			SystemPrivate->WorkingMemoryManager.CurrentPassthroughObjectLoader = nullptr;
+			System->WorkingMemoryManager.EndRunnerThread();
+			System->WorkingMemoryManager.CurrentInstanceCache = nullptr;
+			System->WorkingMemoryManager.CurrentPassthroughObjectLoader = nullptr;
 			
 			return Result;
 		},
@@ -987,7 +939,7 @@ namespace UE::Mutable::Private
 		LLM_SCOPE_BYNAME(TEXT("MutableRuntime"));
 		MUTABLE_CPUPROFILER_SCOPE(EndUpdate);
 
-		FLiveInstance* LiveInstance = m_pD->FindLiveInstance(instanceID);
+		FLiveInstance* LiveInstance = FindLiveInstance(instanceID);
 		if (LiveInstance)
 		{
 			LiveInstance->Instance = nullptr;
@@ -995,74 +947,71 @@ namespace UE::Mutable::Private
 			// Debug check to see if we managed the op-hit-counts correctly
 			LiveInstance->Cache->CheckHitCountsCleared();
 
-			m_pD->WorkingMemoryManager.CurrentInstanceCache = LiveInstance->Cache;
-			m_pD->WorkingMemoryManager.CurrentPassthroughObjectLoader = LiveInstance->PassthroughObjectLoader;
+			WorkingMemoryManager.CurrentInstanceCache = LiveInstance->Cache;
+			WorkingMemoryManager.CurrentPassthroughObjectLoader = LiveInstance->PassthroughObjectLoader;
 
 			// We don't want to clear the cache layer 1 because it contains data that can be useful for a 
 			// future update (same states, just runtime parameters changed).
-			//m_pD->WorkingMemoryManager.ClearCacheLayer1();
+			//WorkingMemoryManager.ClearCacheLayer1();
 
 			// We need to clear the layer 0 cache, because it contains data that is only valid for the current 
 			// parameter values (unless it is data marked as state cache)
-			m_pD->WorkingMemoryManager.ClearCacheLayer0();
+			WorkingMemoryManager.ClearCacheLayer0();
 
-			m_pD->WorkingMemoryManager.CurrentInstanceCache = nullptr;
-			m_pD->WorkingMemoryManager.CurrentPassthroughObjectLoader = nullptr;
+			WorkingMemoryManager.CurrentInstanceCache = nullptr;
+			WorkingMemoryManager.CurrentPassthroughObjectLoader = nullptr;
 		}
 
 		// Reduce the cache until it fits the limit.
-		m_pD->WorkingMemoryManager.EnsureBudgetBelow(0);
+		WorkingMemoryManager.EnsureBudgetBelow(0);
 
-		m_pD->UpdateStats();
+		UpdateStats();
 	}
 
-
-    //---------------------------------------------------------------------------------------------
+	
     void FSystem::ReleaseInstance( FInstance::FID instanceID )
     {
 		LLM_SCOPE_BYNAME(TEXT("MutableRuntime"));
 		MUTABLE_CPUPROFILER_SCOPE(ReleaseInstance);
 
-		for (int32 Index = 0; Index < m_pD->WorkingMemoryManager.LiveInstances.Num(); ++Index)
+		for (int32 Index = 0; Index < WorkingMemoryManager.LiveInstances.Num(); ++Index)
 		{
-			FLiveInstance& Instance = m_pD->WorkingMemoryManager.LiveInstances[Index];
+			FLiveInstance& Instance = WorkingMemoryManager.LiveInstances[Index];
 			if (Instance.InstanceID == instanceID)
 			{
-				m_pD->WorkingMemoryManager.LiveInstances.RemoveAtSwap(Index);
+				WorkingMemoryManager.LiveInstances.RemoveAtSwap(Index);
 				break;
 			}
 		}
 
- 		int32 Removed = m_pD->WorkingMemoryManager.LiveInstances.RemoveAllSwap(
+ 		int32 Removed = WorkingMemoryManager.LiveInstances.RemoveAllSwap(
 			[instanceID](const FLiveInstance& Instance)
 			{
 				return (Instance.InstanceID == instanceID);
 			});
 
-		TRACE_COUNTER_SET(MutableRuntime_LiveInstances, m_pD->WorkingMemoryManager.LiveInstances.Num());
-
+		TRACE_COUNTER_SET(MutableRuntime_LiveInstances, WorkingMemoryManager.LiveInstances.Num());
 	}
 
-
-    //---------------------------------------------------------------------------------------------
+	
     class RelevantParameterVisitor : public UniqueDiscreteCoveredCodeVisitor<>
     {
     public:
 
         RelevantParameterVisitor(
-                FSystem::Private* InSystem,
+                FSystem* InSystem,
 				const TSharedPtr<const FModel>& InModel,
                 const TSharedPtr<const FParameters>& InParams,
                 TArray<uint32>& OutRelevantParameters) :
 			UniqueDiscreteCoveredCodeVisitor<>(InSystem, InModel, InParams, FSystem::AllLODs),
 		    RelevantParameters(OutRelevantParameters)
         {
-            OP::ADDRESS At = InModel->GetPrivate()->Program.States[0].Root;
+            OP::ADDRESS At = InModel->Program.States[0].Root;
 
             Run(At);
         }
 
-        bool Visit(OP::ADDRESS At, FProgram& Program) override
+        bool Visit(OP::ADDRESS At, const FProgram& Program) override
         {
             switch (Program.GetOpType(At))
             {
@@ -1093,8 +1042,7 @@ namespace UE::Mutable::Private
         TArray<uint32>& RelevantParameters;
     };
 
-
-    //---------------------------------------------------------------------------------------------
+	
     void FSystem::GetParameterRelevancy(FInstance::FID InstanceID,
     	const TSharedPtr<const FParameters>& FParameters,
         TArray<uint32>& OutRelevantParameters)
@@ -1102,20 +1050,19 @@ namespace UE::Mutable::Private
 		LLM_SCOPE_BYNAME(TEXT("MutableRuntime"));
 
 		// Find the live instance
-		FLiveInstance* LiveInstance = m_pD->FindLiveInstance(InstanceID);
+		FLiveInstance* LiveInstance = FindLiveInstance(InstanceID);
 		check(LiveInstance);
-		m_pD->WorkingMemoryManager.CurrentInstanceCache = LiveInstance->Cache;
-		m_pD->WorkingMemoryManager.CurrentPassthroughObjectLoader = LiveInstance->PassthroughObjectLoader;
+		WorkingMemoryManager.CurrentInstanceCache = LiveInstance->Cache;
+		WorkingMemoryManager.CurrentPassthroughObjectLoader = LiveInstance->PassthroughObjectLoader;
 		
-		RelevantParameterVisitor visitor(m_pD, LiveInstance->Model, FParameters, OutRelevantParameters);
+		RelevantParameterVisitor visitor(this, LiveInstance->Model, FParameters, OutRelevantParameters);
 
-		m_pD->WorkingMemoryManager.CurrentInstanceCache = nullptr;
-		m_pD->WorkingMemoryManager.CurrentPassthroughObjectLoader = nullptr;
+		WorkingMemoryManager.CurrentInstanceCache = nullptr;
+		WorkingMemoryManager.CurrentPassthroughObjectLoader = nullptr;
     }
 
-
-    //---------------------------------------------------------------------------------------------
-	inline FLiveInstance* FSystem::Private::FindLiveInstance(FInstance::FID id)
+	
+	inline FLiveInstance* FSystem::FindLiveInstance(FInstance::FID id)
 	{
 		for (int32 i = 0; i < WorkingMemoryManager.LiveInstances.Num(); ++i)
 		{
@@ -1128,8 +1075,7 @@ namespace UE::Mutable::Private
 	}
 
 
-	//---------------------------------------------------------------------------------------------
-	bool FSystem::Private::CheckUpdatedParameters( const FLiveInstance* LiveInstance,
+	bool FSystem::CheckUpdatedParameters( const FLiveInstance* LiveInstance,
                                  const TSharedPtr<const FParameters>& Params,
                                  uint64& UpdatedParameters)
     {
@@ -1143,7 +1089,7 @@ namespace UE::Mutable::Private
 
         // check what parameters have changed
 		UpdatedParameters = 0;
-        const FProgram& Program = LiveInstance->Model->GetPrivate()->Program;
+        const FProgram& Program = LiveInstance->Model->Program;
         const TArray<int>& runtimeParams = Program.States[ LiveInstance->State ].m_runtimeParameters;
 
         check( Params->GetCount() == (int)Program.Parameters.Num() );
@@ -1174,8 +1120,7 @@ namespace UE::Mutable::Private
     }
 
 
-	//---------------------------------------------------------------------------------------------
-	void FSystem::Private::BeginBuild(const TSharedPtr<const FModel>& InModel,
+	void FSystem::BeginBuild(const TSharedPtr<const FModel>& InModel,
 		const TSharedPtr<FMeshIdRegistry>& InMeshIdRegistry,
 		const TSharedPtr<FImageIdRegistry>& InImageIdRegistry,
 		const TSharedPtr<FMaterialIdRegistry>& InMaterialIdRegistry,
@@ -1184,7 +1129,7 @@ namespace UE::Mutable::Private
 		// We don't have a FLiveInstance, let's create the memory
 		// \TODO: There is no clear moment to remove this... EndBuild?
 		WorkingMemoryManager.CurrentInstanceCache = MakeShared<FProgramCache>();
-		WorkingMemoryManager.CurrentInstanceCache->Init(InModel->GetPrivate()->Program.OpAddress.Num());
+		WorkingMemoryManager.CurrentInstanceCache->Init(InModel->Program.OpAddress.Num());
 
     	WorkingMemoryManager.CurrentPassthroughObjectLoader = nullptr;
     	
@@ -1200,8 +1145,7 @@ namespace UE::Mutable::Private
 	}
 
 
-	//---------------------------------------------------------------------------------------------
-	void FSystem::Private::EndBuild()
+	void FSystem::EndBuild()
 	{
 		WorkingMemoryManager.CurrentInstanceCache = nullptr;
 		WorkingMemoryManager.CurrentPassthroughObjectLoader = nullptr;
@@ -1212,8 +1156,7 @@ namespace UE::Mutable::Private
 	}
 
 
-	//---------------------------------------------------------------------------------------------
-	void FSystem::Private::RunCode(
+	void FSystem::RunCode(
 		const TSharedPtr<FExternalResourceProvider>& ExternalResourceProvider,
 		const TSharedPtr<const FModel>& InModel,
 		const FParameters* InParameters,
@@ -1241,8 +1184,7 @@ namespace UE::Mutable::Private
 	}
 
 
-	//---------------------------------------------------------------------------------------------
-	bool FSystem::Private::BuildBool(const TSharedPtr<FExternalResourceProvider>& ExternalResourceProvider, const TSharedPtr<const FModel>& pModel, const FParameters* Params, OP::ADDRESS at)
+	bool FSystem::BuildBool(const TSharedPtr<FExternalResourceProvider>& ExternalResourceProvider, const TSharedPtr<const FModel>& pModel, const FParameters* Params, OP::ADDRESS at)
 	{
 		WorkingMemoryManager.BeginRunnerThread();
 
@@ -1260,8 +1202,7 @@ namespace UE::Mutable::Private
 	}
 
 
-	//---------------------------------------------------------------------------------------------
-	float FSystem::Private::BuildScalar(const TSharedPtr<FExternalResourceProvider>& ExternalResourceProvider, const TSharedPtr<const FModel>& pModel, const FParameters* Params, OP::ADDRESS at)
+	float FSystem::BuildScalar(const TSharedPtr<FExternalResourceProvider>& ExternalResourceProvider, const TSharedPtr<const FModel>& pModel, const FParameters* Params, OP::ADDRESS at)
 	{
 		WorkingMemoryManager.BeginRunnerThread();
 
@@ -1279,8 +1220,7 @@ namespace UE::Mutable::Private
 	}
 
 
-	//---------------------------------------------------------------------------------------------
-	int32 FSystem::Private::BuildInt(const TSharedPtr<FExternalResourceProvider>& ExternalResourceProvider, const TSharedPtr<const FModel>& pModel, const FParameters* Params, OP::ADDRESS at)
+	int32 FSystem::BuildInt(const TSharedPtr<FExternalResourceProvider>& ExternalResourceProvider, const TSharedPtr<const FModel>& pModel, const FParameters* Params, OP::ADDRESS at)
 	{
 		WorkingMemoryManager.BeginRunnerThread();
 
@@ -1298,14 +1238,13 @@ namespace UE::Mutable::Private
 	}
 
 
-	//---------------------------------------------------------------------------------------------
-	FVector4f FSystem::Private::BuildColour(const TSharedPtr<FExternalResourceProvider>& ExternalResourceProvider, const TSharedPtr<const FModel>& pModel, const FParameters* Params, OP::ADDRESS at)
+	FVector4f FSystem::BuildColour(const TSharedPtr<FExternalResourceProvider>& ExternalResourceProvider, const TSharedPtr<const FModel>& pModel, const FParameters* Params, OP::ADDRESS at)
 	{
 		WorkingMemoryManager.BeginRunnerThread();
 
 		FVector4f Result(0,0,0,1);
 
-		UE::Mutable::Private::EOpType OpType = pModel->GetPrivate()->Program.GetOpType(at);
+		UE::Mutable::Private::EOpType OpType = pModel->Program.GetOpType(at);
 		if (GetOpDataType(OpType) == EDataType::Color)
 		{
 			RunCode(ExternalResourceProvider, pModel, Params, at);
@@ -1321,8 +1260,7 @@ namespace UE::Mutable::Private
 	}
 
 	
-	//---------------------------------------------------------------------------------------------
-	FProjector FSystem::Private::BuildProjector(const TSharedPtr<FExternalResourceProvider>& ExternalResourceProvider, const TSharedPtr<const FModel>& pModel, const FParameters* Params, OP::ADDRESS at)
+	FProjector FSystem::BuildProjector(const TSharedPtr<FExternalResourceProvider>& ExternalResourceProvider, const TSharedPtr<const FModel>& pModel, const FParameters* Params, OP::ADDRESS at)
 	{
 		WorkingMemoryManager.BeginRunnerThread();
 
@@ -1340,13 +1278,13 @@ namespace UE::Mutable::Private
 	}
 
 	
-	TManagedPtr<const FImage> FSystem::Private::BuildImage(const TSharedPtr<FExternalResourceProvider>& ExternalResourceProvider, const TSharedPtr<const FModel>& pModel, const FParameters* Params, OP::ADDRESS at, int32 MipsToSkip)
+	TManagedPtr<const FImage> FSystem::BuildImage(const TSharedPtr<FExternalResourceProvider>& ExternalResourceProvider, const TSharedPtr<const FModel>& pModel, const FParameters* Params, OP::ADDRESS at, int32 MipsToSkip)
 	{
 		WorkingMemoryManager.BeginRunnerThread();
 
 		TManagedPtr<const FImage> Result;
 
-		UE::Mutable::Private::EOpType OpType = pModel->GetPrivate()->Program.GetOpType(at);
+		UE::Mutable::Private::EOpType OpType = pModel->Program.GetOpType(at);
 		if (GetOpDataType(OpType) == EDataType::Image)
 		{
 			RunCode(ExternalResourceProvider, pModel, Params, at, FSystem::AllLODs, uint8(MipsToSkip));
@@ -1362,13 +1300,13 @@ namespace UE::Mutable::Private
 	}
 
 
-	TManagedPtr<const FMesh> FSystem::Private::BuildMesh(const TSharedPtr<FExternalResourceProvider>& ExternalResourceProvider, const TSharedPtr<const FModel>& Model, const FParameters* Params, OP::ADDRESS RootAddress, EMeshContentFlags MeshContentFilter)
+	TManagedPtr<const FMesh> FSystem::BuildMesh(const TSharedPtr<FExternalResourceProvider>& ExternalResourceProvider, const TSharedPtr<const FModel>& Model, const FParameters* Params, OP::ADDRESS RootAddress, EMeshContentFlags MeshContentFilter)
 	{
 		WorkingMemoryManager.BeginRunnerThread();
 
 		TManagedPtr<const FMesh> Result;
 
-		UE::Mutable::Private::EOpType OpType = Model->GetPrivate()->Program.GetOpType(RootAddress);
+		UE::Mutable::Private::EOpType OpType = Model->Program.GetOpType(RootAddress);
 
 		if (GetOpDataType(OpType) == EDataType::Mesh)
 		{
@@ -1386,13 +1324,13 @@ namespace UE::Mutable::Private
 	}
 
 
-	TManagedPtr<const FInstance> FSystem::Private::BuildInstance(const TSharedPtr<FExternalResourceProvider>& ExternalResourceProvider, const TSharedPtr<const FModel>& pModel, const FParameters* Params, OP::ADDRESS at)
+	TManagedPtr<const FInstance> FSystem::BuildInstance(const TSharedPtr<FExternalResourceProvider>& ExternalResourceProvider, const TSharedPtr<const FModel>& pModel, const FParameters* Params, OP::ADDRESS at)
 	{
 		WorkingMemoryManager.BeginRunnerThread();
 
 		TManagedPtr<const FInstance> Result;
 
-		UE::Mutable::Private::EOpType OpType = pModel->GetPrivate()->Program.GetOpType(at);
+		UE::Mutable::Private::EOpType OpType = pModel->Program.GetOpType(at);
 		if (GetOpDataType(OpType) == EDataType::Instance)
 		{
 			RunCode(ExternalResourceProvider, pModel, Params, at);
@@ -1408,15 +1346,15 @@ namespace UE::Mutable::Private
 	}
 
 
-	TManagedPtr<const FLayout> FSystem::Private::BuildLayout(const TSharedPtr<FExternalResourceProvider>& ExternalResourceProvider, const TSharedPtr<const FModel>& pModel, const FParameters* Params, OP::ADDRESS at)
+	TManagedPtr<const FLayout> FSystem::BuildLayout(const TSharedPtr<FExternalResourceProvider>& ExternalResourceProvider, const TSharedPtr<const FModel>& pModel, const FParameters* Params, OP::ADDRESS at)
 	{
 		WorkingMemoryManager.BeginRunnerThread();
 
 		TManagedPtr<const FLayout> Result;
 
-		if (pModel->GetPrivate()->Program.States[0].Root)
+		if (pModel->Program.States[0].Root)
 		{
-			UE::Mutable::Private::EOpType OpType = pModel->GetPrivate()->Program.GetOpType(at);
+			UE::Mutable::Private::EOpType OpType = pModel->Program.GetOpType(at);
 			if (GetOpDataType(OpType) == EDataType::Layout)
 			{
 				RunCode(ExternalResourceProvider, pModel, Params, at);
@@ -1433,16 +1371,15 @@ namespace UE::Mutable::Private
 	}
 
 
-	//---------------------------------------------------------------------------------------------
-	TManagedPtr<const String> FSystem::Private::BuildString(const TSharedPtr<FExternalResourceProvider>& ExternalResourceProvider, const TSharedPtr<const FModel>& pModel, const FParameters* Params, OP::ADDRESS at)
+	TManagedPtr<const String> FSystem::BuildString(const TSharedPtr<FExternalResourceProvider>& ExternalResourceProvider, const TSharedPtr<const FModel>& pModel, const FParameters* Params, OP::ADDRESS at)
 	{
 		WorkingMemoryManager.BeginRunnerThread();
 
 		TManagedPtr<const String> Result;
 
-		if (pModel->GetPrivate()->Program.States[0].Root)
+		if (pModel->Program.States[0].Root)
 		{
-			UE::Mutable::Private::EOpType OpType = pModel->GetPrivate()->Program.GetOpType(at);
+			UE::Mutable::Private::EOpType OpType = pModel->Program.GetOpType(at);
 			if (GetOpDataType(OpType) == EDataType::String)
 			{
 				RunCode(ExternalResourceProvider, pModel, Params, at);
@@ -1459,12 +1396,11 @@ namespace UE::Mutable::Private
 	}
 
 
-	//---------------------------------------------------------------------------------------------
-	void FSystem::Private::PrepareCache( const FModel* InModel, int32 InState)
+	void FSystem::PrepareCache( const FModel* InModel, int32 InState)
 	{
 		MUTABLE_CPUPROFILER_SCOPE(PrepareCache);
 
-		FProgram& Program = InModel->GetPrivate()->Program;
+		const FProgram& Program = InModel->Program;
 		int32 OpCount = Program.OpAddress.Num();
 		WorkingMemoryManager.CurrentInstanceCache->Init(OpCount);
 
@@ -1488,8 +1424,7 @@ namespace UE::Mutable::Private
 	}
 
 
-	//---------------------------------------------------------------------------------------------
-	void FSystem::Private::UpdateStats()
+	void FSystem::UpdateStats()
 	{
 		// Updater stats
 		int32 WorkingMemoryKb = WorkingMemoryManager.BudgetBytes / 1024;
@@ -1503,7 +1438,41 @@ namespace UE::Mutable::Private
 	}
 
 
-	//---------------------------------------------------------------------------------------------
+	int32 FWorkingMemoryManager::GetRomBytes() const
+	{
+		int32 Result = 0;
+
+		TArray<const FModel*> Models;
+		for (const FLiveInstance& Instance : LiveInstances)
+		{
+			Models.AddUnique(Instance.Model.Get());
+		}
+
+		// Data stored per-model, but related to instance construction
+		for (const FModel* Model : Models)
+		{
+			// Count streamable and currently-loaded resources
+			const FProgram& Program = Model->Program;
+			for (const TPair<uint32,TManagedPtr<const FImage>>& ImageLOD : Program.ConstantImageLODsStreamed)
+			{
+				if (ImageLOD.Value)
+				{
+					Result += ImageLOD.Value->GetDataSize();
+				}
+			}
+			for (const TPair<uint32, TManagedPtr<const FMesh>>& Rom : Program.ConstantMeshesStreamed)
+			{
+				if (Rom.Value)
+				{
+					Result += Rom.Value->GetDataSize();
+				}
+			}
+		}
+
+		return Result;
+	}
+	
+
 	void FWorkingMemoryManager::LogWorkingMemory(const CodeRunner* CurrentRunner) const
 	{
 #if !(UE_BUILD_SHIPPING || UE_BUILD_TEST)
@@ -1550,8 +1519,6 @@ namespace UE::Mutable::Private
 	}
 
 
-
-	//---------------------------------------------------------------------------------------------
 	FWorkingMemoryManager::FModelCacheEntry* FWorkingMemoryManager::FindModelCache(const FModel* InModel)
 	{
 		check(InModel);
@@ -1571,7 +1538,7 @@ namespace UE::Mutable::Private
 		return nullptr;
 	}
 
-    //---------------------------------------------------------------------------------------------
+   
 	FWorkingMemoryManager::FModelCacheEntry& FWorkingMemoryManager::FindOrAddModelCache(const TSharedPtr<const FModel>& InModel)
     {
         check(InModel);
@@ -1599,8 +1566,8 @@ namespace UE::Mutable::Private
         // Not found. Add new
 		FModelCacheEntry n;
         n.Model = TWeakPtr<const FModel>(InModel);
-		n.PendingOpsPerRom.SetNum(InModel->GetPrivate()->Program.Roms.Num());
-    	n.RomWeights.SetNumZeroed(InModel->GetPrivate()->Program.Roms.Num());
+		n.PendingOpsPerRom.SetNum(InModel->Program.Roms.Num());
+    	n.RomWeights.SetNumZeroed(InModel->Program.Roms.Num());
     	CachePerModel.Add(n);
         return CachePerModel.Last();
     }
@@ -1669,7 +1636,7 @@ namespace UE::Mutable::Private
 				TSharedPtr<const FModel> CacheModel = ModelCache.Model.Pin();
 				if (CacheModel)
 				{
-					UE::Mutable::Private::FProgram& Program = CacheModel->GetPrivate()->Program;
+					const FProgram& Program = CacheModel->Program;
 					check(ModelCache.RomWeights.Num() == Program.Roms.Num());
 	
 					check(Program.LoadedMemTrackedRoms.GetMaxIndex() <= Program.Roms.Num());
@@ -1682,7 +1649,7 @@ namespace UE::Mutable::Private
 						check(Program.IsRomLoaded(RomIndex));
 						check(Rom.ResourceType == (uint32)*Iter);
 
-						RomBytes += CacheModel->GetPrivate()->Program.GetRomBytes(RomIndex);
+						RomBytes += CacheModel->Program.GetRomBytes(RomIndex);
 
 						// We cannot unload a rom if some operation is expecting it.
 						const bool bIsRomLocked = ModelCache.PendingOpsPerRom.IsValidIndex(RomIndex) && 
@@ -1707,7 +1674,7 @@ namespace UE::Mutable::Private
 			SSIZE_T CandidatesRomBytes = 0;
 			for (const TPair<float, FRomRef>& Entry : Candidates)
 			{
-				CandidatesRomBytes += Entry.Value.Model->GetPrivate()->Program.GetRomBytes(Entry.Value.RomIndex);
+				CandidatesRomBytes += Entry.Value.Model->Program.GetRomBytes(Entry.Value.RomIndex);
 			}
 		
 			while (!bFinished && Candidates.Num())
@@ -1718,7 +1685,7 @@ namespace UE::Mutable::Private
 
 				// UE_LOG(LogMutableCore,Log, "Unloading rom because of memory budget: %d.", lowestPriorityRom);
 				int32 UnloadedSize = 0;
-				Candidate.Value.Model->GetPrivate()->Program.UnloadRom(Candidate.Value.RomIndex, &UnloadedSize);
+				Candidate.Value.Model->Program.UnloadRom(Candidate.Value.RomIndex, &UnloadedSize);
 				TotalBytes -= UnloadedSize;
 				bFinished = TotalBytes <= BudgetBytes;
 
@@ -1841,13 +1808,12 @@ namespace UE::Mutable::Private
 					LogWorkingMemory(nullptr);
 				}
 			}
-
 		}
 
         return bFinished;
     }
 
-    //---------------------------------------------------------------------------------------------
+   
     void FWorkingMemoryManager::MarkRomUsed( int32 romIndex, const TSharedPtr<const FModel>& pModel )
     {
         check(pModel);
@@ -1870,7 +1836,6 @@ namespace UE::Mutable::Private
     }
 
 
-	//---------------------------------------------------------------------------------------------
 	static void AddMultiValueKeys(FBitWriter& Blob, const TMap< TArray<int32>, FParameterValue >& Multi)
 	{
 		uint16 Num = uint16(Multi.Num());
@@ -1885,7 +1850,6 @@ namespace UE::Mutable::Private
 	}
 
 
-	//---------------------------------------------------------------------------------------------
 	void GetResourceKey(const TSharedPtr<const FModel>& Model, const FParameters* Params, uint32 ParamListIndex, OP::ADDRESS RootAt, TMemoryTrackedArray<uint8>& ParameterValuesBlob)
 	{
 		MUTABLE_CPUPROFILER_SCOPE(GetResourceKey);
@@ -1895,7 +1859,7 @@ namespace UE::Mutable::Private
 			return;
 		}
 
-		const FProgram& Program = Model->GetPrivate()->Program;
+		const FProgram& Program = Model->Program;
 
 		// Find the list of relevant parameters
 		const TArray<uint16>* RelevantParams = nullptr;
@@ -1912,7 +1876,7 @@ namespace UE::Mutable::Private
 		// Generate the relevant parameters blob
 		FBitWriterResourceKey Blob(2048*8, true);
 
-		const TArray<FParameterDesc>& ParamDescs = Params->GetPrivate()->Model->GetPrivate()->Program.Parameters;
+		const TArray<FParameterDesc>& ParamDescs = Params->Model->Program.Parameters;
 
 		// First make a mask with a bit for each relevant parameter. It will be on for parameters included in the blob.
 		// A parameter will be excluded from the blob if it has the default value, and no multivalues.
@@ -1922,11 +1886,11 @@ namespace UE::Mutable::Private
 			for (int32 IndexIndex = 0; IndexIndex < RelevantParams->Num(); ++IndexIndex)
 			{
 				int32 ParamIndex = (*RelevantParams)[IndexIndex];
-				bool bInclude = Params->GetPrivate()->HasMultipleValues(ParamIndex);
+				bool bInclude = Params->HasMultipleValues(ParamIndex);
 				if (!bInclude)
 				{
 					bInclude = !(
-						Params->GetPrivate()->Values[ParamIndex]
+						Params->Values[ParamIndex]
 						==
 						ParamDescs[ParamIndex].DefaultValue
 						);
@@ -1951,12 +1915,12 @@ namespace UE::Mutable::Private
 			switch (Program.Parameters[ParamIndex].Type)
 			{
 			case EParameterType::Bool:
-				Blob.WriteBit(Params->GetPrivate()->Values[ParamIndex].Get<FParamBoolType>() ? 1 : 0);
+				Blob.WriteBit(Params->Values[ParamIndex].Get<FParamBoolType>() ? 1 : 0);
 
 				// Multi-values
-				if (Params->GetPrivate()->HasMultipleValues(ParamIndex))
+				if (Params->HasMultipleValues(ParamIndex))
 				{
-					const TMap< TArray<int32>, FParameterValue >& Multi = Params->GetPrivate()->MultiValues[ParamIndex];
+					const TMap< TArray<int32>, FParameterValue >& Multi = Params->MultiValues[ParamIndex];
 					AddMultiValueKeys(Blob, Multi);
 					for (const TPair< TArray<int32>, FParameterValue >& v : Multi)
 					{
@@ -1968,7 +1932,7 @@ namespace UE::Mutable::Private
 			case EParameterType::Int:
 			{
 				int32 MaxValue = ParamDescs[ParamIndex].PossibleValues.Num();
-				int32 Value = Params->GetPrivate()->Values[ParamIndex].Get<FParamIntType>();
+				int32 Value = Params->Values[ParamIndex].Get<FParamIntType>();
 				if (MaxValue)
 				{
 					// It is an enum
@@ -1983,9 +1947,9 @@ namespace UE::Mutable::Private
 				}
 
 				// Multi-values
-				if (Params->GetPrivate()->HasMultipleValues(ParamIndex))
+				if (Params->HasMultipleValues(ParamIndex))
 				{
-					const TMap< TArray<int32>, FParameterValue >& Multi = Params->GetPrivate()->MultiValues[ParamIndex];
+					const TMap< TArray<int32>, FParameterValue >& Multi = Params->MultiValues[ParamIndex];
 					AddMultiValueKeys(Blob, Multi);
 					for (const TPair< TArray<int32>, FParameterValue >& v : Multi)
 					{
@@ -2009,12 +1973,12 @@ namespace UE::Mutable::Private
 
 			case EParameterType::Float:
 				DataSize = sizeof(FParamFloatType);
-				Blob.Serialize(&Params->GetPrivate()->Values[ParamIndex].Get<FParamFloatType>(), DataSize);
+				Blob.Serialize(&const_cast<FParameters*>(Params)->Values[ParamIndex].Get<FParamFloatType>(), DataSize);
 
 				// Multi-values
-				if (Params->GetPrivate()->HasMultipleValues(ParamIndex))
+				if (Params->HasMultipleValues(ParamIndex))
 				{
-					const TMap< TArray<int32>, FParameterValue >& Multi = Params->GetPrivate()->MultiValues[ParamIndex];
+					const TMap< TArray<int32>, FParameterValue >& Multi = Params->MultiValues[ParamIndex];
 					AddMultiValueKeys(Blob, Multi);
 					for (const TPair< TArray<int32>, FParameterValue >& v : Multi)
 					{
@@ -2025,12 +1989,12 @@ namespace UE::Mutable::Private
 
 			case EParameterType::Color:
 				DataSize = sizeof(FParamColorType);
-				Blob.Serialize(&Params->GetPrivate()->Values[ParamIndex].Get<FParamColorType>(), DataSize);
+				Blob.Serialize(&const_cast<FParameters*>(Params)->Values[ParamIndex].Get<FParamColorType>(), DataSize);
 
 				// Multi-values
-				if (Params->GetPrivate()->HasMultipleValues(ParamIndex))
+				if (Params->HasMultipleValues(ParamIndex))
 				{
-					const TMap< TArray<int32>, FParameterValue >& Multi = Params->GetPrivate()->MultiValues[ParamIndex];
+					const TMap< TArray<int32>, FParameterValue >& Multi = Params->MultiValues[ParamIndex];
 					AddMultiValueKeys(Blob, Multi);
 					for (const TPair< TArray<int32>, FParameterValue >& v : Multi)
 					{
@@ -2042,7 +2006,7 @@ namespace UE::Mutable::Private
 			case EParameterType::Projector:
 			{
 				FPackedNormal TempVec;
-				const FProjector& Value = Params->GetPrivate()->Values[ParamIndex].Get<FParamProjectorType>();
+				const FProjector& Value = Params->Values[ParamIndex].Get<FParamProjectorType>();
 				Blob.Serialize((void*)&Value.position, sizeof(FVector3f));
 				TempVec = FPackedNormal(Value.direction);
 				Blob.Serialize(&TempVec, sizeof(FPackedNormal));
@@ -2052,9 +2016,9 @@ namespace UE::Mutable::Private
 				Blob.Serialize((void*)&Value.projectionAngle, sizeof(float));
 
 				// Multi-values
-				if (Params->GetPrivate()->HasMultipleValues(ParamIndex))
+				if (Params->HasMultipleValues(ParamIndex))
 				{
-					const TMap< TArray<int32>, FParameterValue >& Multi = Params->GetPrivate()->MultiValues[ParamIndex];
+					const TMap< TArray<int32>, FParameterValue >& Multi = Params->MultiValues[ParamIndex];
 					AddMultiValueKeys(Blob, Multi);
 					for (const TPair< TArray<int32>, FParameterValue >& v : Multi)
 					{
@@ -2073,14 +2037,14 @@ namespace UE::Mutable::Private
 
 			case EParameterType::Texture:
 			{
-				TStrongObjectPtr<UTexture> Object = Params->GetPrivate()->Values[ParamIndex].Get<FParamTextureType>();
+				TStrongObjectPtr<UTexture> Object = Params->Values[ParamIndex].Get<FParamTextureType>();
 				FString Path = Object ? Object->GetPathName() : FString();
 				Blob.Serialize(GetData(Path), Path.Len() * sizeof(FString::ElementType));
 
 				// Multi-values
-				if (Params->GetPrivate()->HasMultipleValues(ParamIndex))
+				if (Params->HasMultipleValues(ParamIndex))
 				{
-					const TMap< TArray<int32>, FParameterValue >& Multi = Params->GetPrivate()->MultiValues[ParamIndex];
+					const TMap< TArray<int32>, FParameterValue >& Multi = Params->MultiValues[ParamIndex];
 					AddMultiValueKeys(Blob, Multi);
 					for (const TPair< TArray<int32>, FParameterValue >& v : Multi)
 					{
@@ -2094,14 +2058,14 @@ namespace UE::Mutable::Private
 
 			case EParameterType::SkeletalMesh:
 			{
-				TStrongObjectPtr<USkeletalMesh> Object = Params->GetPrivate()->Values[ParamIndex].Get<FParamSkeletalMeshType>();
+				TStrongObjectPtr<USkeletalMesh> Object = Params->Values[ParamIndex].Get<FParamSkeletalMeshType>();
 				FString Path = Object ? Object->GetPathName() : FString();
 				Blob.Serialize(GetData(Path), Path.Len() * sizeof(FString::ElementType));
 
 				// Multi-values
-				if (Params->GetPrivate()->HasMultipleValues(ParamIndex))
+				if (Params->HasMultipleValues(ParamIndex))
 				{
-					const TMap< TArray<int32>, FParameterValue >& Multi = Params->GetPrivate()->MultiValues[ParamIndex];
+					const TMap< TArray<int32>, FParameterValue >& Multi = Params->MultiValues[ParamIndex];
 					AddMultiValueKeys(Blob, Multi);
 					for (const TPair< TArray<int32>, FParameterValue >& v : Multi)
 					{
@@ -2115,14 +2079,14 @@ namespace UE::Mutable::Private
 
 			case EParameterType::Material:
 			{
-				TStrongObjectPtr<UMaterialInterface> Object = Params->GetPrivate()->Values[ParamIndex].Get<FParamMaterialType>();
+				TStrongObjectPtr<UMaterialInterface> Object = Params->Values[ParamIndex].Get<FParamMaterialType>();
 				FString Path = Object ? Object->GetPathName() : FString();
 				Blob.Serialize(GetData(Path), Path.Len() * sizeof(FString::ElementType));
 
 				// Multi-values
-				if (Params->GetPrivate()->HasMultipleValues(ParamIndex))
+				if (Params->HasMultipleValues(ParamIndex))
 				{
-					const TMap< TArray<int32>, FParameterValue >& Multi = Params->GetPrivate()->MultiValues[ParamIndex];
+					const TMap< TArray<int32>, FParameterValue >& Multi = Params->MultiValues[ParamIndex];
 					AddMultiValueKeys(Blob, Multi);
 					for (const TPair< TArray<int32>, FParameterValue >& v : Multi)
 					{
@@ -2138,13 +2102,13 @@ namespace UE::Mutable::Private
 			{
 				DataSize = sizeof(FMatrix44f);
 
-				const FMatrix44f& Value = Params->GetPrivate()->Values[ParamIndex].Get<FParamMatrixType>();
+				const FMatrix44f& Value = Params->Values[ParamIndex].Get<FParamMatrixType>();
 				Blob.Serialize((void*)&Value, DataSize);
 
 				// Multi-values
-				if (Params->GetPrivate()->HasMultipleValues(ParamIndex))
+				if (Params->HasMultipleValues(ParamIndex))
 				{
-					const TMap<TArray<int32>, FParameterValue>& Multi = Params->GetPrivate()->MultiValues[ParamIndex];
+					const TMap<TArray<int32>, FParameterValue>& Multi = Params->MultiValues[ParamIndex];
 					AddMultiValueKeys(Blob, Multi);
 					for (const TPair<TArray<int32>, FParameterValue>& MultiValuePair : Multi)
 					{
@@ -2157,15 +2121,15 @@ namespace UE::Mutable::Private
 
 			case EParameterType::InstancedStruct:
 			{
-				const FParamInstancedStructType& Value = Params->GetPrivate()->Values[ParamIndex].Get<FParamInstancedStructType>();
+				const FParamInstancedStructType& Value = Params->Values[ParamIndex].Get<FParamInstancedStructType>();
 
 				FInstancedStruct* ValueInstancedStruct = const_cast<FInstancedStruct*>(Value.Get());
 				ValueInstancedStruct->Serialize(Blob);
 
 				// Multi-values
-				if (Params->GetPrivate()->HasMultipleValues(ParamIndex))
+				if (Params->HasMultipleValues(ParamIndex))
 				{
-					const TMap<TArray<int32>, FParameterValue>& Multi = Params->GetPrivate()->MultiValues[ParamIndex];
+					const TMap<TArray<int32>, FParameterValue>& Multi = Params->MultiValues[ParamIndex];
 					AddMultiValueKeys(Blob, Multi);
 					for (const TPair<TArray<int32>, FParameterValue>& MultiValuePair : Multi)
 					{
@@ -2195,7 +2159,7 @@ namespace UE::Mutable::Private
     	FGeneratedMeshKey NewKey;
     	NewKey.Address = RootAt;
 
-    	int32* Result = Model->GetPrivate()->Program.RelevantParameterList.Find(RootAt);
+    	const int32* Result = Model->Program.RelevantParameterList.Find(RootAt);
     	const int32 ParamListIndex = Result ? *Result : INDEX_NONE; 
 		GetResourceKey(Model, Params, ParamListIndex, RootAt, NewKey.ParameterValuesBlob);
 
@@ -2208,7 +2172,7 @@ namespace UE::Mutable::Private
     	FGeneratedImageKey NewKey;
     	NewKey.Address = RootAt;
 
-    	int32* Result = Model->GetPrivate()->Program.RelevantParameterList.Find(RootAt);
+    	const int32* Result = Model->Program.RelevantParameterList.Find(RootAt);
     	const int32 ParamListIndex = Result ? *Result : INDEX_NONE; 
     	GetResourceKey(Model, Params, ParamListIndex, RootAt, NewKey.ParameterValuesBlob);
 
@@ -2221,7 +2185,7 @@ namespace UE::Mutable::Private
     	FGeneratedMaterialKey NewKey;
     	NewKey.Address = RootAt;
 
-    	int32* Result = Model->GetPrivate()->Program.RelevantParameterList.Find(RootAt);
+    	const int32* Result = Model->Program.RelevantParameterList.Find(RootAt);
     	const int32 ParamListIndex = Result ? *Result : INDEX_NONE; 
     	GetResourceKey(Model, Params, ParamListIndex, RootAt, NewKey.ParameterValuesBlob);
 
@@ -2234,7 +2198,7 @@ namespace UE::Mutable::Private
     	FGeneratedSkeletalMeshKey NewKey;
     	NewKey.Address = RootAt;
 
-    	int32* Result = Model->GetPrivate()->Program.RelevantParameterList.Find(RootAt);
+    	const int32* Result = Model->Program.RelevantParameterList.Find(RootAt);
     	const int32 ParamListIndex = Result ? *Result : INDEX_NONE; 
     	GetResourceKey(Model, Params, ParamListIndex, RootAt, NewKey.ParameterValuesBlob);
 
diff --git a/Engine/Plugins/Mutable/Source/MutableTools/Internal/MuT/AST.h b/Engine/Plugins/Mutable/Source/MutableTools/Internal/MuT/AST.h
--- a/Engine/Plugins/Mutable/Source/MutableTools/Internal/MuT/AST.h
+++ b/Engine/Plugins/Mutable/Source/MutableTools/Internal/MuT/AST.h
@@ -6,10 +6,10 @@
 #include "MuT/NodeImage.h"
 #include "MuR/Image.h"
 #include "MuR/Mesh.h"
-#include "MuR/ModelPrivate.h"
+#include "MuR/Model.h"
 #include "MuR/MutableMath.h"
 #include "MuR/Operations.h"
-#include "MuR/ParametersPrivate.h"
+#include "MuR/Parameters.h"
 #include "MuR/Ptr.h"
 #include "MuR/RefCounted.h"
 
diff --git a/Engine/Plugins/Mutable/Source/MutableTools/Internal/MuT/TypeInfo.h b/Engine/Plugins/Mutable/Source/MutableTools/Internal/MuT/TypeInfo.h
--- a/Engine/Plugins/Mutable/Source/MutableTools/Internal/MuT/TypeInfo.h
+++ b/Engine/Plugins/Mutable/Source/MutableTools/Internal/MuT/TypeInfo.h
@@ -6,7 +6,6 @@
 #include "MuR/Image.h"
 #include "MuR/MeshBufferSet.h"
 #include "MuR/Parameters.h"
-#include "MuR/ParametersPrivate.h"
 
 #define UE_API MUTABLETOOLS_API
 
diff --git a/Engine/Plugins/Mutable/Source/MutableTools/Private/MuT/ASTOpInstanceAdd.cpp b/Engine/Plugins/Mutable/Source/MutableTools/Private/MuT/ASTOpInstanceAdd.cpp
--- a/Engine/Plugins/Mutable/Source/MutableTools/Private/MuT/ASTOpInstanceAdd.cpp
+++ b/Engine/Plugins/Mutable/Source/MutableTools/Private/MuT/ASTOpInstanceAdd.cpp
@@ -4,8 +4,8 @@
 
 #include "MuT/ASTOpConditional.h"
 #include "MuT/CodeOptimiser.h"
-#include "MuR/ModelPrivate.h"
-#include "MuR/ParametersPrivate.h"
+#include "MuR/Model.h"
+#include "MuR/Parameters.h"
 #include "MuR/RefCounted.h"
 #include "MuR/Types.h"
 #include "Misc/AssertionMacros.h"
diff --git a/Engine/Plugins/Mutable/Source/MutableTools/Private/MuT/CodeOptimiser.cpp b/Engine/Plugins/Mutable/Source/MutableTools/Private/MuT/CodeOptimiser.cpp
--- a/Engine/Plugins/Mutable/Source/MutableTools/Private/MuT/CodeOptimiser.cpp
+++ b/Engine/Plugins/Mutable/Source/MutableTools/Private/MuT/CodeOptimiser.cpp
@@ -20,8 +20,8 @@
 #include "MuT/ASTOpImageRasterMesh.h"
 #include "MuT/ASTOpReferenceResource.h"
 
-#include "MuR/ModelPrivate.h"
-#include "MuR/SystemPrivate.h"
+#include "MuR/Model.h"
+#include "MuR/System.h"
 #include "MuR/Operations.h"
 #include "MuR/OpMeshMerge.h"
 #include "MuR/MutableRuntimeModule.h"
@@ -172,7 +172,7 @@ namespace UE::Mutable::Private
 			Settings.SetImageCompressionQuality( ImageCompressionQuality );
 			TSharedPtr<FSystem> System = MakeShared<FSystem>(Settings);
 
-			System->GetPrivate()->ImagePixelFormatOverride = ImOp.FormatImageOverride;
+			System->ImagePixelFormatOverride = ImOp.FormatImageOverride;
 
 			FSourceDataDescriptor SourceDataDescriptor;
 			if (DataType == EDataType::Image || DataType == EDataType::Mesh)
@@ -186,15 +186,15 @@ namespace UE::Mutable::Private
 			LinkerOptions.MinTextureResidentMipCount = 255;
 			LinkerOptions.bSeparateImageMips = false;
 
-			TSharedPtr<const FModel> Model = MakeShared<FModel>();
-			OP::ADDRESS at = ASTOp::FullLink(SourceCloned, Model->GetPrivate()->Program, &LinkerOptions);
+			TSharedPtr<FModel> Model = MakeShared<FModel>();
+			OP::ADDRESS at = ASTOp::FullLink(SourceCloned, Model->Program, &LinkerOptions);
 
 			FProgram::FState state;
 			state.Root = at;
-			Model->GetPrivate()->Program.States.Add(state);
+			Model->Program.States.Add(state);
 
 			TSharedPtr<FParameters> LocalParams = FModel::NewParameters(Model);
-			System->GetPrivate()->BeginBuild(Model, nullptr, nullptr, nullptr, nullptr);
+			System->BeginBuild(Model, nullptr, nullptr, nullptr, nullptr);
 
 			// Calculate the value and replace this op by a constant
 			switch( DataType )
@@ -203,7 +203,7 @@ namespace UE::Mutable::Private
 			{
 				MUTABLE_CPUPROFILER_SCOPE(ConstantMesh);
 
-				TManagedPtr<const FMesh> MeshBuild = System->GetPrivate()->BuildMesh(nullptr, Model, LocalParams.Get(), at, EMeshContentFlags::AllFlags);
+				TManagedPtr<const FMesh> MeshBuild = System->BuildMesh(nullptr, Model, LocalParams.Get(), at, EMeshContentFlags::AllFlags);
 
 				if (MeshBuild)
 				{
@@ -220,7 +220,7 @@ namespace UE::Mutable::Private
 			{
 				MUTABLE_CPUPROFILER_SCOPE(ConstantImage);
 
-				TManagedPtr<const FImage> ImageResult = System->GetPrivate()->BuildImage(nullptr, Model, LocalParams.Get(), at, 0);
+				TManagedPtr<const FImage> ImageResult = System->BuildImage(nullptr, Model, LocalParams.Get(), at, 0);
 
 				if (ImageResult)
 				{
@@ -237,7 +237,7 @@ namespace UE::Mutable::Private
 			{
 				MUTABLE_CPUPROFILER_SCOPE(ConstantLayout);
 
-				TManagedPtr<const FLayout> LayoutResult = System->GetPrivate()->BuildLayout(nullptr, Model, LocalParams.Get(), at);
+				TManagedPtr<const FLayout> LayoutResult = System->BuildLayout(nullptr, Model, LocalParams.Get(), at);
 
 				if (LayoutResult)
 				{
@@ -253,7 +253,7 @@ namespace UE::Mutable::Private
 			{
 				MUTABLE_CPUPROFILER_SCOPE(ConstantBool);
 
-				bool value = System->GetPrivate()->BuildBool(nullptr, Model, LocalParams.Get(), at);
+				bool value = System->BuildBool(nullptr, Model, LocalParams.Get(), at);
 				Result = new ASTOpConstantBool(value);
 				break;
 			}
@@ -263,7 +263,7 @@ namespace UE::Mutable::Private
 				MUTABLE_CPUPROFILER_SCOPE(ConstantBool);
 
 				FVector4f ResultColor(0, 0, 0, 0);
-				ResultColor = System->GetPrivate()->BuildColour(nullptr, Model, LocalParams.Get(), at);
+				ResultColor = System->BuildColour(nullptr, Model, LocalParams.Get(), at);
 
 				{
 					UE::Mutable::Private::Ptr<ASTOpConstantColor> ConstantOp = new ASTOpConstantColor();
@@ -284,7 +284,7 @@ namespace UE::Mutable::Private
 				break;
 			}
 
-			System->GetPrivate()->EndBuild();
+			System->EndBuild();
 		}
 
 	};
diff --git a/Engine/Plugins/Mutable/Source/MutableTools/Private/MuT/Compiler.cpp b/Engine/Plugins/Mutable/Source/MutableTools/Private/MuT/Compiler.cpp
--- a/Engine/Plugins/Mutable/Source/MutableTools/Private/MuT/Compiler.cpp
+++ b/Engine/Plugins/Mutable/Source/MutableTools/Private/MuT/Compiler.cpp
@@ -14,10 +14,9 @@
 #include "MuR/Image.h"
 #include "MuR/Mesh.h"
 #include "MuR/Model.h"
-#include "MuR/ModelPrivate.h"
 #include "MuR/MutableTrace.h"
 #include "MuR/Operations.h"
-#include "MuR/ParametersPrivate.h"
+#include "MuR/Parameters.h"
 #include "MuR/Serialisation.h"
 #include "MuR/MutableRuntimeModule.h"
 
@@ -546,7 +545,7 @@ namespace UE::Mutable::Private
 
         // Link the Program and generate state data.
 		TSharedPtr<FModel> Result = MakeShared<FModel>();
-        FProgram& Program = Result->GetPrivate()->Program;
+        FProgram& Program = Result->Program;
 
 		check(Program.Parameters.IsEmpty());
 		Program.Parameters = Parameters;
@@ -718,7 +717,7 @@ namespace UE::Mutable::Private
 
 		uint64 EmbeddedDataBytesLimit = Options->GetPrivate()->EmbeddedDataBytesLimit;
 
-		UE::Mutable::Private::FProgram& Program = Model->GetPrivate()->Program;
+		UE::Mutable::Private::FProgram& Program = Model->Program;
 
 		// These are used for logging only.
 		int32 NumRoms = 0;
diff --git a/Engine/Plugins/Mutable/Source/MutableTools/Private/MuT/ErrorLog.cpp b/Engine/Plugins/Mutable/Source/MutableTools/Private/MuT/ErrorLog.cpp
--- a/Engine/Plugins/Mutable/Source/MutableTools/Private/MuT/ErrorLog.cpp
+++ b/Engine/Plugins/Mutable/Source/MutableTools/Private/MuT/ErrorLog.cpp
@@ -13,10 +13,9 @@
 #include "MuR/Layout.h"
 #include "MuR/Mesh.h"
 #include "MuR/Model.h"
-#include "MuR/ModelPrivate.h"
 #include "MuR/MutableMath.h"
 #include "MuR/Operations.h"
-#include "MuR/ParametersPrivate.h"
+#include "MuR/Parameters.h"
 #include "MuR/Platform.h"
 #include "MuR/Skeleton.h"
 #include "MuR/System.h"
