[Mutable] Fix ASTOp::Replace not working on the roots.

- Fix check when generating grooms.

#Synced-CL 49557694

diff --git a/Engine/Plugins/Mutable/Source/MutableRuntime/Private/MuR/CodeRunnerBegin.cpp b/Engine/Plugins/Mutable/Source/MutableRuntime/Private/MuR/CodeRunnerBegin.cpp
--- a/Engine/Plugins/Mutable/Source/MutableRuntime/Private/MuR/CodeRunnerBegin.cpp
+++ b/Engine/Plugins/Mutable/Source/MutableRuntime/Private/MuR/CodeRunnerBegin.cpp
@@ -1361,7 +1361,7 @@ namespace UE::Mutable::Private
 						{
 							PushOp(FScheduledOpInline(Item, 1));
 							PushOp(FScheduledOpInline(Args.Instance, Item));
-							PushOp(FScheduledOpInline(Args.ExtensionData, Item));
+							PushOp(FScheduledOpInline(Args.ExtensionData, Item, false));
 
 							break;
 						}
diff --git a/Engine/Plugins/Mutable/Source/MutableTools/Private/MuT/ASTOpEntry.cpp b/Engine/Plugins/Mutable/Source/MutableTools/Private/MuT/ASTOpEntry.cpp
new file mode 100644
--- /dev/null
+++ b/Engine/Plugins/Mutable/Source/MutableTools/Private/MuT/ASTOpEntry.cpp
@@ -0,0 +1,71 @@
+// Copyright Epic Games, Inc. All Rights Reserved.
+
+#include "MuT/ASTOpEntry.h"
+#include "MuT/ASTOpConstantBool.h"
+
+
+namespace UE::Mutable::Private
+{
+
+	ASTOpEntry::ASTOpEntry()
+		: A(this)
+	{
+	}
+
+
+	ASTOpEntry::~ASTOpEntry()
+	{
+		// Explicit call needed to avoid recursive destruction
+		ASTOp::RemoveChildren();
+	}
+
+
+	EOpType ASTOpEntry::GetOpType() const
+	{
+		return EOpType::NONE;
+	}
+
+
+	bool ASTOpEntry::IsEqual(const ASTOp& OtherUntyped) const
+	{
+		if (OtherUntyped.GetOpType() == GetOpType())
+		{
+			const ASTOpEntry* Other = static_cast<const ASTOpEntry*>(&OtherUntyped);
+			return A == Other->A;
+		}
+		
+		return false;
+	}
+
+
+	uint64 ASTOpEntry::Hash() const
+	{
+		uint64 Result = std::hash<void*>()(A.child().get());
+		return Result;
+	}
+
+
+	Ptr<ASTOp> ASTOpEntry::Clone(MapChildFuncRef MapChild) const
+	{
+		Ptr<ASTOpEntry> New = new ASTOpEntry();
+		
+		New->A = MapChild(A.child());
+		
+		return New;
+	}
+
+
+	void ASTOpEntry::ForEachChild(const TFunctionRef<void(ASTChild&)> Func)
+	{
+		Func(A);
+	}
+
+
+	void ASTOpEntry::Link(FProgram& Program, FLinkerOptions*)
+	{
+		if (!linkedAddress)
+		{
+			linkedAddress = A->linkedAddress;
+		}
+	}
+}
diff --git a/Engine/Plugins/Mutable/Source/MutableTools/Private/MuT/ASTOpEntry.h b/Engine/Plugins/Mutable/Source/MutableTools/Private/MuT/ASTOpEntry.h
new file mode 100644
--- /dev/null
+++ b/Engine/Plugins/Mutable/Source/MutableTools/Private/MuT/ASTOpEntry.h
@@ -0,0 +1,33 @@
+// Copyright Epic Games, Inc. All Rights Reserved.
+
+#pragma once
+
+#include "HAL/Platform.h"
+#include "MuR/Operations.h"
+#include "MuR/Ptr.h"
+#include "MuT/AST.h"
+
+
+namespace UE::Mutable::Private
+{
+	struct FProgram;
+
+	class ASTOpEntry final : public ASTOp
+	{
+	public:
+		ASTChild A;
+
+		ASTOpEntry();
+		ASTOpEntry(const ASTOpEntry&) = delete;
+		~ASTOpEntry();
+
+		// ASTOp interface
+		virtual EOpType GetOpType() const override;
+		virtual uint64 Hash() const override;
+		virtual bool IsEqual(const ASTOp&) const override;
+		virtual Ptr<ASTOp> Clone(MapChildFuncRef) const override;
+		virtual void ForEachChild(const TFunctionRef<void(ASTChild&)>) override;
+		virtual void Link(FProgram&, FLinkerOptions*) override;
+	};
+}
+
diff --git a/Engine/Plugins/Mutable/Source/MutableTools/Private/MuT/CodeGenerator.cpp b/Engine/Plugins/Mutable/Source/MutableTools/Private/MuT/CodeGenerator.cpp
--- a/Engine/Plugins/Mutable/Source/MutableTools/Private/MuT/CodeGenerator.cpp
+++ b/Engine/Plugins/Mutable/Source/MutableTools/Private/MuT/CodeGenerator.cpp
@@ -4,6 +4,7 @@
 
 #include "ASTOpAddOverrideMaterial.h"
 #include "ASTOpAddSkeletalMesh.h"
+#include "ASTOpEntry.h"
 #include "Containers/Array.h"
 #include "Logging/LogCategory.h"
 #include "Logging/LogMacros.h"
@@ -152,7 +153,11 @@ namespace UE::Mutable::Private
 				GenerateGeneric(Options, Result, InNode.get());
 
 				Ptr<ASTOp> StateRoot = Result.Op;
-				States.Emplace(State, StateRoot);
+
+				Ptr<ASTOpEntry> EntryOp = new ASTOpEntry(); // Dummy op. Required to be able to replace real roots (ASTOp::Replace requires the operation to have a parent).
+				EntryOp->A = StateRoot;
+				
+				States.Emplace(State, EntryOp);
 
 				++CurrentStateIndex;
 			}
