[Mutable] Fix issues with CodeRunnerBegin optional evaluation.

#Synced-CL 49557694

diff --git a/Engine/Plugins/Mutable/Source/MutableRuntime/Private/MuR/CodeRunnerBegin.cpp b/Engine/Plugins/Mutable/Source/MutableRuntime/Private/MuR/CodeRunnerBegin.cpp
--- a/Engine/Plugins/Mutable/Source/MutableRuntime/Private/MuR/CodeRunnerBegin.cpp
+++ b/Engine/Plugins/Mutable/Source/MutableRuntime/Private/MuR/CodeRunnerBegin.cpp
@@ -995,7 +995,7 @@ namespace UE::Mutable::Private
 						{
 							PushOp(FScheduledOpInline(Item, 1));
 							PushOp(FScheduledOpInline(Args.Instance, Item));
-							PushOp(FScheduledOpInline(Args.Material, Item));
+							PushOp(FScheduledOpInline(Args.Material, Item, false));
 
 							break;
 						}
@@ -1361,7 +1361,7 @@ namespace UE::Mutable::Private
 						{
 							PushOp(FScheduledOpInline(Item, 1));
 							PushOp(FScheduledOpInline(Args.Instance, Item));
-							PushOp(FScheduledOpInline(Args.ExtensionData, Item, false));
+							PushOp(FScheduledOpInline(Args.ExtensionData, Item));
 
 							break;
 						}
@@ -1750,7 +1750,7 @@ namespace UE::Mutable::Private
 				
 			case EOpType::ED_CONSTANT:
 				{
-					check(!Item.bEvaluate);
+					check(Item.bEvaluate);
 
 					OP::ResourceConstantArgs Args = Program.GetOpArgs<OP::ResourceConstantArgs>(Item.At);
 
diff --git a/Engine/Plugins/Mutable/Source/MutableTools/Private/MuT/ASTOpEntry.cpp b/Engine/Plugins/Mutable/Source/MutableTools/Private/MuT/ASTOpEntry.cpp
--- a/Engine/Plugins/Mutable/Source/MutableTools/Private/MuT/ASTOpEntry.cpp
+++ b/Engine/Plugins/Mutable/Source/MutableTools/Private/MuT/ASTOpEntry.cpp
@@ -65,7 +65,10 @@ namespace UE::Mutable::Private
 	{
 		if (!linkedAddress)
 		{
-			linkedAddress = A->linkedAddress;
+			if (A)
+			{
+				linkedAddress = A->linkedAddress;
+			}
 		}
 	}
 }
