[Mutable] Remove Live Instance ID

#Synced-CL 49557694

diff --git a/Engine/Plugins/Mutable/Source/CustomizableObject/Internal/MuCO/CustomizableObjectInstancePrivate.h b/Engine/Plugins/Mutable/Source/CustomizableObject/Internal/MuCO/CustomizableObjectInstancePrivate.h
--- a/Engine/Plugins/Mutable/Source/CustomizableObject/Internal/MuCO/CustomizableObjectInstancePrivate.h
+++ b/Engine/Plugins/Mutable/Source/CustomizableObject/Internal/MuCO/CustomizableObjectInstancePrivate.h
@@ -246,7 +246,7 @@ public:
 	TArray<uint32> RelevantParameters;
 
 	// Only used in LiveUpdateMode to reuse core instances between updates and their temp data to speed up updates, but spend way more memory
-	UE::Mutable::Private::FInstance::FID LiveUpdateModeInstanceID = 0;
+	TSharedPtr<UE::Mutable::Private::FLiveInstance> LiveInstance;
 
 #if WITH_EDITOR
 	UE_API virtual void PostDuplicate(bool bDuplicateForPIE) override;
diff --git a/Engine/Plugins/Mutable/Source/CustomizableObject/Internal/MuCO/CustomizableObjectSystemPrivate.h b/Engine/Plugins/Mutable/Source/CustomizableObject/Internal/MuCO/CustomizableObjectSystemPrivate.h
--- a/Engine/Plugins/Mutable/Source/CustomizableObject/Internal/MuCO/CustomizableObjectSystemPrivate.h
+++ b/Engine/Plugins/Mutable/Source/CustomizableObject/Internal/MuCO/CustomizableObjectSystemPrivate.h
@@ -260,8 +260,6 @@ class FMutablePendingInstanceWork
 
 	TSet<FMutablePendingInstanceDiscard, FPendingInstanceDiscardKeyFuncs> PendingInstanceDiscards;
 
-	TSet<UE::Mutable::Private::FInstance::FID> PendingIDsToRelease;
-
 public:
 	// Returns the number of pending instance updates, LOD Updates and discards.
 	int32 Num() const;
@@ -288,15 +286,7 @@ public:
 		return PendingInstanceDiscards.CreateIterator();
 	}
 
-	TSet<UE::Mutable::Private::FInstance::FID>::TIterator GetIDsToReleaseIterator()
-	{
-		return PendingIDsToRelease.CreateIterator();
-	}
-
 	void AddDiscard(const FMutablePendingInstanceDiscard& TaskToEnqueue);
-	void AddIDRelease(UE::Mutable::Private::FInstance::FID IDToRelease);
-
-	void RemoveAllUpdatesAndDiscardsAndReleases();
 };
 
 
@@ -612,6 +602,8 @@ public:
 	TSharedPtr<UE::Mutable::Private::FImageIdRegistry> ImageIdRegistry;
 	TSharedPtr<UE::Mutable::Private::FMaterialIdRegistry> MaterialIdRegistry;
 	TSharedPtr<UE::Mutable::Private::FSkeletalMeshIdRegistry> SkeletalMeshIdRegistry;
+	
+	TSharedPtr<UE::Mutable::Private::FLiveInstance> LiveInstance;
 
 	TSharedPtr<FUnrealMutableResourceProvider> ExternalResourceProvider;
 
@@ -651,8 +643,6 @@ public:
 	
 	int64 UpdateEndRealPeakBytes = 0;
 	
-	UE::Mutable::Private::FInstance::FID InstanceID = 0; // Redundant
-	
 	/** This option comes from the operation request. It is used to reduce the number of mipmaps that mutable must generate for images.  */
 	int32 MipsToSkip = 0;
 	
@@ -860,8 +850,6 @@ public:
 	// Init an async and safe release of the UE and Mutable resources used by the instance without actually destroying the instance, for example if it's very far away
 	void InitDiscardResourcesSkeletalMesh(UCustomizableObjectInstance* InCustomizableObjectInstance);
 
-	// Init the async release of a Mutable Core Instance ID and all the temp resources associated with it
-	void InitInstanceIDRelease(UE::Mutable::Private::FInstance::FID IDToRelease);
 	
 	bool IsReplaceDiscardedWithReferenceMeshEnabled() const;
 	void SetReplaceDiscardedWithReferenceMeshEnabled(bool bIsEnabled);
diff --git a/Engine/Plugins/Mutable/Source/CustomizableObject/Private/MuCO/CustomizableObjectInstance.cpp b/Engine/Plugins/Mutable/Source/CustomizableObject/Private/MuCO/CustomizableObjectInstance.cpp
--- a/Engine/Plugins/Mutable/Source/CustomizableObject/Private/MuCO/CustomizableObjectInstance.cpp
+++ b/Engine/Plugins/Mutable/Source/CustomizableObject/Private/MuCO/CustomizableObjectInstance.cpp
@@ -377,9 +377,6 @@ void UCustomizableObjectInstance::PostInitProperties()
 
 void UCustomizableObjectInstance::BeginDestroy()
 {
-	// Release the Live Instance ID if there it hadn't been released before
-	DestroyLiveUpdateInstance();
-
 	if (PrivateData)
 	{
 #if WITH_EDITOR
@@ -394,17 +391,7 @@ void UCustomizableObjectInstance::BeginDestroy()
 
 void UCustomizableObjectInstance::DestroyLiveUpdateInstance()
 {
-	if (PrivateData && PrivateData->LiveUpdateModeInstanceID)
-	{
-		// If UCustomizableObjectSystemPrivate::SSystem is nullptr it means it has already been destroyed, no point in registering an instanceID release
-		// since the Mutable system has already been destroyed. Just checking UCustomizableObjectSystem::GetInstance() will try to recreate the system when
-		// everything is shutting down, so it's better to check UCustomizableObjectSystemPrivate::SSystem first here
-		if (UCustomizableObjectSystemPrivate::SSystem && UCustomizableObjectSystem::GetInstance() && UCustomizableObjectSystem::GetInstance()->GetPrivate())
-		{
-			UCustomizableObjectSystem::GetInstance()->GetPrivate()->InitInstanceIDRelease(PrivateData->LiveUpdateModeInstanceID);
-			PrivateData->LiveUpdateModeInstanceID = 0;
-		}
-	}
+	GetPrivate()->LiveInstance.Reset();
 }
 
 
diff --git a/Engine/Plugins/Mutable/Source/CustomizableObject/Private/MuCO/CustomizableObjectMeshUpdate.cpp b/Engine/Plugins/Mutable/Source/CustomizableObject/Private/MuCO/CustomizableObjectMeshUpdate.cpp
--- a/Engine/Plugins/Mutable/Source/CustomizableObject/Private/MuCO/CustomizableObjectMeshUpdate.cpp
+++ b/Engine/Plugins/Mutable/Source/CustomizableObject/Private/MuCO/CustomizableObjectMeshUpdate.cpp
@@ -178,18 +178,13 @@ void FCustomizableObjectMeshStreamIn::DoCancelMeshUpdate(const FContext& Context
 
 namespace impl
 {
-	void Task_Mutable_UpdateMesh_End(const TSharedPtr<FMutableUpdateMeshContext> OperationData, TRefCountPtr<FCustomizableObjectMeshStreamIn>& Task, UE::Mutable::Private::FInstance::FID InstanceID)
+	void Task_Mutable_UpdateMesh_End(const TSharedPtr<FMutableUpdateMeshContext> OperationData, TRefCountPtr<FCustomizableObjectMeshStreamIn>& Task, const TSharedRef<UE::Mutable::Private::FLiveInstance>& LiveInstance)
 	{
 		MUTABLE_CPUPROFILER_SCOPE(Task_Mutable_UpdateMesh_End);
 
 		// End update
-		OperationData->System->EndUpdate(InstanceID);
-		OperationData->System->ReleaseInstance(InstanceID);
-
-		if (UCustomizableObjectSystem::ShouldClearWorkingMemoryOnUpdateEnd())
-		{
-			OperationData->System->ClearWorkingMemory();
-		}
+		const bool bClearRoms = UCustomizableObjectSystem::ShouldClearWorkingMemoryOnUpdateEnd();
+		OperationData->System->EndUpdate(LiveInstance, bClearRoms);
 
 		OperationData->Event.Trigger();
 		
@@ -199,27 +194,27 @@ namespace impl
 	void Task_Mutable_UpdateMesh_Loop(
 		const TSharedPtr<FMutableUpdateMeshContext> OperationData,
 		TRefCountPtr<FCustomizableObjectMeshStreamIn>& Task,
-		UE::Mutable::Private::FInstance::FID InstanceID,
+		const TSharedRef<UE::Mutable::Private::FLiveInstance>& LiveInstance,
 		int32 LODIndex)
 	{
 		MUTABLE_CPUPROFILER_SCOPE(Task_Mutable_UpdateMesh_Loop);
 
 		if (Task->IsCancelled() || LODIndex == OperationData->CurrentFirstLODIdx + OperationData->AssetLODBias)
 		{
-			Task_Mutable_UpdateMesh_End(OperationData, Task, InstanceID);
+			Task_Mutable_UpdateMesh_End(OperationData, Task, LiveInstance);
 			return;
 		}
 
 		const UE::Mutable::Private::FMeshId& MeshID = OperationData->MeshIDs[LODIndex];
 
 		constexpr UE::Mutable::Private::EMeshContentFlags MeshContentFilter = UE::Mutable::Private::EMeshContentFlags::AllFlags;
-		UE::Tasks::TTask<UE::Mutable::Private::TManagedPtr<const UE::Mutable::Private::FMesh>> GetMeshTask = OperationData->System->GetMesh({}, InstanceID, MeshID, MeshContentFilter);
+		UE::Tasks::TTask<UE::Mutable::Private::TManagedPtr<const UE::Mutable::Private::FMesh>> GetMeshTask = OperationData->System->GetMesh({}, LiveInstance, MeshID, MeshContentFilter);
 
 		UE::Tasks::AddNested(UE::Tasks::Launch(TEXT("Task_MutableGetMeshes_GetMesh_Post"), [=]() mutable
 			{
 				OperationData->Meshes[LODIndex] = GetMeshTask.GetResult();
 
-				Task_Mutable_UpdateMesh_Loop(OperationData, Task, InstanceID, LODIndex + 1);
+				Task_Mutable_UpdateMesh_Loop(OperationData, Task, LiveInstance, LODIndex + 1);
 			},
 			GetMeshTask,
 			LowLevelTasks::ETaskPriority::Inherit));
@@ -250,14 +245,13 @@ namespace impl
 #endif
 
 		// For now, we are forcing the recreation of mutable-side instances with every update.
-		UE::Mutable::Private::FInstance::FID InstanceID = System->NewInstance(Model, OperationData->ExternalResourceProvider, nullptr);
-		UE_LOG(LogMutable, Verbose, TEXT("Creating Mutable instance with id [%d] for a mesh update"), InstanceID);
+		TSharedRef<UE::Mutable::Private::FLiveInstance> LiveInstance = System->NewInstance(Model, OperationData->ExternalResourceProvider, nullptr);
 
 		// LOD mask, set to all ones to build all LODs
 		const uint32 LODMask = 0xFFFFFFFF;
 
 		// Main instance generation step
-		UE::Mutable::Private::TManagedPtr<const UE::Mutable::Private::FInstance> Instance = System->BeginUpdate_MutableThread(InstanceID,
+		UE::Mutable::Private::TManagedPtr<const UE::Mutable::Private::FInstance> Instance = System->BeginUpdate_MutableThread(LiveInstance,
 			OperationData->Parameters,
 			OperationData->MeshIdRegistry,
 			OperationData->ImageIdRegistry,
@@ -268,7 +262,7 @@ namespace impl
 		
 		check(Instance);
 
-		Task_Mutable_UpdateMesh_Loop(OperationData, Task, InstanceID, OperationData->PendingFirstLODIdx + OperationData->AssetLODBias);
+		Task_Mutable_UpdateMesh_Loop(OperationData, Task, LiveInstance, OperationData->PendingFirstLODIdx + OperationData->AssetLODBias);
 	}
 	
 } // namespace
diff --git a/Engine/Plugins/Mutable/Source/CustomizableObject/Private/MuCO/CustomizableObjectMipDataProvider.cpp b/Engine/Plugins/Mutable/Source/CustomizableObject/Private/MuCO/CustomizableObjectMipDataProvider.cpp
--- a/Engine/Plugins/Mutable/Source/CustomizableObject/Private/MuCO/CustomizableObjectMipDataProvider.cpp
+++ b/Engine/Plugins/Mutable/Source/CustomizableObject/Private/MuCO/CustomizableObjectMipDataProvider.cpp
@@ -131,8 +131,7 @@ namespace Impl
 #endif
 
 		// For now, we are forcing the recreation of mutable-side instances with every update.
-		UE::Mutable::Private::FInstance::FID InstanceID = System->NewInstance(Model, OperationData->UpdateContext->ExternalResourceProvider, nullptr);
-		UE_LOG(LogMutable, Verbose, TEXT("Creating Mutable instance with id [%d] for a single UpdateImage"), InstanceID)
+		TSharedRef<UE::Mutable::Private::FLiveInstance> LiveInstance = System->NewInstance(Model, OperationData->UpdateContext->ExternalResourceProvider, nullptr);
 
 		UE::Mutable::Private::TManagedPtr<const UE::Mutable::Private::FInstance> Instance;
 
@@ -141,7 +140,7 @@ namespace Impl
 			// LOD mask, set to all ones to build all LODs
 			uint32 LODMask = 0xFFFFFFFF;
 
-			Instance = System->BeginUpdate_MutableThread(InstanceID,
+			Instance = System->BeginUpdate_MutableThread(LiveInstance,
 				OperationData->UpdateContext->Parameters,
 				OperationData->UpdateContext->MeshIdRegistry,
 				OperationData->UpdateContext->ImageIdRegistry,
@@ -155,10 +154,10 @@ namespace Impl
 
 		const FMutableImageReference& ImageRef = OperationData->RequestedImage;
 		
-		UE::Tasks::TTask<UE::Mutable::Private::TManagedPtr<const UE::Mutable::Private::FImage>> GetImageTask = System->GetImage({}, InstanceID, ImageRef.ImageID, ImageRef.BaseMip + OperationData->MipsToSkip);
+		UE::Tasks::TTask<UE::Mutable::Private::TManagedPtr<const UE::Mutable::Private::FImage>> GetImageTask = System->GetImage({}, LiveInstance, ImageRef.ImageID, ImageRef.BaseMip + OperationData->MipsToSkip);
 
 		UE::Tasks::AddNested(UE::Tasks::Launch(TEXT("MipDataProvider_EndUpdateImagesTask"),
-				[System, OperationData, InstanceID, StartTime, GetImageTask, EndUpdateImage]() mutable
+				[System, OperationData, LiveInstance, StartTime, GetImageTask, EndUpdateImage]() mutable
 				{
 					check(GetImageTask.IsCompleted());
 
@@ -202,13 +201,8 @@ namespace Impl
 					// End update
 					{
 						MUTABLE_CPUPROFILER_SCOPE(EndUpdate);
-						System->EndUpdate(InstanceID);
-						System->ReleaseInstance(InstanceID);
-
-						if (UCustomizableObjectSystem::ShouldClearWorkingMemoryOnUpdateEnd())
-						{
-							System->ClearWorkingMemory();
-						}
+						const bool bClearRoms = UCustomizableObjectSystem::ShouldClearWorkingMemoryOnUpdateEnd();
+						System->EndUpdate(LiveInstance, bClearRoms);
 					}
 
 					if (FLogBenchmarkUtil::IsBenchmarkingReportingEnabled())
diff --git a/Engine/Plugins/Mutable/Source/CustomizableObject/Private/MuCO/CustomizableObjectSystem.cpp b/Engine/Plugins/Mutable/Source/CustomizableObject/Private/MuCO/CustomizableObjectSystem.cpp
--- a/Engine/Plugins/Mutable/Source/CustomizableObject/Private/MuCO/CustomizableObjectSystem.cpp
+++ b/Engine/Plugins/Mutable/Source/CustomizableObject/Private/MuCO/CustomizableObjectSystem.cpp
@@ -27,6 +27,7 @@
 #include "Components/SkeletalMeshComponent.h"
 #include "MuCO/CustomizableObjectSystemPrivate.h"
 #include "MuR/Parameters.h"
+#include "MuR/System.h"
 
 #if WITH_EDITOR
 #include "Editor.h"
@@ -140,13 +141,11 @@ static FAutoConsoleVariableRef CVarMutableEnableLODManagmentInEditor(
 	TEXT("true/false - If true, enables custom LODManagment in the editor. "),
 	ECVF_Default);
 
-
 TAutoConsoleVariable<bool> CVarMutableLogObjectMemoryOnUpdate(
 	TEXT("mutable.LogObjectMemoryOnUpdate"),
 	false,
 	TEXT("Log the memory used for a CO on every update."),
 	ECVF_Scalability);
-
 #endif
 
 TAutoConsoleVariable<bool> CVarEnableReleaseMeshResources(
@@ -668,7 +667,7 @@ uint32 FPendingInstanceUpdateKeyFuncs::GetKeyHash(const TWeakObjectPtr<const UCu
 
 int32 FMutablePendingInstanceWork::Num() const
 {
-	return PendingInstanceUpdates.Num() + PendingInstanceDiscards.Num() + PendingIDsToRelease.Num();
+	return PendingInstanceUpdates.Num() + PendingInstanceDiscards.Num();
 }
 
 
@@ -745,12 +744,6 @@ void FMutablePendingInstanceWork::AddDiscard(const FMutablePendingInstanceDiscar
 }
 
 
-void FMutablePendingInstanceWork::AddIDRelease(UE::Mutable::Private::FInstance::FID IDToRelease)
-{
-	PendingIDsToRelease.Add(IDToRelease);
-}
-
-
 UCustomizableObjectSystem* UCustomizableObjectSystem::GetInstance()
 {
 	if (!UCustomizableObjectSystemPrivate::SSystem)
@@ -1667,14 +1660,6 @@ void UCustomizableObjectSystemPrivate::InitDiscardResourcesSkeletalMesh(UCustomi
 }
 
 
-void UCustomizableObjectSystemPrivate::InitInstanceIDRelease(UE::Mutable::Private::FInstance::FID IDToRelease)
-{
-	check(IsInGameThread());
-
-	MutablePendingInstanceWork.AddIDRelease(IDToRelease);
-}
-
-
 bool UCustomizableObjectSystem::IsReplaceDiscardedWithReferenceMeshEnabled() const
 {
 	if (Private)
@@ -1881,7 +1866,6 @@ namespace impl
 		MUTABLE_CPUPROFILER_SCOPE(Subtask_Mutable_UpdateParameterRelevancy)
 
 		check(OperationData->Parameters);
-		check(OperationData->InstanceID != 0);
 		
 		// This must run in the mutable thread.
 		check(UCustomizableObjectSystem::GetInstance() != nullptr);
@@ -1890,7 +1874,7 @@ namespace impl
 		// Update the parameter relevancy.
 		{
 			MUTABLE_CPUPROFILER_SCOPE(ParameterRelevancy)
-			OperationData->MutableSystem->GetParameterRelevancy(OperationData->InstanceID, OperationData->Parameters, OperationData->RelevantParameters);
+			OperationData->MutableSystem->GetParameterRelevancy(OperationData->LiveInstance.ToSharedRef(), OperationData->Parameters, OperationData->RelevantParameters);
 		}
 
 		OperationData->TaskGetMeshesTime = FPlatformTime::Seconds() - OperationData->TaskGetMeshesTimeStart;
@@ -2246,7 +2230,7 @@ namespace impl
 		{
 			TSharedRef<FInstanceUpdateData::FMaterial> Material = Context->InstanceUpdateData.Materials[MaterialIndex];
 			
-			UE::Tasks::TTask<UE::Mutable::Private::TManagedPtr<const UE::Mutable::Private::FMaterial>> GetMaterialTask = Context->MutableSystem->GetMaterial(LastTask, Context->InstanceID, Material->MaterialId);
+			UE::Tasks::TTask<UE::Mutable::Private::TManagedPtr<const UE::Mutable::Private::FMaterial>> GetMaterialTask = Context->MutableSystem->GetMaterial(LastTask, Context->LiveInstance.ToSharedRef(), Material->MaterialId);
 			UE::Tasks::AddNested(GetMaterialTask);
 			LastTask = GetMaterialTask;
 			
@@ -2325,7 +2309,7 @@ namespace impl
 			OperationData->MutableSystem->SetImagePixelConversionOverride( OperationData->PixelFormatOverride );
 		}
 		
-		UE::Mutable::Private::TManagedPtr<const UE::Mutable::Private::FInstance> Instance = OperationData->MutableSystem->BeginUpdate_MutableThread(OperationData->InstanceID,
+		UE::Mutable::Private::TManagedPtr<const UE::Mutable::Private::FInstance> Instance = OperationData->MutableSystem->BeginUpdate_MutableThread(OperationData->LiveInstance.ToSharedRef(),
 			OperationData->Parameters,
 			OperationData->MeshIdRegistry,
 			OperationData->ImageIdRegistry,
@@ -2360,7 +2344,7 @@ namespace impl
 
 			Component->SkeletalMeshId = Instance->GetSkeletalMeshId(InstanceComponentIndex);
 
-			UE::Tasks::TTask<UE::Mutable::Private::TManagedPtr<const UE::Mutable::Private::FSkeletalMesh>> GetSkeletalMeshTask = OperationData->MutableSystem->GetSkeletalMesh(LastTask, OperationData->InstanceID, Component->SkeletalMeshId, UE::Mutable::Private::EMeshContentFlags::AllFlags);
+			UE::Tasks::TTask<UE::Mutable::Private::TManagedPtr<const UE::Mutable::Private::FSkeletalMesh>> GetSkeletalMeshTask = OperationData->MutableSystem->GetSkeletalMesh(LastTask, OperationData->LiveInstance.ToSharedRef(), Component->SkeletalMeshId, UE::Mutable::Private::EMeshContentFlags::AllFlags);
 			UE::Tasks::AddNested(GetSkeletalMeshTask);
 			LastTask = GetSkeletalMeshTask;
 			
@@ -2432,7 +2416,7 @@ namespace impl
 					}
 				}
 				
-				UE::Tasks::TTask<UE::Mutable::Private::TManagedPtr<const UE::Mutable::Private::FMesh>> GetMeshTask = OperationData->MutableSystem->GetMesh(LastTask, OperationData->InstanceID, LOD->MeshId, MeshContentFilter);
+				UE::Tasks::TTask<UE::Mutable::Private::TManagedPtr<const UE::Mutable::Private::FMesh>> GetMeshTask = OperationData->MutableSystem->GetMesh(LastTask, OperationData->LiveInstance.ToSharedRef(), LOD->MeshId, MeshContentFilter);
 				UE::Tasks::AddNested(GetMeshTask);
 				LastTask = GetMeshTask;
 				
@@ -2466,7 +2450,7 @@ namespace impl
 				continue;
 			}
 
-			UE::Tasks::TTask<UE::Mutable::Private::FExtendedImageDesc> GetImageDescTask = OperationData->MutableSystem->GetImageDesc(LastTask, OperationData->InstanceID, Image->ImageID);
+			UE::Tasks::TTask<UE::Mutable::Private::FExtendedImageDesc> GetImageDescTask = OperationData->MutableSystem->GetImageDesc(LastTask, OperationData->LiveInstance.ToSharedRef(), Image->ImageID);
 			UE::Tasks::AddNested(GetImageDescTask);
 			LastTask = GetImageDescTask;
 			
@@ -2518,7 +2502,7 @@ namespace impl
 
 			if (Image->bIsPassThrough)
 			{
-				UE::Tasks::TTask<UE::Mutable::Private::TManagedPtr<const UE::Mutable::Private::FImage>> GetImageTask = OperationData->MutableSystem->GetImage(LastTask, OperationData->InstanceID, Image->ImageID, 0);
+				UE::Tasks::TTask<UE::Mutable::Private::TManagedPtr<const UE::Mutable::Private::FImage>> GetImageTask = OperationData->MutableSystem->GetImage(LastTask, OperationData->LiveInstance.ToSharedRef(), Image->ImageID, 0);
 				UE::Tasks::AddNested(GetImageTask);
 				LastTask = GetImageTask;
 					
@@ -2558,7 +2542,7 @@ namespace impl
 				const int32 MipSizeX = FMath::Max(Image->FullImageSizeX >> MipsToSkip, 1);
 				const int32 MipSizeY = FMath::Max(Image->FullImageSizeY >> MipsToSkip, 1);
 
-				UE::Tasks::TTask<UE::Mutable::Private::TManagedPtr<const UE::Mutable::Private::FImage>> GetImageTask = OperationData->MutableSystem->GetImage(LastTask, OperationData->InstanceID, Image->ImageID, Image->BaseMip + MipsToSkip);
+				UE::Tasks::TTask<UE::Mutable::Private::TManagedPtr<const UE::Mutable::Private::FImage>> GetImageTask = OperationData->MutableSystem->GetImage(LastTask, OperationData->LiveInstance.ToSharedRef(), Image->ImageID, Image->BaseMip + MipsToSkip);
 				UE::Tasks::AddNested(GetImageTask);
 				LastTask = GetImageTask;
 
@@ -2648,49 +2632,20 @@ namespace impl
 	
 
 	// This runs in a worker thread.
-	void Task_Mutable_ReleaseInstance(UE::Mutable::Private::FInstance::FID InstanceID, TSharedPtr<UE::Mutable::Private::FSystem> MutableSystem, bool bLiveUpdateMode)
+	void Task_Mutable_EndUpdate(TSharedRef<UE::Mutable::Private::FLiveInstance> LiveInstance, TSharedPtr<UE::Mutable::Private::FSystem> MutableSystem, bool bLiveUpdateMode)
 	{
 		MUTABLE_CPUPROFILER_SCOPE(Task_Mutable_ReleaseInstance)
 
 		check(MutableSystem);
 
-		if (InstanceID > 0)
-		{
-			MutableSystem->EndUpdate(InstanceID);
-
-			if (!bLiveUpdateMode)
-			{
-				MutableSystem->ReleaseInstance(InstanceID);
-			}
-		}
-
+		const bool bClearRoms = UCustomizableObjectSystem::ShouldClearWorkingMemoryOnUpdateEnd();
+		MutableSystem->EndUpdate(LiveInstance, bClearRoms);
+		
 		MutableSystem->SetImagePixelConversionOverride(nullptr);
 
-		if (UCustomizableObjectSystem::ShouldClearWorkingMemoryOnUpdateEnd())
-		{
-			MutableSystem->ClearWorkingMemory();
-		}
-
 		UCustomizableObjectSystem::GetInstance()->GetPrivate()->MutableTaskGraph.AllowLaunchingMutableTaskLowPriority(true, true);
 	}
 
-
-	// This runs in a worker thread.
-	void Task_Mutable_ReleaseInstanceID(const UE::Mutable::Private::FInstance::FID InstanceID, const TSharedPtr<UE::Mutable::Private::FSystem>& MutableSystem)
-	{
-		MUTABLE_CPUPROFILER_SCOPE(Task_Mutable_ReleaseInstanceID)
-
-		if (InstanceID > 0)
-		{
-			MutableSystem->ReleaseInstance(InstanceID);
-		}
-
-		if (UCustomizableObjectSystem::ShouldClearWorkingMemoryOnUpdateEnd())
-		{
-			MutableSystem->ClearWorkingMemory();
-		}
-	}
-
 	
 	void Task_Game_Callbacks(const TSharedRef<FUpdateContextPrivate>& Context)
 	{
@@ -2815,14 +2770,16 @@ namespace impl
 		
 		// Next Task: Release Mutable. We need this regardless if we cancel or not
 		//-------------------------------------------------------------		
-		const  TSharedPtr<UE::Mutable::Private::FSystem> MutableSystem = CustomizableObjectSystemPrivate->MutableSystem;
+		const TSharedPtr<UE::Mutable::Private::FSystem> MutableSystem = CustomizableObjectSystemPrivate->MutableSystem;
 		CustomizableObjectSystemPrivate->LastUpdateMutableTask = CustomizableObjectSystemPrivate->MutableTaskGraph.AddMutableThreadTask(
-			TEXT("Task_Mutable_ReleaseInstance"),
-			[InstanceID = OperationData->InstanceID, MutableSystem, bLiveUpdateMode = static_cast<bool>(OperationData->bLiveUpdateMode)]()
+			TEXT("Task_Mutable_EndUpdate"),
+			[LiveInstance = OperationData->LiveInstance, MutableSystem, bLiveUpdateMode = static_cast<bool>(OperationData->bLiveUpdateMode)]()
 			{
-				Task_Mutable_ReleaseInstance(InstanceID, MutableSystem, bLiveUpdateMode);
+				Task_Mutable_EndUpdate(LiveInstance.ToSharedRef(), MutableSystem, bLiveUpdateMode);
 			});
-
+		
+		// We no longer need it.
+		OperationData->LiveInstance.Reset();
 
 		// If they are still in this container it means they have not been used. No longer required. Free some memory now.
 		OperationData->ImageToPlatformDataMap.Empty();
@@ -2881,17 +2838,6 @@ namespace impl
 		}
 
 		UCustomizableInstancePrivate* ObjectInstancePrivateData = ObjectInstance->GetPrivate();
-
-		if (OperationData->bLiveUpdateMode)
-		{
-			check(OperationData->InstanceID != 0);
-
-			if (ObjectInstancePrivateData->LiveUpdateModeInstanceID == 0)
-			{
-				// From now this instance will reuse this InstanceID until it gets out of LiveUpdateMode
-				ObjectInstancePrivateData->LiveUpdateModeInstanceID = OperationData->InstanceID;
-			}
-		}
 		
 		// TODO: If this is the first code that runs after the CO program has finished AND if it's
 		// guaranteed that the next CO program hasn't started yet, we need to call ClearActiveObject
@@ -2967,30 +2913,6 @@ namespace impl
 			false,
 			{ Game_LoadUnrealAssets, Mutable_PrepareTextures });
 	}
-
-
-	/** Enqueue the release ID operation in the Mutable queue */
-	void Task_Game_ReleaseInstanceID(const UE::Mutable::Private::FInstance::FID IDToRelease)
-	{
-		MUTABLE_CPUPROFILER_SCOPE(Task_Game_ReleaseInstanceID)
-
-		UCustomizableObjectSystem* System = UCustomizableObjectSystem::GetInstanceChecked();
-		UCustomizableObjectSystemPrivate* SystemPrivate = System->GetPrivate();
-
-		const TSharedPtr<UE::Mutable::Private::FSystem> MutableSystem = SystemPrivate->MutableSystem;
-
-		// Task: Release Instance ID
-		//-------------------------------------------------------------
-		{
-			// Task inputs
-			SystemPrivate->MutableTaskGraph.AddMutableThreadTask(
-				TEXT("Task_Mutable_ReleaseInstanceID"),
-				[IDToRelease, MutableSystem]()
-				{
-					impl::Task_Mutable_ReleaseInstanceID(IDToRelease, MutableSystem);
-				});
-		}
-	}
 	
 	
 	/** "Start Update" */
@@ -3126,31 +3048,10 @@ namespace impl
 	
 		Operation->MipsToSkip = !Operation->bNeverStream ? 255 : 0; // // 0 means generate all mips. 255 means skip all possible mips until only UTexture::GetStaticMinTextureResidentMipCount() are left
 		
-		Operation->bLiveUpdateMode = false;
-
-		if (SystemPrivate->EnableMutableLiveUpdate)
-		{
-			Operation->bLiveUpdateMode = StateData ? StateData->bLiveUpdateMode : false;
-		}
-		
-		if (Operation->bLiveUpdateMode && (!Operation->bNeverStream || Operation->MipsToSkip > 0))
-		{
-			UE_LOG(LogMutable, Warning, TEXT("Instance LiveUpdateMode does not yet support progressive streaming of Mutable textures. Disabling LiveUpdateMode for this update."));
-			Operation->bLiveUpdateMode = false;
-		}
-		
 		Operation->bUseCaches = Operation->bKeepOwnershipOfGeneratedResources && !Operation->bBake && !FLogBenchmarkUtil::IsBenchmarkingReportingEnabled();
 		Operation->bUseMeshCache = Operation->bUseCaches && Object->bEnableMeshCache && !Operation->bLiveUpdateMode && CVarEnableMeshCache.GetValueOnAnyThread(true);
 		Operation->bUseResueTextureCache = Operation->bUseCaches && Operation->bNeverStream && CVarEnableMutableReuseInstanceTextures.GetValueOnGameThread();
 		
-		if (!Operation->bLiveUpdateMode && InstancePrivate->LiveUpdateModeInstanceID != 0)
-		{
-			// The instance was in live update mode last update, but now it's not. So the Id and resources have to be released.
-			// Enqueue a new mutable task to release them
-			Task_Game_ReleaseInstanceID(InstancePrivate->LiveUpdateModeInstanceID);
-			InstancePrivate->LiveUpdateModeInstanceID = 0;
-		}
-		
 		Operation->Model = ObjectPrivate->GetModel().ToSharedRef();
 		Operation->MeshIdRegistry = ObjectPrivate->MeshIdRegistry;
 		Operation->ImageIdRegistry = ObjectPrivate->ImageIdRegistry;
@@ -3199,8 +3100,6 @@ namespace impl
 		
 		// Task: Mutable Update and GetMesh
 		//-------------------------------------------------------------
-		Operation->InstanceID = Operation->bLiveUpdateMode ? InstancePrivate->LiveUpdateModeInstanceID : 0;
-		
 		Operation->bStreamMeshLODs = IsStreamingEnabled(*Object, State);
 
 #if WITH_EDITOR
@@ -3281,30 +3180,27 @@ namespace impl
 			Operation->ExternalResourceProvider = MakeShared<FUnrealMutableResourceProvider>(BoneNames);
 			Operation->PassthroughObjectLoader = MakeShared<UE::Mutable::Private::FPassthroughObjectLoader>();
 
-			if (Operation->bLiveUpdateMode)
+			Operation->bLiveUpdateMode = SystemPrivate->EnableMutableLiveUpdate && (StateData ? StateData->bLiveUpdateMode : false);
+		
+			if (Operation->bLiveUpdateMode && (!Operation->bNeverStream || Operation->MipsToSkip > 0))
 			{
-				if (Operation->InstanceID == 0)
-				{
-					// It's the first update since the instance was put in LiveUpdate Mode, this ID will be reused from now on
-					Operation->InstanceID = MutableSystem->NewInstance(Operation->Model, Operation->ExternalResourceProvider, Operation->PassthroughObjectLoader);
-					UE_LOG(LogMutable, Verbose, TEXT("Creating Mutable instance with id [%d] for reuse "), Operation->InstanceID);
-				}
-				else
-				{
-					// The instance was already in LiveUpdate Mode, the ID is reused
-					MutableSystem->ReuseInstance(Operation->InstanceID, Operation->ExternalResourceProvider, Operation->PassthroughObjectLoader);
-					UE_LOG(LogMutable, Verbose, TEXT("Reusing Mutable instance with id [%d] "), Operation->InstanceID);
-				}
+				UE_LOG(LogMutable, Warning, TEXT("Instance LiveUpdateMode does not yet support progressive streaming of Mutable textures. Disabling LiveUpdateMode for this update."));
+				Operation->bLiveUpdateMode = false;
 			}
-			else
+			
+			Operation->LiveInstance = Operation->bLiveUpdateMode ? Instance->GetPrivate()->LiveInstance : nullptr;
+			
+			if (!Operation->LiveInstance)
 			{
-				// In non-LiveUpdate mode, we are forcing the recreation of mutable-side instances with every update.
-				check(Operation->InstanceID == 0);
-				Operation->InstanceID = MutableSystem->NewInstance(Operation->Model, Operation->ExternalResourceProvider, Operation->PassthroughObjectLoader);
-				UE_LOG(LogMutable, Verbose, TEXT("Creating Mutable instance with id [%d] "), Operation->InstanceID);
+				Operation->LiveInstance = MutableSystem->NewInstance(Operation->Model, Operation->ExternalResourceProvider, Operation->PassthroughObjectLoader);
 			}
+
+			Instance->GetPrivate()->LiveInstance = Operation->bLiveUpdateMode ? Operation->LiveInstance : nullptr;
 			
-			Operation->MutableInstance = MutableSystem->BeginUpdate_GameThread(Operation->InstanceID,
+			Operation->LiveInstance->ExternalResourceProvider = Operation->ExternalResourceProvider;
+			Operation->LiveInstance->PassthroughObjectLoader = Operation->PassthroughObjectLoader;
+			
+			Operation->MutableInstance = MutableSystem->BeginUpdate_GameThread(Operation->LiveInstance.ToSharedRef(),
 				Operation->Parameters,
 				Operation->MeshIdRegistry,
 				Operation->ImageIdRegistry,
@@ -3788,7 +3684,6 @@ int32 UCustomizableObjectSystem::TickInternal(const bool bBlocking)
 
 		// Free memory before starting the new update
 		DiscardInstances();
-		ReleaseInstanceIDs();
 	}
 	
 	const int32 RemainingTasks = Private->MutableTaskGraph.Tick();
@@ -3910,30 +3805,6 @@ void UCustomizableObjectSystem::DiscardInstances()
 }
 
 
-TAutoConsoleVariable<int32> CVarMaxNumInstanceIDsToReleasePerTick(
-	TEXT("mutable.MaxNumInstanceIDsToReleasePerTick"),
-	30,
-	TEXT("The maximum number of stale instances IDs that will be released per tick by Mutable."),
-	ECVF_Scalability);
-
-
-void UCustomizableObjectSystem::ReleaseInstanceIDs()
-{
-	// Handle ID discards
-	int32 NumIDsReleased = 0;
-	const int32 IDReleaseLimitPerTick = CVarMaxNumInstanceIDsToReleasePerTick.GetValueOnGameThread();
-
-	for (auto Iterator = Private->MutablePendingInstanceWork.GetIDsToReleaseIterator();
-		Iterator && NumIDsReleased < IDReleaseLimitPerTick; ++Iterator)
-	{
-		impl::Task_Game_ReleaseInstanceID(*Iterator);
-
-		Iterator.RemoveCurrent();
-		NumIDsReleased++;
-	}
-}
-
-
 bool UCustomizableObjectSystem::IsUpdating(const UCustomizableObjectInstance* Instance) const
 {
 	if (!Instance)
diff --git a/Engine/Plugins/Mutable/Source/CustomizableObject/Public/MuCO/CustomizableObjectInstance.h b/Engine/Plugins/Mutable/Source/CustomizableObject/Public/MuCO/CustomizableObjectInstance.h
--- a/Engine/Plugins/Mutable/Source/CustomizableObject/Public/MuCO/CustomizableObjectInstance.h
+++ b/Engine/Plugins/Mutable/Source/CustomizableObject/Public/MuCO/CustomizableObjectInstance.h
@@ -431,7 +431,7 @@ public:
 	/** Copy parameters from the given Instance. */
 	CUSTOMIZABLEOBJECT_API void CopyParametersFromInstance(UCustomizableObjectInstance* Instance);
 
-	/** Immediately destroy the Mutable Core Live Update Instance attached to this (if exists). */
+	/** Destroy the Mutable Core Live Update Instance attached to this (if exists). */
 	CUSTOMIZABLEOBJECT_API void DestroyLiveUpdateInstance();
 	
 	/** Return true if the given parameter has any effect in the current object state, and considering the current values of the other parameters. */
diff --git a/Engine/Plugins/Mutable/Source/CustomizableObject/Public/MuCO/CustomizableObjectSystem.h b/Engine/Plugins/Mutable/Source/CustomizableObject/Public/MuCO/CustomizableObjectSystem.h
--- a/Engine/Plugins/Mutable/Source/CustomizableObject/Public/MuCO/CustomizableObjectSystem.h
+++ b/Engine/Plugins/Mutable/Source/CustomizableObject/Public/MuCO/CustomizableObjectSystem.h
@@ -243,7 +243,6 @@ private:
 	UE_API int32 TickInternal(bool bBlocking);
 	
 	UE_API void DiscardInstances();
-	UE_API void ReleaseInstanceIDs();
 
 public:
 	/** Return true if the instance is being updated. */
diff --git a/Engine/Plugins/Mutable/Source/MutableRuntime/Internal/MuR/Instance.h b/Engine/Plugins/Mutable/Source/MutableRuntime/Internal/MuR/Instance.h
--- a/Engine/Plugins/Mutable/Source/MutableRuntime/Internal/MuR/Instance.h
+++ b/Engine/Plugins/Mutable/Source/MutableRuntime/Internal/MuR/Instance.h
@@ -94,9 +94,6 @@ namespace UE::Mutable::Private
 		//-----------------------------------------------------------------------------------------
 		// Own interface
 		//-----------------------------------------------------------------------------------------
-
-        //! Type for the instance unique identifiers.
-        typedef uint32 FID;
 		
 		//! Get the number of components of this instance.
 		UE_API int32 GetComponentCount() const;
diff --git a/Engine/Plugins/Mutable/Source/MutableRuntime/Internal/MuR/System.h b/Engine/Plugins/Mutable/Source/MutableRuntime/Internal/MuR/System.h
--- a/Engine/Plugins/Mutable/Source/MutableRuntime/Internal/MuR/System.h
+++ b/Engine/Plugins/Mutable/Source/MutableRuntime/Internal/MuR/System.h
@@ -70,7 +70,7 @@ namespace UE::Mutable::Private
 	class FParameters;
     class FMesh;
 	class FExtensionDataStreamer;
-	
+	struct FLiveInstance;
 
 	/** */
 	enum class EExecutionStrategy : uint8
@@ -1487,7 +1487,6 @@ namespace UE::Mutable::Private
 	*/
 	struct FLiveInstance
 	{
-		FInstance::FID InstanceID;
 		int32 State = 0;
 		uint32 LODMask = 0;
 		TManagedPtr<const FInstance> Instance;
@@ -1503,31 +1502,12 @@ namespace UE::Mutable::Private
 		/** Cached data for the generation of this instance. */
 		TSharedPtr<FProgramCache> Cache;
 
+		TSharedPtr<FMeshIdRegistry> MeshIdRegistry;
+		TSharedPtr<FImageIdRegistry> ImageIdRegistry;
+		TSharedPtr<FMaterialIdRegistry> MaterialIdRegistry;
+		TSharedPtr<FSkeletalMeshIdRegistry> SkeletalMeshIdRegistry;
 		TSharedPtr<FExternalResourceProvider> ExternalResourceProvider;
 		TSharedPtr<FPassthroughObjectLoader> PassthroughObjectLoader;
-		
-		~FLiveInstance()
-		{
-			// Manually done to trace mem deallocations
-			MUTABLE_CPUPROFILER_SCOPE(LiveInstanceDestructor);
-			{
-				MUTABLE_CPUPROFILER_SCOPE(LiveInstanceDestructor_DeleteCache);
-				Cache = nullptr;
-			}
-
-			{	
-				MUTABLE_CPUPROFILER_SCOPE(LiveInstanceDestructor_DeleteOldParameters);
-				OldParameters = nullptr;
-			}
-			{
-				MUTABLE_CPUPROFILER_SCOPE(LiveInstanceDestructor_DeleteInstance);
-				Instance = nullptr;
-			}
-			{
-				MUTABLE_CPUPROFILER_SCOPE(LiveInstanceDestructor_DeleteModel);
-				Model = nullptr;
-			}
-		}
 	};
 
 
@@ -1570,7 +1550,7 @@ namespace UE::Mutable::Private
         TMemoryTrackedArray<FModelCacheEntry> CachePerModel;
 
 		/** Data for each mutable instance that is being updated. */
-		TMemoryTrackedArray<FLiveInstance> LiveInstances;
+		TMemoryTrackedArray<TWeakPtr<FLiveInstance>> LiveInstances;
 
 		/** Temporary reference to the memory of the current instance being updated. Only valid during a mutable "atomic" operation, like a BeginUpdate or a GetImage. */
 		TSharedPtr<FProgramCache> CurrentInstanceCache;
@@ -1735,66 +1715,72 @@ namespace UE::Mutable::Private
 
 		int32 GetRomBytes() const;
 
-		/** Calculate the amount of bytes in data cached in the level 0 and 1 cache in all live instances. */
-		int32 GetCacheBytes() const
-		{
-			int32 Result = 0;
-			TSet<TManagedPtr<const FResource>> Cache0Unique;
-			TSet<TManagedPtr<const FResource>> Cache1Unique;
-
-			for (const FLiveInstance& Instance : LiveInstances)
-			{
-			CodeContainer<FProgramCache::FOpExecutionData>::iterator it = Instance.Cache->OpExecutionData.begin();
-			for (; it.IsValid(); ++it)
-			{
-				if (!(*it).DataTypeIndex)
-				{
-					continue;
-				}
-
-				TSet<TManagedPtr<const FResource>>* TargetSet = &Cache0Unique;
-				if ((*it).IsCacheLocked)
-				{
-					TargetSet = &Cache1Unique;
-				}
-
-				switch ((*it).DataType)
-				{
-				case EDataType::Image:
-				{
-					TManagedPtr<const FImage> Value = Instance.Cache->ImageResults[(*it).DataTypeIndex].Pin();
-					if (Value)
-					{
-						TargetSet->Add(MoveTemp(Value));
-					}
-					break;
-				}
-				case EDataType::Mesh:
-				{
-					TManagedPtr<const FMesh> Value = Instance.Cache->MeshResults[(*it).DataTypeIndex].Pin();
-					if (Value)
-					{
-						TargetSet->Add(MoveTemp(Value));
-					}
-					break;
-				}
-				default:
-					break;
-				}
-			}
-		}
-
-			// Merge
-			Cache0Unique.Append(Cache1Unique);
-
-			// Count
-			for (const TManagedPtr<const FResource>& Value : Cache0Unique)
-			{
-				Result += Value->GetDataSize();
-			}
-
-			return Result;
-		}
+    	/** Calculate the amount of bytes in data cached in the level 0 and 1 cache in all live instances. */
+    	int32 GetCacheBytes() const
+    	{
+    		int32 Result = 0;
+    		TSet<TManagedPtr<const FResource>> Cache0Unique;
+    		TSet<TManagedPtr<const FResource>> Cache1Unique;
+
+    		for (const TWeakPtr<FLiveInstance>& WeakInstance : LiveInstances)
+    		{
+    			TSharedPtr<FLiveInstance> Instance = WeakInstance.Pin();
+    			if (!Instance)
+    			{
+    				continue;
+    			}
+
+    			CodeContainer<FProgramCache::FOpExecutionData>::iterator it = Instance->Cache->OpExecutionData.begin();
+    			for (; it.IsValid(); ++it)
+    			{
+    				if (!(*it).DataTypeIndex)
+    				{
+    					continue;
+    				}
+
+    				TSet<TManagedPtr<const FResource>>* TargetSet = &Cache0Unique;
+    				if ((*it).IsCacheLocked)
+    				{
+    					TargetSet = &Cache1Unique;
+    				}
+
+    				switch ((*it).DataType)
+    				{
+    				case EDataType::Image:
+    					{
+    						TManagedPtr<const FImage> Value = Instance->Cache->ImageResults[(*it).DataTypeIndex].Pin();
+    						if (Value)
+    						{
+    							TargetSet->Add(MoveTemp(Value));
+    						}
+    						break;
+    					}
+    				case EDataType::Mesh:
+    					{
+    						TManagedPtr<const FMesh> Value = Instance->Cache->MeshResults[(*it).DataTypeIndex].Pin();
+    						if (Value)
+    						{
+    							TargetSet->Add(MoveTemp(Value));
+    						}
+    						break;
+    					}
+    				default:
+    					break;
+    				}
+    			}
+    		}
+
+    		// Merge
+    		Cache0Unique.Append(Cache1Unique);
+
+    		// Count
+    		for (const TManagedPtr<const FResource>& Value : Cache0Unique)
+    		{
+    			Result += Value->GetDataSize();
+    		}
+
+    		return Result;
+    	}
 
 		/** 
 		 * Remove all intermediate data (big and small) from the memory except for the one that has been explicitely
@@ -2017,11 +2003,6 @@ namespace UE::Mutable::Private
          * Refer to Settings::SetWorkingMemoryBudget for more information.
 		 */
         UE_API void SetWorkingMemoryBytes( uint64 Bytes );
-
-        /** Removes all the possible working memory regardless of the budget set. This may make following 
-		* operations take longer.
-		*/
-		UE_API void ClearWorkingMemory();
 		
 		/** Set a function that will be used to convert image pixel formats instead of the internal conversion. 
 		* \warning The provided function can be called from any thread, and also concurrently.
@@ -2034,14 +2015,9 @@ namespace UE::Mutable::Private
         //! calls to BeginUpdate/EndUpdate.
         //! A call to NewInstance must be paired to a call to ReleasesInstance when the instance is
         //! no longer needed.
-        //! \return An identifier that is always bigger than 0.
-        UE_API FInstance::FID NewInstance(const TSharedPtr<const FModel>& Model,
+        UE_API TSharedRef<FLiveInstance> NewInstance(const TSharedPtr<const FModel>& Model,
         	const TSharedPtr<FExternalResourceProvider>& ExternalResourceProvider,
         	const TSharedPtr<FPassthroughObjectLoader>& PassthroughObjectLoader);
-
-		UE_API void ReuseInstance(FInstance::FID InstanceId,
-			const TSharedPtr<FExternalResourceProvider>& ExternalResourceProvider,
-			const TSharedPtr<FPassthroughObjectLoader>& PassthroughObjectLoader);
 		
         //! \brief Update an instance with a new parameter set and/or state.
         //!
@@ -2057,7 +2033,7 @@ namespace UE::Mutable::Private
         //! the same instanceID parameter.
 		//!
 		//! Very fast. Can be called in the Game Thread.
-		UE_API TManagedPtr<const FInstance> BeginUpdate_GameThread(FInstance::FID InstanceID,
+		UE_API TManagedPtr<const FInstance> BeginUpdate_GameThread(const TSharedRef<FLiveInstance>& LiveInstance,
 			const TSharedPtr<const FParameters>& Params,
 			const TSharedPtr<FMeshIdRegistry>& InMeshIdRegistry,
 			const TSharedPtr<FImageIdRegistry>& InImageIdRegistry,
@@ -2067,7 +2043,7 @@ namespace UE::Mutable::Private
 			uint32 LodMask);
 
 		/** Potentially expensive. Must be called in the Mutable Thread. */
-		UE_API TManagedPtr<const FInstance> BeginUpdate_MutableThread(FInstance::FID InInstanceID,
+		UE_API TManagedPtr<const FInstance> BeginUpdate_MutableThread(const TSharedRef<FLiveInstance>& LiveInstance,
 			const TSharedPtr<const FParameters>& InParams,
 			const TSharedPtr<FMeshIdRegistry>& InMeshIdRegistry,
 			const TSharedPtr<FImageIdRegistry>& InImageIdRegistry,
@@ -2078,49 +2054,44 @@ namespace UE::Mutable::Private
 		
 		//! Only valid between BeginUpdate and EndUpdate
 		//! Calculate the description of an image, without generating it.
-		UE_API UE::Tasks::TTask<FExtendedImageDesc> GetImageDesc(UE::Tasks::FTask Prerequisite, FInstance::FID InstanceID, const FImageId& ImageId);
+		UE_API UE::Tasks::TTask<FExtendedImageDesc> GetImageDesc(UE::Tasks::FTask Prerequisite, const TSharedRef<FLiveInstance>& LiveInstance, const FImageId& ImageId);
 
 		//! Only valid between BeginUpdate and EndUpdate
 		//! \param MipsToSkip Number of mips to skip compared from the full image.
 		//! If 0, all mip levels will be generated. If more levels than possible to discard are specified, 
 		//! the image will still contain a minimum number of mips specified at model compile time.
-		UE_API UE::Tasks::TTask<TManagedPtr<const FImage>> GetImage(UE::Tasks::FTask Prerequisite, FInstance::FID InstanceID, const FImageId& ImageId, int32 MipsToSkip = 0);
+		UE_API UE::Tasks::TTask<TManagedPtr<const FImage>> GetImage(UE::Tasks::FTask Prerequisite, const TSharedRef<FLiveInstance>& LiveInstance, const FImageId& ImageId, int32 MipsToSkip = 0);
 
 		//! Only valid between BeginUpdate and EndUpdate
-		UE_API UE::Tasks::TTask<TManagedPtr<const FMesh>> GetMesh(UE::Tasks::FTask Prerequisite, FInstance::FID InstanceID, const FMeshId& MeshId, EMeshContentFlags MeshContentFilter = EMeshContentFlags::AllFlags);
+		UE_API UE::Tasks::TTask<TManagedPtr<const FMesh>> GetMesh(UE::Tasks::FTask Prerequisite, const TSharedRef<FLiveInstance>& LiveInstance, const FMeshId& MeshId, EMeshContentFlags MeshContentFilter = EMeshContentFlags::AllFlags);
 		
-		UE_API UE::Tasks::TTask<TManagedPtr<const FSkeletalMesh>> GetSkeletalMesh(UE::Tasks::FTask Prerequisite, FInstance::FID InstanceID, const FSkeletalMeshId& SkeletalMeshId, EMeshContentFlags MeshContentFilter);
+		UE_API UE::Tasks::TTask<TManagedPtr<const FSkeletalMesh>> GetSkeletalMesh(UE::Tasks::FTask Prerequisite, const TSharedRef<FLiveInstance>& LiveInstance, const FSkeletalMeshId& SkeletalMeshId, EMeshContentFlags MeshContentFilter);
 
 		//! Only valid between BeginUpdate and EndUpdate
 		//! Calculate the description of an image, without generating it.
-		UE_API FExtendedImageDesc GetImageDescInline(FInstance::FID InstanceID, const FImageId& ImageId);
+		UE_API FExtendedImageDesc GetImageDescInline(const TSharedRef<FLiveInstance>& LiveInstance, const FImageId& ImageId);
 
 		//! Only valid between BeginUpdate and EndUpdate
 		//! \param MipsToSkip Number of mips to skip compared from the full image.
 		//! If 0, all mip levels will be generated. If more levels than possible to discard are specified, 
 		//! the image will still contain a minimum number of mips specified at model compile time.
-		UE_API TManagedPtr<const FImage> GetImageInline(FInstance::FID InstanceID, const FImageId& ImageId, int32 MipsToSkip = 0);
+		UE_API TManagedPtr<const FImage> GetImageInline(const TSharedRef<FLiveInstance>& LiveInstance, const FImageId& ImageId, int32 MipsToSkip = 0);
 
         //! Only valid between BeginUpdate and EndUpdate
-		UE_API TManagedPtr<const FMesh> GetMeshInline(FInstance::FID InstanceID, const FMeshId& MeshId, EMeshContentFlags MeshContentFilter = EMeshContentFlags::AllFlags);
+		UE_API TManagedPtr<const FMesh> GetMeshInline(const TSharedRef<FLiveInstance>& LiveInstance, const FMeshId& MeshId, EMeshContentFlags MeshContentFilter = EMeshContentFlags::AllFlags);
 		
-		UE_API UE::Tasks::TTask<TManagedPtr<const FMaterial>> GetMaterial(UE::Tasks::FTask Prerequisite, FInstance::FID InstanceID, const FMaterialId& MaterialId);
+		UE_API UE::Tasks::TTask<TManagedPtr<const FMaterial>> GetMaterial(UE::Tasks::FTask Prerequisite, const TSharedRef<FLiveInstance>& LiveInstance, const FMaterialId& MaterialId);
 		
         //! Invalidate and free the last Instance data returned by a call to BeginUpdate with
         //! the same instance index. After a call to this method, that Instance cannot be used any
         //! more and its content is undefined.
         //! \param instance The index of the instance whose last data will be invalidated.
-        UE_API void EndUpdate(FInstance::FID InstanceID);
-
-        //! Completely destroy an instance. After a call to this method the given instance cannot be
-        //! updated any more, and its resources may have been freed.
-        //! \param instance The id of the instance to destroy.
-        UE_API void ReleaseInstance(FInstance::FID InstanceID);
-
+        UE_API void EndUpdate(const TSharedRef<FLiveInstance>& LiveInstance, bool bClearRoms);
+		
 		//! Calculate the relevancy of every parameter. Some parameters may be unused depending on
 		//! the values of other parameters. This is useful to hide irrelevant parameters in
 		//! dynamic user interfaces.
-        UE_API void GetParameterRelevancy(FInstance::FID InstanceID,
+        UE_API void GetParameterRelevancy(const TSharedRef<FLiveInstance>& LiveInstance,
         	const TSharedPtr<const FParameters>& FParameters,
 			TArray<uint32>& OutRelevantParameters);
 
@@ -2163,9 +2134,6 @@ namespace UE::Mutable::Private
         //! Data streaming interface, if any.
 		TSharedPtr<FModelReader> StreamInterface;
     	
-		/** Counter used to generate unique IDs for every new instance created in the system. */
-        FInstance::FID LastInstanceID = 0;
-
 		/** This flag is turned on when a streaming error or similar happens. Results are not usable.
 		* This should only happen in-editor.
 		*/
@@ -2176,10 +2144,7 @@ namespace UE::Mutable::Private
 
 		FWorkingMemoryManager WorkingMemoryManager;
 
-		/** The pointer returned by this function is only valid for the duration of the current mutable operation. */
-		inline FLiveInstance* FindLiveInstance(FInstance::FID);
-
-        bool CheckUpdatedParameters( const FLiveInstance*, const TSharedPtr<const FParameters>&, uint64& OutUpdatedParameters );
+		bool CheckUpdatedParameters(const TSharedRef<FLiveInstance>& LiveInstance, const TSharedPtr<const FParameters>&, uint64& OutUpdatedParameters );
 		
 		void RunCode(
 			const TSharedPtr<FExternalResourceProvider>& ExternalResourceProvider,
diff --git a/Engine/Plugins/Mutable/Source/MutableRuntime/Private/MuR/System.cpp b/Engine/Plugins/Mutable/Source/MutableRuntime/Private/MuR/System.cpp
--- a/Engine/Plugins/Mutable/Source/MutableRuntime/Private/MuR/System.cpp
+++ b/Engine/Plugins/Mutable/Source/MutableRuntime/Private/MuR/System.cpp
@@ -116,28 +116,6 @@ namespace UE::Mutable::Private
 	}
 
 	
-	void FSystem::ClearWorkingMemory()
-    {
-		LLM_SCOPE_BYNAME(TEXT("MutableRuntime"));
-
-		// rom caches
-		for (FWorkingMemoryManager::FModelCacheEntry& ModelCache : WorkingMemoryManager.CachePerModel)
-		{
-			if (const TSharedPtr<const FModel> CacheModel = ModelCache.Model.Pin())
-			{
-				const FProgram& Program = CacheModel->Program;
-
-				for (int32 RomIndex = 0; RomIndex < Program.Roms.Num(); ++RomIndex)
-				{
-					Program.UnloadRom(RomIndex);
-				}
-			}
-		}
-    	
-		UpdateStats();
-	}
-
-
 	void FSystem::SetImagePixelConversionOverride(const FImageOperator::FImagePixelFormatFunc& FormatFunc)
 	{
 		LLM_SCOPE_BYNAME(TEXT("MutableRuntime"));
@@ -146,42 +124,35 @@ namespace UE::Mutable::Private
 	}
 
    
-    FInstance::FID FSystem::NewInstance( const TSharedPtr<const FModel>& InModel,
+    TSharedRef<FLiveInstance> FSystem::NewInstance( const TSharedPtr<const FModel>& InModel,
     	const TSharedPtr<FExternalResourceProvider>& ExternalResourceProvider,
     	const TSharedPtr<FPassthroughObjectLoader>& PassthroughObjectLoader)
     {
 		LLM_SCOPE_BYNAME(TEXT("MutableRuntime"));
 		MUTABLE_CPUPROFILER_SCOPE(NewInstance);
 
-		FLiveInstance InstanceData;
-		InstanceData.InstanceID = ++LastInstanceID;
-		InstanceData.Instance = nullptr;
-		InstanceData.Model = InModel;
-		InstanceData.State = -1;
-		InstanceData.Cache = MakeShared<FProgramCache>();
-    	InstanceData.ExternalResourceProvider = ExternalResourceProvider;
-    	InstanceData.PassthroughObjectLoader = PassthroughObjectLoader;
-		WorkingMemoryManager.LiveInstances.Add(InstanceData);
-
-		TRACE_COUNTER_SET(MutableRuntime_LiveInstances, WorkingMemoryManager.LiveInstances.Num());
-
-		return InstanceData.InstanceID;
+    	UE_LOG(LogMutableCore, Verbose, TEXT("Creating Mutable instance"));
+    	
+		TSharedRef<FLiveInstance> InstanceData = MakeShared<FLiveInstance>();
+		InstanceData->Instance = nullptr;
+		InstanceData->Model = InModel;
+		InstanceData->State = -1;
+		InstanceData->Cache = MakeShared<FProgramCache>();
+    	InstanceData->ExternalResourceProvider = ExternalResourceProvider;
+    	InstanceData->PassthroughObjectLoader = PassthroughObjectLoader;
+    	
+    	WorkingMemoryManager.LiveInstances.Add(InstanceData.ToWeakPtr());
+    	WorkingMemoryManager.LiveInstances.RemoveAllSwap([](const TWeakPtr<FLiveInstance>& Instance)
+		{
+    		return !Instance.IsValid();
+		});
+    	
+		return InstanceData;
 	}
+	
 
-
-    void FSystem::ReuseInstance(FInstance::FID InstanceId,
-    	const TSharedPtr<FExternalResourceProvider>& ExternalResourceProvider,
-    	const TSharedPtr<FPassthroughObjectLoader>& PassthroughObjectLoader)
-    {
-    	FLiveInstance* LiveInstance = FindLiveInstance(InstanceId);
-    	check(LiveInstance);
-
-    	LiveInstance->ExternalResourceProvider = ExternalResourceProvider;
-    	LiveInstance->PassthroughObjectLoader = PassthroughObjectLoader;
-    }
-
-   
-    TManagedPtr<const FInstance> FSystem::BeginUpdate_GameThread(FInstance::FID InInstanceID,
+    //---------------------------------------------------------------------------------------------
+    TManagedPtr<const FInstance> FSystem::BeginUpdate_GameThread(const TSharedRef<FLiveInstance>& LiveInstance,
                                      const TSharedPtr<const FParameters>& InParams,
                                      const TSharedPtr<FMeshIdRegistry>& InMeshIdRegistry,
 									 const TSharedPtr<FImageIdRegistry>& InImageIdRegistry,
@@ -193,8 +164,6 @@ namespace UE::Mutable::Private
 		LLM_SCOPE_BYNAME(TEXT("MutableRuntime"));
 		MUTABLE_CPUPROFILER_SCOPE(BeginUpdate_GameThread);
 
-    	FLiveInstance* LiveInstance = FindLiveInstance(InInstanceID);
-    	check(LiveInstance);
 		check(InParams);
 
     	FScheduledOpInline Root(LiveInstance->Model->Program.States[InStateIndex].Root, 0, 0, true);
@@ -221,7 +190,7 @@ namespace UE::Mutable::Private
 	}
 
 	
-    TManagedPtr<const FInstance> FSystem::BeginUpdate_MutableThread(FInstance::FID InInstanceID,
+    TManagedPtr<const FInstance> FSystem::BeginUpdate_MutableThread(const TSharedRef<FLiveInstance>& LiveInstance,
 		const TSharedPtr<const FParameters>& InParams,
 		const TSharedPtr<FMeshIdRegistry>& InMeshIdRegistry,
 		const TSharedPtr<FImageIdRegistry>& InImageIdRegistry,
@@ -239,14 +208,7 @@ namespace UE::Mutable::Private
 			UE_LOG(LogMutableCore, Error, TEXT("Invalid parameters in mutable update."));
 			return nullptr;
 		}
-
-		FLiveInstance* LiveInstance = FindLiveInstance(InInstanceID);
-		if (!LiveInstance)
-		{
-			UE_LOG(LogMutableCore, Error, TEXT("Invalid instance id in mutable update."));
-			return nullptr;
-		}
-
+    	
 		WorkingMemoryManager.CurrentInstanceCache = LiveInstance->Cache;
     	WorkingMemoryManager.CurrentPassthroughObjectLoader = LiveInstance->PassthroughObjectLoader;
 
@@ -330,7 +292,7 @@ namespace UE::Mutable::Private
     }
 	
 
-	TManagedPtr<const FImage> FSystem::GetImageInline(FInstance::FID instanceID, const FImageId& ImageId, int32 MipsToSkip)
+	TManagedPtr<const FImage> FSystem::GetImageInline(const TSharedRef<FLiveInstance>& LiveInstance, const FImageId& ImageId, int32 MipsToSkip)
 	{
 		LLM_SCOPE_BYNAME(TEXT("MutableRuntime"));
 		MUTABLE_CPUPROFILER_SCOPE(SystemGetImage);
@@ -341,8 +303,6 @@ namespace UE::Mutable::Private
     	}
     	
 		// Find the live instance
-		FLiveInstance* LiveInstance = FindLiveInstance(instanceID);
-		check(LiveInstance);
 		WorkingMemoryManager.CurrentInstanceCache = LiveInstance->Cache;
 		WorkingMemoryManager.CurrentPassthroughObjectLoader = LiveInstance->PassthroughObjectLoader;
 
@@ -360,7 +320,7 @@ namespace UE::Mutable::Private
 		return Result;
 	}
 
-	UE::Tasks::TTask<TManagedPtr<const FImage>> FSystem::GetImage(UE::Tasks::FTask Prerequisite, FInstance::FID instanceID, const FImageId& ImageId, int32 MipsToSkip)
+	UE::Tasks::TTask<TManagedPtr<const FImage>> FSystem::GetImage(UE::Tasks::FTask Prerequisite, const TSharedRef<FLiveInstance>& LiveInstance, const FImageId& ImageId, int32 MipsToSkip)
 	{
 		LLM_SCOPE_BYNAME(TEXT("MutableRuntime"));
 		MUTABLE_CPUPROFILER_SCOPE(SystemGetImage);
@@ -378,11 +338,7 @@ namespace UE::Mutable::Private
     	}
 
     	OP::ADDRESS RootAddress = ImageId.GetKey()->Address;
-
-		// Find the live instance
-		FLiveInstance* LiveInstance = FindLiveInstance(instanceID);
-		check(LiveInstance);
-
+    	
     	EOpType OpType = LiveInstance->Model->Program.GetOpType(RootAddress);
     	if (GetOpDataType(OpType) != EDataType::Image)
     	{
@@ -463,7 +419,8 @@ namespace UE::Mutable::Private
 		TEXT("If different than 0, clear the image desc cache at every image."),
 		ECVF_Scalability);
 
-	FExtendedImageDesc FSystem::GetImageDescInline(FInstance::FID instanceID, const FImageId& ImageId)
+	
+	FExtendedImageDesc FSystem::GetImageDescInline(const TSharedRef<FLiveInstance>& LiveInstance, const FImageId& ImageId)
 	{
 		LLM_SCOPE_BYNAME(TEXT("MutableRuntime"));
 		MUTABLE_CPUPROFILER_SCOPE(SystemGetImageDesc);
@@ -476,8 +433,6 @@ namespace UE::Mutable::Private
 		FExtendedImageDesc Result;
 
 		// Find the live instance
-		FLiveInstance* LiveInstance = FindLiveInstance(instanceID);
-		check(LiveInstance);
 		WorkingMemoryManager.CurrentInstanceCache = LiveInstance->Cache;
 		WorkingMemoryManager.CurrentPassthroughObjectLoader = LiveInstance->PassthroughObjectLoader;
 
@@ -528,7 +483,7 @@ namespace UE::Mutable::Private
 	}
 
 
-	UE::Tasks::TTask<FExtendedImageDesc> FSystem::GetImageDesc(UE::Tasks::FTask Prerequisite, FInstance::FID InstanceID, const FImageId& ImageId)
+	UE::Tasks::TTask<FExtendedImageDesc> FSystem::GetImageDesc(UE::Tasks::FTask Prerequisite, const TSharedRef<FLiveInstance>& LiveInstance, const FImageId& ImageId)
 	{
 		LLM_SCOPE_BYNAME(TEXT("MutableRuntime"));
 		MUTABLE_CPUPROFILER_SCOPE(SystemGetImageDesc);
@@ -547,9 +502,6 @@ namespace UE::Mutable::Private
 		
 		OP::ADDRESS RootAddress = ImageId.GetKey()->Address;
 		
-		FLiveInstance* LiveInstance = FindLiveInstance(InstanceID);
-		check(LiveInstance);
-
 		EOpType OpType = LiveInstance->Model->Program.GetOpType(RootAddress);
 		if (GetOpDataType(OpType) != EDataType::Image)
 		{
@@ -623,7 +575,7 @@ namespace UE::Mutable::Private
 	}	
 
 	
-    TManagedPtr<const FMesh> FSystem::GetMeshInline(FInstance::FID instanceID, const FMeshId& MeshId, EMeshContentFlags MeshContentFilter)
+    TManagedPtr<const FMesh> FSystem::GetMeshInline(const TSharedRef<FLiveInstance>& LiveInstance, const FMeshId& MeshId, EMeshContentFlags MeshContentFilter)
     {
 		LLM_SCOPE_BYNAME(TEXT("MutableRuntime"));
 		MUTABLE_CPUPROFILER_SCOPE(SystemGetMesh);
@@ -633,9 +585,6 @@ namespace UE::Mutable::Private
 			return MakeManaged<FMesh>();
 		}
 
-		// Find the live instance
-		FLiveInstance* LiveInstance = FindLiveInstance(instanceID);
-		check(LiveInstance);
 		WorkingMemoryManager.CurrentInstanceCache = LiveInstance->Cache;
 		WorkingMemoryManager.CurrentPassthroughObjectLoader = LiveInstance->PassthroughObjectLoader;
 
@@ -660,7 +609,7 @@ namespace UE::Mutable::Private
 	}
 	
 
-    UE::Tasks::TTask<TManagedPtr<const FMesh>> FSystem::GetMesh(UE::Tasks::FTask Prerequisite, FInstance::FID instanceID, const FMeshId& MeshId, EMeshContentFlags MeshContentFilter)
+    UE::Tasks::TTask<TManagedPtr<const FMesh>> FSystem::GetMesh(UE::Tasks::FTask Prerequisite, const TSharedRef<FLiveInstance>& LiveInstance, const FMeshId& MeshId, EMeshContentFlags MeshContentFilter)
 	{
 		LLM_SCOPE_BYNAME(TEXT("MutableRuntime"));
 		MUTABLE_CPUPROFILER_SCOPE(SystemGetMesh);
@@ -679,10 +628,6 @@ namespace UE::Mutable::Private
 
 		OP::ADDRESS RootAddress = MeshId.GetKey()->Address;
 
-		// Find the live instance
-		FLiveInstance* LiveInstance = FindLiveInstance(instanceID);
-		check(LiveInstance);
-		
 		EOpType OpType = LiveInstance->Model->Program.GetOpType(RootAddress);
 		if (GetOpDataType(OpType) != EDataType::Mesh)
 		{
@@ -752,7 +697,7 @@ namespace UE::Mutable::Private
 	}
 
 
-	UE::Tasks::TTask<TManagedPtr<const FSkeletalMesh>> FSystem::GetSkeletalMesh(UE::Tasks::FTask Prerequisite, FInstance::FID InstanceID, const FSkeletalMeshId& SkeletalMeshId, EMeshContentFlags SkeletalMeshContentFilter)
+	UE::Tasks::TTask<TManagedPtr<const FSkeletalMesh>> FSystem::GetSkeletalMesh(UE::Tasks::FTask Prerequisite, const TSharedRef<FLiveInstance>& LiveInstance, const FSkeletalMeshId& SkeletalMeshId, EMeshContentFlags SkeletalMeshContentFilter)
     {
 		LLM_SCOPE_BYNAME(TEXT("MutableRuntime"));
 		MUTABLE_CPUPROFILER_SCOPE(SystemGetSkeletalMesh);
@@ -771,10 +716,6 @@ namespace UE::Mutable::Private
 
 		OP::ADDRESS RootAddress = SkeletalMeshId.GetKey()->Address;
 		
-		// Find the live instance
-		FLiveInstance* LiveInstance = FindLiveInstance(InstanceID);
-		check(LiveInstance);
-		
 		UE::Mutable::Private::EOpType OpType = LiveInstance->Model->Program.GetOpType(RootAddress);
 		if (GetOpDataType(OpType) != EDataType::SkeletalMesh)
 		{
@@ -843,7 +784,7 @@ namespace UE::Mutable::Private
 	}
 
 	
-	UE::Tasks::TTask<TManagedPtr<const FMaterial>> FSystem::GetMaterial(UE::Tasks::FTask Prerequisite, FInstance::FID InstanceID, const FMaterialId& MaterialId)
+	UE::Tasks::TTask<TManagedPtr<const FMaterial>> FSystem::GetMaterial(UE::Tasks::FTask Prerequisite, const TSharedRef<FLiveInstance>& LiveInstance, const FMaterialId& MaterialId)
 	{
 		LLM_SCOPE_BYNAME(TEXT("MutableRuntime"));
 		MUTABLE_CPUPROFILER_SCOPE(SystemGetMaterial);
@@ -862,10 +803,6 @@ namespace UE::Mutable::Private
 		
 		OP::ADDRESS RootAddress = MaterialId.GetKey()->Address;
 	
-		// Find the live instance
-		FLiveInstance* LiveInstance = FindLiveInstance(InstanceID);
-		check(LiveInstance);
-		
 		EOpType OpType = LiveInstance->Model->Program.GetOpType(RootAddress);
 		if (GetOpDataType(OpType) != EDataType::Material)
 		{
@@ -934,66 +871,54 @@ namespace UE::Mutable::Private
 	}
 	
 
-	void FSystem::EndUpdate(FInstance::FID instanceID)
+	void FSystem::EndUpdate(const TSharedRef<FLiveInstance>& LiveInstance, bool bClearRoms)
 	{
 		LLM_SCOPE_BYNAME(TEXT("MutableRuntime"));
 		MUTABLE_CPUPROFILER_SCOPE(EndUpdate);
 
-		FLiveInstance* LiveInstance = FindLiveInstance(instanceID);
-		if (LiveInstance)
-		{
-			LiveInstance->Instance = nullptr;
-
-			// Debug check to see if we managed the op-hit-counts correctly
-			LiveInstance->Cache->CheckHitCountsCleared();
+		LiveInstance->Instance = nullptr;
 
-			WorkingMemoryManager.CurrentInstanceCache = LiveInstance->Cache;
-			WorkingMemoryManager.CurrentPassthroughObjectLoader = LiveInstance->PassthroughObjectLoader;
+		// Debug check to see if we managed the op-hit-counts correctly
+		LiveInstance->Cache->CheckHitCountsCleared();
 
-			// We don't want to clear the cache layer 1 because it contains data that can be useful for a 
-			// future update (same states, just runtime parameters changed).
-			//WorkingMemoryManager.ClearCacheLayer1();
+		WorkingMemoryManager.CurrentInstanceCache = LiveInstance->Cache;
+		WorkingMemoryManager.CurrentPassthroughObjectLoader = LiveInstance->PassthroughObjectLoader;
 
-			// We need to clear the layer 0 cache, because it contains data that is only valid for the current 
-			// parameter values (unless it is data marked as state cache)
-			WorkingMemoryManager.ClearCacheLayer0();
+		// We don't want to clear the cache layer 1 because it contains data that can be useful for a 
+		// future update (same states, just runtime parameters changed).
+		//WorkingMemoryManager.ClearCacheLayer1();
 
-			WorkingMemoryManager.CurrentInstanceCache = nullptr;
-			WorkingMemoryManager.CurrentPassthroughObjectLoader = nullptr;
-		}
+		// We need to clear the layer 0 cache, because it contains data that is only valid for the current 
+		// parameter values (unless it is data marked as state cache)
+		WorkingMemoryManager.ClearCacheLayer0();
 
+		WorkingMemoryManager.CurrentInstanceCache = nullptr;
+		WorkingMemoryManager.CurrentPassthroughObjectLoader = nullptr;
+		
 		// Reduce the cache until it fits the limit.
 		WorkingMemoryManager.EnsureBudgetBelow(0);
 
-		UpdateStats();
-	}
-
-	
-    void FSystem::ReleaseInstance( FInstance::FID instanceID )
-    {
-		LLM_SCOPE_BYNAME(TEXT("MutableRuntime"));
-		MUTABLE_CPUPROFILER_SCOPE(ReleaseInstance);
-
-		for (int32 Index = 0; Index < WorkingMemoryManager.LiveInstances.Num(); ++Index)
+		if (bClearRoms)
 		{
-			FLiveInstance& Instance = WorkingMemoryManager.LiveInstances[Index];
-			if (Instance.InstanceID == instanceID)
+			for (FWorkingMemoryManager::FModelCacheEntry& ModelCache : WorkingMemoryManager.CachePerModel)
 			{
-				WorkingMemoryManager.LiveInstances.RemoveAtSwap(Index);
-				break;
-			}
-		}
-
- 		int32 Removed = WorkingMemoryManager.LiveInstances.RemoveAllSwap(
-			[instanceID](const FLiveInstance& Instance)
-			{
-				return (Instance.InstanceID == instanceID);
-			});
+				if (const TSharedPtr<const FModel> CacheModel = ModelCache.Model.Pin())
+				{
+					const FProgram& Program = CacheModel->Program;
 
-		TRACE_COUNTER_SET(MutableRuntime_LiveInstances, WorkingMemoryManager.LiveInstances.Num());
+					for (int32 RomIndex = 0; RomIndex < Program.Roms.Num(); ++RomIndex)
+					{
+						Program.UnloadRom(RomIndex);
+					}
+				}
+			}	
+		}
+		
+		UpdateStats();
 	}
-
 	
+
+    //---------------------------------------------------------------------------------------------
     class RelevantParameterVisitor : public UniqueDiscreteCoveredCodeVisitor<>
     {
     public:
@@ -1043,15 +968,13 @@ namespace UE::Mutable::Private
     };
 
 	
-    void FSystem::GetParameterRelevancy(FInstance::FID InstanceID,
+    void FSystem::GetParameterRelevancy(const TSharedRef<FLiveInstance>& LiveInstance,
     	const TSharedPtr<const FParameters>& FParameters,
         TArray<uint32>& OutRelevantParameters)
     {
 		LLM_SCOPE_BYNAME(TEXT("MutableRuntime"));
 
 		// Find the live instance
-		FLiveInstance* LiveInstance = FindLiveInstance(InstanceID);
-		check(LiveInstance);
 		WorkingMemoryManager.CurrentInstanceCache = LiveInstance->Cache;
 		WorkingMemoryManager.CurrentPassthroughObjectLoader = LiveInstance->PassthroughObjectLoader;
 		
@@ -1062,20 +985,7 @@ namespace UE::Mutable::Private
     }
 
 	
-	inline FLiveInstance* FSystem::FindLiveInstance(FInstance::FID id)
-	{
-		for (int32 i = 0; i < WorkingMemoryManager.LiveInstances.Num(); ++i)
-		{
-			if (WorkingMemoryManager.LiveInstances[i].InstanceID == id)
-			{
-				return &WorkingMemoryManager.LiveInstances[i];
-			}
-		}
-		return nullptr;
-	}
-
-
-	bool FSystem::CheckUpdatedParameters( const FLiveInstance* LiveInstance,
+	bool FSystem::CheckUpdatedParameters(const TSharedRef<FLiveInstance>& LiveInstance,
                                  const TSharedPtr<const FParameters>& Params,
                                  uint64& UpdatedParameters)
     {
@@ -1439,38 +1349,44 @@ namespace UE::Mutable::Private
 
 
 	int32 FWorkingMemoryManager::GetRomBytes() const
-	{
-		int32 Result = 0;
+    {
+    	int32 Result = 0;
 
-		TArray<const FModel*> Models;
-		for (const FLiveInstance& Instance : LiveInstances)
-		{
-			Models.AddUnique(Instance.Model.Get());
-		}
+    	TArray<const FModel*> Models;
+    	for (const TWeakPtr<FLiveInstance>& WeakInstance : LiveInstances)
+    	{
+    		TSharedPtr<FLiveInstance> Instance = WeakInstance.Pin();
+    		if (!Instance)
+    		{
+    			continue;
+    		}
 
-		// Data stored per-model, but related to instance construction
-		for (const FModel* Model : Models)
-		{
-			// Count streamable and currently-loaded resources
-			const FProgram& Program = Model->Program;
-			for (const TPair<uint32,TManagedPtr<const FImage>>& ImageLOD : Program.ConstantImageLODsStreamed)
-			{
-				if (ImageLOD.Value)
-				{
-					Result += ImageLOD.Value->GetDataSize();
-				}
-			}
-			for (const TPair<uint32, TManagedPtr<const FMesh>>& Rom : Program.ConstantMeshesStreamed)
-			{
-				if (Rom.Value)
-				{
-					Result += Rom.Value->GetDataSize();
-				}
-			}
-		}
+    		Models.AddUnique(Instance->Model.Get());
+    	}
 
-		return Result;
-	}
+    	// Data stored per-model, but related to instance construction
+    	for (const FModel* Model : Models)
+    	{
+    		// Count streamable and currently-loaded resources
+    		const FProgram& Program = Model->Program;
+    		for (const TPair<uint32,TManagedPtr<const FImage>>& ImageLOD : Program.ConstantImageLODsStreamed)
+    		{
+    			if (ImageLOD.Value)
+    			{
+    				Result += ImageLOD.Value->GetDataSize();
+    			}
+    		}
+    		for (const TPair<uint32, TManagedPtr<const FMesh>>& Rom : Program.ConstantMeshesStreamed)
+    		{
+    			if (Rom.Value)
+    			{
+    				Result += Rom.Value->GetDataSize();
+    			}
+    		}
+    	}
+
+    	return Result;
+    }
 	
 
 	void FWorkingMemoryManager::LogWorkingMemory(const CodeRunner* CurrentRunner) const
@@ -1759,9 +1675,15 @@ namespace UE::Mutable::Private
 		{
 			MUTABLE_CPUPROFILER_SCOPE(EnsureBudgetBelow_FreeCached);	
 			// From other live instances first
-			for (const FLiveInstance& Instance : LiveInstances)
+			for (const TWeakPtr<FLiveInstance>& WeakInstance : LiveInstances)
 			{
-				if (Instance.Cache == CurrentInstanceCache)
+				TSharedPtr<FLiveInstance> Instance = WeakInstance.Pin();
+				if (!Instance)
+				{
+					continue;
+				}
+				
+				if (Instance->Cache == CurrentInstanceCache)
 				{
 					// Ignore the current live instance.
 					continue;
@@ -1770,7 +1692,7 @@ namespace UE::Mutable::Private
 				MUTABLE_CPUPROFILER_SCOPE(EnsureBudgetBelow_FreeCached_GatherAndRemove_Other);
 				// Gather all data in the cache for this instance
 
-				TryRemoveFromCache(*Instance.Cache);
+				TryRemoveFromCache(*Instance->Cache);
 
 				if (bFinished)
 				{
