[Mutable] Allow components having the same name.

#Synced-CL 49557694

diff --git a/Engine/Plugins/Experimental/MutablePopulation/Source/CustomizableObjectPopulationEditor/Private/MuCOPE/CustomizableObjectPopulationClassEditor.cpp b/Engine/Plugins/Experimental/MutablePopulation/Source/CustomizableObjectPopulationEditor/Private/MuCOPE/CustomizableObjectPopulationClassEditor.cpp
--- a/Engine/Plugins/Experimental/MutablePopulation/Source/CustomizableObjectPopulationEditor/Private/MuCOPE/CustomizableObjectPopulationClassEditor.cpp
+++ b/Engine/Plugins/Experimental/MutablePopulation/Source/CustomizableObjectPopulationEditor/Private/MuCOPE/CustomizableObjectPopulationClassEditor.cpp
@@ -22,6 +22,7 @@
 #include "MuCO/CustomizableObject.h"
 #include "MuCO/CustomizableObjectInstance.h"
 #include "MuCO/CustomizableSkeletalComponent.h"
+#include "MuCO/CustomizableObjectInstancePrivate.h"
 #include "MuCOP/CustomizableObjectPopulation.h"
 #include "MuCOP/CustomizableObjectPopulationClass.h"
 #include "MuCOP/CustomizableObjectPopulationGenerator.h"
@@ -880,14 +881,25 @@ void FCustomizableObjectPopulationClassEditor::OpenInstance()
 void FCustomizableObjectPopulationClassEditor::OpenSkeletalMesh()
 {
 	int32 SelectedInstance = Viewport->GetSelectedInstance();
+
+	if (!ViewportInstances[SelectedInstance])
+	{
+		return;
+	}
 	
-	if (ViewportInstances[SelectedInstance])
+	TMap<FName, TObjectPtr<USkeletalMesh>>::TConstIterator It = ViewportInstances[SelectedInstance]->GetPrivate()->SkeletalMeshes.CreateConstIterator();
+	if (!It)
 	{
-		if (USkeletalMesh* SkeletalMesh = ViewportInstances[SelectedInstance]->GetSkeletalMesh())
-		{
-			GEditor->GetEditorSubsystem<UAssetEditorSubsystem>()->OpenEditorForAsset(SkeletalMesh);
-		}
+		return;
 	}
+	
+	USkeletalMesh* SkeletalMesh = It->Value;
+	if (!SkeletalMesh)
+	{
+		return;
+	}	
+
+	GEditor->GetEditorSubsystem<UAssetEditorSubsystem>()->OpenEditorForAsset(SkeletalMesh);
 }
 
 
diff --git a/Engine/Plugins/Experimental/MutablePopulation/Source/CustomizableObjectPopulationEditor/Private/MuCOPE/CustomizableObjectPopulationEditor.cpp b/Engine/Plugins/Experimental/MutablePopulation/Source/CustomizableObjectPopulationEditor/Private/MuCOPE/CustomizableObjectPopulationEditor.cpp
--- a/Engine/Plugins/Experimental/MutablePopulation/Source/CustomizableObjectPopulationEditor/Private/MuCOPE/CustomizableObjectPopulationEditor.cpp
+++ b/Engine/Plugins/Experimental/MutablePopulation/Source/CustomizableObjectPopulationEditor/Private/MuCOPE/CustomizableObjectPopulationEditor.cpp
@@ -34,6 +34,7 @@
 #include "Widgets/Layout/SUniformGridPanel.h"
 #include "WorkspaceMenuStructure.h"
 #include "WorkspaceMenuStructureModule.h"
+#include "MuCO/CustomizableObjectInstancePrivate.h"
 
 class IToolkitHost;
 class SWidget;
@@ -524,13 +525,24 @@ void FCustomizableObjectPopulationEditor::OpenSkeletalMesh()
 {
 	int32 SelectedInstance = Viewport->GetSelectedInstance();
 
-	if (ViewportInstances[SelectedInstance])
+	if (!ViewportInstances[SelectedInstance])
 	{
-		if (USkeletalMesh* SkeletalMesh = ViewportInstances[SelectedInstance]->GetSkeletalMesh())
-		{
-			GEditor->GetEditorSubsystem<UAssetEditorSubsystem>()->OpenEditorForAsset(SkeletalMesh);
-		}
+		return;
+	}
+	
+	TMap<FName, TObjectPtr<USkeletalMesh>>::TConstIterator It = ViewportInstances[SelectedInstance]->GetPrivate()->SkeletalMeshes.CreateConstIterator();
+	if (!It)
+	{
+		return;
 	}
+	
+	USkeletalMesh* SkeletalMesh = It->Value;
+	if (!SkeletalMesh)
+	{
+		return;
+	}	
+
+	GEditor->GetEditorSubsystem<UAssetEditorSubsystem>()->OpenEditorForAsset(SkeletalMesh);
 }
 
 // To move elsewhere
diff --git a/Engine/Plugins/Mutable/Source/CustomizableObject/Internal/MuCO/CustomizableObjectInstancePrivate.h b/Engine/Plugins/Mutable/Source/CustomizableObject/Internal/MuCO/CustomizableObjectInstancePrivate.h
--- a/Engine/Plugins/Mutable/Source/CustomizableObject/Internal/MuCO/CustomizableObjectInstancePrivate.h
+++ b/Engine/Plugins/Mutable/Source/CustomizableObject/Internal/MuCO/CustomizableObjectInstancePrivate.h
@@ -111,6 +111,12 @@ struct FCustomizableInstanceComponentData
 {
 	GENERATED_USTRUCT_BODY();
 
+	FCustomizableInstanceComponentData();
+
+	FName ComponentName;
+	
+	UE::Mutable::Private::FComponentId ComponentId = INDEX_NONE;
+	
 	// AnimBP data gathered for a component from its constituent meshes
 	UPROPERTY(Transient, Category = CustomizableObjectInstance, editfixedsize, VisibleAnywhere)
 	TMap<FName, TSoftClassPtr<UAnimInstance>> AnimSlotToBP;
@@ -260,10 +266,7 @@ public:
 	/** Invalidates the previously generated data and retrieves information from the CObject after specific actions.
 	 *  It'll be called in the PostLoad, after Compiling the CO, and after changing the CO of the Instance. */
 	UE_API void InitCustomizableObjectData(const UCustomizableObject* InCustomizableObject);
-
-	UE_API FCustomizableInstanceComponentData* GetComponentData(const FName& ComponentName);
-	void ClearInstanceComponentData(const FName& ComponentName);
-
+	
 	ECOInstanceFlags GetCOInstanceFlags() const { return InstanceFlagsPrivate; }
 	void SetCOInstanceFlags(ECOInstanceFlags FlagsToSet) { InstanceFlagsPrivate = (ECOInstanceFlags)(InstanceFlagsPrivate | FlagsToSet); }
 	void ClearCOInstanceFlags(ECOInstanceFlags FlagsToClear) { InstanceFlagsPrivate = (ECOInstanceFlags)(InstanceFlagsPrivate & ~FlagsToClear); }
@@ -277,9 +280,6 @@ public:
 	// Returns true if success (?)
 	UE_API bool UpdateSkeletalMesh_PostBeginUpdate0(UCustomizableObjectInstance* Instance, const TSharedRef<FUpdateContextPrivate>& Context);
 	
-	// Copy data generated in the mutable thread over to the instance and initializes additional data required during the update
-	UE_API void PrepareForUpdate();
-
 	// The following method is basically copied from PostEditChangeProperty and/or SkeletalMesh.cpp to be able to replicate PostEditChangeProperty without the editor
 	UE_API void PostEditChangePropertyWithoutEditor();
 	
@@ -287,19 +287,17 @@ public:
 	  * It should not be called if the meshes are still in use or shared with other instances. */
 	UE_API void DiscardResources();
 	
-	/** Set the reference SkeletalMesh, or an empty mesh, to all actors using this instance. */
-	UE_API void SetReferenceSkeletalMesh() const;
+	/** Only use as an emergency fallback! Set the reference SkeletalMesh unconditionally to all actors using this instance. It will hitch. */
+	UE_API void ForceSetReferenceSkeletalMesh() const;
 
 	UE_API const TArray<FAnimInstanceOverridePhysicsAsset>* GetGeneratedPhysicsAssetsForAnimInstance(TSubclassOf<UAnimInstance> AnimInstance) const;
 
-	/** */
 #if WITH_EDITORONLY_DATA
 	UE_API void RegenerateImportedModels(const TSharedRef<FUpdateContextPrivate>& OperationData);
 #endif
 	
-	// 
-	UE_API FCustomizableInstanceComponentData* GetComponentData(FCustomizableObjectComponentIndex ObjectComponentIndex);
-	UE_API const FCustomizableInstanceComponentData* GetComponentData(FCustomizableObjectComponentIndex ObjectComponentIndex) const;
+	UE_API FCustomizableInstanceComponentData* GetComponentData(const FName& ComponentName);
+	UE_API FCustomizableInstanceComponentData& GetComponentData(UE::Mutable::Private::FComponentId ComponentId);
 	
 	UE_API void InvalidateGeneratedData();
 
@@ -347,7 +345,6 @@ public:
 	float MinSquareDistFromComponentToPlayer = FLT_MAX;
 	float LastMinSquareDistFromComponentToPlayer = FLT_MAX; // The same as the previous dist for last frame
 	
-	// To be indexed with object component index
 	UPROPERTY(Transient)
 	TArray<FCustomizableInstanceComponentData> ComponentsData;
 	
@@ -420,23 +417,23 @@ public:
 };
 
 
-UE_API UPhysicsAsset* GetOrBuildMainPhysicsAsset(const TSharedRef<FUpdateContextPrivate>& Context, TObjectPtr<UPhysicsAsset> TamplateAsset, const UE::Mutable::Private::FPhysicsBody* PhysicsBody, bool bDisableCollisionBetweenAssets, FCustomizableObjectInstanceComponentIndex InstanceComponentIndex);
+UE_API UPhysicsAsset* GetOrBuildMainPhysicsAsset(const TSharedRef<FUpdateContextPrivate>& Context, TObjectPtr<UPhysicsAsset> TamplateAsset, const UE::Mutable::Private::FPhysicsBody* PhysicsBody, bool bDisableCollisionBetweenAssets, int32 ComponentIndex);
 
-UE_API USkeleton* MergeSkeletons(UCustomizableObject& CustomizableObject, FCustomizableInstanceComponentData& ObjectComponentIndex, bool& bOutCreatedNewSkeleton);
+UE_API USkeleton* MergeSkeletons(UCustomizableObject& CustomizableObject, FCustomizableInstanceComponentData& ComponentData, bool& bOutCreatedNewSkeleton);
 
-UE_API void InitSkeletalMeshData(const TSharedRef<FUpdateContextPrivate>& Context, USkeletalMesh* SkeletalMesh, const FMutableRefSkeletalMeshData& RefSkeletalMeshData, const UCustomizableObject& CustomizableObject, FCustomizableObjectComponentIndex ObjectComponentIndex);
+UE_API void InitSkeletalMeshData(const TSharedRef<FUpdateContextPrivate>& Context, USkeletalMesh* SkeletalMesh, const FMutableRefSkeletalMeshData& RefSkeletalMeshData, const UCustomizableObject& CustomizableObject, UE::Mutable::Private::FComponentId ComponentId);
 
-UE_API void BuildOrCopyClothingData(const TSharedRef<FUpdateContextPrivate>& Context, USkeletalMesh* SkeletalMesh, const UModelResources& CustomizableObjectInstance, FCustomizableObjectInstanceComponentIndex InstanceComponentIndex, const TArray<TObjectPtr<UPhysicsAsset>>& ClothingPhysicsAssets);
+UE_API void BuildOrCopyClothingData(const TSharedRef<FUpdateContextPrivate>& Context, USkeletalMesh* SkeletalMesh, const UModelResources& CustomizableObjectInstance, int32 ComponentIndex, const TArray<TObjectPtr<UPhysicsAsset>>& ClothingPhysicsAssets);
 
-UE_API bool BuildSkeletonData(const TSharedRef<FUpdateContextPrivate>& Context, USkeletalMesh& SkeletalMesh, FCustomizableObjectInstanceComponentIndex InstanceComponentIndex);
+UE_API bool BuildSkeletonData(const TSharedRef<FUpdateContextPrivate>& Context, USkeletalMesh& SkeletalMesh, int32 ComponentIndex);
 
 UE_API void BuildMeshSockets(USkeletalMesh* SkeletalMesh, const UModelResources& ModelResources, const FMutableRefSkeletalMeshData& RefSkeletalMeshData, UE::Mutable::Private::TManagedPtr<const UE::Mutable::Private::FMesh> MutableMesh);
 
-UE_API bool BuildOrCopyRenderData(const TSharedRef<FUpdateContextPrivate>& Context, USkeletalMesh* SkeletalMesh, UCustomizableObjectInstance* CustomizableObjectInstance, FCustomizableObjectInstanceComponentIndex InstanceComponentIndex);
+UE_API bool BuildOrCopyRenderData(const TSharedRef<FUpdateContextPrivate>& Context, USkeletalMesh* SkeletalMesh, UCustomizableObjectInstance* CustomizableObjectInstance, int32 ComponentIndex);
 
-UE_API void BuildOrCopyElementData(const TSharedRef<FUpdateContextPrivate>& Context, USkeletalMesh* SkeletalMesh, FCustomizableObjectInstanceComponentIndex InstanceComponentIndex);
+UE_API void BuildOrCopyElementData(const TSharedRef<FUpdateContextPrivate>& Context, USkeletalMesh* SkeletalMesh, int32 ComponentIndex);
 
-UE_API void BuildOrCopyMorphTargetsData(const TSharedRef<FUpdateContextPrivate>& Context, USkeletalMesh* SkeletalMesh, FCustomizableObjectInstanceComponentIndex InstanceComponentIndex);
+UE_API void BuildOrCopyMorphTargetsData(const TSharedRef<FUpdateContextPrivate>& Context, USkeletalMesh* SkeletalMesh, int32 ComponentIndex);
 	
 UE_API bool DoComponentsNeedUpdate(const TSharedRef<FUpdateContextPrivate>& Context, bool& bOutEmptyMesh);
 
diff --git a/Engine/Plugins/Mutable/Source/CustomizableObject/Internal/MuCO/CustomizableObjectPrivate.h b/Engine/Plugins/Mutable/Source/CustomizableObject/Internal/MuCO/CustomizableObjectPrivate.h
--- a/Engine/Plugins/Mutable/Source/CustomizableObject/Internal/MuCO/CustomizableObjectPrivate.h
+++ b/Engine/Plugins/Mutable/Source/CustomizableObject/Internal/MuCO/CustomizableObjectPrivate.h
@@ -1221,14 +1221,11 @@ class UModelResources : public UObject
 	GENERATED_BODY()
 
 public:
-	/** 
+	/** Key is FComponentId.
 	 * All the SkeletalMeshes generated for this CustomizableObject instances will use the Reference Skeletal Mesh
 	 * properties for everything that Mutable doesn't create or modify. This struct stores the information used from
 	 * the Reference Skeletal Meshes to avoid having them loaded at all times. This includes data like LOD distances,
-	 * LOD render data settings, Mesh sockets, Bounding volumes, etc.
-	 * 
-	 * Index with CustomizableObject Component index
-	 */
+	 * LOD render data settings, Mesh sockets, Bounding volumes, etc. */
 	UPROPERTY()
 	TArray<FMutableRefSkeletalMeshData> ReferenceSkeletalMeshesData;	
 
@@ -1361,29 +1358,29 @@ public:
 	UPROPERTY()
 	TMap<TObjectPtr<const UScriptStruct>, uint32> ExternalOperationVersions;
 	
-	/** Max number of LODs in the compiled Model. */
+	/** Key is FComponentId. Max number of LODs in the compiled Model. */
 	UPROPERTY()
-	TMap<FName, uint8> NumLODsAvailable;
+	TMap<int32, uint8> NumLODsAvailable;
 
-	/** Max number of LODs to stream. Mutable will always generate at least one LOD. */
+	/** Key is FComponentId. Max number of LODs to stream. Mutable will always generate at least one LOD. */
 	UPROPERTY()
-	TMap<FName, uint8> NumLODsToStream;
+	TMap<int32, uint8> NumLODsToStream;
 
-	/** First LOD available, some platforms may remove lower LODs when cooking, this MinLOD represents the first LOD we can generate */
+	/** Key is FComponentId. First LOD available, some platforms may remove lower LODs when cooking, this MinLOD represents the first LOD we can generate */
 	UPROPERTY()
-	TMap<FName, uint8> FirstLODAvailable;
+	TMap<int32, uint8> FirstLODAvailable;
 
-	/** Name of all possible components. Index is the ObjectComponentIndex. */
+	/** Key is FComponentId. Name of all possible components. */
 	UPROPERTY()
 	TArray<FName> ComponentNamesPerObjectComponent;
 
-	/** Minimum LOD to render per Platform. */
+	/** Key is FComponentId. Minimum LOD to render per Platform. */
 	UPROPERTY()
-	TMap<FName, FPerPlatformInt> MinLODPerComponent;
+	TMap<int32, FPerPlatformInt> MinLODPerComponent;
 
-	/** Minimum LOD to render per Quality level. */
+	/** Key is FComponentId. Minimum LOD to render per Quality level. */
 	UPROPERTY()
-	TMap<FName, FPerQualityLevelInt> MinQualityLevelLODPerComponent;
+	TMap<int32, FPerQualityLevelInt> MinQualityLevelLODPerComponent;
 
 	UPROPERTY()
 	TMap<FString, TObjectPtr<UTexture>> TextureParameterDefaultValues;
@@ -1471,43 +1468,6 @@ struct FMutableMeshComponentData
 };
 
 
-
-/** Strongly typed index for the index of a component in a UCustomizableObject. */
-USTRUCT()
-struct FCustomizableObjectComponentIndex
-{
-	GENERATED_BODY()
-
-private:
-
-	int32 Index = 0;
-
-public:
-
-	explicit inline FCustomizableObjectComponentIndex(int32 InIndex=0)
-		: Index(InIndex)
-	{
-	}
-
-	inline void Invalidate()
-	{
-		Index = INDEX_NONE;
-	}
-
-	inline bool IsValid() const
-	{
-		return Index != INDEX_NONE;
-	}
-
-	inline int32 GetValue() const
-	{
-		return Index;
-	}
-
-	inline bool operator==(const FCustomizableObjectComponentIndex&) const = default;
-};
-
-
 #if WITH_EDITORONLY_DATA 
 // This is a manual version number for the binary blobs in this asset.
 // Increasing it invalidates all the previously compiled models.
@@ -1767,6 +1727,8 @@ enum class ECustomizableObjectVersions : int32
 	DeduplicateFMaterialsSharingPassthroughObject,
 
 	RemovedReuseTextures,
+	
+	ComponentId,
 
 	PassthroughTexturesForBreakNodes,
 
@@ -1911,7 +1873,7 @@ public:
 	CUSTOMIZABLEOBJECT_API void SetIsChildObject(bool bIsChildObject);
 	
 	/** Return the MinLOD index to generate based on the active LODSettings (PerPlatformMinLOD or PerQualityLevelMinLOD) */
-	CUSTOMIZABLEOBJECT_API uint8 GetMinLODIndex(const FName& ComponentName) const;
+	CUSTOMIZABLEOBJECT_API uint8 GetMinLODIndex(const UE::Mutable::Private::FComponentId ComponentId) const;
 
 #if WITH_EDITOR
 	/** See ICustomizableObjectEditorModule::IsCompilationOutOfDate. */
@@ -1936,7 +1898,7 @@ public:
 
 	CUSTOMIZABLEOBJECT_API void BackwardsCompatibleFixup(int32 CustomizableObjectCustomVersion);
 
-	CUSTOMIZABLEOBJECT_API FName GetComponentName(FCustomizableObjectComponentIndex ObjectComponentIndex) const;
+	CUSTOMIZABLEOBJECT_API FName GetComponentName(UE::Mutable::Private::FComponentId ComponentId) const;
 
 #if WITH_EDITORONLY_DATA
 	CUSTOMIZABLEOBJECT_API EMutableCompileMeshType GetMeshCompileType() const;
diff --git a/Engine/Plugins/Mutable/Source/CustomizableObject/Internal/MuCO/CustomizableObjectSystemPrivate.h b/Engine/Plugins/Mutable/Source/CustomizableObject/Internal/MuCO/CustomizableObjectSystemPrivate.h
--- a/Engine/Plugins/Mutable/Source/CustomizableObject/Internal/MuCO/CustomizableObjectSystemPrivate.h
+++ b/Engine/Plugins/Mutable/Source/CustomizableObject/Internal/MuCO/CustomizableObjectSystemPrivate.h
@@ -118,37 +118,6 @@ void LoadClothing(FMutableStreamRequest& MutableDataStreamer, const UE::Mutable:
 void FinishUpdateGlobal(const TSharedRef<FUpdateContextPrivate>& Context);
 
 
-/** Strongly typed index for the index of a component in a UCustomizableObjectInstance. */
-class FCustomizableObjectInstanceComponentIndex
-{
-private:
-
-	int32 Index = 0;
-
-public:
-
-	explicit FCustomizableObjectInstanceComponentIndex(int32 InIndex = 0)
-		: Index(InIndex)
-	{
-	}
-
-	inline void Invalidate()
-	{
-		Index = INDEX_NONE;
-	}
-
-	inline bool IsValid() const
-	{
-		return Index != INDEX_NONE;
-	}
-
-	inline int32 GetValue() const
-	{
-		return Index;
-	}
-};
-
-
 class FMutableUpdateCandidate
 {
 public:
@@ -402,7 +371,7 @@ struct FInstanceUpdateData
 	
 	struct FComponent
 	{
-		FCustomizableObjectComponentIndex Id = FCustomizableObjectComponentIndex(0);
+		UE::Mutable::Private::FComponentId Id = INDEX_NONE;
 		
 		/** Range in the LODs array */
 		uint16 FirstLOD = 0;
@@ -441,7 +410,7 @@ struct FInstanceUpdateData
 	TMap<uint32, FMorphTargetMeshData> RealTimeMorphTargetMeshData;
 	
 	/** Key is the Component Name. Value is the final Morph Target data to copy into the Skeletal Mesh. */
-	TMap<FName, FSkeletalMeshMorphTargets> RealTimeMorphTargets;
+	TMap<UE::Mutable::Private::FComponentId, FSkeletalMeshMorphTargets> RealTimeMorphTargets;
 
 	/** Key is the Block Id. Value is the LoadAdditionalAssetsAndData read destination. */
 	TMap<uint32, FClothingMeshData> ClothingMeshData;
@@ -518,17 +487,12 @@ public:
 	const FDescriptorHash& GetCapturedDescriptorHash() const;
 
 	const FCustomizableObjectInstanceDescriptor&& MoveCommittedDescriptor();
-
-	/** Return the object component index associated with a component in this instance. */
-	FCustomizableObjectComponentIndex GetObjectComponentIndex(FCustomizableObjectInstanceComponentIndex InstanceComponentIndex);
 	
-	const FInstanceUpdateData::FComponent* GetComponentUpdateData(FCustomizableObjectInstanceComponentIndex InstanceComponentIndex);
-
 	void InitComponentDescriptors(int32 Size);
 
 	const TArray<FComponentDescriptor>& GetComponentDescriptors() const;
 	
-	FComponentDescriptor* GetComponentDescriptor(FCustomizableObjectComponentIndex Index);
+	FComponentDescriptor* GetComponentDescriptor(UE::Mutable::Private::FComponentId ComponentId);
 	
 	UE::Mutable::Private::FImageOperator::FImagePixelFormatFunc PixelFormatOverride;
 
@@ -558,7 +522,7 @@ public:
 	  * @return true if the mesh is new or new to this instance (e.g. mesh cached by another instance). */
 	TArray<bool> MeshChangedPerInstanceComponent;
 
-	/** List of component names. Index is the ObjectComponentIndex. */
+	/** List of component names. Key is the FComponentId. */
 	TArray<FName> ComponentNames;
 
 	/** Index of the resource in the StreamedResourceData array of the Model Resources. */
@@ -646,14 +610,6 @@ public:
 	/** This option comes from the operation request. It is used to reduce the number of mipmaps that mutable must generate for images.  */
 	int32 MipsToSkip = 0;
 	
-	// Number of possible components in the entire CO
-	// TODO: Redundant because it is in the CO
-	uint8 NumObjectComponents = 0;
-
-	// Number of components in the instance being generated
-	// TODO: Redundant to store it here because it is implicit in the size of some arrays or in the MutableInstance while it is valid.
-	uint8 NumInstanceComponents = 0;
-	
 	EQueuePriorityType PriorityType = EQueuePriorityType::Low;
 
 	EUpdateResult UpdateResult = EUpdateResult::Success;
diff --git a/Engine/Plugins/Mutable/Source/CustomizableObject/Internal/MuCO/UnrealToMutableTextureConversionUtils.h b/Engine/Plugins/Mutable/Source/CustomizableObject/Internal/MuCO/UnrealToMutableTextureConversionUtils.h
--- a/Engine/Plugins/Mutable/Source/CustomizableObject/Internal/MuCO/UnrealToMutableTextureConversionUtils.h
+++ b/Engine/Plugins/Mutable/Source/CustomizableObject/Internal/MuCO/UnrealToMutableTextureConversionUtils.h
@@ -4,6 +4,7 @@
 
 #include "HAL/Platform.h"
 #include "Engine/Texture.h"
+#include "MuR/ComponentId.h"
 #include "MuR/ManagedPointer.h"
 
 #define UE_API CUSTOMIZABLEOBJECT_API
@@ -92,7 +93,7 @@ struct FMutableSourceMeshData
 	TSoftClassPtr<UAnimInstance> AnimInstance;
 	TSoftObjectPtr<UStreamableRenderAsset> TableReferenceSkeletalMesh;
 
-	FName Component = NAME_None;
+	UE::Mutable::Private::FComponentId ComponentId = INDEX_NONE;
 	
 	/** Required for SurfaceMetadataID */
 
@@ -135,7 +136,7 @@ inline uint32 GetTypeHash(const FMutableSourceMeshData& Key)
 	GuidHash = HashCombineFast(GuidHash, GetTypeHash(Key.SocketPriority));
 	GuidHash = HashCombineFast(GuidHash, GetTypeHash(Key.BonePosePriority));
 	GuidHash = HashCombineFast(GuidHash, GetTypeHash(Key.bOnlyConnectedLOD));
-	GuidHash = HashCombineFast(GuidHash, GetTypeHash(Key.Component));
+	GuidHash = HashCombineFast(GuidHash, GetTypeHash(Key.ComponentId));
 	GuidHash = HashCombineFast(GuidHash, GetTypeHash(Key.Flags));
 	return GuidHash;
 }
diff --git a/Engine/Plugins/Mutable/Source/CustomizableObject/Private/MuCO/CustomizableObject.cpp b/Engine/Plugins/Mutable/Source/CustomizableObject/Private/MuCO/CustomizableObject.cpp
--- a/Engine/Plugins/Mutable/Source/CustomizableObject/Private/MuCO/CustomizableObject.cpp
+++ b/Engine/Plugins/Mutable/Source/CustomizableObject/Private/MuCO/CustomizableObject.cpp
@@ -1005,11 +1005,12 @@ USkeletalMesh* UCustomizableObject::GetComponentMeshReferenceSkeletalMesh(const
 	// currently the graph search will be always performed if the editor only data is accessible.
 	if (const UModelResources* ModelResources = Private->GetModelResources())
 	{
-		int32 ObjectComponentIndex = ModelResources->ComponentNamesPerObjectComponent.IndexOfByKey(ComponentName);
-		if (ModelResources->ReferenceSkeletalMeshesData.IsValidIndex(ObjectComponentIndex))
+		// TODO Due to the adding support of multiple components with the same name this is no longer correct. A single name may point to multiple FComponentId.
+		UE::Mutable::Private::FComponentId ComponentId = ModelResources->ComponentNamesPerObjectComponent.IndexOfByKey(ComponentName);
+		if (ModelResources->ReferenceSkeletalMeshesData.IsValidIndex(ComponentId))
 		{
 			// Can be nullptr if RefSkeletalMeshes are not loaded yet.
-			return ModelResources->ReferenceSkeletalMeshesData[ObjectComponentIndex].SkeletalMesh;
+			return ModelResources->ReferenceSkeletalMeshesData[ComponentId].SkeletalMesh;
 		}
 	}
 	
@@ -1146,14 +1147,14 @@ int32 UCustomizableObject::GetComponentCount() const
 }
 
 
-FName UCustomizableObjectPrivate::GetComponentName(FCustomizableObjectComponentIndex ObjectComponentIndex) const
+FName UCustomizableObjectPrivate::GetComponentName(UE::Mutable::Private::FComponentId ComponentId) const
 {
 	if (const UModelResources* LocalModelResources = GetModelResources())
 	{
 		const TArray<FName>& ComponentNames = LocalModelResources->ComponentNamesPerObjectComponent;
-		if (ComponentNames.IsValidIndex(ObjectComponentIndex.GetValue()))
+		if (ComponentNames.IsValidIndex(ComponentId))
 		{
-			return ComponentNames[ObjectComponentIndex.GetValue()];
+			return ComponentNames[ComponentId];
 		}
 	}
 
@@ -1330,7 +1331,7 @@ int32 UCustomizableObjectPrivate::GetStateParameterIndex(int32 StateIndex, int32
 
 FName UCustomizableObject::GetComponentName(int32 ComponentIndex) const
 {
-	return GetPrivate()->GetComponentName(FCustomizableObjectComponentIndex(ComponentIndex));
+	return GetPrivate()->GetComponentName(ComponentIndex);
 }
 
 
@@ -1962,7 +1963,7 @@ void UCustomizableObjectPrivate::ApplyStateForcedValuesToParameters(FCustomizabl
 }
 
 
-uint8 UCustomizableObjectPrivate::GetMinLODIndex(const FName& ComponentName) const
+uint8 UCustomizableObjectPrivate::GetMinLODIndex(const UE::Mutable::Private::FComponentId ComponentId) const
 {
 	uint8 MinLODIdx = 0;
 
@@ -1978,7 +1979,7 @@ uint8 UCustomizableObjectPrivate::GetMinLODIndex(const FName& ComponentName) con
 				// Value extracted from the Scalability.ini of the current platform (if found).
 				const int32 QualityLevel = UCustomizableObjectSystem::GetInstance()->GetSkeletalMeshMinLODQualityLevel();
 
-				const FPerQualityLevelInt* QualityLevelInt = LocalModelResources->MinQualityLevelLODPerComponent.Find(ComponentName);
+				const FPerQualityLevelInt* QualityLevelInt = LocalModelResources->MinQualityLevelLODPerComponent.Find(ComponentId);
 				if (ensure(QualityLevelInt))
 				{
 					MinLODIdx = QualityLevelInt->GetValue(QualityLevel);
@@ -1987,7 +1988,7 @@ uint8 UCustomizableObjectPrivate::GetMinLODIndex(const FName& ComponentName) con
 		}
 		else
 		{
-			const FPerPlatformInt* PlatformInt = LocalModelResources->MinLODPerComponent.Find(ComponentName);
+			const FPerPlatformInt* PlatformInt = LocalModelResources->MinLODPerComponent.Find(ComponentId);
 			if (ensure(PlatformInt))
 			{
 				MinLODIdx = PlatformInt->GetValue();
@@ -2000,7 +2001,7 @@ uint8 UCustomizableObjectPrivate::GetMinLODIndex(const FName& ComponentName) con
 
 	if (LocalModelResources)
 	{
-		if (const uint8* Result = LocalModelResources->FirstLODAvailable.Find(ComponentName))
+		if (const uint8* Result = LocalModelResources->FirstLODAvailable.Find(ComponentId))
 		{
 			FirstLODAvailable = *Result;
 		}
diff --git a/Engine/Plugins/Mutable/Source/CustomizableObject/Private/MuCO/CustomizableObjectInstance.cpp b/Engine/Plugins/Mutable/Source/CustomizableObject/Private/MuCO/CustomizableObjectInstance.cpp
--- a/Engine/Plugins/Mutable/Source/CustomizableObject/Private/MuCO/CustomizableObjectInstance.cpp
+++ b/Engine/Plugins/Mutable/Source/CustomizableObject/Private/MuCO/CustomizableObjectInstance.cpp
@@ -119,10 +119,7 @@ void UCustomizableInstancePrivate::InvalidateGeneratedData()
 	CommittedDescriptor = {};
 	CommittedDescriptorHash = {};
 
-	// Init Component Data
-	FCustomizableInstanceComponentData TemplateComponentData;
-	TemplateComponentData.LastComponentDescriptor.MeshDescriptors.Init({}, MAX_MESH_LOD_COUNT);
-	ComponentsData.Init(TemplateComponentData, ComponentsData.Num());
+	ComponentsData.Empty();
 }
 
 
@@ -135,77 +132,28 @@ void UCustomizableInstancePrivate::InitCustomizableObjectData(const UCustomizabl
 		return;
 	}
 	
-	// Init Component Data
-	FCustomizableInstanceComponentData TemplateComponentData;
-	TemplateComponentData.LastComponentDescriptor.MeshDescriptors.Init({}, MAX_MESH_LOD_COUNT);
-	ComponentsData.Init(TemplateComponentData, InCustomizableObject->GetComponentCount());
-
 	ExtensionInstanceData.Empty();
 }
 
 
 FCustomizableInstanceComponentData* UCustomizableInstancePrivate::GetComponentData(const FName& ComponentName)
 {
-	UCustomizableObject* Object = GetPublic()->GetCustomizableObject();
-	if (!Object || !Object->IsCompiled())
-	{
-		return nullptr;
-	}
-
-	const int32 ObjectComponentIndex = Object->GetPrivate()->GetModelResourcesChecked().ComponentNamesPerObjectComponent.IndexOfByKey(ComponentName);
-	if (ObjectComponentIndex == INDEX_NONE)
+	return ComponentsData.FindByPredicate([&](const FCustomizableInstanceComponentData& ComponentData)
 	{
-		return nullptr;
-	}
-
-	if (!ComponentsData.IsValidIndex(ObjectComponentIndex))
-	{
-		return nullptr;	
-	}
-
-	return &ComponentsData[ObjectComponentIndex];
-}
-
-
-FCustomizableInstanceComponentData* UCustomizableInstancePrivate::GetComponentData(FCustomizableObjectComponentIndex ObjectComponentIndex)
-{
-	return ComponentsData.IsValidIndex(ObjectComponentIndex.GetValue()) ? &ComponentsData[ObjectComponentIndex.GetValue()] : nullptr;
+		return ComponentData.ComponentName == ComponentName;
+	});
 }
 
 
-const FCustomizableInstanceComponentData* UCustomizableInstancePrivate::GetComponentData(FCustomizableObjectComponentIndex ObjectComponentIndex) const
+FCustomizableInstanceComponentData& UCustomizableInstancePrivate::GetComponentData(UE::Mutable::Private::FComponentId ComponentId)
 {
-	return ComponentsData.IsValidIndex(ObjectComponentIndex.GetValue()) ? &ComponentsData[ObjectComponentIndex.GetValue()] : nullptr;
-}
-
-
-void UCustomizableInstancePrivate::ClearInstanceComponentData(const FName& ComponentName)
-{
-	FCustomizableInstanceComponentData* ComponentData = GetComponentData(ComponentName);
-
-	if (ComponentData)
+	FCustomizableInstanceComponentData* Result = ComponentsData.FindByPredicate([&](const FCustomizableInstanceComponentData& ComponentData)
 	{
-		ComponentData->LastComponentDescriptor.MeshDescriptors.Reset(MAX_MESH_LOD_COUNT);
-		ComponentData->LastComponentDescriptor.MeshDescriptors.AddDefaulted(MAX_MESH_LOD_COUNT);
-
-		ComponentData->AnimSlotToBP.Empty();
-		ComponentData->AssetUserDataArray.Empty();
-		ComponentData->ClothingPhysicsAssetsToStream.Empty();
-
-#if WITH_EDITORONLY_DATA
-		ComponentData->MeshPartPaths.Empty();
-#endif
-		ComponentData->OverlayMaterial = nullptr;
-		ComponentData->OverrideMaterials.Empty();
-		ComponentData->PhysicsAssets.AdditionalPhysicsAssets.Empty();
-		ComponentData->PhysicsAssets.AdditionalPhysicsAssetsToLoad.Empty();
-		ComponentData->PhysicsAssets.PhysicsAssetsToMerge.Empty();
-		ComponentData->PhysicsAssets.PhysicsAssetToLoad.Empty();
-		ComponentData->Skeletons.Skeleton = nullptr;
-		ComponentData->Skeletons.SkeletonIds.Empty();
-		ComponentData->Skeletons.SkeletonsToMerge.Empty();
-		ComponentData->StreamedResourceIndex.Empty();
-	}
+		return ComponentData.ComponentId == ComponentId;
+	});
+	
+	check(Result);
+	return *Result;
 }
 
 
@@ -234,29 +182,6 @@ void UCustomizableInstancePrivate::SetDescriptor(const FCustomizableObjectInstan
 }
 
 
-void UCustomizableInstancePrivate::PrepareForUpdate()
-{
-	// Clear the ComponentData from previous updates
-	for (FCustomizableInstanceComponentData& ComponentData : ComponentsData)
-	{
-		ComponentData.AnimSlotToBP.Empty();
-		ComponentData.AssetUserDataArray.Empty();
-		ComponentData.Skeletons.Skeleton = nullptr;
-		ComponentData.Skeletons.SkeletonIds.Empty();
-		ComponentData.Skeletons.SkeletonsToMerge.Empty();
-		ComponentData.PhysicsAssets.PhysicsAssetToLoad.Empty();
-		ComponentData.PhysicsAssets.PhysicsAssetsToMerge.Empty();
-		ComponentData.ClothingPhysicsAssetsToStream.Empty();
-		ComponentData.StreamedResourceIndex.Empty();
-		ComponentData.OverlayMaterial = nullptr;
-
-#if WITH_EDITORONLY_DATA
-		ComponentData.MeshPartPaths.Empty();
-#endif
-	}
-}
-
-
 uint32 GetTypeHash(const FMaterialReuseCacheKey& Key)
 {
 	uint32 Hash = GetTypeHash(Key.MaterialTemplate);
@@ -1365,7 +1290,7 @@ namespace
 		const TSharedRef<FUpdateContextPrivate>& Context,
 		TObjectPtr<UPhysicsAsset> TemplateAsset,
 		const UE::Mutable::Private::FPhysicsBody* MutablePhysics,
-		FCustomizableObjectInstanceComponentIndex InstanceComponentIndex)
+		int32 ComponentIndex)
 	{
 		check(TemplateAsset);
 		TObjectPtr<UPhysicsAsset> Result = NewObject<UPhysicsAsset>();
@@ -1380,7 +1305,7 @@ namespace
 
 		Result->bNotForDedicatedServer = TemplateAsset->bNotForDedicatedServer;
 
-		const TMap<UE::Mutable::Private::FBoneName, TPair<FName, uint16>>& BoneInfoMap = Context->InstanceUpdateData.SkeletonsPerInstanceComponent[InstanceComponentIndex.GetValue()].BoneInfoMap;
+		const TMap<UE::Mutable::Private::FBoneName, TPair<FName, uint16>>& BoneInfoMap = Context->InstanceUpdateData.SkeletonsPerInstanceComponent[ComponentIndex].BoneInfoMap;
 		TMap<FName, int32> BonesInUse;
 
 		const int32 MutablePhysicsBodyCount = MutablePhysics->GetBodyCount();
@@ -1510,7 +1435,7 @@ UPhysicsAsset* GetOrBuildMainPhysicsAsset(
 	TObjectPtr<UPhysicsAsset> TemplateAsset,
 	const UE::Mutable::Private::FPhysicsBody* MutablePhysics,
 	bool bDisableCollisionsBetweenDifferentAssets,
-	FCustomizableObjectInstanceComponentIndex InstanceComponentIndex)
+	int32 ComponentIndex)
 {
 	MUTABLE_CPUPROFILER_SCOPE(MergePhysicsAssets);
 
@@ -1521,22 +1446,20 @@ UPhysicsAsset* GetOrBuildMainPhysicsAsset(
 	
 	UPhysicsAsset* Result = nullptr;
 
-	const FInstanceUpdateData::FComponent* Component = Context->GetComponentUpdateData(InstanceComponentIndex);
+	const TSharedPtr<FInstanceUpdateData::FComponent>& Component = Context->InstanceUpdateData.Components[ComponentIndex];
 	if (!Component)
 	{
 		return nullptr;
 	}
 
-	FCustomizableObjectComponentIndex ObjectComponentIndex = Component->Id;
-	FCustomizableInstanceComponentData* ComponentData = InstancePrivate->GetComponentData(ObjectComponentIndex);
-	check(ComponentData);
+	FCustomizableInstanceComponentData& ComponentData = InstancePrivate->GetComponentData(Component->Id);
 	
 	TArray<TObjectPtr<UPhysicsAsset>> ValidAssets;
 
-	const int32 NumPhysicsAssets = ComponentData->PhysicsAssets.PhysicsAssetsToMerge.Num();
+	const int32 NumPhysicsAssets = ComponentData.PhysicsAssets.PhysicsAssetsToMerge.Num();
 	for (int32 PhysicsAssetIndex = 0; PhysicsAssetIndex < NumPhysicsAssets; ++PhysicsAssetIndex)
 	{
-		const TObjectPtr<UPhysicsAsset>& PhysicsAsset = ComponentData->PhysicsAssets.PhysicsAssetsToMerge[PhysicsAssetIndex];
+		const TObjectPtr<UPhysicsAsset>& PhysicsAsset = ComponentData.PhysicsAssets.PhysicsAssetsToMerge[PhysicsAssetIndex];
 
 		if (PhysicsAsset)
 		{
@@ -1570,7 +1493,7 @@ UPhysicsAsset* GetOrBuildMainPhysicsAsset(
 
 	Result->bNotForDedicatedServer = TemplateAsset->bNotForDedicatedServer;
 
-	const TMap<UE::Mutable::Private::FBoneName, TPair<FName, uint16>>& BoneInfoMap = Context->InstanceUpdateData.SkeletonsPerInstanceComponent[InstanceComponentIndex.GetValue()].BoneInfoMap;
+	const TMap<UE::Mutable::Private::FBoneName, TPair<FName, uint16>>& BoneInfoMap = Context->InstanceUpdateData.SkeletonsPerInstanceComponent[ComponentIndex].BoneInfoMap;
 	TMap<FName, int32> BonesInUse;
 
 	const int32 MutablePhysicsBodyCount = MutablePhysics->GetBodyCount();
@@ -1799,7 +1722,7 @@ UPhysicsAsset* GetOrBuildMainPhysicsAsset(
 	Result->UpdateBodySetupIndexMap();
 	Result->UpdateBoundsBodiesArray();
 
-	ComponentData->PhysicsAssets.PhysicsAssetsToMerge.Empty();
+	ComponentData.PhysicsAssets.PhysicsAssetsToMerge.Empty();
 
 	return Result;
 }
@@ -1844,17 +1767,17 @@ bool DoComponentsNeedUpdate(const TSharedRef<FUpdateContextPrivate>& Context, bo
 	bHasInvalidMesh = false;
 	bool bNeedsUpdate = false;
 	
-	const int32 NumInstanceComponents = Context->InstanceUpdateData.Components.Num();
-	bNeedsUpdate |= NumInstanceComponents != InstancePrivate->SkeletalMeshes.Num();
+	const int32 NumComponents = Context->InstanceUpdateData.Components.Num();
+	bNeedsUpdate |= NumComponents != InstancePrivate->SkeletalMeshes.Num();
 
 	// Find which components need an update
-	Context->MeshChangedPerInstanceComponent.Init(false, NumInstanceComponents);
+	Context->MeshChangedPerInstanceComponent.Init(false, NumComponents);
 
-	for (int32 InstanceComponentIndex = 0; InstanceComponentIndex < NumInstanceComponents; ++InstanceComponentIndex)
+	for (int32 ComponentIndex = 0; ComponentIndex < NumComponents; ++ComponentIndex)
 	{
-		const TSharedRef<FInstanceUpdateData::FComponent>& Component = Context->InstanceUpdateData.Components[InstanceComponentIndex];
-		const FCustomizableObjectComponentIndex ObjectComponentIndex = Component->Id;
-		const FName ComponentName = Context->ComponentNames[ObjectComponentIndex.GetValue()];
+		const TSharedRef<FInstanceUpdateData::FComponent>& Component = Context->InstanceUpdateData.Components[ComponentIndex];
+		const FName ComponentName = Context->ComponentNames[Component->Id];
+		const FCustomizableInstanceComponentData& ComponentData = InstancePrivate->GetComponentData(Component->Id);
 		
 		bool bComponentWithMesh = false;
 
@@ -1866,10 +1789,7 @@ bool DoComponentsNeedUpdate(const TSharedRef<FUpdateContextPrivate>& Context, bo
 		
 		if (Component->SkeletalMesh->PassthroughObject)
 		{
-			if (const FCustomizableInstanceComponentData* ComponentData = InstancePrivate->GetComponentData(ObjectComponentIndex)) // Could be nullptr if the component has not been generated.
-			{
-				bMeshChanged = Component->SkeletalMeshId != ComponentData->LastComponentDescriptor.SkeletalMeshDescriptor;
-			}
+			bMeshChanged = Component->SkeletalMeshId != ComponentData.LastComponentDescriptor.SkeletalMeshDescriptor;
 		}
 		else
 		{
@@ -1893,7 +1813,7 @@ bool DoComponentsNeedUpdate(const TSharedRef<FUpdateContextPrivate>& Context, bo
 				{
 					UE_LOG(LogMutable, Error, TEXT("Failed to generate SkeletalMesh for CO Instance [%s]. CO [%s] has invalid geometry for LOD [%d] Component [%d]."),
 						*Instance->GetName(), *CustomizableObject->GetName(),
-						LODIndex, InstanceComponentIndex);
+						LODIndex, ComponentIndex);
 					bHasInvalidMesh = true; // Unreal does not support empty LODs.
 				}
 
@@ -1902,16 +1822,13 @@ bool DoComponentsNeedUpdate(const TSharedRef<FUpdateContextPrivate>& Context, bo
 			}
 			
 			// Compare Mesh IDs 
-			const FComponentDescriptor* ComponentDescriptor = Context->GetComponentDescriptor(ObjectComponentIndex);
+			const FComponentDescriptor* ComponentDescriptor = Context->GetComponentDescriptor(Component->Id);
 			if (ensure(ComponentDescriptor))
 			{
-				if (const FCustomizableInstanceComponentData* ComponentData = InstancePrivate->GetComponentData(ObjectComponentIndex)) // Could be nullptr if the component has not been generated.
-				{
-					bMeshChanged = ComponentDescriptor->MeshDescriptors != ComponentData->LastComponentDescriptor.MeshDescriptors;
-				}
+				bMeshChanged = ComponentDescriptor->MeshDescriptors != ComponentData.LastComponentDescriptor.MeshDescriptors;
 			}
 
-			Context->MeshChangedPerInstanceComponent[InstanceComponentIndex] = bMeshChanged;
+			Context->MeshChangedPerInstanceComponent[ComponentIndex] = bMeshChanged;
 		}
 		
 		bNeedsUpdate |= bMeshChanged;
@@ -2066,35 +1983,6 @@ bool UCustomizableInstancePrivate::UpdateSkeletalMesh_PostBeginUpdate0(UCustomiz
 		}
 	}
 
-	// Detect the components that have disappeared since the last update and blank their properties so they aren't reused the
-	// next time they reappear
-	{
-		for (const auto& MeshIt : SkeletalMeshes)
-		{
-			FName DeletedComponentCandidateName = MeshIt.Key;
-			bool bIsDeleted = true;
-
-			for (const auto& InstanceUpdateComponent : Context->InstanceUpdateData.Components)
-			{
-				FName InstanceUpdateComponentName = ModelResources.ComponentNamesPerObjectComponent[InstanceUpdateComponent->Id.GetValue()];
-
-				if (DeletedComponentCandidateName == InstanceUpdateComponentName)
-				{
-					bIsDeleted = false;
-
-					break;
-				}
-			}
-
-			if (bIsDeleted)
-			{
-				ClearInstanceComponentData(DeletedComponentCandidateName);
-
-				bUpdateMeshes = true;
-			}
-		}
-	}
-
 	// None of the current meshes requires a mesh update. Continue to BuildResources
 	if (!bUpdateMeshes)
 	{
@@ -2103,34 +1991,18 @@ bool UCustomizableInstancePrivate::UpdateSkeletalMesh_PostBeginUpdate0(UCustomiz
 
 	SkeletalMeshes.Reset();
 	
-	const int32 NumInstanceComponents = Context->InstanceUpdateData.Components.Num();
-	for (int32 InstanceComponentIndexValue = 0; InstanceComponentIndexValue < NumInstanceComponents; ++InstanceComponentIndexValue)
+	const int32 NumComponents = Context->InstanceUpdateData.Components.Num();
+	for (int32 ComponentIndex = 0; ComponentIndex < NumComponents; ++ComponentIndex)
 	{
-		FCustomizableObjectInstanceComponentIndex InstanceComponentIndex(InstanceComponentIndexValue);
-		const FInstanceUpdateData::FComponent* Component = Context->GetComponentUpdateData( InstanceComponentIndex );
-		if (!Component)
-		{
-			continue;
-		}
-		const FCustomizableObjectComponentIndex ObjectComponentIndex = Component->Id;
-		if (!ModelResources.ComponentNamesPerObjectComponent.IsValidIndex(ObjectComponentIndex.GetValue()))
-		{
-			continue;
-		}
-		const FName ComponentName = ModelResources.ComponentNamesPerObjectComponent[ObjectComponentIndex.GetValue()];
-
-		FCustomizableInstanceComponentData* ComponentData = GetComponentData(ComponentName);
-		if (!ComponentData)
-		{
-			ensure(false);
-			Context->UpdateResult = EUpdateResult::Error;
-			return false;
-		}
+		const TSharedPtr<FInstanceUpdateData::FComponent>& Component = Context->InstanceUpdateData.Components[ComponentIndex];
+		
+		FCustomizableInstanceComponentData& ComponentData = GetComponentData(Component->Id);
+		const FName ComponentName = ComponentData.ComponentName;
 
-		FComponentDescriptor* ComponentDescriptor = Context->GetComponentDescriptor(ObjectComponentIndex);
+		FComponentDescriptor* ComponentDescriptor = Context->GetComponentDescriptor(Component->Id);
 		check(ComponentDescriptor);
 		check(ComponentDescriptor->MeshDescriptors.Num() == MAX_MESH_LOD_COUNT);
-		ComponentData->LastComponentDescriptor = *ComponentDescriptor;
+		ComponentData.LastComponentDescriptor = *ComponentDescriptor;
 
 		if (Component->SkeletalMesh->PassthroughObject)
 		{
@@ -2143,14 +2015,13 @@ bool UCustomizableInstancePrivate::UpdateSkeletalMesh_PostBeginUpdate0(UCustomiz
 		}
 		
 		// If the component doesn't need an update copy the previously generated mesh.
-		if (!Context->MeshChangedPerInstanceComponent[InstanceComponentIndex.GetValue()])
+		if (!Context->MeshChangedPerInstanceComponent[ComponentIndex])
 		{
 			if (TObjectPtr<USkeletalMesh>* Result = OldSkeletalMeshes.Find(ComponentName))
 			{
 				SkeletalMeshes.Add(ComponentName, *Result);
+				continue;
 			}
-
-			continue;
 		}
 
 		const int32 FirstLOD = Context->bStreamMeshLODs ?
@@ -2183,7 +2054,7 @@ bool UCustomizableInstancePrivate::UpdateSkeletalMesh_PostBeginUpdate0(UCustomiz
 			continue;
 		}
 
-		if (!ModelResources.ReferenceSkeletalMeshesData.IsValidIndex(ObjectComponentIndex.GetValue()))
+		if (!ModelResources.ReferenceSkeletalMeshesData.IsValidIndex(Component->Id))
 		{
 			Context->UpdateResult = EUpdateResult::Error;
 			return false;
@@ -2237,13 +2108,13 @@ bool UCustomizableInstancePrivate::UpdateSkeletalMesh_PostBeginUpdate0(UCustomiz
 			SkeletalMeshes.Add(ComponentName, SkeletalMesh);
 		}
 		
-		const FMutableRefSkeletalMeshData& RefSkeletalMeshData = ModelResources.ReferenceSkeletalMeshesData[ObjectComponentIndex.GetValue()];
+		const FMutableRefSkeletalMeshData& RefSkeletalMeshData = ModelResources.ReferenceSkeletalMeshesData[Component->Id];
 
 		// Set up the default information any mesh from this component will have (LODArrayInfos, RenderData, Mesh settings, etc). 
-		InitSkeletalMeshData(Context, SkeletalMesh, RefSkeletalMeshData, *CustomizableObject, ObjectComponentIndex);
+		InitSkeletalMeshData(Context, SkeletalMesh, RefSkeletalMeshData, *CustomizableObject, Component->Id);
 		
 		// Construct a new skeleton, fix up ActiveBones and Bonemap arrays and recompute the RefInvMatrices
-		const bool bBuildSkeletonDataSuccess = BuildSkeletonData(Context, *SkeletalMesh, FCustomizableObjectInstanceComponentIndex(InstanceComponentIndex) );
+		const bool bBuildSkeletonDataSuccess = BuildSkeletonData(Context, *SkeletalMesh, ComponentIndex);
 		if (!bBuildSkeletonDataSuccess)
 		{
 			Context->UpdateResult = EUpdateResult::Error;
@@ -2255,7 +2126,7 @@ bool UCustomizableInstancePrivate::UpdateSkeletalMesh_PostBeginUpdate0(UCustomiz
 		{
 			constexpr bool bDisallowCollisionBetweenAssets = true;
 			UPhysicsAsset* PhysicsAssetResult = GetOrBuildMainPhysicsAsset(
-				Context, RefSkeletalMeshData.PhysicsAsset, MutablePhysics.Get(), bDisallowCollisionBetweenAssets, FCustomizableObjectInstanceComponentIndex(InstanceComponentIndex) );
+				Context, RefSkeletalMeshData.PhysicsAsset, MutablePhysics.Get(), bDisallowCollisionBetweenAssets, ComponentIndex);
 
 			SkeletalMesh->SetPhysicsAsset(PhysicsAssetResult);
 
@@ -2302,7 +2173,7 @@ bool UCustomizableInstancePrivate::UpdateSkeletalMesh_PostBeginUpdate0(UCustomiz
 
 			Entry.PropertyIndex = Info.PropertyIndex;
 			Entry.PhysicsAsset = MakePhysicsAssetFromTemplateAndMutableBody(
-				Context,PhysicsAssetTemplate, AdditionalPhysicsBody.Get(), InstanceComponentIndex);
+				Context,PhysicsAssetTemplate, AdditionalPhysicsBody.Get(), ComponentIndex);
 		}
 
 		// Add sockets from the SkeletalMesh of reference and from the MutableMesh
@@ -2313,16 +2184,16 @@ bool UCustomizableInstancePrivate::UpdateSkeletalMesh_PostBeginUpdate0(UCustomiz
 			Pair.Key->OnSkeletalMeshCreated(Pair.Value, ComponentName, SkeletalMesh);
 		}
 		
-		BuildOrCopyElementData(Context, SkeletalMesh, InstanceComponentIndex);
-		bool const bCopyRenderDataSuccess = BuildOrCopyRenderData(Context, SkeletalMesh, Instance, InstanceComponentIndex);
+		BuildOrCopyElementData(Context, SkeletalMesh, ComponentIndex);
+		bool const bCopyRenderDataSuccess = BuildOrCopyRenderData(Context, SkeletalMesh, Instance, ComponentIndex);
 		if (!bCopyRenderDataSuccess)
 		{
 			Context->UpdateResult = EUpdateResult::Error;
 			return false;
 		}
 
-		BuildOrCopyMorphTargetsData(Context, SkeletalMesh, InstanceComponentIndex);
-		BuildOrCopyClothingData(Context, SkeletalMesh, ModelResources, InstanceComponentIndex, ClothingPhysicsAssets);
+		BuildOrCopyMorphTargetsData(Context, SkeletalMesh, ComponentIndex);
+		BuildOrCopyClothingData(Context, SkeletalMesh, ModelResources, ComponentIndex, ClothingPhysicsAssets);
 
 		FSkeletalMeshRenderData* RenderData = SkeletalMesh->GetResourceForRendering();
 		ensure(RenderData && RenderData->LODRenderData.Num() > 0);
@@ -3084,7 +2955,7 @@ void UCustomizableInstancePrivate::DiscardResources()
 }
 
 
-void UCustomizableInstancePrivate::SetReferenceSkeletalMesh() const
+void UCustomizableInstancePrivate::ForceSetReferenceSkeletalMesh() const
 {
 	UCustomizableObjectInstance* Instance = Cast<UCustomizableObjectInstance>(GetOuter());
 	if (!Instance)
@@ -3118,22 +2989,27 @@ void UCustomizableInstancePrivate::SetReferenceSkeletalMesh() const
 			continue;
 		}
 #endif
-
-		const FName& ComponentName = CustomizableObjectInstanceUsage->GetComponentName();
-		const int32 ObjectComponentIndex = ModelResources->ComponentNamesPerObjectComponent.IndexOfByKey(ComponentName);
-		if (!ModelResources->ReferenceSkeletalMeshesData.IsValidIndex(ObjectComponentIndex))
+		
+		USkeletalMeshComponent* Parent = CustomizableObjectInstanceUsage->GetAttachParent();
+		if (!Parent)
 		{
 			continue;
 		}
 		
-		if (USkeletalMeshComponent* Parent = CustomizableObjectInstanceUsage->GetAttachParent())
+		const FName& ComponentName = CustomizableObjectInstanceUsage->GetComponentName();
+		const UE::Mutable::Private::FComponentId ComponentId = ModelResources->ComponentNamesPerObjectComponent.IndexOfByKey(ComponentName);
+		if (!ModelResources->ReferenceSkeletalMeshesData.IsValidIndex(ComponentId))
 		{
-			Parent->EmptyOverrideMaterials();
-
-			TSoftObjectPtr<USkeletalMesh> SoftObjectPtr = ModelResources->ReferenceSkeletalMeshesData[ObjectComponentIndex].SoftSkeletalMesh;
-			USkeletalMesh* SkeletalMesh = UE::Mutable::Private::LoadObject(SoftObjectPtr);
-			Parent->SetSkeletalMesh(SkeletalMesh);
+			continue;
 		}
+		
+		Parent->EmptyOverrideMaterials();
+
+		// TODO Due to the adding support of multiple components with the same name this is no longer correct. A single name may point to multiple FComponentId.
+		TSoftObjectPtr<USkeletalMesh> SoftObjectPtr = ModelResources->ReferenceSkeletalMeshesData[ComponentId].SoftSkeletalMesh;
+		USkeletalMesh* SkeletalMesh = UE::Mutable::Private::LoadObject(SoftObjectPtr);
+		
+		Parent->SetSkeletalMesh(SkeletalMesh);
 	}
 }
 
@@ -3576,6 +3452,13 @@ void ConvertImage(UTexture2D* Texture, UE::Mutable::Private::TManagedPtr<const U
 	Texture->SetPlatformData(MutableCreateImagePlatformData(MutableImage, OnlyLOD,0,0).Release());
 }
 
+
+FCustomizableInstanceComponentData::FCustomizableInstanceComponentData()
+{
+	LastComponentDescriptor.MeshDescriptors.Init({}, MAX_MESH_LOD_COUNT);
+}
+
+
 static int32 EnableRayTracingFix = 0;
 FAutoConsoleVariableRef CVarMutableEnableRayTracingFix(
 	TEXT("mutable.EnableRayTracingFix"),
@@ -3585,7 +3468,7 @@ FAutoConsoleVariableRef CVarMutableEnableRayTracingFix(
 	TEXT("If 2, Enable fix for all generated meshes. Meshes will have ray tracing disabled.")
 	);
 
-void InitSkeletalMeshData(const TSharedRef<FUpdateContextPrivate>& Context, USkeletalMesh* SkeletalMesh, const FMutableRefSkeletalMeshData& RefSkeletalMeshData, const UCustomizableObject& CustomizableObject, FCustomizableObjectComponentIndex ObjectComponentIndex)
+void InitSkeletalMeshData(const TSharedRef<FUpdateContextPrivate>& Context, USkeletalMesh* SkeletalMesh, const FMutableRefSkeletalMeshData& RefSkeletalMeshData, const UCustomizableObject& CustomizableObject, UE::Mutable::Private::FComponentId ComponentId)
 {
 	MUTABLE_CPUPROFILER_SCOPE(UCustomizableInstancePrivate::InitSkeletalMesh);
 
@@ -3594,7 +3477,7 @@ void InitSkeletalMeshData(const TSharedRef<FUpdateContextPrivate>& Context, USke
 	UCustomizableObjectInstance* Instance = Context->Instance.Get();
 	UCustomizableInstancePrivate* InstancePrivate = Instance->GetPrivate();
 	
-	const FName ComponentName = Context->ComponentNames[ObjectComponentIndex.GetValue()];
+	const FName ComponentName = Context->ComponentNames[ComponentId];
 
 	SkeletalMesh->NeverStream = !Context->bStreamMeshLODs;
 
@@ -3616,21 +3499,20 @@ void InitSkeletalMeshData(const TSharedRef<FUpdateContextPrivate>& Context, USke
 
 	// Asset User Data
 	{
-		const FCustomizableInstanceComponentData* ComponentData = InstancePrivate->GetComponentData(ObjectComponentIndex);
-		check(ComponentData);
-		for (TObjectPtr<UAssetUserData> AssetUserData : ComponentData->AssetUserDataArray)
+		const FCustomizableInstanceComponentData& ComponentData = InstancePrivate->GetComponentData(ComponentId);
+		for (TObjectPtr<UAssetUserData> AssetUserData : ComponentData.AssetUserDataArray)
 		{
 			SkeletalMesh->AddAssetUserData(AssetUserData);
 		}
 
 		//Custom Asset User Data
 		if (Context->Instance->GetAnimationGameplayTags().Num() ||
-			ComponentData->AnimSlotToBP.Num())
+			ComponentData.AnimSlotToBP.Num())
 		{
 			UCustomizableObjectInstanceUserData* InstanceData = NewObject<UCustomizableObjectInstanceUserData>(SkeletalMesh, NAME_None, RF_Public | RF_Transactional);
 			InstanceData->AnimationGameplayTag = Context->Instance->GetAnimationGameplayTags();
 			
-			for (const TTuple<FName, TSoftClassPtr<UAnimInstance>>& AnimSlot : ComponentData->AnimSlotToBP)
+			for (const TTuple<FName, TSoftClassPtr<UAnimInstance>>& AnimSlot : ComponentData.AnimSlotToBP)
 			{
 				FCustomizableObjectAnimationSlot AnimationSlot;
 				AnimationSlot.Name = AnimSlot.Key;
@@ -3726,8 +3608,8 @@ void InitSkeletalMeshData(const TSharedRef<FUpdateContextPrivate>& Context, USke
 
 	// Set Min LOD (Override the Reference Skeletal Mesh LOD Settings)
 	const UModelResources& ModelResources = CustomizableObject.GetPrivate()->GetModelResourcesChecked();
-	SkeletalMesh->SetMinLod(FMath::Max(ModelResources.MinLODPerComponent.FindChecked(ComponentName).GetDefault(), static_cast<int32>(Context->FirstLODAvailable[ComponentName])));
-	SkeletalMesh->SetQualityLevelMinLod(ModelResources.MinQualityLevelLODPerComponent.FindChecked(ComponentName));
+	SkeletalMesh->SetMinLod(FMath::Max(ModelResources.MinLODPerComponent.FindChecked(ComponentId).GetDefault(), static_cast<int32>(Context->FirstLODAvailable[ComponentName])));
+	SkeletalMesh->SetQualityLevelMinLod(ModelResources.MinQualityLevelLODPerComponent.FindChecked(ComponentId));
 
 	// Set up unreal's default material, will be replaced when building materials
 	{
@@ -3743,7 +3625,7 @@ void InitSkeletalMeshData(const TSharedRef<FUpdateContextPrivate>& Context, USke
 }
 
 
-bool BuildSkeletonData(const TSharedRef<FUpdateContextPrivate>& Context, USkeletalMesh& SkeletalMesh, FCustomizableObjectInstanceComponentIndex InstanceComponentIndex)
+bool BuildSkeletonData(const TSharedRef<FUpdateContextPrivate>& Context, USkeletalMesh& SkeletalMesh, int32 ComponentIndex)
 {
 	MUTABLE_CPUPROFILER_SCOPE(UCustomizableInstancePrivate::BuildSkeletonData);
 
@@ -3751,14 +3633,13 @@ bool BuildSkeletonData(const TSharedRef<FUpdateContextPrivate>& Context, USkelet
 	UCustomizableInstancePrivate* InstancePrivate = Instance->GetPrivate();
 
 	UCustomizableObject* CustomizableObject = Instance->GetCustomizableObject();
-	
-	FCustomizableObjectComponentIndex ObjectComponentIndex = Context->GetObjectComponentIndex(InstanceComponentIndex);
 
-	FCustomizableInstanceComponentData* ComponentData = InstancePrivate->GetComponentData(ObjectComponentIndex);
-	check(ComponentData);
+	UE::Mutable::Private::FComponentId ComponentId = Context->InstanceUpdateData.Components[ComponentIndex]->Id;
+
+	FCustomizableInstanceComponentData& ComponentData = InstancePrivate->GetComponentData(ComponentId);
 	
 	bool bCreatedNewSkeleton = false;
-	const TObjectPtr<USkeleton> Skeleton = MergeSkeletons(*CustomizableObject, *ComponentData, bCreatedNewSkeleton);
+	const TObjectPtr<USkeleton> Skeleton = MergeSkeletons(*CustomizableObject, ComponentData, bCreatedNewSkeleton);
 	if (!Skeleton)
 	{
 		return false;
@@ -3772,8 +3653,8 @@ bool BuildSkeletonData(const TSharedRef<FUpdateContextPrivate>& Context, USkelet
 	const TArray<FMeshBoneInfo>& RawRefBoneInfo = ReferenceSkeleton.GetRawRefBoneInfo();
 	const int32 RawRefBoneCount = ReferenceSkeleton.GetRawBoneNum();
 
-	const TArray<FInstanceUpdateData::FBone>& BonePose = Context->InstanceUpdateData.SkeletonsPerInstanceComponent[InstanceComponentIndex.GetValue()].BonePose;
-	TMap<UE::Mutable::Private::FBoneName, TPair<FName,uint16>>& BoneInfoMap = Context->InstanceUpdateData.SkeletonsPerInstanceComponent[InstanceComponentIndex.GetValue()].BoneInfoMap;
+	const TArray<FInstanceUpdateData::FBone>& BonePose = Context->InstanceUpdateData.SkeletonsPerInstanceComponent[ComponentIndex].BonePose;
+	TMap<UE::Mutable::Private::FBoneName, TPair<FName,uint16>>& BoneInfoMap = Context->InstanceUpdateData.SkeletonsPerInstanceComponent[ComponentIndex].BoneInfoMap;
 
 	{
 		MUTABLE_CPUPROFILER_SCOPE(BuildSkeletonData_BuildBoneInfoMap);
@@ -3854,7 +3735,7 @@ bool BuildSkeletonData(const TSharedRef<FUpdateContextPrivate>& Context, USkelet
 		// MergeSkeletons function doesn't do it. Only do it for newly generated ones, not for cached or non-transient ones.
 		GeneratedSkeleton->RecreateBoneTree(&SkeletalMesh);
 		
-		TArray<TObjectPtr<USkeleton>>& SkeletonsToMerge = ComponentData->Skeletons.SkeletonsToMerge;
+		TArray<TObjectPtr<USkeleton>>& SkeletonsToMerge = ComponentData.Skeletons.SkeletonsToMerge;
 		check(SkeletonsToMerge.Num() > 1);
 
 		TMap<FName, EBoneTranslationRetargetingMode::Type> BoneNamesToRetargetingMode;
@@ -3992,17 +3873,17 @@ void BuildMeshSockets(USkeletalMesh* SkeletalMesh, const UModelResources& ModelR
 }
 
 
-void BuildOrCopyElementData(const TSharedRef<FUpdateContextPrivate>& Context, USkeletalMesh* SkeletalMesh, FCustomizableObjectInstanceComponentIndex InstanceComponentIndex)
+void BuildOrCopyElementData(const TSharedRef<FUpdateContextPrivate>& Context, USkeletalMesh* SkeletalMesh, int32 ComponentIndex)
 {
 	MUTABLE_CPUPROFILER_SCOPE(UCustomizableInstancePrivate::BuildOrCopyElementData);
 
-	const FInstanceUpdateData::FComponent* Component = Context->GetComponentUpdateData(InstanceComponentIndex);
+	const TSharedPtr<FInstanceUpdateData::FComponent>& Component = Context->InstanceUpdateData.Components[ComponentIndex];
 	if (!Component)
 	{
 		return;
 	}
 
-	const FName ComponentName = Context->ComponentNames[Component->Id.GetValue()];
+	const FName ComponentName = Context->ComponentNames[Component->Id];
 
 	for (int32 LODIndex = Context->FirstLODAvailable[ComponentName]; LODIndex < Component->LODCount; ++LODIndex)
 	{
@@ -4020,7 +3901,7 @@ void BuildOrCopyElementData(const TSharedRef<FUpdateContextPrivate>& Context, US
 	}
 }
 
-void BuildOrCopyMorphTargetsData(const TSharedRef<FUpdateContextPrivate>& Context, USkeletalMesh* SkeletalMesh, FCustomizableObjectInstanceComponentIndex InstanceComponentIndex)
+void BuildOrCopyMorphTargetsData(const TSharedRef<FUpdateContextPrivate>& Context, USkeletalMesh* SkeletalMesh, int32 ComponentIndex)
 {
 	MUTABLE_CPUPROFILER_SCOPE(UCustomizableInstancePrivate::BuildOrCopyMorphTargetsData);
 
@@ -4034,11 +3915,12 @@ void BuildOrCopyMorphTargetsData(const TSharedRef<FUpdateContextPrivate>& Contex
 	{
 		return;
 	}
-	
-	FCustomizableObjectComponentIndex ObjectComponentIndex = Context->GetObjectComponentIndex(InstanceComponentIndex);
-	const FName ComponentName = Context->ComponentNames[ObjectComponentIndex.GetValue()];
 
-	FSkeletalMeshMorphTargets* ComponentMorphTargets = Context->InstanceUpdateData.RealTimeMorphTargets.Find(ComponentName);
+	const TSharedPtr<FInstanceUpdateData::FComponent>& Component = Context->InstanceUpdateData.Components[ComponentIndex];
+
+	const FName ComponentName = Context->ComponentNames[Component->Id];
+
+	FSkeletalMeshMorphTargets* ComponentMorphTargets = Context->InstanceUpdateData.RealTimeMorphTargets.Find(Component->Id);
 	if (!ComponentMorphTargets)
 	{
 		return;
@@ -4135,7 +4017,7 @@ void BuildOrCopyClothingData(
 		const TSharedRef<FUpdateContextPrivate>& Context, 
 		USkeletalMesh* SkeletalMesh, 
 		const UModelResources& ModelResources,
-		FCustomizableObjectInstanceComponentIndex InstanceComponentIndex,
+		int32 ComponentIndex,
 		const TArray<TObjectPtr<UPhysicsAsset>>& ClothingPhysicsAssets)
 {
 	MUTABLE_CPUPROFILER_SCOPE(UCustomizableInstancePrivate::BuildOrCopyClothingData);
@@ -4166,13 +4048,13 @@ void BuildOrCopyClothingData(
 		return;
 	}
 
-	const FInstanceUpdateData::FComponent* Component = Context->GetComponentUpdateData(InstanceComponentIndex);
+	const TSharedPtr<FInstanceUpdateData::FComponent>& Component = Context->InstanceUpdateData.Components[ComponentIndex];
 	if (!Component)
 	{
 		return;
 	}
 
-	const FName& ComponentName = Context->ComponentNames[Component->Id.GetValue()];
+	const FName& ComponentName = Context->ComponentNames[Component->Id];
 
 	TArray<FSectionClothMetadata> SectionClothMetadata; // Sections must be sorted ascending
 	SectionClothMetadata.Reserve(32);
@@ -4570,7 +4452,7 @@ void BuildOrCopyClothingData(
 
 
 bool BuildOrCopyRenderData(const TSharedRef<FUpdateContextPrivate>& Context, USkeletalMesh* SkeletalMesh, 
-	UCustomizableObjectInstance* Public, FCustomizableObjectInstanceComponentIndex InstanceComponentIndex)
+	UCustomizableObjectInstance* Public, int32 ComponentIndex)
 {
 	MUTABLE_CPUPROFILER_SCOPE(UCustomizableInstancePrivate::BuildOrCopyRenderData);
 
@@ -4578,18 +4460,15 @@ bool BuildOrCopyRenderData(const TSharedRef<FUpdateContextPrivate>& Context, USk
 	check(RenderData);
 
 	UCustomizableObject* CustomizableObject = Public->GetCustomizableObject();
-
-	// It must be not null as it's checked in the calling function
-	check(CustomizableObject);
-
-	const FInstanceUpdateData::FComponent* Component = Context->GetComponentUpdateData(InstanceComponentIndex);
+	
+	const TSharedPtr<FInstanceUpdateData::FComponent>& Component = Context->InstanceUpdateData.Components[ComponentIndex];
 	if (!Component)
 	{
 		return false;
 	}
 
 	const UModelResources& ModelResources = CustomizableObject->GetPrivate()->GetModelResourcesChecked();
-	const FName ComponentName = ModelResources.ComponentNamesPerObjectComponent[Component->Id.GetValue()];
+	const FName ComponentName = ModelResources.ComponentNamesPerObjectComponent[Component->Id];
 
 	const int32 FirstLOD = Context->bStreamMeshLODs ?
 		Context->FirstLODAvailable[ComponentName] :
@@ -4607,17 +4486,17 @@ bool BuildOrCopyRenderData(const TSharedRef<FUpdateContextPrivate>& Context, USk
 			UE_LOG(LogMutable, Warning, TEXT("Building instance: generated mesh [%s] has LOD [%d] of object component index [%d] with no mesh.")
 				, *SkeletalMesh->GetName()
 				, LODIndex
-				, Component->Id.GetValue());
+				, Component->Id);
 
 			// End with failure
 			return false;
 		}
 
-		TRACE_CPUPROFILER_EVENT_SCOPE_TEXT(*FString::Printf(TEXT("BuildRenderData: Component index %d, LOD %d"), Component->Id.GetValue(), LODIndex));
+		TRACE_CPUPROFILER_EVENT_SCOPE_TEXT(*FString::Printf(TEXT("BuildRenderData: Component index %d, LOD %d"), Component->Id, LODIndex));
 		
 		FSkeletalMeshLODRenderData& LODResource = RenderData->LODRenderData[LODIndex];
 
-		const TMap<UE::Mutable::Private::FBoneName, TPair<FName, uint16>>& BoneInfoMap = Context->InstanceUpdateData.SkeletonsPerInstanceComponent[InstanceComponentIndex.GetValue()].BoneInfoMap;
+		const TMap<UE::Mutable::Private::FBoneName, TPair<FName, uint16>>& BoneInfoMap = Context->InstanceUpdateData.SkeletonsPerInstanceComponent[ComponentIndex].BoneInfoMap;
 		
 		// Set active and required bones
 		{
@@ -4642,14 +4521,14 @@ bool BuildOrCopyRenderData(const TSharedRef<FUpdateContextPrivate>& Context, USk
 		for (int32 MeshSectionIndex = 0; MeshSectionIndex < MeshNumSurfaces; ++MeshSectionIndex)
 		{
 			uint32 MeshSurfaceId = LOD->Mesh->GetSurfaceId(MeshSectionIndex);
-			int32 InstanceSurfaceIndex = Context->MutableInstance->FindSurfaceById(InstanceComponentIndex.GetValue(), LODIndex, MeshSurfaceId);
+			int32 InstanceSurfaceIndex = Context->MutableInstance->FindSurfaceById(ComponentIndex, LODIndex, MeshSurfaceId);
 			
 			if (InstanceSurfaceIndex < 0)
 			{
 				continue;
 			}
 			
-			uint32 SurfaceMetadataId = Context->MutableInstance->GetSurfaceCustomId(InstanceComponentIndex.GetValue(), LODIndex, InstanceSurfaceIndex);
+			uint32 SurfaceMetadataId = Context->MutableInstance->GetSurfaceCustomId(ComponentIndex, LODIndex, InstanceSurfaceIndex);
 			
 			uint32 UsedSurfaceMetadataId = 0;
 			if (SurfaceMetadataId != 0)
@@ -4857,33 +4736,29 @@ UE::Tasks::FTask LoadAdditionalAssetsAndData(const TSharedRef<FUpdateContextPriv
 	InstancePrivate->AnimBpPhysicsAssets.Reset();
 
 	// Load Skeletons required by the SubMeshes of the newly generated Mesh, will be merged later
-	for (int32 InstanceComponentIndex = 0; InstanceComponentIndex < Context->NumInstanceComponents; ++InstanceComponentIndex)
+	for (int32 ComponentIndex = 0; ComponentIndex < Context->InstanceUpdateData.Components.Num(); ++ComponentIndex)
 	{
-		FCustomizableObjectComponentIndex ObjectComponentIndex = Context->GetObjectComponentIndex(FCustomizableObjectInstanceComponentIndex(InstanceComponentIndex));
-		if (!ObjectComponentIndex.IsValid())
+		const TSharedPtr<FInstanceUpdateData::FComponent>& Component = Context->InstanceUpdateData.Components[ComponentIndex];
+		if (Component->Id == INDEX_NONE)
 		{
 			continue;
 		}
 
-		const FInstanceUpdateData::FSkeletonData& SkeletonData = Context->InstanceUpdateData.SkeletonsPerInstanceComponent[InstanceComponentIndex];
+		const FInstanceUpdateData::FSkeletonData& SkeletonData = Context->InstanceUpdateData.SkeletonsPerInstanceComponent[ComponentIndex];
 		
-		FCustomizableInstanceComponentData* ComponentData = InstancePrivate->GetComponentData(ObjectComponentIndex);
-		if (!ComponentData)
-		{
-			continue;
-		}
+		FCustomizableInstanceComponentData& ComponentData = InstancePrivate->GetComponentData(Component->Id);
 
 		// Reuse merged Skeleton if cached
-		ComponentData->Skeletons.Skeleton = CustomizableObject->GetPrivate()->SkeletonCache.Get(SkeletonData.SkeletonIds);
-		if (ComponentData->Skeletons.Skeleton)
+		ComponentData.Skeletons.Skeleton = CustomizableObject->GetPrivate()->SkeletonCache.Get(SkeletonData.SkeletonIds);
+		if (ComponentData.Skeletons.Skeleton)
 		{
-			ComponentData->Skeletons.SkeletonIds.Empty();
-			ComponentData->Skeletons.SkeletonsToMerge.Empty();
+			ComponentData.Skeletons.SkeletonIds.Empty();
+			ComponentData.Skeletons.SkeletonsToMerge.Empty();
 			continue;
 		}
 
 		// Add Skeletons to merge
-		ComponentData->Skeletons.SkeletonsToMerge = SkeletonData.Skeletons;
+		ComponentData.Skeletons.SkeletonsToMerge = SkeletonData.Skeletons;
 
 		for (const uint32 SkeletonId : SkeletonData.SkeletonIds)
 		{
@@ -4894,11 +4769,11 @@ UE::Tasks::FTask LoadAdditionalAssetsAndData(const TSharedRef<FUpdateContextPriv
 			}
 
 			// Add referenced skeletons to the assets to stream
-			ComponentData->Skeletons.SkeletonIds.Add(SkeletonId);
+			ComponentData.Skeletons.SkeletonIds.Add(SkeletonId);
 
 			if (USkeleton* Skeleton = AssetPtr.Get())
 			{
-				ComponentData->Skeletons.SkeletonsToMerge.Add(Skeleton);
+				ComponentData.Skeletons.SkeletonsToMerge.Add(Skeleton);
 			}
 			else
 			{
@@ -4911,19 +4786,10 @@ UE::Tasks::FTask LoadAdditionalAssetsAndData(const TSharedRef<FUpdateContextPriv
 	const bool bUpdateMeshes = DoComponentsNeedUpdate(Context, bHasInvalidMesh);
 
 	// Load assets coming from SubMeshes of the newly generated Mesh
-	for (int32 InstanceComponentIndex = 0; InstanceComponentIndex < Context->InstanceUpdateData.Components.Num(); ++InstanceComponentIndex)
+	for (int32 ComponentIndex = 0; ComponentIndex < Context->InstanceUpdateData.Components.Num(); ++ComponentIndex)
 	{
-		const TSharedRef<const FInstanceUpdateData::FComponent>& Component = Components[InstanceComponentIndex];
-		FCustomizableInstanceComponentData* ComponentData = InstancePrivate->GetComponentData(Component->Id);
-		
-		if (!ComponentData)
-        {
-            UE_LOG(LogMutable, Error, TEXT("ComponentData is null: ComponentIDIndex : '%i', COIName :'%s'."), Component->Id.GetValue(), *Instance->GetName());
-			Context->UpdateResult = EUpdateResult::Error;
-			
-			return UE::Tasks::FTask{};
-        }
-        
+		const TSharedRef<const FInstanceUpdateData::FComponent>& Component = Components[ComponentIndex];
+		FCustomizableInstanceComponentData& ComponentData = InstancePrivate->GetComponentData(Component->Id);
 		
 		if (Component->SkeletalMesh->PassthroughObject)
 		{
@@ -4931,18 +4797,17 @@ UE::Tasks::FTask LoadAdditionalAssetsAndData(const TSharedRef<FUpdateContextPriv
 			continue;
 		}
 		
-		const FCustomizableObjectComponentIndex ObjectComponentIndex = Component->Id;
-		const FMutableRefSkeletalMeshData& RefSkeletalMeshData = ModelResources.ReferenceSkeletalMeshesData[ObjectComponentIndex.GetValue()];
+		const FMutableRefSkeletalMeshData& RefSkeletalMeshData = ModelResources.ReferenceSkeletalMeshesData[Component->Id];
 
 		for (int32 AssetUserDataIndex : RefSkeletalMeshData.AssetUserDataIndices)
 		{
 #if !WITH_EDITOR
 			Context->StreamedResourceIndex.AddUnique(AssetUserDataIndex); // Used to hold/release streamed resources in non-editor builds.
 #endif
-			ComponentData->StreamedResourceIndex.AddUnique(AssetUserDataIndex);
+			ComponentData.StreamedResourceIndex.AddUnique(AssetUserDataIndex);
 		}
 
-		const FName ComponentName = Context->ComponentNames[Component->Id.GetValue()];
+		const FName ComponentName = Context->ComponentNames[Component->Id];
 
 		if (bUpdateMeshes)
 		{
@@ -5045,10 +4910,11 @@ UE::Tasks::FTask LoadAdditionalAssetsAndData(const TSharedRef<FUpdateContextPriv
 										{
 											InstancePrivate->ClothingPhysicsAssets[ClothingAssetIndex] = PhysicsAsset.Get();
 										}
+											ComponentData.ClothingPhysicsAssetsToStream.Emplace(ClothingAssetIndex, PhysicsAssetIndex);
 									}
 									else
 									{
-										ComponentData->ClothingPhysicsAssetsToStream.Emplace(ClothingAssetIndex, PhysicsAssetIndex);
+										ComponentData.ClothingPhysicsAssetsToStream.Emplace(ClothingAssetIndex, PhysicsAssetIndex);
 										AssetsToStream.AddUnique(PhysicsAsset.ToSoftObjectPath());
 									}
 								}
@@ -5085,7 +4951,7 @@ UE::Tasks::FTask LoadAdditionalAssetsAndData(const TSharedRef<FUpdateContextPriv
 #if !WITH_EDITOR
 					Context->StreamedResourceIndex.AddUnique(ResourceIndex); // Used to hold/release streamed resources in non-editor builds.
 #endif
-					ComponentData->StreamedResourceIndex.AddUnique(ResourceIndex);
+					ComponentData.StreamedResourceIndex.AddUnique(ResourceIndex);
 				}
 			}
 
@@ -5101,7 +4967,7 @@ UE::Tasks::FTask LoadAdditionalAssetsAndData(const TSharedRef<FUpdateContextPriv
 							// No need to load since the SKM is a StrongObjectPtr
 							if (AssetUserData)
 							{
-								ComponentData->AssetUserDataArray.Add(AssetUserData);
+								ComponentData.AssetUserDataArray.Add(AssetUserData);
 							}
 						}
 					}
@@ -5120,11 +4986,11 @@ UE::Tasks::FTask LoadAdditionalAssetsAndData(const TSharedRef<FUpdateContextPriv
 					{
 						if (UPhysicsAsset* LoadedPhysicsAssetPtr = PhysicsAsset.Get())
 						{
-							ComponentData->PhysicsAssets.PhysicsAssetsToMerge.Add(LoadedPhysicsAssetPtr);
+							ComponentData.PhysicsAssets.PhysicsAssetsToMerge.Add(LoadedPhysicsAssetPtr);
 						}
 						else
 						{
-							ComponentData->PhysicsAssets.PhysicsAssetToLoad.Add(AssetIndex);
+							ComponentData.PhysicsAssets.PhysicsAssetToLoad.Add(AssetIndex);
 							AssetsToStream.AddUnique(PhysicsAsset.ToSoftObjectPath());
 						}
 					}
@@ -5153,11 +5019,11 @@ UE::Tasks::FTask LoadAdditionalAssetsAndData(const TSharedRef<FUpdateContextPriv
 
 						if (!AnimBPAsset.IsNull())
 						{
-							const TSoftClassPtr<UAnimInstance>* FoundAnimBpSlot = ComponentData->AnimSlotToBP.Find(SlotIndex);
+							const TSoftClassPtr<UAnimInstance>* FoundAnimBpSlot = ComponentData.AnimSlotToBP.Find(SlotIndex);
 							bool bIsSameAnimBp = FoundAnimBpSlot && AnimBPAsset == *FoundAnimBpSlot;
 							if (!FoundAnimBpSlot)
 							{
-								ComponentData->AnimSlotToBP.Add(SlotIndex, AnimBPAsset);
+								ComponentData.AnimSlotToBP.Add(SlotIndex, AnimBPAsset);
 
 								if (AnimBPAsset.Get())
 								{
@@ -5190,7 +5056,7 @@ UE::Tasks::FTask LoadAdditionalAssetsAndData(const TSharedRef<FUpdateContextPriv
 #if WITH_EDITORONLY_DATA
 				else if (Tag.RemoveFromStart("__MeshPath:"))
 				{
-					ComponentData->MeshPartPaths.Add(Tag);
+					ComponentData.MeshPartPaths.Add(Tag);
 				}
 #endif
 			}
@@ -5200,13 +5066,13 @@ UE::Tasks::FTask LoadAdditionalAssetsAndData(const TSharedRef<FUpdateContextPriv
 			{
 				const int32 ExternalId = MutableMesh->AdditionalPhysicsBodies[I]->CustomId;
 				
-				ComponentData->PhysicsAssets.AdditionalPhysicsAssetsToLoad.Add(ExternalId);
+				ComponentData.PhysicsAssets.AdditionalPhysicsAssetsToLoad.Add(ExternalId);
 				AssetsToStream.Add(ModelResources.AnimBpOverridePhysiscAssetsInfo[ExternalId].SourceAsset.ToSoftObjectPath());
 			}
 		}
 
 
-		for (int32 ResourceIndex : ComponentData->StreamedResourceIndex)
+		for (int32 ResourceIndex : ComponentData.StreamedResourceIndex)
 		{
 #if WITH_EDITOR
 			if (!ModelResources.StreamedResourceDataEditor.IsValidIndex(ResourceIndex))
@@ -5988,23 +5854,23 @@ void UCustomizableInstancePrivate::BuildResources(const TSharedRef<FUpdateContex
 	Materials.Empty();
 	Textures.Empty();
 
+	const int32 NumComponents = Context->InstanceUpdateData.Components.Num();
+	
 	TArray<bool> RecreateRenderStateOnInstanceComponent;
-	RecreateRenderStateOnInstanceComponent.Init(false, Context->NumInstanceComponents);
+	RecreateRenderStateOnInstanceComponent.Init(false, NumComponents);
 
 	TArray<bool> NotifyUpdateOnInstanceComponent;
-	NotifyUpdateOnInstanceComponent.Init(false, Context->NumInstanceComponents);
+	NotifyUpdateOnInstanceComponent.Init(false, NumComponents);
 
-	for (int32 InstanceComponentIndex = 0; InstanceComponentIndex < Context->NumInstanceComponents; ++InstanceComponentIndex)
+	for (int32 ComponentIndex = 0; ComponentIndex < NumComponents; ++ComponentIndex)
 	{
-		const TSharedRef<const FInstanceUpdateData::FComponent>& Component = Context->InstanceUpdateData.Components[InstanceComponentIndex];
-
-		const FCustomizableObjectComponentIndex ObjectComponentIndex = Component->Id;
+		const TSharedRef<const FInstanceUpdateData::FComponent>& Component = Context->InstanceUpdateData.Components[ComponentIndex];
 
-		if (!ModelResources.ComponentNamesPerObjectComponent.IsValidIndex(ObjectComponentIndex.GetValue()))
+		if (!ModelResources.ComponentNamesPerObjectComponent.IsValidIndex(Component->Id))
 		{
 			continue;
 		}
-		const FName& ComponentName = ModelResources.ComponentNamesPerObjectComponent[ObjectComponentIndex.GetValue()];
+		const FName& ComponentName = ModelResources.ComponentNamesPerObjectComponent[Component->Id];
 		
 		TObjectPtr<USkeletalMesh> SkeletalMesh = SkeletalMeshes.Contains(ComponentName) ? SkeletalMeshes[ComponentName] : nullptr;
 		if (!SkeletalMesh)
@@ -6012,7 +5878,7 @@ void UCustomizableInstancePrivate::BuildResources(const TSharedRef<FUpdateContex
 			continue;
 		}
 
-		const bool bUVChanged = Context->MeshChangedPerInstanceComponent[InstanceComponentIndex];
+		const bool bUVChanged = Context->MeshChangedPerInstanceComponent[ComponentIndex];
 		
 		// If the mesh is not transient, it means it's pass-through so it should use material overrides and not be modified in any way
 		const bool bIsTransientMesh = SkeletalMesh->HasAllFlags(EObjectFlags::RF_Transient);
@@ -6022,45 +5888,42 @@ void UCustomizableInstancePrivate::BuildResources(const TSharedRef<FUpdateContex
 
 		UMaterialInterface* OverlayMaterial = nullptr;
 
-		FCustomizableInstanceComponentData* ComponentData = GetComponentData(ObjectComponentIndex);
-		if (ComponentData)
-		{
-			ComponentData->OverrideMaterials.Reset();
-			ComponentData->OverlayMaterial = nullptr;
+		FCustomizableInstanceComponentData& ComponentData = GetComponentData(Component->Id);
+		ComponentData.OverrideMaterials.Reset();
+		ComponentData.OverlayMaterial = nullptr;
 			
-			if (Component->OverlayMaterialIndex != INDEX_NONE)
-			{
-				FBuildMaterialOut Out = BuildMaterial(Context, Component->OverlayMaterialIndex, !bUVChanged);
-				OverlayMaterial = Out.Material;
-				ComponentData->OverlayMaterial = Out.Material;
+		if (Component->OverlayMaterialIndex != INDEX_NONE)
+		{
+			FBuildMaterialOut Out = BuildMaterial(Context, Component->OverlayMaterialIndex, !bUVChanged);
+			OverlayMaterial = Out.Material;
+			ComponentData.OverlayMaterial = Out.Material;
 				
-				NotifyUpdateOnInstanceComponent[InstanceComponentIndex] = Out.bNotifyPrimitiveUpdated;
-			}
+			NotifyUpdateOnInstanceComponent[ComponentIndex] = Out.bNotifyPrimitiveUpdated;
+		}
 			
-			for (int32 MaterialIndex = 0; MaterialIndex < Component->OverrideMaterialCount; ++MaterialIndex)
-			{
-				const FInstanceUpdateData::FOverrideMaterial& OverrideMaterial = Context->InstanceUpdateData.OverrideMaterials[Component->FirstOverrideMaterial + MaterialIndex];
+		for (int32 MaterialIndex = 0; MaterialIndex < Component->OverrideMaterialCount; ++MaterialIndex)
+		{
+			const FInstanceUpdateData::FOverrideMaterial& OverrideMaterial = Context->InstanceUpdateData.OverrideMaterials[Component->FirstOverrideMaterial + MaterialIndex];
 				
-				const int32 SlotIndex = SkeletalMesh->GetMaterials().IndexOfByPredicate([&](const FSkeletalMaterial& Material)
-				{
-					return Material.MaterialSlotName == OverrideMaterial.SlotName;
-				});
+			const int32 SlotIndex = SkeletalMesh->GetMaterials().IndexOfByPredicate([&](const FSkeletalMaterial& Material)
+			{
+				return Material.MaterialSlotName == OverrideMaterial.SlotName;
+			});
 				
-				if (SlotIndex != INDEX_NONE)
-				{
-					FBuildMaterialOut Out = BuildMaterial(Context, OverrideMaterial.MaterialIndex, !bUVChanged);
+			if (SlotIndex != INDEX_NONE)
+			{
+				FBuildMaterialOut Out = BuildMaterial(Context, OverrideMaterial.MaterialIndex, !bUVChanged);
 					
-					ComponentData->OverrideMaterials.SetNum(SlotIndex + 1);
-					ComponentData->OverrideMaterials[SlotIndex] = Out.Material;
+				ComponentData.OverrideMaterials.SetNum(SlotIndex + 1);
+				ComponentData.OverrideMaterials[SlotIndex] = Out.Material;
 					
-					NotifyUpdateOnInstanceComponent[InstanceComponentIndex] = Out.bNotifyPrimitiveUpdated;
-				}
+				NotifyUpdateOnInstanceComponent[ComponentIndex] = Out.bNotifyPrimitiveUpdated;
 			}
 		}
 
 		if (!bUseOverrideMaterialsOnly)
 		{
-			RecreateRenderStateOnInstanceComponent[InstanceComponentIndex] |= SkeletalMesh->GetOverlayMaterial() != OverlayMaterial;
+			RecreateRenderStateOnInstanceComponent[ComponentIndex] |= SkeletalMesh->GetOverlayMaterial() != OverlayMaterial;
 			SkeletalMesh->SetOverlayMaterial(OverlayMaterial);
 		}
 
@@ -6086,9 +5949,9 @@ void UCustomizableInstancePrivate::BuildResources(const TSharedRef<FUpdateContex
 
 			// Pass-through components will not have a reference mesh.
 			const FMutableRefSkeletalMeshData* RefSkeletalMeshData = nullptr;
-			if (ModelResources.ReferenceSkeletalMeshesData.IsValidIndex(ObjectComponentIndex.GetValue()))
+			if (ModelResources.ReferenceSkeletalMeshesData.IsValidIndex(Component->Id))
 			{
-				RefSkeletalMeshData = &ModelResources.ReferenceSkeletalMeshesData[ObjectComponentIndex.GetValue()];
+				RefSkeletalMeshData = &ModelResources.ReferenceSkeletalMeshesData[Component->Id];
 			}
 
 			for (int32 SurfaceIndex = 0; SurfaceIndex < LOD->SurfaceCount; ++SurfaceIndex)
@@ -6176,19 +6039,16 @@ void UCustomizableInstancePrivate::BuildResources(const TSharedRef<FUpdateContex
 					}
 				}
 
-				if (ComponentData)
-				{
-					FBuildMaterialOut Out = BuildMaterial(Context, Surface->MaterialIndex, !bUVChanged);
-					
-					if (!Out.Material)
-					{
-						Out.Material = UMaterial::GetDefaultMaterial(MD_Surface);
-					}
+				FBuildMaterialOut Out = BuildMaterial(Context, Surface->MaterialIndex, !bUVChanged);
 					
-					MaterialSlot.MaterialInterface = Out.Material;
-					ComponentData->OverrideMaterials.Add(Out.Material);
-					NotifyUpdateOnInstanceComponent[InstanceComponentIndex] = Out.bNotifyPrimitiveUpdated;
+				if (!Out.Material)
+				{
+					Out.Material = UMaterial::GetDefaultMaterial(MD_Surface);
 				}
+					
+				MaterialSlot.MaterialInterface = Out.Material;
+				ComponentData.OverrideMaterials.Add(Out.Material);
+				NotifyUpdateOnInstanceComponent[ComponentIndex] = Out.bNotifyPrimitiveUpdated;
 			}
 		}
 
@@ -6219,12 +6079,12 @@ void UCustomizableInstancePrivate::BuildResources(const TSharedRef<FUpdateContex
 			}
 
 			// Force recreate render state after replacing the materials to avoid a crash in the render pipeline if the old materials are GCed while in use.
-			RecreateRenderStateOnInstanceComponent[InstanceComponentIndex] |= SkeletalMesh->GetResourceForRendering()->IsInitialized() && SkeletalMesh->GetMaterials() != SkeletalMaterials;
+			RecreateRenderStateOnInstanceComponent[ComponentIndex] |= SkeletalMesh->GetResourceForRendering()->IsInitialized() && SkeletalMesh->GetMaterials() != SkeletalMaterials;
 
 			SkeletalMesh->SetMaterials(SkeletalMaterials);
 
 #if WITH_EDITOR
-			if (GEditor && RecreateRenderStateOnInstanceComponent[InstanceComponentIndex])
+			if (GEditor && RecreateRenderStateOnInstanceComponent[ComponentIndex])
 			{
 				// Close all open editors for this mesh to invalidate viewports.
 				GEditor->GetEditorSubsystem<UAssetEditorSubsystem>()->CloseAllEditorsForAsset(SkeletalMesh);
@@ -6259,20 +6119,24 @@ void UCustomizableInstancePrivate::BuildResources(const TSharedRef<FUpdateContex
 #endif
 
 			const FName& ComponentName = CustomizableObjectInstanceUsage->GetComponentName();
-			const int32 ObjectComponentIndex = ModelResources.ComponentNamesPerObjectComponent.IndexOfByKey(ComponentName);
+			FCustomizableInstanceComponentData* ComponentData = GetComponentData(ComponentName);
+			if (!ComponentData)
+			{
+				continue;
+			}
 			
-			int32 InstanceComponentIndex = -1;
+			int32 ComponentIndex = -1;
 			for (int32 CurrentInstanceIndex = 0; CurrentInstanceIndex<Context->InstanceUpdateData.Components.Num(); ++CurrentInstanceIndex)
 			{
-				if (Context->InstanceUpdateData.Components[CurrentInstanceIndex]->Id.GetValue() == ObjectComponentIndex)
+				if (Context->InstanceUpdateData.Components[CurrentInstanceIndex]->Id == ComponentData->ComponentId)
 				{
-					InstanceComponentIndex = CurrentInstanceIndex;
+					ComponentIndex = CurrentInstanceIndex;
 					break;
 				}
 			}
 
-			bool bDoRecreateRenderStateOnComponent = RecreateRenderStateOnInstanceComponent.IsValidIndex(InstanceComponentIndex) && RecreateRenderStateOnInstanceComponent[InstanceComponentIndex];
-			bool bDoNotifyUpdateOnComponent = NotifyUpdateOnInstanceComponent.IsValidIndex(InstanceComponentIndex) && NotifyUpdateOnInstanceComponent[InstanceComponentIndex];
+			bool bDoRecreateRenderStateOnComponent = RecreateRenderStateOnInstanceComponent.IsValidIndex(ComponentIndex) && RecreateRenderStateOnInstanceComponent[ComponentIndex];
+			bool bDoNotifyUpdateOnComponent = NotifyUpdateOnInstanceComponent.IsValidIndex(ComponentIndex) && NotifyUpdateOnInstanceComponent[ComponentIndex];
 
 			if (!bDoRecreateRenderStateOnComponent && !bDoNotifyUpdateOnComponent)
 			{
@@ -6415,19 +6279,21 @@ void UCustomizableObjectInstance::SetRequestedLODs(const TMap<FName, uint8>& InM
 	
 	const TMap<FName, uint8>& FirstRequestedLOD = GetPrivate()->CommittedDescriptorHash.FirstRequestedLOD;
 	
-	for (const FName& ComponentName : ModelResources->ComponentNamesPerObjectComponent)
+	for (UE::Mutable::Private::FComponentId ComponentId = 0; ComponentId < ModelResources->ComponentNamesPerObjectComponent.Num(); ++ComponentId)
 	{	
+		const FName& ComponentName = ModelResources->ComponentNamesPerObjectComponent[ComponentId];
+		
 		uint8& InMinLOD = MutableUpdateCandidate.MinLOD.FindOrAdd(ComponentName);
 		if (const uint8* Result = InMinLODs.Find(ComponentName))
 		{
 			InMinLOD = *Result;
 		}
 
-		const uint8 MinLODIdx = CustomizableObject->GetPrivate()->GetMinLODIndex(ComponentName);
+		const uint8 MinLODIdx = CustomizableObject->GetPrivate()->GetMinLODIndex(ComponentId);
 		MutableUpdateCandidate.QualitySettingMinLODs.Add(ComponentName, MinLODIdx);
 
 		int32 MaxLODIdx = 0;
-		if (const uint8* Found = ModelResources->NumLODsAvailable.Find(ComponentName))
+		if (const uint8* Found = ModelResources->NumLODsAvailable.Find(ComponentId))
 		{
 			MaxLODIdx = *Found - 1;
 		}
@@ -6453,7 +6319,7 @@ void UCustomizableObjectInstance::SetRequestedLODs(const TMap<FName, uint8>& InM
 			PrivateData->SetCOInstanceFlags(InMinLOD > CurrentMinLOD ? PendingLODsDowngrade : ECONone);
 
 			uint8 FirstNonStreamedLODIndex = 0;
-			if (const uint8* Found = ModelResources->NumLODsToStream.Find(ComponentName))
+			if (const uint8* Found = ModelResources->NumLODsToStream.Find(ComponentId))
 			{
 				FirstNonStreamedLODIndex = *Found;
 			}
@@ -6511,12 +6377,6 @@ void UCustomizableObjectInstance::Bake(const FBakingConfiguration& InBakingConfi
 #endif
 
 
-USkeletalMesh* UCustomizableObjectInstance::GetSkeletalMesh(int32 ObjectComponentIndex) const
-{
-	return GetComponentMeshSkeletalMesh(FName(FString::FromInt(ObjectComponentIndex)));
-}
-
-
 USkeletalMesh* UCustomizableObjectInstance::GetComponentMeshSkeletalMesh(const FName& ComponentName) const
 {
 	TObjectPtr<USkeletalMesh>* Result = GetPrivate()->SkeletalMeshes.Find(ComponentName);
@@ -6564,7 +6424,6 @@ TArray<FName> UCustomizableObjectInstance::GetComponentNames() const
 TSubclassOf<UAnimInstance> UCustomizableObjectInstance::GetAnimBP(FName ComponentName, const FName& SlotName) const
 {
 	FCustomizableInstanceComponentData* ComponentData =	GetPrivate()->GetComponentData(ComponentName);
-	
 	if (!ComponentData)
 	{
 		FString ErrorMsg = FString::Printf(TEXT("Tried to access an invalid component index [%s] in a Mutable Instance."), *ComponentName.ToString());
@@ -6666,29 +6525,6 @@ void UCustomizableObjectInstance::ForEachComponentAnimInstance(FName ComponentNa
 	UE::Mutable::Private::InternalForEachAnimInstance<>(GetPrivate(), ComponentName, Delegate);
 }
 
-// Deprecated
-void UCustomizableObjectInstance::ForEachAnimInstance(int32 ObjectComponentIndex, FEachComponentAnimInstanceClassDelegate Delegate) const
-{
-	UCustomizableObject* CO = GetCustomizableObject();
-	if (CO)
-	{
-		FName ComponentName = CO->GetPrivate()->GetComponentName(FCustomizableObjectComponentIndex(ObjectComponentIndex));
-		UE::Mutable::Private::InternalForEachAnimInstance<>(GetPrivate(), ComponentName, Delegate);
-	}
-}
-
-
-// Deprecated
-void UCustomizableObjectInstance::ForEachAnimInstance(int32 ObjectComponentIndex, FEachComponentAnimInstanceClassNativeDelegate Delegate) const
-{
-	UCustomizableObject* CO = GetCustomizableObject();
-	if (CO)
-	{
-		FName ComponentName = CO->GetPrivate()->GetComponentName(FCustomizableObjectComponentIndex(ObjectComponentIndex));
-		UE::Mutable::Private::InternalForEachAnimInstance<>(GetPrivate(), ComponentName, Delegate);
-	}
-}
-
 
 bool UCustomizableObjectInstance::AnimInstanceNeedsFixup(TSubclassOf<UAnimInstance> AnimInstanceClass) const
 {
@@ -6793,26 +6629,23 @@ FInstancedStruct UCustomizableObjectInstance::GetExtensionInstanceData(const UCu
 }
 
 
-TSet<UAssetUserData*> UCustomizableObjectInstance::GetMergedAssetUserData(int32 ComponentIndex) const
+TSet<UAssetUserData*> UCustomizableObjectInstance::GetMergedAssetUserData(const FName& ComponentName) const
 {
-	UCustomizableInstancePrivate* PrivateInstanceData = GetPrivate();
-
-	if (PrivateInstanceData && PrivateInstanceData->ComponentsData.IsValidIndex(ComponentIndex))
+	FCustomizableInstanceComponentData* ComponentData = GetPrivate()->GetComponentData(ComponentName);
+	if (!ComponentData)
 	{
-		TSet<UAssetUserData*> Set;
-		
-		// Have to convert to UAssetUserData* because BP functions don't support TObjectPtr
-		for (const TObjectPtr<UAssetUserData>& Elem : PrivateInstanceData->ComponentsData[ComponentIndex].AssetUserDataArray)
-		{
-			Set.Add(Elem);
-		}
-
-		return Set;
+		return {};
 	}
-	else
+	
+	TSet<UAssetUserData*> Set;
+		
+	for (const TObjectPtr<UAssetUserData>& Elem : ComponentData->AssetUserDataArray)
 	{
-		return TSet<UAssetUserData*>();
+		// Have to convert to UAssetUserData* because BP functions don't support TObjectPtr
+		Set.Add(Elem);
 	}
+
+	return Set;
 }
 
 
diff --git a/Engine/Plugins/Mutable/Source/CustomizableObject/Private/MuCO/CustomizableObjectInstanceUsage.cpp b/Engine/Plugins/Mutable/Source/CustomizableObject/Private/MuCO/CustomizableObjectInstanceUsage.cpp
--- a/Engine/Plugins/Mutable/Source/CustomizableObject/Private/MuCO/CustomizableObjectInstanceUsage.cpp
+++ b/Engine/Plugins/Mutable/Source/CustomizableObject/Private/MuCO/CustomizableObjectInstanceUsage.cpp
@@ -78,19 +78,6 @@ UCustomizableObjectInstance* UCustomizableObjectInstanceUsage::GetCustomizableOb
 }
 
 
-void UCustomizableObjectInstanceUsage::SetComponentIndex(int32 ComponentIndex)
-{
-	if (UCustomizableSkeletalComponent* CustomizableSkeletalComponent = GetPrivate()->GetCustomizableSkeletalComponent())
-	{
-		CustomizableSkeletalComponent->ComponentIndex = ComponentIndex;
-	}
-	else
-	{
-		UsedComponentIndex = ComponentIndex;
-	}
-}
-
-
 int32 UCustomizableObjectInstanceUsage::GetComponentIndex() const
 {
 	if (UCustomizableSkeletalComponent* CustomizableSkeletalComponent = GetPrivate()->GetCustomizableSkeletalComponent())
diff --git a/Engine/Plugins/Mutable/Source/CustomizableObject/Private/MuCO/CustomizableObjectSystem.cpp b/Engine/Plugins/Mutable/Source/CustomizableObject/Private/MuCO/CustomizableObjectSystem.cpp
--- a/Engine/Plugins/Mutable/Source/CustomizableObject/Private/MuCO/CustomizableObjectSystem.cpp
+++ b/Engine/Plugins/Mutable/Source/CustomizableObject/Private/MuCO/CustomizableObjectSystem.cpp
@@ -252,14 +252,19 @@ FUpdateContextPrivate::FUpdateContextPrivate(UCustomizableObjectInstance& InInst
 	{
 		CapturedDescriptor.bStreamingEnabled = IsStreamingEnabled(*InObject, Descriptor.State);
 
-		FirstLODAvailable = ModelResources->FirstLODAvailable;
-		FirstResidentLOD = ModelResources->NumLODsToStream;
 		ComponentNames = ModelResources->ComponentNamesPerObjectComponent;
+		
+		for (UE::Mutable::Private::FComponentId ComponentId = 0; ComponentId < ModelResources->ComponentNamesPerObjectComponent.Num(); ++ComponentId)
+		{
+			const FName& ComponentName = ModelResources->ComponentNamesPerObjectComponent[ComponentId];
+			
+			FirstLODAvailable.Add(ComponentName, ModelResources->FirstLODAvailable[ComponentId]);
+			FirstResidentLOD.Add(ComponentName, ModelResources->NumLODsToStream[ComponentId]);
+		}
 	}
 
 	CapturedDescriptorHash = FDescriptorHash(Descriptor);
-	NumObjectComponents = InObject->GetComponentCount();
-
+	
 	UCustomizableObjectSystem* System = UCustomizableObjectSystem::GetInstanceChecked();
 	UCustomizableObjectSystemPrivate* SystemPrivate = System->GetPrivate();
 
@@ -561,32 +566,7 @@ const FDescriptorHash& FUpdateContextPrivate::GetCapturedDescriptorHash() const
 
 const FCustomizableObjectInstanceDescriptor&& FUpdateContextPrivate::MoveCommittedDescriptor()
 {
-	return MoveTemp(CapturedDescriptor);	
-}
-
-
-FCustomizableObjectComponentIndex FUpdateContextPrivate::GetObjectComponentIndex(FCustomizableObjectInstanceComponentIndex InstanceComponentIndex)
-{
-	FCustomizableObjectComponentIndex ObjectComponentIndex;
-	if (InstanceUpdateData.Components.IsValidIndex(InstanceComponentIndex.GetValue()))
-	{
-		ObjectComponentIndex = InstanceUpdateData.Components[InstanceComponentIndex.GetValue()]->Id;
-	}
-	else
-	{
-		ObjectComponentIndex.Invalidate();
-	}
-	return ObjectComponentIndex;
-}
-
-
-const FInstanceUpdateData::FComponent* FUpdateContextPrivate::GetComponentUpdateData(FCustomizableObjectInstanceComponentIndex InstanceComponentIndex)
-{
-	if (InstanceUpdateData.Components.IsValidIndex(InstanceComponentIndex.GetValue()))
-	{
-		return &InstanceUpdateData.Components[InstanceComponentIndex.GetValue()].Get();
-	}
-	return nullptr;
+	return MoveTemp(CapturedDescriptor);
 }
 
 
@@ -602,13 +582,9 @@ const TArray<FComponentDescriptor>& FUpdateContextPrivate::GetComponentDescripto
 }
 
 
-FComponentDescriptor* FUpdateContextPrivate::GetComponentDescriptor(FCustomizableObjectComponentIndex Index)
+FComponentDescriptor* FUpdateContextPrivate::GetComponentDescriptor(UE::Mutable::Private::FComponentId ComponentId)
 {
-	if (ComponentDescriptors.IsValidIndex(Index.GetValue()))
-	{
-		return &ComponentDescriptors[Index.GetValue()];
-	}
-	return nullptr;
+	return ComponentDescriptors.IsValidIndex(ComponentId) ? &ComponentDescriptors[ComponentId] : nullptr;
 }
 
 
@@ -1336,14 +1312,13 @@ void UpdateSkeletalMesh(const TSharedRef<FUpdateContextPrivate>& Context)
 
 		TArray<TObjectPtr<UMaterialInterface>> NewOverridenMaterials;
 
+		FCustomizableInstanceComponentData* ComponentData = CustomizableObjectInstance->GetPrivate()->GetComponentData(Usage->GetComponentName());
+		
 		const bool bIsTransientMesh = SkeletalMesh ? SkeletalMesh->HasAllFlags(EObjectFlags::RF_Transient) : false;
 		const bool bUseOverrideMaterials = !bIsTransientMesh || Context->bUseMeshCache;
-		if (bUseOverrideMaterials)
+		if (bUseOverrideMaterials && ComponentData)
 		{
-			if (FCustomizableInstanceComponentData* ComponentData = CustomizableObjectInstance->GetPrivate()->GetComponentData(Usage->GetComponentName()))
-			{
-				NewOverridenMaterials = ComponentData->OverrideMaterials;
-			}
+			NewOverridenMaterials = ComponentData->OverrideMaterials;
 		}
 
 		if (Parent->OverrideMaterials != NewOverridenMaterials)
@@ -1365,18 +1340,14 @@ void UpdateSkeletalMesh(const TSharedRef<FUpdateContextPrivate>& Context)
 
 
 		UMaterialInterface* OverlayMaterial = nullptr;
-		if (bUseOverrideMaterials)
+		if (bUseOverrideMaterials && ComponentData)
 		{
-			if (FCustomizableInstanceComponentData* ComponentData = CustomizableObjectInstance->GetPrivate()->GetComponentData(Usage->GetComponentName()))
-			{
-				OverlayMaterial = ComponentData->OverlayMaterial;
-			}
+			OverlayMaterial = ComponentData->OverlayMaterial;
 		}
 
 		bAttachedParentUpdated |= Parent->GetOverlayMaterial() != OverlayMaterial;
 		Parent->SetOverlayMaterial(OverlayMaterial);
-
-
+		
 		if (CustomizableObjectInstancePrivateData->HasCOInstanceFlags(ReplacePhysicsAssets) &&
 			SkeletalMesh &&
 			Parent->GetWorld())
@@ -1541,7 +1512,7 @@ void UCustomizableObjectSystemPrivate::EnqueueUpdateSkeletalMesh(const TSharedRe
 	{
 		// Mutable is disabled. Set the reference SkeletalMesh and finish the update with success to avoid breaking too many things.
 		Context->UpdateResult = EUpdateResult::Success;
-		Instance->GetPrivate()->SetReferenceSkeletalMesh();
+		Instance->GetPrivate()->ForceSetReferenceSkeletalMesh();
 		FinishUpdateGlobal(Context);
 		return;
 	}
@@ -1883,18 +1854,13 @@ namespace impl
 	
 	void FixLODs(const TSharedRef<FUpdateContextPrivate>& Operation)
 	{
-		if (!Operation->NumObjectComponents)
-		{
-			return;
-		}
-
 		TMap<FName, uint8> MinLODs = Operation->GetCapturedDescriptor().MinLOD;
 		TMap<FName, uint8> RequestedLODs = Operation->GetFirstRequestedLOD();
 		
-		for (int32 InstanceComponentIndex = 0; InstanceComponentIndex < Operation->NumInstanceComponents; ++InstanceComponentIndex)
+		for (int32 ComponentIndex = 0; ComponentIndex < Operation->MutableInstance->GetComponentCount(); ++ComponentIndex)
 		{
-			const int32 ObjectComponentIndex = Operation->MutableInstance->GetComponentId(InstanceComponentIndex);
-			const FName& ComponentName = Operation->ComponentNames[ObjectComponentIndex];
+			const UE::Mutable::Private::FComponentId ComponentId = Operation->MutableInstance->GetComponentId(ComponentIndex);
+			const FName& ComponentName = Operation->ComponentNames[ComponentId];
 
 			uint8& MinLOD = MinLODs.FindOrAdd(ComponentName, 0);
 			uint8& RequestedLOD = RequestedLODs.FindOrAdd(ComponentName, 0);
@@ -1902,7 +1868,7 @@ namespace impl
 			uint8& FirstResidentLOD = Operation->FirstResidentLOD.FindOrAdd(ComponentName, 0);
 			const uint8& FirstLODAvailable = Operation->FirstLODAvailable.FindOrAdd(ComponentName, 0);
 
-			NumLODsAvailable = Operation->MutableInstance->GetLODCount(InstanceComponentIndex);
+			NumLODsAvailable = Operation->MutableInstance->GetLODCount(ComponentIndex);
 			
 			if (Operation->bStreamMeshLODs)
 			{
@@ -1928,12 +1894,12 @@ namespace impl
 	{
 		MUTABLE_CPUPROFILER_SCOPE(Subtask_Mutable_PrepareSkeletonData);
 
-		const int32 NumInstanceComponents = OperationData->InstanceUpdateData.Components.Num();
-		OperationData->InstanceUpdateData.SkeletonsPerInstanceComponent.SetNum(NumInstanceComponents);
+		const int32 NumComponents = OperationData->InstanceUpdateData.Components.Num();
+		OperationData->InstanceUpdateData.SkeletonsPerInstanceComponent.SetNum(NumComponents);
 
-		for (int32 InstanceComponentIndex = 0; InstanceComponentIndex< NumInstanceComponents; ++InstanceComponentIndex)
+		for (int32 ComponentIndex = 0; ComponentIndex< NumComponents; ++ComponentIndex)
 		{
-			const TSharedRef<FInstanceUpdateData::FComponent>& Component = OperationData->InstanceUpdateData.Components[InstanceComponentIndex];
+			const TSharedRef<FInstanceUpdateData::FComponent>& Component = OperationData->InstanceUpdateData.Components[ComponentIndex];
 
 			for (int32 LODIndex = 0; LODIndex < Component->LODCount; ++LODIndex)
 			{
@@ -1945,7 +1911,7 @@ namespace impl
 					continue;
 				}
 
-				FInstanceUpdateData::FSkeletonData& SkeletonData = OperationData->InstanceUpdateData.SkeletonsPerInstanceComponent[InstanceComponentIndex];
+				FInstanceUpdateData::FSkeletonData& SkeletonData = OperationData->InstanceUpdateData.SkeletonsPerInstanceComponent[ComponentIndex];
 
 				// Add SkeletonIds 
 				const int32 SkeletonIDsCount = Mesh->GetSkeletonIDsCount();
@@ -2010,17 +1976,18 @@ namespace impl
 		
 		const TSharedPtr<FModelStreamableBulkData>& ModelStreamableBulkData = Object->GetPrivate()->GetModelStreamableBulkData();
 
-		const int32 NumInstanceComponents = OperationData->InstanceUpdateData.Components.Num();
+		const int32 NumComponents = OperationData->InstanceUpdateData.Components.Num();
 		check(OperationData->InstanceUpdateData.RealTimeMorphTargets.IsEmpty());
-		OperationData->InstanceUpdateData.RealTimeMorphTargets.Reserve(NumInstanceComponents);
-		for ( int32 InstanceComponentIndex = 0; InstanceComponentIndex< NumInstanceComponents; ++InstanceComponentIndex)
+		OperationData->InstanceUpdateData.RealTimeMorphTargets.Reserve(NumComponents);
+		
+		for (int32 ComponentIndex = 0; ComponentIndex< NumComponents; ++ComponentIndex)
 		{
-			TSharedRef<FInstanceUpdateData::FComponent>& Component = OperationData->InstanceUpdateData.Components[InstanceComponentIndex];
-			check(Component->Id.IsValid());
+			TSharedRef<FInstanceUpdateData::FComponent>& Component = OperationData->InstanceUpdateData.Components[ComponentIndex];
+			check(Component->Id != INDEX_NONE);
 
-			const FName& ComponentName = OperationData->ComponentNames[Component->Id.GetValue()];
+			const FName& ComponentName = OperationData->ComponentNames[Component->Id];
 
-			FSkeletalMeshMorphTargets& ComponentMorphTargetsData = OperationData->InstanceUpdateData.RealTimeMorphTargets.Add(ComponentName);
+			FSkeletalMeshMorphTargets& ComponentMorphTargetsData = OperationData->InstanceUpdateData.RealTimeMorphTargets.Add(Component->Id);
 			
 			ComponentMorphTargetsData.RealTimeMorphTargetNames.Empty();
 
@@ -2155,10 +2122,10 @@ namespace impl
 		
 		UE::Tasks::FTask LastTask = UE::Tasks::MakeCompletedTask<void>();
 		
-		for (int32 InstanceComponentIndex = 0; InstanceComponentIndex < OperationData->NumInstanceComponents; ++InstanceComponentIndex)
+		for (int32 ComponentIndex = 0; ComponentIndex < OperationData->InstanceUpdateData.Components.Num(); ++ComponentIndex)
 		{
-			TSharedRef<FInstanceUpdateData::FComponent>& Component = OperationData->InstanceUpdateData.Components[InstanceComponentIndex];
-			const FName& ComponentName = ModelResources.ComponentNamesPerObjectComponent[Component->Id.GetValue()];
+			TSharedRef<FInstanceUpdateData::FComponent>& Component = OperationData->InstanceUpdateData.Components[ComponentIndex];
+			const FName& ComponentName = ModelResources.ComponentNamesPerObjectComponent[Component->Id];
 
 			for (int32 MutableLODIndex = OperationData->FirstLODAvailable[ComponentName]; MutableLODIndex < Component->LODCount; ++MutableLODIndex)
 			{
@@ -2178,7 +2145,7 @@ namespace impl
 				{
 					const uint32 SurfaceId = LOD->Mesh->GetSurfaceId(MeshSurfaceIndex);
 
-					const int32 InstanceSurfaceIndex = MutableInstance->FindSurfaceById(InstanceComponentIndex, MutableLODIndex, SurfaceId);
+					const int32 InstanceSurfaceIndex = MutableInstance->FindSurfaceById(ComponentIndex, MutableLODIndex, SurfaceId);
 					check(LOD->Mesh->GetVertexCount() > 0 || InstanceSurfaceIndex >= 0);
 
 					if (InstanceSurfaceIndex < 0)
@@ -2201,9 +2168,9 @@ namespace impl
 					TSharedRef<FInstanceUpdateData::FSurface>& Surface = OperationData->InstanceUpdateData.Surfaces.Add_GetRef(MakeShared<FInstanceUpdateData::FSurface>());
 
 					Surface->SurfaceId = SurfaceId;
-					Surface->SurfaceMetadataId = MutableInstance->GetSurfaceCustomId(InstanceComponentIndex, MutableLODIndex, InstanceSurfaceIndex);
+					Surface->SurfaceMetadataId = MutableInstance->GetSurfaceCustomId(ComponentIndex, MutableLODIndex, InstanceSurfaceIndex);
 
-					if (UE::Mutable::Private::FMaterialId MaterialId = MutableInstance->GetMaterialId(InstanceComponentIndex, MutableLODIndex, InstanceSurfaceIndex))
+					if (UE::Mutable::Private::FMaterialId MaterialId = MutableInstance->GetMaterialId(ComponentIndex, MutableLODIndex, InstanceSurfaceIndex))
 					{
 						Surface->MaterialIndex = OperationData->InstanceUpdateData.Materials.Num();
 						TSharedRef<FInstanceUpdateData::FMaterial>& Material = OperationData->InstanceUpdateData.Materials.Add_GetRef(MakeShared<FInstanceUpdateData::FMaterial>());
@@ -2324,25 +2291,19 @@ namespace impl
 
 		UE::Tasks::FTask LastTask;
 		
-		OperationData->InstanceUpdateData.Components.Reserve(OperationData->NumInstanceComponents);
-		for (int32 InstanceComponentIndex = 0; InstanceComponentIndex < OperationData->NumInstanceComponents; ++InstanceComponentIndex)
+		const int32 NumComponents = OperationData->MutableInstance->GetComponentCount();
+		
+		OperationData->InstanceUpdateData.Components.Reserve(NumComponents);
+		for (int32 ComponentIndex = 0; ComponentIndex < NumComponents; ++ComponentIndex)
 		{
 			TSharedRef<FInstanceUpdateData::FComponent>& Component = OperationData->InstanceUpdateData.Components.Add_GetRef(MakeShared<FInstanceUpdateData::FComponent>());
 			Component->FirstLOD = OperationData->InstanceUpdateData.LODs.Num();
-			uint64 MutableComponentId = Instance->GetComponentId(InstanceComponentIndex);
-			if (MutableComponentId < 65535)
-			{
-				Component->Id = FCustomizableObjectComponentIndex(MutableComponentId);
-			}
-			else
-			{
-				Component->Id.Invalidate();
-			}
+			UE::Mutable::Private::FComponentId MutableComponentId = Instance->GetComponentId(ComponentIndex);
+			Component->Id = MutableComponentId;
 
-			FCustomizableObjectComponentIndex ObjectComponentIndex = Component->Id;
-			const FName ComponentName = OperationData->ComponentNames[ObjectComponentIndex.GetValue()];
+			const FName ComponentName = OperationData->ComponentNames[Component->Id];
 
-			Component->SkeletalMeshId = Instance->GetSkeletalMeshId(InstanceComponentIndex);
+			Component->SkeletalMeshId = Instance->GetSkeletalMeshId(ComponentIndex);
 
 			UE::Tasks::TTask<UE::Mutable::Private::TManagedPtr<const UE::Mutable::Private::FSkeletalMesh>> GetSkeletalMeshTask = OperationData->MutableSystem->GetSkeletalMesh(LastTask, OperationData->LiveInstance.ToSharedRef(), Component->SkeletalMeshId, UE::Mutable::Private::EMeshContentFlags::AllFlags);
 			UE::Tasks::AddNested(GetSkeletalMeshTask);
@@ -2355,22 +2316,22 @@ namespace impl
 			GetSkeletalMeshTask,
 			LowLevelTasks::ETaskPriority::Inherit));
 
-			if (UE::Mutable::Private::FMaterialId MaterialId = Instance->GetOverlayMaterialId(InstanceComponentIndex))
+			if (UE::Mutable::Private::FMaterialId MaterialId = Instance->GetOverlayMaterialId(ComponentIndex))
 			{
 				Component->OverlayMaterialIndex = OperationData->InstanceUpdateData.Materials.Num();
 				TSharedRef<FInstanceUpdateData::FMaterial> Material = OperationData->InstanceUpdateData.Materials.Add_GetRef(MakeShared<FInstanceUpdateData::FMaterial>());
 				Material->MaterialId = MaterialId;
 			}
 			
-			const int32 OverrideMaterialCount = Instance->GetOverrideMaterialCount(InstanceComponentIndex);
+			const int32 OverrideMaterialCount = Instance->GetOverrideMaterialCount(ComponentIndex);
 			Component->FirstOverrideMaterial = OperationData->InstanceUpdateData.OverrideMaterials.Num();
 			for (int32 MaterialIndex = 0; MaterialIndex < OverrideMaterialCount; ++MaterialIndex)
 			{
-				if (UE::Mutable::Private::FMaterialId MaterialId = Instance->GetOverrideMaterialId(InstanceComponentIndex, MaterialIndex))
+				if (UE::Mutable::Private::FMaterialId MaterialId = Instance->GetOverrideMaterialId(ComponentIndex, MaterialIndex))
 				{
 					FInstanceUpdateData::FOverrideMaterial& OverrideMaterial = OperationData->InstanceUpdateData.OverrideMaterials.AddDefaulted_GetRef();
 					++Component->OverrideMaterialCount;
-					OverrideMaterial.SlotName = Instance->GetOverrideMaterialSlotSlotName(InstanceComponentIndex, MaterialIndex);
+					OverrideMaterial.SlotName = Instance->GetOverrideMaterialSlotSlotName(ComponentIndex, MaterialIndex);
 					
 					OverrideMaterial.MaterialIndex = OperationData->InstanceUpdateData.Materials.Num();
 					TSharedRef<FInstanceUpdateData::FMaterial> Material = OperationData->InstanceUpdateData.Materials.Add_GetRef(MakeShared<FInstanceUpdateData::FMaterial>());
@@ -2404,7 +2365,7 @@ namespace impl
 
 				TSharedRef<FInstanceUpdateData::FLOD>& LOD = OperationData->InstanceUpdateData.LODs[Component->FirstLOD + LODIndex];
 				
-				LOD->MeshId = Instance->GetMeshId(InstanceComponentIndex, LODIndex);
+				LOD->MeshId = Instance->GetMeshId(ComponentIndex, LODIndex);
 				
 				UE::Mutable::Private::EMeshContentFlags MeshContentFilter = UE::Mutable::Private::EMeshContentFlags::AllFlags; 
 				
@@ -2866,9 +2827,34 @@ namespace impl
 		// Any external texture that may be needed for this update will be requested from Mutable Core's GetImage
 		// which will safely access the GlobalExternalImages map, and then just get the cached image or issue a disk read
 
-		// Copy data generated in the mutable thread over to the instance
-		ObjectInstancePrivateData->PrepareForUpdate();
+		const UModelResources& ModelResources = CustomizableObject->GetPrivate()->GetModelResourcesChecked();
+
+		const int32 NumComponents = OperationData->MutableInstance->GetComponentCount();
+		ObjectInstancePrivateData->ComponentsData.SetNum(NumComponents);
+		
+		TArray<FName> ComponentNames;
+		ComponentNames.Reserve(NumComponents);
 
+		for (int32 ComponentIndex = 0; ComponentIndex < NumComponents; ++ComponentIndex)
+		{
+			const UE::Mutable::Private::FComponentId ComponentId = OperationData->MutableInstance->GetComponentId(ComponentIndex);
+			const FName& ComponentName = ModelResources.ComponentNamesPerObjectComponent[ComponentId];
+			
+			if (ComponentNames.Contains(ComponentName)) // Components with the same name. Abort.
+			{
+				UE_LOG(LogMutable, Error, TEXT("Update error. Trying to generate multiple components with the same name: %s"), *ComponentName.ToString());
+
+				OperationData->UpdateResult = EUpdateResult::Error;
+				FinishUpdateGlobal(OperationData);
+				return;
+			}
+			
+			ComponentNames.Add(ComponentName);
+			
+			ObjectInstancePrivateData->ComponentsData[ComponentIndex].ComponentId = ComponentId;
+			ObjectInstancePrivateData->ComponentsData[ComponentIndex].ComponentName = ComponentName;
+		}
+		
 		UE::Tasks::FTask Mutable_GetImageDescriptors_Task = SystemPrivate->MutableTaskGraph.AddMutableThreadTask(
 		TEXT("Task_Mutable_GetImageDescriptors"),
 		[OperationData]()
@@ -3208,20 +3194,18 @@ namespace impl
 				Operation->SkeletalMeshIdRegistry,
 				Operation->GetCapturedDescriptor().GetState(),
 				UE::Mutable::Private::FSystem::AllLODs);
-			
-			Operation->NumInstanceComponents = Operation->MutableInstance->GetComponentCount();
 		}
 		
 		FixLODs(Operation);
 		
-		Operation->InitComponentDescriptors(Operation->NumObjectComponents);
+		Operation->InitComponentDescriptors(Object->GetComponentCount());
 
 		const TMap<FName, uint8>& NumLODsAvailable = Operation->NumLODsAvailable;
 
-		for (int32 InstanceComponentIndex = 0; InstanceComponentIndex < Operation->NumInstanceComponents; ++InstanceComponentIndex)
+		for (int32 ComponentIndex = 0; ComponentIndex < Operation->MutableInstance->GetComponentCount(); ++ComponentIndex)
 		{
-			FCustomizableObjectComponentIndex ObjectComponentIndex(Operation->MutableInstance->GetComponentId(InstanceComponentIndex));
-			FComponentDescriptor* ComponentDescriptor = Operation->GetComponentDescriptor(ObjectComponentIndex);
+			UE::Mutable::Private::FComponentId ComponentId = Operation->MutableInstance->GetComponentId(ComponentIndex);
+			FComponentDescriptor* ComponentDescriptor = Operation->GetComponentDescriptor(ComponentId);
 			check(ComponentDescriptor);
 
 			if (!ComponentDescriptor)
@@ -3231,7 +3215,7 @@ namespace impl
 
 			ComponentDescriptor->MeshDescriptors.Init({}, MAX_MESH_LOD_COUNT);
 
-			const FName& ComponentName = Operation->ComponentNames[ObjectComponentIndex.GetValue()];
+			const FName& ComponentName = Operation->ComponentNames[ComponentId];
 
 			const int32 FirstLOD = Operation->bStreamMeshLODs ?
 				Operation->FirstLODAvailable[ComponentName] :
@@ -3239,7 +3223,7 @@ namespace impl
 
 			for (int32 LODIndex = FirstLOD; LODIndex < NumLODsAvailable[ComponentName]; ++LODIndex)
 			{
-				ComponentDescriptor->MeshDescriptors[LODIndex] = Operation->MutableInstance->GetMeshId(InstanceComponentIndex, LODIndex);
+				ComponentDescriptor->MeshDescriptors[LODIndex] = Operation->MutableInstance->GetMeshId(ComponentIndex, LODIndex);
 			}
 		}
 
@@ -3771,20 +3755,12 @@ void UCustomizableObjectSystem::DiscardInstances()
 
 							if (USkeletalMeshComponent* Parent = Cast<USkeletalMeshComponent>(InstanceUsage->GetAttachParent()))
 							{
-								const FName& ComponentName = InstanceUsage->GetComponentName();
-								const int32 ObjectComponentIndex = ModelResources->ComponentNamesPerObjectComponent.IndexOfByKey(ComponentName);
-								if (ModelResources->ReferenceSkeletalMeshesData.IsValidIndex(ObjectComponentIndex))
-								{
-									USkeletalMesh* SkeletalMesh = nullptr;
-
-									if (IsReplaceDiscardedWithReferenceMeshEnabled())
-									{
-										// Force load the reference mesh if necessary. 
-										TSoftObjectPtr<USkeletalMesh> SoftObjectPtr = ModelResources->ReferenceSkeletalMeshesData[ObjectComponentIndex].SoftSkeletalMesh;
-										SkeletalMesh = UE::Mutable::Private::LoadObject(SoftObjectPtr);
-									}
+								Parent->EmptyOverrideMaterials();
 
-									Parent->EmptyOverrideMaterials();
+								if (IsReplaceDiscardedWithReferenceMeshEnabled())
+								{
+									const FName& ComponentName = InstanceUsage->GetComponentName();
+									USkeletalMesh* SkeletalMesh = CustomizableObject->GetSkeletalMeshComponentReferenceSkeletalMesh(ComponentName);
 									Parent->SetSkeletalMesh(SkeletalMesh);
 								}
 							}
diff --git a/Engine/Plugins/Mutable/Source/CustomizableObject/Private/MuCO/CustomizableSkeletalMeshActor.cpp b/Engine/Plugins/Mutable/Source/CustomizableObject/Private/MuCO/CustomizableSkeletalMeshActor.cpp
--- a/Engine/Plugins/Mutable/Source/CustomizableObject/Private/MuCO/CustomizableSkeletalMeshActor.cpp
+++ b/Engine/Plugins/Mutable/Source/CustomizableObject/Private/MuCO/CustomizableSkeletalMeshActor.cpp
@@ -76,13 +76,13 @@ void UCustomizableSkeletalMeshActorPrivate::Init(UCustomizableObjectInstance* In
 		return;
 	}
 	
-	for (int32 ObjectComponentIndex = 0; ObjectComponentIndex < Object->GetComponentCount(); ++ObjectComponentIndex)
+	for (UE::Mutable::Private::FComponentId ComponentId = 0; ComponentId < Object->GetComponentCount(); ++ComponentId)
 	{
-		const FName& ComponentName = Object->GetComponentName(ObjectComponentIndex);
+		const FName& ComponentName = Object->GetComponentName(ComponentId);
 		const FString ComponentNameString = ComponentName.ToString();
 
 		USkeletalMeshComponent* Component;
-		if (ObjectComponentIndex == 0)
+		if (ComponentId == 0)
 		{
 			Component = GetPublic()->SkeletalMeshComponents[0];
 			// Renaming the component will make it disappear. Keep it as it is.
@@ -102,7 +102,7 @@ void UCustomizableSkeletalMeshActorPrivate::Init(UCustomizableObjectInstance* In
 		}
 
 		UCustomizableSkeletalComponent* CustomizableComponent;
-		if (ObjectComponentIndex == 0)
+		if (ComponentId == 0)
 		{
 			CustomizableComponent = GetPublic()->CustomizableSkeletalComponents[0];
 			// Renaming the component will make it disappear. Keep it as it is.
@@ -181,20 +181,20 @@ void ACustomizableSkeletalMeshActor::SwitchComponentsMaterials(UCustomizableObje
 			return;
 		}
 
-		for (int32 ObjectComponentIndex = 0; ObjectComponentIndex < SkeletalMeshComponents.Num(); ++ObjectComponentIndex)
+		for (UE::Mutable::Private::FComponentId ComponentId = 0; ComponentId < SkeletalMeshComponents.Num(); ++ComponentId)
 		{
-			int32 NumMaterials = SkeletalMeshComponents[ObjectComponentIndex]->GetNumMaterials();
+			int32 NumMaterials = SkeletalMeshComponents[ComponentId]->GetNumMaterials();
 
 			if (bDebugMaterialEnabled)
 			{
 				for (int32 MatIndex = 0; MatIndex < NumMaterials; ++MatIndex)
 				{
-					SkeletalMeshComponents[ObjectComponentIndex]->SetMaterial(MatIndex, DebugMaterial);
+					SkeletalMeshComponents[ComponentId]->SetMaterial(MatIndex, DebugMaterial);
 				}
 			}
 			else // Remove debugmaterial
 			{
-				const FName ComponentName = CustomizableSkeletalComponents[ObjectComponentIndex]->GetComponentName();
+				const FName ComponentName = CustomizableSkeletalComponents[ComponentId]->GetComponentName();
 				
 				// check if original materials already overriden
 				const TArray<UMaterialInterface*> OverrideMaterials = COInstance->GetSkeletalMeshComponentOverrideMaterials(ComponentName);
@@ -204,12 +204,12 @@ void ACustomizableSkeletalMeshActor::SwitchComponentsMaterials(UCustomizableObje
 				{
 					for (int32 MatIndex = 0; MatIndex < OverrideMaterials.Num(); ++MatIndex)
 					{
-						SkeletalMeshComponents[ObjectComponentIndex]->SetMaterial(MatIndex, OverrideMaterials[MatIndex]);
+						SkeletalMeshComponents[ComponentId]->SetMaterial(MatIndex, OverrideMaterials[MatIndex]);
 					}
 				}
 				else
 				{
-					SkeletalMeshComponents[ObjectComponentIndex]->EmptyOverrideMaterials();
+					SkeletalMeshComponents[ComponentId]->EmptyOverrideMaterials();
 				}
 
 				bRemoveDebugMaterial = false;
diff --git a/Engine/Plugins/Mutable/Source/CustomizableObject/Private/MuCO/LogBenchmarkUtil.cpp b/Engine/Plugins/Mutable/Source/CustomizableObject/Private/MuCO/LogBenchmarkUtil.cpp
--- a/Engine/Plugins/Mutable/Source/CustomizableObject/Private/MuCO/LogBenchmarkUtil.cpp
+++ b/Engine/Plugins/Mutable/Source/CustomizableObject/Private/MuCO/LogBenchmarkUtil.cpp
@@ -307,9 +307,9 @@ void FLogBenchmarkUtil::FinishUpdateMesh(const TSharedRef<FUpdateContextPrivate>
 
 	// Cache the amount of triangles of this instance
 	uint32 TriangleCount = 0;	
-	for (int32 ComponentIndex = 0; ComponentIndex < Object->GetComponentCount(); ComponentIndex++)
+	for (UE::Mutable::Private::FComponentId ComponentId = 0; ComponentId < Object->GetComponentCount(); ++ComponentId)
 	{
-		const FName ComponentName = Object->GetPrivate()->GetComponentName(FCustomizableObjectComponentIndex(ComponentIndex));
+		const FName ComponentName = Object->GetPrivate()->GetComponentName(ComponentId);
 		
 		// Process the generated components (not null)
 		if (const USkeletalMesh* InstanceSkeletalMesh = Instance->GetComponentMeshSkeletalMesh(ComponentName))
diff --git a/Engine/Plugins/Mutable/Source/CustomizableObject/Public/MuCO/CustomizableObjectInstance.h b/Engine/Plugins/Mutable/Source/CustomizableObject/Public/MuCO/CustomizableObjectInstance.h
--- a/Engine/Plugins/Mutable/Source/CustomizableObject/Public/MuCO/CustomizableObjectInstance.h
+++ b/Engine/Plugins/Mutable/Source/CustomizableObject/Public/MuCO/CustomizableObjectInstance.h
@@ -364,11 +364,7 @@ public:
 	/** Set the CustomizableObject state that this instance will be generated into. */
 	UFUNCTION(BlueprintCallable, Category = CustomizableObjectInstance)
 	CUSTOMIZABLEOBJECT_API void SetCurrentState(const FString& StateName);
-
-	// DEPRECATED: Use GetSkeletalMeshComponentSkeletalMesh 
-	UFUNCTION(BlueprintCallable, Category = CustomizableObjectInstance)
-	CUSTOMIZABLEOBJECT_API USkeletalMesh* GetSkeletalMesh(int32 ComponentIndex = 0) const;
-
+	
 	// DEPRECATED: Use GetSkeletalMeshComponentSkeletalMesh 
 	UFUNCTION(BlueprintCallable, Category = CustomizableObjectInstance)
 	CUSTOMIZABLEOBJECT_API USkeletalMesh* GetComponentMeshSkeletalMesh(const FName& ComponentName) const;
@@ -407,7 +403,7 @@ public:
 	/** Returns the AssetUserData that was gathered from all the constituent mesh parts during the last update. 
 	  * It requires that the CustomizableObject had the bEnableAssetUserDataMerge set to true during compilation. */
 	UFUNCTION(BlueprintCallable, Category = CustomizableObjectInstance)
-	CUSTOMIZABLEOBJECT_API TSet<UAssetUserData*> GetMergedAssetUserData(int32 ComponentIndex) const;
+	CUSTOMIZABLEOBJECT_API TSet<UAssetUserData*> GetMergedAssetUserData(const FName& ComponentName) const;
 	
 	/** Generate the instance with the current parameters and update all the components Skeletal Meshes asynchronously. */
 	UFUNCTION(BlueprintCallable, Category = CustomizableObjectInstance)
@@ -761,13 +757,6 @@ public:
 
 	CUSTOMIZABLEOBJECT_API void ForEachComponentAnimInstance(FName ComponentName, FEachComponentAnimInstanceClassNativeDelegate Delegate) const;
 
-	/** DEPRECATED: Use ForEachComponentAnimInstance. */
-	UFUNCTION(BlueprintCallable, Category = CustomizableObjectInstance)
-	CUSTOMIZABLEOBJECT_API void ForEachAnimInstance(int32 ObjectComponentIndex, FEachComponentAnimInstanceClassDelegate Delegate) const;
-
-	/** DEPRECATED: Use ForEachComponentAnimInstance. */
-	CUSTOMIZABLEOBJECT_API void ForEachAnimInstance(int32 ObjectComponentIndex, FEachComponentAnimInstanceClassNativeDelegate Delegate) const;
-
 	/** Check if the given UAnimInstance class requires to be fixed up. */
 	CUSTOMIZABLEOBJECT_API bool AnimInstanceNeedsFixup(TSubclassOf<UAnimInstance> AnimInstance) const;
 
diff --git a/Engine/Plugins/Mutable/Source/CustomizableObject/Public/MuCO/CustomizableObjectInstanceDescriptor.h b/Engine/Plugins/Mutable/Source/CustomizableObject/Public/MuCO/CustomizableObjectInstanceDescriptor.h
--- a/Engine/Plugins/Mutable/Source/CustomizableObject/Public/MuCO/CustomizableObjectInstanceDescriptor.h
+++ b/Engine/Plugins/Mutable/Source/CustomizableObject/Public/MuCO/CustomizableObjectInstanceDescriptor.h
@@ -20,7 +20,6 @@ class UCustomizableObjectInstance;
 class UMaterialInterface;
 class FDescriptorHash;
 class FMutableUpdateCandidate;
-struct FCustomizableObjectComponentIndex;
 
 typedef TMap<const UCustomizableObjectInstance*, FMutableUpdateCandidate> FMutableInstanceUpdateMap;
 
diff --git a/Engine/Plugins/Mutable/Source/CustomizableObject/Public/MuCO/CustomizableObjectInstanceUsage.h b/Engine/Plugins/Mutable/Source/CustomizableObject/Public/MuCO/CustomizableObjectInstanceUsage.h
--- a/Engine/Plugins/Mutable/Source/CustomizableObject/Public/MuCO/CustomizableObjectInstanceUsage.h
+++ b/Engine/Plugins/Mutable/Source/CustomizableObject/Public/MuCO/CustomizableObjectInstanceUsage.h
@@ -42,9 +42,6 @@ public:
 	UFUNCTION(BlueprintCallable, Category = CustomizableObjectInstanceUsage)
 	UE_API UCustomizableObjectInstance* GetCustomizableObjectInstance() const;
 	
-	// DEPRECATED: use SetComponentName
-	UE_API void SetComponentIndex(int32 ObjectComponentIndex);
-
 	// DEPRECATED: use GetComponentName
 	UE_API int32 GetComponentIndex() const;
 
diff --git a/Engine/Plugins/Mutable/Source/CustomizableObjectEditor/Private/MuCOE/CustomizableObjectCompiler.cpp b/Engine/Plugins/Mutable/Source/CustomizableObjectEditor/Private/MuCOE/CustomizableObjectCompiler.cpp
--- a/Engine/Plugins/Mutable/Source/CustomizableObjectEditor/Private/MuCOE/CustomizableObjectCompiler.cpp
+++ b/Engine/Plugins/Mutable/Source/CustomizableObjectEditor/Private/MuCOE/CustomizableObjectCompiler.cpp
@@ -992,20 +992,19 @@ void FCustomizableObjectCompiler::CompileInternal(bool bAsync)
 										ImageProp.AddressX, ImageProp.AddressY });
 		}
 
-		const TArray<FMutableComponentInfo>& ComponentInfos = GenerationContext.CompilationContext->ComponentInfos;
-		for (const FMutableComponentInfo& ComponentInfo : ComponentInfos)
+		const TArray<FSkeletalMeshComponentInfo>& ComponentInfos = GenerationContext.CompilationContext->ComponentInfos;
+		for (const FSkeletalMeshComponentInfo& ComponentInfo : ComponentInfos)
 		{
 			if (!Cast<UCustomizableObjectNodeComponentMesh>(ComponentInfo.Node))
 			{
 				continue;
 			}
 
-			const FName ComponentName = ComponentInfo.ComponentName;
 			const FMutableLODSettings& ComponentLODSettings = ComponentInfo.LODSettings;
 
 			// Copy the LODSettings data found in the Component into the ModelResources
-			ModelResources->MinLODPerComponent.Add(ComponentName, ComponentLODSettings.MinLOD);
-			ModelResources->MinQualityLevelLODPerComponent.Add(ComponentName, ComponentLODSettings.MinQualityLevelLOD);
+			ModelResources->MinLODPerComponent.Add(ComponentInfo.ComponentId, ComponentLODSettings.MinLOD);
+			ModelResources->MinQualityLevelLODPerComponent.Add(ComponentInfo.ComponentId, ComponentLODSettings.MinQualityLevelLOD);
 		}
 		
 		ModelResources->ParameterUIDataMap = MoveTemp(GenerationContext.ParameterUIDataMap);
@@ -1047,7 +1046,7 @@ void FCustomizableObjectCompiler::CompileInternal(bool bAsync)
 		}
 		else
 		{
-			for (TTuple<FName, uint8>& Pair : ModelResources->NumLODsToStream)
+			for (TTuple<UE::Mutable::Private::FComponentId, uint8>& Pair : ModelResources->NumLODsToStream)
 			{
 				Pair.Value = 0;
 			}
@@ -1667,6 +1666,10 @@ void FCustomizableObjectCompiler::FinishCompilationTask()
 			Offset += ResourceSize;
 		}
 	}
+	else
+	{
+		CompilerLog(LOCTEXT("NoModel", "Unable to generate Model"), nullptr, EMessageSeverity::Error);
+	}
 
 	// Second part of model resources can be filled now.
 	{
diff --git a/Engine/Plugins/Mutable/Source/CustomizableObjectEditor/Private/MuCOE/CustomizableObjectEditorViewportClient.cpp b/Engine/Plugins/Mutable/Source/CustomizableObjectEditor/Private/MuCOE/CustomizableObjectEditorViewportClient.cpp
--- a/Engine/Plugins/Mutable/Source/CustomizableObjectEditor/Private/MuCOE/CustomizableObjectEditorViewportClient.cpp
+++ b/Engine/Plugins/Mutable/Source/CustomizableObjectEditor/Private/MuCOE/CustomizableObjectEditorViewportClient.cpp
@@ -875,9 +875,9 @@ void FCustomizableObjectEditorViewportClient::OnPreSetSkeletalMesh(const FPreSet
 	}
 
 	// Add new components
-	for (int32 ObjectComponentIndex = 0; ObjectComponentIndex < CO->GetComponentCount(); ++ObjectComponentIndex)
+	for (UE::Mutable::Private::FComponentId ComponentId = 0; ComponentId < CO->GetComponentCount(); ++ComponentId)
 	{
-		FName Name = CO->GetPrivate()->GetComponentName(FCustomizableObjectComponentIndex(ObjectComponentIndex));
+		FName Name = CO->GetPrivate()->GetComponentName(ComponentId);
 		if (Instance->GetComponentMeshSkeletalMesh(Name))
 		{
 			if (!SkeletalMeshComponents.Contains(Name))
diff --git a/Engine/Plugins/Mutable/Source/CustomizableObjectEditor/Private/MuCOE/CustomizableObjectInstanceBakingUtils.cpp b/Engine/Plugins/Mutable/Source/CustomizableObjectEditor/Private/MuCOE/CustomizableObjectInstanceBakingUtils.cpp
--- a/Engine/Plugins/Mutable/Source/CustomizableObjectEditor/Private/MuCOE/CustomizableObjectInstanceBakingUtils.cpp
+++ b/Engine/Plugins/Mutable/Source/CustomizableObjectEditor/Private/MuCOE/CustomizableObjectInstanceBakingUtils.cpp
@@ -579,11 +579,11 @@ bool BakeCustomizableObjectInstance(
 	TArray<UObject*> HandledSourceObjects;
 	
 	
-	for (int32 ObjectComponentIndex = 0; ObjectComponentIndex < InstanceCO->GetComponentCount(); ++ObjectComponentIndex)
+	for (UE::Mutable::Private::FComponentId ComponentId = 0; ComponentId < InstanceCO->GetComponentCount(); ++ComponentId)
 	{
 		check(DuplicatedObjects.Num() == HandledSourceObjects.Num())
 		
-		const FName ComponentName = InstanceCO->GetPrivate()->GetComponentName(FCustomizableObjectComponentIndex(ObjectComponentIndex));
+		const FName ComponentName = InstanceCO->GetPrivate()->GetComponentName(ComponentId);
 		USkeletalMesh* Mesh = InInstance.GetComponentMeshSkeletalMesh(ComponentName);
 
 		if (!Mesh)
@@ -1068,16 +1068,15 @@ bool BakeCustomizableObjectInstance(
 
 			// Add Instance Info in a custom AssetUserData
 			{
-				const FCustomizableInstanceComponentData* ComponentData = InInstance.GetPrivate()->GetComponentData(ComponentName);
-				check(ComponentData);
+				const FCustomizableInstanceComponentData& ComponentData = InInstance.GetPrivate()->GetComponentData(ComponentId);
 				
 				if (InInstance.GetAnimationGameplayTags().Num() ||
-					ComponentData->AnimSlotToBP.Num())
+					ComponentData.AnimSlotToBP.Num())
 				{
 					UCustomizableObjectInstanceUserData* InstanceData = NewObject<UCustomizableObjectInstanceUserData>(SkeletalMesh, NAME_None, RF_Public | RF_Transactional);
 					InstanceData->AnimationGameplayTag = InInstance.GetAnimationGameplayTags();
 
-					for (const TTuple<FName, TSoftClassPtr<UAnimInstance>>& AnimSlot : ComponentData->AnimSlotToBP)
+					for (const TTuple<FName, TSoftClassPtr<UAnimInstance>>& AnimSlot : ComponentData.AnimSlotToBP)
 					{
 						FCustomizableObjectAnimationSlot AnimationSlot;
 						AnimationSlot.Name = AnimSlot.Key;
@@ -1092,9 +1091,9 @@ bool BakeCustomizableObjectInstance(
 
 			// Copy LODSettings from the Reference Skeletal Mesh
 			{
-				if (ModelResources->ReferenceSkeletalMeshesData.IsValidIndex(ObjectComponentIndex))
+				if (ModelResources->ReferenceSkeletalMeshesData.IsValidIndex(ComponentId))
 				{
-					USkeletalMeshLODSettings* LODSettings = ModelResources->ReferenceSkeletalMeshesData[ObjectComponentIndex].SkeletalMeshLODSettings;
+					USkeletalMeshLODSettings* LODSettings = ModelResources->ReferenceSkeletalMeshesData[ComponentId].SkeletalMeshLODSettings;
 					SkeletalMesh->SetLODSettings(LODSettings);
 				}
 			}
diff --git a/Engine/Plugins/Mutable/Source/CustomizableObjectEditor/Private/MuCOE/GenerateMutableSource/GenerateMutableSource.cpp b/Engine/Plugins/Mutable/Source/CustomizableObjectEditor/Private/MuCOE/GenerateMutableSource/GenerateMutableSource.cpp
--- a/Engine/Plugins/Mutable/Source/CustomizableObjectEditor/Private/MuCOE/GenerateMutableSource/GenerateMutableSource.cpp
+++ b/Engine/Plugins/Mutable/Source/CustomizableObjectEditor/Private/MuCOE/GenerateMutableSource/GenerateMutableSource.cpp
@@ -222,12 +222,21 @@ FMutableCompilationContext::FMutableCompilationContext(const UCustomizableObject
 }
 
 
-FMutableComponentInfo* FMutableCompilationContext::GetComponentInfo(FName ComponentName)
+FSkeletalMeshComponentInfo* FMutableCompilationContext::GetComponentInfo(const UCustomizableObjectNodeComponentMesh* NodeComponentMesh)
 {
-	return ComponentInfos.FindByPredicate([ComponentName](const FMutableComponentInfo& Component)
-		{
-			return Component.ComponentName == ComponentName;
-		});
+	return ComponentInfos.FindByPredicate([&](const FSkeletalMeshComponentInfo& ComponentInfo)
+	{
+		return ComponentInfo.Node == NodeComponentMesh;
+	});
+}
+
+
+FSkeletalMeshComponentInfo* FMutableCompilationContext::GetComponentInfo(UE::Mutable::Private::FComponentId ComponentId)
+{
+	return ComponentInfos.FindByPredicate([&](const FSkeletalMeshComponentInfo& ComponentInfo)
+	{
+		return ComponentInfo.ComponentId == ComponentId;
+	});
 }
 
 
@@ -373,9 +382,12 @@ const FGuid FMutableGraphGenerationContext::GetNodeIdUnchecked(const UCustomizab
 }
 
 
-FMutableComponentInfo* FMutableGraphGenerationContext::GetCurrentComponentInfo()
+FSkeletalMeshComponentInfo* FMutableGraphGenerationContext::GetCurrentComponentInfo()
 {
-	return CompilationContext->GetComponentInfo(CurrentMeshComponent);
+	return CompilationContext->ComponentInfos.FindByPredicate([&](const FSkeletalMeshComponentInfo& ComponentInfo)
+	{
+		return ComponentInfo.ComponentId == CurrentSkeletalMeshComponent;
+	});
 }
 
 
@@ -967,7 +979,7 @@ UE::Mutable::Private::Ptr<UE::Mutable::Private::NodeObject> GenerateMutableSourc
 			int32 MinQualityLevelLOD = INT32_MAX;
 			int32 MaxQualityLevelLOD = 0;
 
-			for (const FMutableComponentInfo& ComponentInfo : GenerationContext.CompilationContext->ComponentInfos)
+			for (const FSkeletalMeshComponentInfo& ComponentInfo : GenerationContext.CompilationContext->ComponentInfos)
 			{
 				if (!Cast<UCustomizableObjectNodeComponentMesh>(ComponentInfo.Node))
 				{
@@ -1025,15 +1037,15 @@ UE::Mutable::Private::Ptr<UE::Mutable::Private::NodeObject> GenerateMutableSourc
 			TArray<UEdGraphPin*> ConnectedModifierPins = FollowInputPinArray(*ModifierPin);
 			for (UEdGraphPin* const ConnectedModifier : ConnectedModifierPins)
 			{
-				for (const FMutableComponentInfo& Pair : GenerationContext.CompilationContext->ComponentInfos)
+				for (const FSkeletalMeshComponentInfo& Pair : GenerationContext.CompilationContext->ComponentInfos)
 				{
-					GenerationContext.CurrentMeshComponent = Pair.ComponentName;
+					GenerationContext.CurrentSkeletalMeshComponent = Pair.ComponentId;
 					GenerationContext.CurrentAutoLODStrategy = Pair.Node->AutoLODStrategy;
 					
 					UE::Mutable::Private::Ptr<UE::Mutable::Private::NodeModifier> ModifierNode = GenerateMutableSourceModifier(ConnectedModifier, GenerationContext);
 					ObjectNode->Modifiers.Add(ModifierNode);
 
-					GenerationContext.CurrentMeshComponent = {};
+					GenerationContext.CurrentSkeletalMeshComponent = INDEX_NONE;
 					GenerationContext.CurrentAutoLODStrategy = {};
 				}
 			}
@@ -1468,8 +1480,7 @@ void PopulateReferenceSkeletalMeshesData(FMutableGraphGenerationContext& Generat
 	GenerationContext.ReferenceSkeletalMeshesData.AddDefaulted(ComponentCount);
 	for (uint32 ComponentIndex = 0; ComponentIndex < ComponentCount; ++ComponentIndex)
 	{
-		const FMutableComponentInfo& ComponentInfo = GenerationContext.CompilationContext->ComponentInfos[ComponentIndex];
-		const FName& ComponentName = ComponentInfo.ComponentName;
+		const FSkeletalMeshComponentInfo& ComponentInfo = GenerationContext.CompilationContext->ComponentInfos[ComponentIndex];
 		
 		USkeletalMesh* RefSkeletalMesh = ComponentInfo.RefSkeletalMesh.Get();
 		check(RefSkeletalMesh);
@@ -1489,11 +1500,11 @@ void PopulateReferenceSkeletalMeshesData(FMutableGraphGenerationContext& Generat
 		Data.SkeletalMeshLODSettings = RefSkeletalMesh->GetLODSettings();
 
 		// Gather LODData, this may include per LOD settings such as render data config or LODDataInfoArray
-		Data.LODData.AddDefaulted(GenerationContext.NumLODs[ComponentName]);
+		Data.LODData.AddDefaulted(GenerationContext.NumLODs[ComponentInfo.ComponentId]);
 		
 		const uint32 RefSkeletalMeshLODCount = RefSkeletalMesh->GetLODNum();
 
-		const uint32 LODCount = GenerationContext.NumLODs[ComponentName];
+		const uint32 LODCount = GenerationContext.NumLODs[ComponentInfo.ComponentId];
 		for (uint32 LODIndex = 0; LODIndex < LODCount; ++LODIndex)
 		{
 			FMutableRefLODData& LODData = Data.LODData[LODIndex];
@@ -1845,76 +1856,6 @@ UE::Mutable::Private::TManagedPtr<UE::Mutable::Private::FMesh> GenerateMeshConst
 	return Result;
 }
 
-
-FMutableComponentInfo::FMutableComponentInfo(FName InComponentName, USkeletalMesh* InRefSkeletalMesh)
-{
-	ComponentName = InComponentName;
-
-	if (!InRefSkeletalMesh || !InRefSkeletalMesh->GetSkeleton())
-	{
-		return;
-	}
-
-	RefSkeletalMesh = TStrongObjectPtr<USkeletalMesh>(InRefSkeletalMesh);
-	RefSkeleton = RefSkeletalMesh->GetSkeleton();
-
-	const int32 NumBones = RefSkeleton->GetReferenceSkeleton().GetRawBoneNum();
-	BoneNamesToPathHash.Reserve(NumBones);
-
-	const TArray<FMeshBoneInfo>& Bones = RefSkeleton->GetReferenceSkeleton().GetRawRefBoneInfo();
-
-	for (int32 BoneIndex = 0; BoneIndex < NumBones; ++BoneIndex)
-	{
-		const FMeshBoneInfo& Bone = Bones[BoneIndex];
-
-		// Retrieve parent bone name and respective hash, root-bone is assumed to have a parent hash of 0
-		const FName ParentName = Bone.ParentIndex != INDEX_NONE ? Bones[Bone.ParentIndex].Name : NAME_None;
-		const uint32 ParentHash = Bone.ParentIndex != INDEX_NONE ? GetTypeHash(ParentName) : 0;
-
-		// Look-up the path-hash from root to the parent bone
-		const uint32* ParentPath = BoneNamesToPathHash.Find(ParentName);
-		const uint32 ParentPathHash = ParentPath ? *ParentPath : 0;
-
-		// Append parent hash to path to give full path hash to current bone
-		const uint32 BonePathHash = HashCombine(ParentPathHash, ParentHash);
-
-		// Add path hash to current bone
-		BoneNamesToPathHash.Add(Bone.Name, BonePathHash);
-	}
-}
-
-
-void FMutableComponentInfo::AccumulateBonesToRemovePerLOD(const TArray<FLODReductionSettings>& LODReductionSettings, int32 NumLODs)
-{
-	BonesToRemovePerLOD.SetNum(NumLODs);
-
-	TMap<FName, bool> BonesToRemove;
-
-	const int32 ComponentSettingsLODCount = LODReductionSettings.Num();
-	for (int32 LODIndex = 0; LODIndex < NumLODs; ++LODIndex)
-	{
-		if (LODIndex < ComponentSettingsLODCount)
-		{
-			const FLODReductionSettings& LODReductionSetting = LODReductionSettings[LODIndex];
-
-			for (const FBoneToRemove& Bone : LODReductionSetting.BonesToRemove)
-			{
-				if (bool* bOnlyRemoveChildren = BonesToRemove.Find(Bone.BoneName))
-				{
-					// Removed by a previous LOD
-					*bOnlyRemoveChildren = (*bOnlyRemoveChildren) && Bone.bOnlyRemoveChildren;
-				}
-				else
-				{
-					BonesToRemove.Add(Bone.BoneName, Bone.bOnlyRemoveChildren);
-				}
-			}
-		}
-		
-		BonesToRemovePerLOD[LODIndex] = BonesToRemove;
-	}
-}
-
 #undef LOCTEXT_NAMESPACE
 
 
diff --git a/Engine/Plugins/Mutable/Source/CustomizableObjectEditor/Private/MuCOE/GenerateMutableSource/GenerateMutableSource.h b/Engine/Plugins/Mutable/Source/CustomizableObjectEditor/Private/MuCOE/GenerateMutableSource/GenerateMutableSource.h
--- a/Engine/Plugins/Mutable/Source/CustomizableObjectEditor/Private/MuCOE/GenerateMutableSource/GenerateMutableSource.h
+++ b/Engine/Plugins/Mutable/Source/CustomizableObjectEditor/Private/MuCOE/GenerateMutableSource/GenerateMutableSource.h
@@ -184,7 +184,7 @@ public:
 	/** UV Layout modes */
 	FLayoutGenerationFlags LayoutFlags;
 	
-	FName CurrentMeshComponent;
+	UE::Mutable::Private::FComponentId CurrentMeshComponent = INDEX_NONE;
 	
 	/** When caching a generated mesh, true if we force to generate the connected LOD when using Automatic LODs From Mesh. */
 	bool bOnlyConnectedLOD = false;
@@ -432,22 +432,18 @@ struct FGroupNodeIdsTempData
 
 
 /** Struct used to store info specific to each component during compilation */
-struct FMutableComponentInfo
+struct FSkeletalMeshComponentInfo
 {
-	FMutableComponentInfo(FName InComponentName, USkeletalMesh* InRefSkeletalMesh);
-
-	void AccumulateBonesToRemovePerLOD(const TArray<FLODReductionSettings>& LODReductionSettings, int32 NumLODs);
-
 	FName ComponentName;
-
+	
+	UE::Mutable::Private::FComponentId ComponentId = INDEX_NONE;
+	
 	FMutableLODSettings LODSettings;
 	
 	// Each component must have a reference SkeletalMesh with a valid Skeleton
 	TStrongObjectPtr<USkeletalMesh> RefSkeletalMesh;
 	TObjectPtr<USkeleton> RefSkeleton;
 	
-	UCustomizableObjectNodeComponentMesh* NodeComponentMesh = nullptr;
-	
 	// Map to check skeleton compatibility
 	TMap<SIZE_T, bool> SkeletonCompatibility;
 	
@@ -521,7 +517,7 @@ private:
 struct FGeneratedGroupProjectorsKey
 {
 	UCustomizableObjectNodeGroupProjectorParameter* Node = nullptr;
-	FName CurrentComponent;
+	UE::Mutable::Private::FComponentId CurrentComponent;
 
 	bool operator==(const FGeneratedGroupProjectorsKey&) const = default;
 };
@@ -578,7 +574,7 @@ struct FMutableCompilationContext
 	TMap<uint32, FMutableMeshMetadata> MeshMetadata;
 
 	/** Only Mesh Components (no passthrough). */
-	TArray<FMutableComponentInfo> ComponentInfos;
+	TArray<FSkeletalMeshComponentInfo> ComponentInfos;
 
 	/** Array of unique Bone identifiers. */
 	TMap<UE::Mutable::Private::FBoneName, FString> UniqueBoneNames;
@@ -609,7 +605,9 @@ public:
 	void Log(const FText& Message, const UObject* Context = nullptr, const EMessageSeverity::Type MessageSeverity = EMessageSeverity::Warning, const bool bAddBaseObjectInfo = true, const ELoggerSpamBin SpamBin = ELoggerSpamBin::ShowAll) const;
 
 	/** Component access (for Mesh components only). */
-	FMutableComponentInfo* GetComponentInfo(FName ComponentName);
+	FSkeletalMeshComponentInfo* GetComponentInfo(const UCustomizableObjectNodeComponentMesh* NodeComponentMesh);
+
+	FSkeletalMeshComponentInfo* GetComponentInfo(UE::Mutable::Private::FComponentId ComponentName);
 
 	/** Get unique identifier for BoneName built from its FString. */
 	UE::Mutable::Private::FBoneName GetBoneUnique(const FName& BoneName);
@@ -849,7 +847,7 @@ struct FMutableGraphGenerationContext
 	const FGuid GetNodeIdUnchecked(const UCustomizableObjectNode* Node);
 
 	/** Get the reference skeletal mesh associated to the current mesh component being generated */
-	FMutableComponentInfo* GetCurrentComponentInfo();
+	FSkeletalMeshComponentInfo* GetCurrentComponentInfo();
 	
 	UObject* LoadObject(const FSoftObjectPtr& SoftObject);
 	
@@ -871,7 +869,7 @@ struct FMutableGraphGenerationContext
 	 */
 	FMaterialBreakParameter GetCurrentMaterialBreakParameter();
 	
-	/** Only compiled components. All component types. Index is the ObjectComponentIndex. */
+	/** Only compiled components. All component types. Index is the ComponentId. */
 	TArray<FName> ComponentNames;
 	
 	TArray<FMutableRefSkeletalMeshData> ReferenceSkeletalMeshesData;
@@ -892,16 +890,16 @@ struct FMutableGraphGenerationContext
 	
 	uint8 FromLOD = 0; // LOD to append to the CurrentLOD when using AutomaticLODs. 
 	uint8 CurrentLOD = 0;
-	FName CurrentMeshComponent;
+	UE::Mutable::Private::FComponentId CurrentSkeletalMeshComponent = INDEX_NONE;
 
 	/** If this is set, we are generating materials for a "passthrough" component, with a fixed mesh. */
 	UE::Mutable::Private::Ptr<UE::Mutable::Private::NodeMesh> ComponentMeshOverride;
 
-	TMap<FName, uint8> NumLODs;
+	TMap<UE::Mutable::Private::FComponentId, uint8> NumLODs;
 
-	TMap<FName, uint8> FirstLODAvailable;
+	TMap<UE::Mutable::Private::FComponentId, uint8> FirstLODAvailable;
 	
-	TMap<FName, uint8> NumMaxLODsToStream;
+	TMap<UE::Mutable::Private::FComponentId, uint8> NumMaxLODsToStream;
 
 	bool bEnableLODStreaming = true;
 
diff --git a/Engine/Plugins/Mutable/Source/CustomizableObjectEditor/Private/MuCOE/GenerateMutableSource/GenerateMutableSourceComponent.cpp b/Engine/Plugins/Mutable/Source/CustomizableObjectEditor/Private/MuCOE/GenerateMutableSource/GenerateMutableSourceComponent.cpp
--- a/Engine/Plugins/Mutable/Source/CustomizableObjectEditor/Private/MuCOE/GenerateMutableSource/GenerateMutableSourceComponent.cpp
+++ b/Engine/Plugins/Mutable/Source/CustomizableObjectEditor/Private/MuCOE/GenerateMutableSource/GenerateMutableSourceComponent.cpp
@@ -61,55 +61,48 @@ UE::Mutable::Private::Ptr<UE::Mutable::Private::NodeComponent> GenerateMutableSo
 	
 	if (const UCustomizableObjectNodeComponentMesh* TypedComponentMesh = Cast<UCustomizableObjectNodeComponentMesh>(Node))
 	{
-		if (!GenerationContext.CompilationContext->ComponentInfos.ContainsByPredicate([&](const FMutableComponentInfo& ComponentInfo)
+		FSkeletalMeshComponentInfo* ComponentInfo = GenerationContext.CompilationContext->GetComponentInfo(TypedComponentMesh);
+		if (ComponentInfo)
 		{
-			return ComponentInfo.Node == TypedComponentMesh;
-		}))
-		{
-			return nullptr; // Not generated in the first pass.
-		}
+			UE::Mutable::Private::Ptr<UE::Mutable::Private::NodeComponentNew> NodeComponentNew = new UE::Mutable::Private::NodeComponentNew();
+			NodeComponentNew->Id = ComponentInfo->ComponentId;
+			NodeComponentNew->SetMessageContext(Node);
 		
-		UE::Mutable::Private::Ptr<UE::Mutable::Private::NodeComponentNew> NodeComponentNew = new UE::Mutable::Private::NodeComponentNew();
-		const FName& ComponentName = TypedComponentMesh->GetComponentName(&GenerationContext.MacroNodesStack);
-
-		NodeComponentNew->Id = GenerationContext.ComponentNames.Find(ComponentName);
-		NodeComponentNew->SetMessageContext(Node);
-		
-		UEdGraphPin* MaterialAssetPin = TypedComponentMesh->GetOverlayMaterialAssetPin();
-		if (const UEdGraphPin* ConnectedPin = MaterialAssetPin ? FollowInputPin(*MaterialAssetPin) : nullptr)
-		{
-			GenerationContext.CurrentMaterialParameterId = ConnectedPin->PinId.ToString();
-			UE::Mutable::Private::Ptr<UE::Mutable::Private::NodeMaterial> MaterialNode = GenerateMutableSourceMaterial(ConnectedPin, GenerationContext, FGenerationMaterialOptions());
-			NodeComponentNew->OverlayMaterial = MaterialNode;
-		}
-		else if(UMaterialInterface* OverlayMaterial = GenerationContext.LoadObject(TypedComponentMesh->GetOverlayMaterial()))
-		{
-			UE::Mutable::Private::Ptr<UE::Mutable::Private::NodeMaterialConstant> MaterialNode = new UE::Mutable::Private::NodeMaterialConstant();
-			MaterialNode->MaterialId = GenerationContext.CompilationContext->PassthroughObjectFactory.Add(*OverlayMaterial);
-			NodeComponentNew->OverlayMaterial = MaterialNode;
-		}
+			UEdGraphPin* MaterialAssetPin = TypedComponentMesh->GetOverlayMaterialAssetPin();
+			if (const UEdGraphPin* ConnectedPin = MaterialAssetPin ? FollowInputPin(*MaterialAssetPin) : nullptr)
+			{
+				GenerationContext.CurrentMaterialParameterId = ConnectedPin->PinId.ToString();
+				UE::Mutable::Private::Ptr<UE::Mutable::Private::NodeMaterial> MaterialNode = GenerateMutableSourceMaterial(ConnectedPin, GenerationContext, FGenerationMaterialOptions());
+				NodeComponentNew->OverlayMaterial = MaterialNode;
+			}
+			else if(UMaterialInterface* OverlayMaterial = GenerationContext.LoadObject(TypedComponentMesh->GetOverlayMaterial()))
+			{
+				UE::Mutable::Private::Ptr<UE::Mutable::Private::NodeMaterialConstant> MaterialNode = new UE::Mutable::Private::NodeMaterialConstant();
+				MaterialNode->MaterialId = GenerationContext.CompilationContext->PassthroughObjectFactory.Add(*OverlayMaterial);
+				NodeComponentNew->OverlayMaterial = MaterialNode;
+			}
 
-		Result = NodeComponentNew;
+			Result = NodeComponentNew;
 
-		GenerationContext.CurrentMeshComponent = ComponentName;
-		GenerationContext.CurrentAutoLODStrategy = TypedComponentMesh->AutoLODStrategy;
+			GenerationContext.CurrentSkeletalMeshComponent = ComponentInfo->ComponentId;
+			GenerationContext.CurrentAutoLODStrategy = TypedComponentMesh->AutoLODStrategy;
 		
-		NodeComponentNew->SkeletalMesh = UE::Mutable::Private::GenerateMutableSourceSkeletalMesh(Pin, GenerationContext);
-
-		GenerationContext.CurrentMeshComponent = {};
-		GenerationContext.CurrentAutoLODStrategy = {};
+			NodeComponentNew->SkeletalMesh = UE::Mutable::Private::GenerateMutableSourceSkeletalMesh(Pin, GenerationContext);
+		
+			GenerationContext.CurrentSkeletalMeshComponent = INDEX_NONE;
+			GenerationContext.CurrentAutoLODStrategy = {};
+		}
 	}
 
 	else if (const UCustomizableObjectNodeComponentMeshAddTo* TypedComponentMeshExtend = Cast<UCustomizableObjectNodeComponentMeshAddTo>(Node))
 	{
 		const FName& ParentComponentName = TypedComponentMeshExtend->GetParentComponentName(&GenerationContext.MacroNodesStack);
-
-		if (FMutableComponentInfo* FindResult = GenerationContext.CompilationContext->ComponentInfos.FindByPredicate([&](const FMutableComponentInfo& Element)
-		{
-			return Element.ComponentName == ParentComponentName;
-		}))
+		const UE::Mutable::Private::FComponentId ComponentId = GenerationContext.ComponentNames.IndexOfByKey(ParentComponentName);
+		
+		FSkeletalMeshComponentInfo* ComponentInfo = GenerationContext.CompilationContext->GetComponentInfo(ComponentId);
+		if (ComponentInfo)
 		{
-			UCustomizableObjectNodeComponentMesh* TypedParentComponentMesh = FindResult->Node;
+			UCustomizableObjectNodeComponentMesh* TypedParentComponentMesh = ComponentInfo->Node;
 
 			if (TypedComponentMeshExtend->NumLODs > TypedParentComponentMesh->NumLODs)
 			{
@@ -119,12 +112,12 @@ UE::Mutable::Private::Ptr<UE::Mutable::Private::NodeComponent> GenerateMutableSo
 
 			// Swap the macro contex since this component can be in another macro
 			TArray<const UCustomizableObjectNodeMacroInstance*> MacroContextCopy = GenerationContext.MacroNodesStack;
-			GenerationContext.MacroNodesStack = FindResult->MacroContext;
+			GenerationContext.MacroNodesStack = ComponentInfo->MacroContext;
 
 			UE::Mutable::Private::Ptr<UE::Mutable::Private::NodeComponent> ParentNodeComponent = GenerateMutableSourceComponent(TypedParentComponentMesh->OutputPin.Get(), GenerationContext);
 
 			// Get the parent name using the parent's macro context
-			GenerationContext.CurrentMeshComponent = TypedParentComponentMesh->GetComponentName(&GenerationContext.MacroNodesStack);
+			GenerationContext.CurrentSkeletalMeshComponent = ComponentInfo->ComponentId;
 
 			// Restore GenerationContext MacroContex
 			GenerationContext.MacroNodesStack = MacroContextCopy;
@@ -139,7 +132,7 @@ UE::Mutable::Private::Ptr<UE::Mutable::Private::NodeComponent> GenerateMutableSo
 			
 			NodeComponentEdit->SkeletalMesh = UE::Mutable::Private::GenerateMutableSourceSkeletalMeshNew(Pin, GenerationContext);
 
-			GenerationContext.CurrentMeshComponent = {};
+			GenerationContext.CurrentSkeletalMeshComponent = INDEX_NONE;
 			GenerationContext.CurrentAutoLODStrategy = {};
 			
 			Result = NodeComponentEdit;	
@@ -153,23 +146,9 @@ UE::Mutable::Private::Ptr<UE::Mutable::Private::NodeComponent> GenerateMutableSo
 	
 	else if (const UCustomizableObjectNodeComponentPassthroughMesh* TypedComponentPassthroughMesh = Cast<UCustomizableObjectNodeComponentPassthroughMesh>(Node))
 	{
-		const FName& ComponentName = TypedComponentPassthroughMesh->GetComponentName(&GenerationContext.MacroNodesStack);
-		GenerationContext.CurrentMeshComponent = ComponentName;
-
-		if (ComponentName.IsNone())
-		{
-			FString Msg = FString::Printf(TEXT("Invalid Component Name."));
-			GenerationContext.Log(FText::FromString(Msg), TypedComponentPassthroughMesh, EMessageSeverity::Warning);
-			return nullptr;
-		}
-		
-		// Create the component node
 		UE::Mutable::Private::Ptr<UE::Mutable::Private::NodeComponentNew> ComponentNode = new UE::Mutable::Private::NodeComponentNew;
-		if (GenerationContext.ComponentNames.Contains(ComponentName))
-		{
-			FString Msg = FString::Printf(TEXT("More than one component with the same name [%s] have been found. This is not supported."), *ComponentName.ToString());
-			GenerationContext.Log(FText::FromString(Msg), TypedComponentPassthroughMesh, EMessageSeverity::Warning);
-		}
+
+		const FName& ComponentName = TypedComponentPassthroughMesh->GetComponentName(&GenerationContext.MacroNodesStack);
 		ComponentNode->Id = GenerationContext.ComponentNames.Add(ComponentName);
 		ComponentNode->SkeletalMesh = UE::Mutable::Private::GenerateMutableSourceSkeletalMesh(Pin, GenerationContext);
 
@@ -202,7 +181,6 @@ UE::Mutable::Private::Ptr<UE::Mutable::Private::NodeComponent> GenerateMutableSo
 			}
 		}
 		
-		GenerationContext.CurrentMeshComponent = FName();
 		Result = ComponentNode;
 	}
 	
@@ -359,36 +337,6 @@ UE::Mutable::Private::Ptr<UE::Mutable::Private::NodeComponent> GenerateMutableSo
 
 void FirstPass(UCustomizableObjectNodeComponentMesh& Node, FMutableGraphGenerationContext& GenerationContext)
 {
-	FName ComponentName = Node.GetComponentName(&GenerationContext.MacroNodesStack);
-
-	if (ComponentName.IsNone())
-	{
-		GenerationContext.Log(LOCTEXT("EmptyComponentNameError", "Error! Missing name in a component of the Customizable Object."), &Node, EMessageSeverity::Error);
-		return;
-	}
-	
-	if (FMutableComponentInfo* Result = GenerationContext.CompilationContext->ComponentInfos.FindByPredicate([&](const FMutableComponentInfo& Element)
-	{
-		return Element.ComponentName == ComponentName;
-	}))
-	{
-		FText Msg;
-		if (Result->Node->IsInMacro())
-		{
-			const UCustomizableObjectNodeMacroInstance* CurrentMacro = GenerationContext.MacroNodesStack.Last();
-			check(CurrentMacro);
-			
-			Msg = FText::Format(LOCTEXT("ComponentNodeWithSameNameExists_Macro", "Error! A Mesh Component node with the same name already exists in a Macro Instance [{0}]"), FText::FromName(CurrentMacro->ParentMacro->Name));
-		}
-		else
-		{
-			Msg = FText::Format(LOCTEXT("ComponentNodeWithSameNameExists_CO", "Error! A Mesh Component node with the same name already exists in the Customizable Object [{0}]"), FText::FromString(GraphTraversal::GetObject(*Result->Node)->GetName()));
-		}
-		
-		GenerationContext.Log(Msg, &Node, EMessageSeverity::Error);
-		return;
-	}
-	
 	USkeletalMesh* RefSkeletalMesh = Node.ReferenceSkeletalMesh;
 	if (!RefSkeletalMesh)
 	{
@@ -486,23 +434,83 @@ void FirstPass(UCustomizableObjectNodeComponentMesh& Node, FMutableGraphGenerati
 
 	NumMaxLODsToStream = FMath::Clamp(NumMaxLODsToStream, 0, NumLODs - 1);
 
-	GenerationContext.NumLODs.Add(ComponentName, NumLODs);
-	GenerationContext.FirstLODAvailable.Add(ComponentName, FirstLODAvailable);
-	GenerationContext.NumMaxLODsToStream.Add(ComponentName, NumMaxLODsToStream);
+	const FName ComponentName = Node.GetComponentName(&GenerationContext.MacroNodesStack);
+	UE::Mutable::Private::FComponentId ComponentId = GenerationContext.ComponentNames.Add(ComponentName);
+
+	GenerationContext.NumLODs.Add(ComponentId, NumLODs);
+	GenerationContext.FirstLODAvailable.Add(ComponentId, FirstLODAvailable);
+	GenerationContext.NumMaxLODsToStream.Add(ComponentId, NumMaxLODsToStream);
 	
 	// Add a new entry to the list of Component Infos
-	FMutableComponentInfo ComponentInfo(ComponentName, RefSkeletalMesh);
+	FSkeletalMeshComponentInfo ComponentInfo;
+	ComponentInfo.ComponentName = ComponentName;
+	ComponentInfo.ComponentId = ComponentId;
 	ComponentInfo.Node = &Node;
-	ComponentInfo.AccumulateBonesToRemovePerLOD(Node.LODReductionSettings, Node.NumLODs);
 	ComponentInfo.MacroContext = GenerationContext.MacroNodesStack;
 	ComponentInfo.LODSettings = LODSettings;
+	
+	if (RefSkeletalMesh && RefSkeleton)
+	{
+		ComponentInfo.RefSkeletalMesh = TStrongObjectPtr(RefSkeletalMesh);
+		ComponentInfo.RefSkeleton = RefSkeleton;
+
+		const int32 NumBones = RefSkeleton->GetReferenceSkeleton().GetRawBoneNum();
+		ComponentInfo.BoneNamesToPathHash.Reserve(NumBones);
+
+		const TArray<FMeshBoneInfo>& Bones = RefSkeleton->GetReferenceSkeleton().GetRawRefBoneInfo();
+
+		for (int32 BoneIndex = 0; BoneIndex < NumBones; ++BoneIndex)
+		{
+			const FMeshBoneInfo& Bone = Bones[BoneIndex];
+
+			// Retrieve parent bone name and respective hash, root-bone is assumed to have a parent hash of 0
+			const FName ParentName = Bone.ParentIndex != INDEX_NONE ? Bones[Bone.ParentIndex].Name : NAME_None;
+			const uint32 ParentHash = Bone.ParentIndex != INDEX_NONE ? GetTypeHash(ParentName) : 0;
+
+			// Look-up the path-hash from root to the parent bone
+			const uint32* ParentPath = ComponentInfo.BoneNamesToPathHash.Find(ParentName);
+			const uint32 ParentPathHash = ParentPath ? *ParentPath : 0;
+
+			// Append parent hash to path to give full path hash to current bone
+			const uint32 BonePathHash = HashCombine(ParentPathHash, ParentHash);
+
+			// Add path hash to current bone
+			ComponentInfo.BoneNamesToPathHash.Add(Bone.Name, BonePathHash);
+		}
+	}
+
+	ComponentInfo.BonesToRemovePerLOD.SetNum(NumLODs);
+
+	TMap<FName, bool> BonesToRemove;
+
+	const int32 ComponentSettingsLODCount = Node.LODReductionSettings.Num();
+	for (int32 LODIndex = 0; LODIndex < NumLODs; ++LODIndex)
+	{
+		if (LODIndex < ComponentSettingsLODCount)
+		{
+			const FLODReductionSettings& LODReductionSetting = Node.LODReductionSettings[LODIndex];
+
+			for (const FBoneToRemove& Bone : LODReductionSetting.BonesToRemove)
+			{
+				if (bool* bOnlyRemoveChildren = BonesToRemove.Find(Bone.BoneName))
+				{
+					// Removed by a previous LOD
+					*bOnlyRemoveChildren = (*bOnlyRemoveChildren) && Bone.bOnlyRemoveChildren;
+				}
+				else
+				{
+					BonesToRemove.Add(Bone.BoneName, Bone.bOnlyRemoveChildren);
+				}
+			}
+		}
+		
+		ComponentInfo.BonesToRemovePerLOD[LODIndex] = BonesToRemove;
+	}
 
 	GenerationContext.CompilationContext->ComponentInfos.Add(ComponentInfo);
 
 	// Make sure the Skeleton from the reference mesh is added to the list of referenced Skeletons.
 	GenerationContext.CompilationContext->ReferencedSkeletons.Add(RefSkeleton);
-
-	GenerationContext.ComponentNames.Add(ComponentName);
 }
 
 
diff --git a/Engine/Plugins/Mutable/Source/CustomizableObjectEditor/Private/MuCOE/GenerateMutableSource/GenerateMutableSourceGroupProjector.cpp b/Engine/Plugins/Mutable/Source/CustomizableObjectEditor/Private/MuCOE/GenerateMutableSource/GenerateMutableSourceGroupProjector.cpp
--- a/Engine/Plugins/Mutable/Source/CustomizableObjectEditor/Private/MuCOE/GenerateMutableSource/GenerateMutableSourceGroupProjector.cpp
+++ b/Engine/Plugins/Mutable/Source/CustomizableObjectEditor/Private/MuCOE/GenerateMutableSource/GenerateMutableSourceGroupProjector.cpp
@@ -299,7 +299,7 @@ TOptional<FGroupProjectorTempData> GenerateMutableGroupProjector(UCustomizableOb
 {
 	FGeneratedGroupProjectorsKey Key;
 	Key.Node = ProjParamNode;
-	Key.CurrentComponent = GenerationContext.CurrentMeshComponent;
+	Key.CurrentComponent = GenerationContext.CurrentSkeletalMeshComponent;
 	
 	if (FGroupProjectorTempData* Result = GenerationContext.GeneratedGroupProjectors.Find(Key))
 	{
@@ -429,7 +429,7 @@ TOptional<FGroupProjectorTempData> GenerateMutableGroupProjector(UCustomizableOb
 			return {};
 		}
 
-		FMutableComponentInfo* CurrentComponentInfo = GenerationContext.GetCurrentComponentInfo();
+		FSkeletalMeshComponentInfo* CurrentComponentInfo = GenerationContext.GetCurrentComponentInfo();
 		check(CurrentComponentInfo);
 		
 		PoseEnumParameterNode->Options[0].Value = 0.f;
diff --git a/Engine/Plugins/Mutable/Source/CustomizableObjectEditor/Private/MuCOE/GenerateMutableSource/GenerateMutableSourceMesh.cpp b/Engine/Plugins/Mutable/Source/CustomizableObjectEditor/Private/MuCOE/GenerateMutableSource/GenerateMutableSourceMesh.cpp
--- a/Engine/Plugins/Mutable/Source/CustomizableObjectEditor/Private/MuCOE/GenerateMutableSource/GenerateMutableSourceMesh.cpp
+++ b/Engine/Plugins/Mutable/Source/CustomizableObjectEditor/Private/MuCOE/GenerateMutableSource/GenerateMutableSourceMesh.cpp
@@ -371,7 +371,7 @@ void GetLODAndSectionForAutomaticLODs(const FMutableCompilationContext& Context,
 }
 
 
-void BuildRemappedBonesArray(const FMutableComponentInfo& InComponentInfo, TObjectPtr<const USkeletalMesh> InSkeletalMesh, int32 InLODIndex, const TArray<FBoneIndexType>& InRequiredBones, TArray<FBoneIndexType>& OutRemappedBones)
+void BuildRemappedBonesArray(const FSkeletalMeshComponentInfo& InComponentInfo, TObjectPtr<const USkeletalMesh> InSkeletalMesh, int32 InLODIndex, const TArray<FBoneIndexType>& InRequiredBones, TArray<FBoneIndexType>& OutRemappedBones)
 {
 	if (!InSkeletalMesh)
 	{
@@ -612,7 +612,7 @@ void ConvertVertexSkinningData(uint16* InfluenceBones, const UE::Mutable::Privat
 }
 
 
-bool IsSkeletalMeshCompatibleWithRefSkeleton(FMutableComponentInfo& ComponentInfo, TObjectPtr<const USkeletalMesh> InSkeletalMesh, FString& OutErrorMessage)
+bool IsSkeletalMeshCompatibleWithRefSkeleton(FSkeletalMeshComponentInfo& ComponentInfo, TObjectPtr<const USkeletalMesh> InSkeletalMesh, FString& OutErrorMessage)
 {
 	TObjectPtr<const USkeleton> Skeleton = InSkeletalMesh->GetSkeleton();
 
@@ -1314,7 +1314,7 @@ UE::Tasks::FTask ConvertSkeletalMeshToMutable(TSharedRef<FMeshConversionContext>
 			return UE::Tasks::MakeCompletedTask<void>();
 		}
 
-		FMutableComponentInfo* MutComponentInfo = Context.GetComponentInfo(Source.Component);
+		FSkeletalMeshComponentInfo* MutComponentInfo = Context.GetComponentInfo(Source.ComponentId);
 		check(MutComponentInfo);
 
 		USkeletalMesh* ComponentRefSkeletalMesh = MutComponentInfo->RefSkeletalMesh.Get();
@@ -2657,7 +2657,7 @@ UE::Mutable::Private::TManagedPtr<UE::Mutable::Private::FMesh> ConvertSkeletalMe
 	Source.LODOffset = GenerationContext.CurrentLOD - GenerationContext.FromLOD;
 	Source.Flags = GenerationContext.MeshGenerationFlags.Last();
 	Source.MessageContext = CurrentNode;
-	Source.Component = GenerationContext.CurrentMeshComponent;
+	Source.ComponentId = GenerationContext.CurrentSkeletalMeshComponent;
 	Source.SocketPriority = GenerationContext.SocketPriorityStack.Last();
 	Source.BonePosePriority = GenerationContext.BonePosePriorityStack.Last();
 
@@ -3880,7 +3880,7 @@ UE::Mutable::Private::Ptr<UE::Mutable::Private::NodeMesh> GenerateMutableSourceM
 			// 
 			// Generate Layouts when using manual LODs or when processing the base Automatic LOD (the other automatic LODs will use the base layouts).
 			const bool bGenerateLayouts = GenerationContext.CurrentAutoLODStrategy != ECustomizableObjectAutomaticLODStrategy::AutomaticFromMesh
-				|| (GenerationContext.CurrentLOD == GenerationContext.FromLOD || GenerationContext.CurrentLOD <= GenerationContext.FirstLODAvailable[GenerationContext.CurrentMeshComponent]);
+				|| (GenerationContext.CurrentLOD == GenerationContext.FromLOD || GenerationContext.CurrentLOD <= GenerationContext.FirstLODAvailable[GenerationContext.CurrentSkeletalMeshComponent]);
 
 			if(bGenerateLayouts)
 			{
@@ -4533,7 +4533,7 @@ UE::Mutable::Private::Ptr<UE::Mutable::Private::NodeMesh> GenerateMutableSourceM
 							// 
 							// Generate Layouts when using manual LODs or when processing the base Automatic LOD (the other automatic LODs will use the base layouts).
 							const bool bGenerateLayouts = GenerationContext.CurrentAutoLODStrategy != ECustomizableObjectAutomaticLODStrategy::AutomaticFromMesh
-								|| (GenerationContext.CurrentLOD == GenerationContext.FromLOD || GenerationContext.CurrentLOD <= GenerationContext.FirstLODAvailable[GenerationContext.CurrentMeshComponent]);
+								|| (GenerationContext.CurrentLOD == GenerationContext.FromLOD || GenerationContext.CurrentLOD <= GenerationContext.FirstLODAvailable[GenerationContext.CurrentSkeletalMeshComponent]);
 
 							if (bGenerateLayouts)
 							{
diff --git a/Engine/Plugins/Mutable/Source/CustomizableObjectEditor/Private/MuCOE/GenerateMutableSource/GenerateMutableSourceModifier.cpp b/Engine/Plugins/Mutable/Source/CustomizableObjectEditor/Private/MuCOE/GenerateMutableSource/GenerateMutableSourceModifier.cpp
--- a/Engine/Plugins/Mutable/Source/CustomizableObjectEditor/Private/MuCOE/GenerateMutableSource/GenerateMutableSourceModifier.cpp
+++ b/Engine/Plugins/Mutable/Source/CustomizableObjectEditor/Private/MuCOE/GenerateMutableSource/GenerateMutableSourceModifier.cpp
@@ -59,7 +59,7 @@ UE::Mutable::Private::Ptr<UE::Mutable::Private::NodeModifier> GenerateMutableSou
 	UCustomizableObjectNode* Node = CastChecked<UCustomizableObjectNode>(Pin->GetOwningNode());
 
 	FGeneratedKey Key(reinterpret_cast<void*>(&GenerateMutableSourceModifier), *Pin, *Node, GenerationContext, true);
-	Key.CurrentMeshComponent = GenerationContext.CurrentMeshComponent;
+	Key.CurrentMeshComponent = GenerationContext.CurrentSkeletalMeshComponent;
 	
 	if (const FGeneratedData* Generated = GenerationContext.Generated.Find(Key))
 	{
@@ -265,7 +265,7 @@ UE::Mutable::Private::Ptr<UE::Mutable::Private::NodeModifier> GenerateMutableSou
 		SurfNode->RequiredTags = TypedNodeExt->GetNodeRequiredTags(&GenerationContext.MacroNodesStack);
 
 		// Is this enough? Should we try to narrow down with potential mesh sections modified by this?
-		int32 LODCount = GenerationContext.NumLODs[GenerationContext.CurrentMeshComponent];
+		int32 LODCount = GenerationContext.NumLODs[GenerationContext.CurrentSkeletalMeshComponent];
 		SurfNode->LODs.SetNum(LODCount);
 
 		for (int32 LODIndex= TypedNodeExt->FirstLOD; LODIndex<LODCount; ++LODIndex)
@@ -388,7 +388,7 @@ UE::Mutable::Private::Ptr<UE::Mutable::Private::NodeModifier> GenerateMutableSou
 		if (const UEdGraphPin* ConnectedPin = FollowInputPin(*TypedNodeRem->RemoveMeshPin()))
 		{
 			// Is this enough? Should we try to narrow down with potential mesh sections modified by this?
-			int32 LODCount = GenerationContext.NumLODs[GenerationContext.CurrentMeshComponent];
+			int32 LODCount = GenerationContext.NumLODs[GenerationContext.CurrentSkeletalMeshComponent];
 			SurfNode->LODs.SetNum(LODCount);
 
 			SurfNode->FaceCullStrategy = TypedNodeRem->FaceCullStrategy;
@@ -441,7 +441,7 @@ UE::Mutable::Private::Ptr<UE::Mutable::Private::NodeModifier> GenerateMutableSou
 		SurfNode->RequiredTags = TypedNodeEdit->GetNodeRequiredTags(&GenerationContext.MacroNodesStack);
 
 		// Is this enough? Should we try to narrow down with potential mesh sections modified by this?
-		int32 LODCount = GenerationContext.NumLODs[GenerationContext.CurrentMeshComponent];
+		int32 LODCount = GenerationContext.NumLODs[GenerationContext.CurrentSkeletalMeshComponent];
 		SurfNode->LODs.SetNum(LODCount);
 
 		for (int32 LODIndex = 0; LODIndex < LODCount; ++LODIndex)
@@ -686,9 +686,8 @@ UE::Mutable::Private::Ptr<UE::Mutable::Private::NodeModifier> GenerateMutableSou
 	{
 		Result->SetMessageContext(Node);
 
-		int32 ComponentId = GenerationContext.ComponentNames.IndexOfByKey(GenerationContext.CurrentMeshComponent);
-		check(ComponentId>=0);
-		Result->RequiredComponentId = ComponentId;
+		check(GenerationContext.CurrentSkeletalMeshComponent != INDEX_NONE);
+		Result->RequiredComponentId = GenerationContext.CurrentSkeletalMeshComponent;
 	}
 
 	if (bCacheNode)
diff --git a/Engine/Plugins/Mutable/Source/CustomizableObjectEditor/Private/MuCOE/GenerateMutableSource/GenerateMutableSourceSkeletalMesh.cpp b/Engine/Plugins/Mutable/Source/CustomizableObjectEditor/Private/MuCOE/GenerateMutableSource/GenerateMutableSourceSkeletalMesh.cpp
--- a/Engine/Plugins/Mutable/Source/CustomizableObjectEditor/Private/MuCOE/GenerateMutableSource/GenerateMutableSourceSkeletalMesh.cpp
+++ b/Engine/Plugins/Mutable/Source/CustomizableObjectEditor/Private/MuCOE/GenerateMutableSource/GenerateMutableSourceSkeletalMesh.cpp
@@ -75,7 +75,7 @@ namespace UE::Mutable::Private
 
 			int32 FirstLOD = -1;
 
-			const int32 NumLODsInRoot = GenerationContext.NumLODs[GenerationContext.CurrentMeshComponent];
+			const int32 NumLODsInRoot = GenerationContext.NumLODs[GenerationContext.CurrentSkeletalMeshComponent];
 			for (int32 CurrentLOD = 0; CurrentLOD < NumLODsInRoot; ++CurrentLOD)
 			{
 				GenerationContext.CurrentLOD = CurrentLOD;
@@ -99,7 +99,7 @@ namespace UE::Mutable::Private
 					continue;
 				}
 
-				if (GenerationContext.CurrentLOD < GenerationContext.FirstLODAvailable[GenerationContext.CurrentMeshComponent])
+				if (GenerationContext.CurrentLOD < GenerationContext.FirstLODAvailable[GenerationContext.CurrentSkeletalMeshComponent])
 				{
 					continue;
 				} 
diff --git a/Engine/Plugins/Mutable/Source/CustomizableObjectEditor/Private/MuCOE/SMutableCodeViewer.cpp b/Engine/Plugins/Mutable/Source/CustomizableObjectEditor/Private/MuCOE/SMutableCodeViewer.cpp
--- a/Engine/Plugins/Mutable/Source/CustomizableObjectEditor/Private/MuCOE/SMutableCodeViewer.cpp
+++ b/Engine/Plugins/Mutable/Source/CustomizableObjectEditor/Private/MuCOE/SMutableCodeViewer.cpp
@@ -1691,7 +1691,6 @@ void SMutableCodeViewer::GenerateElementRecursive(const int32& InStateIndex, UE:
 
 	case UE::Mutable::Private::EOpType::IN_ADDMESH:
 	case UE::Mutable::Private::EOpType::IN_ADDSTRING:
-	case UE::Mutable::Private::EOpType::IN_ADDCOMPONENT:
 	case UE::Mutable::Private::EOpType::IN_ADDSURFACE:
 	case UE::Mutable::Private::EOpType::IN_ADDSKELETALMESH:
 	{
@@ -1701,6 +1700,14 @@ void SMutableCodeViewer::GenerateElementRecursive(const int32& InStateIndex, UE:
 		break;
 	}
 		
+	case UE::Mutable::Private::EOpType::IN_ADDCOMPONENT:
+	{
+		UE::Mutable::Private::OP::InstanceAddComponentArgs Args = InProgram.GetOpArgs<UE::Mutable::Private::OP::InstanceAddComponentArgs>(InParentAddress);
+		AddOpFunc(Args.Instance, TEXT("instance "));
+		AddOpFunc(Args.Value, TEXT("value "));
+		break;
+	}
+		
 	case UE::Mutable::Private::EOpType::IN_ADDOVERLAYMATERIAL:
 	{
 		UE::Mutable::Private::OP::InstanceAddMaterialArgs Args = InProgram.GetOpArgs<UE::Mutable::Private::OP::InstanceAddMaterialArgs>(InParentAddress);
@@ -2627,10 +2634,6 @@ bool SMutableCodeViewer::IsConstantResourceUsedByOperation(const int32 IndexOnCo
 				{
 					bResourceLocated = IndexOnConstantsArray == InProgram.GetOpArgs<UE::Mutable::Private::OP::InstanceAddArgs>(OperationAddress).name;
 				}
-				else if (OperationType == UE::Mutable::Private::EOpType::IN_ADDCOMPONENT)
-				{
-					bResourceLocated = IndexOnConstantsArray == InProgram.GetOpArgs<UE::Mutable::Private::OP::InstanceAddArgs>(OperationAddress).name;
-				}
 				else if (OperationType == UE::Mutable::Private::EOpType::IN_ADDSURFACE)
 				{
 					bResourceLocated = IndexOnConstantsArray == InProgram.GetOpArgs<UE::Mutable::Private::OP::InstanceAddArgs>(OperationAddress).name;
diff --git a/Engine/Plugins/Mutable/Source/MutableRuntime/Internal/MuR/ComponentId.h b/Engine/Plugins/Mutable/Source/MutableRuntime/Internal/MuR/ComponentId.h
new file mode 100644
--- /dev/null
+++ b/Engine/Plugins/Mutable/Source/MutableRuntime/Internal/MuR/ComponentId.h
@@ -0,0 +1,11 @@
+﻿// Copyright Epic Games, Inc. All Rights Reserved.
+
+#pragma once
+
+#include "CoreTypes.h"
+
+
+namespace UE::Mutable::Private
+{
+	using FComponentId = int32;
+}
diff --git a/Engine/Plugins/Mutable/Source/MutableRuntime/Internal/MuR/Instance.h b/Engine/Plugins/Mutable/Source/MutableRuntime/Internal/MuR/Instance.h
--- a/Engine/Plugins/Mutable/Source/MutableRuntime/Internal/MuR/Instance.h
+++ b/Engine/Plugins/Mutable/Source/MutableRuntime/Internal/MuR/Instance.h
@@ -54,7 +54,7 @@ namespace UE::Mutable::Private
 
 	struct FInstanceComponent
 	{
-		uint16 Id;
+		FComponentId Id = INDEX_NONE;
 
 		FMaterialId OverlayMaterialId;
 
@@ -100,7 +100,7 @@ namespace UE::Mutable::Private
 
 		//! Get the Id of a component
 		//! \param ComponentIndex Index of the component, from 0 to GetComponentCount()-1
-		UE_API uint16 GetComponentId(int32 ComponentIndex) const;
+		UE_API FComponentId GetComponentId(int32 ComponentIndex) const;
 
         //! Get the number of LODs in a component
         //! \param ComponentIndex Index of the component, from 0 to GetComponentCount()-1
diff --git a/Engine/Plugins/Mutable/Source/MutableRuntime/Internal/MuR/Operations.h b/Engine/Plugins/Mutable/Source/MutableRuntime/Internal/MuR/Operations.h
--- a/Engine/Plugins/Mutable/Source/MutableRuntime/Internal/MuR/Operations.h
+++ b/Engine/Plugins/Mutable/Source/MutableRuntime/Internal/MuR/Operations.h
@@ -2,6 +2,7 @@
 
 #pragma once
 
+#include "ComponentId.h"
 #include "PassthroughObject.h"
 #include "MuR/ImageTypes.h"
 #include "MuR/Types.h"
@@ -15,6 +16,7 @@
 #define MUTABLE_OP_MAX_INTERPOLATE_COUNT	6
 #define MUTABLE_OP_MAX_SWIZZLE_CHANNELS		4
 
+
 namespace UE::Mutable::Private
 {
 	enum class EImageFormat : uint8;
@@ -1134,6 +1136,13 @@ namespace UE::Mutable::Private
 
             ADDRESS name;
         };
+    	
+    	struct InstanceAddComponentArgs
+    	{
+    		ADDRESS Instance;
+    		ADDRESS Value;
+    		FComponentId ExternalId;
+    	};
 
 		struct InstanceAddExtensionDataArgs
 		{
diff --git a/Engine/Plugins/Mutable/Source/MutableRuntime/Private/MuR/CodeRunner.cpp b/Engine/Plugins/Mutable/Source/MutableRuntime/Private/MuR/CodeRunner.cpp
--- a/Engine/Plugins/Mutable/Source/MutableRuntime/Private/MuR/CodeRunner.cpp
+++ b/Engine/Plugins/Mutable/Source/MutableRuntime/Private/MuR/CodeRunner.cpp
@@ -573,18 +573,18 @@ namespace UE::Mutable::Private
 
         case EOpType::IN_ADDCOMPONENT:
         {
-			OP::InstanceAddArgs Args = Program.GetOpArgs<OP::InstanceAddArgs>(Item.At);
+			OP::InstanceAddComponentArgs Args = Program.GetOpArgs<OP::InstanceAddComponentArgs>(Item.At);
             switch (Item.Stage)
             {
             case 0:
                     AddOp( FScheduledOp( Item.At, Item, 1),
-                           FScheduledOp( Args.instance, Item),
-                           FScheduledOp( Args.value, Item) );
+                           FScheduledOp( Args.Instance, Item),
+                           FScheduledOp( Args.Value, Item) );
                 break;
 
             case 1:
             {
-				TManagedPtr<const FInstance> Base = LoadInstance( FCacheAddress(Args.instance,Item) );
+				TManagedPtr<const FInstance> Base = LoadInstance( FCacheAddress(Args.Instance,Item) );
 				TManagedPtr<FInstance> Result = nullptr;
 
 				if (!Base)
@@ -596,9 +596,9 @@ namespace UE::Mutable::Private
 					Result = UE::Mutable::Private::CloneOrTakeOver<FInstance>(MoveTemp(Base));
 				}
 
-                if ( Args.value )
+                if ( Args.Value )
                 {
-                    TManagedPtr<const FInstance> pComp = LoadInstance( FCacheAddress(Args.value,Item) );
+                    TManagedPtr<const FInstance> pComp = LoadInstance( FCacheAddress(Args.Value,Item) );
 
                     int32 NewComponentIndex = Result->AddComponent();
 
diff --git a/Engine/Plugins/Mutable/Source/MutableRuntime/Private/MuR/CodeRunnerBegin.cpp b/Engine/Plugins/Mutable/Source/MutableRuntime/Private/MuR/CodeRunnerBegin.cpp
--- a/Engine/Plugins/Mutable/Source/MutableRuntime/Private/MuR/CodeRunnerBegin.cpp
+++ b/Engine/Plugins/Mutable/Source/MutableRuntime/Private/MuR/CodeRunnerBegin.cpp
@@ -1043,14 +1043,14 @@ namespace UE::Mutable::Private
 				{
 					check(Item.bEvaluate);
 
-					OP::InstanceAddArgs Args = Program.GetOpArgs<OP::InstanceAddArgs>(Item.At);
+					OP::InstanceAddComponentArgs Args = Program.GetOpArgs<OP::InstanceAddComponentArgs>(Item.At);
 					switch (Item.Stage)
 					{
 					case 0:
 						{
 							PushOp(FScheduledOpInline(Item, 1));
-							PushOp(FScheduledOpInline(Args.instance, Item));
-							PushOp(FScheduledOpInline(Args.value, Item));
+							PushOp(FScheduledOpInline(Args.Instance, Item));
+							PushOp(FScheduledOpInline(Args.Value, Item));
 							
 							break;
 						}
@@ -1058,9 +1058,9 @@ namespace UE::Mutable::Private
 					case 1:
 						{
 							TManagedPtr<const FInstance> InInstance;
-							if (Args.instance)
+							if (Args.Instance)
 							{
-								InInstance = LoadInstance(Args.instance);
+								InInstance = LoadInstance(Args.Instance);
 							}
 
 							TManagedPtr<FInstance> Instance;
@@ -1073,9 +1073,9 @@ namespace UE::Mutable::Private
 								Instance = MakeManaged<FInstance>();
 							}
 
-							if (Args.value)
+							if (Args.Value)
 							{
-								TManagedPtr<const FInstance> pComp = LoadInstance(Args.value);
+								TManagedPtr<const FInstance> pComp = LoadInstance(Args.Value);
 
 								int32 NewComponentIndex = Instance->AddComponent();
 
diff --git a/Engine/Plugins/Mutable/Source/MutableRuntime/Private/MuR/Instance.cpp b/Engine/Plugins/Mutable/Source/MutableRuntime/Private/MuR/Instance.cpp
--- a/Engine/Plugins/Mutable/Source/MutableRuntime/Private/MuR/Instance.cpp
+++ b/Engine/Plugins/Mutable/Source/MutableRuntime/Private/MuR/Instance.cpp
@@ -43,7 +43,7 @@ namespace UE::Mutable::Private
 	}
 	
 	
-	uint16 FInstance::GetComponentId( int32 ComponentIndex ) const
+	FComponentId FInstance::GetComponentId( int32 ComponentIndex ) const
 	{
 		if (Components.IsValidIndex(ComponentIndex))
 		{
@@ -54,7 +54,7 @@ namespace UE::Mutable::Private
 			check(false);
 		}
 
-		return 0;
+		return INDEX_NONE;
 	}
 
 	
diff --git a/Engine/Plugins/Mutable/Source/MutableRuntime/Private/MuR/Operations.cpp b/Engine/Plugins/Mutable/Source/MutableRuntime/Private/MuR/Operations.cpp
--- a/Engine/Plugins/Mutable/Source/MutableRuntime/Private/MuR/Operations.cpp
+++ b/Engine/Plugins/Mutable/Source/MutableRuntime/Private/MuR/Operations.cpp
@@ -900,7 +900,6 @@ namespace UE::Mutable::Private
         //-------------------------------------------------------------------------------------
         case EOpType::IN_ADDMESH:
         case EOpType::IN_ADDSTRING:
-        case EOpType::IN_ADDCOMPONENT:
         case EOpType::IN_ADDSURFACE:
         {
 			OP::InstanceAddArgs args = program.GetOpArgs<OP::InstanceAddArgs>(at);
@@ -909,6 +908,14 @@ namespace UE::Mutable::Private
             break;
         }
 
+        case EOpType::IN_ADDCOMPONENT:
+        {
+        	OP::InstanceAddComponentArgs args = program.GetOpArgs<OP::InstanceAddComponentArgs>(at);
+        	f(args.Instance );
+        	f(args.Value );
+        	break;
+        }
+        	
         case EOpType::IN_ADDLOD:
         {
 			const uint8* Data = program.GetOpArgsPointer(at);
diff --git a/Engine/Plugins/Mutable/Source/MutableTools/Internal/MuT/NodeComponentNew.h b/Engine/Plugins/Mutable/Source/MutableTools/Internal/MuT/NodeComponentNew.h
--- a/Engine/Plugins/Mutable/Source/MutableTools/Internal/MuT/NodeComponentNew.h
+++ b/Engine/Plugins/Mutable/Source/MutableTools/Internal/MuT/NodeComponentNew.h
@@ -2,7 +2,7 @@
 
 #pragma once
 
-#include "HAL/Platform.h"
+#include "MuR/ComponentId.h"
 #include "MuR/Ptr.h"
 #include "MuT/NodeSkeletalMeshNew.h"
 #include "MuT/Node.h"
@@ -19,7 +19,7 @@ namespace UE::Mutable::Private
 	public:
 
 		/** Externally managed id assigned to this component. */
-		uint16 Id = 0;
+		FComponentId Id = INDEX_NONE;
 
 		Ptr<NodeMaterial> OverlayMaterial;
 
diff --git a/Engine/Plugins/Mutable/Source/MutableTools/Private/MuT/ASTOpComponentAdd.cpp b/Engine/Plugins/Mutable/Source/MutableTools/Private/MuT/ASTOpComponentAdd.cpp
new file mode 100644
--- /dev/null
+++ b/Engine/Plugins/Mutable/Source/MutableTools/Private/MuT/ASTOpComponentAdd.cpp
@@ -0,0 +1,102 @@
+// Copyright Epic Games, Inc. All Rights Reserved.
+
+#include "MuT/ASTOpComponentAdd.h"
+
+#include "MuT/ASTOpConditional.h"
+#include "MuR/Model.h"
+#include "MuR/Parameters.h"
+#include "MuR/RefCounted.h"
+
+
+namespace UE::Mutable::Private
+{
+	ASTOpComponentAdd::ASTOpComponentAdd()
+		: Instance(this)
+		, Value(this)
+	{
+	}
+
+
+	ASTOpComponentAdd::~ASTOpComponentAdd()
+	{
+		// Explicit call needed to avoid recursive destruction
+		ASTOp::RemoveChildren();
+	}
+
+	
+	EOpType ASTOpComponentAdd::GetOpType() const
+	{
+		return EOpType::IN_ADDCOMPONENT;
+	}
+
+
+	bool ASTOpComponentAdd::IsEqual(const ASTOp& OtherUntyped) const
+	{
+		if (OtherUntyped.GetOpType() == GetOpType())
+		{
+			const ASTOpComponentAdd* other = static_cast<const ASTOpComponentAdd*>(&OtherUntyped);
+			return GetOpType() == other->GetOpType() &&
+				Instance == other->Instance &&
+				Value == other->Value &&
+				ExternalId == other->ExternalId;
+		}
+		
+		return false;
+	}
+
+
+	UE::Mutable::Private::Ptr<ASTOp> ASTOpComponentAdd::Clone(MapChildFuncRef MapChild) const
+	{
+		Ptr<ASTOpComponentAdd> New = new ASTOpComponentAdd();
+		New->Instance = MapChild(Instance.child());
+		New->Value = MapChild(Value.child());
+		New->ExternalId = ExternalId;
+		
+		return New;
+	}
+
+
+	uint64 ASTOpComponentAdd::Hash() const
+	{
+		uint64 res = std::hash<size_t>()(size_t(GetOpType()));
+
+		hash_combine(res, Instance.child().get());
+		hash_combine(res, Value.child().get());
+		
+		return res;
+	}
+	
+
+	void ASTOpComponentAdd::ForEachChild(const TFunctionRef<void(ASTChild&)> Func)
+	{
+		Func(Instance);
+		Func(Value);
+	}
+
+
+	void ASTOpComponentAdd::Link(FProgram& Program, FLinkerOptions*)
+	{
+		// Already linked?
+		if (!linkedAddress)
+		{
+			OP::InstanceAddComponentArgs Args;
+			FMemory::Memzero(Args);
+			Args.ExternalId = ExternalId;
+
+			if (Instance)
+			{
+				Args.Instance = Instance->linkedAddress;
+			}
+				
+			if (Value)
+			{
+				Args.Value = Value->linkedAddress;
+			}
+
+			linkedAddress = (OP::ADDRESS)Program.OpAddress.Num();
+			Program.OpAddress.Add((uint32_t)Program.ByteCode.Num());
+			AppendCode(Program.ByteCode, GetOpType());
+			AppendCode(Program.ByteCode, Args);
+		}
+	}
+}
diff --git a/Engine/Plugins/Mutable/Source/MutableTools/Private/MuT/ASTOpComponentAdd.h b/Engine/Plugins/Mutable/Source/MutableTools/Private/MuT/ASTOpComponentAdd.h
new file mode 100644
--- /dev/null
+++ b/Engine/Plugins/Mutable/Source/MutableTools/Private/MuT/ASTOpComponentAdd.h
@@ -0,0 +1,39 @@
+// Copyright Epic Games, Inc. All Rights Reserved.
+
+#pragma once
+
+#include "HAL/PlatformMath.h"
+#include "MuR/Operations.h"
+#include "MuR/Ptr.h"
+#include "MuT/AST.h"
+
+
+namespace UE::Mutable::Private
+{
+	struct FProgram;
+
+
+	//---------------------------------------------------------------------------------------------
+	//! Operations to add elements to an instance
+	//---------------------------------------------------------------------------------------------
+	class ASTOpComponentAdd final : public ASTOp
+	{
+	public:
+		ASTChild Instance;
+		ASTChild Value;
+		FComponentId ExternalId = INDEX_NONE;
+
+		ASTOpComponentAdd();
+		ASTOpComponentAdd(const ASTOpComponentAdd&) = delete;
+		~ASTOpComponentAdd();
+
+		virtual EOpType GetOpType() const override;
+
+		virtual void ForEachChild(const TFunctionRef<void(ASTChild&)>) override;
+		virtual bool IsEqual(const ASTOp& OtherUntyped) const override;
+		virtual Ptr<ASTOp> Clone(MapChildFuncRef MapChild) const override;
+		virtual uint64 Hash() const override;
+		virtual void Link(FProgram& Program, FLinkerOptions* Options) override;
+	};
+}
+
diff --git a/Engine/Plugins/Mutable/Source/MutableTools/Private/MuT/ASTOpInstanceAdd.cpp b/Engine/Plugins/Mutable/Source/MutableTools/Private/MuT/ASTOpInstanceAdd.cpp
--- a/Engine/Plugins/Mutable/Source/MutableTools/Private/MuT/ASTOpInstanceAdd.cpp
+++ b/Engine/Plugins/Mutable/Source/MutableTools/Private/MuT/ASTOpInstanceAdd.cpp
@@ -73,7 +73,6 @@ namespace UE::Mutable::Private
 		case EOpType::IN_ADDMESH:
 		case EOpType::IN_ADDSTRING:
 		case EOpType::IN_ADDSURFACE:
-		case EOpType::IN_ADDCOMPONENT:
 		case EOpType::IN_ADDLOD:
 			break;
 		default:
diff --git a/Engine/Plugins/Mutable/Source/MutableTools/Private/MuT/CodeGenerator.cpp b/Engine/Plugins/Mutable/Source/MutableTools/Private/MuT/CodeGenerator.cpp
--- a/Engine/Plugins/Mutable/Source/MutableTools/Private/MuT/CodeGenerator.cpp
+++ b/Engine/Plugins/Mutable/Source/MutableTools/Private/MuT/CodeGenerator.cpp
@@ -5,6 +5,7 @@
 #include "ASTOpAddOverrideMaterial.h"
 #include "ASTOpAddSkeletalMesh.h"
 #include "ASTOpEntry.h"
+#include "ASTOpComponentAdd.h"
 #include "Containers/Array.h"
 #include "Logging/LogCategory.h"
 #include "Logging/LogMacros.h"
@@ -637,10 +638,9 @@ namespace UE::Mutable::Private
 					LastInstOp = AddOverrideMaterialOp;
 				}
 
-				Ptr<ASTOpInstanceAdd> InstanceOp = new ASTOpInstanceAdd();
-				InstanceOp->type = EOpType::IN_ADDCOMPONENT;
-				InstanceOp->instance = Options.BaseInstance;
-				InstanceOp->value = LastInstOp;
+				Ptr<ASTOpComponentAdd> InstanceOp = new ASTOpComponentAdd();
+				InstanceOp->Instance = Options.BaseInstance;
+				InstanceOp->Value = LastInstOp;
 				InstanceOp->ExternalId = InNode->Id;
 
 				Result.Op = InstanceOp;
