[Mutable] Avoid LoadSync the full CO hierarchy.

- Reorganized tickers.

#Synced-CL 49557694

diff --git a/Engine/Plugins/Mutable/Source/CustomizableObject/Internal/MuCO/CustomizableObjectSystemPrivate.h b/Engine/Plugins/Mutable/Source/CustomizableObject/Internal/MuCO/CustomizableObjectSystemPrivate.h
--- a/Engine/Plugins/Mutable/Source/CustomizableObject/Internal/MuCO/CustomizableObjectSystemPrivate.h
+++ b/Engine/Plugins/Mutable/Source/CustomizableObject/Internal/MuCO/CustomizableObjectSystemPrivate.h
@@ -795,8 +795,12 @@ public:
 	virtual void SetDisregardWorldResourcesForFrames(int32 NumFrames) override {}
 	virtual void NotifyLevelOffset(ULevel* Level, const FVector& Offset) override {}
 
-	UCustomizableObjectSystem* GetPublic() const;
+	UE_API int32 MainTick(bool bBlocking);
+	
+	int32 UpdateTick(bool bBlocking);
 	
+	UCustomizableObjectSystem* GetPublic() const;
+
 	EUpdateRequired IsUpdateRequired(const UCustomizableObjectInstance& Instance, bool bOnlyUpdateIfNotGenerated, bool bOnlyUpdateIfLOD, bool bIgnoreCloseDist) const;
 
 	EQueuePriorityType GetUpdatePriority(const UCustomizableObjectInstance& Instance, bool bForceHighPriority) const;
@@ -806,7 +810,6 @@ public:
 	// Init an async and safe release of the UE and Mutable resources used by the instance without actually destroying the instance, for example if it's very far away
 	void InitDiscardResourcesSkeletalMesh(UCustomizableObjectInstance* InCustomizableObjectInstance);
 
-	
 	bool IsReplaceDiscardedWithReferenceMeshEnabled() const;
 	void SetReplaceDiscardedWithReferenceMeshEnabled(bool bIsEnabled);
 
@@ -840,29 +843,12 @@ public:
 	/** See UCustomizableObjectInstance::IsUpdating. */
 	bool IsUpdating(const UCustomizableObjectInstance& Instance) const;
 
-	/** Update stats at each tick.
-	 * Used for stats that are costly to update. */
-	void UpdateStats();
-
-#if WITH_EDITOR
-	/** Add a CO to the pending - load list.COs need to wait until all related objects are fully loaded before being able
-	 * to do things like check-if-up-to-date or compile.
-	 */
-	void AddPendingLoad(UCustomizableObject*);
-#endif
-	
 	/** Mutable TaskGraph system (Mutable Thread). */
 	FMutableTaskGraph MutableTaskGraph;
 
 	/** Last Mutable task from the previous update. The next update can not start until this task has been completed. */
 	UE::Tasks::FTask LastUpdateMutableTask = UE::Tasks::MakeCompletedTask<void>();
 	
-#if WITH_EDITORONLY_DATA
-	/** List of CustomizableObjects pending to complete loading. */
-	UPROPERTY(Transient)
-	TArray<TObjectPtr<UCustomizableObject>> ObjectsPendingLoad;
-#endif
-
 	FLogBenchmarkUtil LogBenchmarkUtil;
 
 	int32 NumSkeletalMeshes = 0;
@@ -898,6 +884,8 @@ public:
 
 	int32 NumLODUpdatesLastTick = 0;
 
+	int32 RemainingWork = 0;
+	
 	/** Time when the "Started Update Skeletal Mesh Async" log will be unmuted (in seconds). */
 	float LogStartedUpdateUnmute = 0.0;
 
@@ -918,7 +906,6 @@ public:
 	CUSTOMIZABLEOBJECT_API static void SetUsageOfBenchmarkingSettings(bool bUseBenchmarkingOptimizedSettings);
 
 private:
-
 	/** Flag that controls some of the settings used for the generation of instances. */
 	inline static bool bUseBenchmarkingSettings = false;
 };
diff --git a/Engine/Plugins/Mutable/Source/CustomizableObject/Internal/MuCO/ICustomizableObjectEditorModule.h b/Engine/Plugins/Mutable/Source/CustomizableObject/Internal/MuCO/ICustomizableObjectEditorModule.h
--- a/Engine/Plugins/Mutable/Source/CustomizableObject/Internal/MuCO/ICustomizableObjectEditorModule.h
+++ b/Engine/Plugins/Mutable/Source/CustomizableObject/Internal/MuCO/ICustomizableObjectEditorModule.h
@@ -71,11 +71,7 @@ public:
 	  * @param bReleaseVersionDiff true if the Release Version has changed since the last compilation.
 	  * @return true if the compilation is out of date. */
 	virtual bool IsCompilationOutOfDate(const UCustomizableObject& Object, bool bSkipIndirectReferences, TArray<FName>& OutOfDatePackages, TArray<FName>& AddedPackages, TArray<FName>& RemovedPackages, bool& bReleaseVersionDiff) const = 0;
-
-	/** Method called once all parent COs of this CO have already been loaded and it is safe to get data from the root CO.
-	* @note The compiled data will not be yet available. */
-	virtual void OnUpstreamCOsLoaded(UCustomizableObject* Object) const = 0;
-
+	
 	/** Fixup operations executed once the upstream COs (parent COs) have been loaded.
 	 * @param CustomizableObjectCustomVersion Version index used to determine what change in the data should be performed.
 	 */
@@ -100,9 +96,8 @@ public:
 	/** See GraphTraversal::GetRootObject(...) */
 	virtual const UCustomizableObject* GetRootObject(const UCustomizableObject* ChildObject) const = 0;
 
-	/** Return all the CustomizableObjects related to the given one. */
-	virtual void GetRelatedObjects(UCustomizableObject*, TSet<UCustomizableObject*>& OutRelated ) const = 0;
-
+	virtual void GetAllObjectsInGraph(UCustomizableObject*, TSet<UCustomizableObject*>& OutRelated ) const = 0;
+	
 	/**
 	 * Execute this method in order to bake the provided instance. It will schedule a special type of instance update before proceeding with the bake itself.
 	 * @param InTargetInstance The instance we want to bake
@@ -136,5 +131,9 @@ public:
 	virtual void BeginCacheForCookedPlatformData(UCustomizableObject& Object, const ITargetPlatform* TargetPlatform) = 0;
 	
 	virtual bool IsCachedCookedPlatformDataLoaded(UCustomizableObject& Object, const ITargetPlatform* TargetPlatform) = 0;
+	
+	/** Add a CO to the pending - load list.COs need to wait until all related objects are fully loaded before being able
+	 * to do things like check-if-up-to-date or compile. */
+	virtual void AddPendingLoad(UCustomizableObject* Object) = 0;
 };
  
diff --git a/Engine/Plugins/Mutable/Source/CustomizableObject/Private/MuCO/CustomizableObject.cpp b/Engine/Plugins/Mutable/Source/CustomizableObject/Private/MuCO/CustomizableObject.cpp
--- a/Engine/Plugins/Mutable/Source/CustomizableObject/Private/MuCO/CustomizableObject.cpp
+++ b/Engine/Plugins/Mutable/Source/CustomizableObject/Private/MuCO/CustomizableObject.cpp
@@ -371,8 +371,12 @@ void UCustomizableObject::PostLoad()
 			});
 	}
 	
-		GetPrivate()->Status.NextState(FCustomizableObjectStatusTypes::EState::Loading);
-		UCustomizableObjectSystem::GetInstance()->GetPrivate()->AddPendingLoad(this);
+	GetPrivate()->Status.NextState(FCustomizableObjectStatusTypes::EState::Loading);
+
+	if (ICustomizableObjectEditorModule* Module = ICustomizableObjectEditorModule::Get())
+	{
+		Module->AddPendingLoad(this);
+	}
 #endif
 }
 
diff --git a/Engine/Plugins/Mutable/Source/CustomizableObject/Private/MuCO/CustomizableObjectSystem.cpp b/Engine/Plugins/Mutable/Source/CustomizableObject/Private/MuCO/CustomizableObjectSystem.cpp
--- a/Engine/Plugins/Mutable/Source/CustomizableObject/Private/MuCO/CustomizableObjectSystem.cpp
+++ b/Engine/Plugins/Mutable/Source/CustomizableObject/Private/MuCO/CustomizableObjectSystem.cpp
@@ -13,6 +13,7 @@
 #include "MuCO/CustomizableObjectUIData.h"
 #include "MuCO/CustomizableObjectInstanceUsage.h"
 #include "MuCO/ICustomizableObjectModule.h"
+#include "MuCO/ICustomizableObjectEditorModule.h"
 #include "MuCO/LogBenchmarkUtil.h"
 #include "MuCO/UnrealMutableImageProvider.h"
 #include "MuCO/UnrealMutableModelDiskStreamer.h"
@@ -35,7 +36,6 @@
 #include "Misc/ConfigCacheIni.h"
 #include "Engine/World.h"
 #include "Interfaces/ITargetPlatform.h"
-#include "MuCO/ICustomizableObjectEditorModule.h"
 #include "MuT/UnrealPixelFormatOverride.h"
 #include "MuR/Parameters.h"
 #else
@@ -1652,14 +1652,6 @@ void UCustomizableObjectSystem::SetReplaceDiscardedWithReferenceMeshEnabled(bool
 
 
 #if WITH_EDITOR
-
-void UCustomizableObjectSystemPrivate::AddPendingLoad(UCustomizableObject* CO)
-{
-	check(IsInGameThread());
-	ObjectsPendingLoad.AddUnique(CO);
-}
-
-
 bool UCustomizableObjectSystem::LockObject(UCustomizableObject* InObject)
 {
 	check(InObject != nullptr);
@@ -3344,101 +3336,22 @@ namespace impl
 } // namespace impl
 
 
-bool UCustomizableObjectSystem::Tick(float DeltaTime)
+int32 UCustomizableObjectSystemPrivate::UpdateTick(const bool bBlocking)
 {
-	TickInternal(false);
-	return true;
-}
-
-
-int32 UCustomizableObjectSystem::TickInternal(const bool bBlocking)
-{
-	MUTABLE_CPUPROFILER_SCOPE(UCustomizableObjectSystem::TickInternal)
-
-	check(IsInGameThread());
-	
-	// Building instances is not enabled in servers. If at some point relevant collision or animation data is necessary for server logic this will need to be changed.
-#if UE_SERVER
-	return 0;
-#else // !UE_SERVER
-
-	if (!IsValid(this) || !IsValid(Private))
-	{
-		return 0;
-	}
-
-	if (IsEngineExitRequested())
-	{	
-		// TODO BEGIN. Remove once UE-281921  
-		GetPrivate()->MutableTaskGraph.UnlockMutableThread();
-   		GetPrivate()->MutableTaskGraph.AllowLaunchingMutableTaskLowPriority(true, false);
-		// TODO END
-
-		return 0;
-	}
-
-	if (IsValid(GWorld))
-	{
-		const EWorldType::Type WorldType = GWorld->WorldType;
-
-		if (WorldType != EWorldType::PIE && WorldType != EWorldType::Game && WorldType != EWorldType::Editor && WorldType != EWorldType::GamePreview)
-		{
-			return 0;
-		}
-	}
+	NumSkeletalMeshes = 0;
 	
-	// \TODO: Review: We should never compile an object from this tick, so this could be removed
-#if WITH_EDITOR
-
-	// See if any COs pending to load can be completed.
-	for ( TArray<TObjectPtr<UCustomizableObject>>::TIterator It = GetPrivate()->ObjectsPendingLoad.CreateIterator(); It; ++It )
+	for (TObjectIterator<UCustomizableObjectInstance> Instance; Instance; ++Instance)
 	{
-		UCustomizableObject* CO = (*It);
-		if (!IsValid(CO))
+		if (!IsValid(*Instance))
 		{
-			It.RemoveCurrent();
 			continue;
 		}
 
-		bool bReady = true;
-
-		if (ICustomizableObjectEditorModule* Module = ICustomizableObjectEditorModule::Get())
-		{
-			TSet<UCustomizableObject*> RelatedCustomizableObjects;
-			Module->GetRelatedObjects(CO, RelatedCustomizableObjects);
-
-			for (const UObject* CustomizableObject : RelatedCustomizableObjects)
-			{
-				if (CustomizableObject && !CustomizableObject->HasAnyFlags(EObjectFlags::RF_LoadCompleted))
-				{
-					bReady = false;
-					break;
-				}
-			}
-
-			if (bReady)
-			{
-				Module->OnUpstreamCOsLoaded(CO);
-			}
-		}
-
-		if (bReady)
-		{
-			CO->GetPrivate()->LoadCompiledDataFromDisk();
-			It.RemoveCurrent();
-		}
+		NumSkeletalMeshes += Instance->GetPrivate()->SkeletalMeshes.Num();
 	}
 
-	// Do not tick if the CookCommandlet is running.
-	if (IsRunningCookCommandlet())
-	{
-		return GetPrivate()->ObjectsPendingLoad.Num();
-	}
+	FMutableUpdateCandidate* LODUpdateCandidateFound = nullptr;
 	
-#endif
-
-	Private->UpdateStats();
-
 	for (TObjectIterator<UCustomizableObjectInstanceUsage> CustomizableObjectInstanceUsageItr; CustomizableObjectInstanceUsageItr; ++CustomizableObjectInstanceUsageItr)
 	{
 		TObjectPtr<UCustomizableObjectInstanceUsage> CustomizableObjectInstanceUsage = *CustomizableObjectInstanceUsageItr;
@@ -3492,7 +3405,6 @@ int32 UCustomizableObjectSystem::TickInternal(const bool bBlocking)
 		Instance->GetPrivate()->bAutomaticUpdateRequired = true;
 	}
 	
-	FMutableUpdateCandidate* LODUpdateCandidateFound = nullptr;
 	
 	bool bPendingCompilation = false;
 #if WITH_EDITOR
@@ -3501,13 +3413,13 @@ int32 UCustomizableObjectSystem::TickInternal(const bool bBlocking)
 #endif
 
 	// Get a new operation if we aren't working on one
-	if (!Private->CurrentMutableOperation && bIsMutableEnabled && !bPendingCompilation)
+	if (!CurrentMutableOperation && bIsMutableEnabled && !bPendingCompilation)
 	{
 		// Reset the instance relevancy
 		// The RequestedUpdates only refer to LOD changes. User Customization and discards are handled separately
 		FMutableInstanceUpdateMap RequestedLODUpdates;
 		
-		GetPrivate()->CurrentInstanceLODManagement->UpdateInstanceDistsAndLODs(RequestedLODUpdates);
+		CurrentInstanceLODManagement->UpdateInstanceDistsAndLODs(RequestedLODUpdates);
 
 		for (TObjectIterator<UCustomizableObjectInstance> CustomizableObjectInstance; CustomizableObjectInstance; ++CustomizableObjectInstance)
 		{
@@ -3545,13 +3457,13 @@ int32 UCustomizableObjectSystem::TickInternal(const bool bBlocking)
 			const FMutablePendingInstanceUpdate* PendingInstanceUpdateFound = nullptr;
 
 			// Look for the highest priority Pending Update
-			for (auto Iterator = Private->MutablePendingInstanceWork.GetUpdateIterator(); Iterator; ++Iterator)
+			for (auto Iterator = MutablePendingInstanceWork.GetUpdateIterator(); Iterator; ++Iterator)
 			{
 				FMutablePendingInstanceUpdate& PendingUpdate = *Iterator;
 
 				if (PendingUpdate.Context->Instance.IsValid() && PendingUpdate.Context->Instance->GetCustomizableObject())
 				{
-					const EQueuePriorityType PriorityType = Private->GetUpdatePriority(*PendingUpdate.Context->Instance, false);
+					const EQueuePriorityType PriorityType = GetUpdatePriority(*PendingUpdate.Context->Instance, false);
 					
 					if (PendingUpdate.Context->PriorityType <= MaxPriorityFound)
 					{
@@ -3606,13 +3518,13 @@ int32 UCustomizableObjectSystem::TickInternal(const bool bBlocking)
 				}
 			}
 
-			Private->NumLODUpdatesLastTick = RequestedLODUpdates.Num();
+			NumLODUpdatesLastTick = RequestedLODUpdates.Num();
 
 			// If the chosen LODUpdate has the same instance as a PendingUpdate, choose the PendingUpdate to apply both the LOD update
 			// and customization change
 			if (LODUpdateCandidateFound)
 			{
-				if (const FMutablePendingInstanceUpdate *PendingUpdateWithSameInstance = Private->MutablePendingInstanceWork.GetUpdate(LODUpdateCandidateFound->CustomizableObjectInstance))
+				if (const FMutablePendingInstanceUpdate *PendingUpdateWithSameInstance = MutablePendingInstanceWork.GetUpdate(LODUpdateCandidateFound->CustomizableObjectInstance))
 				{
 					PendingInstanceUpdateFound = PendingUpdateWithSameInstance;
 					LODUpdateCandidateFound = nullptr;
@@ -3634,8 +3546,8 @@ int32 UCustomizableObjectSystem::TickInternal(const bool bBlocking)
 					LODUpdateWithSameInstance->ApplyLODUpdateParamsToInstance(PendingInstanceUpdateFound->Context.Get());
 				}
 
-				Private->StartUpdateSkeletalMesh(PendingInstanceUpdateFound->Context);
-				Private->MutablePendingInstanceWork.RemoveUpdate(PendingInstanceUpdateFound->Context->Instance);
+				StartUpdateSkeletalMesh(PendingInstanceUpdateFound->Context);
+				MutablePendingInstanceWork.RemoveUpdate(PendingInstanceUpdateFound->Context->Instance);
 			}
 			else if (LODUpdateCandidateFound)
 			{
@@ -3648,7 +3560,7 @@ int32 UCustomizableObjectSystem::TickInternal(const bool bBlocking)
 				// Commit the LOD changes
 				LODUpdateCandidateFound->ApplyLODUpdateParamsToInstance(*Context);
 
-				Private->StartUpdateSkeletalMesh(Context);
+				StartUpdateSkeletalMesh(Context);
 			}
 		}
 
@@ -3664,36 +3576,77 @@ int32 UCustomizableObjectSystem::TickInternal(const bool bBlocking)
 		}
 
 		// Update the streaming limit if it has changed. It is safe to do this now.
-		Private->UpdateMemoryLimit();
+		UpdateMemoryLimit();
 
 		// Free memory before starting the new update
-		DiscardInstances();
+		GetPublic()->DiscardInstances();
 	}
 	
-	const int32 RemainingTasks = Private->MutableTaskGraph.Tick();
+	LogBenchmarkUtil.UpdateStats(); // Must to be the last thing to perform
+
+	if (!bIsMutableEnabled && !CurrentMutableOperation)
+	{
+		FStreamingManagerCollection::Get().RemoveStreamingManager(this);
+	}
+	
+	return CurrentMutableOperation.IsValid() + 
+	   MutablePendingInstanceWork.Num() +
+	   static_cast<int32>(LODUpdateCandidateFound != nullptr); // Still a pending LOD update. We can not use the size of RequestedLODUpdates since not all requests valid in future ticks.
+}
+
+
+int32 UCustomizableObjectSystemPrivate::MainTick(const bool bBlocking)
+{
+	MUTABLE_CPUPROFILER_SCOPE(UCustomizableObjectSystem::Tick)
 
-	Private->LogBenchmarkUtil.UpdateStats(); // Must to be the last thing to perform
+	check(IsInGameThread());
+	
+	// Building instances is not enabled in servers. If at some point relevant collision or animation data is necessary for server logic this will need to be changed.
+#if UE_SERVER
+	return 0;
+#else // !UE_SERVER
 
-	if (!bIsMutableEnabled && !Private->CurrentMutableOperation)
+	if (!IsValid(GetPublic()))
 	{
-		FStreamingManagerCollection::Get().RemoveStreamingManager(GetPrivate());
+		return 0;
 	}
 
-	int32 RemainingWork = Private->CurrentMutableOperation.IsValid() + 
-		Private->MutablePendingInstanceWork.Num() +
-		static_cast<int32>(LODUpdateCandidateFound != nullptr) + // Still a pending LOD update. We can not use the size of RequestedLODUpdates since not all requests valid in future ticks.
-		RemainingTasks;
+	if (IsEngineExitRequested())
+	{	
+		// TODO BEGIN. Remove once UE-281921  
+		MutableTaskGraph.UnlockMutableThread();
+   		MutableTaskGraph.AllowLaunchingMutableTaskLowPriority(true, false);
+		// TODO END
 
-#if WITH_EDITOR
-	if (bBlocking)
+		return 0;
+	}
+
+	if (IsValid(GWorld))
 	{
-		RemainingWork += EditorModule ? EditorModule->Tick(true) : 0;
+		const EWorldType::Type WorldType = GWorld->WorldType;
+
+		if (WorldType != EWorldType::PIE && WorldType != EWorldType::Game && WorldType != EWorldType::Editor && WorldType != EWorldType::GamePreview)
+		{
+			return 0;
+		}
 	}
 	
-	RemainingWork += GetPrivate()->ObjectsPendingLoad.Num();
-#endif
+	RemainingWork = 0;
+	
+	// Do not tick if the CookCommandlet is running.
+	if (!IsRunningCookCommandlet())
+	{
+		RemainingWork += UpdateTick(bBlocking);
+	}
+
+	RemainingWork = MutableTaskGraph.Tick();
+
+	if (ICustomizableObjectEditorModule* Module = ICustomizableObjectEditorModule::Get())
+	{
+		RemainingWork += Module->Tick(bBlocking);
+	}
 	
-	RemainingWork += GetPrivate()->StreamableManager->Tick(bBlocking);
+	RemainingWork += StreamableManager->Tick(bBlocking);
 	
 	return RemainingWork;
 #endif // !UE_SERVER
@@ -4103,22 +4056,6 @@ bool UCustomizableObjectSystemPrivate::IsUpdating(const UCustomizableObjectInsta
 }
 
 
-void UCustomizableObjectSystemPrivate::UpdateStats()
-{
-	NumSkeletalMeshes = 0;
-	
-	for (TObjectIterator<UCustomizableObjectInstance> Instance; Instance; ++Instance)
-	{
-		if (!IsValid(*Instance))
-		{
-			continue;
-		}
-
-		NumSkeletalMeshes += Instance->GetPrivate()->SkeletalMeshes.Num();
-	}
-}
-
-
 bool UCustomizableObjectSystem::IsMutableAnimInfoDebuggingEnabled() const
 {
 #if WITH_EDITOR
@@ -4131,7 +4068,7 @@ bool UCustomizableObjectSystem::IsMutableAnimInfoDebuggingEnabled() const
 
 void UCustomizableObjectSystemPrivate::UpdateResourceStreaming(float DeltaTime, bool bProcessEverything)
 {
-	GetPublic()->TickInternal(false);
+	MainTick(false);
 }
 
 
@@ -4139,25 +4076,25 @@ int32 UCustomizableObjectSystemPrivate::BlockTillAllRequestsFinished(float TimeL
 {
 	const double BlockEndTime = FPlatformTime::Seconds() + TimeLimit;
 
-	int32 RemainingWork = TNumericLimits<int32>::Max();
+	int32 LocalRemainingWork = TNumericLimits<int32>::Max();
 	
 	if (TimeLimit == 0.0f)
 	{
-		while (RemainingWork > 0)
+		while (LocalRemainingWork > 0)
 		{
-			RemainingWork = GetPublic()->TickInternal(true);
+			LocalRemainingWork = MainTick(true);
 		}
 	}
 	else
 	{
-		while (RemainingWork > 0)
+		while (LocalRemainingWork > 0)
 		{			
 			if (FPlatformTime::Seconds() > BlockEndTime)
 			{
-				return RemainingWork;
+				return LocalRemainingWork;
 			}
 			
-			RemainingWork = GetPublic()->TickInternal(true);
+			LocalRemainingWork = MainTick(true);
 		}
 	}
 
@@ -4467,3 +4404,4 @@ bool IsStreamingEnabled(const UCustomizableObject& Object, const int32 State)
 		bStreamMeshLODs &&
 		IStreamingManager::Get().IsRenderAssetStreamingEnabled(EStreamableRenderAssetType::SkeletalMesh);
 }
+
diff --git a/Engine/Plugins/Mutable/Source/CustomizableObject/Public/MuCO/CustomizableObjectSystem.h b/Engine/Plugins/Mutable/Source/CustomizableObject/Public/MuCO/CustomizableObjectSystem.h
--- a/Engine/Plugins/Mutable/Source/CustomizableObject/Public/MuCO/CustomizableObjectSystem.h
+++ b/Engine/Plugins/Mutable/Source/CustomizableObject/Public/MuCO/CustomizableObjectSystem.h
@@ -236,12 +236,6 @@ public:
 #endif
 
 private:
-	// Most of the work in this plugin happens here.
-	UE_API bool Tick(float DeltaTime);
-
-	/** Returns the number of remaining operations. */
-	UE_API int32 TickInternal(bool bBlocking);
-	
 	UE_API void DiscardInstances();
 
 public:
diff --git a/Engine/Plugins/Mutable/Source/CustomizableObjectEditor/Private/MuCOE/CustomizableObjectCompiler.cpp b/Engine/Plugins/Mutable/Source/CustomizableObjectEditor/Private/MuCOE/CustomizableObjectCompiler.cpp
--- a/Engine/Plugins/Mutable/Source/CustomizableObjectEditor/Private/MuCOE/CustomizableObjectCompiler.cpp
+++ b/Engine/Plugins/Mutable/Source/CustomizableObjectEditor/Private/MuCOE/CustomizableObjectCompiler.cpp
@@ -171,7 +171,7 @@ void FCustomizableObjectCompiler::PreloadReferencerAssets()
 		UCustomizableObjectSystem* System = UCustomizableObjectSystem::GetInstanceChecked();
 		const TSharedRef<FMutableStreamableManager>& Streamable = System->GetPrivate()->StreamableManager;
 
-		if (bAsync && !CurrentOptions.bIsCooking)
+		if (bAsync)
 		{
 			AddCompileNotification(LOCTEXT("LoadingReferencerAssets", "Loading assets"));
 
@@ -478,32 +478,6 @@ void FCustomizableObjectCompiler::EnqueueCompileRequest(const TSharedRef<FCompil
 }
 
 
-bool FCustomizableObjectCompiler::IsTickable() const
-{
-	return NumCompilationRequests > 0 || CurrentRequest;
-}
-
-
-void FCustomizableObjectCompiler::Tick(float InDeltaTime)
-{
-	MUTABLE_CPUPROFILER_SCOPE(FCustomizableObjectCompiler::Tick);
-	Tick();
-}
-
-
-TStatId FCustomizableObjectCompiler::GetStatId() const
-{
-	RETURN_QUICK_DECLARE_CYCLE_STAT(FCustomizableObjectCompiler, STATGROUP_Tickables);
-}
-
-
-void FCustomizableObjectCompiler::TickCook(float DeltaTime, bool bCookCompete)
-{
-	MUTABLE_CPUPROFILER_SCOPE(FCustomizableObjectCompiler::TickCook);
-	Tick();
-}
-
-
 bool FCustomizableObjectCompiler::IsRequestQueued(const TSharedRef<FCompilationRequest>& InCompileRequest) const
 {
 	return CurrentRequest == InCompileRequest || 
diff --git a/Engine/Plugins/Mutable/Source/CustomizableObjectEditor/Private/MuCOE/CustomizableObjectCompiler.h b/Engine/Plugins/Mutable/Source/CustomizableObjectEditor/Private/MuCOE/CustomizableObjectCompiler.h
--- a/Engine/Plugins/Mutable/Source/CustomizableObjectEditor/Private/MuCOE/CustomizableObjectCompiler.h
+++ b/Engine/Plugins/Mutable/Source/CustomizableObjectEditor/Private/MuCOE/CustomizableObjectCompiler.h
@@ -36,21 +36,12 @@ namespace UE::Mutable::Private
 UE::Mutable::Private::Ptr<UE::Mutable::Private::NodeObject> GenerateMutableRoot(const UCustomizableObject* Object, FMutableGraphGenerationContext& GenerationContext);
 
 
-class FCustomizableObjectCompiler : public TSharedFromThis<FCustomizableObjectCompiler>, public FTickableEditorObject, public FTickableCookObject, public FTickableGameObject, public FGCObject
+class FCustomizableObjectCompiler : public TSharedFromThis<FCustomizableObjectCompiler>, public FGCObject
 {
 public:
 	/** Check for pending compilation process. Returns true if an object has been updated. */
 	bool Tick(bool bBlocking = false);
 	int32 GetNumRemainingWork() const;
-
-	// FTickableGameObject interface
-	virtual ETickableTickType GetTickableTickType() const override { return ETickableTickType::Conditional; };
-	virtual bool IsTickable() const override;
-	virtual void Tick( float InDeltaTime ) override;
-	virtual TStatId GetStatId() const override;
-
-	// FTickableCookObject interface
-	virtual void TickCook(float DeltaTime, bool bCookCompete) override;
 	
 	/** Generate the Mutable Graph from the Unreal Graph. */
 	UE::Mutable::Private::Ptr<UE::Mutable::Private::Node> Export(UCustomizableObject* Object, const FCompilationOptions& Options, 
diff --git a/Engine/Plugins/Mutable/Source/CustomizableObjectEditor/Private/MuCOE/CustomizableObjectEditorModule.cpp b/Engine/Plugins/Mutable/Source/CustomizableObjectEditor/Private/MuCOE/CustomizableObjectEditorModule.cpp
--- a/Engine/Plugins/Mutable/Source/CustomizableObjectEditor/Private/MuCOE/CustomizableObjectEditorModule.cpp
+++ b/Engine/Plugins/Mutable/Source/CustomizableObjectEditor/Private/MuCOE/CustomizableObjectEditorModule.cpp
@@ -806,10 +806,9 @@ const UCustomizableObject* FCustomizableObjectEditorModule::GetRootObject(const
 }
 
 
-void FCustomizableObjectEditorModule::GetRelatedObjects(UCustomizableObject* CO, TSet<UCustomizableObject*>& OutRelated) const
+void FCustomizableObjectEditorModule::GetAllObjectsInGraph(UCustomizableObject* Object, TSet<UCustomizableObject*>& OutRelated) const
 {
-	UCustomizableObject* RootObject = GraphTraversal::GetRootObject(CO);
-	GetAllObjectsInGraph(RootObject, OutRelated);
+	GraphTraversal::GetAllObjectsInGraph(Object, OutRelated);
 }
 
 
@@ -1149,6 +1148,13 @@ void FCustomizableObjectEditorModule::PostBackwardsCompatibleFixup(UEdGraph& Gra
 }
 
 
+void FCustomizableObjectEditorModule::AddPendingLoad(UCustomizableObject* Object)
+{
+	check(IsInGameThread());
+	ObjectsPendingLoad.Add(TStrongObjectPtr(Object));
+}
+
+
 void FCustomizableObjectEditorModule::EnqueueCompileRequest(const TSharedRef<FCompilationRequest>& InCompilationRequest, bool bForceRequest)
 {
 	Compiler->EnqueueCompileRequest(InCompilationRequest, bForceRequest);
@@ -1211,8 +1217,43 @@ void FCustomizableObjectEditorModule::CompileCustomizableObject(UCustomizableObj
 
 int32 FCustomizableObjectEditorModule::Tick(bool bBlocking)
 {
+	if (TSet<TStrongObjectPtr<UCustomizableObject>>::TIterator It = ObjectsPendingLoad.CreateIterator())
+	{
+		UCustomizableObject* Object = It->Get();
+		check(IsValid(Object));
+		
+		TSet<UCustomizableObject*> RelatedObjects;
+		const bool bPendingLoad = GraphTraversal::GetAllObjectsInGraphIncremental(Object, RelatedObjects);
+
+		if (!bPendingLoad)
+		{
+			bool bHierarchyLoaded = true;
+			for (const UCustomizableObject* RelatedObject : RelatedObjects)
+			{
+				if (RelatedObject && !RelatedObject->HasAnyFlags(EObjectFlags::RF_LoadCompleted))
+				{
+					bHierarchyLoaded = false;
+					break;
+				}
+			}
+			
+			if (bHierarchyLoaded)
+			{
+				for (UCustomizableObject* RelatedObject : RelatedObjects)
+				{
+					if (ObjectsPendingLoad.Remove(TStrongObjectPtr(RelatedObject)))
+					{
+						OnUpstreamCOsLoaded(RelatedObject);
+						RelatedObject->GetPrivate()->LoadCompiledDataFromDisk();
+					}
+				}
+			}
+		}
+	}
+	
 	Compiler->Tick(bBlocking);
-	return Compiler->GetNumRemainingWork();
+
+	return Compiler->GetNumRemainingWork() + static_cast<int32>(!ObjectsPendingLoad.IsEmpty());
 }
 
 
@@ -1367,4 +1408,31 @@ void FCustomizableObjectEditorModule::OnPreBeginPIE(const bool bIsSimulatingInEd
 	}
 }
 
+
+bool FCustomizableObjectEditorModule::IsTickable() const
+{
+	return UCustomizableObjectSystem::GetInstance()->GetPrivate()->RemainingWork > 0;
+}
+
+
+void FCustomizableObjectEditorModule::Tick(float InDeltaTime)
+{
+	MUTABLE_CPUPROFILER_SCOPE(FCustomizableObjectEditorModule::Tick);
+	UCustomizableObjectSystem::GetInstance()->GetPrivate()->MainTick(false);
+}
+
+
+TStatId FCustomizableObjectEditorModule::GetStatId() const
+{
+	RETURN_QUICK_DECLARE_CYCLE_STAT(FCustomizableObjectEditorModule, STATGROUP_Tickables);
+}
+
+
+void FCustomizableObjectEditorModule::TickCook(float DeltaTime, bool bCookCompete)
+{
+	MUTABLE_CPUPROFILER_SCOPE(FCustomizableObjectEditorModule::TickCook);
+	UCustomizableObjectSystem::GetInstance()->GetPrivate()->MainTick(false);
+}
+
+
 #undef LOCTEXT_NAMESPACE
diff --git a/Engine/Plugins/Mutable/Source/CustomizableObjectEditor/Private/MuCOE/CustomizableObjectEditorModule.h b/Engine/Plugins/Mutable/Source/CustomizableObjectEditor/Private/MuCOE/CustomizableObjectEditorModule.h
--- a/Engine/Plugins/Mutable/Source/CustomizableObjectEditor/Private/MuCOE/CustomizableObjectEditorModule.h
+++ b/Engine/Plugins/Mutable/Source/CustomizableObjectEditor/Private/MuCOE/CustomizableObjectEditorModule.h
@@ -3,9 +3,10 @@
 #pragma once
 
 #include "ICustomizableObjectEditorModulePrivate.h"
+#include "TickableEditorObject.h"
+
 #include "MuCO/ICustomizableObjectEditorModule.h"
 #include "MuCO/CustomizableObjectCompilerTypes.h"
-
 #include "MuCOE/CustomizableObjectCompiler.h"
 #include "MuCOE/CustomizableObjectEditorLogger.h"
 #include "Toolkits/AssetEditorToolkit.h"
@@ -40,7 +41,7 @@ UE::DerivedData::ECachePolicy GetDerivedDataCachePolicyForEditor();
 /**
  * StaticMesh editor module
  */
-class FCustomizableObjectEditorModule : public ICustomizableObjectEditorModulePrivate
+class FCustomizableObjectEditorModule : public ICustomizableObjectEditorModulePrivate, public FTickableCookObject
 {
 public:
 	// IModuleInterface interface
@@ -55,8 +56,7 @@ public:
 	virtual FString GetCurrentReleaseVersionForObject(const UCustomizableObject& Object) const override;
 	virtual UCustomizableObject* GetRootObject(UCustomizableObject* ChildObject) const override;
 	virtual const UCustomizableObject* GetRootObject(const UCustomizableObject* ChildObject) const override;
-	virtual void GetRelatedObjects(UCustomizableObject*, TSet<UCustomizableObject*>& OutRelated) const override;
-	virtual void OnUpstreamCOsLoaded(UCustomizableObject* Object) const override;
+	virtual void GetAllObjectsInGraph(UCustomizableObject*, TSet<UCustomizableObject*>& OutRelated) const override;
 	virtual void OnUpstreamCOsLoadedFixup(UCustomizableObject* Object, int32 CustomizableObjectCustomVersion) const override;
 	virtual void BakeCustomizableObjectInstance(UCustomizableObjectInstance* InTargetInstance, const FBakingConfiguration& InBakingConfig) override;
 	virtual USkeletalMesh* GetReferenceSkeletalMesh(const UCustomizableObject& Object, const FName& ComponentName) const override;
@@ -72,10 +72,20 @@ public:
 	virtual bool IsCachedCookedPlatformDataLoaded(UCustomizableObject& Object, const ITargetPlatform* TargetPlatform) override;
 	virtual TSharedPtr<FExtensibilityManager> GetCustomizableObjectEditorToolBarExtensibilityManager() override { return CustomizableObjectEditor_ToolBarExtensibilityManager; }
 	virtual TSharedPtr<FExtensibilityManager> GetCustomizableObjectEditorMenuExtensibilityManager() override { return CustomizableObjectEditor_MenuExtensibilityManager; }
+	virtual void AddPendingLoad(UCustomizableObject* Object) override;
 
 	// ICustomizableObjectEditorModulePrivate interface
 	virtual void EnqueueCompileRequest(const TSharedRef<FCompilationRequest>& InCompilationRequest, bool bForceRequest = false) override;
 
+	// FTickableGameObject interface
+	virtual ETickableTickType GetTickableTickType() const override { return ETickableTickType::Conditional; };
+	virtual bool IsTickable() const override;
+	virtual void Tick( float InDeltaTime ) override;
+	virtual TStatId GetStatId() const override;
+
+	// FTickableCookObject interface
+	virtual void TickCook(float DeltaTime, bool bCookCompete) override;
+	
 private:	
 	TSharedPtr<FExtensibilityManager> CustomizableObjectEditor_ToolBarExtensibilityManager;
 	TSharedPtr<FExtensibilityManager> CustomizableObjectEditor_MenuExtensibilityManager;
@@ -83,6 +93,8 @@ private:
 	/** List of registered custom details to remove later. */
 	TArray<FName> RegisteredCustomDetails;
 
+	virtual void OnUpstreamCOsLoaded(UCustomizableObject* Object) const;
+	
 	/** Register Custom details. Also adds them to RegisteredCustomDetails list. */
 	void RegisterCustomDetails(FPropertyEditorModule& PropertyModule, const UClass* Class, FOnGetDetailCustomizationInstance DetailLayoutDelegate);
 
@@ -109,4 +121,7 @@ private:
 
 	/** Cook requests. */
 	TMap<TWeakObjectPtr<UCustomizableObject>, TArray<TSharedRef<FCompilationRequest>>> CookCompileRequests;
+	
+	/** List of CustomizableObjects pending to complete loading. */
+	TSet<TStrongObjectPtr<UCustomizableObject>> ObjectsPendingLoad;
 };
\ No newline at end of file
diff --git a/Engine/Plugins/Mutable/Source/CustomizableObjectEditor/Private/MuCOE/GraphTraversal.cpp b/Engine/Plugins/Mutable/Source/CustomizableObjectEditor/Private/MuCOE/GraphTraversal.cpp
--- a/Engine/Plugins/Mutable/Source/CustomizableObjectEditor/Private/MuCOE/GraphTraversal.cpp
+++ b/Engine/Plugins/Mutable/Source/CustomizableObjectEditor/Private/MuCOE/GraphTraversal.cpp
@@ -3,6 +3,7 @@
 #include "MuCOE/GraphTraversal.h"
 
 #include "AssetRegistry/AssetRegistryModule.h"
+#include "MuCO/CustomizableObjectSystemPrivate.h"
 #include "MuCO/LoadUtils.h"
 #include "MuCOE/CustomizableObjectPin.h"
 #include "MuCOE/EdGraphSchema_CustomizableObject.h"
@@ -772,10 +773,27 @@ const UEdGraphPin* FindMeshBaseSource(const UEdGraphPin& Pin, const bool bOnlyLo
 }
 
 
-void GetNodeGroupObjectNodeMappingImmersive(UCustomizableObject* Object, IAssetRegistry& AssetRegistry, TSet<UCustomizableObject*>& Visited, TMultiMap<FGuid, UCustomizableObjectNodeObject*>& Mapping)
+void GetNodeGroupObjectNodeMappingImmersive(UCustomizableObject* Object, IAssetRegistry& AssetRegistry, TSet<UCustomizableObject*>& Visited, TMultiMap<FGuid, UCustomizableObjectNodeObject*>& Mapping, bool bAsync, bool& bPendingLoad)
 {
-	Visited.Add(Object);
+	MUTABLE_CPUPROFILER_SCOPE(GetNodeGroupObjectNodeMappingImmersive);
+	
+	if (Visited.Contains(Object))
+	{
+		return;
+	}
 
+	Visited.Add(Object);
+	
+	const bool bHasGroupNodes = Object->GetPrivate()->GetSource()->Nodes.ContainsByPredicate([](const TObjectPtr<UEdGraphNode>& Node)
+	{
+		return Node.IsA(UCustomizableObjectNodeObjectGroup::StaticClass());
+	});
+	
+	if (!bHasGroupNodes)
+	{
+		return;
+	}
+	
 	TArray<FName> ArrayReferenceNames;
 	AssetRegistry.GetReferencers(*Object->GetOuter()->GetPathName(), ArrayReferenceNames, UE::AssetRegistry::EDependencyCategory::Package, UE::AssetRegistry::EDependencyQuery::Hard);
 
@@ -790,12 +808,39 @@ void GetNodeGroupObjectNodeMappingImmersive(UCustomizableObject* Object, IAssetR
 		}
 	}
 	
+	Filter.ClassPaths.Add(UCustomizableObject::StaticClass()->GetClassPathName());
+	
 	TArray<FAssetData> ArrayAssetData;
 	AssetRegistry.GetAssets(Filter, ArrayAssetData);
 
 	for (FAssetData& AssetData : ArrayAssetData)
 	{
-		UCustomizableObject* ChildObject = Cast<UCustomizableObject>(UE::Mutable::Private::LoadObject(AssetData));
+		UCustomizableObject* ChildObject = nullptr;
+		{
+			MUTABLE_CPUPROFILER_SCOPE(LoadObject);
+			
+			if (bAsync)
+			{
+				if (AssetData.IsAssetLoaded())
+				{
+					ChildObject = Cast<UCustomizableObject>(UE::Mutable::Private::LoadObject(AssetData));
+				}
+				else
+				{
+					UCustomizableObjectSystem* System = UCustomizableObjectSystem::GetInstance();
+					
+					FSoftObjectPath SoftObjectPath = AssetData.ToSoftObjectPath();
+					System->GetPrivate()->StreamableManager->RequestAsyncLoad(SoftObjectPath);
+					
+					bPendingLoad |= true;
+				}
+			}
+			else
+			{
+				ChildObject = Cast<UCustomizableObject>(UE::Mutable::Private::LoadObject(AssetData));
+			}
+		}
+		
 		if (!ChildObject)
 		{
 			continue;
@@ -812,10 +857,7 @@ void GetNodeGroupObjectNodeMappingImmersive(UCustomizableObject* Object, IAssetR
 			}
 		}
 
-		if (!Visited.Contains(ChildObject))
-		{
-			GetNodeGroupObjectNodeMappingImmersive(ChildObject, AssetRegistry, Visited, Mapping);
-		}
+		GetNodeGroupObjectNodeMappingImmersive(ChildObject, AssetRegistry, Visited, Mapping, bAsync, bPendingLoad);
 	}
 }
 
@@ -829,13 +871,14 @@ TMultiMap<FGuid, UCustomizableObjectNodeObject*> GetNodeGroupObjectNodeMapping(U
 	TSet<UCustomizableObject*> Visited;
 	TMultiMap<FGuid, UCustomizableObjectNodeObject*> Mapping;
 
-	GetNodeGroupObjectNodeMappingImmersive(Object, AssetRegistry, Visited, Mapping);
+	bool bPendingLoad = false;
+	GetNodeGroupObjectNodeMappingImmersive(Object, AssetRegistry, Visited, Mapping, false, bPendingLoad);
 	
 	return Mapping;
 }
 
 
-void GetAllObjectsInGraph(UCustomizableObject* Object, TSet<UCustomizableObject*>& OutObjects)
+void GraphTraversal::GetAllObjectsInGraph(UCustomizableObject* Object, TSet<UCustomizableObject*>& OutObjects)
 {
 	if (!Object)
 	{
@@ -847,7 +890,29 @@ void GetAllObjectsInGraph(UCustomizableObject* Object, TSet<UCustomizableObject*
 	TMultiMap<FGuid, UCustomizableObjectNodeObject*> DummyMap;
 
 	IAssetRegistry& AssetRegistry = FModuleManager::LoadModuleChecked<FAssetRegistryModule>("AssetRegistry").Get();
-	GetNodeGroupObjectNodeMappingImmersive(RootObject, AssetRegistry, OutObjects, DummyMap);
+	
+	bool bPendingLoad = false;
+	GetNodeGroupObjectNodeMappingImmersive(RootObject, AssetRegistry, OutObjects, DummyMap, false, bPendingLoad);
+}
+
+
+bool GraphTraversal::GetAllObjectsInGraphIncremental(UCustomizableObject* Object, TSet<UCustomizableObject*>& OutObjects)
+{
+	if (!Object)
+	{
+		return true;
+	}
+
+	// Search the root of the CO's graph
+	UCustomizableObject* RootObject = GraphTraversal::GetRootObject(Object);
+	TMultiMap<FGuid, UCustomizableObjectNodeObject*> DummyMap;
+
+	IAssetRegistry& AssetRegistry = FModuleManager::LoadModuleChecked<FAssetRegistryModule>("AssetRegistry").Get();
+
+	bool bPendingLoad = false;
+	GetNodeGroupObjectNodeMappingImmersive(RootObject, AssetRegistry, OutObjects, DummyMap, true, bPendingLoad);
+	
+	return bPendingLoad;
 }
 
 
diff --git a/Engine/Plugins/Mutable/Source/CustomizableObjectEditor/Private/MuCOE/GraphTraversal.h b/Engine/Plugins/Mutable/Source/CustomizableObjectEditor/Private/MuCOE/GraphTraversal.h
--- a/Engine/Plugins/Mutable/Source/CustomizableObjectEditor/Private/MuCOE/GraphTraversal.h
+++ b/Engine/Plugins/Mutable/Source/CustomizableObjectEditor/Private/MuCOE/GraphTraversal.h
@@ -111,11 +111,15 @@ const UEdGraphPin* FindMeshBaseSource(const UEdGraphPin& Pin, const bool bOnlyLo
  * @param Object Child or root Object to start the search from. */
 TMultiMap<FGuid, UCustomizableObjectNodeObject*> GetNodeGroupObjectNodeMapping(UCustomizableObject* Object);
 
-/** Returns all the Customizable Objects in a graph starting at the root object */
-void GetAllObjectsInGraph(UCustomizableObject* Object, TSet<UCustomizableObject*>& OutObjects);
-
 namespace GraphTraversal
 {
+	/** Returns all the Customizable Objects in a graph starting at the root object */
+	void GetAllObjectsInGraph(UCustomizableObject* Object, TSet<UCustomizableObject*>& OutObjects);
+
+	/** Returns all the Customizable Objects in a graph starting at the root object.
+	 * @return true if the full hierarchy could not be discovered due to UCustomizableObjects still async loading. Async loadings are called automatically. */
+	bool GetAllObjectsInGraphIncremental(UCustomizableObject* Object, TSet<UCustomizableObject*>& OutObjects);
+
     /** Return ture if the given Customizable Object is Root Object (not a Child Object). */
     bool IsRootObject(const UCustomizableObject& Object); 
 }
diff --git a/Engine/Plugins/Mutable/Source/CustomizableObjectEditor/Private/MuCOE/Nodes/CustomizableObjectNode.cpp b/Engine/Plugins/Mutable/Source/CustomizableObjectEditor/Private/MuCOE/Nodes/CustomizableObjectNode.cpp
--- a/Engine/Plugins/Mutable/Source/CustomizableObjectEditor/Private/MuCOE/Nodes/CustomizableObjectNode.cpp
+++ b/Engine/Plugins/Mutable/Source/CustomizableObjectEditor/Private/MuCOE/Nodes/CustomizableObjectNode.cpp
@@ -442,7 +442,7 @@ bool UCustomizableObjectNode::FindNodeForInternalTag(const FString& Tag, UCustom
 	UCustomizableObject* RootObject = GraphTraversal::GetRootObject(ThisNodeObject);
 
 	TSet<UCustomizableObject*> AllCustomizableObject;
-	GetAllObjectsInGraph(RootObject, AllCustomizableObject);
+	GraphTraversal::GetAllObjectsInGraph(RootObject, AllCustomizableObject);
 
 	for (UCustomizableObject* CustObject : AllCustomizableObject)
 	{
diff --git a/Engine/Plugins/Mutable/Source/CustomizableObjectEditor/Private/MuCOE/Nodes/CustomizableObjectNodeModifierBase.cpp b/Engine/Plugins/Mutable/Source/CustomizableObjectEditor/Private/MuCOE/Nodes/CustomizableObjectNodeModifierBase.cpp
--- a/Engine/Plugins/Mutable/Source/CustomizableObjectEditor/Private/MuCOE/Nodes/CustomizableObjectNodeModifierBase.cpp
+++ b/Engine/Plugins/Mutable/Source/CustomizableObjectEditor/Private/MuCOE/Nodes/CustomizableObjectNodeModifierBase.cpp
@@ -124,7 +124,7 @@ void UCustomizableObjectNodeModifierBase::GetPossiblyModifiedNodes(TArray<UCusto
 	UCustomizableObject* RootObject = GraphTraversal::GetRootObject(ThisNodeObject);
 
 	TSet<UCustomizableObject*> AllCustomizableObject;
-	GetAllObjectsInGraph(RootObject, AllCustomizableObject);
+	GraphTraversal::GetAllObjectsInGraph(RootObject, AllCustomizableObject);
 
 	for (const UCustomizableObject* CustObject : AllCustomizableObject)
 	{
diff --git a/Engine/Plugins/Mutable/Source/CustomizableObjectEditor/Private/MuCOE/Nodes/CustomizableObjectNodeModifierMorphMeshSectionDetails.cpp b/Engine/Plugins/Mutable/Source/CustomizableObjectEditor/Private/MuCOE/Nodes/CustomizableObjectNodeModifierMorphMeshSectionDetails.cpp
--- a/Engine/Plugins/Mutable/Source/CustomizableObjectEditor/Private/MuCOE/Nodes/CustomizableObjectNodeModifierMorphMeshSectionDetails.cpp
+++ b/Engine/Plugins/Mutable/Source/CustomizableObjectEditor/Private/MuCOE/Nodes/CustomizableObjectNodeModifierMorphMeshSectionDetails.cpp
@@ -141,7 +141,7 @@ void FCustomizableObjectNodeModifierMorphMeshSectionDetails::RefreshMorphOptions
 		UCustomizableObject* RootObject = GraphTraversal::GetRootObject(ThisNodeObject);
 
 		TSet<UCustomizableObject*> AllCustomizableObject;
-		GetAllObjectsInGraph(RootObject, AllCustomizableObject);
+		GraphTraversal::GetAllObjectsInGraph(RootObject, AllCustomizableObject);
 
 		for (const UCustomizableObject* CustObject : AllCustomizableObject)
 		{
diff --git a/Engine/Plugins/Mutable/Source/CustomizableObjectEditor/Private/MuCOE/Nodes/CustomizableObjectNodeObjectDetails.cpp b/Engine/Plugins/Mutable/Source/CustomizableObjectEditor/Private/MuCOE/Nodes/CustomizableObjectNodeObjectDetails.cpp
--- a/Engine/Plugins/Mutable/Source/CustomizableObjectEditor/Private/MuCOE/Nodes/CustomizableObjectNodeObjectDetails.cpp
+++ b/Engine/Plugins/Mutable/Source/CustomizableObjectEditor/Private/MuCOE/Nodes/CustomizableObjectNodeObjectDetails.cpp
@@ -368,7 +368,7 @@ void FCustomizableObjectNodeObjectDetails::FillParameterNamesArray()
 	TSet<UCustomizableObject*> CustomObjectTree;
 
 	// Get and load the whole tree of customizable object
-	GetAllObjectsInGraph(RootObjet, CustomObjectTree);
+	GraphTraversal::GetAllObjectsInGraph(RootObjet, CustomObjectTree);
 
 	// Array to store all the ids of group nodes of type toggle
 	TArray<FGuid> ToggleGroupObjectIds;
diff --git a/Engine/Plugins/Mutable/Source/CustomizableObjectEditor/Private/MuCOE/SCustomizableObjectEditorTagExplorer.cpp b/Engine/Plugins/Mutable/Source/CustomizableObjectEditor/Private/MuCOE/SCustomizableObjectEditorTagExplorer.cpp
--- a/Engine/Plugins/Mutable/Source/CustomizableObjectEditor/Private/MuCOE/SCustomizableObjectEditorTagExplorer.cpp
+++ b/Engine/Plugins/Mutable/Source/CustomizableObjectEditor/Private/MuCOE/SCustomizableObjectEditorTagExplorer.cpp
@@ -121,7 +121,7 @@ TSharedRef<SWidget> SCustomizableObjectEditorTagExplorer::OnGetTagsMenuContent()
 		TArray<FString> Tags;
 		TSet<UCustomizableObject*> CustomizableObjectTree;
 
-		GetAllObjectsInGraph(CustomizableObject, CustomizableObjectTree);
+		GraphTraversal::GetAllObjectsInGraph(CustomizableObject, CustomizableObjectTree);
 
 		for (const UCustomizableObject* CustObject : CustomizableObjectTree)
 		{
diff --git a/Engine/Plugins/Mutable/Source/CustomizableObjectEditor/Private/MuCOE/SMutableTagListWidget.cpp b/Engine/Plugins/Mutable/Source/CustomizableObjectEditor/Private/MuCOE/SMutableTagListWidget.cpp
--- a/Engine/Plugins/Mutable/Source/CustomizableObjectEditor/Private/MuCOE/SMutableTagListWidget.cpp
+++ b/Engine/Plugins/Mutable/Source/CustomizableObjectEditor/Private/MuCOE/SMutableTagListWidget.cpp
@@ -185,7 +185,7 @@ void SMutableTagComboBox::RefreshOptions()
 	UCustomizableObject* RootObject = GraphTraversal::GetRootObject(ThisNodeObject);
 
 	TSet<UCustomizableObject*> AllCustomizableObject;
-	GetAllObjectsInGraph(RootObject, AllCustomizableObject);
+	GraphTraversal::GetAllObjectsInGraph(RootObject, AllCustomizableObject);
 
 	TMap<UEdGraphNode*, TSharedPtr<SMutableSearchComboBox::FFilteredOption>> AddedOptions;
 
diff --git a/Engine/Source/Runtime/Core/Public/GenericPlatform/GenericPlatformTime.h b/Engine/Source/Runtime/Core/Public/GenericPlatform/GenericPlatformTime.h
--- a/Engine/Source/Runtime/Core/Public/GenericPlatform/GenericPlatformTime.h
+++ b/Engine/Source/Runtime/Core/Public/GenericPlatform/GenericPlatformTime.h
@@ -118,7 +118,7 @@ struct FGenericPlatformTime
 		return false;
 	}
 
-	/** Updates current thread CPU utilization, calling is user defined per-thread (unused float parameter, is for FTicker compatibility). */
+	/** Updates current thread CPU utilization, calling is user defined per-thread (unused float parameter, is for FCookTicker compatibility). */
 	static bool UpdateThreadCPUTime(float = 0.0)
 	{
 		return false;
