[Mutable] Replace custom hash functions with engine ones.

#Synced-CL 49557694

diff --git a/Engine/Plugins/Mutable/Source/CustomizableObject/Private/MuCO/CustomizableObjectInstance.cpp b/Engine/Plugins/Mutable/Source/CustomizableObject/Private/MuCO/CustomizableObjectInstance.cpp
--- a/Engine/Plugins/Mutable/Source/CustomizableObject/Private/MuCO/CustomizableObjectInstance.cpp
+++ b/Engine/Plugins/Mutable/Source/CustomizableObject/Private/MuCO/CustomizableObjectInstance.cpp
@@ -6344,7 +6344,7 @@ void UCustomizableObjectInstance::SetRequestedLODs(const TMap<FName, uint8>& InM
 		}
 	}
 	
-	if (bMinLODChanged || FirstRequestedLOD != MutableUpdateCandidate.FirstRequestedLOD)
+	if (bMinLODChanged || !FirstRequestedLOD.OrderIndependentCompareEqual(MutableUpdateCandidate.FirstRequestedLOD))
 	{
 		// TODO: Remove this flag as it will become redundant with the new InOutRequestedUpdates system
 		PrivateData->SetCOInstanceFlags(PendingLODsUpdate);
diff --git a/Engine/Plugins/Mutable/Source/MutableRuntime/Internal/MuR/ImageTypes.h b/Engine/Plugins/Mutable/Source/MutableRuntime/Internal/MuR/ImageTypes.h
--- a/Engine/Plugins/Mutable/Source/MutableRuntime/Internal/MuR/ImageTypes.h
+++ b/Engine/Plugins/Mutable/Source/MutableRuntime/Internal/MuR/ImageTypes.h
@@ -71,6 +71,10 @@ namespace UE::Mutable::Private
 	};
 	MUTABLE_DEFINE_ENUM_SERIALISABLE(EImageFormat);
 
+	inline uint32 GetTypeHash(const EImageFormat& Format)
+	{
+		return static_cast<uint32>(Format);
+	}
 
 	struct FImageDesc
 	{
@@ -144,6 +148,11 @@ namespace UE::Mutable::Private
 	};
 	MUTABLE_DEFINE_ENUM_SERIALISABLE(EAddressMode);
 	
+	inline uint32 GetTypeHash(const EAddressMode& Mode)
+	{
+		return static_cast<uint32>(Mode);
+	}
+	
 	enum class EMipmapFilterType
 	{
 		Unfiltered,
diff --git a/Engine/Plugins/Mutable/Source/MutableRuntime/Internal/MuR/Mesh.h b/Engine/Plugins/Mutable/Source/MutableRuntime/Internal/MuR/Mesh.h
--- a/Engine/Plugins/Mutable/Source/MutableRuntime/Internal/MuR/Mesh.h
+++ b/Engine/Plugins/Mutable/Source/MutableRuntime/Internal/MuR/Mesh.h
@@ -166,6 +166,11 @@ namespace UE::Mutable::Private
 		ReshapeMaskWeight = 1,
 		ReshapeClusterId = 2
 	};
+	
+	inline uint32 GetTypeHash(const EVertexColorUsage& Usage)
+	{
+		return GetTypeHash(static_cast<uint32>(Usage));
+	}
 
 	enum class EMeshCopyFlags : uint32
 	{
diff --git a/Engine/Plugins/Mutable/Source/MutableRuntime/Internal/MuR/Operations.h b/Engine/Plugins/Mutable/Source/MutableRuntime/Internal/MuR/Operations.h
--- a/Engine/Plugins/Mutable/Source/MutableRuntime/Internal/MuR/Operations.h
+++ b/Engine/Plugins/Mutable/Source/MutableRuntime/Internal/MuR/Operations.h
@@ -403,8 +403,14 @@ namespace UE::Mutable::Private
 
         //!
         COUNT
-
     };
+	
+	
+	inline uint32 GetTypeHash(const EOpType& OpType)
+	{
+		return static_cast<uint32>(OpType);
+	}
+
 
 	enum class EMeshBindShapeFlags : uint32
 	{
diff --git a/Engine/Plugins/Mutable/Source/MutableRuntime/Internal/MuR/Parameters.h b/Engine/Plugins/Mutable/Source/MutableRuntime/Internal/MuR/Parameters.h
--- a/Engine/Plugins/Mutable/Source/MutableRuntime/Internal/MuR/Parameters.h
+++ b/Engine/Plugins/Mutable/Source/MutableRuntime/Internal/MuR/Parameters.h
@@ -70,8 +70,13 @@ namespace UE::Mutable::Private
 
         /** Utility enumeration value, not really a parameter type. */
         Count
-
     };
+	
+	
+	inline uint32 GetTypeHash(const EParameterType& Type)
+	{
+		return static_cast<uint32>(Type);
+	}
 
 
     /** Types of 3D projectors. */
diff --git a/Engine/Plugins/Mutable/Source/MutableRuntime/Internal/MuR/Types.h b/Engine/Plugins/Mutable/Source/MutableRuntime/Internal/MuR/Types.h
--- a/Engine/Plugins/Mutable/Source/MutableRuntime/Internal/MuR/Types.h
+++ b/Engine/Plugins/Mutable/Source/MutableRuntime/Internal/MuR/Types.h
@@ -7,6 +7,7 @@
 #include "Containers/Array.h"
 #include "Containers/UnrealString.h"
 #include "StructUtils/InstancedStruct.h"
+#include "Curves/RichCurve.h"
 
 #include "Types.generated.h"
 
@@ -36,38 +37,15 @@ namespace UE::Mutable::Private
 {
 	class FOutputArchive;
 	class FInputArchive;
+
 	
 	enum class EMemoryInitPolicy
 	{
 		Uninitialized,
 		Zeroed,
 	};
-}
-
-
-//
-template<typename T>
-inline uint32 GetTypeHash(const TArray<T>& Key)
-{
-	uint32 Hash = 0;
-	for (const T& e : Key)
-	{
-		Hash = HashCombine(Hash, GetTypeHash(e));
-	}
-	return Hash;
-}
-
-
-//
-template<typename K, typename V>
-inline bool operator==(const TMap<K,V>& A, const TMap<K,V>& B )
-{	
-	return A.OrderIndependentCompareEqual(B);
-}
-
-
-namespace UE::Mutable::Private
-{
+	
+	
 	inline uint32 GetTypeHash(const FInstancedStruct& InstancedStruct)
 	{
 		if (!InstancedStruct.IsValid())
@@ -80,4 +58,55 @@ namespace UE::Mutable::Private
 
 		return HashCombine(ValueHash, PointerHash(ScriptedStruct));
 	}
+	
+	
+	inline uint32 GetTypeHash(const float& Key)
+	{
+		return static_cast<uint32>(Key);
+	}
+	
+	
+	inline uint32 GetTypeHash(const FRichCurveKey& Key)
+	{
+		uint32 Result = 0;
+		
+		Result = HashCombineFast(Result, GetTypeHash(Key.InterpMode));
+		Result = HashCombineFast(Result, GetTypeHash(Key.TangentMode));
+		Result = HashCombineFast(Result, GetTypeHash(Key.TangentWeightMode));
+		Result = HashCombineFast(Result, GetTypeHash(Key.Time));
+		Result = HashCombineFast(Result, GetTypeHash(Key.Value));
+		Result = HashCombineFast(Result, GetTypeHash(Key.ArriveTangent));
+		Result = HashCombineFast(Result, GetTypeHash(Key.ArriveTangentWeight));
+		Result = HashCombineFast(Result, GetTypeHash(Key.LeaveTangent));
+		Result = HashCombineFast(Result, GetTypeHash(Key.LeaveTangentWeight));
+		 	
+		return Result;
+	}
+	
+	
+	template<typename T>
+	uint32 GetTypeHash(const TArray<T>& Key)
+	{
+		uint32 Hash = 0;
+		
+		for (const T& Element : Key)
+		{
+			Hash = HashCombine(Hash, GetTypeHash(Element));
+		}
+		
+		return Hash;
+	}
+
+
+	template<typename K, typename V>
+	bool operator==(const TMap<K,V>& A, const TMap<K,V>& B )
+	{	
+		return A.OrderIndependentCompareEqual(B);
+	}
+	
+	
+	inline uint32 GetTypeHash(const FRichCurve& Curve)
+	{
+		return GetTypeHash(Curve.Keys);
+	}
 }
diff --git a/Engine/Plugins/Mutable/Source/MutableTools/Internal/MuT/AST.h b/Engine/Plugins/Mutable/Source/MutableTools/Internal/MuT/AST.h
--- a/Engine/Plugins/Mutable/Source/MutableTools/Internal/MuT/AST.h
+++ b/Engine/Plugins/Mutable/Source/MutableTools/Internal/MuT/AST.h
@@ -41,15 +41,6 @@ namespace std
 
 }
 
-namespace 
-{
-	// TODO: Replace with UE hashing
-	template <class T> inline void hash_combine(uint64& seed, const T& v) {
-		std::hash<T> hasher;
-		seed ^= hasher(v) + 0x9e3779b9 + (seed << 6) + (seed >> 2);
-	}
-}
-
 
 namespace UE::Mutable::Private
 {
@@ -524,6 +515,12 @@ namespace UE::Mutable::Private
         inline void ClearParent();
     };
 
+	
+	inline uint32 GetTypeHash(const ASTChild& Child)
+	{
+		return PointerHash(Child.child().get());
+	}
+
 
     //---------------------------------------------------------------------------------------------
     //! Abstract Syntax Tree of operations in the mutable virtual machine.
@@ -566,7 +563,7 @@ namespace UE::Mutable::Private
         virtual bool operator==( const ASTOp& other ) const;
 
         //! Hint hash method for op sorting and containers. It's a hash of the actual operation.
-        virtual uint64 Hash() const = 0;
+        virtual uint32 Hash() const = 0;
 
         //! Shallow clone. New node will have no parents but reference to the same children.
 		using MapChildFunc = TFunction<Ptr<ASTOp>(const Ptr<ASTOp>&)>;
@@ -937,6 +934,7 @@ namespace UE::Mutable::Private
         Ptr<DERIVED> t = static_cast<DERIVED*>(c.get());
 		return t;
     }
+	
 
     //---------------------------------------------------------------------------------------------
     //!
diff --git a/Engine/Plugins/Mutable/Source/MutableTools/Private/MuT/ASTOpAddExtensionData.cpp b/Engine/Plugins/Mutable/Source/MutableTools/Private/MuT/ASTOpAddExtensionData.cpp
--- a/Engine/Plugins/Mutable/Source/MutableTools/Private/MuT/ASTOpAddExtensionData.cpp
+++ b/Engine/Plugins/Mutable/Source/MutableTools/Private/MuT/ASTOpAddExtensionData.cpp
@@ -40,12 +40,12 @@ bool ASTOpAddExtensionData::IsEqual(const ASTOp& OtherUntyped) const
 	return false;
 }
 
-uint64 ASTOpAddExtensionData::Hash() const
+uint32 ASTOpAddExtensionData::Hash() const
 {
-	uint64 Result = std::hash<uint64>()(uint64(EOpType::IN_ADDEXTENSIONDATA));
+	uint32 Result = GetTypeHash(GetOpType());
 	
-	hash_combine(Result, Instance.child().get());
-	hash_combine(Result, ExtensionData.child().get());
+	Result = HashCombineFast(Result, GetTypeHash(Instance));
+	Result = HashCombineFast(Result, GetTypeHash(ExtensionData));
 	
 	return Result;
 }
diff --git a/Engine/Plugins/Mutable/Source/MutableTools/Private/MuT/ASTOpAddExtensionData.h b/Engine/Plugins/Mutable/Source/MutableTools/Private/MuT/ASTOpAddExtensionData.h
--- a/Engine/Plugins/Mutable/Source/MutableTools/Private/MuT/ASTOpAddExtensionData.h
+++ b/Engine/Plugins/Mutable/Source/MutableTools/Private/MuT/ASTOpAddExtensionData.h
@@ -32,7 +32,7 @@ public:
 	~ASTOpAddExtensionData();
 
 	EOpType GetOpType() const override { return EOpType::IN_ADDEXTENSIONDATA; }
-	uint64 Hash() const override;
+	uint32 Hash() const override;
 	void ForEachChild(const TFunctionRef<void(ASTChild&)> F) override;
 	bool IsEqual(const ASTOp& OtherUntyped) const override;
 	Ptr<ASTOp> Clone(MapChildFuncRef MapChild) const override;
diff --git a/Engine/Plugins/Mutable/Source/MutableTools/Private/MuT/ASTOpAddLOD.cpp b/Engine/Plugins/Mutable/Source/MutableTools/Private/MuT/ASTOpAddLOD.cpp
--- a/Engine/Plugins/Mutable/Source/MutableTools/Private/MuT/ASTOpAddLOD.cpp
+++ b/Engine/Plugins/Mutable/Source/MutableTools/Private/MuT/ASTOpAddLOD.cpp
@@ -37,14 +37,16 @@ namespace UE::Mutable::Private
 	}
 
 
-	uint64 ASTOpAddLOD::Hash() const
+	uint32 ASTOpAddLOD::Hash() const
 	{
-		uint64 res = std::hash<uint64>()(uint64(EOpType::IN_ADDLOD));
-		for (auto& c : lods)
+		uint32 Result = GetTypeHash(GetOpType());
+
+		for (auto& Child : lods)
 		{
-			hash_combine(res, c.child().get());
+			Result = HashCombineFast(Result, GetTypeHash(Child));
 		}
-		return res;
+		
+		return Result;
 	}
 
 
diff --git a/Engine/Plugins/Mutable/Source/MutableTools/Private/MuT/ASTOpAddLOD.h b/Engine/Plugins/Mutable/Source/MutableTools/Private/MuT/ASTOpAddLOD.h
--- a/Engine/Plugins/Mutable/Source/MutableTools/Private/MuT/ASTOpAddLOD.h
+++ b/Engine/Plugins/Mutable/Source/MutableTools/Private/MuT/ASTOpAddLOD.h
@@ -27,7 +27,7 @@ struct FProgram;
 		~ASTOpAddLOD();
 
 		EOpType GetOpType() const override { return EOpType::IN_ADDLOD; }
-		uint64 Hash() const override;
+		uint32 Hash() const override;
 		void ForEachChild(const TFunctionRef<void(ASTChild&)>) override;
 		bool IsEqual(const ASTOp& otherUntyped) const override;
 		Ptr<ASTOp> Clone(MapChildFuncRef mapChild) const override;
diff --git a/Engine/Plugins/Mutable/Source/MutableTools/Private/MuT/ASTOpAddMaterial.cpp b/Engine/Plugins/Mutable/Source/MutableTools/Private/MuT/ASTOpAddMaterial.cpp
--- a/Engine/Plugins/Mutable/Source/MutableTools/Private/MuT/ASTOpAddMaterial.cpp
+++ b/Engine/Plugins/Mutable/Source/MutableTools/Private/MuT/ASTOpAddMaterial.cpp
@@ -39,12 +39,12 @@ bool ASTOpAddMaterial::IsEqual(const ASTOp& OtherUntyped) const
 	return false;
 }
 
-uint64 ASTOpAddMaterial::Hash() const
+uint32 ASTOpAddMaterial::Hash() const
 {
-	uint64 Result = std::hash<uint64>()(uint64(Type));
+	uint32 Result = GetTypeHash(GetOpType());
 	
-	hash_combine(Result, Instance.child().get());
-	hash_combine(Result, Material.child().get());
+	Result = HashCombineFast(Result, GetTypeHash(Instance));
+	Result = HashCombineFast(Result, GetTypeHash(Material));
 	
 	return Result;
 }
diff --git a/Engine/Plugins/Mutable/Source/MutableTools/Private/MuT/ASTOpAddOverrideMaterial.cpp b/Engine/Plugins/Mutable/Source/MutableTools/Private/MuT/ASTOpAddOverrideMaterial.cpp
--- a/Engine/Plugins/Mutable/Source/MutableTools/Private/MuT/ASTOpAddOverrideMaterial.cpp
+++ b/Engine/Plugins/Mutable/Source/MutableTools/Private/MuT/ASTOpAddOverrideMaterial.cpp
@@ -42,13 +42,13 @@ namespace UE::Mutable::Private
 	}
 
 	
-	uint64 ASTOpAddOverrideMaterial::Hash() const
+	uint32 ASTOpAddOverrideMaterial::Hash() const
 	{
-		uint32 Result = GetTypeHash(static_cast<uint32>(GetOpType()));
+		uint32 Result = GetTypeHash(GetOpType());
 		
-		Result = HashCombine(Result, PointerHash(Instance.child().get()));
-		Result = HashCombine(Result, PointerHash(Material.child().get()));
-		Result = HashCombine(Result, GetTypeHash(SlotName));
+		Result = HashCombineFast(Result, GetTypeHash(Instance));
+		Result = HashCombineFast(Result, GetTypeHash(Material));
+		Result = HashCombineFast(Result, GetTypeHash(SlotName));
 		
 		return Result;
 	}
diff --git a/Engine/Plugins/Mutable/Source/MutableTools/Private/MuT/ASTOpAddOverrideMaterial.h b/Engine/Plugins/Mutable/Source/MutableTools/Private/MuT/ASTOpAddOverrideMaterial.h
--- a/Engine/Plugins/Mutable/Source/MutableTools/Private/MuT/ASTOpAddOverrideMaterial.h
+++ b/Engine/Plugins/Mutable/Source/MutableTools/Private/MuT/ASTOpAddOverrideMaterial.h
@@ -26,7 +26,7 @@ namespace UE::Mutable::Private
 		virtual ~ASTOpAddOverrideMaterial() override;
 
 		virtual EOpType GetOpType() const override;
-		virtual uint64 Hash() const override;
+		virtual uint32 Hash() const override;
 		virtual void ForEachChild(const TFunctionRef<void(ASTChild&)> Func) override;
 		virtual bool IsEqual(const ASTOp& OtherUntyped) const override;
 		virtual Ptr<ASTOp> Clone(MapChildFuncRef MapChild) const override;
diff --git a/Engine/Plugins/Mutable/Source/MutableTools/Private/MuT/ASTOpAddSkeletalMesh.cpp b/Engine/Plugins/Mutable/Source/MutableTools/Private/MuT/ASTOpAddSkeletalMesh.cpp
--- a/Engine/Plugins/Mutable/Source/MutableTools/Private/MuT/ASTOpAddSkeletalMesh.cpp
+++ b/Engine/Plugins/Mutable/Source/MutableTools/Private/MuT/ASTOpAddSkeletalMesh.cpp
@@ -33,11 +33,12 @@ namespace UE::Mutable::Private
 	}
 
 		
-	uint64 ASTOpAddSkeletalMesh::Hash() const
+	uint32 ASTOpAddSkeletalMesh::Hash() const
 	{
-		uint64 Result = std::hash<void*>()(Instance.child().get());
+		uint32 Result = GetTypeHash(GetOpType());
 		
-		hash_combine(Result, SkeletalMesh.child().get());
+		Result = HashCombineFast(Result, GetTypeHash(Instance));
+		Result = HashCombineFast(Result, GetTypeHash(SkeletalMesh));
 		
 		return Result;
 	}
diff --git a/Engine/Plugins/Mutable/Source/MutableTools/Private/MuT/ASTOpBoolAnd.cpp b/Engine/Plugins/Mutable/Source/MutableTools/Private/MuT/ASTOpBoolAnd.cpp
--- a/Engine/Plugins/Mutable/Source/MutableTools/Private/MuT/ASTOpBoolAnd.cpp
+++ b/Engine/Plugins/Mutable/Source/MutableTools/Private/MuT/ASTOpBoolAnd.cpp
@@ -37,10 +37,13 @@ namespace UE::Mutable::Private
 	}
 
 
-	uint64 ASTOpBoolAnd::Hash() const
+	uint32 ASTOpBoolAnd::Hash() const
 	{
-		uint64 Result = std::hash<void*>()(A.child().get());
-		hash_combine(Result, B.child().get());
+		uint32 Result = GetTypeHash(GetOpType());
+		
+		Result = HashCombineFast(Result, GetTypeHash(A));
+		Result = HashCombineFast(Result, GetTypeHash(B));
+		
 		return Result;
 	}
 
diff --git a/Engine/Plugins/Mutable/Source/MutableTools/Private/MuT/ASTOpBoolAnd.h b/Engine/Plugins/Mutable/Source/MutableTools/Private/MuT/ASTOpBoolAnd.h
--- a/Engine/Plugins/Mutable/Source/MutableTools/Private/MuT/ASTOpBoolAnd.h
+++ b/Engine/Plugins/Mutable/Source/MutableTools/Private/MuT/ASTOpBoolAnd.h
@@ -28,7 +28,7 @@ namespace UE::Mutable::Private
 
 		// ASTOp interface
 		virtual EOpType GetOpType() const override { return EOpType::BO_AND; }
-		virtual uint64 Hash() const override;
+		virtual uint32 Hash() const override;
 		virtual bool IsEqual(const ASTOp&) const override;
 		virtual Ptr<ASTOp> Clone(MapChildFuncRef) const override;
 		virtual void ForEachChild(const TFunctionRef<void(ASTChild&)>) override;
diff --git a/Engine/Plugins/Mutable/Source/MutableTools/Private/MuT/ASTOpBoolEqualIntConst.cpp b/Engine/Plugins/Mutable/Source/MutableTools/Private/MuT/ASTOpBoolEqualIntConst.cpp
--- a/Engine/Plugins/Mutable/Source/MutableTools/Private/MuT/ASTOpBoolEqualIntConst.cpp
+++ b/Engine/Plugins/Mutable/Source/MutableTools/Private/MuT/ASTOpBoolEqualIntConst.cpp
@@ -51,11 +51,14 @@ namespace UE::Mutable::Private
 	}
 
 
-	uint64 ASTOpBoolEqualIntConst::Hash() const
+	uint32 ASTOpBoolEqualIntConst::Hash() const
 	{
-		uint64 Res = std::hash<ASTOp*>()(Value.child().get());
-		hash_combine(Res, Constant);
-		return Res;
+		uint32 Result = GetTypeHash(GetOpType());
+		
+		Result = HashCombineFast(Result, GetTypeHash(Value));
+		Result = HashCombineFast(Result, GetTypeHash(Constant));
+
+		return Result;
 	}
 
 
diff --git a/Engine/Plugins/Mutable/Source/MutableTools/Private/MuT/ASTOpBoolNot.cpp b/Engine/Plugins/Mutable/Source/MutableTools/Private/MuT/ASTOpBoolNot.cpp
--- a/Engine/Plugins/Mutable/Source/MutableTools/Private/MuT/ASTOpBoolNot.cpp
+++ b/Engine/Plugins/Mutable/Source/MutableTools/Private/MuT/ASTOpBoolNot.cpp
@@ -35,9 +35,12 @@ namespace UE::Mutable::Private
 	}
 
 
-	uint64 ASTOpBoolNot::Hash() const
+	uint32 ASTOpBoolNot::Hash() const
 	{
-		uint64 Result = std::hash<void*>()(A.child().get());
+		uint32 Result = GetTypeHash(GetOpType());
+
+		Result = HashCombineFast(Result, GetTypeHash(A));
+		
 		return Result;
 	}
 
diff --git a/Engine/Plugins/Mutable/Source/MutableTools/Private/MuT/ASTOpBoolOr.cpp b/Engine/Plugins/Mutable/Source/MutableTools/Private/MuT/ASTOpBoolOr.cpp
--- a/Engine/Plugins/Mutable/Source/MutableTools/Private/MuT/ASTOpBoolOr.cpp
+++ b/Engine/Plugins/Mutable/Source/MutableTools/Private/MuT/ASTOpBoolOr.cpp
@@ -37,10 +37,13 @@ namespace UE::Mutable::Private
 	}
 
 
-	uint64 ASTOpBoolOr::Hash() const
+	uint32 ASTOpBoolOr::Hash() const
 	{
-		uint64 Result = std::hash<void*>()(A.child().get());
-		hash_combine(Result, B.child().get());
+		uint32 Result = GetTypeHash(GetOpType());
+
+		Result = HashCombineFast(Result, GetTypeHash(A));
+		Result = HashCombineFast(Result, GetTypeHash(B));
+		
 		return Result;
 	}
 
diff --git a/Engine/Plugins/Mutable/Source/MutableTools/Private/MuT/ASTOpColorArithmetic.cpp b/Engine/Plugins/Mutable/Source/MutableTools/Private/MuT/ASTOpColorArithmetic.cpp
--- a/Engine/Plugins/Mutable/Source/MutableTools/Private/MuT/ASTOpColorArithmetic.cpp
+++ b/Engine/Plugins/Mutable/Source/MutableTools/Private/MuT/ASTOpColorArithmetic.cpp
@@ -37,11 +37,14 @@ namespace UE::Mutable::Private
 	}
 
 
-	uint64 ASTOpColorArithmetic::Hash() const
+	uint32 ASTOpColorArithmetic::Hash() const
 	{
-		uint64 Result = std::hash<void*>()(A.child().get());
-		hash_combine(Result, B.child().get());
-		hash_combine(Result, Operation);
+		uint32 Result = GetTypeHash(GetOpType());
+
+		Result = HashCombineFast(Result, GetTypeHash(A));
+		Result = HashCombineFast(Result, GetTypeHash(B));
+		Result = HashCombineFast(Result, Operation);
+
 		return Result;
 	}
 
diff --git a/Engine/Plugins/Mutable/Source/MutableTools/Private/MuT/ASTOpColorFromScalars.cpp b/Engine/Plugins/Mutable/Source/MutableTools/Private/MuT/ASTOpColorFromScalars.cpp
--- a/Engine/Plugins/Mutable/Source/MutableTools/Private/MuT/ASTOpColorFromScalars.cpp
+++ b/Engine/Plugins/Mutable/Source/MutableTools/Private/MuT/ASTOpColorFromScalars.cpp
@@ -57,13 +57,16 @@ namespace UE::Mutable::Private
 	}
 
 
-	uint64 ASTOpColorFromScalars::Hash() const
+	uint32 ASTOpColorFromScalars::Hash() const
 	{
-		uint64 res = std::hash<void*>()(V[0].child().get());
-		hash_combine(res, std::hash<void*>()(V[1].child().get()));
-		hash_combine(res, std::hash<void*>()(V[2].child().get()));
-		hash_combine(res, std::hash<void*>()(V[3].child().get()));
-		return res;
+		uint32 Result = GetTypeHash(GetOpType());
+		
+		Result = HashCombineFast(Result, GetTypeHash(V[0]));
+		Result = HashCombineFast(Result, GetTypeHash(V[1]));
+		Result = HashCombineFast(Result, GetTypeHash(V[2]));
+		Result = HashCombineFast(Result, GetTypeHash(V[3]));
+
+		return Result;
 	}
 
 
diff --git a/Engine/Plugins/Mutable/Source/MutableTools/Private/MuT/ASTOpColorSampleImage.cpp b/Engine/Plugins/Mutable/Source/MutableTools/Private/MuT/ASTOpColorSampleImage.cpp
--- a/Engine/Plugins/Mutable/Source/MutableTools/Private/MuT/ASTOpColorSampleImage.cpp
+++ b/Engine/Plugins/Mutable/Source/MutableTools/Private/MuT/ASTOpColorSampleImage.cpp
@@ -39,12 +39,15 @@ namespace UE::Mutable::Private
 	}
 
 
-	uint64 ASTOpColorSampleImage::Hash() const
+	uint32 ASTOpColorSampleImage::Hash() const
 	{
-		uint64 Result = std::hash<void*>()(Image.child().get());
-		hash_combine(Result, X.child().get());
-		hash_combine(Result, Y.child().get());
-		hash_combine(Result, Filter);
+		uint32 Result = GetTypeHash(GetOpType());
+
+		Result = HashCombineFast(Result, GetTypeHash(Image));
+		Result = HashCombineFast(Result, GetTypeHash(X));
+		Result = HashCombineFast(Result, GetTypeHash(Y));
+		Result = HashCombineFast(Result, GetTypeHash(Filter));
+
 		return Result;
 	}
 
diff --git a/Engine/Plugins/Mutable/Source/MutableTools/Private/MuT/ASTOpColorSampleImage.h b/Engine/Plugins/Mutable/Source/MutableTools/Private/MuT/ASTOpColorSampleImage.h
--- a/Engine/Plugins/Mutable/Source/MutableTools/Private/MuT/ASTOpColorSampleImage.h
+++ b/Engine/Plugins/Mutable/Source/MutableTools/Private/MuT/ASTOpColorSampleImage.h
@@ -20,7 +20,7 @@ namespace UE::Mutable::Private
 		ASTChild Image;
 		ASTChild X;
 		ASTChild Y;
-		uint8 Filter=0;
+		uint8 Filter = 0;
 
 	public:
 
@@ -30,7 +30,7 @@ namespace UE::Mutable::Private
 
 		// ASTOp interface
 		virtual EOpType GetOpType() const override { return EOpType::CO_SAMPLEIMAGE; }
-		virtual uint64 Hash() const override;
+		virtual uint32 Hash() const override;
 		virtual bool IsEqual(const ASTOp&) const override;
 		virtual Ptr<ASTOp> Clone(MapChildFuncRef) const override;
 		virtual void ForEachChild(const TFunctionRef<void(ASTChild&)>) override;
diff --git a/Engine/Plugins/Mutable/Source/MutableTools/Private/MuT/ASTOpColorSwizzle.cpp b/Engine/Plugins/Mutable/Source/MutableTools/Private/MuT/ASTOpColorSwizzle.cpp
--- a/Engine/Plugins/Mutable/Source/MutableTools/Private/MuT/ASTOpColorSwizzle.cpp
+++ b/Engine/Plugins/Mutable/Source/MutableTools/Private/MuT/ASTOpColorSwizzle.cpp
@@ -58,17 +58,20 @@ namespace UE::Mutable::Private
 	}
 
 
-	uint64 ASTOpColorSwizzle::Hash() const
+	uint32 ASTOpColorSwizzle::Hash() const
 	{
-		uint64 res = std::hash<void*>()(Sources[0].child().get());
-		hash_combine(res, std::hash<void*>()(Sources[1].child().get()));
-		hash_combine(res, std::hash<void*>()(Sources[2].child().get()));
-		hash_combine(res, std::hash<void*>()(Sources[3].child().get()));
-		hash_combine(res, std::hash<uint8>()(SourceChannels[0]));
-		hash_combine(res, std::hash<uint8>()(SourceChannels[1]));
-		hash_combine(res, std::hash<uint8>()(SourceChannels[2]));
-		hash_combine(res, std::hash<uint8>()(SourceChannels[3]));
-		return res;
+		uint32 Result = GetTypeHash(GetOpType());
+		
+		Result = HashCombineFast(Result, GetTypeHash(Sources[0]));
+		Result = HashCombineFast(Result, GetTypeHash(Sources[1]));
+		Result = HashCombineFast(Result, GetTypeHash(Sources[2]));
+		Result = HashCombineFast(Result, GetTypeHash(Sources[3]));
+		Result = HashCombineFast(Result, GetTypeHash(SourceChannels[0]));
+		Result = HashCombineFast(Result, GetTypeHash(SourceChannels[1]));
+		Result = HashCombineFast(Result, GetTypeHash(SourceChannels[2]));
+		Result = HashCombineFast(Result, GetTypeHash(SourceChannels[3]));
+
+		return Result;
 	}
 
 
diff --git a/Engine/Plugins/Mutable/Source/MutableTools/Private/MuT/ASTOpColorToSRGB.cpp b/Engine/Plugins/Mutable/Source/MutableTools/Private/MuT/ASTOpColorToSRGB.cpp
--- a/Engine/Plugins/Mutable/Source/MutableTools/Private/MuT/ASTOpColorToSRGB.cpp
+++ b/Engine/Plugins/Mutable/Source/MutableTools/Private/MuT/ASTOpColorToSRGB.cpp
@@ -34,10 +34,13 @@ namespace UE::Mutable::Private
 	}
 
 
-	uint64 ASTOpColorToSRGB::Hash() const
+	uint32 ASTOpColorToSRGB::Hash() const
 	{
-		uint64 Res = std::hash<void*>()(Color.child().get());
-		return Res;
+		uint32 Result = GetTypeHash(GetOpType());
+
+		Result = HashCombineFast(Result, GetTypeHash(Color));
+		
+		return Result;
 	}
 
 
diff --git a/Engine/Plugins/Mutable/Source/MutableTools/Private/MuT/ASTOpComponentAdd.cpp b/Engine/Plugins/Mutable/Source/MutableTools/Private/MuT/ASTOpComponentAdd.cpp
--- a/Engine/Plugins/Mutable/Source/MutableTools/Private/MuT/ASTOpComponentAdd.cpp
+++ b/Engine/Plugins/Mutable/Source/MutableTools/Private/MuT/ASTOpComponentAdd.cpp
@@ -56,14 +56,14 @@ namespace UE::Mutable::Private
 	}
 
 
-	uint64 ASTOpComponentAdd::Hash() const
+	uint32 ASTOpComponentAdd::Hash() const
 	{
-		uint64 res = std::hash<size_t>()(size_t(GetOpType()));
-
-		hash_combine(res, Instance.child().get());
-		hash_combine(res, Value.child().get());
+		uint32 Result = GetTypeHash(GetOpType());
+		
+		Result = HashCombineFast(Result, GetTypeHash(Instance));
+		Result = HashCombineFast(Result, GetTypeHash(Value));
 		
-		return res;
+		return Result;
 	}
 	
 
diff --git a/Engine/Plugins/Mutable/Source/MutableTools/Private/MuT/ASTOpConditional.cpp b/Engine/Plugins/Mutable/Source/MutableTools/Private/MuT/ASTOpConditional.cpp
--- a/Engine/Plugins/Mutable/Source/MutableTools/Private/MuT/ASTOpConditional.cpp
+++ b/Engine/Plugins/Mutable/Source/MutableTools/Private/MuT/ASTOpConditional.cpp
@@ -61,12 +61,15 @@ namespace UE::Mutable::Private
 
 
 	//-------------------------------------------------------------------------------------------------
-	uint64 ASTOpConditional::Hash() const
+	uint32 ASTOpConditional::Hash() const
 	{
-		uint64 res = std::hash<void*>()(condition.Child.get());
-		hash_combine(res, yes.Child.get());
-		hash_combine(res, no.Child.get());
-		return res;
+		uint32 Result = GetTypeHash(GetOpType());
+
+		Result = HashCombineFast(Result, GetTypeHash(condition));
+		Result = HashCombineFast(Result, GetTypeHash(yes));
+		Result = HashCombineFast(Result, GetTypeHash(no));
+		
+		return Result;
 	}
 
 
diff --git a/Engine/Plugins/Mutable/Source/MutableTools/Private/MuT/ASTOpConstantBool.cpp b/Engine/Plugins/Mutable/Source/MutableTools/Private/MuT/ASTOpConstantBool.cpp
--- a/Engine/Plugins/Mutable/Source/MutableTools/Private/MuT/ASTOpConstantBool.cpp
+++ b/Engine/Plugins/Mutable/Source/MutableTools/Private/MuT/ASTOpConstantBool.cpp
@@ -33,10 +33,12 @@ namespace UE::Mutable::Private
 	}
 
 
-	uint64 ASTOpConstantBool::Hash() const
+	uint32 ASTOpConstantBool::Hash() const
 	{
-		uint64 Result = std::hash<uint64>()(uint64(EOpType::BO_CONSTANT));
-		hash_combine(Result, bValue);
+		uint32 Result = GetTypeHash(GetOpType());
+		
+		Result = HashCombineFast(Result, GetTypeHash(bValue));
+
 		return Result;
 	}
 
diff --git a/Engine/Plugins/Mutable/Source/MutableTools/Private/MuT/ASTOpConstantColor.cpp b/Engine/Plugins/Mutable/Source/MutableTools/Private/MuT/ASTOpConstantColor.cpp
--- a/Engine/Plugins/Mutable/Source/MutableTools/Private/MuT/ASTOpConstantColor.cpp
+++ b/Engine/Plugins/Mutable/Source/MutableTools/Private/MuT/ASTOpConstantColor.cpp
@@ -35,13 +35,12 @@ namespace UE::Mutable::Private
 	}
 
 
-	uint64 ASTOpConstantColor::Hash() const
+	uint32 ASTOpConstantColor::Hash() const
 	{
-		uint64 Result = std::hash<uint64>()(uint64(GetOpType()));
-		hash_combine(Result, Value[0]);
-		hash_combine(Result, Value[1]);
-		hash_combine(Result, Value[2]);
-		hash_combine(Result, Value[3]);
+		uint32 Result = GetTypeHash(GetOpType());
+
+		Result = HashCombineFast(Result, GetTypeHash(Value));
+		
 		return Result;
 	}
 
diff --git a/Engine/Plugins/Mutable/Source/MutableTools/Private/MuT/ASTOpConstantExtensionData.cpp b/Engine/Plugins/Mutable/Source/MutableTools/Private/MuT/ASTOpConstantExtensionData.cpp
--- a/Engine/Plugins/Mutable/Source/MutableTools/Private/MuT/ASTOpConstantExtensionData.cpp
+++ b/Engine/Plugins/Mutable/Source/MutableTools/Private/MuT/ASTOpConstantExtensionData.cpp
@@ -35,9 +35,9 @@ UE::Mutable::Private::Ptr<ASTOp> ASTOpConstantExtensionData::Clone(MapChildFuncR
 
 
 //-------------------------------------------------------------------------------------------------
-uint64 ASTOpConstantExtensionData::Hash() const
+uint32 ASTOpConstantExtensionData::Hash() const
 {
-	return std::hash<uint64>()(Value != nullptr ? Value->Hash() : 0ll);
+	return Value != nullptr ? Value->Hash() : 0;
 }
 
 
diff --git a/Engine/Plugins/Mutable/Source/MutableTools/Private/MuT/ASTOpConstantInt.cpp b/Engine/Plugins/Mutable/Source/MutableTools/Private/MuT/ASTOpConstantInt.cpp
--- a/Engine/Plugins/Mutable/Source/MutableTools/Private/MuT/ASTOpConstantInt.cpp
+++ b/Engine/Plugins/Mutable/Source/MutableTools/Private/MuT/ASTOpConstantInt.cpp
@@ -33,10 +33,12 @@ namespace UE::Mutable::Private
 	}
 
 
-	uint64 ASTOpConstantInt::Hash() const
+	uint32 ASTOpConstantInt::Hash() const
 	{
-		uint64 Result = std::hash<uint64>()(uint64(GetOpType()));
-		hash_combine(Result, Value);
+		uint32 Result = GetTypeHash(GetOpType());
+
+		Result = HashCombineFast(Result, GetTypeHash(Value));
+		
 		return Result;
 	}
 
diff --git a/Engine/Plugins/Mutable/Source/MutableTools/Private/MuT/ASTOpConstantMatrix.cpp b/Engine/Plugins/Mutable/Source/MutableTools/Private/MuT/ASTOpConstantMatrix.cpp
--- a/Engine/Plugins/Mutable/Source/MutableTools/Private/MuT/ASTOpConstantMatrix.cpp
+++ b/Engine/Plugins/Mutable/Source/MutableTools/Private/MuT/ASTOpConstantMatrix.cpp
@@ -6,11 +6,13 @@ UE::Mutable::Private::ASTOpConstantMatrix::ASTOpConstantMatrix(FMatrix44f InitVa
 {
 }
 
-uint64 UE::Mutable::Private::ASTOpConstantMatrix::Hash() const
+uint32 UE::Mutable::Private::ASTOpConstantMatrix::Hash() const
 {
-	uint64 res = std::hash<uint64>()(uint64(EOpType::MA_CONSTANT));
-	hash_combine(res, value.ComputeHash());
-	return res;
+	uint32 Result = GetTypeHash(GetOpType());
+
+	Result = HashCombineFast(Result, value.ComputeHash());
+	
+	return Result;
 }
 
 bool UE::Mutable::Private::ASTOpConstantMatrix::IsEqual(const ASTOp& otherUntyped) const
diff --git a/Engine/Plugins/Mutable/Source/MutableTools/Private/MuT/ASTOpConstantProjector.cpp b/Engine/Plugins/Mutable/Source/MutableTools/Private/MuT/ASTOpConstantProjector.cpp
--- a/Engine/Plugins/Mutable/Source/MutableTools/Private/MuT/ASTOpConstantProjector.cpp
+++ b/Engine/Plugins/Mutable/Source/MutableTools/Private/MuT/ASTOpConstantProjector.cpp
@@ -19,11 +19,14 @@ namespace UE::Mutable::Private
 	}
 
 
-	uint64 ASTOpConstantProjector::Hash() const
+	uint32 ASTOpConstantProjector::Hash() const
 	{
-		uint64 res = std::hash<float>()(value.position[0]);
-		hash_combine(res, value.direction[0]);
-		return res;
+		uint32 Result = GetTypeHash(GetOpType());
+
+		Result = HashCombineFast(Result, GetTypeHash(value.position));
+		Result = HashCombineFast(Result, GetTypeHash(value.direction));
+		
+		return Result;
 	}
 
 
diff --git a/Engine/Plugins/Mutable/Source/MutableTools/Private/MuT/ASTOpConstantResource.cpp b/Engine/Plugins/Mutable/Source/MutableTools/Private/MuT/ASTOpConstantResource.cpp
--- a/Engine/Plugins/Mutable/Source/MutableTools/Private/MuT/ASTOpConstantResource.cpp
+++ b/Engine/Plugins/Mutable/Source/MutableTools/Private/MuT/ASTOpConstantResource.cpp
@@ -65,11 +65,13 @@ namespace UE::Mutable::Private
 
 
 	//-------------------------------------------------------------------------------------------------
-	uint64 ASTOpConstantResource::Hash() const
+	uint32 ASTOpConstantResource::Hash() const
 	{
-		uint64 res = std::hash<uint64>()(uint64(Type));
-		hash_combine(res, ValueHash);
-		return res;
+		uint32 Result = GetTypeHash(GetOpType());
+		
+		Result = HashCombineFast(Result, GetTypeHash(ValueHash));
+
+		return Result;
 	}
 
 
@@ -679,7 +681,7 @@ namespace UE::Mutable::Private
 
 
 	//-------------------------------------------------------------------------------------------------
-	uint64 ASTOpConstantResource::GetValueHash() const
+	uint32 ASTOpConstantResource::GetValueHash() const
 	{
 		return ValueHash;
 	}
diff --git a/Engine/Plugins/Mutable/Source/MutableTools/Private/MuT/ASTOpConstantResource.h b/Engine/Plugins/Mutable/Source/MutableTools/Private/MuT/ASTOpConstantResource.h
--- a/Engine/Plugins/Mutable/Source/MutableTools/Private/MuT/ASTOpConstantResource.h
+++ b/Engine/Plugins/Mutable/Source/MutableTools/Private/MuT/ASTOpConstantResource.h
@@ -23,7 +23,7 @@ namespace UE::Mutable::Private
 		TManagedPtr<const FResource> Value;
 
 		//! Value hash
-		uint64 ValueHash;
+		uint64 ValueHash = 0;
 
 		//! We tried to link already but the result is a null op.
 		bool bLinkedAndNull = false;
@@ -46,7 +46,7 @@ namespace UE::Mutable::Private
 		// Own interface
 
 		//! Get a hash of the stored value.
-		uint64 GetValueHash() const;
+		uint32 GetValueHash() const;
 
 		//! Get a pointer to the stored value.
 		TManagedPtr<const FResource> GetValue() const;
@@ -60,7 +60,7 @@ namespace UE::Mutable::Private
 		virtual void ForEachChild(const TFunctionRef<void(ASTChild&)>) override;
 		virtual bool IsEqual(const ASTOp& otherUntyped) const override;
 		virtual Ptr<ASTOp> Clone(MapChildFuncRef mapChild) const override;
-		virtual uint64 Hash() const override;
+		virtual uint32 Hash() const override;
 		virtual void Link(FProgram& program, FLinkerOptions*) override;
 		virtual FImageDesc GetImageDesc(bool, class FGetImageDescContext*) const override;
 		virtual void GetBlockLayoutSize(uint64 BlockId, int32* pBlockX, int32* pBlockY, FBlockLayoutSizeCache* cache) override;
diff --git a/Engine/Plugins/Mutable/Source/MutableTools/Private/MuT/ASTOpConstantScalar.cpp b/Engine/Plugins/Mutable/Source/MutableTools/Private/MuT/ASTOpConstantScalar.cpp
--- a/Engine/Plugins/Mutable/Source/MutableTools/Private/MuT/ASTOpConstantScalar.cpp
+++ b/Engine/Plugins/Mutable/Source/MutableTools/Private/MuT/ASTOpConstantScalar.cpp
@@ -33,10 +33,12 @@ namespace UE::Mutable::Private
 	}
 
 
-	uint64 ASTOpConstantScalar::Hash() const
+	uint32 ASTOpConstantScalar::Hash() const
 	{
-		uint64 Result = std::hash<uint64>()(uint64(GetOpType()));
-		hash_combine(Result, Value);
+		uint32 Result = GetTypeHash(GetOpType());
+
+		Result = HashCombineFast(Result, GetTypeHash(Value));
+		
 		return Result;
 	}
 
diff --git a/Engine/Plugins/Mutable/Source/MutableTools/Private/MuT/ASTOpConstantString.cpp b/Engine/Plugins/Mutable/Source/MutableTools/Private/MuT/ASTOpConstantString.cpp
--- a/Engine/Plugins/Mutable/Source/MutableTools/Private/MuT/ASTOpConstantString.cpp
+++ b/Engine/Plugins/Mutable/Source/MutableTools/Private/MuT/ASTOpConstantString.cpp
@@ -20,10 +20,13 @@ namespace UE::Mutable::Private
 	}
 
 
-	uint64 ASTOpConstantString::Hash() const
+	uint32 ASTOpConstantString::Hash() const
 	{
-		uint64 res = std::hash<int32>()(value.Len());
-		return res;
+		uint32 Result = GetTypeHash(GetOpType());
+
+		Result = HashCombineFast(Result, GetTypeHash(value));
+
+		return Result;
 	}
 
 
diff --git a/Engine/Plugins/Mutable/Source/MutableTools/Private/MuT/ASTOpEntry.cpp b/Engine/Plugins/Mutable/Source/MutableTools/Private/MuT/ASTOpEntry.cpp
--- a/Engine/Plugins/Mutable/Source/MutableTools/Private/MuT/ASTOpEntry.cpp
+++ b/Engine/Plugins/Mutable/Source/MutableTools/Private/MuT/ASTOpEntry.cpp
@@ -38,9 +38,12 @@ namespace UE::Mutable::Private
 	}
 
 
-	uint64 ASTOpEntry::Hash() const
+	uint32 ASTOpEntry::Hash() const
 	{
-		uint64 Result = std::hash<void*>()(A.child().get());
+		uint32 Result = GetTypeHash(GetOpType());
+		
+		Result = HashCombineFast(Result, GetTypeHash(A));
+		
 		return Result;
 	}
 
diff --git a/Engine/Plugins/Mutable/Source/MutableTools/Private/MuT/ASTOpExternal.cpp b/Engine/Plugins/Mutable/Source/MutableTools/Private/MuT/ASTOpExternal.cpp
--- a/Engine/Plugins/Mutable/Source/MutableTools/Private/MuT/ASTOpExternal.cpp
+++ b/Engine/Plugins/Mutable/Source/MutableTools/Private/MuT/ASTOpExternal.cpp
@@ -11,14 +11,15 @@ namespace UE::Mutable::Private
 	}
 
 
-	uint64 ASTOpExternal::Hash() const
+	uint32 ASTOpExternal::Hash() const
 	{
-		uint64 Result = GetTypeHash(OperationInstancedStruct);
+		uint32 Result = GetTypeHash(OperationInstancedStruct);
 		
 		for (const ASTChild& Input : Inputs)
 		{
-			hash_combine(Result, Input.child().get());
+			Result = HashCombineFast(Result, GetTypeHash(Input));
 		}
+		
 		return Result;
 	}
 
diff --git a/Engine/Plugins/Mutable/Source/MutableTools/Private/MuT/ASTOpImageBinarize.cpp b/Engine/Plugins/Mutable/Source/MutableTools/Private/MuT/ASTOpImageBinarize.cpp
--- a/Engine/Plugins/Mutable/Source/MutableTools/Private/MuT/ASTOpImageBinarize.cpp
+++ b/Engine/Plugins/Mutable/Source/MutableTools/Private/MuT/ASTOpImageBinarize.cpp
@@ -50,12 +50,14 @@ namespace UE::Mutable::Private
 	}
 
 
-	uint64 ASTOpImageBinarize::Hash() const
+	uint32 ASTOpImageBinarize::Hash() const
 	{
-		uint64 Res = std::hash<EOpType>()(GetOpType());
-		hash_combine(Res, Base.child().get());
-		hash_combine(Res, Threshold.child().get());
-		return Res;
+		uint32 Result = GetTypeHash(GetOpType());
+		
+		Result = HashCombineFast(Result, GetTypeHash(Base));
+		Result = HashCombineFast(Result, GetTypeHash(Threshold));
+		
+		return Result;
 	}
 
 
diff --git a/Engine/Plugins/Mutable/Source/MutableTools/Private/MuT/ASTOpImageBlankLayout.cpp b/Engine/Plugins/Mutable/Source/MutableTools/Private/MuT/ASTOpImageBlankLayout.cpp
--- a/Engine/Plugins/Mutable/Source/MutableTools/Private/MuT/ASTOpImageBlankLayout.cpp
+++ b/Engine/Plugins/Mutable/Source/MutableTools/Private/MuT/ASTOpImageBlankLayout.cpp
@@ -33,7 +33,7 @@ namespace UE::Mutable::Private
 
 	bool ASTOpImageBlankLayout::IsEqual(const ASTOp& InOther) const
 	{
-		if (InOther.GetOpType()==GetOpType())
+		if (InOther.GetOpType() == GetOpType())
 		{
 			const ASTOpImageBlankLayout* Other = static_cast<const ASTOpImageBlankLayout*>(&InOther);
 			return Layout == Other->Layout &&
@@ -42,19 +42,21 @@ namespace UE::Mutable::Private
 				GenerateMipmaps == Other->GenerateMipmaps &&
 				MipmapCount == Other->MipmapCount;
 		}
+		
 		return false;
 	}
 
 
-	uint64 ASTOpImageBlankLayout::Hash() const
+	uint32 ASTOpImageBlankLayout::Hash() const
 	{
-		uint64 Result = std::hash<EOpType>()(GetOpType());
-		hash_combine(Result, Layout.child().get());
-		hash_combine(Result, BlockSize[0]);
-		hash_combine(Result, BlockSize[1]);
-		hash_combine(Result, Format);
-		hash_combine(Result, GenerateMipmaps);
-		hash_combine(Result, MipmapCount);
+		uint32 Result = GetTypeHash(GetOpType());
+
+		Result = HashCombineFast(Result, GetTypeHash(Layout));
+		Result = HashCombineFast(Result, GetTypeHash(BlockSize));
+		Result = HashCombineFast(Result, GetTypeHash(Format));
+		Result = HashCombineFast(Result, GetTypeHash(GenerateMipmaps));
+		Result = HashCombineFast(Result, GetTypeHash(MipmapCount));
+
 		return Result;
 	}
 
diff --git a/Engine/Plugins/Mutable/Source/MutableTools/Private/MuT/ASTOpImageColorMap.cpp b/Engine/Plugins/Mutable/Source/MutableTools/Private/MuT/ASTOpImageColorMap.cpp
--- a/Engine/Plugins/Mutable/Source/MutableTools/Private/MuT/ASTOpImageColorMap.cpp
+++ b/Engine/Plugins/Mutable/Source/MutableTools/Private/MuT/ASTOpImageColorMap.cpp
@@ -52,13 +52,15 @@ namespace UE::Mutable::Private
 	}
 
 
-	uint64 ASTOpImageColorMap::Hash() const
+	uint32 ASTOpImageColorMap::Hash() const
 	{
-		uint64 Res = std::hash<EOpType>()(GetOpType());
-		hash_combine(Res, Base.child().get());
-		hash_combine(Res, Mask.child().get());
-		hash_combine(Res, Map.child().get());
-		return Res;
+		uint32 Result = GetTypeHash(GetOpType());
+		
+		Result = HashCombineFast(Result, GetTypeHash(Base));
+		Result = HashCombineFast(Result, GetTypeHash(Mask));
+		Result = HashCombineFast(Result, GetTypeHash(Map));
+
+		return Result;
 	}
 
 
diff --git a/Engine/Plugins/Mutable/Source/MutableTools/Private/MuT/ASTOpImageCompose.cpp b/Engine/Plugins/Mutable/Source/MutableTools/Private/MuT/ASTOpImageCompose.cpp
--- a/Engine/Plugins/Mutable/Source/MutableTools/Private/MuT/ASTOpImageCompose.cpp
+++ b/Engine/Plugins/Mutable/Source/MutableTools/Private/MuT/ASTOpImageCompose.cpp
@@ -58,15 +58,17 @@ bool ASTOpImageCompose::IsEqual(const ASTOp& otherUntyped) const
 
 
 //-------------------------------------------------------------------------------------------------
-uint64 ASTOpImageCompose::Hash() const
+uint32 ASTOpImageCompose::Hash() const
 {
-	uint64 res = std::hash<EOpType>()(EOpType::IM_COMPOSE);
-    hash_combine( res, Layout.child().get() );
-	hash_combine(res, Base.child().get());
-	hash_combine(res, BlockImage.child().get());
-	hash_combine(res, Mask.child().get());
-	hash_combine(res, BlockId);
-	return res;
+	uint32 Result = GetTypeHash(GetOpType());
+	
+    Result = HashCombineFast(Result, GetTypeHash(Layout));
+	Result = HashCombineFast(Result, GetTypeHash(Base));
+	Result = HashCombineFast(Result, GetTypeHash(BlockImage));
+	Result = HashCombineFast(Result, GetTypeHash(Mask));
+	Result = HashCombineFast(Result, BlockId);
+	
+	return Result;
 }
 
 
diff --git a/Engine/Plugins/Mutable/Source/MutableTools/Private/MuT/ASTOpImageCrop.cpp b/Engine/Plugins/Mutable/Source/MutableTools/Private/MuT/ASTOpImageCrop.cpp
--- a/Engine/Plugins/Mutable/Source/MutableTools/Private/MuT/ASTOpImageCrop.cpp
+++ b/Engine/Plugins/Mutable/Source/MutableTools/Private/MuT/ASTOpImageCrop.cpp
@@ -47,15 +47,15 @@ namespace UE::Mutable::Private
 	}
 
 
-	uint64 ASTOpImageCrop::Hash() const
+	uint32 ASTOpImageCrop::Hash() const
 	{
-		uint64 res = std::hash<EOpType>()(GetOpType());
-		hash_combine(res, Source.child().get());
-		hash_combine(res, Min[0]);
-		hash_combine(res, Min[1]);
-		hash_combine(res, Size[0]);
-		hash_combine(res, Size[1]);
-		return res;
+		uint32 Result = GetTypeHash(GetOpType());
+		
+		Result = HashCombineFast(Result, GetTypeHash(Source));
+		Result = HashCombineFast(Result, GetTypeHash(Min));
+		Result = HashCombineFast(Result, GetTypeHash(Size));
+
+		return Result;
 	}
 
 
diff --git a/Engine/Plugins/Mutable/Source/MutableTools/Private/MuT/ASTOpImageDisplace.cpp b/Engine/Plugins/Mutable/Source/MutableTools/Private/MuT/ASTOpImageDisplace.cpp
--- a/Engine/Plugins/Mutable/Source/MutableTools/Private/MuT/ASTOpImageDisplace.cpp
+++ b/Engine/Plugins/Mutable/Source/MutableTools/Private/MuT/ASTOpImageDisplace.cpp
@@ -50,12 +50,14 @@ namespace UE::Mutable::Private
 	}
 
 
-	uint64 ASTOpImageDisplace::Hash() const
+	uint32 ASTOpImageDisplace::Hash() const
 	{
-		uint64 Res = std::hash<EOpType>()(GetOpType());
-		hash_combine(Res, Source.child().get());
-		hash_combine(Res, DisplacementMap.child().get());
-		return Res;
+		uint32 Result = GetTypeHash(GetOpType());
+		
+		Result = HashCombineFast(Result, GetTypeHash(Source));
+		Result = HashCombineFast(Result, GetTypeHash(DisplacementMap));
+		
+		return Result;
 	}
 
 
diff --git a/Engine/Plugins/Mutable/Source/MutableTools/Private/MuT/ASTOpImageInterpolate.cpp b/Engine/Plugins/Mutable/Source/MutableTools/Private/MuT/ASTOpImageInterpolate.cpp
--- a/Engine/Plugins/Mutable/Source/MutableTools/Private/MuT/ASTOpImageInterpolate.cpp
+++ b/Engine/Plugins/Mutable/Source/MutableTools/Private/MuT/ASTOpImageInterpolate.cpp
@@ -57,16 +57,19 @@ namespace UE::Mutable::Private
 	}
 
 
-	uint64 ASTOpImageInterpolate::Hash() const
+	uint32 ASTOpImageInterpolate::Hash() const
 	{
-		uint64 res = std::hash<void*>()(Targets[0].child().get());
-		hash_combine(res, std::hash<void*>()(Targets[1].child().get()));
-		hash_combine(res, std::hash<void*>()(Targets[2].child().get()));
-		hash_combine(res, std::hash<void*>()(Targets[3].child().get()));
-		hash_combine(res, std::hash<void*>()(Targets[4].child().get()));
-		hash_combine(res, std::hash<void*>()(Targets[5].child().get()));
-		hash_combine(res, std::hash<void*>()(Factor.child().get()));
-		return res;
+		uint32 Result = GetTypeHash(GetOpType());
+		
+		Result = HashCombineFast(Result, GetTypeHash(Targets[0]));
+		Result = HashCombineFast(Result, GetTypeHash(Targets[1]));
+		Result = HashCombineFast(Result, GetTypeHash(Targets[2]));
+		Result = HashCombineFast(Result, GetTypeHash(Targets[3]));
+		Result = HashCombineFast(Result, GetTypeHash(Targets[4]));
+		Result = HashCombineFast(Result, GetTypeHash(Targets[5]));
+		Result = HashCombineFast(Result, GetTypeHash(Factor));
+
+		return Result;
 	}
 
 
diff --git a/Engine/Plugins/Mutable/Source/MutableTools/Private/MuT/ASTOpImageInvert.cpp b/Engine/Plugins/Mutable/Source/MutableTools/Private/MuT/ASTOpImageInvert.cpp
--- a/Engine/Plugins/Mutable/Source/MutableTools/Private/MuT/ASTOpImageInvert.cpp
+++ b/Engine/Plugins/Mutable/Source/MutableTools/Private/MuT/ASTOpImageInvert.cpp
@@ -48,11 +48,13 @@ namespace UE::Mutable::Private
 	}
 
 
-	uint64 ASTOpImageInvert::Hash() const
+	uint32 ASTOpImageInvert::Hash() const
 	{
-		uint64 Res = std::hash<EOpType>()(GetOpType());
-		hash_combine(Res, Base.child().get());
-		return Res;
+		uint32 Result = GetTypeHash(GetOpType());
+
+		Result = HashCombineFast(Result, GetTypeHash(Base));
+		
+		return Result;
 	}
 
 
diff --git a/Engine/Plugins/Mutable/Source/MutableTools/Private/MuT/ASTOpImageLayer.cpp b/Engine/Plugins/Mutable/Source/MutableTools/Private/MuT/ASTOpImageLayer.cpp
--- a/Engine/Plugins/Mutable/Source/MutableTools/Private/MuT/ASTOpImageLayer.cpp
+++ b/Engine/Plugins/Mutable/Source/MutableTools/Private/MuT/ASTOpImageLayer.cpp
@@ -52,13 +52,15 @@ namespace UE::Mutable::Private
 	}
 
 
-	uint64 ASTOpImageLayer::Hash() const
+	uint32 ASTOpImageLayer::Hash() const
 	{
-		uint64 res = std::hash<EOpType>()(GetOpType());
-		hash_combine(res, base.child().get());
-		hash_combine(res, blend.child().get());
-		hash_combine(res, mask.child().get());
-		return res;
+		uint32 Result = GetTypeHash(GetOpType());
+
+		Result = HashCombineFast(Result, GetTypeHash(base));
+		Result = HashCombineFast(Result, GetTypeHash(blend));
+		Result = HashCombineFast(Result, GetTypeHash(mask));
+		
+		return Result;
 	}
 
 
diff --git a/Engine/Plugins/Mutable/Source/MutableTools/Private/MuT/ASTOpImageLayerColor.cpp b/Engine/Plugins/Mutable/Source/MutableTools/Private/MuT/ASTOpImageLayerColor.cpp
--- a/Engine/Plugins/Mutable/Source/MutableTools/Private/MuT/ASTOpImageLayerColor.cpp
+++ b/Engine/Plugins/Mutable/Source/MutableTools/Private/MuT/ASTOpImageLayerColor.cpp
@@ -52,13 +52,15 @@ namespace UE::Mutable::Private
 
 
 	//-------------------------------------------------------------------------------------------------
-	uint64 ASTOpImageLayerColor::Hash() const
+	uint32 ASTOpImageLayerColor::Hash() const
 	{
-		uint64 res = std::hash<EOpType>()(GetOpType());
-		hash_combine(res, base.child().get());
-		hash_combine(res, color.child().get());
-		hash_combine(res, mask.child().get());
-		return res;
+		uint32 Result = GetTypeHash(GetOpType());
+		
+		Result = HashCombineFast(Result, GetTypeHash(base));
+		Result = HashCombineFast(Result, GetTypeHash(color));
+		Result = HashCombineFast(Result, GetTypeHash(mask));
+		
+		return Result;
 	}
 
 
diff --git a/Engine/Plugins/Mutable/Source/MutableTools/Private/MuT/ASTOpImageLuminance.cpp b/Engine/Plugins/Mutable/Source/MutableTools/Private/MuT/ASTOpImageLuminance.cpp
--- a/Engine/Plugins/Mutable/Source/MutableTools/Private/MuT/ASTOpImageLuminance.cpp
+++ b/Engine/Plugins/Mutable/Source/MutableTools/Private/MuT/ASTOpImageLuminance.cpp
@@ -48,11 +48,13 @@ namespace UE::Mutable::Private
 	}
 
 
-	uint64 ASTOpImageLuminance::Hash() const
+	uint32 ASTOpImageLuminance::Hash() const
 	{
-		uint64 Res = std::hash<EOpType>()(GetOpType());
-		hash_combine(Res, Base.child().get());
-		return Res;
+		uint32 Result = GetTypeHash(GetOpType());
+	
+		Result = HashCombineFast(Result, GetTypeHash(Base));
+		
+		return Result;
 	}
 
 
diff --git a/Engine/Plugins/Mutable/Source/MutableTools/Private/MuT/ASTOpImageMakeGrowMap.cpp b/Engine/Plugins/Mutable/Source/MutableTools/Private/MuT/ASTOpImageMakeGrowMap.cpp
--- a/Engine/Plugins/Mutable/Source/MutableTools/Private/MuT/ASTOpImageMakeGrowMap.cpp
+++ b/Engine/Plugins/Mutable/Source/MutableTools/Private/MuT/ASTOpImageMakeGrowMap.cpp
@@ -46,11 +46,14 @@ namespace UE::Mutable::Private
 	}
 
 
-	uint64 ASTOpImageMakeGrowMap::Hash() const
+	uint32 ASTOpImageMakeGrowMap::Hash() const
 	{
-		uint64 res = std::hash<void*>()(Mask.child().get());
-		hash_combine(res, Border);
-		return res;
+		uint32 Result = GetTypeHash(GetOpType());		
+
+		Result = HashCombineFast(Result, GetTypeHash(Mask));
+		Result = HashCombineFast(Result, GetTypeHash(Border));
+		
+		return Result;
 	}
 
 
diff --git a/Engine/Plugins/Mutable/Source/MutableTools/Private/MuT/ASTOpImageMipmap.cpp b/Engine/Plugins/Mutable/Source/MutableTools/Private/MuT/ASTOpImageMipmap.cpp
--- a/Engine/Plugins/Mutable/Source/MutableTools/Private/MuT/ASTOpImageMipmap.cpp
+++ b/Engine/Plugins/Mutable/Source/MutableTools/Private/MuT/ASTOpImageMipmap.cpp
@@ -54,11 +54,14 @@ namespace UE::Mutable::Private
 	}
 
 
-	uint64 ASTOpImageMipmap::Hash() const
+	uint32 ASTOpImageMipmap::Hash() const
 	{
-		uint64 res = std::hash<void*>()(Source.child().get());
-		hash_combine(res, Levels);
-		return res;
+		uint32 Result = GetTypeHash(GetOpType());
+		
+		Result = HashCombineFast(Result, GetTypeHash(Source));
+		Result = HashCombineFast(Result, Levels);
+
+		return Result;
 	}
 
 
diff --git a/Engine/Plugins/Mutable/Source/MutableTools/Private/MuT/ASTOpImageMultiCompose.cpp b/Engine/Plugins/Mutable/Source/MutableTools/Private/MuT/ASTOpImageMultiCompose.cpp
--- a/Engine/Plugins/Mutable/Source/MutableTools/Private/MuT/ASTOpImageMultiCompose.cpp
+++ b/Engine/Plugins/Mutable/Source/MutableTools/Private/MuT/ASTOpImageMultiCompose.cpp
@@ -57,14 +57,16 @@ bool ASTOpImageMultiCompose::IsEqual(const ASTOp& OtherUntyped) const
 }
 
 
-uint64 ASTOpImageMultiCompose::Hash() const
+uint32 ASTOpImageMultiCompose::Hash() const
 {
-	uint64 Res = std::hash<EOpType>()(EOpType::IM_MULTICOMPOSE);
-    hash_combine(Res, Layout.child().get());
-	hash_combine(Res, Base.child().get());
-	hash_combine(Res, SourceLayout.child().get());
-	hash_combine(Res, SourceImage.child().get());
-	return Res;
+	uint32 Result = GetTypeHash(GetOpType());
+	
+	Result = HashCombineFast(Result, GetTypeHash(Layout));
+	Result = HashCombineFast(Result, GetTypeHash(Base));
+	Result = HashCombineFast(Result, GetTypeHash(SourceLayout));
+	Result = HashCombineFast(Result, GetTypeHash(SourceImage));
+
+	return Result;
 }
 
 
diff --git a/Engine/Plugins/Mutable/Source/MutableTools/Private/MuT/ASTOpImageMultiLayer.cpp b/Engine/Plugins/Mutable/Source/MutableTools/Private/MuT/ASTOpImageMultiLayer.cpp
--- a/Engine/Plugins/Mutable/Source/MutableTools/Private/MuT/ASTOpImageMultiLayer.cpp
+++ b/Engine/Plugins/Mutable/Source/MutableTools/Private/MuT/ASTOpImageMultiLayer.cpp
@@ -50,13 +50,15 @@ namespace UE::Mutable::Private
 
 
 	//-------------------------------------------------------------------------------------------------
-	uint64 ASTOpImageMultiLayer::Hash() const
+	uint32 ASTOpImageMultiLayer::Hash() const
 	{
-		uint64 res = std::hash<EOpType>()(GetOpType());
-		hash_combine(res, base.child().get());
-		hash_combine(res, blend.child().get());
-		hash_combine(res, mask.child().get());
-		return res;
+		uint32 Result = GetTypeHash(GetOpType());
+		
+		Result = HashCombineFast(Result, GetTypeHash(base));
+		Result = HashCombineFast(Result, GetTypeHash(blend));
+		Result = HashCombineFast(Result, GetTypeHash(mask));
+		
+		return Result;
 	}
 
 
diff --git a/Engine/Plugins/Mutable/Source/MutableTools/Private/MuT/ASTOpImageNormalComposite.cpp b/Engine/Plugins/Mutable/Source/MutableTools/Private/MuT/ASTOpImageNormalComposite.cpp
--- a/Engine/Plugins/Mutable/Source/MutableTools/Private/MuT/ASTOpImageNormalComposite.cpp
+++ b/Engine/Plugins/Mutable/Source/MutableTools/Private/MuT/ASTOpImageNormalComposite.cpp
@@ -45,12 +45,14 @@ namespace UE::Mutable::Private
 	}
 
 
-	uint64 ASTOpImageNormalComposite::Hash() const
+	uint32 ASTOpImageNormalComposite::Hash() const
 	{
-		uint64 Res = std::hash<void*>()(Base.child().get());
-		hash_combine(Res, Normal.child().get());
+		uint32 Result = GetTypeHash(GetOpType());
 
-		return Res;
+		Result = HashCombineFast(Result, GetTypeHash(Base));
+		Result = HashCombineFast(Result, GetTypeHash(Normal));
+
+		return Result;
 	}
 
 
diff --git a/Engine/Plugins/Mutable/Source/MutableTools/Private/MuT/ASTOpImageParameterConvert.cpp b/Engine/Plugins/Mutable/Source/MutableTools/Private/MuT/ASTOpImageParameterConvert.cpp
--- a/Engine/Plugins/Mutable/Source/MutableTools/Private/MuT/ASTOpImageParameterConvert.cpp
+++ b/Engine/Plugins/Mutable/Source/MutableTools/Private/MuT/ASTOpImageParameterConvert.cpp
@@ -57,13 +57,13 @@ namespace UE::Mutable::Private
 	}
 
 
-	uint64 ASTOpImageParameterConvert::Hash() const
+	uint32 ASTOpImageParameterConvert::Hash() const
 	{
-		uint64 res = std::hash<uint64>()(uint64(Type));
+		uint32 Result = GetTypeHash(GetOpType());
 
-		hash_combine(res, ImageParameter.child().get());
+		Result = HashCombineFast(Result, GetTypeHash(ImageParameter));
 
-		return res;
+		return Result;
 	}
 
 
diff --git a/Engine/Plugins/Mutable/Source/MutableTools/Private/MuT/ASTOpImagePatch.cpp b/Engine/Plugins/Mutable/Source/MutableTools/Private/MuT/ASTOpImagePatch.cpp
--- a/Engine/Plugins/Mutable/Source/MutableTools/Private/MuT/ASTOpImagePatch.cpp
+++ b/Engine/Plugins/Mutable/Source/MutableTools/Private/MuT/ASTOpImagePatch.cpp
@@ -41,12 +41,14 @@ namespace UE::Mutable::Private
 
 
 	//-------------------------------------------------------------------------------------------------
-	uint64 ASTOpImagePatch::Hash() const
+	uint32 ASTOpImagePatch::Hash() const
 	{
-		uint64 res = std::hash<EOpType>()(EOpType::IM_PATCH);
-		hash_combine(res, base.child().get());
-		hash_combine(res, patch.child().get());
-		return res;
+		uint32 Result = GetTypeHash(GetOpType());
+		
+		Result = HashCombineFast(Result, GetTypeHash(base));
+		Result = HashCombineFast(Result, GetTypeHash(patch));
+
+		return Result;
 	}
 
 
diff --git a/Engine/Plugins/Mutable/Source/MutableTools/Private/MuT/ASTOpImagePixelFormat.cpp b/Engine/Plugins/Mutable/Source/MutableTools/Private/MuT/ASTOpImagePixelFormat.cpp
--- a/Engine/Plugins/Mutable/Source/MutableTools/Private/MuT/ASTOpImagePixelFormat.cpp
+++ b/Engine/Plugins/Mutable/Source/MutableTools/Private/MuT/ASTOpImagePixelFormat.cpp
@@ -58,11 +58,14 @@ namespace UE::Mutable::Private
 	}
 
 
-	uint64 ASTOpImagePixelFormat::Hash() const
+	uint32 ASTOpImagePixelFormat::Hash() const
 	{
-		uint64 res = std::hash<void*>()(Source.child().get());
-		hash_combine(res, Format);
-		return res;
+		uint32 Result = GetTypeHash(GetOpType());
+		
+		Result = HashCombineFast(Result, GetTypeHash(Source));
+		Result = HashCombineFast(Result, GetTypeHash(Format));
+		
+		return Result;
 	}
 
 
diff --git a/Engine/Plugins/Mutable/Source/MutableTools/Private/MuT/ASTOpImagePlainColor.cpp b/Engine/Plugins/Mutable/Source/MutableTools/Private/MuT/ASTOpImagePlainColor.cpp
--- a/Engine/Plugins/Mutable/Source/MutableTools/Private/MuT/ASTOpImagePlainColor.cpp
+++ b/Engine/Plugins/Mutable/Source/MutableTools/Private/MuT/ASTOpImagePlainColor.cpp
@@ -46,15 +46,16 @@ namespace UE::Mutable::Private
 	}
 
 
-	uint64 ASTOpImagePlainColor::Hash() const
+	uint32 ASTOpImagePlainColor::Hash() const
 	{
-		uint64 res = std::hash<EOpType>()(GetOpType());
-		hash_combine(res, Color.child().get());
-		hash_combine(res, Size[0]);
-		hash_combine(res, Size[1]);
-		hash_combine(res, Format);
-		hash_combine(res, LODs);
-		return res;
+		uint32 Result = GetTypeHash(GetOpType());
+		
+		Result = HashCombineFast(Result, GetTypeHash(Color));
+		Result = HashCombineFast(Result, GetTypeHash(Size));
+		Result = HashCombineFast(Result, GetTypeHash(Format));
+		Result = HashCombineFast(Result, GetTypeHash(LODs));
+
+		return Result;
 	}
 
 
diff --git a/Engine/Plugins/Mutable/Source/MutableTools/Private/MuT/ASTOpImageRasterMesh.cpp b/Engine/Plugins/Mutable/Source/MutableTools/Private/MuT/ASTOpImageRasterMesh.cpp
--- a/Engine/Plugins/Mutable/Source/MutableTools/Private/MuT/ASTOpImageRasterMesh.cpp
+++ b/Engine/Plugins/Mutable/Source/MutableTools/Private/MuT/ASTOpImageRasterMesh.cpp
@@ -80,16 +80,18 @@ namespace UE::Mutable::Private
 
 
 	//-------------------------------------------------------------------------------------------------
-	uint64 ASTOpImageRasterMesh::Hash() const
+	uint32 ASTOpImageRasterMesh::Hash() const
 	{
-		uint64 res = std::hash<EOpType>()(GetOpType());
-		hash_combine(res, mesh.child().get());
-		hash_combine(res, image.child().get());
-		hash_combine(res, angleFadeProperties.child().get());
-		hash_combine(res, mask.child().get());
-		hash_combine(res, projector.child().get());
-		hash_combine(res, BlockId);
-		return res;
+		uint32 Result = GetTypeHash(GetOpType());
+		
+		Result = HashCombineFast(Result, GetTypeHash(mesh));
+		Result = HashCombineFast(Result, GetTypeHash(image));
+		Result = HashCombineFast(Result, GetTypeHash(angleFadeProperties));
+		Result = HashCombineFast(Result, GetTypeHash(mask));
+		Result = HashCombineFast(Result, GetTypeHash(projector));
+		Result = HashCombineFast(Result, GetTypeHash(BlockId));
+		
+		return Result;
 	}
 
 
diff --git a/Engine/Plugins/Mutable/Source/MutableTools/Private/MuT/ASTOpImageResize.cpp b/Engine/Plugins/Mutable/Source/MutableTools/Private/MuT/ASTOpImageResize.cpp
--- a/Engine/Plugins/Mutable/Source/MutableTools/Private/MuT/ASTOpImageResize.cpp
+++ b/Engine/Plugins/Mutable/Source/MutableTools/Private/MuT/ASTOpImageResize.cpp
@@ -53,13 +53,14 @@ namespace UE::Mutable::Private
 	}
 
 
-	uint64 ASTOpImageResize::Hash() const
+	uint32 ASTOpImageResize::Hash() const
 	{
-		uint64 Res = std::hash<EOpType>()(GetOpType());
-		hash_combine(Res, Source.child().get());
-		hash_combine(Res, Size[0]);
-		hash_combine(Res, Size[1]);
-		return Res;
+		uint32 Result = GetTypeHash(GetOpType());
+		
+		Result = HashCombineFast(Result, GetTypeHash(Source));
+		Result = HashCombineFast(Result, GetTypeHash(Size));
+
+		return Result;
 	}
 
 
diff --git a/Engine/Plugins/Mutable/Source/MutableTools/Private/MuT/ASTOpImageResizeLike.cpp b/Engine/Plugins/Mutable/Source/MutableTools/Private/MuT/ASTOpImageResizeLike.cpp
--- a/Engine/Plugins/Mutable/Source/MutableTools/Private/MuT/ASTOpImageResizeLike.cpp
+++ b/Engine/Plugins/Mutable/Source/MutableTools/Private/MuT/ASTOpImageResizeLike.cpp
@@ -50,12 +50,14 @@ namespace UE::Mutable::Private
 	}
 
 
-	uint64 ASTOpImageResizeLike::Hash() const
+	uint32 ASTOpImageResizeLike::Hash() const
 	{
-		uint64 Res = std::hash<EOpType>()(GetOpType());
-		hash_combine(Res, Source.child().get());
-		hash_combine(Res, SizeSource.child().get());
-		return Res;
+		uint32 Result = GetTypeHash(GetOpType());
+		
+		Result = HashCombineFast(Result, GetTypeHash(Source));
+		Result = HashCombineFast(Result, GetTypeHash(SizeSource));
+		
+		return Result;
 	}
 
 
diff --git a/Engine/Plugins/Mutable/Source/MutableTools/Private/MuT/ASTOpImageResizeRel.cpp b/Engine/Plugins/Mutable/Source/MutableTools/Private/MuT/ASTOpImageResizeRel.cpp
--- a/Engine/Plugins/Mutable/Source/MutableTools/Private/MuT/ASTOpImageResizeRel.cpp
+++ b/Engine/Plugins/Mutable/Source/MutableTools/Private/MuT/ASTOpImageResizeRel.cpp
@@ -58,13 +58,14 @@ namespace UE::Mutable::Private
 	}
 
 
-	uint64 ASTOpImageResizeRel::Hash() const
+	uint32 ASTOpImageResizeRel::Hash() const
 	{
-		uint64 Res = std::hash<EOpType>()(GetOpType());
-		hash_combine(Res, Source.child().get());
-		hash_combine(Res, Factor[0]);
-		hash_combine(Res, Factor[1]);
-		return Res;
+		uint32 Result = GetTypeHash(GetOpType());
+		
+		Result = HashCombineFast(Result, GetTypeHash(Source));
+		Result = HashCombineFast(Result, GetTypeHash(Factor));
+		
+		return Result;
 	}
 
 
diff --git a/Engine/Plugins/Mutable/Source/MutableTools/Private/MuT/ASTOpImageSaturate.cpp b/Engine/Plugins/Mutable/Source/MutableTools/Private/MuT/ASTOpImageSaturate.cpp
--- a/Engine/Plugins/Mutable/Source/MutableTools/Private/MuT/ASTOpImageSaturate.cpp
+++ b/Engine/Plugins/Mutable/Source/MutableTools/Private/MuT/ASTOpImageSaturate.cpp
@@ -50,12 +50,14 @@ namespace UE::Mutable::Private
 	}
 
 
-	uint64 ASTOpImageSaturate::Hash() const
+	uint32 ASTOpImageSaturate::Hash() const
 	{
-		uint64 Res = std::hash<EOpType>()(GetOpType());
-		hash_combine(Res, Base.child().get());
-		hash_combine(Res, Factor.child().get());
-		return Res;
+		uint32 Result = GetTypeHash(GetOpType());
+		
+		Result = HashCombineFast(Result, GetTypeHash(Base));
+		Result = HashCombineFast(Result, GetTypeHash(Factor));
+		
+		return Result;
 	}
 
 
diff --git a/Engine/Plugins/Mutable/Source/MutableTools/Private/MuT/ASTOpImageSwizzle.cpp b/Engine/Plugins/Mutable/Source/MutableTools/Private/MuT/ASTOpImageSwizzle.cpp
--- a/Engine/Plugins/Mutable/Source/MutableTools/Private/MuT/ASTOpImageSwizzle.cpp
+++ b/Engine/Plugins/Mutable/Source/MutableTools/Private/MuT/ASTOpImageSwizzle.cpp
@@ -64,18 +64,21 @@ namespace UE::Mutable::Private
 	}
 
 
-	uint64 ASTOpImageSwizzle::Hash() const
+	uint32 ASTOpImageSwizzle::Hash() const
 	{
-		uint64 res = std::hash<void*>()(Sources[0].child().get());
-		hash_combine(res, std::hash<void*>()(Sources[1].child().get()));
-		hash_combine(res, std::hash<void*>()(Sources[2].child().get()));
-		hash_combine(res, std::hash<void*>()(Sources[3].child().get()));
-		hash_combine(res, std::hash<uint8>()(SourceChannels[0]));
-		hash_combine(res, std::hash<uint8>()(SourceChannels[1]));
-		hash_combine(res, std::hash<uint8>()(SourceChannels[2]));
-		hash_combine(res, std::hash<uint8>()(SourceChannels[3]));
-		hash_combine(res, Format);
-		return res;
+		uint32 Result = GetTypeHash(GetOpType());
+
+		Result = HashCombineFast(Result, GetTypeHash(Sources[0]));
+		Result = HashCombineFast(Result, GetTypeHash(Sources[1]));
+		Result = HashCombineFast(Result, GetTypeHash(Sources[2]));
+		Result = HashCombineFast(Result, GetTypeHash(Sources[3]));
+		Result = HashCombineFast(Result, GetTypeHash(SourceChannels[0]));
+		Result = HashCombineFast(Result, GetTypeHash(SourceChannels[1]));
+		Result = HashCombineFast(Result, GetTypeHash(SourceChannels[2]));
+		Result = HashCombineFast(Result, GetTypeHash(SourceChannels[3]));
+		Result = HashCombineFast(Result, GetTypeHash(Format));
+		
+		return Result;
 	}
 
 
diff --git a/Engine/Plugins/Mutable/Source/MutableTools/Private/MuT/ASTOpImageTransform.cpp b/Engine/Plugins/Mutable/Source/MutableTools/Private/MuT/ASTOpImageTransform.cpp
--- a/Engine/Plugins/Mutable/Source/MutableTools/Private/MuT/ASTOpImageTransform.cpp
+++ b/Engine/Plugins/Mutable/Source/MutableTools/Private/MuT/ASTOpImageTransform.cpp
@@ -57,22 +57,24 @@ namespace UE::Mutable::Private
 
 
 	//-------------------------------------------------------------------------------------------------
-	uint64 ASTOpImageTransform::Hash() const
+	uint32 ASTOpImageTransform::Hash() const
 	{
-		uint64 Res = std::hash<EOpType>()(EOpType::IM_TRANSFORM);
-		hash_combine(Res, Base.child().get());
-		hash_combine(Res, OffsetX.child().get());
-		hash_combine(Res, OffsetY.child().get());
-		hash_combine(Res, ScaleX.child().get());
-		hash_combine(Res, ScaleY.child().get());
-		hash_combine(Res, Rotation.child().get());
-		hash_combine(Res, std::hash<uint32>()(static_cast<uint32>(AddressMode)));
-		hash_combine(Res, SizeX);
-		hash_combine(Res, SizeY);
-		hash_combine(Res, SourceSizeX);
-		hash_combine(Res, SourceSizeY);
-		hash_combine(Res, bKeepAspectRatio);
-		return Res;
+		uint32 Result = GetTypeHash(GetOpType());
+		
+		Result = HashCombineFast(Result, GetTypeHash(Base));
+		Result = HashCombineFast(Result, GetTypeHash(OffsetX));
+		Result = HashCombineFast(Result, GetTypeHash(OffsetY));
+		Result = HashCombineFast(Result, GetTypeHash(ScaleX));
+		Result = HashCombineFast(Result, GetTypeHash(ScaleY));
+		Result = HashCombineFast(Result, GetTypeHash(Rotation));
+		Result = HashCombineFast(Result, GetTypeHash(AddressMode));
+		Result = HashCombineFast(Result, GetTypeHash(SizeX));
+		Result = HashCombineFast(Result, GetTypeHash(SizeY));
+		Result = HashCombineFast(Result, GetTypeHash(SourceSizeX));
+		Result = HashCombineFast(Result, GetTypeHash(SourceSizeY));
+		Result = HashCombineFast(Result, GetTypeHash(bKeepAspectRatio));
+			
+		return Result;
 	}
 
 
diff --git a/Engine/Plugins/Mutable/Source/MutableTools/Private/MuT/ASTOpInstanceAdd.cpp b/Engine/Plugins/Mutable/Source/MutableTools/Private/MuT/ASTOpInstanceAdd.cpp
--- a/Engine/Plugins/Mutable/Source/MutableTools/Private/MuT/ASTOpInstanceAdd.cpp
+++ b/Engine/Plugins/Mutable/Source/MutableTools/Private/MuT/ASTOpInstanceAdd.cpp
@@ -57,12 +57,13 @@ namespace UE::Mutable::Private
 	}
 
 
-	uint64 ASTOpInstanceAdd::Hash() const
+	uint32 ASTOpInstanceAdd::Hash() const
 	{
-		uint64 res = std::hash<size_t>()(size_t(type));
-		hash_combine(res, instance.child().get());
-		hash_combine(res, value.child().get());
-		return res;
+		uint32 Result = GetTypeHash(GetOpType());
+
+		Result = HashCombineFast(Result, GetTypeHash(instance));
+		Result = HashCombineFast(Result, GetTypeHash(value));
+		return Result;
 	}
 
 
diff --git a/Engine/Plugins/Mutable/Source/MutableTools/Private/MuT/ASTOpLayoutFromMesh.cpp b/Engine/Plugins/Mutable/Source/MutableTools/Private/MuT/ASTOpLayoutFromMesh.cpp
--- a/Engine/Plugins/Mutable/Source/MutableTools/Private/MuT/ASTOpLayoutFromMesh.cpp
+++ b/Engine/Plugins/Mutable/Source/MutableTools/Private/MuT/ASTOpLayoutFromMesh.cpp
@@ -54,11 +54,14 @@ namespace UE::Mutable::Private
 	}
 
 
-	uint64 ASTOpLayoutFromMesh::Hash() const
+	uint32 ASTOpLayoutFromMesh::Hash() const
 	{
-		uint64 res = std::hash<void*>()(Mesh.child().get());
-		hash_combine(res, LayoutIndex);
-		return res;
+		uint32 Result = GetTypeHash(GetOpType());
+
+		Result = HashCombineFast(Result, GetTypeHash(Mesh));
+		Result = HashCombineFast(Result, GetTypeHash(LayoutIndex));
+		
+		return Result;
 	}
 
 
diff --git a/Engine/Plugins/Mutable/Source/MutableTools/Private/MuT/ASTOpLayoutMerge.cpp b/Engine/Plugins/Mutable/Source/MutableTools/Private/MuT/ASTOpLayoutMerge.cpp
--- a/Engine/Plugins/Mutable/Source/MutableTools/Private/MuT/ASTOpLayoutMerge.cpp
+++ b/Engine/Plugins/Mutable/Source/MutableTools/Private/MuT/ASTOpLayoutMerge.cpp
@@ -49,11 +49,14 @@ namespace UE::Mutable::Private
 	}
 
 
-	uint64 ASTOpLayoutMerge::Hash() const
+	uint32 ASTOpLayoutMerge::Hash() const
 	{
-		uint64 res = std::hash<void*>()(Base.child().get());
-		hash_combine(res, Added.child().get());
-		return res;
+		uint32 Result = GetTypeHash(GetOpType());
+		
+		Result = HashCombineFast(Result, GetTypeHash(Base));
+		Result = HashCombineFast(Result, GetTypeHash(Added));
+
+		return Result;
 	}
 
 
diff --git a/Engine/Plugins/Mutable/Source/MutableTools/Private/MuT/ASTOpLayoutPack.cpp b/Engine/Plugins/Mutable/Source/MutableTools/Private/MuT/ASTOpLayoutPack.cpp
--- a/Engine/Plugins/Mutable/Source/MutableTools/Private/MuT/ASTOpLayoutPack.cpp
+++ b/Engine/Plugins/Mutable/Source/MutableTools/Private/MuT/ASTOpLayoutPack.cpp
@@ -46,10 +46,13 @@ namespace UE::Mutable::Private
 	}
 
 
-	uint64 ASTOpLayoutPack::Hash() const
+	uint32 ASTOpLayoutPack::Hash() const
 	{
-		uint64 res = std::hash<void*>()(Source.child().get());
-		return res;
+		uint32 Result = GetTypeHash(GetOpType());
+
+		Result = HashCombineFast(Result, GetTypeHash(Source));
+		
+		return Result;
 	}
 
 
diff --git a/Engine/Plugins/Mutable/Source/MutableTools/Private/MuT/ASTOpLayoutRemoveBlocks.cpp b/Engine/Plugins/Mutable/Source/MutableTools/Private/MuT/ASTOpLayoutRemoveBlocks.cpp
--- a/Engine/Plugins/Mutable/Source/MutableTools/Private/MuT/ASTOpLayoutRemoveBlocks.cpp
+++ b/Engine/Plugins/Mutable/Source/MutableTools/Private/MuT/ASTOpLayoutRemoveBlocks.cpp
@@ -47,11 +47,14 @@ namespace UE::Mutable::Private
 	}
 
 
-	uint64 ASTOpLayoutRemoveBlocks::Hash() const
+	uint32 ASTOpLayoutRemoveBlocks::Hash() const
 	{
-		uint64 res = std::hash<void*>()(Source.child().get());
-		hash_combine(res, ReferenceLayout.child().get());
-		return res;
+		uint32 Result = GetTypeHash(GetOpType());
+
+		Result = HashCombineFast(Result, GetTypeHash(Source));
+		Result = HashCombineFast(Result, GetTypeHash(ReferenceLayout));
+
+		return Result;
 	}
 
 
diff --git a/Engine/Plugins/Mutable/Source/MutableTools/Private/MuT/ASTOpMaterialBreak.cpp b/Engine/Plugins/Mutable/Source/MutableTools/Private/MuT/ASTOpMaterialBreak.cpp
--- a/Engine/Plugins/Mutable/Source/MutableTools/Private/MuT/ASTOpMaterialBreak.cpp
+++ b/Engine/Plugins/Mutable/Source/MutableTools/Private/MuT/ASTOpMaterialBreak.cpp
@@ -61,13 +61,14 @@ namespace UE::Mutable::Private
 	}
 
 
-	uint64 ASTOpMaterialBreak::Hash() const
+	uint32 ASTOpMaterialBreak::Hash() const
 	{
-		uint64 res = std::hash<uint64>()(uint64(Type));
-		hash_combine(res, GetTypeHash(ParameterKey));
-		hash_combine(res, Material.child().get());
+		uint32 Result = GetTypeHash(GetOpType());
+	
+		Result = HashCombineFast(Result, GetTypeHash(ParameterKey));
+		Result = HashCombineFast(Result, GetTypeHash(Material));
 		
-		return res;
+		return Result;
 	}
 
 
diff --git a/Engine/Plugins/Mutable/Source/MutableTools/Private/MuT/ASTOpMaterialBreakImageParameter.cpp b/Engine/Plugins/Mutable/Source/MutableTools/Private/MuT/ASTOpMaterialBreakImageParameter.cpp
--- a/Engine/Plugins/Mutable/Source/MutableTools/Private/MuT/ASTOpMaterialBreakImageParameter.cpp
+++ b/Engine/Plugins/Mutable/Source/MutableTools/Private/MuT/ASTOpMaterialBreakImageParameter.cpp
@@ -51,13 +51,14 @@ namespace UE::Mutable::Private
 	}
 
 
-	uint64 ASTOpImageFromMaterialParameter::Hash() const
+	uint32 ASTOpImageFromMaterialParameter::Hash() const
 	{
-		uint64 res = std::hash<EOpType>()(GetOpType());
-		hash_combine(res, GetTypeHash(ParameterName));
-		hash_combine(res, GetTypeHash(MaterialParameter.Name));
+		uint32 Result = GetTypeHash(GetOpType());
 		
-		return res;
+		Result = HashCombineFast(Result, GetTypeHash(ParameterName));
+		Result = HashCombineFast(Result, GetTypeHash(MaterialParameter.Name));
+		
+		return Result;
 	}
 
 
diff --git a/Engine/Plugins/Mutable/Source/MutableTools/Private/MuT/ASTOpMaterialModify.cpp b/Engine/Plugins/Mutable/Source/MutableTools/Private/MuT/ASTOpMaterialModify.cpp
--- a/Engine/Plugins/Mutable/Source/MutableTools/Private/MuT/ASTOpMaterialModify.cpp
+++ b/Engine/Plugins/Mutable/Source/MutableTools/Private/MuT/ASTOpMaterialModify.cpp
@@ -77,24 +77,25 @@ namespace UE::Mutable::Private
 	}
 
 
-	uint64 ASTOpMaterialModify::Hash() const
+	uint32 ASTOpMaterialModify::Hash() const
 	{
-		uint64 res = std::hash<uint64>()(uint64(GetOpType()));
-		hash_combine(res, Material.child().get());
+		uint32 Result = GetTypeHash(GetOpType());
+
+		Result = HashCombineFast(Result, GetTypeHash(Material));
 
 		for (const TPair<FParameterKey, ASTChild>& Parameter : ParametersToModify)
 		{
-			hash_combine(res, GetTypeHash(Parameter.Key));
-			hash_combine(res, Parameter.Value.child().get());
+			Result = HashCombineFast(Result, GetTypeHash(Parameter.Key));
+			Result = HashCombineFast(Result, GetTypeHash(Parameter.Value));
 		}
 
 		for (const TPair<FParameterKey, int32>& Parameter : ImagePropertyIndexMap)
 		{
-			hash_combine(res, GetTypeHash(Parameter.Key));
-			hash_combine(res, GetTypeHash(Parameter.Value));
+			Result = HashCombineFast(Result, GetTypeHash(Parameter.Key));
+			Result = HashCombineFast(Result, GetTypeHash(Parameter.Value));
 		}
 
-		return res;
+		return Result;
 	}
 
 
diff --git a/Engine/Plugins/Mutable/Source/MutableTools/Private/MuT/ASTOpMaterialSkeletalMeshBreak.cpp b/Engine/Plugins/Mutable/Source/MutableTools/Private/MuT/ASTOpMaterialSkeletalMeshBreak.cpp
--- a/Engine/Plugins/Mutable/Source/MutableTools/Private/MuT/ASTOpMaterialSkeletalMeshBreak.cpp
+++ b/Engine/Plugins/Mutable/Source/MutableTools/Private/MuT/ASTOpMaterialSkeletalMeshBreak.cpp
@@ -22,12 +22,14 @@ namespace UE::Mutable::Private
 		return EOpType::MI_SKELETALMESH_BREAK;
 	}
 
-	uint64 ASTOpMaterialSkeletalMeshBreak::Hash() const
+	uint32 ASTOpMaterialSkeletalMeshBreak::Hash() const
 	{
-		uint64 Hash = std::hash<uint64>()(uint64(GetOpType()));
-		hash_combine(Hash, GetTypeHash(SlotName));
-		hash_combine(Hash, SkeletalMesh.child().get());
-		return Hash;
+		uint32 Result = GetTypeHash(GetOpType());
+
+		Result = HashCombineFast(Result, GetTypeHash(SlotName));
+		Result = HashCombineFast(Result, GetTypeHash(SkeletalMesh));
+		
+		return Result;
 	}
 
 	void ASTOpMaterialSkeletalMeshBreak::ForEachChild(const TFunctionRef<void(ASTChild&)> F)
diff --git a/Engine/Plugins/Mutable/Source/MutableTools/Private/MuT/ASTOpMeshAddMetadata.cpp b/Engine/Plugins/Mutable/Source/MutableTools/Private/MuT/ASTOpMeshAddMetadata.cpp
--- a/Engine/Plugins/Mutable/Source/MutableTools/Private/MuT/ASTOpMeshAddMetadata.cpp
+++ b/Engine/Plugins/Mutable/Source/MutableTools/Private/MuT/ASTOpMeshAddMetadata.cpp
@@ -61,9 +61,12 @@ namespace UE::Mutable::Private
 	}
 
 
-	uint64 ASTOpMeshAddMetadata::Hash() const
+	uint32 ASTOpMeshAddMetadata::Hash() const
 	{
-		uint64 Result = std::hash<ASTOp*>()(Source.child().get());
+		uint32 Result = GetTypeHash(GetOpType());
+
+		Result = HashCombineFast(Result, GetTypeHash(Source));
+
 		return Result;
 	}
 
diff --git a/Engine/Plugins/Mutable/Source/MutableTools/Private/MuT/ASTOpMeshApplyLayout.cpp b/Engine/Plugins/Mutable/Source/MutableTools/Private/MuT/ASTOpMeshApplyLayout.cpp
--- a/Engine/Plugins/Mutable/Source/MutableTools/Private/MuT/ASTOpMeshApplyLayout.cpp
+++ b/Engine/Plugins/Mutable/Source/MutableTools/Private/MuT/ASTOpMeshApplyLayout.cpp
@@ -39,11 +39,14 @@ namespace UE::Mutable::Private
 	}
 
 
-	uint64 ASTOpMeshApplyLayout::Hash() const
+	uint32 ASTOpMeshApplyLayout::Hash() const
 	{
-		uint64 Result = std::hash<void*>()(Mesh.child().get());
-		hash_combine(Result, Layout.child().get());
-		hash_combine(Result, Channel);
+		uint32 Result = GetTypeHash(GetOpType());
+
+		Result = HashCombineFast(Result, GetTypeHash(Mesh));
+		Result = HashCombineFast(Result, GetTypeHash(Layout));
+		Result = HashCombineFast(Result, GetTypeHash(Channel));
+		
 		return Result;
 	}
 
diff --git a/Engine/Plugins/Mutable/Source/MutableTools/Private/MuT/ASTOpMeshApplyPose.cpp b/Engine/Plugins/Mutable/Source/MutableTools/Private/MuT/ASTOpMeshApplyPose.cpp
--- a/Engine/Plugins/Mutable/Source/MutableTools/Private/MuT/ASTOpMeshApplyPose.cpp
+++ b/Engine/Plugins/Mutable/Source/MutableTools/Private/MuT/ASTOpMeshApplyPose.cpp
@@ -37,10 +37,13 @@ namespace UE::Mutable::Private
 	}
 
 
-	uint64 ASTOpMeshApplyPose::Hash() const
+	uint32 ASTOpMeshApplyPose::Hash() const
 	{
-		uint64 Result = std::hash<void*>()(Base.child().get());
-		hash_combine(Result, Pose.child().get());
+		uint32 Result = GetTypeHash(GetOpType());
+
+		Result = HashCombineFast(Result, GetTypeHash(Base));
+		Result = HashCombineFast(Result, GetTypeHash(Pose));
+		
 		return Result;
 	}
 
diff --git a/Engine/Plugins/Mutable/Source/MutableTools/Private/MuT/ASTOpMeshApplyShape.cpp b/Engine/Plugins/Mutable/Source/MutableTools/Private/MuT/ASTOpMeshApplyShape.cpp
--- a/Engine/Plugins/Mutable/Source/MutableTools/Private/MuT/ASTOpMeshApplyShape.cpp
+++ b/Engine/Plugins/Mutable/Source/MutableTools/Private/MuT/ASTOpMeshApplyShape.cpp
@@ -50,15 +50,17 @@ namespace UE::Mutable::Private
 	}
 
 
-	uint64 ASTOpMeshApplyShape::Hash() const
+	uint32 ASTOpMeshApplyShape::Hash() const
 	{
-		uint64 Result = std::hash<void*>()(Mesh.child().get());
-		hash_combine(Result, Shape.child().get());
-		hash_combine(Result, bool(bRecomputeNormals));
-		hash_combine(Result, bool(bReshapeSkeleton));
-		hash_combine(Result, bool(bReshapePhysicsVolumes));
-		hash_combine(Result, bool(bReshapeVertices));
-		hash_combine(Result, bool(bApplyLaplacian));
+		uint32 Result = GetTypeHash(GetOpType());
+
+		Result = HashCombineFast(Result, GetTypeHash(Mesh));
+		Result = HashCombineFast(Result, GetTypeHash(Shape));
+		Result = HashCombineFast(Result, GetTypeHash(bRecomputeNormals));
+		Result = HashCombineFast(Result, GetTypeHash(bReshapeSkeleton));
+		Result = HashCombineFast(Result, GetTypeHash(bReshapePhysicsVolumes));
+		Result = HashCombineFast(Result, GetTypeHash(bReshapeVertices));
+		Result = HashCombineFast(Result, GetTypeHash(bApplyLaplacian));
 
 		return Result;
 	}
diff --git a/Engine/Plugins/Mutable/Source/MutableTools/Private/MuT/ASTOpMeshBindShape.cpp b/Engine/Plugins/Mutable/Source/MutableTools/Private/MuT/ASTOpMeshBindShape.cpp
--- a/Engine/Plugins/Mutable/Source/MutableTools/Private/MuT/ASTOpMeshBindShape.cpp
+++ b/Engine/Plugins/Mutable/Source/MutableTools/Private/MuT/ASTOpMeshBindShape.cpp
@@ -66,30 +66,31 @@ namespace UE::Mutable::Private
 	}
 
 
-	uint64 ASTOpMeshBindShape::Hash() const
+	uint32 ASTOpMeshBindShape::Hash() const
 	{
-		uint64 Result = std::hash<void*>()(Mesh.child().get());
-		hash_combine(Result, Shape.child().get());
-		hash_combine(Result, bool(bRecomputeNormals));
-		hash_combine(Result, bool(bReshapeSkeleton));
-		hash_combine(Result, bool(bReshapePhysicsVolumes));
-		hash_combine(Result, bool(bReshapeVertices));
-		hash_combine(Result, bool(bApplyLaplacian));
-		hash_combine(Result, bool(BindingMethod));
-
-		hash_combine(Result, static_cast<uint32>(RChannelUsage));
-		hash_combine(Result, static_cast<uint32>(GChannelUsage));
-		hash_combine(Result, static_cast<uint32>(BChannelUsage));
-		hash_combine(Result, static_cast<uint32>(AChannelUsage));
+		uint32 Result = GetTypeHash(GetOpType());
+		
+		Result = HashCombineFast(Result, GetTypeHash(Mesh));
+		Result = HashCombineFast(Result, GetTypeHash(Shape));
+		Result = HashCombineFast(Result, GetTypeHash(bRecomputeNormals));
+		Result = HashCombineFast(Result, GetTypeHash(bReshapeSkeleton));
+		Result = HashCombineFast(Result, GetTypeHash(bReshapePhysicsVolumes));
+		Result = HashCombineFast(Result, GetTypeHash(bReshapeVertices));
+		Result = HashCombineFast(Result, GetTypeHash(bApplyLaplacian));
+		Result = HashCombineFast(Result, GetTypeHash(BindingMethod));
+		Result = HashCombineFast(Result, GetTypeHash(RChannelUsage));
+		Result = HashCombineFast(Result, GetTypeHash(GChannelUsage));
+		Result = HashCombineFast(Result, GetTypeHash(BChannelUsage));
+		Result = HashCombineFast(Result, GetTypeHash(AChannelUsage));
 
 		for (const FBoneName& S : BonesToDeform)
 		{
-			hash_combine(Result, S.Id);
+			Result = HashCombineFast(Result, GetTypeHash(S.Id));
 		}
 
 		for (const FBoneName& S : PhysicsToDeform)
 		{
-			hash_combine(Result, S.Id);
+			Result = HashCombineFast(Result, GetTypeHash(S.Id));
 		}
 
 		return Result;
diff --git a/Engine/Plugins/Mutable/Source/MutableTools/Private/MuT/ASTOpMeshClipDeform.cpp b/Engine/Plugins/Mutable/Source/MutableTools/Private/MuT/ASTOpMeshClipDeform.cpp
--- a/Engine/Plugins/Mutable/Source/MutableTools/Private/MuT/ASTOpMeshClipDeform.cpp
+++ b/Engine/Plugins/Mutable/Source/MutableTools/Private/MuT/ASTOpMeshClipDeform.cpp
@@ -38,11 +38,14 @@ namespace UE::Mutable::Private
 	}
 
 
-	uint64 ASTOpMeshClipDeform::Hash() const
+	uint32 ASTOpMeshClipDeform::Hash() const
 	{
-		uint64 res = std::hash<void*>()(Mesh.child().get());
-		hash_combine(res, ClipShape.child().get());
-		return res;
+		uint32 Result = GetTypeHash(GetOpType());
+
+		Result = HashCombineFast(Result, GetTypeHash(Mesh));
+		Result = HashCombineFast(Result, GetTypeHash(ClipShape));
+		
+		return Result;
 	}
 
 
diff --git a/Engine/Plugins/Mutable/Source/MutableTools/Private/MuT/ASTOpMeshClipMorphPlane.cpp b/Engine/Plugins/Mutable/Source/MutableTools/Private/MuT/ASTOpMeshClipMorphPlane.cpp
--- a/Engine/Plugins/Mutable/Source/MutableTools/Private/MuT/ASTOpMeshClipMorphPlane.cpp
+++ b/Engine/Plugins/Mutable/Source/MutableTools/Private/MuT/ASTOpMeshClipMorphPlane.cpp
@@ -49,12 +49,15 @@ namespace UE::Mutable::Private
 	}
 
 
-	uint64 ASTOpMeshClipMorphPlane::Hash() const
+	uint32 ASTOpMeshClipMorphPlane::Hash() const
 	{
-		uint64 res = std::hash<void*>()(Source.child().get());
-		hash_combine(res, Factor);
-		hash_combine(res, Dist);
-		return res;
+		uint32 Result = GetTypeHash(GetOpType());
+
+		Result = HashCombineFast(Result, GetTypeHash(Source));
+		Result = HashCombineFast(Result, GetTypeHash(Factor));
+		Result = HashCombineFast(Result, GetTypeHash(Dist));
+		
+		return Result;
 	}
 
 
diff --git a/Engine/Plugins/Mutable/Source/MutableTools/Private/MuT/ASTOpMeshClipWithMesh.cpp b/Engine/Plugins/Mutable/Source/MutableTools/Private/MuT/ASTOpMeshClipWithMesh.cpp
--- a/Engine/Plugins/Mutable/Source/MutableTools/Private/MuT/ASTOpMeshClipWithMesh.cpp
+++ b/Engine/Plugins/Mutable/Source/MutableTools/Private/MuT/ASTOpMeshClipWithMesh.cpp
@@ -37,10 +37,13 @@ namespace UE::Mutable::Private
 	}
 
 
-	uint64 ASTOpMeshClipWithMesh::Hash() const
+	uint32 ASTOpMeshClipWithMesh::Hash() const
 	{
-		uint64 Result = std::hash<void*>()(Source.child().get());
-		hash_combine(Result, ClipMesh.child().get());
+		uint32 Result = GetTypeHash(GetOpType());
+
+		Result = HashCombineFast(Result, GetTypeHash(Source));
+		Result = HashCombineFast(Result, GetTypeHash(ClipMesh));
+		
 		return Result;
 	}
 
diff --git a/Engine/Plugins/Mutable/Source/MutableTools/Private/MuT/ASTOpMeshDifference.cpp b/Engine/Plugins/Mutable/Source/MutableTools/Private/MuT/ASTOpMeshDifference.cpp
--- a/Engine/Plugins/Mutable/Source/MutableTools/Private/MuT/ASTOpMeshDifference.cpp
+++ b/Engine/Plugins/Mutable/Source/MutableTools/Private/MuT/ASTOpMeshDifference.cpp
@@ -66,11 +66,14 @@ namespace UE::Mutable::Private
 
 
 	//---------------------------------------------------------------------------------------------
-	uint64 ASTOpMeshDifference::Hash() const
+	uint32 ASTOpMeshDifference::Hash() const
 	{
-		uint64 res = std::hash<ASTOp*>()(Base.child().get());
-		hash_combine(res, Target.child().get());
-		return res;
+		uint32 Result = GetTypeHash(GetOpType());
+
+		Result = HashCombineFast(Result, GetTypeHash(Base));
+		Result = HashCombineFast(Result, GetTypeHash(Target));
+		
+		return Result;
 	}
 
 
diff --git a/Engine/Plugins/Mutable/Source/MutableTools/Private/MuT/ASTOpMeshExtractLayoutBlocks.cpp b/Engine/Plugins/Mutable/Source/MutableTools/Private/MuT/ASTOpMeshExtractLayoutBlocks.cpp
--- a/Engine/Plugins/Mutable/Source/MutableTools/Private/MuT/ASTOpMeshExtractLayoutBlocks.cpp
+++ b/Engine/Plugins/Mutable/Source/MutableTools/Private/MuT/ASTOpMeshExtractLayoutBlocks.cpp
@@ -72,10 +72,13 @@ namespace UE::Mutable::Private
 	}
 
 
-	uint64 ASTOpMeshExtractLayoutBlocks::Hash() const
+	uint32 ASTOpMeshExtractLayoutBlocks::Hash() const
 	{
-		uint64 res = std::hash<size_t>()(size_t(Source.child().get()));
-		return res;
+		uint32 Result = GetTypeHash(GetOpType());
+
+		Result = HashCombineFast(Result, GetTypeHash(Source));
+
+		return Result;
 	}
 
 
diff --git a/Engine/Plugins/Mutable/Source/MutableTools/Private/MuT/ASTOpMeshFormat.cpp b/Engine/Plugins/Mutable/Source/MutableTools/Private/MuT/ASTOpMeshFormat.cpp
--- a/Engine/Plugins/Mutable/Source/MutableTools/Private/MuT/ASTOpMeshFormat.cpp
+++ b/Engine/Plugins/Mutable/Source/MutableTools/Private/MuT/ASTOpMeshFormat.cpp
@@ -57,11 +57,14 @@ bool ASTOpMeshFormat::IsEqual(const ASTOp& otherUntyped) const
 }
 
 
-uint64 ASTOpMeshFormat::Hash() const
+uint32 ASTOpMeshFormat::Hash() const
 {
-	uint64 res = std::hash<void*>()(Source.child().get() );
-    hash_combine( res, Format.child().get() );
-    return res;
+	uint32 Result = GetTypeHash(GetOpType());
+
+	Result = HashCombineFast(Result, GetTypeHash(Source));
+    Result = HashCombineFast(Result, GetTypeHash(Format));
+    
+	return Result;
 }
 
 
diff --git a/Engine/Plugins/Mutable/Source/MutableTools/Private/MuT/ASTOpMeshMaskClipMesh.cpp b/Engine/Plugins/Mutable/Source/MutableTools/Private/MuT/ASTOpMeshMaskClipMesh.cpp
--- a/Engine/Plugins/Mutable/Source/MutableTools/Private/MuT/ASTOpMeshMaskClipMesh.cpp
+++ b/Engine/Plugins/Mutable/Source/MutableTools/Private/MuT/ASTOpMeshMaskClipMesh.cpp
@@ -48,11 +48,14 @@ namespace UE::Mutable::Private
 
 
 	//-------------------------------------------------------------------------------------------------
-	uint64 ASTOpMeshMaskClipMesh::Hash() const
+	uint32 ASTOpMeshMaskClipMesh::Hash() const
 	{
-		uint64 res = std::hash<void*>()(source.child().get());
-		hash_combine(res, clip.child());
-		return res;
+		uint32 Result = GetTypeHash(GetOpType());
+
+		Result = HashCombineFast(Result, GetTypeHash(source));
+		Result = HashCombineFast(Result, GetTypeHash(clip));
+		
+		return Result;
 	}
 
 
diff --git a/Engine/Plugins/Mutable/Source/MutableTools/Private/MuT/ASTOpMeshMaskClipUVMask.cpp b/Engine/Plugins/Mutable/Source/MutableTools/Private/MuT/ASTOpMeshMaskClipUVMask.cpp
--- a/Engine/Plugins/Mutable/Source/MutableTools/Private/MuT/ASTOpMeshMaskClipUVMask.cpp
+++ b/Engine/Plugins/Mutable/Source/MutableTools/Private/MuT/ASTOpMeshMaskClipUVMask.cpp
@@ -54,13 +54,16 @@ namespace UE::Mutable::Private
 
 
 	//-------------------------------------------------------------------------------------------------
-	uint64 ASTOpMeshMaskClipUVMask::Hash() const
+	uint32 ASTOpMeshMaskClipUVMask::Hash() const
 	{
-		uint64 res = std::hash<void*>()(Source.child().get());
-		hash_combine(res, UVSource.child());
-		hash_combine(res, MaskImage.child());
-		hash_combine(res, MaskLayout.child());
-		return res;
+		uint32 Result = GetTypeHash(GetOpType());
+		
+		Result = HashCombineFast(Result, GetTypeHash(Source));
+		Result = HashCombineFast(Result, GetTypeHash(UVSource));
+		Result = HashCombineFast(Result, GetTypeHash(MaskImage));
+		Result = HashCombineFast(Result, GetTypeHash(MaskLayout));
+
+		return Result;
 	}
 
 
diff --git a/Engine/Plugins/Mutable/Source/MutableTools/Private/MuT/ASTOpMeshMaskDiff.cpp b/Engine/Plugins/Mutable/Source/MutableTools/Private/MuT/ASTOpMeshMaskDiff.cpp
--- a/Engine/Plugins/Mutable/Source/MutableTools/Private/MuT/ASTOpMeshMaskDiff.cpp
+++ b/Engine/Plugins/Mutable/Source/MutableTools/Private/MuT/ASTOpMeshMaskDiff.cpp
@@ -38,10 +38,13 @@ namespace UE::Mutable::Private
 	}
 
 
-	uint64 ASTOpMeshMaskDiff::Hash() const
+	uint32 ASTOpMeshMaskDiff::Hash() const
 	{
-		uint64 Result = std::hash<void*>()(Source.child().get());
-		hash_combine(Result, Fragment.child().get());
+		uint32 Result = GetTypeHash(GetOpType());
+
+		Result = HashCombineFast(Result, GetTypeHash(Source));
+		Result = HashCombineFast(Result, GetTypeHash(Fragment));
+		
 		return Result;
 	}
 
diff --git a/Engine/Plugins/Mutable/Source/MutableTools/Private/MuT/ASTOpMeshMerge.cpp b/Engine/Plugins/Mutable/Source/MutableTools/Private/MuT/ASTOpMeshMerge.cpp
--- a/Engine/Plugins/Mutable/Source/MutableTools/Private/MuT/ASTOpMeshMerge.cpp
+++ b/Engine/Plugins/Mutable/Source/MutableTools/Private/MuT/ASTOpMeshMerge.cpp
@@ -37,11 +37,14 @@ namespace UE::Mutable::Private
 	}
 
 
-	uint64 ASTOpMeshMerge::Hash() const
+	uint32 ASTOpMeshMerge::Hash() const
 	{
-		uint64 Result = std::hash<void*>()(Base.child().get());
-		hash_combine(Result, Added.child().get());
-		hash_combine(Result, NewSurfaceID);
+		uint32 Result = GetTypeHash(GetOpType());
+		
+		Result = HashCombineFast(Result, GetTypeHash(Base));
+		Result = HashCombineFast(Result, GetTypeHash(Added));
+		Result = HashCombineFast(Result, GetTypeHash(NewSurfaceID));
+
 		return Result;
 	}
 
diff --git a/Engine/Plugins/Mutable/Source/MutableTools/Private/MuT/ASTOpMeshMorph.cpp b/Engine/Plugins/Mutable/Source/MutableTools/Private/MuT/ASTOpMeshMorph.cpp
--- a/Engine/Plugins/Mutable/Source/MutableTools/Private/MuT/ASTOpMeshMorph.cpp
+++ b/Engine/Plugins/Mutable/Source/MutableTools/Private/MuT/ASTOpMeshMorph.cpp
@@ -61,13 +61,16 @@ namespace UE::Mutable::Private
 
 
 	//---------------------------------------------------------------------------------------------
-	uint64 ASTOpMeshMorph::Hash() const
+	uint32 ASTOpMeshMorph::Hash() const
 	{
-		uint64 res = GetTypeHash(Name);
-		hash_combine(res, Factor.child().get());
-		hash_combine(res, Base.child().get());
-		hash_combine(res, Target.child().get());
-		return res;
+		uint32 Result = GetTypeHash(GetOpType());
+
+		Result = HashCombineFast(Result, GetTypeHash(Name));
+		Result = HashCombineFast(Result, GetTypeHash(Factor));
+		Result = HashCombineFast(Result, GetTypeHash(Base));
+		Result = HashCombineFast(Result, GetTypeHash(Target));
+		
+		return Result;
 	}
 
 
diff --git a/Engine/Plugins/Mutable/Source/MutableTools/Private/MuT/ASTOpMeshMorphReshape.cpp b/Engine/Plugins/Mutable/Source/MutableTools/Private/MuT/ASTOpMeshMorphReshape.cpp
--- a/Engine/Plugins/Mutable/Source/MutableTools/Private/MuT/ASTOpMeshMorphReshape.cpp
+++ b/Engine/Plugins/Mutable/Source/MutableTools/Private/MuT/ASTOpMeshMorphReshape.cpp
@@ -38,12 +38,14 @@ namespace UE::Mutable::Private
 	}
 
 
-	uint64 ASTOpMeshMorphReshape::Hash() const
+	uint32 ASTOpMeshMorphReshape::Hash() const
 	{
-		uint64 res = std::hash<void*>()(Morph.child().get());
-		hash_combine(res, Reshape.child().get());
+		uint32 Result = GetTypeHash(GetOpType());
 
-		return res;
+		Result = HashCombineFast(Result, GetTypeHash(Morph));
+		Result = HashCombineFast(Result, GetTypeHash(Reshape));
+
+		return Result;
 	}
 
 
diff --git a/Engine/Plugins/Mutable/Source/MutableTools/Private/MuT/ASTOpMeshOptimizeSkinning.cpp b/Engine/Plugins/Mutable/Source/MutableTools/Private/MuT/ASTOpMeshOptimizeSkinning.cpp
--- a/Engine/Plugins/Mutable/Source/MutableTools/Private/MuT/ASTOpMeshOptimizeSkinning.cpp
+++ b/Engine/Plugins/Mutable/Source/MutableTools/Private/MuT/ASTOpMeshOptimizeSkinning.cpp
@@ -40,10 +40,13 @@ namespace UE::Mutable::Private
 	}
 
 
-	uint64 ASTOpMeshOptimizeSkinning::Hash() const
+	uint32 ASTOpMeshOptimizeSkinning::Hash() const
 	{
-		uint64 Res = std::hash<void*>()(Source.child().get());
-		return Res;
+		uint32 Result = GetTypeHash(GetOpType());
+		
+		Result = HashCombineFast(Result, GetTypeHash(Source));
+		
+		return Result;
 	}
 
 
diff --git a/Engine/Plugins/Mutable/Source/MutableTools/Private/MuT/ASTOpMeshPrepareLayout.cpp b/Engine/Plugins/Mutable/Source/MutableTools/Private/MuT/ASTOpMeshPrepareLayout.cpp
--- a/Engine/Plugins/Mutable/Source/MutableTools/Private/MuT/ASTOpMeshPrepareLayout.cpp
+++ b/Engine/Plugins/Mutable/Source/MutableTools/Private/MuT/ASTOpMeshPrepareLayout.cpp
@@ -42,11 +42,14 @@ namespace UE::Mutable::Private
 	}
 
 
-	uint64 ASTOpMeshPrepareLayout::Hash() const
+	uint32 ASTOpMeshPrepareLayout::Hash() const
 	{
-		uint64 Result = std::hash<void*>()(Mesh.child().get());
-		hash_combine(Result, Layout.child().get());
-		hash_combine(Result, bUseAbsoluteBlockIds);
+		uint32 Result = GetTypeHash(GetOpType());
+
+		Result = HashCombineFast(Result, GetTypeHash(Mesh));
+		Result = HashCombineFast(Result, GetTypeHash(Layout));
+		Result = HashCombineFast(Result, GetTypeHash(bUseAbsoluteBlockIds));
+		
 		return Result;
 	}
 
diff --git a/Engine/Plugins/Mutable/Source/MutableTools/Private/MuT/ASTOpMeshProject.cpp b/Engine/Plugins/Mutable/Source/MutableTools/Private/MuT/ASTOpMeshProject.cpp
--- a/Engine/Plugins/Mutable/Source/MutableTools/Private/MuT/ASTOpMeshProject.cpp
+++ b/Engine/Plugins/Mutable/Source/MutableTools/Private/MuT/ASTOpMeshProject.cpp
@@ -39,10 +39,13 @@ namespace UE::Mutable::Private
 	}
 
 
-	uint64 ASTOpMeshProject::Hash() const
+	uint32 ASTOpMeshProject::Hash() const
 	{
-		uint64 Result = std::hash<void*>()(Mesh.child().get());
-		hash_combine(Result, Projector.child().get());
+		uint32 Result = GetTypeHash(GetOpType());
+
+		Result = HashCombineFast(Result, GetTypeHash(Mesh));
+		Result = HashCombineFast(Result, GetTypeHash(Projector));
+		
 		return Result;
 	}
 
diff --git a/Engine/Plugins/Mutable/Source/MutableTools/Private/MuT/ASTOpMeshRemoveMask.cpp b/Engine/Plugins/Mutable/Source/MutableTools/Private/MuT/ASTOpMeshRemoveMask.cpp
--- a/Engine/Plugins/Mutable/Source/MutableTools/Private/MuT/ASTOpMeshRemoveMask.cpp
+++ b/Engine/Plugins/Mutable/Source/MutableTools/Private/MuT/ASTOpMeshRemoveMask.cpp
@@ -75,15 +75,19 @@ namespace UE::Mutable::Private
 
 
 	//---------------------------------------------------------------------------------------------
-	uint64 ASTOpMeshRemoveMask::Hash() const
+	uint32 ASTOpMeshRemoveMask::Hash() const
 	{
-		uint64 res = std::hash<ASTOp*>()(source.child().get());
-		for (const TPair<ASTChild, ASTChild>& r : removes)
+		uint32 Result = GetTypeHash(GetOpType());
+
+		Result = HashCombineFast(Result, GetTypeHash(source));
+		
+		for (const TPair<ASTChild, ASTChild>& Remove : removes)
 		{
-			hash_combine(res, r.Key.child().get());
-			hash_combine(res, r.Value.child().get());
+			Result = HashCombineFast(Result, GetTypeHash(Remove.Key));
+			Result = HashCombineFast(Result, GetTypeHash(Remove.Value));
 		}
-		return res;
+		
+		return Result;
 	}
 
 
diff --git a/Engine/Plugins/Mutable/Source/MutableTools/Private/MuT/ASTOpMeshSetSkeleton.cpp b/Engine/Plugins/Mutable/Source/MutableTools/Private/MuT/ASTOpMeshSetSkeleton.cpp
--- a/Engine/Plugins/Mutable/Source/MutableTools/Private/MuT/ASTOpMeshSetSkeleton.cpp
+++ b/Engine/Plugins/Mutable/Source/MutableTools/Private/MuT/ASTOpMeshSetSkeleton.cpp
@@ -36,10 +36,13 @@ namespace UE::Mutable::Private
 	}
 
 
-	uint64 ASTOpMeshSetSkeleton::Hash() const
+	uint32 ASTOpMeshSetSkeleton::Hash() const
 	{
-		uint64 Result = std::hash<void*>()(Source.child().get());
-		hash_combine(Result, Skeleton.child().get());
+		uint32 Result = GetTypeHash(GetOpType());
+
+		Result = HashCombineFast(Result, GetTypeHash(Source));
+		Result = HashCombineFast(Result, GetTypeHash(Skeleton));
+		
 		return Result;
 	}
 
diff --git a/Engine/Plugins/Mutable/Source/MutableTools/Private/MuT/ASTOpMeshSkeletalMeshBreak.cpp b/Engine/Plugins/Mutable/Source/MutableTools/Private/MuT/ASTOpMeshSkeletalMeshBreak.cpp
--- a/Engine/Plugins/Mutable/Source/MutableTools/Private/MuT/ASTOpMeshSkeletalMeshBreak.cpp
+++ b/Engine/Plugins/Mutable/Source/MutableTools/Private/MuT/ASTOpMeshSkeletalMeshBreak.cpp
@@ -61,16 +61,16 @@ namespace UE::Mutable::Private
 	}
 
 
-	uint64 ASTOpMeshSkeletalMeshBreak::Hash() const
+	uint32 ASTOpMeshSkeletalMeshBreak::Hash() const
 	{
-		uint64 res = std::hash<uint64>()(uint64(Type));
+		uint32 Result = GetTypeHash(GetOpType());
 
-		hash_combine(res, SkeletalMeshParameter.child().get());
-		hash_combine(res, LODIndex);
-		hash_combine(res, SectionIndex);
-		hash_combine(res, MeshID);
+		Result = HashCombineFast(Result, GetTypeHash(SkeletalMeshParameter));
+		Result = HashCombineFast(Result, GetTypeHash(LODIndex));
+		Result = HashCombineFast(Result, GetTypeHash(SectionIndex));
+		Result = HashCombineFast(Result, GetTypeHash(MeshID));
 
-		return res;
+		return Result;
 	}
 
 
diff --git a/Engine/Plugins/Mutable/Source/MutableTools/Private/MuT/ASTOpMeshTransform.cpp b/Engine/Plugins/Mutable/Source/MutableTools/Private/MuT/ASTOpMeshTransform.cpp
--- a/Engine/Plugins/Mutable/Source/MutableTools/Private/MuT/ASTOpMeshTransform.cpp
+++ b/Engine/Plugins/Mutable/Source/MutableTools/Private/MuT/ASTOpMeshTransform.cpp
@@ -36,11 +36,14 @@ namespace UE::Mutable::Private
 	}
 
 
-	uint64 ASTOpMeshTransform::Hash() const
+	uint32 ASTOpMeshTransform::Hash() const
 	{
-		uint64 res = std::hash<EOpType>()(EOpType::ME_TRANSFORM);
-		hash_combine(res, source.child().get());
-		return res;
+		uint32 Result = GetTypeHash(GetOpType());
+		
+		Result = HashCombineFast(Result, GetTypeHash(source));
+		Result = HashCombineFast(Result, matrix.ComputeHash());
+
+		return Result;
 	}
 
 
diff --git a/Engine/Plugins/Mutable/Source/MutableTools/Private/MuT/ASTOpMeshTransformWithBone.cpp b/Engine/Plugins/Mutable/Source/MutableTools/Private/MuT/ASTOpMeshTransformWithBone.cpp
--- a/Engine/Plugins/Mutable/Source/MutableTools/Private/MuT/ASTOpMeshTransformWithBone.cpp
+++ b/Engine/Plugins/Mutable/Source/MutableTools/Private/MuT/ASTOpMeshTransformWithBone.cpp
@@ -17,15 +17,16 @@ UE::Mutable::Private::ASTOpMeshTransformWithBone::~ASTOpMeshTransformWithBone()
 	ASTOp::RemoveChildren();
 }
 
-uint64 UE::Mutable::Private::ASTOpMeshTransformWithBone::Hash() const
+uint32 UE::Mutable::Private::ASTOpMeshTransformWithBone::Hash() const
 {
-	uint64 res = std::hash<EOpType>()(GetOpType());
-	hash_combine(res, SourceMesh.child());
-	hash_combine(res, Matrix.child());
-	hash_combine(res, BoneName.Id);
-	hash_combine(res, ThresholdFactor);
-
-	return res;
+	uint32 Result = GetTypeHash(GetOpType());
+	
+	Result = HashCombineFast(Result, GetTypeHash(SourceMesh));
+	Result = HashCombineFast(Result, GetTypeHash(Matrix));
+	Result = HashCombineFast(Result, GetTypeHash(BoneName));
+	Result = HashCombineFast(Result, GetTypeHash(ThresholdFactor));
+
+	return Result;
 }
 
 bool UE::Mutable::Private::ASTOpMeshTransformWithBone::IsEqual(const ASTOp& OtherUntyped) const
diff --git a/Engine/Plugins/Mutable/Source/MutableTools/Private/MuT/ASTOpMeshTransformWithBoundingMesh.cpp b/Engine/Plugins/Mutable/Source/MutableTools/Private/MuT/ASTOpMeshTransformWithBoundingMesh.cpp
--- a/Engine/Plugins/Mutable/Source/MutableTools/Private/MuT/ASTOpMeshTransformWithBoundingMesh.cpp
+++ b/Engine/Plugins/Mutable/Source/MutableTools/Private/MuT/ASTOpMeshTransformWithBoundingMesh.cpp
@@ -18,14 +18,15 @@ UE::Mutable::Private::ASTOpMeshTransformWithBoundingMesh::~ASTOpMeshTransformWit
 	ASTOp::RemoveChildren();
 }
 
-uint64 UE::Mutable::Private::ASTOpMeshTransformWithBoundingMesh::Hash() const
+uint32 UE::Mutable::Private::ASTOpMeshTransformWithBoundingMesh::Hash() const
 {
-	uint64 res = std::hash<EOpType>()(GetOpType());
-	hash_combine(res, source.child());
-	hash_combine(res, boundingMesh.child());
-	hash_combine(res, matrix.child());
+	uint32 Result = GetTypeHash(GetOpType());
+	
+	Result = HashCombineFast(Result, GetTypeHash(source));
+	Result = HashCombineFast(Result, GetTypeHash(boundingMesh));
+	Result = HashCombineFast(Result, GetTypeHash(matrix));
 
-	return res;
+	return Result;
 }
 
 bool UE::Mutable::Private::ASTOpMeshTransformWithBoundingMesh::IsEqual(const ASTOp& otherUntyped) const
diff --git a/Engine/Plugins/Mutable/Source/MutableTools/Private/MuT/ASTOpParameter.cpp b/Engine/Plugins/Mutable/Source/MutableTools/Private/MuT/ASTOpParameter.cpp
--- a/Engine/Plugins/Mutable/Source/MutableTools/Private/MuT/ASTOpParameter.cpp
+++ b/Engine/Plugins/Mutable/Source/MutableTools/Private/MuT/ASTOpParameter.cpp
@@ -73,13 +73,14 @@ namespace UE::Mutable::Private
 	}
 
 
-	uint64 ASTOpParameter::Hash() const
+	uint32 ASTOpParameter::Hash() const
 	{
-		uint64 res = std::hash<uint64>()(uint64(Type));
-		hash_combine(res, Parameter.Type);
-		hash_combine(res, Parameter.Name.Len());
+		uint32 Result = GetTypeHash(GetOpType());
+		
+		Result = HashCombineFast(Result, GetTypeHash(Parameter.Type));
+		Result = HashCombineFast(Result, GetTypeHash(Parameter.Name.Len()));
 
-		return res;
+		return Result;
 	}
 
 
diff --git a/Engine/Plugins/Mutable/Source/MutableTools/Private/MuT/ASTOpReferenceResource.cpp b/Engine/Plugins/Mutable/Source/MutableTools/Private/MuT/ASTOpReferenceResource.cpp
--- a/Engine/Plugins/Mutable/Source/MutableTools/Private/MuT/ASTOpReferenceResource.cpp
+++ b/Engine/Plugins/Mutable/Source/MutableTools/Private/MuT/ASTOpReferenceResource.cpp
@@ -42,11 +42,13 @@ namespace UE::Mutable::Private
 
 
 	//-------------------------------------------------------------------------------------------------
-	uint64 ASTOpReferenceResource::Hash() const
+	uint32 ASTOpReferenceResource::Hash() const
 	{
-		uint64 res = std::hash<uint64>()(uint64(Type));
-		hash_combine(res, ID);
-		return res;
+		uint32 Result = GetTypeHash(GetOpType());
+		
+		Result = HashCombineFast(Result, GetTypeHash(ID));
+		
+		return Result;
 	}
 
 
diff --git a/Engine/Plugins/Mutable/Source/MutableTools/Private/MuT/ASTOpScalarArithmetic.cpp b/Engine/Plugins/Mutable/Source/MutableTools/Private/MuT/ASTOpScalarArithmetic.cpp
--- a/Engine/Plugins/Mutable/Source/MutableTools/Private/MuT/ASTOpScalarArithmetic.cpp
+++ b/Engine/Plugins/Mutable/Source/MutableTools/Private/MuT/ASTOpScalarArithmetic.cpp
@@ -37,11 +37,14 @@ namespace UE::Mutable::Private
 	}
 
 
-	uint64 ASTOpScalarArithmetic::Hash() const
+	uint32 ASTOpScalarArithmetic::Hash() const
 	{
-		uint64 Result = std::hash<void*>()(A.child().get());
-		hash_combine(Result, B.child().get());
-		hash_combine(Result, Operation);
+		uint32 Result = GetTypeHash(GetOpType());
+
+		Result = HashCombineFast(Result, GetTypeHash(A));
+		Result = HashCombineFast(Result, GetTypeHash(B));
+		Result = HashCombineFast(Result, GetTypeHash(Operation));
+		
 		return Result;
 	}
 
