[Mutable] Replace custom hash functions with engine ones.

diff --git a/Engine/Plugins/Mutable/Source/MutableTools/Private/MuT/ASTOpScalarCurve.cpp b/Engine/Plugins/Mutable/Source/MutableTools/Private/MuT/ASTOpScalarCurve.cpp
--- a/Engine/Plugins/Mutable/Source/MutableTools/Private/MuT/ASTOpScalarCurve.cpp
+++ b/Engine/Plugins/Mutable/Source/MutableTools/Private/MuT/ASTOpScalarCurve.cpp
@@ -31,11 +31,13 @@ namespace UE::Mutable::Private
 	}
 
 
-	uint64 ASTOpScalarCurve::Hash() const
+	uint32 ASTOpScalarCurve::Hash() const
 	{
-		uint64 res = std::hash<uint64>()(size_t(EOpType::SC_CURVE));
-		hash_combine(res, Curve.Keys.Num());
-		return res;
+		uint32 Result = GetTypeHash(GetOpType());
+
+		Result = HashCombineFast(Result, GetTypeHash(Curve));
+		
+		return Result;
 	}
 
 
diff --git a/Engine/Plugins/Mutable/Source/MutableTools/Private/MuT/ASTOpSwitch.cpp b/Engine/Plugins/Mutable/Source/MutableTools/Private/MuT/ASTOpSwitch.cpp
--- a/Engine/Plugins/Mutable/Source/MutableTools/Private/MuT/ASTOpSwitch.cpp
+++ b/Engine/Plugins/Mutable/Source/MutableTools/Private/MuT/ASTOpSwitch.cpp
@@ -85,13 +85,14 @@ namespace UE::Mutable::Private
 	}
 
 
-	uint64 ASTOpSwitch::Hash() const
+	uint32 ASTOpSwitch::Hash() const
 	{
-		uint64 Result = std::hash<uint64>()(uint64(Type));
+		uint32 Result = GetTypeHash(GetOpType());
+
 		for (const FCase& Case : Cases)
 		{
-			hash_combine(Result, Case.Condition);
-			hash_combine(Result, Case.Branch.child().get());
+			Result = HashCombineFast(Result, GetTypeHash(Case.Condition));
+			Result = HashCombineFast(Result, GetTypeHash(Case.Branch));
 		}
 
 		return Result;
