[Mutable] Reenable sync loads of referenced assets during cook.

#Synced-CL 49557694

diff --git a/Engine/Plugins/Mutable/Source/CustomizableObjectEditor/Private/MuCOE/CustomizableObjectCompiler.cpp b/Engine/Plugins/Mutable/Source/CustomizableObjectEditor/Private/MuCOE/CustomizableObjectCompiler.cpp
--- a/Engine/Plugins/Mutable/Source/CustomizableObjectEditor/Private/MuCOE/CustomizableObjectCompiler.cpp
+++ b/Engine/Plugins/Mutable/Source/CustomizableObjectEditor/Private/MuCOE/CustomizableObjectCompiler.cpp
@@ -171,7 +171,7 @@ void FCustomizableObjectCompiler::PreloadReferencerAssets()
 		UCustomizableObjectSystem* System = UCustomizableObjectSystem::GetInstanceChecked();
 		const TSharedRef<FMutableStreamableManager>& Streamable = System->GetPrivate()->StreamableManager;
 
-		if (bAsync)
+		if (bAsync && !CurrentOptions.bIsCooking)
 		{
 			AddCompileNotification(LOCTEXT("LoadingReferencerAssets", "Loading assets"));
 
