[Mutable] Change ModelResources UPROPERTY names to avoid type warnings.

#Synced-CL 49557694

diff --git a/Engine/Plugins/Mutable/Source/CustomizableObject/Internal/MuCO/CustomizableObjectPrivate.h b/Engine/Plugins/Mutable/Source/CustomizableObject/Internal/MuCO/CustomizableObjectPrivate.h
--- a/Engine/Plugins/Mutable/Source/CustomizableObject/Internal/MuCO/CustomizableObjectPrivate.h
+++ b/Engine/Plugins/Mutable/Source/CustomizableObject/Internal/MuCO/CustomizableObjectPrivate.h
@@ -1360,15 +1360,15 @@ public:
 	
 	/** Key is FComponentId. Max number of LODs in the compiled Model. */
 	UPROPERTY()
-	TMap<int32, uint8> NumLODsAvailable;
+	TMap<int32, uint8> NumLODs;
 
 	/** Key is FComponentId. Max number of LODs to stream. Mutable will always generate at least one LOD. */
 	UPROPERTY()
-	TMap<int32, uint8> NumLODsToStream;
+	TMap<int32, uint8> MaxNumStreamedLODs;
 
 	/** Key is FComponentId. First LOD available, some platforms may remove lower LODs when cooking, this MinLOD represents the first LOD we can generate */
 	UPROPERTY()
-	TMap<int32, uint8> FirstLODAvailable;
+	TMap<int32, uint8> FirstLODs;
 
 	/** Key is FComponentId. Name of all possible components. */
 	UPROPERTY()
@@ -1376,11 +1376,11 @@ public:
 
 	/** Key is FComponentId. Minimum LOD to render per Platform. */
 	UPROPERTY()
-	TMap<int32, FPerPlatformInt> MinLODPerComponent;
+	TMap<int32, FPerPlatformInt> MinLODs;
 
 	/** Key is FComponentId. Minimum LOD to render per Quality level. */
 	UPROPERTY()
-	TMap<int32, FPerQualityLevelInt> MinQualityLevelLODPerComponent;
+	TMap<int32, FPerQualityLevelInt> MinQualityLevelLODs;
 
 	UPROPERTY()
 	TMap<FString, TObjectPtr<UTexture>> TextureParameterDefaultValues;
diff --git a/Engine/Plugins/Mutable/Source/CustomizableObject/Private/MuCO/CustomizableObject.cpp b/Engine/Plugins/Mutable/Source/CustomizableObject/Private/MuCO/CustomizableObject.cpp
--- a/Engine/Plugins/Mutable/Source/CustomizableObject/Private/MuCO/CustomizableObject.cpp
+++ b/Engine/Plugins/Mutable/Source/CustomizableObject/Private/MuCO/CustomizableObject.cpp
@@ -1983,7 +1983,7 @@ uint8 UCustomizableObjectPrivate::GetMinLODIndex(const UE::Mutable::Private::FCo
 				// Value extracted from the Scalability.ini of the current platform (if found).
 				const int32 QualityLevel = UCustomizableObjectSystem::GetInstance()->GetSkeletalMeshMinLODQualityLevel();
 
-				const FPerQualityLevelInt* QualityLevelInt = LocalModelResources->MinQualityLevelLODPerComponent.Find(ComponentId);
+				const FPerQualityLevelInt* QualityLevelInt = LocalModelResources->MinQualityLevelLODs.Find(ComponentId);
 				if (ensure(QualityLevelInt))
 				{
 					MinLODIdx = QualityLevelInt->GetValue(QualityLevel);
@@ -1992,7 +1992,7 @@ uint8 UCustomizableObjectPrivate::GetMinLODIndex(const UE::Mutable::Private::FCo
 		}
 		else
 		{
-			const FPerPlatformInt* PlatformInt = LocalModelResources->MinLODPerComponent.Find(ComponentId);
+			const FPerPlatformInt* PlatformInt = LocalModelResources->MinLODs.Find(ComponentId);
 			if (ensure(PlatformInt))
 			{
 				MinLODIdx = PlatformInt->GetValue();
@@ -2005,7 +2005,7 @@ uint8 UCustomizableObjectPrivate::GetMinLODIndex(const UE::Mutable::Private::FCo
 
 	if (LocalModelResources)
 	{
-		if (const uint8* Result = LocalModelResources->FirstLODAvailable.Find(ComponentId))
+		if (const uint8* Result = LocalModelResources->FirstLODs.Find(ComponentId))
 		{
 			FirstLODAvailable = *Result;
 		}
@@ -3465,8 +3465,8 @@ void UCustomizableObjectPrivate::LogMemory()
 				Arch << ModelResources->ClothSharedConfigsData;
 				Arch << ModelResources->ClothingAssetsData;
 				Arch << ModelResources->ComponentNamesPerObjectComponent;
-				Arch << ModelResources->MinLODPerComponent;
-				Arch << ModelResources->MinQualityLevelLODPerComponent;
+				Arch << ModelResources->MinLODs;
+				Arch << ModelResources->MinQualityLevelLODs;
 				
 				ModeleResData += ModelResources->StreamedResourceData.GetAllocatedSize();
 				for (FCustomizableObjectStreamedResourceData& Data : ModelResources->StreamedResourceData)
diff --git a/Engine/Plugins/Mutable/Source/CustomizableObject/Private/MuCO/CustomizableObjectInstance.cpp b/Engine/Plugins/Mutable/Source/CustomizableObject/Private/MuCO/CustomizableObjectInstance.cpp
--- a/Engine/Plugins/Mutable/Source/CustomizableObject/Private/MuCO/CustomizableObjectInstance.cpp
+++ b/Engine/Plugins/Mutable/Source/CustomizableObject/Private/MuCO/CustomizableObjectInstance.cpp
@@ -3608,8 +3608,8 @@ void InitSkeletalMeshData(const TSharedRef<FUpdateContextPrivate>& Context, USke
 
 	// Set Min LOD (Override the Reference Skeletal Mesh LOD Settings)
 	const UModelResources& ModelResources = CustomizableObject.GetPrivate()->GetModelResourcesChecked();
-	SkeletalMesh->SetMinLod(FMath::Max(ModelResources.MinLODPerComponent.FindChecked(ComponentId).GetDefault(), static_cast<int32>(Context->FirstLODAvailable[ComponentName])));
-	SkeletalMesh->SetQualityLevelMinLod(ModelResources.MinQualityLevelLODPerComponent.FindChecked(ComponentId));
+	SkeletalMesh->SetMinLod(FMath::Max(ModelResources.MinLODs.FindChecked(ComponentId).GetDefault(), static_cast<int32>(Context->FirstLODAvailable[ComponentName])));
+	SkeletalMesh->SetQualityLevelMinLod(ModelResources.MinQualityLevelLODs.FindChecked(ComponentId));
 
 	// Set up unreal's default material, will be replaced when building materials
 	{
@@ -6293,7 +6293,7 @@ void UCustomizableObjectInstance::SetRequestedLODs(const TMap<FName, uint8>& InM
 		MutableUpdateCandidate.QualitySettingMinLODs.Add(ComponentName, MinLODIdx);
 
 		int32 MaxLODIdx = 0;
-		if (const uint8* Found = ModelResources->NumLODsAvailable.Find(ComponentId))
+		if (const uint8* Found = ModelResources->NumLODs.Find(ComponentId))
 		{
 			MaxLODIdx = *Found - 1;
 		}
@@ -6319,7 +6319,7 @@ void UCustomizableObjectInstance::SetRequestedLODs(const TMap<FName, uint8>& InM
 			PrivateData->SetCOInstanceFlags(InMinLOD > CurrentMinLOD ? PendingLODsDowngrade : ECONone);
 
 			uint8 FirstNonStreamedLODIndex = 0;
-			if (const uint8* Found = ModelResources->NumLODsToStream.Find(ComponentId))
+			if (const uint8* Found = ModelResources->MaxNumStreamedLODs.Find(ComponentId))
 			{
 				FirstNonStreamedLODIndex = *Found;
 			}
diff --git a/Engine/Plugins/Mutable/Source/CustomizableObject/Private/MuCO/CustomizableObjectSystem.cpp b/Engine/Plugins/Mutable/Source/CustomizableObject/Private/MuCO/CustomizableObjectSystem.cpp
--- a/Engine/Plugins/Mutable/Source/CustomizableObject/Private/MuCO/CustomizableObjectSystem.cpp
+++ b/Engine/Plugins/Mutable/Source/CustomizableObject/Private/MuCO/CustomizableObjectSystem.cpp
@@ -258,8 +258,8 @@ FUpdateContextPrivate::FUpdateContextPrivate(UCustomizableObjectInstance& InInst
 		{
 			const FName& ComponentName = ModelResources->ComponentNamesPerObjectComponent[ComponentId];
 			
-			FirstLODAvailable.Add(ComponentName, ModelResources->FirstLODAvailable[ComponentId]);
-			FirstResidentLOD.Add(ComponentName, ModelResources->NumLODsToStream[ComponentId]);
+			FirstLODAvailable.Add(ComponentName, ModelResources->FirstLODs[ComponentId]);
+			FirstResidentLOD.Add(ComponentName, ModelResources->MaxNumStreamedLODs[ComponentId]);
 		}
 	}
 
