[Mutable] Change ModelResources UPROPERTY names to avoid type warnings.

Missing file.

diff --git a/Engine/Plugins/Mutable/Source/CustomizableObjectEditor/Private/MuCOE/CustomizableObjectCompiler.cpp b/Engine/Plugins/Mutable/Source/CustomizableObjectEditor/Private/MuCOE/CustomizableObjectCompiler.cpp
--- a/Engine/Plugins/Mutable/Source/CustomizableObjectEditor/Private/MuCOE/CustomizableObjectCompiler.cpp
+++ b/Engine/Plugins/Mutable/Source/CustomizableObjectEditor/Private/MuCOE/CustomizableObjectCompiler.cpp
@@ -977,8 +977,8 @@ void FCustomizableObjectCompiler::CompileInternal(bool bAsync)
 			const FMutableLODSettings& ComponentLODSettings = ComponentInfo.LODSettings;
 
 			// Copy the LODSettings data found in the Component into the ModelResources
-			ModelResources->MinLODPerComponent.Add(ComponentInfo.ComponentId, ComponentLODSettings.MinLOD);
-			ModelResources->MinQualityLevelLODPerComponent.Add(ComponentInfo.ComponentId, ComponentLODSettings.MinQualityLevelLOD);
+			ModelResources->MinLODs.Add(ComponentInfo.ComponentId, ComponentLODSettings.MinLOD);
+			ModelResources->MinQualityLevelLODs.Add(ComponentInfo.ComponentId, ComponentLODSettings.MinQualityLevelLOD);
 		}
 		
 		ModelResources->ParameterUIDataMap = MoveTemp(GenerationContext.ParameterUIDataMap);
@@ -1012,21 +1012,21 @@ void FCustomizableObjectCompiler::CompileInternal(bool bAsync)
 			ModelResources->ReleaseVersion = VersionBridge->GetCurrentVersionAsString();
 		}
 
-		ModelResources->NumLODsAvailable = GenerationContext.NumLODs;
+		ModelResources->NumLODs = GenerationContext.NumLODs;
 
 		if (GenerationContext.bEnableLODStreaming)
 		{
-			ModelResources->NumLODsToStream = GenerationContext.NumMaxLODsToStream;
+			ModelResources->MaxNumStreamedLODs = GenerationContext.NumMaxLODsToStream;
 		}
 		else
 		{
-			for (TTuple<UE::Mutable::Private::FComponentId, uint8>& Pair : ModelResources->NumLODsToStream)
+			for (TTuple<UE::Mutable::Private::FComponentId, uint8>& Pair : ModelResources->MaxNumStreamedLODs)
 			{
 				Pair.Value = 0;
 			}
 		}
 
-		ModelResources->FirstLODAvailable = GenerationContext.FirstLODAvailable;
+		ModelResources->FirstLODs = GenerationContext.FirstLODAvailable;
 		
 		ModelResources->ParticipatingObjects = MoveTemp(ParticipatingObjects);
 		
