[Backout] - CL49656333
#rnx
Original CL Desc
-----------------------------------------------------------------
[Mutable] Change ModelResources UPROPERTY names to avoid type warnings.

Missing file.

diff --git a/Engine/Plugins/Mutable/Source/CustomizableObjectEditor/Private/MuCOE/CustomizableObjectCompiler.cpp b/Engine/Plugins/Mutable/Source/CustomizableObjectEditor/Private/MuCOE/CustomizableObjectCompiler.cpp
--- a/Engine/Plugins/Mutable/Source/CustomizableObjectEditor/Private/MuCOE/CustomizableObjectCompiler.cpp
+++ b/Engine/Plugins/Mutable/Source/CustomizableObjectEditor/Private/MuCOE/CustomizableObjectCompiler.cpp
@@ -1003,8 +1003,8 @@ void FCustomizableObjectCompiler::CompileInternal(bool bAsync)
 			const FMutableLODSettings& ComponentLODSettings = ComponentInfo.LODSettings;
 
 			// Copy the LODSettings data found in the Component into the ModelResources
-			ModelResources->MinLODs.Add(ComponentInfo.ComponentId, ComponentLODSettings.MinLOD);
-			ModelResources->MinQualityLevelLODs.Add(ComponentInfo.ComponentId, ComponentLODSettings.MinQualityLevelLOD);
+			ModelResources->MinLODPerComponent.Add(ComponentInfo.ComponentId, ComponentLODSettings.MinLOD);
+			ModelResources->MinQualityLevelLODPerComponent.Add(ComponentInfo.ComponentId, ComponentLODSettings.MinQualityLevelLOD);
 		}
 		
 		ModelResources->ParameterUIDataMap = MoveTemp(GenerationContext.ParameterUIDataMap);
@@ -1038,21 +1038,21 @@ void FCustomizableObjectCompiler::CompileInternal(bool bAsync)
 			ModelResources->ReleaseVersion = VersionBridge->GetCurrentVersionAsString();
 		}
 
-		ModelResources->NumLODs = GenerationContext.NumLODs;
+		ModelResources->NumLODsAvailable = GenerationContext.NumLODs;
 
 		if (GenerationContext.bEnableLODStreaming)
 		{
-			ModelResources->MaxNumStreamedLODs = GenerationContext.NumMaxLODsToStream;
+			ModelResources->NumLODsToStream = GenerationContext.NumMaxLODsToStream;
 		}
 		else
 		{
-			for (TTuple<UE::Mutable::Private::FComponentId, uint8>& Pair : ModelResources->MaxNumStreamedLODs)
+			for (TTuple<UE::Mutable::Private::FComponentId, uint8>& Pair : ModelResources->NumLODsToStream)
 			{
 				Pair.Value = 0;
 			}
 		}
 
-		ModelResources->FirstLODs = GenerationContext.FirstLODAvailable;
+		ModelResources->FirstLODAvailable = GenerationContext.FirstLODAvailable;
 		
 		ModelResources->ParticipatingObjects = MoveTemp(ParticipatingObjects);
 		
