[Backout] - CL49648897
#rnx
Original CL Desc
-----------------------------------------------------------------
[Mutable] Allow components having the same name.

#Synced-CL 49557694

diff --git a/Engine/Plugins/Experimental/MutablePopulation/Source/CustomizableObjectPopulationEditor/Private/MuCOPE/CustomizableObjectPopulationClassEditor.cpp b/Engine/Plugins/Experimental/MutablePopulation/Source/CustomizableObjectPopulationEditor/Private/MuCOPE/CustomizableObjectPopulationClassEditor.cpp
--- a/Engine/Plugins/Experimental/MutablePopulation/Source/CustomizableObjectPopulationEditor/Private/MuCOPE/CustomizableObjectPopulationClassEditor.cpp
+++ b/Engine/Plugins/Experimental/MutablePopulation/Source/CustomizableObjectPopulationEditor/Private/MuCOPE/CustomizableObjectPopulationClassEditor.cpp
@@ -22,7 +22,6 @@
 #include "MuCO/CustomizableObject.h"
 #include "MuCO/CustomizableObjectInstance.h"
 #include "MuCO/CustomizableSkeletalComponent.h"
-#include "MuCO/CustomizableObjectInstancePrivate.h"
 #include "MuCOP/CustomizableObjectPopulation.h"
 #include "MuCOP/CustomizableObjectPopulationClass.h"
 #include "MuCOP/CustomizableObjectPopulationGenerator.h"
@@ -881,25 +880,14 @@ void FCustomizableObjectPopulationClassEditor::OpenInstance()
 void FCustomizableObjectPopulationClassEditor::OpenSkeletalMesh()
 {
 	int32 SelectedInstance = Viewport->GetSelectedInstance();
-
-	if (!ViewportInstances[SelectedInstance])
-	{
-		return;
-	}
 	
-	TMap<FName, TObjectPtr<USkeletalMesh>>::TConstIterator It = ViewportInstances[SelectedInstance]->GetPrivate()->SkeletalMeshes.CreateConstIterator();
-	if (!It)
+	if (ViewportInstances[SelectedInstance])
 	{
-		return;
+		if (USkeletalMesh* SkeletalMesh = ViewportInstances[SelectedInstance]->GetSkeletalMesh())
+		{
+			GEditor->GetEditorSubsystem<UAssetEditorSubsystem>()->OpenEditorForAsset(SkeletalMesh);
+		}
 	}
-	
-	USkeletalMesh* SkeletalMesh = It->Value;
-	if (!SkeletalMesh)
-	{
-		return;
-	}	
-
-	GEditor->GetEditorSubsystem<UAssetEditorSubsystem>()->OpenEditorForAsset(SkeletalMesh);
 }
 
 
diff --git a/Engine/Plugins/Experimental/MutablePopulation/Source/CustomizableObjectPopulationEditor/Private/MuCOPE/CustomizableObjectPopulationEditor.cpp b/Engine/Plugins/Experimental/MutablePopulation/Source/CustomizableObjectPopulationEditor/Private/MuCOPE/CustomizableObjectPopulationEditor.cpp
--- a/Engine/Plugins/Experimental/MutablePopulation/Source/CustomizableObjectPopulationEditor/Private/MuCOPE/CustomizableObjectPopulationEditor.cpp
+++ b/Engine/Plugins/Experimental/MutablePopulation/Source/CustomizableObjectPopulationEditor/Private/MuCOPE/CustomizableObjectPopulationEditor.cpp
@@ -34,7 +34,6 @@
 #include "Widgets/Layout/SUniformGridPanel.h"
 #include "WorkspaceMenuStructure.h"
 #include "WorkspaceMenuStructureModule.h"
-#include "MuCO/CustomizableObjectInstancePrivate.h"
 
 class IToolkitHost;
 class SWidget;
@@ -525,24 +524,13 @@ void FCustomizableObjectPopulationEditor::OpenSkeletalMesh()
 {
 	int32 SelectedInstance = Viewport->GetSelectedInstance();
 
-	if (!ViewportInstances[SelectedInstance])
-	{
-		return;
-	}
-	
-	TMap<FName, TObjectPtr<USkeletalMesh>>::TConstIterator It = ViewportInstances[SelectedInstance]->GetPrivate()->SkeletalMeshes.CreateConstIterator();
-	if (!It)
+	if (ViewportInstances[SelectedInstance])
 	{
-		return;
+		if (USkeletalMesh* SkeletalMesh = ViewportInstances[SelectedInstance]->GetSkeletalMesh())
+		{
+			GEditor->GetEditorSubsystem<UAssetEditorSubsystem>()->OpenEditorForAsset(SkeletalMesh);
+		}
 	}
-	
-	USkeletalMesh* SkeletalMesh = It->Value;
-	if (!SkeletalMesh)
-	{
-		return;
-	}	
-
-	GEditor->GetEditorSubsystem<UAssetEditorSubsystem>()->OpenEditorForAsset(SkeletalMesh);
 }
 
 // To move elsewhere
diff --git a/Engine/Plugins/Mutable/Source/CustomizableObject/Internal/MuCO/CustomizableObjectInstancePrivate.h b/Engine/Plugins/Mutable/Source/CustomizableObject/Internal/MuCO/CustomizableObjectInstancePrivate.h
--- a/Engine/Plugins/Mutable/Source/CustomizableObject/Internal/MuCO/CustomizableObjectInstancePrivate.h
+++ b/Engine/Plugins/Mutable/Source/CustomizableObject/Internal/MuCO/CustomizableObjectInstancePrivate.h
@@ -111,12 +111,6 @@ struct FCustomizableInstanceComponentData
 {
 	GENERATED_USTRUCT_BODY();
 
-	FCustomizableInstanceComponentData();
-
-	FName ComponentName;
-	
-	UE::Mutable::Private::FComponentId ComponentId = INDEX_NONE;
-	
 	// AnimBP data gathered for a component from its constituent meshes
 	UPROPERTY(Transient, Category = CustomizableObjectInstance, editfixedsize, VisibleAnywhere)
 	TMap<FName, TSoftClassPtr<UAnimInstance>> AnimSlotToBP;
@@ -266,7 +260,10 @@ public:
 	/** Invalidates the previously generated data and retrieves information from the CObject after specific actions.
 	 *  It'll be called in the PostLoad, after Compiling the CO, and after changing the CO of the Instance. */
 	UE_API void InitCustomizableObjectData(const UCustomizableObject* InCustomizableObject);
-	
+
+	UE_API FCustomizableInstanceComponentData* GetComponentData(const FName& ComponentName);
+	void ClearInstanceComponentData(const FName& ComponentName);
+
 	ECOInstanceFlags GetCOInstanceFlags() const { return InstanceFlagsPrivate; }
 	void SetCOInstanceFlags(ECOInstanceFlags FlagsToSet) { InstanceFlagsPrivate = (ECOInstanceFlags)(InstanceFlagsPrivate | FlagsToSet); }
 	void ClearCOInstanceFlags(ECOInstanceFlags FlagsToClear) { InstanceFlagsPrivate = (ECOInstanceFlags)(InstanceFlagsPrivate & ~FlagsToClear); }
@@ -280,6 +277,9 @@ public:
 	// Returns true if success (?)
 	UE_API bool UpdateSkeletalMesh_PostBeginUpdate0(UCustomizableObjectInstance* Instance, const TSharedRef<FUpdateContextPrivate>& Context);
 	
+	// Copy data generated in the mutable thread over to the instance and initializes additional data required during the update
+	UE_API void PrepareForUpdate();
+
 	// The following method is basically copied from PostEditChangeProperty and/or SkeletalMesh.cpp to be able to replicate PostEditChangeProperty without the editor
 	UE_API void PostEditChangePropertyWithoutEditor();
 	
@@ -287,17 +287,19 @@ public:
 	  * It should not be called if the meshes are still in use or shared with other instances. */
 	UE_API void DiscardResources();
 	
-	/** Only use as an emergency fallback! Set the reference SkeletalMesh unconditionally to all actors using this instance. It will hitch. */
-	UE_API void ForceSetReferenceSkeletalMesh() const;
+	/** Set the reference SkeletalMesh, or an empty mesh, to all actors using this instance. */
+	UE_API void SetReferenceSkeletalMesh() const;
 
 	UE_API const TArray<FAnimInstanceOverridePhysicsAsset>* GetGeneratedPhysicsAssetsForAnimInstance(TSubclassOf<UAnimInstance> AnimInstance) const;
 
+	/** */
 #if WITH_EDITORONLY_DATA
 	UE_API void RegenerateImportedModels(const TSharedRef<FUpdateContextPrivate>& OperationData);
 #endif
 	
-	UE_API FCustomizableInstanceComponentData* GetComponentData(const FName& ComponentName);
-	UE_API FCustomizableInstanceComponentData& GetComponentData(UE::Mutable::Private::FComponentId ComponentId);
+	// 
+	UE_API FCustomizableInstanceComponentData* GetComponentData(FCustomizableObjectComponentIndex ObjectComponentIndex);
+	UE_API const FCustomizableInstanceComponentData* GetComponentData(FCustomizableObjectComponentIndex ObjectComponentIndex) const;
 	
 	UE_API void InvalidateGeneratedData();
 
@@ -345,6 +347,7 @@ public:
 	float MinSquareDistFromComponentToPlayer = FLT_MAX;
 	float LastMinSquareDistFromComponentToPlayer = FLT_MAX; // The same as the previous dist for last frame
 	
+	// To be indexed with object component index
 	UPROPERTY(Transient)
 	TArray<FCustomizableInstanceComponentData> ComponentsData;
 	
@@ -417,23 +420,23 @@ public:
 };
 
 
-UE_API UPhysicsAsset* GetOrBuildMainPhysicsAsset(const TSharedRef<FUpdateContextPrivate>& Context, TObjectPtr<UPhysicsAsset> TamplateAsset, const UE::Mutable::Private::FPhysicsBody* PhysicsBody, bool bDisableCollisionBetweenAssets, int32 ComponentIndex);
+UE_API UPhysicsAsset* GetOrBuildMainPhysicsAsset(const TSharedRef<FUpdateContextPrivate>& Context, TObjectPtr<UPhysicsAsset> TamplateAsset, const UE::Mutable::Private::FPhysicsBody* PhysicsBody, bool bDisableCollisionBetweenAssets, FCustomizableObjectInstanceComponentIndex InstanceComponentIndex);
 
-UE_API USkeleton* MergeSkeletons(UCustomizableObject& CustomizableObject, FCustomizableInstanceComponentData& ComponentData, bool& bOutCreatedNewSkeleton);
+UE_API USkeleton* MergeSkeletons(UCustomizableObject& CustomizableObject, FCustomizableInstanceComponentData& ObjectComponentIndex, bool& bOutCreatedNewSkeleton);
 
-UE_API void InitSkeletalMeshData(const TSharedRef<FUpdateContextPrivate>& Context, USkeletalMesh* SkeletalMesh, const FMutableRefSkeletalMeshData& RefSkeletalMeshData, const UCustomizableObject& CustomizableObject, UE::Mutable::Private::FComponentId ComponentId);
+UE_API void InitSkeletalMeshData(const TSharedRef<FUpdateContextPrivate>& Context, USkeletalMesh* SkeletalMesh, const FMutableRefSkeletalMeshData& RefSkeletalMeshData, const UCustomizableObject& CustomizableObject, FCustomizableObjectComponentIndex ObjectComponentIndex);
 
-UE_API void BuildOrCopyClothingData(const TSharedRef<FUpdateContextPrivate>& Context, USkeletalMesh* SkeletalMesh, const UModelResources& CustomizableObjectInstance, int32 ComponentIndex, const TArray<TObjectPtr<UPhysicsAsset>>& ClothingPhysicsAssets);
+UE_API void BuildOrCopyClothingData(const TSharedRef<FUpdateContextPrivate>& Context, USkeletalMesh* SkeletalMesh, const UModelResources& CustomizableObjectInstance, FCustomizableObjectInstanceComponentIndex InstanceComponentIndex, const TArray<TObjectPtr<UPhysicsAsset>>& ClothingPhysicsAssets);
 
-UE_API bool BuildSkeletonData(const TSharedRef<FUpdateContextPrivate>& Context, USkeletalMesh& SkeletalMesh, int32 ComponentIndex);
+UE_API bool BuildSkeletonData(const TSharedRef<FUpdateContextPrivate>& Context, USkeletalMesh& SkeletalMesh, FCustomizableObjectInstanceComponentIndex InstanceComponentIndex);
 
 UE_API void BuildMeshSockets(USkeletalMesh* SkeletalMesh, const UModelResources& ModelResources, const FMutableRefSkeletalMeshData& RefSkeletalMeshData, UE::Mutable::Private::TManagedPtr<const UE::Mutable::Private::FMesh> MutableMesh);
 
-UE_API bool BuildOrCopyRenderData(const TSharedRef<FUpdateContextPrivate>& Context, USkeletalMesh* SkeletalMesh, UCustomizableObjectInstance* CustomizableObjectInstance, int32 ComponentIndex);
+UE_API bool BuildOrCopyRenderData(const TSharedRef<FUpdateContextPrivate>& Context, USkeletalMesh* SkeletalMesh, UCustomizableObjectInstance* CustomizableObjectInstance, FCustomizableObjectInstanceComponentIndex InstanceComponentIndex);
 
-UE_API void BuildOrCopyElementData(const TSharedRef<FUpdateContextPrivate>& Context, USkeletalMesh* SkeletalMesh, int32 ComponentIndex);
+UE_API void BuildOrCopyElementData(const TSharedRef<FUpdateContextPrivate>& Context, USkeletalMesh* SkeletalMesh, FCustomizableObjectInstanceComponentIndex InstanceComponentIndex);
 
-UE_API void BuildOrCopyMorphTargetsData(const TSharedRef<FUpdateContextPrivate>& Context, USkeletalMesh* SkeletalMesh, int32 ComponentIndex);
+UE_API void BuildOrCopyMorphTargetsData(const TSharedRef<FUpdateContextPrivate>& Context, USkeletalMesh* SkeletalMesh, FCustomizableObjectInstanceComponentIndex InstanceComponentIndex);
 	
 UE_API bool DoComponentsNeedUpdate(const TSharedRef<FUpdateContextPrivate>& Context, bool& bOutEmptyMesh);
 
diff --git a/Engine/Plugins/Mutable/Source/CustomizableObject/Internal/MuCO/CustomizableObjectPrivate.h b/Engine/Plugins/Mutable/Source/CustomizableObject/Internal/MuCO/CustomizableObjectPrivate.h
--- a/Engine/Plugins/Mutable/Source/CustomizableObject/Internal/MuCO/CustomizableObjectPrivate.h
+++ b/Engine/Plugins/Mutable/Source/CustomizableObject/Internal/MuCO/CustomizableObjectPrivate.h
@@ -1221,11 +1221,14 @@ class UModelResources : public UObject
 	GENERATED_BODY()
 
 public:
-	/** Key is FComponentId.
+	/** 
 	 * All the SkeletalMeshes generated for this CustomizableObject instances will use the Reference Skeletal Mesh
 	 * properties for everything that Mutable doesn't create or modify. This struct stores the information used from
 	 * the Reference Skeletal Meshes to avoid having them loaded at all times. This includes data like LOD distances,
-	 * LOD render data settings, Mesh sockets, Bounding volumes, etc. */
+	 * LOD render data settings, Mesh sockets, Bounding volumes, etc.
+	 * 
+	 * Index with CustomizableObject Component index
+	 */
 	UPROPERTY()
 	TArray<FMutableRefSkeletalMeshData> ReferenceSkeletalMeshesData;	
 
@@ -1358,29 +1361,29 @@ public:
 	UPROPERTY()
 	TMap<TObjectPtr<const UScriptStruct>, uint32> ExternalOperationVersions;
 	
-	/** Key is FComponentId. Max number of LODs in the compiled Model. */
+	/** Max number of LODs in the compiled Model. */
 	UPROPERTY()
-	TMap<int32, uint8> NumLODsAvailable;
+	TMap<FName, uint8> NumLODsAvailable;
 
-	/** Key is FComponentId. Max number of LODs to stream. Mutable will always generate at least one LOD. */
+	/** Max number of LODs to stream. Mutable will always generate at least one LOD. */
 	UPROPERTY()
-	TMap<int32, uint8> NumLODsToStream;
+	TMap<FName, uint8> NumLODsToStream;
 
-	/** Key is FComponentId. First LOD available, some platforms may remove lower LODs when cooking, this MinLOD represents the first LOD we can generate */
+	/** First LOD available, some platforms may remove lower LODs when cooking, this MinLOD represents the first LOD we can generate */
 	UPROPERTY()
-	TMap<int32, uint8> FirstLODAvailable;
+	TMap<FName, uint8> FirstLODAvailable;
 
-	/** Key is FComponentId. Name of all possible components. */
+	/** Name of all possible components. Index is the ObjectComponentIndex. */
 	UPROPERTY()
 	TArray<FName> ComponentNamesPerObjectComponent;
 
-	/** Key is FComponentId. Minimum LOD to render per Platform. */
+	/** Minimum LOD to render per Platform. */
 	UPROPERTY()
-	TMap<int32, FPerPlatformInt> MinLODPerComponent;
+	TMap<FName, FPerPlatformInt> MinLODPerComponent;
 
-	/** Key is FComponentId. Minimum LOD to render per Quality level. */
+	/** Minimum LOD to render per Quality level. */
 	UPROPERTY()
-	TMap<int32, FPerQualityLevelInt> MinQualityLevelLODPerComponent;
+	TMap<FName, FPerQualityLevelInt> MinQualityLevelLODPerComponent;
 
 	UPROPERTY()
 	TMap<FString, TObjectPtr<UTexture>> TextureParameterDefaultValues;
@@ -1468,6 +1471,43 @@ struct FMutableMeshComponentData
 };
 
 
+
+/** Strongly typed index for the index of a component in a UCustomizableObject. */
+USTRUCT()
+struct FCustomizableObjectComponentIndex
+{
+	GENERATED_BODY()
+
+private:
+
+	int32 Index = 0;
+
+public:
+
+	explicit inline FCustomizableObjectComponentIndex(int32 InIndex=0)
+		: Index(InIndex)
+	{
+	}
+
+	inline void Invalidate()
+	{
+		Index = INDEX_NONE;
+	}
+
+	inline bool IsValid() const
+	{
+		return Index != INDEX_NONE;
+	}
+
+	inline int32 GetValue() const
+	{
+		return Index;
+	}
+
+	inline bool operator==(const FCustomizableObjectComponentIndex&) const = default;
+};
+
+
 #if WITH_EDITORONLY_DATA 
 // This is a manual version number for the binary blobs in this asset.
 // Increasing it invalidates all the previously compiled models.
@@ -1727,8 +1767,6 @@ enum class ECustomizableObjectVersions : int32
 	DeduplicateFMaterialsSharingPassthroughObject,
 
 	RemovedReuseTextures,
-	
-	ComponentId,
 
 	PassthroughTexturesForBreakNodes,
 
@@ -1873,7 +1911,7 @@ public:
 	CUSTOMIZABLEOBJECT_API void SetIsChildObject(bool bIsChildObject);
 	
 	/** Return the MinLOD index to generate based on the active LODSettings (PerPlatformMinLOD or PerQualityLevelMinLOD) */
-	CUSTOMIZABLEOBJECT_API uint8 GetMinLODIndex(const UE::Mutable::Private::FComponentId ComponentId) const;
+	CUSTOMIZABLEOBJECT_API uint8 GetMinLODIndex(const FName& ComponentName) const;
 
 #if WITH_EDITOR
 	/** See ICustomizableObjectEditorModule::IsCompilationOutOfDate. */
@@ -1898,7 +1936,7 @@ public:
 
 	CUSTOMIZABLEOBJECT_API void BackwardsCompatibleFixup(int32 CustomizableObjectCustomVersion);
 
-	CUSTOMIZABLEOBJECT_API FName GetComponentName(UE::Mutable::Private::FComponentId ComponentId) const;
+	CUSTOMIZABLEOBJECT_API FName GetComponentName(FCustomizableObjectComponentIndex ObjectComponentIndex) const;
 
 #if WITH_EDITORONLY_DATA
 	CUSTOMIZABLEOBJECT_API EMutableCompileMeshType GetMeshCompileType() const;
diff --git a/Engine/Plugins/Mutable/Source/CustomizableObject/Internal/MuCO/CustomizableObjectSystemPrivate.h b/Engine/Plugins/Mutable/Source/CustomizableObject/Internal/MuCO/CustomizableObjectSystemPrivate.h
--- a/Engine/Plugins/Mutable/Source/CustomizableObject/Internal/MuCO/CustomizableObjectSystemPrivate.h
+++ b/Engine/Plugins/Mutable/Source/CustomizableObject/Internal/MuCO/CustomizableObjectSystemPrivate.h
@@ -118,6 +118,37 @@ void LoadClothing(FMutableStreamRequest& MutableDataStreamer, const UE::Mutable:
 void FinishUpdateGlobal(const TSharedRef<FUpdateContextPrivate>& Context);
 
 
+/** Strongly typed index for the index of a component in a UCustomizableObjectInstance. */
+class FCustomizableObjectInstanceComponentIndex
+{
+private:
+
+	int32 Index = 0;
+
+public:
+
+	explicit FCustomizableObjectInstanceComponentIndex(int32 InIndex = 0)
+		: Index(InIndex)
+	{
+	}
+
+	inline void Invalidate()
+	{
+		Index = INDEX_NONE;
+	}
+
+	inline bool IsValid() const
+	{
+		return Index != INDEX_NONE;
+	}
+
+	inline int32 GetValue() const
+	{
+		return Index;
+	}
+};
+
+
 class FMutableUpdateCandidate
 {
 public:
@@ -371,7 +402,7 @@ struct FInstanceUpdateData
 	
 	struct FComponent
 	{
-		UE::Mutable::Private::FComponentId Id = INDEX_NONE;
+		FCustomizableObjectComponentIndex Id = FCustomizableObjectComponentIndex(0);
 		
 		/** Range in the LODs array */
 		uint16 FirstLOD = 0;
@@ -410,7 +441,7 @@ struct FInstanceUpdateData
 	TMap<uint32, FMorphTargetMeshData> RealTimeMorphTargetMeshData;
 	
 	/** Key is the Component Name. Value is the final Morph Target data to copy into the Skeletal Mesh. */
-	TMap<UE::Mutable::Private::FComponentId, FSkeletalMeshMorphTargets> RealTimeMorphTargets;
+	TMap<FName, FSkeletalMeshMorphTargets> RealTimeMorphTargets;
 
 	/** Key is the Block Id. Value is the LoadAdditionalAssetsAndData read destination. */
 	TMap<uint32, FClothingMeshData> ClothingMeshData;
@@ -487,12 +518,17 @@ public:
 	const FDescriptorHash& GetCapturedDescriptorHash() const;
 
 	const FCustomizableObjectInstanceDescriptor&& MoveCommittedDescriptor();
+
+	/** Return the object component index associated with a component in this instance. */
+	FCustomizableObjectComponentIndex GetObjectComponentIndex(FCustomizableObjectInstanceComponentIndex InstanceComponentIndex);
 	
+	const FInstanceUpdateData::FComponent* GetComponentUpdateData(FCustomizableObjectInstanceComponentIndex InstanceComponentIndex);
+
 	void InitComponentDescriptors(int32 Size);
 
 	const TArray<FComponentDescriptor>& GetComponentDescriptors() const;
 	
-	FComponentDescriptor* GetComponentDescriptor(UE::Mutable::Private::FComponentId ComponentId);
+	FComponentDescriptor* GetComponentDescriptor(FCustomizableObjectComponentIndex Index);
 	
 	UE::Mutable::Private::FImageOperator::FImagePixelFormatFunc PixelFormatOverride;
 
@@ -522,7 +558,7 @@ public:
 	  * @return true if the mesh is new or new to this instance (e.g. mesh cached by another instance). */
 	TArray<bool> MeshChangedPerInstanceComponent;
 
-	/** List of component names. Key is the FComponentId. */
+	/** List of component names. Index is the ObjectComponentIndex. */
 	TArray<FName> ComponentNames;
 
 	/** Index of the resource in the StreamedResourceData array of the Model Resources. */
@@ -610,6 +646,14 @@ public:
 	/** This option comes from the operation request. It is used to reduce the number of mipmaps that mutable must generate for images.  */
 	int32 MipsToSkip = 0;
 	
+	// Number of possible components in the entire CO
+	// TODO: Redundant because it is in the CO
+	uint8 NumObjectComponents = 0;
+
+	// Number of components in the instance being generated
+	// TODO: Redundant to store it here because it is implicit in the size of some arrays or in the MutableInstance while it is valid.
+	uint8 NumInstanceComponents = 0;
+	
 	EQueuePriorityType PriorityType = EQueuePriorityType::Low;
 
 	EUpdateResult UpdateResult = EUpdateResult::Success;
diff --git a/Engine/Plugins/Mutable/Source/CustomizableObject/Internal/MuCO/UnrealToMutableTextureConversionUtils.h b/Engine/Plugins/Mutable/Source/CustomizableObject/Internal/MuCO/UnrealToMutableTextureConversionUtils.h
--- a/Engine/Plugins/Mutable/Source/CustomizableObject/Internal/MuCO/UnrealToMutableTextureConversionUtils.h
+++ b/Engine/Plugins/Mutable/Source/CustomizableObject/Internal/MuCO/UnrealToMutableTextureConversionUtils.h
@@ -4,7 +4,6 @@
 
 #include "HAL/Platform.h"
 #include "Engine/Texture.h"
-#include "MuR/ComponentId.h"
 #include "MuR/ManagedPointer.h"
 
 #define UE_API CUSTOMIZABLEOBJECT_API
@@ -93,7 +92,7 @@ struct FMutableSourceMeshData
 	TSoftClassPtr<UAnimInstance> AnimInstance;
 	TSoftObjectPtr<UStreamableRenderAsset> TableReferenceSkeletalMesh;
 
-	UE::Mutable::Private::FComponentId ComponentId = INDEX_NONE;
+	FName Component = NAME_None;
 	
 	/** Required for SurfaceMetadataID */
 
@@ -136,7 +135,7 @@ inline uint32 GetTypeHash(const FMutableSourceMeshData& Key)
 	GuidHash = HashCombineFast(GuidHash, GetTypeHash(Key.SocketPriority));
 	GuidHash = HashCombineFast(GuidHash, GetTypeHash(Key.BonePosePriority));
 	GuidHash = HashCombineFast(GuidHash, GetTypeHash(Key.bOnlyConnectedLOD));
-	GuidHash = HashCombineFast(GuidHash, GetTypeHash(Key.ComponentId));
+	GuidHash = HashCombineFast(GuidHash, GetTypeHash(Key.Component));
 	GuidHash = HashCombineFast(GuidHash, GetTypeHash(Key.Flags));
 	return GuidHash;
 }
diff --git a/Engine/Plugins/Mutable/Source/CustomizableObject/Private/MuCO/CustomizableObject.cpp b/Engine/Plugins/Mutable/Source/CustomizableObject/Private/MuCO/CustomizableObject.cpp
--- a/Engine/Plugins/Mutable/Source/CustomizableObject/Private/MuCO/CustomizableObject.cpp
+++ b/Engine/Plugins/Mutable/Source/CustomizableObject/Private/MuCO/CustomizableObject.cpp
@@ -1005,12 +1005,11 @@ USkeletalMesh* UCustomizableObject::GetComponentMeshReferenceSkeletalMesh(const
 	// currently the graph search will be always performed if the editor only data is accessible.
 	if (const UModelResources* ModelResources = Private->GetModelResources())
 	{
-		// TODO Due to the adding support of multiple components with the same name this is no longer correct. A single name may point to multiple FComponentId.
-		UE::Mutable::Private::FComponentId ComponentId = ModelResources->ComponentNamesPerObjectComponent.IndexOfByKey(ComponentName);
-		if (ModelResources->ReferenceSkeletalMeshesData.IsValidIndex(ComponentId))
+		int32 ObjectComponentIndex = ModelResources->ComponentNamesPerObjectComponent.IndexOfByKey(ComponentName);
+		if (ModelResources->ReferenceSkeletalMeshesData.IsValidIndex(ObjectComponentIndex))
 		{
 			// Can be nullptr if RefSkeletalMeshes are not loaded yet.
-			return ModelResources->ReferenceSkeletalMeshesData[ComponentId].SkeletalMesh;
+			return ModelResources->ReferenceSkeletalMeshesData[ObjectComponentIndex].SkeletalMesh;
 		}
 	}
 	
@@ -1147,14 +1146,14 @@ int32 UCustomizableObject::GetComponentCount() const
 }
 
 
-FName UCustomizableObjectPrivate::GetComponentName(UE::Mutable::Private::FComponentId ComponentId) const
+FName UCustomizableObjectPrivate::GetComponentName(FCustomizableObjectComponentIndex ObjectComponentIndex) const
 {
 	if (const UModelResources* LocalModelResources = GetModelResources())
 	{
 		const TArray<FName>& ComponentNames = LocalModelResources->ComponentNamesPerObjectComponent;
-		if (ComponentNames.IsValidIndex(ComponentId))
+		if (ComponentNames.IsValidIndex(ObjectComponentIndex.GetValue()))
 		{
-			return ComponentNames[ComponentId];
+			return ComponentNames[ObjectComponentIndex.GetValue()];
 		}
 	}
 
@@ -1331,7 +1330,7 @@ int32 UCustomizableObjectPrivate::GetStateParameterIndex(int32 StateIndex, int32
 
 FName UCustomizableObject::GetComponentName(int32 ComponentIndex) const
 {
-	return GetPrivate()->GetComponentName(ComponentIndex);
+	return GetPrivate()->GetComponentName(FCustomizableObjectComponentIndex(ComponentIndex));
 }
 
 
@@ -1963,7 +1962,7 @@ void UCustomizableObjectPrivate::ApplyStateForcedValuesToParameters(FCustomizabl
 }
 
 
-uint8 UCustomizableObjectPrivate::GetMinLODIndex(const UE::Mutable::Private::FComponentId ComponentId) const
+uint8 UCustomizableObjectPrivate::GetMinLODIndex(const FName& ComponentName) const
 {
 	uint8 MinLODIdx = 0;
 
@@ -1979,7 +1978,7 @@ uint8 UCustomizableObjectPrivate::GetMinLODIndex(const UE::Mutable::Private::FCo
 				// Value extracted from the Scalability.ini of the current platform (if found).
 				const int32 QualityLevel = UCustomizableObjectSystem::GetInstance()->GetSkeletalMeshMinLODQualityLevel();
 
-				const FPerQualityLevelInt* QualityLevelInt = LocalModelResources->MinQualityLevelLODPerComponent.Find(ComponentId);
+				const FPerQualityLevelInt* QualityLevelInt = LocalModelResources->MinQualityLevelLODPerComponent.Find(ComponentName);
 				if (ensure(QualityLevelInt))
 				{
 					MinLODIdx = QualityLevelInt->GetValue(QualityLevel);
@@ -1988,7 +1987,7 @@ uint8 UCustomizableObjectPrivate::GetMinLODIndex(const UE::Mutable::Private::FCo
 		}
 		else
 		{
-			const FPerPlatformInt* PlatformInt = LocalModelResources->MinLODPerComponent.Find(ComponentId);
+			const FPerPlatformInt* PlatformInt = LocalModelResources->MinLODPerComponent.Find(ComponentName);
 			if (ensure(PlatformInt))
 			{
 				MinLODIdx = PlatformInt->GetValue();
@@ -2001,7 +2000,7 @@ uint8 UCustomizableObjectPrivate::GetMinLODIndex(const UE::Mutable::Private::FCo
 
 	if (LocalModelResources)
 	{
-		if (const uint8* Result = LocalModelResources->FirstLODAvailable.Find(ComponentId))
+		if (const uint8* Result = LocalModelResources->FirstLODAvailable.Find(ComponentName))
 		{
 			FirstLODAvailable = *Result;
 		}
diff --git a/Engine/Plugins/Mutable/Source/CustomizableObject/Private/MuCO/CustomizableObjectInstance.cpp b/Engine/Plugins/Mutable/Source/CustomizableObject/Private/MuCO/CustomizableObjectInstance.cpp
--- a/Engine/Plugins/Mutable/Source/CustomizableObject/Private/MuCO/CustomizableObjectInstance.cpp
+++ b/Engine/Plugins/Mutable/Source/CustomizableObject/Private/MuCO/CustomizableObjectInstance.cpp
@@ -119,7 +119,10 @@ void UCustomizableInstancePrivate::InvalidateGeneratedData()
 	CommittedDescriptor = {};
 	CommittedDescriptorHash = {};
 
-	ComponentsData.Empty();
+	// Init Component Data
+	FCustomizableInstanceComponentData TemplateComponentData;
+	TemplateComponentData.LastComponentDescriptor.MeshDescriptors.Init({}, MAX_MESH_LOD_COUNT);
+	ComponentsData.Init(TemplateComponentData, ComponentsData.Num());
 }
 
 
@@ -132,28 +135,77 @@ void UCustomizableInstancePrivate::InitCustomizableObjectData(const UCustomizabl
 		return;
 	}
 	
+	// Init Component Data
+	FCustomizableInstanceComponentData TemplateComponentData;
+	TemplateComponentData.LastComponentDescriptor.MeshDescriptors.Init({}, MAX_MESH_LOD_COUNT);
+	ComponentsData.Init(TemplateComponentData, InCustomizableObject->GetComponentCount());
+
 	ExtensionInstanceData.Empty();
 }
 
 
 FCustomizableInstanceComponentData* UCustomizableInstancePrivate::GetComponentData(const FName& ComponentName)
 {
-	return ComponentsData.FindByPredicate([&](const FCustomizableInstanceComponentData& ComponentData)
+	UCustomizableObject* Object = GetPublic()->GetCustomizableObject();
+	if (!Object || !Object->IsCompiled())
+	{
+		return nullptr;
+	}
+
+	const int32 ObjectComponentIndex = Object->GetPrivate()->GetModelResourcesChecked().ComponentNamesPerObjectComponent.IndexOfByKey(ComponentName);
+	if (ObjectComponentIndex == INDEX_NONE)
 	{
-		return ComponentData.ComponentName == ComponentName;
-	});
+		return nullptr;
+	}
+
+	if (!ComponentsData.IsValidIndex(ObjectComponentIndex))
+	{
+		return nullptr;	
+	}
+
+	return &ComponentsData[ObjectComponentIndex];
+}
+
+
+FCustomizableInstanceComponentData* UCustomizableInstancePrivate::GetComponentData(FCustomizableObjectComponentIndex ObjectComponentIndex)
+{
+	return ComponentsData.IsValidIndex(ObjectComponentIndex.GetValue()) ? &ComponentsData[ObjectComponentIndex.GetValue()] : nullptr;
 }
 
 
-FCustomizableInstanceComponentData& UCustomizableInstancePrivate::GetComponentData(UE::Mutable::Private::FComponentId ComponentId)
+const FCustomizableInstanceComponentData* UCustomizableInstancePrivate::GetComponentData(FCustomizableObjectComponentIndex ObjectComponentIndex) const
 {
-	FCustomizableInstanceComponentData* Result = ComponentsData.FindByPredicate([&](const FCustomizableInstanceComponentData& ComponentData)
+	return ComponentsData.IsValidIndex(ObjectComponentIndex.GetValue()) ? &ComponentsData[ObjectComponentIndex.GetValue()] : nullptr;
+}
+
+
+void UCustomizableInstancePrivate::ClearInstanceComponentData(const FName& ComponentName)
+{
+	FCustomizableInstanceComponentData* ComponentData = GetComponentData(ComponentName);
+
+	if (ComponentData)
 	{
-		return ComponentData.ComponentId == ComponentId;
-	});
-	
-	check(Result);
-	return *Result;
+		ComponentData->LastComponentDescriptor.MeshDescriptors.Reset(MAX_MESH_LOD_COUNT);
+		ComponentData->LastComponentDescriptor.MeshDescriptors.AddDefaulted(MAX_MESH_LOD_COUNT);
+
+		ComponentData->AnimSlotToBP.Empty();
+		ComponentData->AssetUserDataArray.Empty();
+		ComponentData->ClothingPhysicsAssetsToStream.Empty();
+
+#if WITH_EDITORONLY_DATA
+		ComponentData->MeshPartPaths.Empty();
+#endif
+		ComponentData->OverlayMaterial = nullptr;
+		ComponentData->OverrideMaterials.Empty();
+		ComponentData->PhysicsAssets.AdditionalPhysicsAssets.Empty();
+		ComponentData->PhysicsAssets.AdditionalPhysicsAssetsToLoad.Empty();
+		ComponentData->PhysicsAssets.PhysicsAssetsToMerge.Empty();
+		ComponentData->PhysicsAssets.PhysicsAssetToLoad.Empty();
+		ComponentData->Skeletons.Skeleton = nullptr;
+		ComponentData->Skeletons.SkeletonIds.Empty();
+		ComponentData->Skeletons.SkeletonsToMerge.Empty();
+		ComponentData->StreamedResourceIndex.Empty();
+	}
 }
 
 
@@ -182,6 +234,29 @@ void UCustomizableInstancePrivate::SetDescriptor(const FCustomizableObjectInstan
 }
 
 
+void UCustomizableInstancePrivate::PrepareForUpdate()
+{
+	// Clear the ComponentData from previous updates
+	for (FCustomizableInstanceComponentData& ComponentData : ComponentsData)
+	{
+		ComponentData.AnimSlotToBP.Empty();
+		ComponentData.AssetUserDataArray.Empty();
+		ComponentData.Skeletons.Skeleton = nullptr;
+		ComponentData.Skeletons.SkeletonIds.Empty();
+		ComponentData.Skeletons.SkeletonsToMerge.Empty();
+		ComponentData.PhysicsAssets.PhysicsAssetToLoad.Empty();
+		ComponentData.PhysicsAssets.PhysicsAssetsToMerge.Empty();
+		ComponentData.ClothingPhysicsAssetsToStream.Empty();
+		ComponentData.StreamedResourceIndex.Empty();
+		ComponentData.OverlayMaterial = nullptr;
+
+#if WITH_EDITORONLY_DATA
+		ComponentData.MeshPartPaths.Empty();
+#endif
+	}
+}
+
+
 uint32 GetTypeHash(const FMaterialReuseCacheKey& Key)
 {
 	uint32 Hash = GetTypeHash(Key.MaterialTemplate);
@@ -1290,7 +1365,7 @@ namespace
 		const TSharedRef<FUpdateContextPrivate>& Context,
 		TObjectPtr<UPhysicsAsset> TemplateAsset,
 		const UE::Mutable::Private::FPhysicsBody* MutablePhysics,
-		int32 ComponentIndex)
+		FCustomizableObjectInstanceComponentIndex InstanceComponentIndex)
 	{
 		check(TemplateAsset);
 		TObjectPtr<UPhysicsAsset> Result = NewObject<UPhysicsAsset>();
@@ -1305,7 +1380,7 @@ namespace
 
 		Result->bNotForDedicatedServer = TemplateAsset->bNotForDedicatedServer;
 
-		const TMap<UE::Mutable::Private::FBoneName, TPair<FName, uint16>>& BoneInfoMap = Context->InstanceUpdateData.SkeletonsPerInstanceComponent[ComponentIndex].BoneInfoMap;
+		const TMap<UE::Mutable::Private::FBoneName, TPair<FName, uint16>>& BoneInfoMap = Context->InstanceUpdateData.SkeletonsPerInstanceComponent[InstanceComponentIndex.GetValue()].BoneInfoMap;
 		TMap<FName, int32> BonesInUse;
 
 		const int32 MutablePhysicsBodyCount = MutablePhysics->GetBodyCount();
@@ -1435,7 +1510,7 @@ UPhysicsAsset* GetOrBuildMainPhysicsAsset(
 	TObjectPtr<UPhysicsAsset> TemplateAsset,
 	const UE::Mutable::Private::FPhysicsBody* MutablePhysics,
 	bool bDisableCollisionsBetweenDifferentAssets,
-	int32 ComponentIndex)
+	FCustomizableObjectInstanceComponentIndex InstanceComponentIndex)
 {
 	MUTABLE_CPUPROFILER_SCOPE(MergePhysicsAssets);
 
@@ -1446,20 +1521,22 @@ UPhysicsAsset* GetOrBuildMainPhysicsAsset(
 	
 	UPhysicsAsset* Result = nullptr;
 
-	const TSharedPtr<FInstanceUpdateData::FComponent>& Component = Context->InstanceUpdateData.Components[ComponentIndex];
+	const FInstanceUpdateData::FComponent* Component = Context->GetComponentUpdateData(InstanceComponentIndex);
 	if (!Component)
 	{
 		return nullptr;
 	}
 
-	FCustomizableInstanceComponentData& ComponentData = InstancePrivate->GetComponentData(Component->Id);
+	FCustomizableObjectComponentIndex ObjectComponentIndex = Component->Id;
+	FCustomizableInstanceComponentData* ComponentData = InstancePrivate->GetComponentData(ObjectComponentIndex);
+	check(ComponentData);
 	
 	TArray<TObjectPtr<UPhysicsAsset>> ValidAssets;
 
-	const int32 NumPhysicsAssets = ComponentData.PhysicsAssets.PhysicsAssetsToMerge.Num();
+	const int32 NumPhysicsAssets = ComponentData->PhysicsAssets.PhysicsAssetsToMerge.Num();
 	for (int32 PhysicsAssetIndex = 0; PhysicsAssetIndex < NumPhysicsAssets; ++PhysicsAssetIndex)
 	{
-		const TObjectPtr<UPhysicsAsset>& PhysicsAsset = ComponentData.PhysicsAssets.PhysicsAssetsToMerge[PhysicsAssetIndex];
+		const TObjectPtr<UPhysicsAsset>& PhysicsAsset = ComponentData->PhysicsAssets.PhysicsAssetsToMerge[PhysicsAssetIndex];
 
 		if (PhysicsAsset)
 		{
@@ -1493,7 +1570,7 @@ UPhysicsAsset* GetOrBuildMainPhysicsAsset(
 
 	Result->bNotForDedicatedServer = TemplateAsset->bNotForDedicatedServer;
 
-	const TMap<UE::Mutable::Private::FBoneName, TPair<FName, uint16>>& BoneInfoMap = Context->InstanceUpdateData.SkeletonsPerInstanceComponent[ComponentIndex].BoneInfoMap;
+	const TMap<UE::Mutable::Private::FBoneName, TPair<FName, uint16>>& BoneInfoMap = Context->InstanceUpdateData.SkeletonsPerInstanceComponent[InstanceComponentIndex.GetValue()].BoneInfoMap;
 	TMap<FName, int32> BonesInUse;
 
 	const int32 MutablePhysicsBodyCount = MutablePhysics->GetBodyCount();
@@ -1722,7 +1799,7 @@ UPhysicsAsset* GetOrBuildMainPhysicsAsset(
 	Result->UpdateBodySetupIndexMap();
 	Result->UpdateBoundsBodiesArray();
 
-	ComponentData.PhysicsAssets.PhysicsAssetsToMerge.Empty();
+	ComponentData->PhysicsAssets.PhysicsAssetsToMerge.Empty();
 
 	return Result;
 }
@@ -1767,17 +1844,17 @@ bool DoComponentsNeedUpdate(const TSharedRef<FUpdateContextPrivate>& Context, bo
 	bHasInvalidMesh = false;
 	bool bNeedsUpdate = false;
 	
-	const int32 NumComponents = Context->InstanceUpdateData.Components.Num();
-	bNeedsUpdate |= NumComponents != InstancePrivate->SkeletalMeshes.Num();
+	const int32 NumInstanceComponents = Context->InstanceUpdateData.Components.Num();
+	bNeedsUpdate |= NumInstanceComponents != InstancePrivate->SkeletalMeshes.Num();
 
 	// Find which components need an update
-	Context->MeshChangedPerInstanceComponent.Init(false, NumComponents);
+	Context->MeshChangedPerInstanceComponent.Init(false, NumInstanceComponents);
 
-	for (int32 ComponentIndex = 0; ComponentIndex < NumComponents; ++ComponentIndex)
+	for (int32 InstanceComponentIndex = 0; InstanceComponentIndex < NumInstanceComponents; ++InstanceComponentIndex)
 	{
-		const TSharedRef<FInstanceUpdateData::FComponent>& Component = Context->InstanceUpdateData.Components[ComponentIndex];
-		const FName ComponentName = Context->ComponentNames[Component->Id];
-		const FCustomizableInstanceComponentData& ComponentData = InstancePrivate->GetComponentData(Component->Id);
+		const TSharedRef<FInstanceUpdateData::FComponent>& Component = Context->InstanceUpdateData.Components[InstanceComponentIndex];
+		const FCustomizableObjectComponentIndex ObjectComponentIndex = Component->Id;
+		const FName ComponentName = Context->ComponentNames[ObjectComponentIndex.GetValue()];
 		
 		bool bComponentWithMesh = false;
 
@@ -1789,7 +1866,10 @@ bool DoComponentsNeedUpdate(const TSharedRef<FUpdateContextPrivate>& Context, bo
 		
 		if (Component->SkeletalMesh->PassthroughObject)
 		{
-			bMeshChanged = Component->SkeletalMeshId != ComponentData.LastComponentDescriptor.SkeletalMeshDescriptor;
+			if (const FCustomizableInstanceComponentData* ComponentData = InstancePrivate->GetComponentData(ObjectComponentIndex)) // Could be nullptr if the component has not been generated.
+			{
+				bMeshChanged = Component->SkeletalMeshId != ComponentData->LastComponentDescriptor.SkeletalMeshDescriptor;
+			}
 		}
 		else
 		{
@@ -1813,7 +1893,7 @@ bool DoComponentsNeedUpdate(const TSharedRef<FUpdateContextPrivate>& Context, bo
 				{
 					UE_LOG(LogMutable, Error, TEXT("Failed to generate SkeletalMesh for CO Instance [%s]. CO [%s] has invalid geometry for LOD [%d] Component [%d]."),
 						*Instance->GetName(), *CustomizableObject->GetName(),
-						LODIndex, ComponentIndex);
+						LODIndex, InstanceComponentIndex);
 					bHasInvalidMesh = true; // Unreal does not support empty LODs.
 				}
 
@@ -1822,13 +1902,16 @@ bool DoComponentsNeedUpdate(const TSharedRef<FUpdateContextPrivate>& Context, bo
 			}
 			
 			// Compare Mesh IDs 
-			const FComponentDescriptor* ComponentDescriptor = Context->GetComponentDescriptor(Component->Id);
+			const FComponentDescriptor* ComponentDescriptor = Context->GetComponentDescriptor(ObjectComponentIndex);
 			if (ensure(ComponentDescriptor))
 			{
-				bMeshChanged = ComponentDescriptor->MeshDescriptors != ComponentData.LastComponentDescriptor.MeshDescriptors;
+				if (const FCustomizableInstanceComponentData* ComponentData = InstancePrivate->GetComponentData(ObjectComponentIndex)) // Could be nullptr if the component has not been generated.
+				{
+					bMeshChanged = ComponentDescriptor->MeshDescriptors != ComponentData->LastComponentDescriptor.MeshDescriptors;
+				}
 			}
 
-			Context->MeshChangedPerInstanceComponent[ComponentIndex] = bMeshChanged;
+			Context->MeshChangedPerInstanceComponent[InstanceComponentIndex] = bMeshChanged;
 		}
 		
 		bNeedsUpdate |= bMeshChanged;
@@ -1983,6 +2066,35 @@ bool UCustomizableInstancePrivate::UpdateSkeletalMesh_PostBeginUpdate0(UCustomiz
 		}
 	}
 
+	// Detect the components that have disappeared since the last update and blank their properties so they aren't reused the
+	// next time they reappear
+	{
+		for (const auto& MeshIt : SkeletalMeshes)
+		{
+			FName DeletedComponentCandidateName = MeshIt.Key;
+			bool bIsDeleted = true;
+
+			for (const auto& InstanceUpdateComponent : Context->InstanceUpdateData.Components)
+			{
+				FName InstanceUpdateComponentName = ModelResources.ComponentNamesPerObjectComponent[InstanceUpdateComponent->Id.GetValue()];
+
+				if (DeletedComponentCandidateName == InstanceUpdateComponentName)
+				{
+					bIsDeleted = false;
+
+					break;
+				}
+			}
+
+			if (bIsDeleted)
+			{
+				ClearInstanceComponentData(DeletedComponentCandidateName);
+
+				bUpdateMeshes = true;
+			}
+		}
+	}
+
 	// None of the current meshes requires a mesh update. Continue to BuildResources
 	if (!bUpdateMeshes)
 	{
@@ -1991,18 +2103,34 @@ bool UCustomizableInstancePrivate::UpdateSkeletalMesh_PostBeginUpdate0(UCustomiz
 
 	SkeletalMeshes.Reset();
 	
-	const int32 NumComponents = Context->InstanceUpdateData.Components.Num();
-	for (int32 ComponentIndex = 0; ComponentIndex < NumComponents; ++ComponentIndex)
+	const int32 NumInstanceComponents = Context->InstanceUpdateData.Components.Num();
+	for (int32 InstanceComponentIndexValue = 0; InstanceComponentIndexValue < NumInstanceComponents; ++InstanceComponentIndexValue)
 	{
-		const TSharedPtr<FInstanceUpdateData::FComponent>& Component = Context->InstanceUpdateData.Components[ComponentIndex];
-		
-		FCustomizableInstanceComponentData& ComponentData = GetComponentData(Component->Id);
-		const FName ComponentName = ComponentData.ComponentName;
+		FCustomizableObjectInstanceComponentIndex InstanceComponentIndex(InstanceComponentIndexValue);
+		const FInstanceUpdateData::FComponent* Component = Context->GetComponentUpdateData( InstanceComponentIndex );
+		if (!Component)
+		{
+			continue;
+		}
+		const FCustomizableObjectComponentIndex ObjectComponentIndex = Component->Id;
+		if (!ModelResources.ComponentNamesPerObjectComponent.IsValidIndex(ObjectComponentIndex.GetValue()))
+		{
+			continue;
+		}
+		const FName ComponentName = ModelResources.ComponentNamesPerObjectComponent[ObjectComponentIndex.GetValue()];
+
+		FCustomizableInstanceComponentData* ComponentData = GetComponentData(ComponentName);
+		if (!ComponentData)
+		{
+			ensure(false);
+			Context->UpdateResult = EUpdateResult::Error;
+			return false;
+		}
 
-		FComponentDescriptor* ComponentDescriptor = Context->GetComponentDescriptor(Component->Id);
+		FComponentDescriptor* ComponentDescriptor = Context->GetComponentDescriptor(ObjectComponentIndex);
 		check(ComponentDescriptor);
 		check(ComponentDescriptor->MeshDescriptors.Num() == MAX_MESH_LOD_COUNT);
-		ComponentData.LastComponentDescriptor = *ComponentDescriptor;
+		ComponentData->LastComponentDescriptor = *ComponentDescriptor;
 
 		if (Component->SkeletalMesh->PassthroughObject)
 		{
@@ -2015,13 +2143,14 @@ bool UCustomizableInstancePrivate::UpdateSkeletalMesh_PostBeginUpdate0(UCustomiz
 		}
 		
 		// If the component doesn't need an update copy the previously generated mesh.
-		if (!Context->MeshChangedPerInstanceComponent[ComponentIndex])
+		if (!Context->MeshChangedPerInstanceComponent[InstanceComponentIndex.GetValue()])
 		{
 			if (TObjectPtr<USkeletalMesh>* Result = OldSkeletalMeshes.Find(ComponentName))
 			{
 				SkeletalMeshes.Add(ComponentName, *Result);
-				continue;
 			}
+
+			continue;
 		}
 
 		const int32 FirstLOD = Context->bStreamMeshLODs ?
@@ -2054,7 +2183,7 @@ bool UCustomizableInstancePrivate::UpdateSkeletalMesh_PostBeginUpdate0(UCustomiz
 			continue;
 		}
 
-		if (!ModelResources.ReferenceSkeletalMeshesData.IsValidIndex(Component->Id))
+		if (!ModelResources.ReferenceSkeletalMeshesData.IsValidIndex(ObjectComponentIndex.GetValue()))
 		{
 			Context->UpdateResult = EUpdateResult::Error;
 			return false;
@@ -2108,13 +2237,13 @@ bool UCustomizableInstancePrivate::UpdateSkeletalMesh_PostBeginUpdate0(UCustomiz
 			SkeletalMeshes.Add(ComponentName, SkeletalMesh);
 		}
 		
-		const FMutableRefSkeletalMeshData& RefSkeletalMeshData = ModelResources.ReferenceSkeletalMeshesData[Component->Id];
+		const FMutableRefSkeletalMeshData& RefSkeletalMeshData = ModelResources.ReferenceSkeletalMeshesData[ObjectComponentIndex.GetValue()];
 
 		// Set up the default information any mesh from this component will have (LODArrayInfos, RenderData, Mesh settings, etc). 
-		InitSkeletalMeshData(Context, SkeletalMesh, RefSkeletalMeshData, *CustomizableObject, Component->Id);
+		InitSkeletalMeshData(Context, SkeletalMesh, RefSkeletalMeshData, *CustomizableObject, ObjectComponentIndex);
 		
 		// Construct a new skeleton, fix up ActiveBones and Bonemap arrays and recompute the RefInvMatrices
-		const bool bBuildSkeletonDataSuccess = BuildSkeletonData(Context, *SkeletalMesh, ComponentIndex);
+		const bool bBuildSkeletonDataSuccess = BuildSkeletonData(Context, *SkeletalMesh, FCustomizableObjectInstanceComponentIndex(InstanceComponentIndex) );
 		if (!bBuildSkeletonDataSuccess)
 		{
 			Context->UpdateResult = EUpdateResult::Error;
@@ -2126,7 +2255,7 @@ bool UCustomizableInstancePrivate::UpdateSkeletalMesh_PostBeginUpdate0(UCustomiz
 		{
 			constexpr bool bDisallowCollisionBetweenAssets = true;
 			UPhysicsAsset* PhysicsAssetResult = GetOrBuildMainPhysicsAsset(
-				Context, RefSkeletalMeshData.PhysicsAsset, MutablePhysics.Get(), bDisallowCollisionBetweenAssets, ComponentIndex);
+				Context, RefSkeletalMeshData.PhysicsAsset, MutablePhysics.Get(), bDisallowCollisionBetweenAssets, FCustomizableObjectInstanceComponentIndex(InstanceComponentIndex) );
 
 			SkeletalMesh->SetPhysicsAsset(PhysicsAssetResult);
 
@@ -2173,7 +2302,7 @@ bool UCustomizableInstancePrivate::UpdateSkeletalMesh_PostBeginUpdate0(UCustomiz
 
 			Entry.PropertyIndex = Info.PropertyIndex;
 			Entry.PhysicsAsset = MakePhysicsAssetFromTemplateAndMutableBody(
-				Context,PhysicsAssetTemplate, AdditionalPhysicsBody.Get(), ComponentIndex);
+				Context,PhysicsAssetTemplate, AdditionalPhysicsBody.Get(), InstanceComponentIndex);
 		}
 
 		// Add sockets from the SkeletalMesh of reference and from the MutableMesh
@@ -2184,16 +2313,16 @@ bool UCustomizableInstancePrivate::UpdateSkeletalMesh_PostBeginUpdate0(UCustomiz
 			Pair.Key->OnSkeletalMeshCreated(Pair.Value, ComponentName, SkeletalMesh);
 		}
 		
-		BuildOrCopyElementData(Context, SkeletalMesh, ComponentIndex);
-		bool const bCopyRenderDataSuccess = BuildOrCopyRenderData(Context, SkeletalMesh, Instance, ComponentIndex);
+		BuildOrCopyElementData(Context, SkeletalMesh, InstanceComponentIndex);
+		bool const bCopyRenderDataSuccess = BuildOrCopyRenderData(Context, SkeletalMesh, Instance, InstanceComponentIndex);
 		if (!bCopyRenderDataSuccess)
 		{
 			Context->UpdateResult = EUpdateResult::Error;
 			return false;
 		}
 
-		BuildOrCopyMorphTargetsData(Context, SkeletalMesh, ComponentIndex);
-		BuildOrCopyClothingData(Context, SkeletalMesh, ModelResources, ComponentIndex, ClothingPhysicsAssets);
+		BuildOrCopyMorphTargetsData(Context, SkeletalMesh, InstanceComponentIndex);
+		BuildOrCopyClothingData(Context, SkeletalMesh, ModelResources, InstanceComponentIndex, ClothingPhysicsAssets);
 
 		FSkeletalMeshRenderData* RenderData = SkeletalMesh->GetResourceForRendering();
 		ensure(RenderData && RenderData->LODRenderData.Num() > 0);
@@ -2955,7 +3084,7 @@ void UCustomizableInstancePrivate::DiscardResources()
 }
 
 
-void UCustomizableInstancePrivate::ForceSetReferenceSkeletalMesh() const
+void UCustomizableInstancePrivate::SetReferenceSkeletalMesh() const
 {
 	UCustomizableObjectInstance* Instance = Cast<UCustomizableObjectInstance>(GetOuter());
 	if (!Instance)
@@ -2989,27 +3118,22 @@ void UCustomizableInstancePrivate::ForceSetReferenceSkeletalMesh() const
 			continue;
 		}
 #endif
-		
-		USkeletalMeshComponent* Parent = CustomizableObjectInstanceUsage->GetAttachParent();
-		if (!Parent)
-		{
-			continue;
-		}
-		
+
 		const FName& ComponentName = CustomizableObjectInstanceUsage->GetComponentName();
-		const UE::Mutable::Private::FComponentId ComponentId = ModelResources->ComponentNamesPerObjectComponent.IndexOfByKey(ComponentName);
-		if (!ModelResources->ReferenceSkeletalMeshesData.IsValidIndex(ComponentId))
+		const int32 ObjectComponentIndex = ModelResources->ComponentNamesPerObjectComponent.IndexOfByKey(ComponentName);
+		if (!ModelResources->ReferenceSkeletalMeshesData.IsValidIndex(ObjectComponentIndex))
 		{
 			continue;
 		}
 		
-		Parent->EmptyOverrideMaterials();
+		if (USkeletalMeshComponent* Parent = CustomizableObjectInstanceUsage->GetAttachParent())
+		{
+			Parent->EmptyOverrideMaterials();
 
-		// TODO Due to the adding support of multiple components with the same name this is no longer correct. A single name may point to multiple FComponentId.
-		TSoftObjectPtr<USkeletalMesh> SoftObjectPtr = ModelResources->ReferenceSkeletalMeshesData[ComponentId].SoftSkeletalMesh;
-		USkeletalMesh* SkeletalMesh = UE::Mutable::Private::LoadObject(SoftObjectPtr);
-		
-		Parent->SetSkeletalMesh(SkeletalMesh);
+			TSoftObjectPtr<USkeletalMesh> SoftObjectPtr = ModelResources->ReferenceSkeletalMeshesData[ObjectComponentIndex].SoftSkeletalMesh;
+			USkeletalMesh* SkeletalMesh = UE::Mutable::Private::LoadObject(SoftObjectPtr);
+			Parent->SetSkeletalMesh(SkeletalMesh);
+		}
 	}
 }
 
@@ -3452,13 +3576,6 @@ void ConvertImage(UTexture2D* Texture, UE::Mutable::Private::TManagedPtr<const U
 	Texture->SetPlatformData(MutableCreateImagePlatformData(MutableImage, OnlyLOD,0,0).Release());
 }
 
-
-FCustomizableInstanceComponentData::FCustomizableInstanceComponentData()
-{
-	LastComponentDescriptor.MeshDescriptors.Init({}, MAX_MESH_LOD_COUNT);
-}
-
-
 static int32 EnableRayTracingFix = 0;
 FAutoConsoleVariableRef CVarMutableEnableRayTracingFix(
 	TEXT("mutable.EnableRayTracingFix"),
@@ -3468,7 +3585,7 @@ FAutoConsoleVariableRef CVarMutableEnableRayTracingFix(
 	TEXT("If 2, Enable fix for all generated meshes. Meshes will have ray tracing disabled.")
 	);
 
-void InitSkeletalMeshData(const TSharedRef<FUpdateContextPrivate>& Context, USkeletalMesh* SkeletalMesh, const FMutableRefSkeletalMeshData& RefSkeletalMeshData, const UCustomizableObject& CustomizableObject, UE::Mutable::Private::FComponentId ComponentId)
+void InitSkeletalMeshData(const TSharedRef<FUpdateContextPrivate>& Context, USkeletalMesh* SkeletalMesh, const FMutableRefSkeletalMeshData& RefSkeletalMeshData, const UCustomizableObject& CustomizableObject, FCustomizableObjectComponentIndex ObjectComponentIndex)
 {
 	MUTABLE_CPUPROFILER_SCOPE(UCustomizableInstancePrivate::InitSkeletalMesh);
 
@@ -3477,7 +3594,7 @@ void InitSkeletalMeshData(const TSharedRef<FUpdateContextPrivate>& Context, USke
 	UCustomizableObjectInstance* Instance = Context->Instance.Get();
 	UCustomizableInstancePrivate* InstancePrivate = Instance->GetPrivate();
 	
-	const FName ComponentName = Context->ComponentNames[ComponentId];
+	const FName ComponentName = Context->ComponentNames[ObjectComponentIndex.GetValue()];
 
 	SkeletalMesh->NeverStream = !Context->bStreamMeshLODs;
 
@@ -3499,20 +3616,21 @@ void InitSkeletalMeshData(const TSharedRef<FUpdateContextPrivate>& Context, USke
 
 	// Asset User Data
 	{
-		const FCustomizableInstanceComponentData& ComponentData = InstancePrivate->GetComponentData(ComponentId);
-		for (TObjectPtr<UAssetUserData> AssetUserData : ComponentData.AssetUserDataArray)
+		const FCustomizableInstanceComponentData* ComponentData = InstancePrivate->GetComponentData(ObjectComponentIndex);
+		check(ComponentData);
+		for (TObjectPtr<UAssetUserData> AssetUserData : ComponentData->AssetUserDataArray)
 		{
 			SkeletalMesh->AddAssetUserData(AssetUserData);
 		}
 
 		//Custom Asset User Data
 		if (Context->Instance->GetAnimationGameplayTags().Num() ||
-			ComponentData.AnimSlotToBP.Num())
+			ComponentData->AnimSlotToBP.Num())
 		{
 			UCustomizableObjectInstanceUserData* InstanceData = NewObject<UCustomizableObjectInstanceUserData>(SkeletalMesh, NAME_None, RF_Public | RF_Transactional);
 			InstanceData->AnimationGameplayTag = Context->Instance->GetAnimationGameplayTags();
 			
-			for (const TTuple<FName, TSoftClassPtr<UAnimInstance>>& AnimSlot : ComponentData.AnimSlotToBP)
+			for (const TTuple<FName, TSoftClassPtr<UAnimInstance>>& AnimSlot : ComponentData->AnimSlotToBP)
 			{
 				FCustomizableObjectAnimationSlot AnimationSlot;
 				AnimationSlot.Name = AnimSlot.Key;
@@ -3608,8 +3726,8 @@ void InitSkeletalMeshData(const TSharedRef<FUpdateContextPrivate>& Context, USke
 
 	// Set Min LOD (Override the Reference Skeletal Mesh LOD Settings)
 	const UModelResources& ModelResources = CustomizableObject.GetPrivate()->GetModelResourcesChecked();
-	SkeletalMesh->SetMinLod(FMath::Max(ModelResources.MinLODPerComponent.FindChecked(ComponentId).GetDefault(), static_cast<int32>(Context->FirstLODAvailable[ComponentName])));
-	SkeletalMesh->SetQualityLevelMinLod(ModelResources.MinQualityLevelLODPerComponent.FindChecked(ComponentId));
+	SkeletalMesh->SetMinLod(FMath::Max(ModelResources.MinLODPerComponent.FindChecked(ComponentName).GetDefault(), static_cast<int32>(Context->FirstLODAvailable[ComponentName])));
+	SkeletalMesh->SetQualityLevelMinLod(ModelResources.MinQualityLevelLODPerComponent.FindChecked(ComponentName));
 
 	// Set up unreal's default material, will be replaced when building materials
 	{
@@ -3625,7 +3743,7 @@ void InitSkeletalMeshData(const TSharedRef<FUpdateContextPrivate>& Context, USke
 }
 
 
-bool BuildSkeletonData(const TSharedRef<FUpdateContextPrivate>& Context, USkeletalMesh& SkeletalMesh, int32 ComponentIndex)
+bool BuildSkeletonData(const TSharedRef<FUpdateContextPrivate>& Context, USkeletalMesh& SkeletalMesh, FCustomizableObjectInstanceComponentIndex InstanceComponentIndex)
 {
 	MUTABLE_CPUPROFILER_SCOPE(UCustomizableInstancePrivate::BuildSkeletonData);
 
@@ -3633,13 +3751,14 @@ bool BuildSkeletonData(const TSharedRef<FUpdateContextPrivate>& Context, USkelet
 	UCustomizableInstancePrivate* InstancePrivate = Instance->GetPrivate();
 
 	UCustomizableObject* CustomizableObject = Instance->GetCustomizableObject();
+	
+	FCustomizableObjectComponentIndex ObjectComponentIndex = Context->GetObjectComponentIndex(InstanceComponentIndex);
 
-	UE::Mutable::Private::FComponentId ComponentId = Context->InstanceUpdateData.Components[ComponentIndex]->Id;
-
-	FCustomizableInstanceComponentData& ComponentData = InstancePrivate->GetComponentData(ComponentId);
+	FCustomizableInstanceComponentData* ComponentData = InstancePrivate->GetComponentData(ObjectComponentIndex);
+	check(ComponentData);
 	
 	bool bCreatedNewSkeleton = false;
-	const TObjectPtr<USkeleton> Skeleton = MergeSkeletons(*CustomizableObject, ComponentData, bCreatedNewSkeleton);
+	const TObjectPtr<USkeleton> Skeleton = MergeSkeletons(*CustomizableObject, *ComponentData, bCreatedNewSkeleton);
 	if (!Skeleton)
 	{
 		return false;
@@ -3653,8 +3772,8 @@ bool BuildSkeletonData(const TSharedRef<FUpdateContextPrivate>& Context, USkelet
 	const TArray<FMeshBoneInfo>& RawRefBoneInfo = ReferenceSkeleton.GetRawRefBoneInfo();
 	const int32 RawRefBoneCount = ReferenceSkeleton.GetRawBoneNum();
 
-	const TArray<FInstanceUpdateData::FBone>& BonePose = Context->InstanceUpdateData.SkeletonsPerInstanceComponent[ComponentIndex].BonePose;
-	TMap<UE::Mutable::Private::FBoneName, TPair<FName,uint16>>& BoneInfoMap = Context->InstanceUpdateData.SkeletonsPerInstanceComponent[ComponentIndex].BoneInfoMap;
+	const TArray<FInstanceUpdateData::FBone>& BonePose = Context->InstanceUpdateData.SkeletonsPerInstanceComponent[InstanceComponentIndex.GetValue()].BonePose;
+	TMap<UE::Mutable::Private::FBoneName, TPair<FName,uint16>>& BoneInfoMap = Context->InstanceUpdateData.SkeletonsPerInstanceComponent[InstanceComponentIndex.GetValue()].BoneInfoMap;
 
 	{
 		MUTABLE_CPUPROFILER_SCOPE(BuildSkeletonData_BuildBoneInfoMap);
@@ -3735,7 +3854,7 @@ bool BuildSkeletonData(const TSharedRef<FUpdateContextPrivate>& Context, USkelet
 		// MergeSkeletons function doesn't do it. Only do it for newly generated ones, not for cached or non-transient ones.
 		GeneratedSkeleton->RecreateBoneTree(&SkeletalMesh);
 		
-		TArray<TObjectPtr<USkeleton>>& SkeletonsToMerge = ComponentData.Skeletons.SkeletonsToMerge;
+		TArray<TObjectPtr<USkeleton>>& SkeletonsToMerge = ComponentData->Skeletons.SkeletonsToMerge;
 		check(SkeletonsToMerge.Num() > 1);
 
 		TMap<FName, EBoneTranslationRetargetingMode::Type> BoneNamesToRetargetingMode;
@@ -3873,17 +3992,17 @@ void BuildMeshSockets(USkeletalMesh* SkeletalMesh, const UModelResources& ModelR
 }
 
 
-void BuildOrCopyElementData(const TSharedRef<FUpdateContextPrivate>& Context, USkeletalMesh* SkeletalMesh, int32 ComponentIndex)
+void BuildOrCopyElementData(const TSharedRef<FUpdateContextPrivate>& Context, USkeletalMesh* SkeletalMesh, FCustomizableObjectInstanceComponentIndex InstanceComponentIndex)
 {
 	MUTABLE_CPUPROFILER_SCOPE(UCustomizableInstancePrivate::BuildOrCopyElementData);
 
-	const TSharedPtr<FInstanceUpdateData::FComponent>& Component = Context->InstanceUpdateData.Components[ComponentIndex];
+	const FInstanceUpdateData::FComponent* Component = Context->GetComponentUpdateData(InstanceComponentIndex);
 	if (!Component)
 	{
 		return;
 	}
 
-	const FName ComponentName = Context->ComponentNames[Component->Id];
+	const FName ComponentName = Context->ComponentNames[Component->Id.GetValue()];
 
 	for (int32 LODIndex = Context->FirstLODAvailable[ComponentName]; LODIndex < Component->LODCount; ++LODIndex)
 	{
@@ -3901,7 +4020,7 @@ void BuildOrCopyElementData(const TSharedRef<FUpdateContextPrivate>& Context, US
 	}
 }
 
-void BuildOrCopyMorphTargetsData(const TSharedRef<FUpdateContextPrivate>& Context, USkeletalMesh* SkeletalMesh, int32 ComponentIndex)
+void BuildOrCopyMorphTargetsData(const TSharedRef<FUpdateContextPrivate>& Context, USkeletalMesh* SkeletalMesh, FCustomizableObjectInstanceComponentIndex InstanceComponentIndex)
 {
 	MUTABLE_CPUPROFILER_SCOPE(UCustomizableInstancePrivate::BuildOrCopyMorphTargetsData);
 
@@ -3915,12 +4034,11 @@ void BuildOrCopyMorphTargetsData(const TSharedRef<FUpdateContextPrivate>& Contex
 	{
 		return;
 	}
+	
+	FCustomizableObjectComponentIndex ObjectComponentIndex = Context->GetObjectComponentIndex(InstanceComponentIndex);
+	const FName ComponentName = Context->ComponentNames[ObjectComponentIndex.GetValue()];
 
-	const TSharedPtr<FInstanceUpdateData::FComponent>& Component = Context->InstanceUpdateData.Components[ComponentIndex];
-
-	const FName ComponentName = Context->ComponentNames[Component->Id];
-
-	FSkeletalMeshMorphTargets* ComponentMorphTargets = Context->InstanceUpdateData.RealTimeMorphTargets.Find(Component->Id);
+	FSkeletalMeshMorphTargets* ComponentMorphTargets = Context->InstanceUpdateData.RealTimeMorphTargets.Find(ComponentName);
 	if (!ComponentMorphTargets)
 	{
 		return;
@@ -4017,7 +4135,7 @@ void BuildOrCopyClothingData(
 		const TSharedRef<FUpdateContextPrivate>& Context, 
 		USkeletalMesh* SkeletalMesh, 
 		const UModelResources& ModelResources,
-		int32 ComponentIndex,
+		FCustomizableObjectInstanceComponentIndex InstanceComponentIndex,
 		const TArray<TObjectPtr<UPhysicsAsset>>& ClothingPhysicsAssets)
 {
 	MUTABLE_CPUPROFILER_SCOPE(UCustomizableInstancePrivate::BuildOrCopyClothingData);
@@ -4048,13 +4166,13 @@ void BuildOrCopyClothingData(
 		return;
 	}
 
-	const TSharedPtr<FInstanceUpdateData::FComponent>& Component = Context->InstanceUpdateData.Components[ComponentIndex];
+	const FInstanceUpdateData::FComponent* Component = Context->GetComponentUpdateData(InstanceComponentIndex);
 	if (!Component)
 	{
 		return;
 	}
 
-	const FName& ComponentName = Context->ComponentNames[Component->Id];
+	const FName& ComponentName = Context->ComponentNames[Component->Id.GetValue()];
 
 	TArray<FSectionClothMetadata> SectionClothMetadata; // Sections must be sorted ascending
 	SectionClothMetadata.Reserve(32);
@@ -4452,7 +4570,7 @@ void BuildOrCopyClothingData(
 
 
 bool BuildOrCopyRenderData(const TSharedRef<FUpdateContextPrivate>& Context, USkeletalMesh* SkeletalMesh, 
-	UCustomizableObjectInstance* Public, int32 ComponentIndex)
+	UCustomizableObjectInstance* Public, FCustomizableObjectInstanceComponentIndex InstanceComponentIndex)
 {
 	MUTABLE_CPUPROFILER_SCOPE(UCustomizableInstancePrivate::BuildOrCopyRenderData);
 
@@ -4460,15 +4578,18 @@ bool BuildOrCopyRenderData(const TSharedRef<FUpdateContextPrivate>& Context, USk
 	check(RenderData);
 
 	UCustomizableObject* CustomizableObject = Public->GetCustomizableObject();
-	
-	const TSharedPtr<FInstanceUpdateData::FComponent>& Component = Context->InstanceUpdateData.Components[ComponentIndex];
+
+	// It must be not null as it's checked in the calling function
+	check(CustomizableObject);
+
+	const FInstanceUpdateData::FComponent* Component = Context->GetComponentUpdateData(InstanceComponentIndex);
 	if (!Component)
 	{
 		return false;
 	}
 
 	const UModelResources& ModelResources = CustomizableObject->GetPrivate()->GetModelResourcesChecked();
-	const FName ComponentName = ModelResources.ComponentNamesPerObjectComponent[Component->Id];
+	const FName ComponentName = ModelResources.ComponentNamesPerObjectComponent[Component->Id.GetValue()];
 
 	const int32 FirstLOD = Context->bStreamMeshLODs ?
 		Context->FirstLODAvailable[ComponentName] :
@@ -4486,17 +4607,17 @@ bool BuildOrCopyRenderData(const TSharedRef<FUpdateContextPrivate>& Context, USk
 			UE_LOG(LogMutable, Warning, TEXT("Building instance: generated mesh [%s] has LOD [%d] of object component index [%d] with no mesh.")
 				, *SkeletalMesh->GetName()
 				, LODIndex
-				, Component->Id);
+				, Component->Id.GetValue());
 
 			// End with failure
 			return false;
 		}
 
-		TRACE_CPUPROFILER_EVENT_SCOPE_TEXT(*FString::Printf(TEXT("BuildRenderData: Component index %d, LOD %d"), Component->Id, LODIndex));
+		TRACE_CPUPROFILER_EVENT_SCOPE_TEXT(*FString::Printf(TEXT("BuildRenderData: Component index %d, LOD %d"), Component->Id.GetValue(), LODIndex));
 		
 		FSkeletalMeshLODRenderData& LODResource = RenderData->LODRenderData[LODIndex];
 
-		const TMap<UE::Mutable::Private::FBoneName, TPair<FName, uint16>>& BoneInfoMap = Context->InstanceUpdateData.SkeletonsPerInstanceComponent[ComponentIndex].BoneInfoMap;
+		const TMap<UE::Mutable::Private::FBoneName, TPair<FName, uint16>>& BoneInfoMap = Context->InstanceUpdateData.SkeletonsPerInstanceComponent[InstanceComponentIndex.GetValue()].BoneInfoMap;
 		
 		// Set active and required bones
 		{
@@ -4521,14 +4642,14 @@ bool BuildOrCopyRenderData(const TSharedRef<FUpdateContextPrivate>& Context, USk
 		for (int32 MeshSectionIndex = 0; MeshSectionIndex < MeshNumSurfaces; ++MeshSectionIndex)
 		{
 			uint32 MeshSurfaceId = LOD->Mesh->GetSurfaceId(MeshSectionIndex);
-			int32 InstanceSurfaceIndex = Context->MutableInstance->FindSurfaceById(ComponentIndex, LODIndex, MeshSurfaceId);
+			int32 InstanceSurfaceIndex = Context->MutableInstance->FindSurfaceById(InstanceComponentIndex.GetValue(), LODIndex, MeshSurfaceId);
 			
 			if (InstanceSurfaceIndex < 0)
 			{
 				continue;
 			}
 			
-			uint32 SurfaceMetadataId = Context->MutableInstance->GetSurfaceCustomId(ComponentIndex, LODIndex, InstanceSurfaceIndex);
+			uint32 SurfaceMetadataId = Context->MutableInstance->GetSurfaceCustomId(InstanceComponentIndex.GetValue(), LODIndex, InstanceSurfaceIndex);
 			
 			uint32 UsedSurfaceMetadataId = 0;
 			if (SurfaceMetadataId != 0)
@@ -4736,29 +4857,33 @@ UE::Tasks::FTask LoadAdditionalAssetsAndData(const TSharedRef<FUpdateContextPriv
 	InstancePrivate->AnimBpPhysicsAssets.Reset();
 
 	// Load Skeletons required by the SubMeshes of the newly generated Mesh, will be merged later
-	for (int32 ComponentIndex = 0; ComponentIndex < Context->InstanceUpdateData.Components.Num(); ++ComponentIndex)
+	for (int32 InstanceComponentIndex = 0; InstanceComponentIndex < Context->NumInstanceComponents; ++InstanceComponentIndex)
 	{
-		const TSharedPtr<FInstanceUpdateData::FComponent>& Component = Context->InstanceUpdateData.Components[ComponentIndex];
-		if (Component->Id == INDEX_NONE)
+		FCustomizableObjectComponentIndex ObjectComponentIndex = Context->GetObjectComponentIndex(FCustomizableObjectInstanceComponentIndex(InstanceComponentIndex));
+		if (!ObjectComponentIndex.IsValid())
 		{
 			continue;
 		}
 
-		const FInstanceUpdateData::FSkeletonData& SkeletonData = Context->InstanceUpdateData.SkeletonsPerInstanceComponent[ComponentIndex];
+		const FInstanceUpdateData::FSkeletonData& SkeletonData = Context->InstanceUpdateData.SkeletonsPerInstanceComponent[InstanceComponentIndex];
 		
-		FCustomizableInstanceComponentData& ComponentData = InstancePrivate->GetComponentData(Component->Id);
+		FCustomizableInstanceComponentData* ComponentData = InstancePrivate->GetComponentData(ObjectComponentIndex);
+		if (!ComponentData)
+		{
+			continue;
+		}
 
 		// Reuse merged Skeleton if cached
-		ComponentData.Skeletons.Skeleton = CustomizableObject->GetPrivate()->SkeletonCache.Get(SkeletonData.SkeletonIds);
-		if (ComponentData.Skeletons.Skeleton)
+		ComponentData->Skeletons.Skeleton = CustomizableObject->GetPrivate()->SkeletonCache.Get(SkeletonData.SkeletonIds);
+		if (ComponentData->Skeletons.Skeleton)
 		{
-			ComponentData.Skeletons.SkeletonIds.Empty();
-			ComponentData.Skeletons.SkeletonsToMerge.Empty();
+			ComponentData->Skeletons.SkeletonIds.Empty();
+			ComponentData->Skeletons.SkeletonsToMerge.Empty();
 			continue;
 		}
 
 		// Add Skeletons to merge
-		ComponentData.Skeletons.SkeletonsToMerge = SkeletonData.Skeletons;
+		ComponentData->Skeletons.SkeletonsToMerge = SkeletonData.Skeletons;
 
 		for (const uint32 SkeletonId : SkeletonData.SkeletonIds)
 		{
@@ -4769,11 +4894,11 @@ UE::Tasks::FTask LoadAdditionalAssetsAndData(const TSharedRef<FUpdateContextPriv
 			}
 
 			// Add referenced skeletons to the assets to stream
-			ComponentData.Skeletons.SkeletonIds.Add(SkeletonId);
+			ComponentData->Skeletons.SkeletonIds.Add(SkeletonId);
 
 			if (USkeleton* Skeleton = AssetPtr.Get())
 			{
-				ComponentData.Skeletons.SkeletonsToMerge.Add(Skeleton);
+				ComponentData->Skeletons.SkeletonsToMerge.Add(Skeleton);
 			}
 			else
 			{
@@ -4786,10 +4911,19 @@ UE::Tasks::FTask LoadAdditionalAssetsAndData(const TSharedRef<FUpdateContextPriv
 	const bool bUpdateMeshes = DoComponentsNeedUpdate(Context, bHasInvalidMesh);
 
 	// Load assets coming from SubMeshes of the newly generated Mesh
-	for (int32 ComponentIndex = 0; ComponentIndex < Context->InstanceUpdateData.Components.Num(); ++ComponentIndex)
+	for (int32 InstanceComponentIndex = 0; InstanceComponentIndex < Context->InstanceUpdateData.Components.Num(); ++InstanceComponentIndex)
 	{
-		const TSharedRef<const FInstanceUpdateData::FComponent>& Component = Components[ComponentIndex];
-		FCustomizableInstanceComponentData& ComponentData = InstancePrivate->GetComponentData(Component->Id);
+		const TSharedRef<const FInstanceUpdateData::FComponent>& Component = Components[InstanceComponentIndex];
+		FCustomizableInstanceComponentData* ComponentData = InstancePrivate->GetComponentData(Component->Id);
+		
+		if (!ComponentData)
+        {
+            UE_LOG(LogMutable, Error, TEXT("ComponentData is null: ComponentIDIndex : '%i', COIName :'%s'."), Component->Id.GetValue(), *Instance->GetName());
+			Context->UpdateResult = EUpdateResult::Error;
+			
+			return UE::Tasks::FTask{};
+        }
+        
 		
 		if (Component->SkeletalMesh->PassthroughObject)
 		{
@@ -4797,17 +4931,18 @@ UE::Tasks::FTask LoadAdditionalAssetsAndData(const TSharedRef<FUpdateContextPriv
 			continue;
 		}
 		
-		const FMutableRefSkeletalMeshData& RefSkeletalMeshData = ModelResources.ReferenceSkeletalMeshesData[Component->Id];
+		const FCustomizableObjectComponentIndex ObjectComponentIndex = Component->Id;
+		const FMutableRefSkeletalMeshData& RefSkeletalMeshData = ModelResources.ReferenceSkeletalMeshesData[ObjectComponentIndex.GetValue()];
 
 		for (int32 AssetUserDataIndex : RefSkeletalMeshData.AssetUserDataIndices)
 		{
 #if !WITH_EDITOR
 			Context->StreamedResourceIndex.AddUnique(AssetUserDataIndex); // Used to hold/release streamed resources in non-editor builds.
 #endif
-			ComponentData.StreamedResourceIndex.AddUnique(AssetUserDataIndex);
+			ComponentData->StreamedResourceIndex.AddUnique(AssetUserDataIndex);
 		}
 
-		const FName ComponentName = Context->ComponentNames[Component->Id];
+		const FName ComponentName = Context->ComponentNames[Component->Id.GetValue()];
 
 		if (bUpdateMeshes)
 		{
@@ -4910,11 +5045,10 @@ UE::Tasks::FTask LoadAdditionalAssetsAndData(const TSharedRef<FUpdateContextPriv
 										{
 											InstancePrivate->ClothingPhysicsAssets[ClothingAssetIndex] = PhysicsAsset.Get();
 										}
-											ComponentData.ClothingPhysicsAssetsToStream.Emplace(ClothingAssetIndex, PhysicsAssetIndex);
 									}
 									else
 									{
-										ComponentData.ClothingPhysicsAssetsToStream.Emplace(ClothingAssetIndex, PhysicsAssetIndex);
+										ComponentData->ClothingPhysicsAssetsToStream.Emplace(ClothingAssetIndex, PhysicsAssetIndex);
 										AssetsToStream.AddUnique(PhysicsAsset.ToSoftObjectPath());
 									}
 								}
@@ -4951,7 +5085,7 @@ UE::Tasks::FTask LoadAdditionalAssetsAndData(const TSharedRef<FUpdateContextPriv
 #if !WITH_EDITOR
 					Context->StreamedResourceIndex.AddUnique(ResourceIndex); // Used to hold/release streamed resources in non-editor builds.
 #endif
-					ComponentData.StreamedResourceIndex.AddUnique(ResourceIndex);
+					ComponentData->StreamedResourceIndex.AddUnique(ResourceIndex);
 				}
 			}
 
@@ -4967,7 +5101,7 @@ UE::Tasks::FTask LoadAdditionalAssetsAndData(const TSharedRef<FUpdateContextPriv
 							// No need to load since the SKM is a StrongObjectPtr
 							if (AssetUserData)
 							{
-								ComponentData.AssetUserDataArray.Add(AssetUserData);
+								ComponentData->AssetUserDataArray.Add(AssetUserData);
 							}
 						}
 					}
@@ -4986,11 +5120,11 @@ UE::Tasks::FTask LoadAdditionalAssetsAndData(const TSharedRef<FUpdateContextPriv
 					{
 						if (UPhysicsAsset* LoadedPhysicsAssetPtr = PhysicsAsset.Get())
 						{
-							ComponentData.PhysicsAssets.PhysicsAssetsToMerge.Add(LoadedPhysicsAssetPtr);
+							ComponentData->PhysicsAssets.PhysicsAssetsToMerge.Add(LoadedPhysicsAssetPtr);
 						}
 						else
 						{
-							ComponentData.PhysicsAssets.PhysicsAssetToLoad.Add(AssetIndex);
+							ComponentData->PhysicsAssets.PhysicsAssetToLoad.Add(AssetIndex);
 							AssetsToStream.AddUnique(PhysicsAsset.ToSoftObjectPath());
 						}
 					}
@@ -5019,11 +5153,11 @@ UE::Tasks::FTask LoadAdditionalAssetsAndData(const TSharedRef<FUpdateContextPriv
 
 						if (!AnimBPAsset.IsNull())
 						{
-							const TSoftClassPtr<UAnimInstance>* FoundAnimBpSlot = ComponentData.AnimSlotToBP.Find(SlotIndex);
+							const TSoftClassPtr<UAnimInstance>* FoundAnimBpSlot = ComponentData->AnimSlotToBP.Find(SlotIndex);
 							bool bIsSameAnimBp = FoundAnimBpSlot && AnimBPAsset == *FoundAnimBpSlot;
 							if (!FoundAnimBpSlot)
 							{
-								ComponentData.AnimSlotToBP.Add(SlotIndex, AnimBPAsset);
+								ComponentData->AnimSlotToBP.Add(SlotIndex, AnimBPAsset);
 
 								if (AnimBPAsset.Get())
 								{
@@ -5056,7 +5190,7 @@ UE::Tasks::FTask LoadAdditionalAssetsAndData(const TSharedRef<FUpdateContextPriv
 #if WITH_EDITORONLY_DATA
 				else if (Tag.RemoveFromStart("__MeshPath:"))
 				{
-					ComponentData.MeshPartPaths.Add(Tag);
+					ComponentData->MeshPartPaths.Add(Tag);
 				}
 #endif
 			}
@@ -5066,13 +5200,13 @@ UE::Tasks::FTask LoadAdditionalAssetsAndData(const TSharedRef<FUpdateContextPriv
 			{
 				const int32 ExternalId = MutableMesh->AdditionalPhysicsBodies[I]->CustomId;
 				
-				ComponentData.PhysicsAssets.AdditionalPhysicsAssetsToLoad.Add(ExternalId);
+				ComponentData->PhysicsAssets.AdditionalPhysicsAssetsToLoad.Add(ExternalId);
 				AssetsToStream.Add(ModelResources.AnimBpOverridePhysiscAssetsInfo[ExternalId].SourceAsset.ToSoftObjectPath());
 			}
 		}
 
 
-		for (int32 ResourceIndex : ComponentData.StreamedResourceIndex)
+		for (int32 ResourceIndex : ComponentData->StreamedResourceIndex)
 		{
 #if WITH_EDITOR
 			if (!ModelResources.StreamedResourceDataEditor.IsValidIndex(ResourceIndex))
@@ -5854,23 +5988,23 @@ void UCustomizableInstancePrivate::BuildResources(const TSharedRef<FUpdateContex
 	Materials.Empty();
 	Textures.Empty();
 
-	const int32 NumComponents = Context->InstanceUpdateData.Components.Num();
-	
 	TArray<bool> RecreateRenderStateOnInstanceComponent;
-	RecreateRenderStateOnInstanceComponent.Init(false, NumComponents);
+	RecreateRenderStateOnInstanceComponent.Init(false, Context->NumInstanceComponents);
 
 	TArray<bool> NotifyUpdateOnInstanceComponent;
-	NotifyUpdateOnInstanceComponent.Init(false, NumComponents);
+	NotifyUpdateOnInstanceComponent.Init(false, Context->NumInstanceComponents);
 
-	for (int32 ComponentIndex = 0; ComponentIndex < NumComponents; ++ComponentIndex)
+	for (int32 InstanceComponentIndex = 0; InstanceComponentIndex < Context->NumInstanceComponents; ++InstanceComponentIndex)
 	{
-		const TSharedRef<const FInstanceUpdateData::FComponent>& Component = Context->InstanceUpdateData.Components[ComponentIndex];
+		const TSharedRef<const FInstanceUpdateData::FComponent>& Component = Context->InstanceUpdateData.Components[InstanceComponentIndex];
+
+		const FCustomizableObjectComponentIndex ObjectComponentIndex = Component->Id;
 
-		if (!ModelResources.ComponentNamesPerObjectComponent.IsValidIndex(Component->Id))
+		if (!ModelResources.ComponentNamesPerObjectComponent.IsValidIndex(ObjectComponentIndex.GetValue()))
 		{
 			continue;
 		}
-		const FName& ComponentName = ModelResources.ComponentNamesPerObjectComponent[Component->Id];
+		const FName& ComponentName = ModelResources.ComponentNamesPerObjectComponent[ObjectComponentIndex.GetValue()];
 		
 		TObjectPtr<USkeletalMesh> SkeletalMesh = SkeletalMeshes.Contains(ComponentName) ? SkeletalMeshes[ComponentName] : nullptr;
 		if (!SkeletalMesh)
@@ -5878,7 +6012,7 @@ void UCustomizableInstancePrivate::BuildResources(const TSharedRef<FUpdateContex
 			continue;
 		}
 
-		const bool bUVChanged = Context->MeshChangedPerInstanceComponent[ComponentIndex];
+		const bool bUVChanged = Context->MeshChangedPerInstanceComponent[InstanceComponentIndex];
 		
 		// If the mesh is not transient, it means it's pass-through so it should use material overrides and not be modified in any way
 		const bool bIsTransientMesh = SkeletalMesh->HasAllFlags(EObjectFlags::RF_Transient);
@@ -5888,42 +6022,45 @@ void UCustomizableInstancePrivate::BuildResources(const TSharedRef<FUpdateContex
 
 		UMaterialInterface* OverlayMaterial = nullptr;
 
-		FCustomizableInstanceComponentData& ComponentData = GetComponentData(Component->Id);
-		ComponentData.OverrideMaterials.Reset();
-		ComponentData.OverlayMaterial = nullptr;
-			
-		if (Component->OverlayMaterialIndex != INDEX_NONE)
+		FCustomizableInstanceComponentData* ComponentData = GetComponentData(ObjectComponentIndex);
+		if (ComponentData)
 		{
-			FBuildMaterialOut Out = BuildMaterial(Context, Component->OverlayMaterialIndex, !bUVChanged);
-			OverlayMaterial = Out.Material;
-			ComponentData.OverlayMaterial = Out.Material;
-				
-			NotifyUpdateOnInstanceComponent[ComponentIndex] = Out.bNotifyPrimitiveUpdated;
-		}
+			ComponentData->OverrideMaterials.Reset();
+			ComponentData->OverlayMaterial = nullptr;
 			
-		for (int32 MaterialIndex = 0; MaterialIndex < Component->OverrideMaterialCount; ++MaterialIndex)
-		{
-			const FInstanceUpdateData::FOverrideMaterial& OverrideMaterial = Context->InstanceUpdateData.OverrideMaterials[Component->FirstOverrideMaterial + MaterialIndex];
-				
-			const int32 SlotIndex = SkeletalMesh->GetMaterials().IndexOfByPredicate([&](const FSkeletalMaterial& Material)
+			if (Component->OverlayMaterialIndex != INDEX_NONE)
 			{
-				return Material.MaterialSlotName == OverrideMaterial.SlotName;
-			});
+				FBuildMaterialOut Out = BuildMaterial(Context, Component->OverlayMaterialIndex, !bUVChanged);
+				OverlayMaterial = Out.Material;
+				ComponentData->OverlayMaterial = Out.Material;
 				
-			if (SlotIndex != INDEX_NONE)
+				NotifyUpdateOnInstanceComponent[InstanceComponentIndex] = Out.bNotifyPrimitiveUpdated;
+			}
+			
+			for (int32 MaterialIndex = 0; MaterialIndex < Component->OverrideMaterialCount; ++MaterialIndex)
 			{
-				FBuildMaterialOut Out = BuildMaterial(Context, OverrideMaterial.MaterialIndex, !bUVChanged);
+				const FInstanceUpdateData::FOverrideMaterial& OverrideMaterial = Context->InstanceUpdateData.OverrideMaterials[Component->FirstOverrideMaterial + MaterialIndex];
+				
+				const int32 SlotIndex = SkeletalMesh->GetMaterials().IndexOfByPredicate([&](const FSkeletalMaterial& Material)
+				{
+					return Material.MaterialSlotName == OverrideMaterial.SlotName;
+				});
+				
+				if (SlotIndex != INDEX_NONE)
+				{
+					FBuildMaterialOut Out = BuildMaterial(Context, OverrideMaterial.MaterialIndex, !bUVChanged);
 					
-				ComponentData.OverrideMaterials.SetNum(SlotIndex + 1);
-				ComponentData.OverrideMaterials[SlotIndex] = Out.Material;
+					ComponentData->OverrideMaterials.SetNum(SlotIndex + 1);
+					ComponentData->OverrideMaterials[SlotIndex] = Out.Material;
 					
-				NotifyUpdateOnInstanceComponent[ComponentIndex] = Out.bNotifyPrimitiveUpdated;
+					NotifyUpdateOnInstanceComponent[InstanceComponentIndex] = Out.bNotifyPrimitiveUpdated;
+				}
 			}
 		}
 
 		if (!bUseOverrideMaterialsOnly)
 		{
-			RecreateRenderStateOnInstanceComponent[ComponentIndex] |= SkeletalMesh->GetOverlayMaterial() != OverlayMaterial;
+			RecreateRenderStateOnInstanceComponent[InstanceComponentIndex] |= SkeletalMesh->GetOverlayMaterial() != OverlayMaterial;
 			SkeletalMesh->SetOverlayMaterial(OverlayMaterial);
 		}
 
@@ -5949,9 +6086,9 @@ void UCustomizableInstancePrivate::BuildResources(const TSharedRef<FUpdateContex
 
 			// Pass-through components will not have a reference mesh.
 			const FMutableRefSkeletalMeshData* RefSkeletalMeshData = nullptr;
-			if (ModelResources.ReferenceSkeletalMeshesData.IsValidIndex(Component->Id))
+			if (ModelResources.ReferenceSkeletalMeshesData.IsValidIndex(ObjectComponentIndex.GetValue()))
 			{
-				RefSkeletalMeshData = &ModelResources.ReferenceSkeletalMeshesData[Component->Id];
+				RefSkeletalMeshData = &ModelResources.ReferenceSkeletalMeshesData[ObjectComponentIndex.GetValue()];
 			}
 
 			for (int32 SurfaceIndex = 0; SurfaceIndex < LOD->SurfaceCount; ++SurfaceIndex)
@@ -6039,16 +6176,19 @@ void UCustomizableInstancePrivate::BuildResources(const TSharedRef<FUpdateContex
 					}
 				}
 
-				FBuildMaterialOut Out = BuildMaterial(Context, Surface->MaterialIndex, !bUVChanged);
-					
-				if (!Out.Material)
+				if (ComponentData)
 				{
-					Out.Material = UMaterial::GetDefaultMaterial(MD_Surface);
-				}
+					FBuildMaterialOut Out = BuildMaterial(Context, Surface->MaterialIndex, !bUVChanged);
+					
+					if (!Out.Material)
+					{
+						Out.Material = UMaterial::GetDefaultMaterial(MD_Surface);
+					}
 					
-				MaterialSlot.MaterialInterface = Out.Material;
-				ComponentData.OverrideMaterials.Add(Out.Material);
-				NotifyUpdateOnInstanceComponent[ComponentIndex] = Out.bNotifyPrimitiveUpdated;
+					MaterialSlot.MaterialInterface = Out.Material;
+					ComponentData->OverrideMaterials.Add(Out.Material);
+					NotifyUpdateOnInstanceComponent[InstanceComponentIndex] = Out.bNotifyPrimitiveUpdated;
+				}
 			}
 		}
 
@@ -6079,12 +6219,12 @@ void UCustomizableInstancePrivate::BuildResources(const TSharedRef<FUpdateContex
 			}
 
 			// Force recreate render state after replacing the materials to avoid a crash in the render pipeline if the old materials are GCed while in use.
-			RecreateRenderStateOnInstanceComponent[ComponentIndex] |= SkeletalMesh->GetResourceForRendering()->IsInitialized() && SkeletalMesh->GetMaterials() != SkeletalMaterials;
+			RecreateRenderStateOnInstanceComponent[InstanceComponentIndex] |= SkeletalMesh->GetResourceForRendering()->IsInitialized() && SkeletalMesh->GetMaterials() != SkeletalMaterials;
 
 			SkeletalMesh->SetMaterials(SkeletalMaterials);
 
 #if WITH_EDITOR
-			if (GEditor && RecreateRenderStateOnInstanceComponent[ComponentIndex])
+			if (GEditor && RecreateRenderStateOnInstanceComponent[InstanceComponentIndex])
 			{
 				// Close all open editors for this mesh to invalidate viewports.
 				GEditor->GetEditorSubsystem<UAssetEditorSubsystem>()->CloseAllEditorsForAsset(SkeletalMesh);
@@ -6119,24 +6259,20 @@ void UCustomizableInstancePrivate::BuildResources(const TSharedRef<FUpdateContex
 #endif
 
 			const FName& ComponentName = CustomizableObjectInstanceUsage->GetComponentName();
-			FCustomizableInstanceComponentData* ComponentData = GetComponentData(ComponentName);
-			if (!ComponentData)
-			{
-				continue;
-			}
+			const int32 ObjectComponentIndex = ModelResources.ComponentNamesPerObjectComponent.IndexOfByKey(ComponentName);
 			
-			int32 ComponentIndex = -1;
+			int32 InstanceComponentIndex = -1;
 			for (int32 CurrentInstanceIndex = 0; CurrentInstanceIndex<Context->InstanceUpdateData.Components.Num(); ++CurrentInstanceIndex)
 			{
-				if (Context->InstanceUpdateData.Components[CurrentInstanceIndex]->Id == ComponentData->ComponentId)
+				if (Context->InstanceUpdateData.Components[CurrentInstanceIndex]->Id.GetValue() == ObjectComponentIndex)
 				{
-					ComponentIndex = CurrentInstanceIndex;
+					InstanceComponentIndex = CurrentInstanceIndex;
 					break;
 				}
 			}
 
-			bool bDoRecreateRenderStateOnComponent = RecreateRenderStateOnInstanceComponent.IsValidIndex(ComponentIndex) && RecreateRenderStateOnInstanceComponent[ComponentIndex];
-			bool bDoNotifyUpdateOnComponent = NotifyUpdateOnInstanceComponent.IsValidIndex(ComponentIndex) && NotifyUpdateOnInstanceComponent[ComponentIndex];
+			bool bDoRecreateRenderStateOnComponent = RecreateRenderStateOnInstanceComponent.IsValidIndex(InstanceComponentIndex) && RecreateRenderStateOnInstanceComponent[InstanceComponentIndex];
+			bool bDoNotifyUpdateOnComponent = NotifyUpdateOnInstanceComponent.IsValidIndex(InstanceComponentIndex) && NotifyUpdateOnInstanceComponent[InstanceComponentIndex];
 
 			if (!bDoRecreateRenderStateOnComponent && !bDoNotifyUpdateOnComponent)
 			{
@@ -6279,21 +6415,19 @@ void UCustomizableObjectInstance::SetRequestedLODs(const TMap<FName, uint8>& InM
 	
 	const TMap<FName, uint8>& FirstRequestedLOD = GetPrivate()->CommittedDescriptorHash.FirstRequestedLOD;
 	
-	for (UE::Mutable::Private::FComponentId ComponentId = 0; ComponentId < ModelResources->ComponentNamesPerObjectComponent.Num(); ++ComponentId)
+	for (const FName& ComponentName : ModelResources->ComponentNamesPerObjectComponent)
 	{	
-		const FName& ComponentName = ModelResources->ComponentNamesPerObjectComponent[ComponentId];
-		
 		uint8& InMinLOD = MutableUpdateCandidate.MinLOD.FindOrAdd(ComponentName);
 		if (const uint8* Result = InMinLODs.Find(ComponentName))
 		{
 			InMinLOD = *Result;
 		}
 
-		const uint8 MinLODIdx = CustomizableObject->GetPrivate()->GetMinLODIndex(ComponentId);
+		const uint8 MinLODIdx = CustomizableObject->GetPrivate()->GetMinLODIndex(ComponentName);
 		MutableUpdateCandidate.QualitySettingMinLODs.Add(ComponentName, MinLODIdx);
 
 		int32 MaxLODIdx = 0;
-		if (const uint8* Found = ModelResources->NumLODsAvailable.Find(ComponentId))
+		if (const uint8* Found = ModelResources->NumLODsAvailable.Find(ComponentName))
 		{
 			MaxLODIdx = *Found - 1;
 		}
@@ -6319,7 +6453,7 @@ void UCustomizableObjectInstance::SetRequestedLODs(const TMap<FName, uint8>& InM
 			PrivateData->SetCOInstanceFlags(InMinLOD > CurrentMinLOD ? PendingLODsDowngrade : ECONone);
 
 			uint8 FirstNonStreamedLODIndex = 0;
-			if (const uint8* Found = ModelResources->NumLODsToStream.Find(ComponentId))
+			if (const uint8* Found = ModelResources->NumLODsToStream.Find(ComponentName))
 			{
 				FirstNonStreamedLODIndex = *Found;
 			}
@@ -6377,6 +6511,12 @@ void UCustomizableObjectInstance::Bake(const FBakingConfiguration& InBakingConfi
 #endif
 
 
+USkeletalMesh* UCustomizableObjectInstance::GetSkeletalMesh(int32 ObjectComponentIndex) const
+{
+	return GetComponentMeshSkeletalMesh(FName(FString::FromInt(ObjectComponentIndex)));
+}
+
+
 USkeletalMesh* UCustomizableObjectInstance::GetComponentMeshSkeletalMesh(const FName& ComponentName) const
 {
 	TObjectPtr<USkeletalMesh>* Result = GetPrivate()->SkeletalMeshes.Find(ComponentName);
@@ -6424,6 +6564,7 @@ TArray<FName> UCustomizableObjectInstance::GetComponentNames() const
 TSubclassOf<UAnimInstance> UCustomizableObjectInstance::GetAnimBP(FName ComponentName, const FName& SlotName) const
 {
 	FCustomizableInstanceComponentData* ComponentData =	GetPrivate()->GetComponentData(ComponentName);
+	
 	if (!ComponentData)
 	{
 		FString ErrorMsg = FString::Printf(TEXT("Tried to access an invalid component index [%s] in a Mutable Instance."), *ComponentName.ToString());
@@ -6525,6 +6666,29 @@ void UCustomizableObjectInstance::ForEachComponentAnimInstance(FName ComponentNa
 	UE::Mutable::Private::InternalForEachAnimInstance<>(GetPrivate(), ComponentName, Delegate);
 }
 
+// Deprecated
+void UCustomizableObjectInstance::ForEachAnimInstance(int32 ObjectComponentIndex, FEachComponentAnimInstanceClassDelegate Delegate) const
+{
+	UCustomizableObject* CO = GetCustomizableObject();
+	if (CO)
+	{
+		FName ComponentName = CO->GetPrivate()->GetComponentName(FCustomizableObjectComponentIndex(ObjectComponentIndex));
+		UE::Mutable::Private::InternalForEachAnimInstance<>(GetPrivate(), ComponentName, Delegate);
+	}
+}
+
+
+// Deprecated
+void UCustomizableObjectInstance::ForEachAnimInstance(int32 ObjectComponentIndex, FEachComponentAnimInstanceClassNativeDelegate Delegate) const
+{
+	UCustomizableObject* CO = GetCustomizableObject();
+	if (CO)
+	{
+		FName ComponentName = CO->GetPrivate()->GetComponentName(FCustomizableObjectComponentIndex(ObjectComponentIndex));
+		UE::Mutable::Private::InternalForEachAnimInstance<>(GetPrivate(), ComponentName, Delegate);
+	}
+}
+
 
 bool UCustomizableObjectInstance::AnimInstanceNeedsFixup(TSubclassOf<UAnimInstance> AnimInstanceClass) const
 {
@@ -6629,23 +6793,26 @@ FInstancedStruct UCustomizableObjectInstance::GetExtensionInstanceData(const UCu
 }
 
 
-TSet<UAssetUserData*> UCustomizableObjectInstance::GetMergedAssetUserData(const FName& ComponentName) const
+TSet<UAssetUserData*> UCustomizableObjectInstance::GetMergedAssetUserData(int32 ComponentIndex) const
 {
-	FCustomizableInstanceComponentData* ComponentData = GetPrivate()->GetComponentData(ComponentName);
-	if (!ComponentData)
+	UCustomizableInstancePrivate* PrivateInstanceData = GetPrivate();
+
+	if (PrivateInstanceData && PrivateInstanceData->ComponentsData.IsValidIndex(ComponentIndex))
 	{
-		return {};
-	}
-	
-	TSet<UAssetUserData*> Set;
+		TSet<UAssetUserData*> Set;
 		
-	for (const TObjectPtr<UAssetUserData>& Elem : ComponentData->AssetUserDataArray)
-	{
 		// Have to convert to UAssetUserData* because BP functions don't support TObjectPtr
-		Set.Add(Elem);
-	}
+		for (const TObjectPtr<UAssetUserData>& Elem : PrivateInstanceData->ComponentsData[ComponentIndex].AssetUserDataArray)
+		{
+			Set.Add(Elem);
+		}
 
-	return Set;
+		return Set;
+	}
+	else
+	{
+		return TSet<UAssetUserData*>();
+	}
 }
 
 
diff --git a/Engine/Plugins/Mutable/Source/CustomizableObject/Private/MuCO/CustomizableObjectInstanceUsage.cpp b/Engine/Plugins/Mutable/Source/CustomizableObject/Private/MuCO/CustomizableObjectInstanceUsage.cpp
--- a/Engine/Plugins/Mutable/Source/CustomizableObject/Private/MuCO/CustomizableObjectInstanceUsage.cpp
+++ b/Engine/Plugins/Mutable/Source/CustomizableObject/Private/MuCO/CustomizableObjectInstanceUsage.cpp
@@ -78,6 +78,19 @@ UCustomizableObjectInstance* UCustomizableObjectInstanceUsage::GetCustomizableOb
 }
 
 
+void UCustomizableObjectInstanceUsage::SetComponentIndex(int32 ComponentIndex)
+{
+	if (UCustomizableSkeletalComponent* CustomizableSkeletalComponent = GetPrivate()->GetCustomizableSkeletalComponent())
+	{
+		CustomizableSkeletalComponent->ComponentIndex = ComponentIndex;
+	}
+	else
+	{
+		UsedComponentIndex = ComponentIndex;
+	}
+}
+
+
 int32 UCustomizableObjectInstanceUsage::GetComponentIndex() const
 {
 	if (UCustomizableSkeletalComponent* CustomizableSkeletalComponent = GetPrivate()->GetCustomizableSkeletalComponent())
diff --git a/Engine/Plugins/Mutable/Source/CustomizableObject/Private/MuCO/CustomizableObjectSystem.cpp b/Engine/Plugins/Mutable/Source/CustomizableObject/Private/MuCO/CustomizableObjectSystem.cpp
--- a/Engine/Plugins/Mutable/Source/CustomizableObject/Private/MuCO/CustomizableObjectSystem.cpp
+++ b/Engine/Plugins/Mutable/Source/CustomizableObject/Private/MuCO/CustomizableObjectSystem.cpp
@@ -252,19 +252,14 @@ FUpdateContextPrivate::FUpdateContextPrivate(UCustomizableObjectInstance& InInst
 	{
 		CapturedDescriptor.bStreamingEnabled = IsStreamingEnabled(*InObject, Descriptor.State);
 
+		FirstLODAvailable = ModelResources->FirstLODAvailable;
+		FirstResidentLOD = ModelResources->NumLODsToStream;
 		ComponentNames = ModelResources->ComponentNamesPerObjectComponent;
-		
-		for (UE::Mutable::Private::FComponentId ComponentId = 0; ComponentId < ModelResources->ComponentNamesPerObjectComponent.Num(); ++ComponentId)
-		{
-			const FName& ComponentName = ModelResources->ComponentNamesPerObjectComponent[ComponentId];
-			
-			FirstLODAvailable.Add(ComponentName, ModelResources->FirstLODAvailable[ComponentId]);
-			FirstResidentLOD.Add(ComponentName, ModelResources->NumLODsToStream[ComponentId]);
-		}
 	}
 
 	CapturedDescriptorHash = FDescriptorHash(Descriptor);
-	
+	NumObjectComponents = InObject->GetComponentCount();
+
 	UCustomizableObjectSystem* System = UCustomizableObjectSystem::GetInstanceChecked();
 	UCustomizableObjectSystemPrivate* SystemPrivate = System->GetPrivate();
 
@@ -566,7 +561,32 @@ const FDescriptorHash& FUpdateContextPrivate::GetCapturedDescriptorHash() const
 
 const FCustomizableObjectInstanceDescriptor&& FUpdateContextPrivate::MoveCommittedDescriptor()
 {
-	return MoveTemp(CapturedDescriptor);
+	return MoveTemp(CapturedDescriptor);	
+}
+
+
+FCustomizableObjectComponentIndex FUpdateContextPrivate::GetObjectComponentIndex(FCustomizableObjectInstanceComponentIndex InstanceComponentIndex)
+{
+	FCustomizableObjectComponentIndex ObjectComponentIndex;
+	if (InstanceUpdateData.Components.IsValidIndex(InstanceComponentIndex.GetValue()))
+	{
+		ObjectComponentIndex = InstanceUpdateData.Components[InstanceComponentIndex.GetValue()]->Id;
+	}
+	else
+	{
+		ObjectComponentIndex.Invalidate();
+	}
+	return ObjectComponentIndex;
+}
+
+
+const FInstanceUpdateData::FComponent* FUpdateContextPrivate::GetComponentUpdateData(FCustomizableObjectInstanceComponentIndex InstanceComponentIndex)
+{
+	if (InstanceUpdateData.Components.IsValidIndex(InstanceComponentIndex.GetValue()))
+	{
+		return &InstanceUpdateData.Components[InstanceComponentIndex.GetValue()].Get();
+	}
+	return nullptr;
 }
 
 
@@ -582,9 +602,13 @@ const TArray<FComponentDescriptor>& FUpdateContextPrivate::GetComponentDescripto
 }
 
 
-FComponentDescriptor* FUpdateContextPrivate::GetComponentDescriptor(UE::Mutable::Private::FComponentId ComponentId)
+FComponentDescriptor* FUpdateContextPrivate::GetComponentDescriptor(FCustomizableObjectComponentIndex Index)
 {
-	return ComponentDescriptors.IsValidIndex(ComponentId) ? &ComponentDescriptors[ComponentId] : nullptr;
+	if (ComponentDescriptors.IsValidIndex(Index.GetValue()))
+	{
+		return &ComponentDescriptors[Index.GetValue()];
+	}
+	return nullptr;
 }
 
 
@@ -1312,13 +1336,14 @@ void UpdateSkeletalMesh(const TSharedRef<FUpdateContextPrivate>& Context)
 
 		TArray<TObjectPtr<UMaterialInterface>> NewOverridenMaterials;
 
-		FCustomizableInstanceComponentData* ComponentData = CustomizableObjectInstance->GetPrivate()->GetComponentData(Usage->GetComponentName());
-		
 		const bool bIsTransientMesh = SkeletalMesh ? SkeletalMesh->HasAllFlags(EObjectFlags::RF_Transient) : false;
 		const bool bUseOverrideMaterials = !bIsTransientMesh || Context->bUseMeshCache;
-		if (bUseOverrideMaterials && ComponentData)
+		if (bUseOverrideMaterials)
 		{
-			NewOverridenMaterials = ComponentData->OverrideMaterials;
+			if (FCustomizableInstanceComponentData* ComponentData = CustomizableObjectInstance->GetPrivate()->GetComponentData(Usage->GetComponentName()))
+			{
+				NewOverridenMaterials = ComponentData->OverrideMaterials;
+			}
 		}
 
 		if (Parent->OverrideMaterials != NewOverridenMaterials)
@@ -1340,14 +1365,18 @@ void UpdateSkeletalMesh(const TSharedRef<FUpdateContextPrivate>& Context)
 
 
 		UMaterialInterface* OverlayMaterial = nullptr;
-		if (bUseOverrideMaterials && ComponentData)
+		if (bUseOverrideMaterials)
 		{
-			OverlayMaterial = ComponentData->OverlayMaterial;
+			if (FCustomizableInstanceComponentData* ComponentData = CustomizableObjectInstance->GetPrivate()->GetComponentData(Usage->GetComponentName()))
+			{
+				OverlayMaterial = ComponentData->OverlayMaterial;
+			}
 		}
 
 		bAttachedParentUpdated |= Parent->GetOverlayMaterial() != OverlayMaterial;
 		Parent->SetOverlayMaterial(OverlayMaterial);
-		
+
+
 		if (CustomizableObjectInstancePrivateData->HasCOInstanceFlags(ReplacePhysicsAssets) &&
 			SkeletalMesh &&
 			Parent->GetWorld())
@@ -1512,7 +1541,7 @@ void UCustomizableObjectSystemPrivate::EnqueueUpdateSkeletalMesh(const TSharedRe
 	{
 		// Mutable is disabled. Set the reference SkeletalMesh and finish the update with success to avoid breaking too many things.
 		Context->UpdateResult = EUpdateResult::Success;
-		Instance->GetPrivate()->ForceSetReferenceSkeletalMesh();
+		Instance->GetPrivate()->SetReferenceSkeletalMesh();
 		FinishUpdateGlobal(Context);
 		return;
 	}
@@ -1854,13 +1883,18 @@ namespace impl
 	
 	void FixLODs(const TSharedRef<FUpdateContextPrivate>& Operation)
 	{
+		if (!Operation->NumObjectComponents)
+		{
+			return;
+		}
+
 		TMap<FName, uint8> MinLODs = Operation->GetCapturedDescriptor().MinLOD;
 		TMap<FName, uint8> RequestedLODs = Operation->GetFirstRequestedLOD();
 		
-		for (int32 ComponentIndex = 0; ComponentIndex < Operation->MutableInstance->GetComponentCount(); ++ComponentIndex)
+		for (int32 InstanceComponentIndex = 0; InstanceComponentIndex < Operation->NumInstanceComponents; ++InstanceComponentIndex)
 		{
-			const UE::Mutable::Private::FComponentId ComponentId = Operation->MutableInstance->GetComponentId(ComponentIndex);
-			const FName& ComponentName = Operation->ComponentNames[ComponentId];
+			const int32 ObjectComponentIndex = Operation->MutableInstance->GetComponentId(InstanceComponentIndex);
+			const FName& ComponentName = Operation->ComponentNames[ObjectComponentIndex];
 
 			uint8& MinLOD = MinLODs.FindOrAdd(ComponentName, 0);
 			uint8& RequestedLOD = RequestedLODs.FindOrAdd(ComponentName, 0);
@@ -1868,7 +1902,7 @@ namespace impl
 			uint8& FirstResidentLOD = Operation->FirstResidentLOD.FindOrAdd(ComponentName, 0);
 			const uint8& FirstLODAvailable = Operation->FirstLODAvailable.FindOrAdd(ComponentName, 0);
 
-			NumLODsAvailable = Operation->MutableInstance->GetLODCount(ComponentIndex);
+			NumLODsAvailable = Operation->MutableInstance->GetLODCount(InstanceComponentIndex);
 			
 			if (Operation->bStreamMeshLODs)
 			{
@@ -1894,12 +1928,12 @@ namespace impl
 	{
 		MUTABLE_CPUPROFILER_SCOPE(Subtask_Mutable_PrepareSkeletonData);
 
-		const int32 NumComponents = OperationData->InstanceUpdateData.Components.Num();
-		OperationData->InstanceUpdateData.SkeletonsPerInstanceComponent.SetNum(NumComponents);
+		const int32 NumInstanceComponents = OperationData->InstanceUpdateData.Components.Num();
+		OperationData->InstanceUpdateData.SkeletonsPerInstanceComponent.SetNum(NumInstanceComponents);
 
-		for (int32 ComponentIndex = 0; ComponentIndex< NumComponents; ++ComponentIndex)
+		for (int32 InstanceComponentIndex = 0; InstanceComponentIndex< NumInstanceComponents; ++InstanceComponentIndex)
 		{
-			const TSharedRef<FInstanceUpdateData::FComponent>& Component = OperationData->InstanceUpdateData.Components[ComponentIndex];
+			const TSharedRef<FInstanceUpdateData::FComponent>& Component = OperationData->InstanceUpdateData.Components[InstanceComponentIndex];
 
 			for (int32 LODIndex = 0; LODIndex < Component->LODCount; ++LODIndex)
 			{
@@ -1911,7 +1945,7 @@ namespace impl
 					continue;
 				}
 
-				FInstanceUpdateData::FSkeletonData& SkeletonData = OperationData->InstanceUpdateData.SkeletonsPerInstanceComponent[ComponentIndex];
+				FInstanceUpdateData::FSkeletonData& SkeletonData = OperationData->InstanceUpdateData.SkeletonsPerInstanceComponent[InstanceComponentIndex];
 
 				// Add SkeletonIds 
 				const int32 SkeletonIDsCount = Mesh->GetSkeletonIDsCount();
@@ -1976,18 +2010,17 @@ namespace impl
 		
 		const TSharedPtr<FModelStreamableBulkData>& ModelStreamableBulkData = Object->GetPrivate()->GetModelStreamableBulkData();
 
-		const int32 NumComponents = OperationData->InstanceUpdateData.Components.Num();
+		const int32 NumInstanceComponents = OperationData->InstanceUpdateData.Components.Num();
 		check(OperationData->InstanceUpdateData.RealTimeMorphTargets.IsEmpty());
-		OperationData->InstanceUpdateData.RealTimeMorphTargets.Reserve(NumComponents);
-		
-		for (int32 ComponentIndex = 0; ComponentIndex< NumComponents; ++ComponentIndex)
+		OperationData->InstanceUpdateData.RealTimeMorphTargets.Reserve(NumInstanceComponents);
+		for ( int32 InstanceComponentIndex = 0; InstanceComponentIndex< NumInstanceComponents; ++InstanceComponentIndex)
 		{
-			TSharedRef<FInstanceUpdateData::FComponent>& Component = OperationData->InstanceUpdateData.Components[ComponentIndex];
-			check(Component->Id != INDEX_NONE);
+			TSharedRef<FInstanceUpdateData::FComponent>& Component = OperationData->InstanceUpdateData.Components[InstanceComponentIndex];
+			check(Component->Id.IsValid());
 
-			const FName& ComponentName = OperationData->ComponentNames[Component->Id];
+			const FName& ComponentName = OperationData->ComponentNames[Component->Id.GetValue()];
 
-			FSkeletalMeshMorphTargets& ComponentMorphTargetsData = OperationData->InstanceUpdateData.RealTimeMorphTargets.Add(Component->Id);
+			FSkeletalMeshMorphTargets& ComponentMorphTargetsData = OperationData->InstanceUpdateData.RealTimeMorphTargets.Add(ComponentName);
 			
 			ComponentMorphTargetsData.RealTimeMorphTargetNames.Empty();
 
@@ -2122,10 +2155,10 @@ namespace impl
 		
 		UE::Tasks::FTask LastTask = UE::Tasks::MakeCompletedTask<void>();
 		
-		for (int32 ComponentIndex = 0; ComponentIndex < OperationData->InstanceUpdateData.Components.Num(); ++ComponentIndex)
+		for (int32 InstanceComponentIndex = 0; InstanceComponentIndex < OperationData->NumInstanceComponents; ++InstanceComponentIndex)
 		{
-			TSharedRef<FInstanceUpdateData::FComponent>& Component = OperationData->InstanceUpdateData.Components[ComponentIndex];
-			const FName& ComponentName = ModelResources.ComponentNamesPerObjectComponent[Component->Id];
+			TSharedRef<FInstanceUpdateData::FComponent>& Component = OperationData->InstanceUpdateData.Components[InstanceComponentIndex];
+			const FName& ComponentName = ModelResources.ComponentNamesPerObjectComponent[Component->Id.GetValue()];
 
 			for (int32 MutableLODIndex = OperationData->FirstLODAvailable[ComponentName]; MutableLODIndex < Component->LODCount; ++MutableLODIndex)
 			{
@@ -2145,7 +2178,7 @@ namespace impl
 				{
 					const uint32 SurfaceId = LOD->Mesh->GetSurfaceId(MeshSurfaceIndex);
 
-					const int32 InstanceSurfaceIndex = MutableInstance->FindSurfaceById(ComponentIndex, MutableLODIndex, SurfaceId);
+					const int32 InstanceSurfaceIndex = MutableInstance->FindSurfaceById(InstanceComponentIndex, MutableLODIndex, SurfaceId);
 					check(LOD->Mesh->GetVertexCount() > 0 || InstanceSurfaceIndex >= 0);
 
 					if (InstanceSurfaceIndex < 0)
@@ -2168,9 +2201,9 @@ namespace impl
 					TSharedRef<FInstanceUpdateData::FSurface>& Surface = OperationData->InstanceUpdateData.Surfaces.Add_GetRef(MakeShared<FInstanceUpdateData::FSurface>());
 
 					Surface->SurfaceId = SurfaceId;
-					Surface->SurfaceMetadataId = MutableInstance->GetSurfaceCustomId(ComponentIndex, MutableLODIndex, InstanceSurfaceIndex);
+					Surface->SurfaceMetadataId = MutableInstance->GetSurfaceCustomId(InstanceComponentIndex, MutableLODIndex, InstanceSurfaceIndex);
 
-					if (UE::Mutable::Private::FMaterialId MaterialId = MutableInstance->GetMaterialId(ComponentIndex, MutableLODIndex, InstanceSurfaceIndex))
+					if (UE::Mutable::Private::FMaterialId MaterialId = MutableInstance->GetMaterialId(InstanceComponentIndex, MutableLODIndex, InstanceSurfaceIndex))
 					{
 						Surface->MaterialIndex = OperationData->InstanceUpdateData.Materials.Num();
 						TSharedRef<FInstanceUpdateData::FMaterial>& Material = OperationData->InstanceUpdateData.Materials.Add_GetRef(MakeShared<FInstanceUpdateData::FMaterial>());
@@ -2291,19 +2324,25 @@ namespace impl
 
 		UE::Tasks::FTask LastTask;
 		
-		const int32 NumComponents = OperationData->MutableInstance->GetComponentCount();
-		
-		OperationData->InstanceUpdateData.Components.Reserve(NumComponents);
-		for (int32 ComponentIndex = 0; ComponentIndex < NumComponents; ++ComponentIndex)
+		OperationData->InstanceUpdateData.Components.Reserve(OperationData->NumInstanceComponents);
+		for (int32 InstanceComponentIndex = 0; InstanceComponentIndex < OperationData->NumInstanceComponents; ++InstanceComponentIndex)
 		{
 			TSharedRef<FInstanceUpdateData::FComponent>& Component = OperationData->InstanceUpdateData.Components.Add_GetRef(MakeShared<FInstanceUpdateData::FComponent>());
 			Component->FirstLOD = OperationData->InstanceUpdateData.LODs.Num();
-			UE::Mutable::Private::FComponentId MutableComponentId = Instance->GetComponentId(ComponentIndex);
-			Component->Id = MutableComponentId;
+			uint64 MutableComponentId = Instance->GetComponentId(InstanceComponentIndex);
+			if (MutableComponentId < 65535)
+			{
+				Component->Id = FCustomizableObjectComponentIndex(MutableComponentId);
+			}
+			else
+			{
+				Component->Id.Invalidate();
+			}
 
-			const FName ComponentName = OperationData->ComponentNames[Component->Id];
+			FCustomizableObjectComponentIndex ObjectComponentIndex = Component->Id;
+			const FName ComponentName = OperationData->ComponentNames[ObjectComponentIndex.GetValue()];
 
-			Component->SkeletalMeshId = Instance->GetSkeletalMeshId(ComponentIndex);
+			Component->SkeletalMeshId = Instance->GetSkeletalMeshId(InstanceComponentIndex);
 
 			UE::Tasks::TTask<UE::Mutable::Private::TManagedPtr<const UE::Mutable::Private::FSkeletalMesh>> GetSkeletalMeshTask = OperationData->MutableSystem->GetSkeletalMesh(LastTask, OperationData->LiveInstance.ToSharedRef(), Component->SkeletalMeshId, UE::Mutable::Private::EMeshContentFlags::AllFlags);
 			UE::Tasks::AddNested(GetSkeletalMeshTask);
@@ -2316,22 +2355,22 @@ namespace impl
 			GetSkeletalMeshTask,
 			LowLevelTasks::ETaskPriority::Inherit));
 
-			if (UE::Mutable::Private::FMaterialId MaterialId = Instance->GetOverlayMaterialId(ComponentIndex))
+			if (UE::Mutable::Private::FMaterialId MaterialId = Instance->GetOverlayMaterialId(InstanceComponentIndex))
 			{
 				Component->OverlayMaterialIndex = OperationData->InstanceUpdateData.Materials.Num();
 				TSharedRef<FInstanceUpdateData::FMaterial> Material = OperationData->InstanceUpdateData.Materials.Add_GetRef(MakeShared<FInstanceUpdateData::FMaterial>());
 				Material->MaterialId = MaterialId;
 			}
 			
-			const int32 OverrideMaterialCount = Instance->GetOverrideMaterialCount(ComponentIndex);
+			const int32 OverrideMaterialCount = Instance->GetOverrideMaterialCount(InstanceComponentIndex);
 			Component->FirstOverrideMaterial = OperationData->InstanceUpdateData.OverrideMaterials.Num();
 			for (int32 MaterialIndex = 0; MaterialIndex < OverrideMaterialCount; ++MaterialIndex)
 			{
-				if (UE::Mutable::Private::FMaterialId MaterialId = Instance->GetOverrideMaterialId(ComponentIndex, MaterialIndex))
+				if (UE::Mutable::Private::FMaterialId MaterialId = Instance->GetOverrideMaterialId(InstanceComponentIndex, MaterialIndex))
 				{
 					FInstanceUpdateData::FOverrideMaterial& OverrideMaterial = OperationData->InstanceUpdateData.OverrideMaterials.AddDefaulted_GetRef();
 					++Component->OverrideMaterialCount;
-					OverrideMaterial.SlotName = Instance->GetOverrideMaterialSlotSlotName(ComponentIndex, MaterialIndex);
+					OverrideMaterial.SlotName = Instance->GetOverrideMaterialSlotSlotName(InstanceComponentIndex, MaterialIndex);
 					
 					OverrideMaterial.MaterialIndex = OperationData->InstanceUpdateData.Materials.Num();
 					TSharedRef<FInstanceUpdateData::FMaterial> Material = OperationData->InstanceUpdateData.Materials.Add_GetRef(MakeShared<FInstanceUpdateData::FMaterial>());
@@ -2365,7 +2404,7 @@ namespace impl
 
 				TSharedRef<FInstanceUpdateData::FLOD>& LOD = OperationData->InstanceUpdateData.LODs[Component->FirstLOD + LODIndex];
 				
-				LOD->MeshId = Instance->GetMeshId(ComponentIndex, LODIndex);
+				LOD->MeshId = Instance->GetMeshId(InstanceComponentIndex, LODIndex);
 				
 				UE::Mutable::Private::EMeshContentFlags MeshContentFilter = UE::Mutable::Private::EMeshContentFlags::AllFlags; 
 				
@@ -2827,34 +2866,9 @@ namespace impl
 		// Any external texture that may be needed for this update will be requested from Mutable Core's GetImage
 		// which will safely access the GlobalExternalImages map, and then just get the cached image or issue a disk read
 
-		const UModelResources& ModelResources = CustomizableObject->GetPrivate()->GetModelResourcesChecked();
-
-		const int32 NumComponents = OperationData->MutableInstance->GetComponentCount();
-		ObjectInstancePrivateData->ComponentsData.SetNum(NumComponents);
-		
-		TArray<FName> ComponentNames;
-		ComponentNames.Reserve(NumComponents);
+		// Copy data generated in the mutable thread over to the instance
+		ObjectInstancePrivateData->PrepareForUpdate();
 
-		for (int32 ComponentIndex = 0; ComponentIndex < NumComponents; ++ComponentIndex)
-		{
-			const UE::Mutable::Private::FComponentId ComponentId = OperationData->MutableInstance->GetComponentId(ComponentIndex);
-			const FName& ComponentName = ModelResources.ComponentNamesPerObjectComponent[ComponentId];
-			
-			if (ComponentNames.Contains(ComponentName)) // Components with the same name. Abort.
-			{
-				UE_LOG(LogMutable, Error, TEXT("Update error. Trying to generate multiple components with the same name: %s"), *ComponentName.ToString());
-
-				OperationData->UpdateResult = EUpdateResult::Error;
-				FinishUpdateGlobal(OperationData);
-				return;
-			}
-			
-			ComponentNames.Add(ComponentName);
-			
-			ObjectInstancePrivateData->ComponentsData[ComponentIndex].ComponentId = ComponentId;
-			ObjectInstancePrivateData->ComponentsData[ComponentIndex].ComponentName = ComponentName;
-		}
-		
 		UE::Tasks::FTask Mutable_GetImageDescriptors_Task = SystemPrivate->MutableTaskGraph.AddMutableThreadTask(
 		TEXT("Task_Mutable_GetImageDescriptors"),
 		[OperationData]()
@@ -3194,18 +3208,20 @@ namespace impl
 				Operation->SkeletalMeshIdRegistry,
 				Operation->GetCapturedDescriptor().GetState(),
 				UE::Mutable::Private::FSystem::AllLODs);
+			
+			Operation->NumInstanceComponents = Operation->MutableInstance->GetComponentCount();
 		}
 		
 		FixLODs(Operation);
 		
-		Operation->InitComponentDescriptors(Object->GetComponentCount());
+		Operation->InitComponentDescriptors(Operation->NumObjectComponents);
 
 		const TMap<FName, uint8>& NumLODsAvailable = Operation->NumLODsAvailable;
 
-		for (int32 ComponentIndex = 0; ComponentIndex < Operation->MutableInstance->GetComponentCount(); ++ComponentIndex)
+		for (int32 InstanceComponentIndex = 0; InstanceComponentIndex < Operation->NumInstanceComponents; ++InstanceComponentIndex)
 		{
-			UE::Mutable::Private::FComponentId ComponentId = Operation->MutableInstance->GetComponentId(ComponentIndex);
-			FComponentDescriptor* ComponentDescriptor = Operation->GetComponentDescriptor(ComponentId);
+			FCustomizableObjectComponentIndex ObjectComponentIndex(Operation->MutableInstance->GetComponentId(InstanceComponentIndex));
+			FComponentDescriptor* ComponentDescriptor = Operation->GetComponentDescriptor(ObjectComponentIndex);
 			check(ComponentDescriptor);
 
 			if (!ComponentDescriptor)
@@ -3215,7 +3231,7 @@ namespace impl
 
 			ComponentDescriptor->MeshDescriptors.Init({}, MAX_MESH_LOD_COUNT);
 
-			const FName& ComponentName = Operation->ComponentNames[ComponentId];
+			const FName& ComponentName = Operation->ComponentNames[ObjectComponentIndex.GetValue()];
 
 			const int32 FirstLOD = Operation->bStreamMeshLODs ?
 				Operation->FirstLODAvailable[ComponentName] :
@@ -3223,7 +3239,7 @@ namespace impl
 
 			for (int32 LODIndex = FirstLOD; LODIndex < NumLODsAvailable[ComponentName]; ++LODIndex)
 			{
-				ComponentDescriptor->MeshDescriptors[LODIndex] = Operation->MutableInstance->GetMeshId(ComponentIndex, LODIndex);
+				ComponentDescriptor->MeshDescriptors[LODIndex] = Operation->MutableInstance->GetMeshId(InstanceComponentIndex, LODIndex);
 			}
 		}
 
@@ -3755,12 +3771,20 @@ void UCustomizableObjectSystem::DiscardInstances()
 
 							if (USkeletalMeshComponent* Parent = Cast<USkeletalMeshComponent>(InstanceUsage->GetAttachParent()))
 							{
-								Parent->EmptyOverrideMaterials();
-
-								if (IsReplaceDiscardedWithReferenceMeshEnabled())
+								const FName& ComponentName = InstanceUsage->GetComponentName();
+								const int32 ObjectComponentIndex = ModelResources->ComponentNamesPerObjectComponent.IndexOfByKey(ComponentName);
+								if (ModelResources->ReferenceSkeletalMeshesData.IsValidIndex(ObjectComponentIndex))
 								{
-									const FName& ComponentName = InstanceUsage->GetComponentName();
-									USkeletalMesh* SkeletalMesh = CustomizableObject->GetSkeletalMeshComponentReferenceSkeletalMesh(ComponentName);
+									USkeletalMesh* SkeletalMesh = nullptr;
+
+									if (IsReplaceDiscardedWithReferenceMeshEnabled())
+									{
+										// Force load the reference mesh if necessary. 
+										TSoftObjectPtr<USkeletalMesh> SoftObjectPtr = ModelResources->ReferenceSkeletalMeshesData[ObjectComponentIndex].SoftSkeletalMesh;
+										SkeletalMesh = UE::Mutable::Private::LoadObject(SoftObjectPtr);
+									}
+
+									Parent->EmptyOverrideMaterials();
 									Parent->SetSkeletalMesh(SkeletalMesh);
 								}
 							}
diff --git a/Engine/Plugins/Mutable/Source/CustomizableObject/Private/MuCO/CustomizableSkeletalMeshActor.cpp b/Engine/Plugins/Mutable/Source/CustomizableObject/Private/MuCO/CustomizableSkeletalMeshActor.cpp
--- a/Engine/Plugins/Mutable/Source/CustomizableObject/Private/MuCO/CustomizableSkeletalMeshActor.cpp
+++ b/Engine/Plugins/Mutable/Source/CustomizableObject/Private/MuCO/CustomizableSkeletalMeshActor.cpp
@@ -76,13 +76,13 @@ void UCustomizableSkeletalMeshActorPrivate::Init(UCustomizableObjectInstance* In
 		return;
 	}
 	
-	for (UE::Mutable::Private::FComponentId ComponentId = 0; ComponentId < Object->GetComponentCount(); ++ComponentId)
+	for (int32 ObjectComponentIndex = 0; ObjectComponentIndex < Object->GetComponentCount(); ++ObjectComponentIndex)
 	{
-		const FName& ComponentName = Object->GetComponentName(ComponentId);
+		const FName& ComponentName = Object->GetComponentName(ObjectComponentIndex);
 		const FString ComponentNameString = ComponentName.ToString();
 
 		USkeletalMeshComponent* Component;
-		if (ComponentId == 0)
+		if (ObjectComponentIndex == 0)
 		{
 			Component = GetPublic()->SkeletalMeshComponents[0];
 			// Renaming the component will make it disappear. Keep it as it is.
@@ -102,7 +102,7 @@ void UCustomizableSkeletalMeshActorPrivate::Init(UCustomizableObjectInstance* In
 		}
 
 		UCustomizableSkeletalComponent* CustomizableComponent;
-		if (ComponentId == 0)
+		if (ObjectComponentIndex == 0)
 		{
 			CustomizableComponent = GetPublic()->CustomizableSkeletalComponents[0];
 			// Renaming the component will make it disappear. Keep it as it is.
@@ -181,20 +181,20 @@ void ACustomizableSkeletalMeshActor::SwitchComponentsMaterials(UCustomizableObje
 			return;
 		}
 
-		for (UE::Mutable::Private::FComponentId ComponentId = 0; ComponentId < SkeletalMeshComponents.Num(); ++ComponentId)
+		for (int32 ObjectComponentIndex = 0; ObjectComponentIndex < SkeletalMeshComponents.Num(); ++ObjectComponentIndex)
 		{
-			int32 NumMaterials = SkeletalMeshComponents[ComponentId]->GetNumMaterials();
+			int32 NumMaterials = SkeletalMeshComponents[ObjectComponentIndex]->GetNumMaterials();
 
 			if (bDebugMaterialEnabled)
 			{
 				for (int32 MatIndex = 0; MatIndex < NumMaterials; ++MatIndex)
 				{
-					SkeletalMeshComponents[ComponentId]->SetMaterial(MatIndex, DebugMaterial);
+					SkeletalMeshComponents[ObjectComponentIndex]->SetMaterial(MatIndex, DebugMaterial);
 				}
 			}
 			else // Remove debugmaterial
 			{
-				const FName ComponentName = CustomizableSkeletalComponents[ComponentId]->GetComponentName();
+				const FName ComponentName = CustomizableSkeletalComponents[ObjectComponentIndex]->GetComponentName();
 				
 				// check if original materials already overriden
 				const TArray<UMaterialInterface*> OverrideMaterials = COInstance->GetSkeletalMeshComponentOverrideMaterials(ComponentName);
@@ -204,12 +204,12 @@ void ACustomizableSkeletalMeshActor::SwitchComponentsMaterials(UCustomizableObje
 				{
 					for (int32 MatIndex = 0; MatIndex < OverrideMaterials.Num(); ++MatIndex)
 					{
-						SkeletalMeshComponents[ComponentId]->SetMaterial(MatIndex, OverrideMaterials[MatIndex]);
+						SkeletalMeshComponents[ObjectComponentIndex]->SetMaterial(MatIndex, OverrideMaterials[MatIndex]);
 					}
 				}
 				else
 				{
-					SkeletalMeshComponents[ComponentId]->EmptyOverrideMaterials();
+					SkeletalMeshComponents[ObjectComponentIndex]->EmptyOverrideMaterials();
 				}
 
 				bRemoveDebugMaterial = false;
diff --git a/Engine/Plugins/Mutable/Source/CustomizableObject/Private/MuCO/LogBenchmarkUtil.cpp b/Engine/Plugins/Mutable/Source/CustomizableObject/Private/MuCO/LogBenchmarkUtil.cpp
--- a/Engine/Plugins/Mutable/Source/CustomizableObject/Private/MuCO/LogBenchmarkUtil.cpp
+++ b/Engine/Plugins/Mutable/Source/CustomizableObject/Private/MuCO/LogBenchmarkUtil.cpp
@@ -307,9 +307,9 @@ void FLogBenchmarkUtil::FinishUpdateMesh(const TSharedRef<FUpdateContextPrivate>
 
 	// Cache the amount of triangles of this instance
 	uint32 TriangleCount = 0;	
-	for (UE::Mutable::Private::FComponentId ComponentId = 0; ComponentId < Object->GetComponentCount(); ++ComponentId)
+	for (int32 ComponentIndex = 0; ComponentIndex < Object->GetComponentCount(); ComponentIndex++)
 	{
-		const FName ComponentName = Object->GetPrivate()->GetComponentName(ComponentId);
+		const FName ComponentName = Object->GetPrivate()->GetComponentName(FCustomizableObjectComponentIndex(ComponentIndex));
 		
 		// Process the generated components (not null)
 		if (const USkeletalMesh* InstanceSkeletalMesh = Instance->GetComponentMeshSkeletalMesh(ComponentName))
diff --git a/Engine/Plugins/Mutable/Source/CustomizableObject/Public/MuCO/CustomizableObjectInstance.h b/Engine/Plugins/Mutable/Source/CustomizableObject/Public/MuCO/CustomizableObjectInstance.h
--- a/Engine/Plugins/Mutable/Source/CustomizableObject/Public/MuCO/CustomizableObjectInstance.h
+++ b/Engine/Plugins/Mutable/Source/CustomizableObject/Public/MuCO/CustomizableObjectInstance.h
@@ -364,7 +364,11 @@ public:
 	/** Set the CustomizableObject state that this instance will be generated into. */
 	UFUNCTION(BlueprintCallable, Category = CustomizableObjectInstance)
 	CUSTOMIZABLEOBJECT_API void SetCurrentState(const FString& StateName);
-	
+
+	// DEPRECATED: Use GetSkeletalMeshComponentSkeletalMesh 
+	UFUNCTION(BlueprintCallable, Category = CustomizableObjectInstance)
+	CUSTOMIZABLEOBJECT_API USkeletalMesh* GetSkeletalMesh(int32 ComponentIndex = 0) const;
+
 	// DEPRECATED: Use GetSkeletalMeshComponentSkeletalMesh 
 	UFUNCTION(BlueprintCallable, Category = CustomizableObjectInstance)
 	CUSTOMIZABLEOBJECT_API USkeletalMesh* GetComponentMeshSkeletalMesh(const FName& ComponentName) const;
@@ -403,7 +407,7 @@ public:
 	/** Returns the AssetUserData that was gathered from all the constituent mesh parts during the last update. 
 	  * It requires that the CustomizableObject had the bEnableAssetUserDataMerge set to true during compilation. */
 	UFUNCTION(BlueprintCallable, Category = CustomizableObjectInstance)
-	CUSTOMIZABLEOBJECT_API TSet<UAssetUserData*> GetMergedAssetUserData(const FName& ComponentName) const;
+	CUSTOMIZABLEOBJECT_API TSet<UAssetUserData*> GetMergedAssetUserData(int32 ComponentIndex) const;
 	
 	/** Generate the instance with the current parameters and update all the components Skeletal Meshes asynchronously. */
 	UFUNCTION(BlueprintCallable, Category = CustomizableObjectInstance)
@@ -757,6 +761,13 @@ public:
 
 	CUSTOMIZABLEOBJECT_API void ForEachComponentAnimInstance(FName ComponentName, FEachComponentAnimInstanceClassNativeDelegate Delegate) const;
 
+	/** DEPRECATED: Use ForEachComponentAnimInstance. */
+	UFUNCTION(BlueprintCallable, Category = CustomizableObjectInstance)
+	CUSTOMIZABLEOBJECT_API void ForEachAnimInstance(int32 ObjectComponentIndex, FEachComponentAnimInstanceClassDelegate Delegate) const;
+
+	/** DEPRECATED: Use ForEachComponentAnimInstance. */
+	CUSTOMIZABLEOBJECT_API void ForEachAnimInstance(int32 ObjectComponentIndex, FEachComponentAnimInstanceClassNativeDelegate Delegate) const;
+
 	/** Check if the given UAnimInstance class requires to be fixed up. */
 	CUSTOMIZABLEOBJECT_API bool AnimInstanceNeedsFixup(TSubclassOf<UAnimInstance> AnimInstance) const;
 
diff --git a/Engine/Plugins/Mutable/Source/CustomizableObject/Public/MuCO/CustomizableObjectInstanceDescriptor.h b/Engine/Plugins/Mutable/Source/CustomizableObject/Public/MuCO/CustomizableObjectInstanceDescriptor.h
--- a/Engine/Plugins/Mutable/Source/CustomizableObject/Public/MuCO/CustomizableObjectInstanceDescriptor.h
+++ b/Engine/Plugins/Mutable/Source/CustomizableObject/Public/MuCO/CustomizableObjectInstanceDescriptor.h
@@ -20,6 +20,7 @@ class UCustomizableObjectInstance;
 class UMaterialInterface;
 class FDescriptorHash;
 class FMutableUpdateCandidate;
+struct FCustomizableObjectComponentIndex;
 
 typedef TMap<const UCustomizableObjectInstance*, FMutableUpdateCandidate> FMutableInstanceUpdateMap;
 
diff --git a/Engine/Plugins/Mutable/Source/CustomizableObject/Public/MuCO/CustomizableObjectInstanceUsage.h b/Engine/Plugins/Mutable/Source/CustomizableObject/Public/MuCO/CustomizableObjectInstanceUsage.h
--- a/Engine/Plugins/Mutable/Source/CustomizableObject/Public/MuCO/CustomizableObjectInstanceUsage.h
+++ b/Engine/Plugins/Mutable/Source/CustomizableObject/Public/MuCO/CustomizableObjectInstanceUsage.h
@@ -42,6 +42,9 @@ public:
 	UFUNCTION(BlueprintCallable, Category = CustomizableObjectInstanceUsage)
 	UE_API UCustomizableObjectInstance* GetCustomizableObjectInstance() const;
 	
+	// DEPRECATED: use SetComponentName
+	UE_API void SetComponentIndex(int32 ObjectComponentIndex);
+
 	// DEPRECATED: use GetComponentName
 	UE_API int32 GetComponentIndex() const;
 
diff --git a/Engine/Plugins/Mutable/Source/CustomizableObjectEditor/Private/MuCOE/CustomizableObjectCompiler.cpp b/Engine/Plugins/Mutable/Source/CustomizableObjectEditor/Private/MuCOE/CustomizableObjectCompiler.cpp
--- a/Engine/Plugins/Mutable/Source/CustomizableObjectEditor/Private/MuCOE/CustomizableObjectCompiler.cpp
+++ b/Engine/Plugins/Mutable/Source/CustomizableObjectEditor/Private/MuCOE/CustomizableObjectCompiler.cpp
@@ -992,19 +992,20 @@ void FCustomizableObjectCompiler::CompileInternal(bool bAsync)
 										ImageProp.AddressX, ImageProp.AddressY });
 		}
 
-		const TArray<FSkeletalMeshComponentInfo>& ComponentInfos = GenerationContext.CompilationContext->ComponentInfos;
-		for (const FSkeletalMeshComponentInfo& ComponentInfo : ComponentInfos)
+		const TArray<FMutableComponentInfo>& ComponentInfos = GenerationContext.CompilationContext->ComponentInfos;
+		for (const FMutableComponentInfo& ComponentInfo : ComponentInfos)
 		{
 			if (!Cast<UCustomizableObjectNodeComponentMesh>(ComponentInfo.Node))
 			{
 				continue;
 			}
 
+			const FName ComponentName = ComponentInfo.ComponentName;
 			const FMutableLODSettings& ComponentLODSettings = ComponentInfo.LODSettings;
 
 			// Copy the LODSettings data found in the Component into the ModelResources
-			ModelResources->MinLODPerComponent.Add(ComponentInfo.ComponentId, ComponentLODSettings.MinLOD);
-			ModelResources->MinQualityLevelLODPerComponent.Add(ComponentInfo.ComponentId, ComponentLODSettings.MinQualityLevelLOD);
+			ModelResources->MinLODPerComponent.Add(ComponentName, ComponentLODSettings.MinLOD);
+			ModelResources->MinQualityLevelLODPerComponent.Add(ComponentName, ComponentLODSettings.MinQualityLevelLOD);
 		}
 		
 		ModelResources->ParameterUIDataMap = MoveTemp(GenerationContext.ParameterUIDataMap);
@@ -1046,7 +1047,7 @@ void FCustomizableObjectCompiler::CompileInternal(bool bAsync)
 		}
 		else
 		{
-			for (TTuple<UE::Mutable::Private::FComponentId, uint8>& Pair : ModelResources->NumLODsToStream)
+			for (TTuple<FName, uint8>& Pair : ModelResources->NumLODsToStream)
 			{
 				Pair.Value = 0;
 			}
@@ -1666,10 +1667,6 @@ void FCustomizableObjectCompiler::FinishCompilationTask()
 			Offset += ResourceSize;
 		}
 	}
-	else
-	{
-		CompilerLog(LOCTEXT("NoModel", "Unable to generate Model"), nullptr, EMessageSeverity::Error);
-	}
 
 	// Second part of model resources can be filled now.
 	{
diff --git a/Engine/Plugins/Mutable/Source/CustomizableObjectEditor/Private/MuCOE/CustomizableObjectEditorViewportClient.cpp b/Engine/Plugins/Mutable/Source/CustomizableObjectEditor/Private/MuCOE/CustomizableObjectEditorViewportClient.cpp
--- a/Engine/Plugins/Mutable/Source/CustomizableObjectEditor/Private/MuCOE/CustomizableObjectEditorViewportClient.cpp
+++ b/Engine/Plugins/Mutable/Source/CustomizableObjectEditor/Private/MuCOE/CustomizableObjectEditorViewportClient.cpp
@@ -875,9 +875,9 @@ void FCustomizableObjectEditorViewportClient::OnPreSetSkeletalMesh(const FPreSet
 	}
 
 	// Add new components
-	for (UE::Mutable::Private::FComponentId ComponentId = 0; ComponentId < CO->GetComponentCount(); ++ComponentId)
+	for (int32 ObjectComponentIndex = 0; ObjectComponentIndex < CO->GetComponentCount(); ++ObjectComponentIndex)
 	{
-		FName Name = CO->GetPrivate()->GetComponentName(ComponentId);
+		FName Name = CO->GetPrivate()->GetComponentName(FCustomizableObjectComponentIndex(ObjectComponentIndex));
 		if (Instance->GetComponentMeshSkeletalMesh(Name))
 		{
 			if (!SkeletalMeshComponents.Contains(Name))
diff --git a/Engine/Plugins/Mutable/Source/CustomizableObjectEditor/Private/MuCOE/CustomizableObjectInstanceBakingUtils.cpp b/Engine/Plugins/Mutable/Source/CustomizableObjectEditor/Private/MuCOE/CustomizableObjectInstanceBakingUtils.cpp
--- a/Engine/Plugins/Mutable/Source/CustomizableObjectEditor/Private/MuCOE/CustomizableObjectInstanceBakingUtils.cpp
+++ b/Engine/Plugins/Mutable/Source/CustomizableObjectEditor/Private/MuCOE/CustomizableObjectInstanceBakingUtils.cpp
@@ -579,11 +579,11 @@ bool BakeCustomizableObjectInstance(
 	TArray<UObject*> HandledSourceObjects;
 	
 	
-	for (UE::Mutable::Private::FComponentId ComponentId = 0; ComponentId < InstanceCO->GetComponentCount(); ++ComponentId)
+	for (int32 ObjectComponentIndex = 0; ObjectComponentIndex < InstanceCO->GetComponentCount(); ++ObjectComponentIndex)
 	{
 		check(DuplicatedObjects.Num() == HandledSourceObjects.Num())
 		
-		const FName ComponentName = InstanceCO->GetPrivate()->GetComponentName(ComponentId);
+		const FName ComponentName = InstanceCO->GetPrivate()->GetComponentName(FCustomizableObjectComponentIndex(ObjectComponentIndex));
 		USkeletalMesh* Mesh = InInstance.GetComponentMeshSkeletalMesh(ComponentName);
 
 		if (!Mesh)
@@ -1068,15 +1068,16 @@ bool BakeCustomizableObjectInstance(
 
 			// Add Instance Info in a custom AssetUserData
 			{
-				const FCustomizableInstanceComponentData& ComponentData = InInstance.GetPrivate()->GetComponentData(ComponentId);
+				const FCustomizableInstanceComponentData* ComponentData = InInstance.GetPrivate()->GetComponentData(ComponentName);
+				check(ComponentData);
 				
 				if (InInstance.GetAnimationGameplayTags().Num() ||
-					ComponentData.AnimSlotToBP.Num())
+					ComponentData->AnimSlotToBP.Num())
 				{
 					UCustomizableObjectInstanceUserData* InstanceData = NewObject<UCustomizableObjectInstanceUserData>(SkeletalMesh, NAME_None, RF_Public | RF_Transactional);
 					InstanceData->AnimationGameplayTag = InInstance.GetAnimationGameplayTags();
 
-					for (const TTuple<FName, TSoftClassPtr<UAnimInstance>>& AnimSlot : ComponentData.AnimSlotToBP)
+					for (const TTuple<FName, TSoftClassPtr<UAnimInstance>>& AnimSlot : ComponentData->AnimSlotToBP)
 					{
 						FCustomizableObjectAnimationSlot AnimationSlot;
 						AnimationSlot.Name = AnimSlot.Key;
@@ -1091,9 +1092,9 @@ bool BakeCustomizableObjectInstance(
 
 			// Copy LODSettings from the Reference Skeletal Mesh
 			{
-				if (ModelResources->ReferenceSkeletalMeshesData.IsValidIndex(ComponentId))
+				if (ModelResources->ReferenceSkeletalMeshesData.IsValidIndex(ObjectComponentIndex))
 				{
-					USkeletalMeshLODSettings* LODSettings = ModelResources->ReferenceSkeletalMeshesData[ComponentId].SkeletalMeshLODSettings;
+					USkeletalMeshLODSettings* LODSettings = ModelResources->ReferenceSkeletalMeshesData[ObjectComponentIndex].SkeletalMeshLODSettings;
 					SkeletalMesh->SetLODSettings(LODSettings);
 				}
 			}
diff --git a/Engine/Plugins/Mutable/Source/CustomizableObjectEditor/Private/MuCOE/GenerateMutableSource/GenerateMutableSource.cpp b/Engine/Plugins/Mutable/Source/CustomizableObjectEditor/Private/MuCOE/GenerateMutableSource/GenerateMutableSource.cpp
--- a/Engine/Plugins/Mutable/Source/CustomizableObjectEditor/Private/MuCOE/GenerateMutableSource/GenerateMutableSource.cpp
+++ b/Engine/Plugins/Mutable/Source/CustomizableObjectEditor/Private/MuCOE/GenerateMutableSource/GenerateMutableSource.cpp
@@ -222,21 +222,12 @@ FMutableCompilationContext::FMutableCompilationContext(const UCustomizableObject
 }
 
 
-FSkeletalMeshComponentInfo* FMutableCompilationContext::GetComponentInfo(const UCustomizableObjectNodeComponentMesh* NodeComponentMesh)
+FMutableComponentInfo* FMutableCompilationContext::GetComponentInfo(FName ComponentName)
 {
-	return ComponentInfos.FindByPredicate([&](const FSkeletalMeshComponentInfo& ComponentInfo)
-	{
-		return ComponentInfo.Node == NodeComponentMesh;
-	});
-}
-
-
-FSkeletalMeshComponentInfo* FMutableCompilationContext::GetComponentInfo(UE::Mutable::Private::FComponentId ComponentId)
-{
-	return ComponentInfos.FindByPredicate([&](const FSkeletalMeshComponentInfo& ComponentInfo)
-	{
-		return ComponentInfo.ComponentId == ComponentId;
-	});
+	return ComponentInfos.FindByPredicate([ComponentName](const FMutableComponentInfo& Component)
+		{
+			return Component.ComponentName == ComponentName;
+		});
 }
 
 
@@ -382,12 +373,9 @@ const FGuid FMutableGraphGenerationContext::GetNodeIdUnchecked(const UCustomizab
 }
 
 
-FSkeletalMeshComponentInfo* FMutableGraphGenerationContext::GetCurrentComponentInfo()
+FMutableComponentInfo* FMutableGraphGenerationContext::GetCurrentComponentInfo()
 {
-	return CompilationContext->ComponentInfos.FindByPredicate([&](const FSkeletalMeshComponentInfo& ComponentInfo)
-	{
-		return ComponentInfo.ComponentId == CurrentSkeletalMeshComponent;
-	});
+	return CompilationContext->GetComponentInfo(CurrentMeshComponent);
 }
 
 
@@ -979,7 +967,7 @@ UE::Mutable::Private::Ptr<UE::Mutable::Private::NodeObject> GenerateMutableSourc
 			int32 MinQualityLevelLOD = INT32_MAX;
 			int32 MaxQualityLevelLOD = 0;
 
-			for (const FSkeletalMeshComponentInfo& ComponentInfo : GenerationContext.CompilationContext->ComponentInfos)
+			for (const FMutableComponentInfo& ComponentInfo : GenerationContext.CompilationContext->ComponentInfos)
 			{
 				if (!Cast<UCustomizableObjectNodeComponentMesh>(ComponentInfo.Node))
 				{
@@ -1037,15 +1025,15 @@ UE::Mutable::Private::Ptr<UE::Mutable::Private::NodeObject> GenerateMutableSourc
 			TArray<UEdGraphPin*> ConnectedModifierPins = FollowInputPinArray(*ModifierPin);
 			for (UEdGraphPin* const ConnectedModifier : ConnectedModifierPins)
 			{
-				for (const FSkeletalMeshComponentInfo& Pair : GenerationContext.CompilationContext->ComponentInfos)
+				for (const FMutableComponentInfo& Pair : GenerationContext.CompilationContext->ComponentInfos)
 				{
-					GenerationContext.CurrentSkeletalMeshComponent = Pair.ComponentId;
+					GenerationContext.CurrentMeshComponent = Pair.ComponentName;
 					GenerationContext.CurrentAutoLODStrategy = Pair.Node->AutoLODStrategy;
 					
 					UE::Mutable::Private::Ptr<UE::Mutable::Private::NodeModifier> ModifierNode = GenerateMutableSourceModifier(ConnectedModifier, GenerationContext);
 					ObjectNode->Modifiers.Add(ModifierNode);
 
-					GenerationContext.CurrentSkeletalMeshComponent = INDEX_NONE;
+					GenerationContext.CurrentMeshComponent = {};
 					GenerationContext.CurrentAutoLODStrategy = {};
 				}
 			}
@@ -1480,7 +1468,8 @@ void PopulateReferenceSkeletalMeshesData(FMutableGraphGenerationContext& Generat
 	GenerationContext.ReferenceSkeletalMeshesData.AddDefaulted(ComponentCount);
 	for (uint32 ComponentIndex = 0; ComponentIndex < ComponentCount; ++ComponentIndex)
 	{
-		const FSkeletalMeshComponentInfo& ComponentInfo = GenerationContext.CompilationContext->ComponentInfos[ComponentIndex];
+		const FMutableComponentInfo& ComponentInfo = GenerationContext.CompilationContext->ComponentInfos[ComponentIndex];
+		const FName& ComponentName = ComponentInfo.ComponentName;
 		
 		USkeletalMesh* RefSkeletalMesh = ComponentInfo.RefSkeletalMesh.Get();
 		check(RefSkeletalMesh);
@@ -1500,11 +1489,11 @@ void PopulateReferenceSkeletalMeshesData(FMutableGraphGenerationContext& Generat
 		Data.SkeletalMeshLODSettings = RefSkeletalMesh->GetLODSettings();
 
 		// Gather LODData, this may include per LOD settings such as render data config or LODDataInfoArray
-		Data.LODData.AddDefaulted(GenerationContext.NumLODs[ComponentInfo.ComponentId]);
+		Data.LODData.AddDefaulted(GenerationContext.NumLODs[ComponentName]);
 		
 		const uint32 RefSkeletalMeshLODCount = RefSkeletalMesh->GetLODNum();
 
-		const uint32 LODCount = GenerationContext.NumLODs[ComponentInfo.ComponentId];
+		const uint32 LODCount = GenerationContext.NumLODs[ComponentName];
 		for (uint32 LODIndex = 0; LODIndex < LODCount; ++LODIndex)
 		{
 			FMutableRefLODData& LODData = Data.LODData[LODIndex];
@@ -1856,6 +1845,76 @@ UE::Mutable::Private::TManagedPtr<UE::Mutable::Private::FMesh> GenerateMeshConst
 	return Result;
 }
 
+
+FMutableComponentInfo::FMutableComponentInfo(FName InComponentName, USkeletalMesh* InRefSkeletalMesh)
+{
+	ComponentName = InComponentName;
+
+	if (!InRefSkeletalMesh || !InRefSkeletalMesh->GetSkeleton())
+	{
+		return;
+	}
+
+	RefSkeletalMesh = TStrongObjectPtr<USkeletalMesh>(InRefSkeletalMesh);
+	RefSkeleton = RefSkeletalMesh->GetSkeleton();
+
+	const int32 NumBones = RefSkeleton->GetReferenceSkeleton().GetRawBoneNum();
+	BoneNamesToPathHash.Reserve(NumBones);
+
+	const TArray<FMeshBoneInfo>& Bones = RefSkeleton->GetReferenceSkeleton().GetRawRefBoneInfo();
+
+	for (int32 BoneIndex = 0; BoneIndex < NumBones; ++BoneIndex)
+	{
+		const FMeshBoneInfo& Bone = Bones[BoneIndex];
+
+		// Retrieve parent bone name and respective hash, root-bone is assumed to have a parent hash of 0
+		const FName ParentName = Bone.ParentIndex != INDEX_NONE ? Bones[Bone.ParentIndex].Name : NAME_None;
+		const uint32 ParentHash = Bone.ParentIndex != INDEX_NONE ? GetTypeHash(ParentName) : 0;
+
+		// Look-up the path-hash from root to the parent bone
+		const uint32* ParentPath = BoneNamesToPathHash.Find(ParentName);
+		const uint32 ParentPathHash = ParentPath ? *ParentPath : 0;
+
+		// Append parent hash to path to give full path hash to current bone
+		const uint32 BonePathHash = HashCombine(ParentPathHash, ParentHash);
+
+		// Add path hash to current bone
+		BoneNamesToPathHash.Add(Bone.Name, BonePathHash);
+	}
+}
+
+
+void FMutableComponentInfo::AccumulateBonesToRemovePerLOD(const TArray<FLODReductionSettings>& LODReductionSettings, int32 NumLODs)
+{
+	BonesToRemovePerLOD.SetNum(NumLODs);
+
+	TMap<FName, bool> BonesToRemove;
+
+	const int32 ComponentSettingsLODCount = LODReductionSettings.Num();
+	for (int32 LODIndex = 0; LODIndex < NumLODs; ++LODIndex)
+	{
+		if (LODIndex < ComponentSettingsLODCount)
+		{
+			const FLODReductionSettings& LODReductionSetting = LODReductionSettings[LODIndex];
+
+			for (const FBoneToRemove& Bone : LODReductionSetting.BonesToRemove)
+			{
+				if (bool* bOnlyRemoveChildren = BonesToRemove.Find(Bone.BoneName))
+				{
+					// Removed by a previous LOD
+					*bOnlyRemoveChildren = (*bOnlyRemoveChildren) && Bone.bOnlyRemoveChildren;
+				}
+				else
+				{
+					BonesToRemove.Add(Bone.BoneName, Bone.bOnlyRemoveChildren);
+				}
+			}
+		}
+		
+		BonesToRemovePerLOD[LODIndex] = BonesToRemove;
+	}
+}
+
 #undef LOCTEXT_NAMESPACE
 
 
diff --git a/Engine/Plugins/Mutable/Source/CustomizableObjectEditor/Private/MuCOE/GenerateMutableSource/GenerateMutableSource.h b/Engine/Plugins/Mutable/Source/CustomizableObjectEditor/Private/MuCOE/GenerateMutableSource/GenerateMutableSource.h
--- a/Engine/Plugins/Mutable/Source/CustomizableObjectEditor/Private/MuCOE/GenerateMutableSource/GenerateMutableSource.h
+++ b/Engine/Plugins/Mutable/Source/CustomizableObjectEditor/Private/MuCOE/GenerateMutableSource/GenerateMutableSource.h
@@ -184,7 +184,7 @@ public:
 	/** UV Layout modes */
 	FLayoutGenerationFlags LayoutFlags;
 	
-	UE::Mutable::Private::FComponentId CurrentMeshComponent = INDEX_NONE;
+	FName CurrentMeshComponent;
 	
 	/** When caching a generated mesh, true if we force to generate the connected LOD when using Automatic LODs From Mesh. */
 	bool bOnlyConnectedLOD = false;
@@ -432,18 +432,22 @@ struct FGroupNodeIdsTempData
 
 
 /** Struct used to store info specific to each component during compilation */
-struct FSkeletalMeshComponentInfo
+struct FMutableComponentInfo
 {
+	FMutableComponentInfo(FName InComponentName, USkeletalMesh* InRefSkeletalMesh);
+
+	void AccumulateBonesToRemovePerLOD(const TArray<FLODReductionSettings>& LODReductionSettings, int32 NumLODs);
+
 	FName ComponentName;
-	
-	UE::Mutable::Private::FComponentId ComponentId = INDEX_NONE;
-	
+
 	FMutableLODSettings LODSettings;
 	
 	// Each component must have a reference SkeletalMesh with a valid Skeleton
 	TStrongObjectPtr<USkeletalMesh> RefSkeletalMesh;
 	TObjectPtr<USkeleton> RefSkeleton;
 	
+	UCustomizableObjectNodeComponentMesh* NodeComponentMesh = nullptr;
+	
 	// Map to check skeleton compatibility
 	TMap<SIZE_T, bool> SkeletonCompatibility;
 	
@@ -517,7 +521,7 @@ private:
 struct FGeneratedGroupProjectorsKey
 {
 	UCustomizableObjectNodeGroupProjectorParameter* Node = nullptr;
-	UE::Mutable::Private::FComponentId CurrentComponent;
+	FName CurrentComponent;
 
 	bool operator==(const FGeneratedGroupProjectorsKey&) const = default;
 };
@@ -574,7 +578,7 @@ struct FMutableCompilationContext
 	TMap<uint32, FMutableMeshMetadata> MeshMetadata;
 
 	/** Only Mesh Components (no passthrough). */
-	TArray<FSkeletalMeshComponentInfo> ComponentInfos;
+	TArray<FMutableComponentInfo> ComponentInfos;
 
 	/** Array of unique Bone identifiers. */
 	TMap<UE::Mutable::Private::FBoneName, FString> UniqueBoneNames;
@@ -605,9 +609,7 @@ public:
 	void Log(const FText& Message, const UObject* Context = nullptr, const EMessageSeverity::Type MessageSeverity = EMessageSeverity::Warning, const bool bAddBaseObjectInfo = true, const ELoggerSpamBin SpamBin = ELoggerSpamBin::ShowAll) const;
 
 	/** Component access (for Mesh components only). */
-	FSkeletalMeshComponentInfo* GetComponentInfo(const UCustomizableObjectNodeComponentMesh* NodeComponentMesh);
-
-	FSkeletalMeshComponentInfo* GetComponentInfo(UE::Mutable::Private::FComponentId ComponentName);
+	FMutableComponentInfo* GetComponentInfo(FName ComponentName);
 
 	/** Get unique identifier for BoneName built from its FString. */
 	UE::Mutable::Private::FBoneName GetBoneUnique(const FName& BoneName);
@@ -847,7 +849,7 @@ struct FMutableGraphGenerationContext
 	const FGuid GetNodeIdUnchecked(const UCustomizableObjectNode* Node);
 
 	/** Get the reference skeletal mesh associated to the current mesh component being generated */
-	FSkeletalMeshComponentInfo* GetCurrentComponentInfo();
+	FMutableComponentInfo* GetCurrentComponentInfo();
 	
 	UObject* LoadObject(const FSoftObjectPtr& SoftObject);
 	
@@ -869,7 +871,7 @@ struct FMutableGraphGenerationContext
 	 */
 	FMaterialBreakParameter GetCurrentMaterialBreakParameter();
 	
-	/** Only compiled components. All component types. Index is the ComponentId. */
+	/** Only compiled components. All component types. Index is the ObjectComponentIndex. */
 	TArray<FName> ComponentNames;
 	
 	TArray<FMutableRefSkeletalMeshData> ReferenceSkeletalMeshesData;
@@ -890,16 +892,16 @@ struct FMutableGraphGenerationContext
 	
 	uint8 FromLOD = 0; // LOD to append to the CurrentLOD when using AutomaticLODs. 
 	uint8 CurrentLOD = 0;
-	UE::Mutable::Private::FComponentId CurrentSkeletalMeshComponent = INDEX_NONE;
+	FName CurrentMeshComponent;
 
 	/** If this is set, we are generating materials for a "passthrough" component, with a fixed mesh. */
 	UE::Mutable::Private::Ptr<UE::Mutable::Private::NodeMesh> ComponentMeshOverride;
 
-	TMap<UE::Mutable::Private::FComponentId, uint8> NumLODs;
+	TMap<FName, uint8> NumLODs;
 
-	TMap<UE::Mutable::Private::FComponentId, uint8> FirstLODAvailable;
+	TMap<FName, uint8> FirstLODAvailable;
 	
-	TMap<UE::Mutable::Private::FComponentId, uint8> NumMaxLODsToStream;
+	TMap<FName, uint8> NumMaxLODsToStream;
 
 	bool bEnableLODStreaming = true;
 
diff --git a/Engine/Plugins/Mutable/Source/CustomizableObjectEditor/Private/MuCOE/GenerateMutableSource/GenerateMutableSourceComponent.cpp b/Engine/Plugins/Mutable/Source/CustomizableObjectEditor/Private/MuCOE/GenerateMutableSource/GenerateMutableSourceComponent.cpp
--- a/Engine/Plugins/Mutable/Source/CustomizableObjectEditor/Private/MuCOE/GenerateMutableSource/GenerateMutableSourceComponent.cpp
+++ b/Engine/Plugins/Mutable/Source/CustomizableObjectEditor/Private/MuCOE/GenerateMutableSource/GenerateMutableSourceComponent.cpp
@@ -61,48 +61,55 @@ UE::Mutable::Private::Ptr<UE::Mutable::Private::NodeComponent> GenerateMutableSo
 	
 	if (const UCustomizableObjectNodeComponentMesh* TypedComponentMesh = Cast<UCustomizableObjectNodeComponentMesh>(Node))
 	{
-		FSkeletalMeshComponentInfo* ComponentInfo = GenerationContext.CompilationContext->GetComponentInfo(TypedComponentMesh);
-		if (ComponentInfo)
+		if (!GenerationContext.CompilationContext->ComponentInfos.ContainsByPredicate([&](const FMutableComponentInfo& ComponentInfo)
 		{
-			UE::Mutable::Private::Ptr<UE::Mutable::Private::NodeComponentNew> NodeComponentNew = new UE::Mutable::Private::NodeComponentNew();
-			NodeComponentNew->Id = ComponentInfo->ComponentId;
-			NodeComponentNew->SetMessageContext(Node);
+			return ComponentInfo.Node == TypedComponentMesh;
+		}))
+		{
+			return nullptr; // Not generated in the first pass.
+		}
 		
-			UEdGraphPin* MaterialAssetPin = TypedComponentMesh->GetOverlayMaterialAssetPin();
-			if (const UEdGraphPin* ConnectedPin = MaterialAssetPin ? FollowInputPin(*MaterialAssetPin) : nullptr)
-			{
-				GenerationContext.CurrentMaterialParameterId = ConnectedPin->PinId.ToString();
-				UE::Mutable::Private::Ptr<UE::Mutable::Private::NodeMaterial> MaterialNode = GenerateMutableSourceMaterial(ConnectedPin, GenerationContext, FGenerationMaterialOptions());
-				NodeComponentNew->OverlayMaterial = MaterialNode;
-			}
-			else if(UMaterialInterface* OverlayMaterial = GenerationContext.LoadObject(TypedComponentMesh->GetOverlayMaterial()))
-			{
-				UE::Mutable::Private::Ptr<UE::Mutable::Private::NodeMaterialConstant> MaterialNode = new UE::Mutable::Private::NodeMaterialConstant();
-				MaterialNode->MaterialId = GenerationContext.CompilationContext->PassthroughObjectFactory.Add(*OverlayMaterial);
-				NodeComponentNew->OverlayMaterial = MaterialNode;
-			}
-
-			Result = NodeComponentNew;
+		UE::Mutable::Private::Ptr<UE::Mutable::Private::NodeComponentNew> NodeComponentNew = new UE::Mutable::Private::NodeComponentNew();
+		const FName& ComponentName = TypedComponentMesh->GetComponentName(&GenerationContext.MacroNodesStack);
 
-			GenerationContext.CurrentSkeletalMeshComponent = ComponentInfo->ComponentId;
-			GenerationContext.CurrentAutoLODStrategy = TypedComponentMesh->AutoLODStrategy;
-		
-			NodeComponentNew->SkeletalMesh = UE::Mutable::Private::GenerateMutableSourceSkeletalMesh(Pin, GenerationContext);
+		NodeComponentNew->Id = GenerationContext.ComponentNames.Find(ComponentName);
+		NodeComponentNew->SetMessageContext(Node);
 		
-			GenerationContext.CurrentSkeletalMeshComponent = INDEX_NONE;
-			GenerationContext.CurrentAutoLODStrategy = {};
+		UEdGraphPin* MaterialAssetPin = TypedComponentMesh->GetOverlayMaterialAssetPin();
+		if (const UEdGraphPin* ConnectedPin = MaterialAssetPin ? FollowInputPin(*MaterialAssetPin) : nullptr)
+		{
+			GenerationContext.CurrentMaterialParameterId = ConnectedPin->PinId.ToString();
+			UE::Mutable::Private::Ptr<UE::Mutable::Private::NodeMaterial> MaterialNode = GenerateMutableSourceMaterial(ConnectedPin, GenerationContext, FGenerationMaterialOptions());
+			NodeComponentNew->OverlayMaterial = MaterialNode;
 		}
+		else if(UMaterialInterface* OverlayMaterial = GenerationContext.LoadObject(TypedComponentMesh->GetOverlayMaterial()))
+		{
+			UE::Mutable::Private::Ptr<UE::Mutable::Private::NodeMaterialConstant> MaterialNode = new UE::Mutable::Private::NodeMaterialConstant();
+			MaterialNode->MaterialId = GenerationContext.CompilationContext->PassthroughObjectFactory.Add(*OverlayMaterial);
+			NodeComponentNew->OverlayMaterial = MaterialNode;
+		}
+
+		Result = NodeComponentNew;
+
+		GenerationContext.CurrentMeshComponent = ComponentName;
+		GenerationContext.CurrentAutoLODStrategy = TypedComponentMesh->AutoLODStrategy;
+		
+		NodeComponentNew->SkeletalMesh = UE::Mutable::Private::GenerateMutableSourceSkeletalMesh(Pin, GenerationContext);
+
+		GenerationContext.CurrentMeshComponent = {};
+		GenerationContext.CurrentAutoLODStrategy = {};
 	}
 
 	else if (const UCustomizableObjectNodeComponentMeshAddTo* TypedComponentMeshExtend = Cast<UCustomizableObjectNodeComponentMeshAddTo>(Node))
 	{
 		const FName& ParentComponentName = TypedComponentMeshExtend->GetParentComponentName(&GenerationContext.MacroNodesStack);
-		const UE::Mutable::Private::FComponentId ComponentId = GenerationContext.ComponentNames.IndexOfByKey(ParentComponentName);
-		
-		FSkeletalMeshComponentInfo* ComponentInfo = GenerationContext.CompilationContext->GetComponentInfo(ComponentId);
-		if (ComponentInfo)
+
+		if (FMutableComponentInfo* FindResult = GenerationContext.CompilationContext->ComponentInfos.FindByPredicate([&](const FMutableComponentInfo& Element)
+		{
+			return Element.ComponentName == ParentComponentName;
+		}))
 		{
-			UCustomizableObjectNodeComponentMesh* TypedParentComponentMesh = ComponentInfo->Node;
+			UCustomizableObjectNodeComponentMesh* TypedParentComponentMesh = FindResult->Node;
 
 			if (TypedComponentMeshExtend->NumLODs > TypedParentComponentMesh->NumLODs)
 			{
@@ -112,12 +119,12 @@ UE::Mutable::Private::Ptr<UE::Mutable::Private::NodeComponent> GenerateMutableSo
 
 			// Swap the macro contex since this component can be in another macro
 			TArray<const UCustomizableObjectNodeMacroInstance*> MacroContextCopy = GenerationContext.MacroNodesStack;
-			GenerationContext.MacroNodesStack = ComponentInfo->MacroContext;
+			GenerationContext.MacroNodesStack = FindResult->MacroContext;
 
 			UE::Mutable::Private::Ptr<UE::Mutable::Private::NodeComponent> ParentNodeComponent = GenerateMutableSourceComponent(TypedParentComponentMesh->OutputPin.Get(), GenerationContext);
 
 			// Get the parent name using the parent's macro context
-			GenerationContext.CurrentSkeletalMeshComponent = ComponentInfo->ComponentId;
+			GenerationContext.CurrentMeshComponent = TypedParentComponentMesh->GetComponentName(&GenerationContext.MacroNodesStack);
 
 			// Restore GenerationContext MacroContex
 			GenerationContext.MacroNodesStack = MacroContextCopy;
@@ -132,7 +139,7 @@ UE::Mutable::Private::Ptr<UE::Mutable::Private::NodeComponent> GenerateMutableSo
 			
 			NodeComponentEdit->SkeletalMesh = UE::Mutable::Private::GenerateMutableSourceSkeletalMeshNew(Pin, GenerationContext);
 
-			GenerationContext.CurrentSkeletalMeshComponent = INDEX_NONE;
+			GenerationContext.CurrentMeshComponent = {};
 			GenerationContext.CurrentAutoLODStrategy = {};
 			
 			Result = NodeComponentEdit;	
@@ -146,9 +153,23 @@ UE::Mutable::Private::Ptr<UE::Mutable::Private::NodeComponent> GenerateMutableSo
 	
 	else if (const UCustomizableObjectNodeComponentPassthroughMesh* TypedComponentPassthroughMesh = Cast<UCustomizableObjectNodeComponentPassthroughMesh>(Node))
 	{
-		UE::Mutable::Private::Ptr<UE::Mutable::Private::NodeComponentNew> ComponentNode = new UE::Mutable::Private::NodeComponentNew;
-
 		const FName& ComponentName = TypedComponentPassthroughMesh->GetComponentName(&GenerationContext.MacroNodesStack);
+		GenerationContext.CurrentMeshComponent = ComponentName;
+
+		if (ComponentName.IsNone())
+		{
+			FString Msg = FString::Printf(TEXT("Invalid Component Name."));
+			GenerationContext.Log(FText::FromString(Msg), TypedComponentPassthroughMesh, EMessageSeverity::Warning);
+			return nullptr;
+		}
+		
+		// Create the component node
+		UE::Mutable::Private::Ptr<UE::Mutable::Private::NodeComponentNew> ComponentNode = new UE::Mutable::Private::NodeComponentNew;
+		if (GenerationContext.ComponentNames.Contains(ComponentName))
+		{
+			FString Msg = FString::Printf(TEXT("More than one component with the same name [%s] have been found. This is not supported."), *ComponentName.ToString());
+			GenerationContext.Log(FText::FromString(Msg), TypedComponentPassthroughMesh, EMessageSeverity::Warning);
+		}
 		ComponentNode->Id = GenerationContext.ComponentNames.Add(ComponentName);
 		ComponentNode->SkeletalMesh = UE::Mutable::Private::GenerateMutableSourceSkeletalMesh(Pin, GenerationContext);
 
@@ -181,6 +202,7 @@ UE::Mutable::Private::Ptr<UE::Mutable::Private::NodeComponent> GenerateMutableSo
 			}
 		}
 		
+		GenerationContext.CurrentMeshComponent = FName();
 		Result = ComponentNode;
 	}
 	
@@ -337,6 +359,36 @@ UE::Mutable::Private::Ptr<UE::Mutable::Private::NodeComponent> GenerateMutableSo
 
 void FirstPass(UCustomizableObjectNodeComponentMesh& Node, FMutableGraphGenerationContext& GenerationContext)
 {
+	FName ComponentName = Node.GetComponentName(&GenerationContext.MacroNodesStack);
+
+	if (ComponentName.IsNone())
+	{
+		GenerationContext.Log(LOCTEXT("EmptyComponentNameError", "Error! Missing name in a component of the Customizable Object."), &Node, EMessageSeverity::Error);
+		return;
+	}
+	
+	if (FMutableComponentInfo* Result = GenerationContext.CompilationContext->ComponentInfos.FindByPredicate([&](const FMutableComponentInfo& Element)
+	{
+		return Element.ComponentName == ComponentName;
+	}))
+	{
+		FText Msg;
+		if (Result->Node->IsInMacro())
+		{
+			const UCustomizableObjectNodeMacroInstance* CurrentMacro = GenerationContext.MacroNodesStack.Last();
+			check(CurrentMacro);
+			
+			Msg = FText::Format(LOCTEXT("ComponentNodeWithSameNameExists_Macro", "Error! A Mesh Component node with the same name already exists in a Macro Instance [{0}]"), FText::FromName(CurrentMacro->ParentMacro->Name));
+		}
+		else
+		{
+			Msg = FText::Format(LOCTEXT("ComponentNodeWithSameNameExists_CO", "Error! A Mesh Component node with the same name already exists in the Customizable Object [{0}]"), FText::FromString(GraphTraversal::GetObject(*Result->Node)->GetName()));
+		}
+		
+		GenerationContext.Log(Msg, &Node, EMessageSeverity::Error);
+		return;
+	}
+	
 	USkeletalMesh* RefSkeletalMesh = Node.ReferenceSkeletalMesh;
 	if (!RefSkeletalMesh)
 	{
@@ -434,83 +486,23 @@ void FirstPass(UCustomizableObjectNodeComponentMesh& Node, FMutableGraphGenerati
 
 	NumMaxLODsToStream = FMath::Clamp(NumMaxLODsToStream, 0, NumLODs - 1);
 
-	const FName ComponentName = Node.GetComponentName(&GenerationContext.MacroNodesStack);
-	UE::Mutable::Private::FComponentId ComponentId = GenerationContext.ComponentNames.Add(ComponentName);
-
-	GenerationContext.NumLODs.Add(ComponentId, NumLODs);
-	GenerationContext.FirstLODAvailable.Add(ComponentId, FirstLODAvailable);
-	GenerationContext.NumMaxLODsToStream.Add(ComponentId, NumMaxLODsToStream);
+	GenerationContext.NumLODs.Add(ComponentName, NumLODs);
+	GenerationContext.FirstLODAvailable.Add(ComponentName, FirstLODAvailable);
+	GenerationContext.NumMaxLODsToStream.Add(ComponentName, NumMaxLODsToStream);
 	
 	// Add a new entry to the list of Component Infos
-	FSkeletalMeshComponentInfo ComponentInfo;
-	ComponentInfo.ComponentName = ComponentName;
-	ComponentInfo.ComponentId = ComponentId;
+	FMutableComponentInfo ComponentInfo(ComponentName, RefSkeletalMesh);
 	ComponentInfo.Node = &Node;
+	ComponentInfo.AccumulateBonesToRemovePerLOD(Node.LODReductionSettings, Node.NumLODs);
 	ComponentInfo.MacroContext = GenerationContext.MacroNodesStack;
 	ComponentInfo.LODSettings = LODSettings;
-	
-	if (RefSkeletalMesh && RefSkeleton)
-	{
-		ComponentInfo.RefSkeletalMesh = TStrongObjectPtr(RefSkeletalMesh);
-		ComponentInfo.RefSkeleton = RefSkeleton;
-
-		const int32 NumBones = RefSkeleton->GetReferenceSkeleton().GetRawBoneNum();
-		ComponentInfo.BoneNamesToPathHash.Reserve(NumBones);
-
-		const TArray<FMeshBoneInfo>& Bones = RefSkeleton->GetReferenceSkeleton().GetRawRefBoneInfo();
-
-		for (int32 BoneIndex = 0; BoneIndex < NumBones; ++BoneIndex)
-		{
-			const FMeshBoneInfo& Bone = Bones[BoneIndex];
-
-			// Retrieve parent bone name and respective hash, root-bone is assumed to have a parent hash of 0
-			const FName ParentName = Bone.ParentIndex != INDEX_NONE ? Bones[Bone.ParentIndex].Name : NAME_None;
-			const uint32 ParentHash = Bone.ParentIndex != INDEX_NONE ? GetTypeHash(ParentName) : 0;
-
-			// Look-up the path-hash from root to the parent bone
-			const uint32* ParentPath = ComponentInfo.BoneNamesToPathHash.Find(ParentName);
-			const uint32 ParentPathHash = ParentPath ? *ParentPath : 0;
-
-			// Append parent hash to path to give full path hash to current bone
-			const uint32 BonePathHash = HashCombine(ParentPathHash, ParentHash);
-
-			// Add path hash to current bone
-			ComponentInfo.BoneNamesToPathHash.Add(Bone.Name, BonePathHash);
-		}
-	}
-
-	ComponentInfo.BonesToRemovePerLOD.SetNum(NumLODs);
-
-	TMap<FName, bool> BonesToRemove;
-
-	const int32 ComponentSettingsLODCount = Node.LODReductionSettings.Num();
-	for (int32 LODIndex = 0; LODIndex < NumLODs; ++LODIndex)
-	{
-		if (LODIndex < ComponentSettingsLODCount)
-		{
-			const FLODReductionSettings& LODReductionSetting = Node.LODReductionSettings[LODIndex];
-
-			for (const FBoneToRemove& Bone : LODReductionSetting.BonesToRemove)
-			{
-				if (bool* bOnlyRemoveChildren = BonesToRemove.Find(Bone.BoneName))
-				{
-					// Removed by a previous LOD
-					*bOnlyRemoveChildren = (*bOnlyRemoveChildren) && Bone.bOnlyRemoveChildren;
-				}
-				else
-				{
-					BonesToRemove.Add(Bone.BoneName, Bone.bOnlyRemoveChildren);
-				}
-			}
-		}
-		
-		ComponentInfo.BonesToRemovePerLOD[LODIndex] = BonesToRemove;
-	}
 
 	GenerationContext.CompilationContext->ComponentInfos.Add(ComponentInfo);
 
 	// Make sure the Skeleton from the reference mesh is added to the list of referenced Skeletons.
 	GenerationContext.CompilationContext->ReferencedSkeletons.Add(RefSkeleton);
+
+	GenerationContext.ComponentNames.Add(ComponentName);
 }
 
 
diff --git a/Engine/Plugins/Mutable/Source/CustomizableObjectEditor/Private/MuCOE/GenerateMutableSource/GenerateMutableSourceGroupProjector.cpp b/Engine/Plugins/Mutable/Source/CustomizableObjectEditor/Private/MuCOE/GenerateMutableSource/GenerateMutableSourceGroupProjector.cpp
--- a/Engine/Plugins/Mutable/Source/CustomizableObjectEditor/Private/MuCOE/GenerateMutableSource/GenerateMutableSourceGroupProjector.cpp
+++ b/Engine/Plugins/Mutable/Source/CustomizableObjectEditor/Private/MuCOE/GenerateMutableSource/GenerateMutableSourceGroupProjector.cpp
@@ -299,7 +299,7 @@ TOptional<FGroupProjectorTempData> GenerateMutableGroupProjector(UCustomizableOb
 {
 	FGeneratedGroupProjectorsKey Key;
 	Key.Node = ProjParamNode;
-	Key.CurrentComponent = GenerationContext.CurrentSkeletalMeshComponent;
+	Key.CurrentComponent = GenerationContext.CurrentMeshComponent;
 	
 	if (FGroupProjectorTempData* Result = GenerationContext.GeneratedGroupProjectors.Find(Key))
 	{
@@ -429,7 +429,7 @@ TOptional<FGroupProjectorTempData> GenerateMutableGroupProjector(UCustomizableOb
 			return {};
 		}
 
-		FSkeletalMeshComponentInfo* CurrentComponentInfo = GenerationContext.GetCurrentComponentInfo();
+		FMutableComponentInfo* CurrentComponentInfo = GenerationContext.GetCurrentComponentInfo();
 		check(CurrentComponentInfo);
 		
 		PoseEnumParameterNode->Options[0].Value = 0.f;
diff --git a/Engine/Plugins/Mutable/Source/CustomizableObjectEditor/Private/MuCOE/GenerateMutableSource/GenerateMutableSourceMesh.cpp b/Engine/Plugins/Mutable/Source/CustomizableObjectEditor/Private/MuCOE/GenerateMutableSource/GenerateMutableSourceMesh.cpp
--- a/Engine/Plugins/Mutable/Source/CustomizableObjectEditor/Private/MuCOE/GenerateMutableSource/GenerateMutableSourceMesh.cpp
+++ b/Engine/Plugins/Mutable/Source/CustomizableObjectEditor/Private/MuCOE/GenerateMutableSource/GenerateMutableSourceMesh.cpp
@@ -371,7 +371,7 @@ void GetLODAndSectionForAutomaticLODs(const FMutableCompilationContext& Context,
 }
 
 
-void BuildRemappedBonesArray(const FSkeletalMeshComponentInfo& InComponentInfo, TObjectPtr<const USkeletalMesh> InSkeletalMesh, int32 InLODIndex, const TArray<FBoneIndexType>& InRequiredBones, TArray<FBoneIndexType>& OutRemappedBones)
+void BuildRemappedBonesArray(const FMutableComponentInfo& InComponentInfo, TObjectPtr<const USkeletalMesh> InSkeletalMesh, int32 InLODIndex, const TArray<FBoneIndexType>& InRequiredBones, TArray<FBoneIndexType>& OutRemappedBones)
 {
 	if (!InSkeletalMesh)
 	{
@@ -612,7 +612,7 @@ void ConvertVertexSkinningData(uint16* InfluenceBones, const UE::Mutable::Privat
 }
 
 
-bool IsSkeletalMeshCompatibleWithRefSkeleton(FSkeletalMeshComponentInfo& ComponentInfo, TObjectPtr<const USkeletalMesh> InSkeletalMesh, FString& OutErrorMessage)
+bool IsSkeletalMeshCompatibleWithRefSkeleton(FMutableComponentInfo& ComponentInfo, TObjectPtr<const USkeletalMesh> InSkeletalMesh, FString& OutErrorMessage)
 {
 	TObjectPtr<const USkeleton> Skeleton = InSkeletalMesh->GetSkeleton();
 
@@ -1314,7 +1314,7 @@ UE::Tasks::FTask ConvertSkeletalMeshToMutable(TSharedRef<FMeshConversionContext>
 			return UE::Tasks::MakeCompletedTask<void>();
 		}
 
-		FSkeletalMeshComponentInfo* MutComponentInfo = Context.GetComponentInfo(Source.ComponentId);
+		FMutableComponentInfo* MutComponentInfo = Context.GetComponentInfo(Source.Component);
 		check(MutComponentInfo);
 
 		USkeletalMesh* ComponentRefSkeletalMesh = MutComponentInfo->RefSkeletalMesh.Get();
@@ -2657,7 +2657,7 @@ UE::Mutable::Private::TManagedPtr<UE::Mutable::Private::FMesh> ConvertSkeletalMe
 	Source.LODOffset = GenerationContext.CurrentLOD - GenerationContext.FromLOD;
 	Source.Flags = GenerationContext.MeshGenerationFlags.Last();
 	Source.MessageContext = CurrentNode;
-	Source.ComponentId = GenerationContext.CurrentSkeletalMeshComponent;
+	Source.Component = GenerationContext.CurrentMeshComponent;
 	Source.SocketPriority = GenerationContext.SocketPriorityStack.Last();
 	Source.BonePosePriority = GenerationContext.BonePosePriorityStack.Last();
 
@@ -3880,7 +3880,7 @@ UE::Mutable::Private::Ptr<UE::Mutable::Private::NodeMesh> GenerateMutableSourceM
 			// 
 			// Generate Layouts when using manual LODs or when processing the base Automatic LOD (the other automatic LODs will use the base layouts).
 			const bool bGenerateLayouts = GenerationContext.CurrentAutoLODStrategy != ECustomizableObjectAutomaticLODStrategy::AutomaticFromMesh
-				|| (GenerationContext.CurrentLOD == GenerationContext.FromLOD || GenerationContext.CurrentLOD <= GenerationContext.FirstLODAvailable[GenerationContext.CurrentSkeletalMeshComponent]);
+				|| (GenerationContext.CurrentLOD == GenerationContext.FromLOD || GenerationContext.CurrentLOD <= GenerationContext.FirstLODAvailable[GenerationContext.CurrentMeshComponent]);
 
 			if(bGenerateLayouts)
 			{
@@ -4533,7 +4533,7 @@ UE::Mutable::Private::Ptr<UE::Mutable::Private::NodeMesh> GenerateMutableSourceM
 							// 
 							// Generate Layouts when using manual LODs or when processing the base Automatic LOD (the other automatic LODs will use the base layouts).
 							const bool bGenerateLayouts = GenerationContext.CurrentAutoLODStrategy != ECustomizableObjectAutomaticLODStrategy::AutomaticFromMesh
-								|| (GenerationContext.CurrentLOD == GenerationContext.FromLOD || GenerationContext.CurrentLOD <= GenerationContext.FirstLODAvailable[GenerationContext.CurrentSkeletalMeshComponent]);
+								|| (GenerationContext.CurrentLOD == GenerationContext.FromLOD || GenerationContext.CurrentLOD <= GenerationContext.FirstLODAvailable[GenerationContext.CurrentMeshComponent]);
 
 							if (bGenerateLayouts)
 							{
diff --git a/Engine/Plugins/Mutable/Source/CustomizableObjectEditor/Private/MuCOE/GenerateMutableSource/GenerateMutableSourceModifier.cpp b/Engine/Plugins/Mutable/Source/CustomizableObjectEditor/Private/MuCOE/GenerateMutableSource/GenerateMutableSourceModifier.cpp
--- a/Engine/Plugins/Mutable/Source/CustomizableObjectEditor/Private/MuCOE/GenerateMutableSource/GenerateMutableSourceModifier.cpp
+++ b/Engine/Plugins/Mutable/Source/CustomizableObjectEditor/Private/MuCOE/GenerateMutableSource/GenerateMutableSourceModifier.cpp
@@ -59,7 +59,7 @@ UE::Mutable::Private::Ptr<UE::Mutable::Private::NodeModifier> GenerateMutableSou
 	UCustomizableObjectNode* Node = CastChecked<UCustomizableObjectNode>(Pin->GetOwningNode());
 
 	FGeneratedKey Key(reinterpret_cast<void*>(&GenerateMutableSourceModifier), *Pin, *Node, GenerationContext, true);
-	Key.CurrentMeshComponent = GenerationContext.CurrentSkeletalMeshComponent;
+	Key.CurrentMeshComponent = GenerationContext.CurrentMeshComponent;
 	
 	if (const FGeneratedData* Generated = GenerationContext.Generated.Find(Key))
 	{
@@ -265,7 +265,7 @@ UE::Mutable::Private::Ptr<UE::Mutable::Private::NodeModifier> GenerateMutableSou
 		SurfNode->RequiredTags = TypedNodeExt->GetNodeRequiredTags(&GenerationContext.MacroNodesStack);
 
 		// Is this enough? Should we try to narrow down with potential mesh sections modified by this?
-		int32 LODCount = GenerationContext.NumLODs[GenerationContext.CurrentSkeletalMeshComponent];
+		int32 LODCount = GenerationContext.NumLODs[GenerationContext.CurrentMeshComponent];
 		SurfNode->LODs.SetNum(LODCount);
 
 		for (int32 LODIndex= TypedNodeExt->FirstLOD; LODIndex<LODCount; ++LODIndex)
@@ -388,7 +388,7 @@ UE::Mutable::Private::Ptr<UE::Mutable::Private::NodeModifier> GenerateMutableSou
 		if (const UEdGraphPin* ConnectedPin = FollowInputPin(*TypedNodeRem->RemoveMeshPin()))
 		{
 			// Is this enough? Should we try to narrow down with potential mesh sections modified by this?
-			int32 LODCount = GenerationContext.NumLODs[GenerationContext.CurrentSkeletalMeshComponent];
+			int32 LODCount = GenerationContext.NumLODs[GenerationContext.CurrentMeshComponent];
 			SurfNode->LODs.SetNum(LODCount);
 
 			SurfNode->FaceCullStrategy = TypedNodeRem->FaceCullStrategy;
@@ -441,7 +441,7 @@ UE::Mutable::Private::Ptr<UE::Mutable::Private::NodeModifier> GenerateMutableSou
 		SurfNode->RequiredTags = TypedNodeEdit->GetNodeRequiredTags(&GenerationContext.MacroNodesStack);
 
 		// Is this enough? Should we try to narrow down with potential mesh sections modified by this?
-		int32 LODCount = GenerationContext.NumLODs[GenerationContext.CurrentSkeletalMeshComponent];
+		int32 LODCount = GenerationContext.NumLODs[GenerationContext.CurrentMeshComponent];
 		SurfNode->LODs.SetNum(LODCount);
 
 		for (int32 LODIndex = 0; LODIndex < LODCount; ++LODIndex)
@@ -686,8 +686,9 @@ UE::Mutable::Private::Ptr<UE::Mutable::Private::NodeModifier> GenerateMutableSou
 	{
 		Result->SetMessageContext(Node);
 
-		check(GenerationContext.CurrentSkeletalMeshComponent != INDEX_NONE);
-		Result->RequiredComponentId = GenerationContext.CurrentSkeletalMeshComponent;
+		int32 ComponentId = GenerationContext.ComponentNames.IndexOfByKey(GenerationContext.CurrentMeshComponent);
+		check(ComponentId>=0);
+		Result->RequiredComponentId = ComponentId;
 	}
 
 	if (bCacheNode)
diff --git a/Engine/Plugins/Mutable/Source/CustomizableObjectEditor/Private/MuCOE/GenerateMutableSource/GenerateMutableSourceSkeletalMesh.cpp b/Engine/Plugins/Mutable/Source/CustomizableObjectEditor/Private/MuCOE/GenerateMutableSource/GenerateMutableSourceSkeletalMesh.cpp
--- a/Engine/Plugins/Mutable/Source/CustomizableObjectEditor/Private/MuCOE/GenerateMutableSource/GenerateMutableSourceSkeletalMesh.cpp
+++ b/Engine/Plugins/Mutable/Source/CustomizableObjectEditor/Private/MuCOE/GenerateMutableSource/GenerateMutableSourceSkeletalMesh.cpp
@@ -75,7 +75,7 @@ namespace UE::Mutable::Private
 
 			int32 FirstLOD = -1;
 
-			const int32 NumLODsInRoot = GenerationContext.NumLODs[GenerationContext.CurrentSkeletalMeshComponent];
+			const int32 NumLODsInRoot = GenerationContext.NumLODs[GenerationContext.CurrentMeshComponent];
 			for (int32 CurrentLOD = 0; CurrentLOD < NumLODsInRoot; ++CurrentLOD)
 			{
 				GenerationContext.CurrentLOD = CurrentLOD;
@@ -99,7 +99,7 @@ namespace UE::Mutable::Private
 					continue;
 				}
 
-				if (GenerationContext.CurrentLOD < GenerationContext.FirstLODAvailable[GenerationContext.CurrentSkeletalMeshComponent])
+				if (GenerationContext.CurrentLOD < GenerationContext.FirstLODAvailable[GenerationContext.CurrentMeshComponent])
 				{
 					continue;
 				} 
diff --git a/Engine/Plugins/Mutable/Source/CustomizableObjectEditor/Private/MuCOE/SMutableCodeViewer.cpp b/Engine/Plugins/Mutable/Source/CustomizableObjectEditor/Private/MuCOE/SMutableCodeViewer.cpp
--- a/Engine/Plugins/Mutable/Source/CustomizableObjectEditor/Private/MuCOE/SMutableCodeViewer.cpp
+++ b/Engine/Plugins/Mutable/Source/CustomizableObjectEditor/Private/MuCOE/SMutableCodeViewer.cpp
@@ -1691,6 +1691,7 @@ void SMutableCodeViewer::GenerateElementRecursive(const int32& InStateIndex, UE:
 
 	case UE::Mutable::Private::EOpType::IN_ADDMESH:
 	case UE::Mutable::Private::EOpType::IN_ADDSTRING:
+	case UE::Mutable::Private::EOpType::IN_ADDCOMPONENT:
 	case UE::Mutable::Private::EOpType::IN_ADDSURFACE:
 	case UE::Mutable::Private::EOpType::IN_ADDSKELETALMESH:
 	{
@@ -1700,14 +1701,6 @@ void SMutableCodeViewer::GenerateElementRecursive(const int32& InStateIndex, UE:
 		break;
 	}
 		
-	case UE::Mutable::Private::EOpType::IN_ADDCOMPONENT:
-	{
-		UE::Mutable::Private::OP::InstanceAddComponentArgs Args = InProgram.GetOpArgs<UE::Mutable::Private::OP::InstanceAddComponentArgs>(InParentAddress);
-		AddOpFunc(Args.Instance, TEXT("instance "));
-		AddOpFunc(Args.Value, TEXT("value "));
-		break;
-	}
-		
 	case UE::Mutable::Private::EOpType::IN_ADDOVERLAYMATERIAL:
 	{
 		UE::Mutable::Private::OP::InstanceAddMaterialArgs Args = InProgram.GetOpArgs<UE::Mutable::Private::OP::InstanceAddMaterialArgs>(InParentAddress);
@@ -2634,6 +2627,10 @@ bool SMutableCodeViewer::IsConstantResourceUsedByOperation(const int32 IndexOnCo
 				{
 					bResourceLocated = IndexOnConstantsArray == InProgram.GetOpArgs<UE::Mutable::Private::OP::InstanceAddArgs>(OperationAddress).name;
 				}
+				else if (OperationType == UE::Mutable::Private::EOpType::IN_ADDCOMPONENT)
+				{
+					bResourceLocated = IndexOnConstantsArray == InProgram.GetOpArgs<UE::Mutable::Private::OP::InstanceAddArgs>(OperationAddress).name;
+				}
 				else if (OperationType == UE::Mutable::Private::EOpType::IN_ADDSURFACE)
 				{
 					bResourceLocated = IndexOnConstantsArray == InProgram.GetOpArgs<UE::Mutable::Private::OP::InstanceAddArgs>(OperationAddress).name;
diff --git a/Engine/Plugins/Mutable/Source/MutableRuntime/Internal/MuR/ComponentId.h b/Engine/Plugins/Mutable/Source/MutableRuntime/Internal/MuR/ComponentId.h
deleted file mode 100644
diff --git a/Engine/Plugins/Mutable/Source/MutableRuntime/Internal/MuR/Instance.h b/Engine/Plugins/Mutable/Source/MutableRuntime/Internal/MuR/Instance.h
--- a/Engine/Plugins/Mutable/Source/MutableRuntime/Internal/MuR/Instance.h
+++ b/Engine/Plugins/Mutable/Source/MutableRuntime/Internal/MuR/Instance.h
@@ -54,7 +54,7 @@ namespace UE::Mutable::Private
 
 	struct FInstanceComponent
 	{
-		FComponentId Id = INDEX_NONE;
+		uint16 Id;
 
 		FMaterialId OverlayMaterialId;
 
@@ -100,7 +100,7 @@ namespace UE::Mutable::Private
 
 		//! Get the Id of a component
 		//! \param ComponentIndex Index of the component, from 0 to GetComponentCount()-1
-		UE_API FComponentId GetComponentId(int32 ComponentIndex) const;
+		UE_API uint16 GetComponentId(int32 ComponentIndex) const;
 
         //! Get the number of LODs in a component
         //! \param ComponentIndex Index of the component, from 0 to GetComponentCount()-1
diff --git a/Engine/Plugins/Mutable/Source/MutableRuntime/Internal/MuR/Operations.h b/Engine/Plugins/Mutable/Source/MutableRuntime/Internal/MuR/Operations.h
--- a/Engine/Plugins/Mutable/Source/MutableRuntime/Internal/MuR/Operations.h
+++ b/Engine/Plugins/Mutable/Source/MutableRuntime/Internal/MuR/Operations.h
@@ -2,7 +2,6 @@
 
 #pragma once
 
-#include "ComponentId.h"
 #include "PassthroughObject.h"
 #include "MuR/ImageTypes.h"
 #include "MuR/Types.h"
@@ -16,7 +15,6 @@
 #define MUTABLE_OP_MAX_INTERPOLATE_COUNT	6
 #define MUTABLE_OP_MAX_SWIZZLE_CHANNELS		4
 
-
 namespace UE::Mutable::Private
 {
 	enum class EImageFormat : uint8;
@@ -1142,13 +1140,6 @@ namespace UE::Mutable::Private
 
             ADDRESS name;
         };
-    	
-    	struct InstanceAddComponentArgs
-    	{
-    		ADDRESS Instance;
-    		ADDRESS Value;
-    		FComponentId ExternalId;
-    	};
 
 		struct InstanceAddExtensionDataArgs
 		{
diff --git a/Engine/Plugins/Mutable/Source/MutableRuntime/Private/MuR/CodeRunner.cpp b/Engine/Plugins/Mutable/Source/MutableRuntime/Private/MuR/CodeRunner.cpp
--- a/Engine/Plugins/Mutable/Source/MutableRuntime/Private/MuR/CodeRunner.cpp
+++ b/Engine/Plugins/Mutable/Source/MutableRuntime/Private/MuR/CodeRunner.cpp
@@ -573,18 +573,18 @@ namespace UE::Mutable::Private
 
         case EOpType::IN_ADDCOMPONENT:
         {
-			OP::InstanceAddComponentArgs Args = Program.GetOpArgs<OP::InstanceAddComponentArgs>(Item.At);
+			OP::InstanceAddArgs Args = Program.GetOpArgs<OP::InstanceAddArgs>(Item.At);
             switch (Item.Stage)
             {
             case 0:
                     AddOp( FScheduledOp( Item.At, Item, 1),
-                           FScheduledOp( Args.Instance, Item),
-                           FScheduledOp( Args.Value, Item) );
+                           FScheduledOp( Args.instance, Item),
+                           FScheduledOp( Args.value, Item) );
                 break;
 
             case 1:
             {
-				TManagedPtr<const FInstance> Base = LoadInstance( FCacheAddress(Args.Instance,Item) );
+				TManagedPtr<const FInstance> Base = LoadInstance( FCacheAddress(Args.instance,Item) );
 				TManagedPtr<FInstance> Result = nullptr;
 
 				if (!Base)
@@ -596,9 +596,9 @@ namespace UE::Mutable::Private
 					Result = UE::Mutable::Private::CloneOrTakeOver<FInstance>(MoveTemp(Base));
 				}
 
-                if ( Args.Value )
+                if ( Args.value )
                 {
-                    TManagedPtr<const FInstance> pComp = LoadInstance( FCacheAddress(Args.Value,Item) );
+                    TManagedPtr<const FInstance> pComp = LoadInstance( FCacheAddress(Args.value,Item) );
 
                     int32 NewComponentIndex = Result->AddComponent();
 
diff --git a/Engine/Plugins/Mutable/Source/MutableRuntime/Private/MuR/CodeRunnerBegin.cpp b/Engine/Plugins/Mutable/Source/MutableRuntime/Private/MuR/CodeRunnerBegin.cpp
--- a/Engine/Plugins/Mutable/Source/MutableRuntime/Private/MuR/CodeRunnerBegin.cpp
+++ b/Engine/Plugins/Mutable/Source/MutableRuntime/Private/MuR/CodeRunnerBegin.cpp
@@ -1043,14 +1043,14 @@ namespace UE::Mutable::Private
 				{
 					check(Item.bEvaluate);
 
-					OP::InstanceAddComponentArgs Args = Program.GetOpArgs<OP::InstanceAddComponentArgs>(Item.At);
+					OP::InstanceAddArgs Args = Program.GetOpArgs<OP::InstanceAddArgs>(Item.At);
 					switch (Item.Stage)
 					{
 					case 0:
 						{
 							PushOp(FScheduledOpInline(Item, 1));
-							PushOp(FScheduledOpInline(Args.Instance, Item));
-							PushOp(FScheduledOpInline(Args.Value, Item));
+							PushOp(FScheduledOpInline(Args.instance, Item));
+							PushOp(FScheduledOpInline(Args.value, Item));
 							
 							break;
 						}
@@ -1058,9 +1058,9 @@ namespace UE::Mutable::Private
 					case 1:
 						{
 							TManagedPtr<const FInstance> InInstance;
-							if (Args.Instance)
+							if (Args.instance)
 							{
-								InInstance = LoadInstance(Args.Instance);
+								InInstance = LoadInstance(Args.instance);
 							}
 
 							TManagedPtr<FInstance> Instance;
@@ -1073,9 +1073,9 @@ namespace UE::Mutable::Private
 								Instance = MakeManaged<FInstance>();
 							}
 
-							if (Args.Value)
+							if (Args.value)
 							{
-								TManagedPtr<const FInstance> pComp = LoadInstance(Args.Value);
+								TManagedPtr<const FInstance> pComp = LoadInstance(Args.value);
 
 								int32 NewComponentIndex = Instance->AddComponent();
 
diff --git a/Engine/Plugins/Mutable/Source/MutableRuntime/Private/MuR/Instance.cpp b/Engine/Plugins/Mutable/Source/MutableRuntime/Private/MuR/Instance.cpp
--- a/Engine/Plugins/Mutable/Source/MutableRuntime/Private/MuR/Instance.cpp
+++ b/Engine/Plugins/Mutable/Source/MutableRuntime/Private/MuR/Instance.cpp
@@ -43,7 +43,7 @@ namespace UE::Mutable::Private
 	}
 	
 	
-	FComponentId FInstance::GetComponentId( int32 ComponentIndex ) const
+	uint16 FInstance::GetComponentId( int32 ComponentIndex ) const
 	{
 		if (Components.IsValidIndex(ComponentIndex))
 		{
@@ -54,7 +54,7 @@ namespace UE::Mutable::Private
 			check(false);
 		}
 
-		return INDEX_NONE;
+		return 0;
 	}
 
 	
diff --git a/Engine/Plugins/Mutable/Source/MutableRuntime/Private/MuR/Operations.cpp b/Engine/Plugins/Mutable/Source/MutableRuntime/Private/MuR/Operations.cpp
--- a/Engine/Plugins/Mutable/Source/MutableRuntime/Private/MuR/Operations.cpp
+++ b/Engine/Plugins/Mutable/Source/MutableRuntime/Private/MuR/Operations.cpp
@@ -900,6 +900,7 @@ namespace UE::Mutable::Private
         //-------------------------------------------------------------------------------------
         case EOpType::IN_ADDMESH:
         case EOpType::IN_ADDSTRING:
+        case EOpType::IN_ADDCOMPONENT:
         case EOpType::IN_ADDSURFACE:
         {
 			OP::InstanceAddArgs args = program.GetOpArgs<OP::InstanceAddArgs>(at);
@@ -908,14 +909,6 @@ namespace UE::Mutable::Private
             break;
         }
 
-        case EOpType::IN_ADDCOMPONENT:
-        {
-        	OP::InstanceAddComponentArgs args = program.GetOpArgs<OP::InstanceAddComponentArgs>(at);
-        	f(args.Instance );
-        	f(args.Value );
-        	break;
-        }
-        	
         case EOpType::IN_ADDLOD:
         {
 			const uint8* Data = program.GetOpArgsPointer(at);
diff --git a/Engine/Plugins/Mutable/Source/MutableTools/Internal/MuT/NodeComponentNew.h b/Engine/Plugins/Mutable/Source/MutableTools/Internal/MuT/NodeComponentNew.h
--- a/Engine/Plugins/Mutable/Source/MutableTools/Internal/MuT/NodeComponentNew.h
+++ b/Engine/Plugins/Mutable/Source/MutableTools/Internal/MuT/NodeComponentNew.h
@@ -2,7 +2,7 @@
 
 #pragma once
 
-#include "MuR/ComponentId.h"
+#include "HAL/Platform.h"
 #include "MuR/Ptr.h"
 #include "MuT/NodeSkeletalMeshNew.h"
 #include "MuT/Node.h"
@@ -19,7 +19,7 @@ namespace UE::Mutable::Private
 	public:
 
 		/** Externally managed id assigned to this component. */
-		FComponentId Id = INDEX_NONE;
+		uint16 Id = 0;
 
 		Ptr<NodeMaterial> OverlayMaterial;
 
diff --git a/Engine/Plugins/Mutable/Source/MutableTools/Private/MuT/ASTOpComponentAdd.cpp b/Engine/Plugins/Mutable/Source/MutableTools/Private/MuT/ASTOpComponentAdd.cpp
deleted file mode 100644
diff --git a/Engine/Plugins/Mutable/Source/MutableTools/Private/MuT/ASTOpComponentAdd.h b/Engine/Plugins/Mutable/Source/MutableTools/Private/MuT/ASTOpComponentAdd.h
deleted file mode 100644
diff --git a/Engine/Plugins/Mutable/Source/MutableTools/Private/MuT/ASTOpInstanceAdd.cpp b/Engine/Plugins/Mutable/Source/MutableTools/Private/MuT/ASTOpInstanceAdd.cpp
--- a/Engine/Plugins/Mutable/Source/MutableTools/Private/MuT/ASTOpInstanceAdd.cpp
+++ b/Engine/Plugins/Mutable/Source/MutableTools/Private/MuT/ASTOpInstanceAdd.cpp
@@ -74,6 +74,7 @@ namespace UE::Mutable::Private
 		case EOpType::IN_ADDMESH:
 		case EOpType::IN_ADDSTRING:
 		case EOpType::IN_ADDSURFACE:
+		case EOpType::IN_ADDCOMPONENT:
 		case EOpType::IN_ADDLOD:
 			break;
 		default:
diff --git a/Engine/Plugins/Mutable/Source/MutableTools/Private/MuT/CodeGenerator.cpp b/Engine/Plugins/Mutable/Source/MutableTools/Private/MuT/CodeGenerator.cpp
--- a/Engine/Plugins/Mutable/Source/MutableTools/Private/MuT/CodeGenerator.cpp
+++ b/Engine/Plugins/Mutable/Source/MutableTools/Private/MuT/CodeGenerator.cpp
@@ -5,7 +5,6 @@
 #include "ASTOpAddOverrideMaterial.h"
 #include "ASTOpAddSkeletalMesh.h"
 #include "ASTOpEntry.h"
-#include "ASTOpComponentAdd.h"
 #include "Containers/Array.h"
 #include "Logging/LogCategory.h"
 #include "Logging/LogMacros.h"
@@ -638,9 +637,10 @@ namespace UE::Mutable::Private
 					LastInstOp = AddOverrideMaterialOp;
 				}
 
-				Ptr<ASTOpComponentAdd> InstanceOp = new ASTOpComponentAdd();
-				InstanceOp->Instance = Options.BaseInstance;
-				InstanceOp->Value = LastInstOp;
+				Ptr<ASTOpInstanceAdd> InstanceOp = new ASTOpInstanceAdd();
+				InstanceOp->type = EOpType::IN_ADDCOMPONENT;
+				InstanceOp->instance = Options.BaseInstance;
+				InstanceOp->value = LastInstOp;
 				InstanceOp->ExternalId = InNode->Id;
 
 				Result.Op = InstanceOp;
