[Mutable] CO Version bump.

#Synced-CL 49557694

diff --git a/Engine/Plugins/Mutable/Source/CustomizableObject/Internal/MuCO/CustomizableObjectPrivate.h b/Engine/Plugins/Mutable/Source/CustomizableObject/Internal/MuCO/CustomizableObjectPrivate.h
--- a/Engine/Plugins/Mutable/Source/CustomizableObject/Internal/MuCO/CustomizableObjectPrivate.h
+++ b/Engine/Plugins/Mutable/Source/CustomizableObject/Internal/MuCO/CustomizableObjectPrivate.h
@@ -1768,8 +1768,12 @@ enum class ECustomizableObjectVersions : int32
 
 	RemovedReuseTextures,
 
-	PassthroughTexturesForBreakNodes,
+	ComponentId,
 
+	PassthroughTexturesForBreakNodes,
+	
+	VersionBumpBackout,
+	
 	// -----<new versions can be added above this line>--------
 	LastCustomizableObjectVersion
 };
