[mutable] Made mutable, able to report an indeterminate amount of nodes causing an issue.
- Reporting the existence of repeated nodes was not working that well, as we were only able to report two nodes. This greatly helps the end user to be able to locate all the instances of a repeated node with just a single CO compilation. Before, the user had to compile the CO multiple times to be able to find all the nodes sharing the same name.
- Fixed some typos and removed unused includes.

#rnx
#Synced-CL 49674194

diff --git a/Engine/Plugins/Mutable/Source/CustomizableObjectEditor/Private/MuCOE/CustomizableObjectCompileRunnable.cpp b/Engine/Plugins/Mutable/Source/CustomizableObjectEditor/Private/MuCOE/CustomizableObjectCompileRunnable.cpp
--- a/Engine/Plugins/Mutable/Source/CustomizableObjectEditor/Private/MuCOE/CustomizableObjectCompileRunnable.cpp
+++ b/Engine/Plugins/Mutable/Source/CustomizableObjectEditor/Private/MuCOE/CustomizableObjectCompileRunnable.cpp
@@ -16,14 +16,11 @@
 #include "MuT/ErrorLog.h"
 #include "MuT/UnrealPixelFormatOverride.h"
 #include "Serialization/MemoryWriter.h"
-#include "Async/Async.h"
-#include "Containers/Ticker.h"
 #include "Serialization/ObjectAndNameAsStringProxyArchive.h"
 #include "Trace/Trace.inl"
 #include "UObject/StrongObjectPtr.h"
 #include "UObject/SoftObjectPtr.h"
 #include "UObject/ICookInfo.h"
-#include "Engine/AssetUserData.h"
 #include "Engine/SkeletalMesh.h"
 #include "Animation/AnimInstance.h"
 
@@ -360,20 +357,22 @@ uint32 FCustomizableObjectCompileRunnable::Run()
 			}
 			}();
 
-			if (MessageAttachedData.UnassignedUVs && MessageAttachedData.UnassignedUVsSize > 0) 
-			{			
-				TSharedPtr<FErrorAttachedData> ErrorAttachedData = MakeShared<FErrorAttachedData>();
-				ErrorAttachedData->UnassignedUVs.Reset();
-				ErrorAttachedData->UnassignedUVs.Append(MessageAttachedData.UnassignedUVs, MessageAttachedData.UnassignedUVsSize);
-				const UObject* Context = static_cast<const UObject*>(pLog->GetMessageContext(i));
-				ArrayErrors.Add(FError(Severity, FText::AsCultureInvariant(Message), ErrorAttachedData, Context, SpamBin));
-			}
-			else
+			const TArray<const void*, TInlineAllocator<4>>* MessageContexts = pLog->GetMessageContexts(i);
+			if (ensure(MessageContexts))
 			{
-				// TODO: Review, and probably propagate the UObject type into the runtime.
-				const UObject* Context = static_cast<const UObject*>(pLog->GetMessageContext(i));
-				const UObject* Context2 = static_cast<const UObject*>(pLog->GetMessageContext2(i));
-				ArrayErrors.Add(FError(Severity, FText::AsCultureInvariant(Message), Context, Context2, SpamBin));
+				if (MessageAttachedData.UnassignedUVs && MessageAttachedData.UnassignedUVsSize > 0) 
+				{			
+					TSharedPtr<FErrorAttachedData> ErrorAttachedData = MakeShared<FErrorAttachedData>();
+					ErrorAttachedData->UnassignedUVs.Reset();
+					ErrorAttachedData->UnassignedUVs.Append(MessageAttachedData.UnassignedUVs, MessageAttachedData.UnassignedUVsSize);
+				
+					const void* Context = MessageContexts->IsValidIndex(0) ? (*MessageContexts)[0] : nullptr;
+					ArrayErrors.Add(FError(Severity, FText::AsCultureInvariant(Message), Context, ErrorAttachedData, SpamBin));
+				}
+				else
+				{
+					ArrayErrors.Add(FError(Severity, FText::AsCultureInvariant(Message), *MessageContexts, SpamBin));
+				}
 			}
 		}
 	}
diff --git a/Engine/Plugins/Mutable/Source/CustomizableObjectEditor/Private/MuCOE/CustomizableObjectCompileRunnable.h b/Engine/Plugins/Mutable/Source/CustomizableObjectEditor/Private/MuCOE/CustomizableObjectCompileRunnable.h
--- a/Engine/Plugins/Mutable/Source/CustomizableObjectEditor/Private/MuCOE/CustomizableObjectCompileRunnable.h
+++ b/Engine/Plugins/Mutable/Source/CustomizableObjectEditor/Private/MuCOE/CustomizableObjectCompileRunnable.h
@@ -36,13 +36,13 @@ public:
 		ELoggerSpamBin SpamBin = ELoggerSpamBin::ShowAll;
 		FText Message;
 		TSharedPtr<FErrorAttachedData> AttachedData;
-		TObjectPtr<const UObject> Context = nullptr;
-		TObjectPtr<const UObject> Context2 = nullptr;
+		
+		TArray<const void*, TInlineAllocator<4>> Contexts;
 
-		FError(const EMessageSeverity::Type InSeverity, const FText& InMessage, const UObject* InContext, const UObject* InContext2=nullptr, const ELoggerSpamBin InSpamBin = ELoggerSpamBin::ShowAll )
-			: Severity(InSeverity), SpamBin(InSpamBin), Message(InMessage), Context(InContext), Context2(InContext2) {}
-		FError(const EMessageSeverity::Type InSeverity, const FText& InMessage, const TSharedPtr<FErrorAttachedData>& InAttachedData, const UObject* InContext, const ELoggerSpamBin InSpamBin = ELoggerSpamBin::ShowAll)
-			: Severity(InSeverity), SpamBin(InSpamBin), Message(InMessage), AttachedData(InAttachedData), Context(InContext) {}
+		FError(const EMessageSeverity::Type InSeverity, const FText& InMessage, const TArray<const void*, TInlineAllocator<4>>& InContexts, const ELoggerSpamBin InSpamBin = ELoggerSpamBin::ShowAll )
+			: Severity(InSeverity), SpamBin(InSpamBin), Message(InMessage), Contexts(InContexts) {}
+		FError(const EMessageSeverity::Type InSeverity, const FText& InMessage, const void* InContext, const TSharedPtr<FErrorAttachedData>& InAttachedData, const ELoggerSpamBin InSpamBin = ELoggerSpamBin::ShowAll)
+			: Severity(InSeverity), SpamBin(InSpamBin), Message(InMessage), AttachedData(InAttachedData), Contexts{InContext} {}
 	};
 
 private:
diff --git a/Engine/Plugins/Mutable/Source/CustomizableObjectEditor/Private/MuCOE/CustomizableObjectCompiler.cpp b/Engine/Plugins/Mutable/Source/CustomizableObjectEditor/Private/MuCOE/CustomizableObjectCompiler.cpp
--- a/Engine/Plugins/Mutable/Source/CustomizableObjectEditor/Private/MuCOE/CustomizableObjectCompiler.cpp
+++ b/Engine/Plugins/Mutable/Source/CustomizableObjectEditor/Private/MuCOE/CustomizableObjectCompiler.cpp
@@ -7,7 +7,6 @@
 #include "ClothConfig.h"
 #include "HAL/RunnableThread.h"
 #include "Interfaces/ITargetPlatformManagerModule.h"
-#include "Materials/MaterialInterface.h"
 #include "MessageLogModule.h"
 #include "GenerateMutableSource/GenerateMutableSourceComponent.h"
 #include "MuCO/CustomizableObject.h"
@@ -1992,25 +1991,24 @@ void FCustomizableObjectCompiler::UpdateCompilerLogData()
 {
 	FMessageLogModule& MessageLogModule = FModuleManager::LoadModuleChecked<FMessageLogModule>("MessageLog");
 	MessageLogModule.RegisterLogListing(FName("Mutable"), LOCTEXT("MutableLog", "Mutable"));
-	const TArray<FCustomizableObjectCompileRunnable::FError>& ArrayCompileErrors = CompileTask->GetArrayErrors();
 
 	const FText ObjectName = CurrentObject ? FText::FromString(CurrentObject->GetName()) : LOCTEXT("Unknown Object", "Unknown Object");
 
+	const TArray<FCustomizableObjectCompileRunnable::FError>& ArrayCompileErrors = CompileTask->GetArrayErrors();
 	for (const FCustomizableObjectCompileRunnable::FError& CompileError : ArrayCompileErrors)
 	{
 		TArray<const UObject*> ObjectArray;
-		if (CompileError.Context)
+		ObjectArray.Reserve(CompileError.Contexts.Num());
+		for (const void* Context : CompileError.Contexts)
 		{
-			ObjectArray.Add(CompileError.Context);
-		}
-		if (CompileError.Context2)
-		{
-			ObjectArray.Add(CompileError.Context2);
+			ObjectArray.Add(static_cast<const UObject*>(Context));
 		}
 
-		if (CompileError.Context && CompileError.AttachedData)
+		const UObject* CompileErrorContext = ObjectArray.IsValidIndex(0) ? ObjectArray[0] : nullptr; 
+		
+		if (CompileErrorContext && CompileError.AttachedData)
 		{
-			if (const UCustomizableObjectNode* Node = Cast<UCustomizableObjectNode>(CompileError.Context))
+			if (const UCustomizableObjectNode* Node = Cast<UCustomizableObjectNode>(CompileErrorContext))
 			{
 				UCustomizableObjectNode::FAttachedErrorDataView ErrorDataView;
 				ErrorDataView.UnassignedUVs = { CompileError.AttachedData->UnassignedUVs.GetData(),
diff --git a/Engine/Plugins/Mutable/Source/MutableTools/Internal/MuT/ErrorLog.h b/Engine/Plugins/Mutable/Source/MutableTools/Internal/MuT/ErrorLog.h
--- a/Engine/Plugins/Mutable/Source/MutableTools/Internal/MuT/ErrorLog.h
+++ b/Engine/Plugins/Mutable/Source/MutableTools/Internal/MuT/ErrorLog.h
@@ -37,7 +37,7 @@ namespace UE::Mutable::Private
 	};
 
 	/** Storage for mutable error, warning and information messages from several processes performed by the tools library, like compilation. 
-	* It support concurrent addition of messages, but not retrieval.
+	* It supports concurrent addition of messages, but not retrieval.
 	*/	
 	class FErrorLog
 	{
@@ -54,8 +54,8 @@ namespace UE::Mutable::Private
 			ErrorLogMessageSpamBin Spam = ELMSB_ALL;
 			FString Text;
 			TSharedPtr<FErrorData> Data;
-			const void* Context = nullptr;
-			const void* Context2 = nullptr;
+			
+			TArray<const void*, TInlineAllocator<4>> Contexts;
 		};
 
 	private:
@@ -79,10 +79,9 @@ namespace UE::Mutable::Private
 		//! \param index index of the message from 0 to GetMessageCount()-1
 		UE_API const FString& GetMessageText( int32 index ) const;
 
-		//! Get the opaque context of a message.
-		//! \param index index of the message from 0 to GetMessageCount()-1
-		UE_API const void* GetMessageContext(int32 index) const;
-		UE_API const void* GetMessageContext2(int32 index) const;
+		//! Get the opaque contexts of a message.
+		//! \param Index index of the message from 0 to GetMessageCount()-1
+		UE_API const TArray<const void*, TInlineAllocator<4>>* GetMessageContexts(int32 Index) const;
 
 		//!
 		UE_API ErrorLogMessageType GetMessageType( int32 index ) const;
@@ -103,7 +102,8 @@ namespace UE::Mutable::Private
 		//!
 		UE_API void Add(const FString& Message, ErrorLogMessageType Type, const void* Context, ErrorLogMessageSpamBin SpamBin = ELMSB_ALL);
 		UE_API void Add(const FString& Message, ErrorLogMessageType Type, const void* Context, const void* Context2, ErrorLogMessageSpamBin SpamBin = ELMSB_ALL);
-
+		UE_API void Add(const FString& Message, ErrorLogMessageType Type, TArray<const void*, TInlineAllocator<4>>& Contexts, ErrorLogMessageSpamBin SpamBin = ELMSB_ALL);
+		
 		//!
 		UE_API void Add(const FString& Message, const ErrorLogMessageAttachedDataView& Data, ErrorLogMessageType Type, const void* Context, ErrorLogMessageSpamBin SpamBin = ELMSB_ALL);
 
diff --git a/Engine/Plugins/Mutable/Source/MutableTools/Private/MuT/CodeGenerator.cpp b/Engine/Plugins/Mutable/Source/MutableTools/Private/MuT/CodeGenerator.cpp
--- a/Engine/Plugins/Mutable/Source/MutableTools/Private/MuT/CodeGenerator.cpp
+++ b/Engine/Plugins/Mutable/Source/MutableTools/Private/MuT/CodeGenerator.cpp
@@ -6,7 +6,6 @@
 #include "ASTOpAddSkeletalMesh.h"
 #include "ASTOpEntry.h"
 #include "Containers/Array.h"
-#include "Logging/LogCategory.h"
 #include "Logging/LogMacros.h"
 #include "Math/IntPoint.h"
 #include "Math/UnrealMathUtility.h"
@@ -16,7 +15,6 @@
 #include "MuR/Operations.h"
 #include "MuR/Parameters.h"
 #include "MuR/Platform.h"
-#include "MuT/ASTOpAddLOD.h"
 #include "MuT/ASTOpAddExtensionData.h"
 #include "MuT/ASTOpAddMaterial.h"
 #include "MuT/ASTOpConditional.h"
@@ -24,17 +22,8 @@
 #include "MuT/ASTOpConstantResource.h"
 #include "MuT/ASTOpBoolAnd.h"
 #include "MuT/ASTOpImageCompose.h"
-#include "MuT/ASTOpImageMipmap.h"
-#include "MuT/ASTOpImagePixelFormat.h"
-#include "MuT/ASTOpImageSwizzle.h"
 #include "MuT/ASTOpImageLayer.h"
-#include "MuT/ASTOpImageLayerColor.h"
-#include "MuT/ASTOpImagePatch.h"
-#include "MuT/ASTOpImageCrop.h"
-#include "MuT/ASTOpImageBlankLayout.h"
-#include "MuT/ASTOpImagePlainColor.h"
 #include "MuT/ASTOpInstanceAdd.h"
-#include "MuT/ASTOpMaterialModify.h"
 #include "MuT/ASTOpMeshBindShape.h"
 #include "MuT/ASTOpMeshClipDeform.h"
 #include "MuT/ASTOpMeshClipMorphPlane.h"
@@ -50,14 +39,11 @@
 #include "MuT/ASTOpMeshMaskDiff.h"
 #include "MuT/ASTOpMeshTransformWithBoundingMesh.h"
 #include "MuT/ASTOpMeshTransformWithBone.h"
-#include "MuT/ASTOpParameter.h"
 #include "MuT/ASTOpLayoutRemoveBlocks.h"
 #include "MuT/ASTOpLayoutFromMesh.h"
 #include "MuT/ASTOpLayoutMerge.h"
 #include "MuT/ASTOpLayoutPack.h"
-#include "MuT/ASTOpReferenceResource.h"
 #include "MuT/CodeGenerator_SecondPass.h"
-#include "MuT/CodeOptimiser.h"
 #include "MuT/CompilerPrivate.h"
 #include "MuT/ErrorLog.h"
 #include "MuT/NodeColour.h"
@@ -67,17 +53,7 @@
 #include "MuT/NodeComponentVariation.h"
 #include "MuT/NodeImage.h"
 #include "MuT/NodeImageConstant.h"
-#include "MuT/NodeImageFormat.h"
-#include "MuT/NodeImageMipmap.h"
-#include "MuT/NodeImageSwizzle.h"
-#include "MuT/NodeMatrixConstant.h"
 #include "MuT/NodeMesh.h"
-#include "MuT/NodeMeshClipMorphPlane.h"
-#include "MuT/NodeMeshClipWithMesh.h"
-#include "MuT/NodeMeshFormat.h"
-#include "MuT/NodeMeshFragment.h"
-#include "MuT/NodeMeshMorph.h"
-#include "MuT/NodeMeshReshape.h"
 #include "MuT/NodeMeshConstant.h"
 #include "MuT/NodeModifier.h"
 #include "MuT/NodeModifierMeshClipDeform.h"
@@ -94,12 +70,10 @@
 #include "MuT/NodeScalar.h"
 #include "MuT/NodeScalarConstant.h"
 #include "MuT/NodeScalarEnumParameter.h"
-#include "MuT/NodeSurface.h"
 #include "MuT/NodeSurfaceNew.h"
 #include "MuT/NodeSurfaceVariation.h"
 #include "MuT/NodeSurfaceSwitch.h"
 #include "MuT/TablePrivate.h"
-#include "Trace/Detail/Channel.h"
 
 
 namespace UE::Mutable::Private
@@ -1788,7 +1762,7 @@ namespace UE::Mutable::Private
 								if (BaseMeshResult.GeneratedLayouts.Num() != AddResults.GeneratedLayouts.Num())
 								{
 									// When extending a mesh section the added mesh section will use the layout strategy of the base one
-									FString Msg = FString::Printf(TEXT("Extended mesh section layout count is differenta than the mesh being extended."));
+									const FString Msg = FString::Printf(TEXT("Extended mesh section layout count is different than the mesh being extended."));
 									ErrorLog->Add(Msg, ELMT_INFO, Edit->GetMessageContext(), ErrorContext);
 								}
 
@@ -1809,7 +1783,7 @@ namespace UE::Mutable::Private
 										&&
 										BaseMeshResult.GeneratedLayouts[LayoutIndex].Layout->Strategy != AddResults.GeneratedLayouts[LayoutIndex].Layout->Strategy)
 									{
-										// When extending a mesh section the added mesh sectio will use the layout strategy of the base one
+										// When extending a mesh section the added mesh section will use the layout strategy of the base one
 										FString Msg = FString::Printf(TEXT("Extended mesh section layout [%d] is using a different strategy than the section being extended. The base strategy will be used."), LayoutIndex);
 										ErrorLog->Add(Msg, ELMT_INFO, Edit->GetMessageContext(), ErrorContext);
 									}
diff --git a/Engine/Plugins/Mutable/Source/MutableTools/Private/MuT/Compiler.cpp b/Engine/Plugins/Mutable/Source/MutableTools/Private/MuT/Compiler.cpp
--- a/Engine/Plugins/Mutable/Source/MutableTools/Private/MuT/Compiler.cpp
+++ b/Engine/Plugins/Mutable/Source/MutableTools/Private/MuT/Compiler.cpp
@@ -20,13 +20,11 @@
 #include "MuR/Serialisation.h"
 #include "MuR/MutableRuntimeModule.h"
 
-#include "Trace/Detail/Channel.h"
 #include "Containers/Array.h"
 #include "Containers/Map.h"
 #include "HAL/LowLevelMemTracker.h"
 #include "HAL/PlatformCrt.h"
 #include "HAL/PlatformTime.h"
-#include "Hash/CityHash.h"
 #include "Logging/LogCategory.h"
 #include "Logging/LogMacros.h"
 #include "Misc/EnumClassFlags.h"
@@ -516,7 +514,7 @@ namespace UE::Mutable::Private
 							{
 								// Report repeated parameter name alongside with the nodes that use it
 								const FString WarningText = FString::Printf(TEXT("Repeated parameter found : \"%s\". Please use a different name."), *CurrentParameterName);
-								GenErrorLog->Add(WarningText, ELMT_ERROR, MessageContexts[0], MessageContexts[1]);
+								GenErrorLog->Add(WarningText, ELMT_ERROR, MessageContexts);
 								RepeatedParameterNames.Add(CurrentParameterName);
 							}
 						}
diff --git a/Engine/Plugins/Mutable/Source/MutableTools/Private/MuT/ErrorLog.cpp b/Engine/Plugins/Mutable/Source/MutableTools/Private/MuT/ErrorLog.cpp
--- a/Engine/Plugins/Mutable/Source/MutableTools/Private/MuT/ErrorLog.cpp
+++ b/Engine/Plugins/Mutable/Source/MutableTools/Private/MuT/ErrorLog.cpp
@@ -3,25 +3,16 @@
 #include "MuT/ErrorLog.h"
 
 #include "Containers/Array.h"
-#include "Containers/Map.h"
 #include "Containers/UnrealString.h"
-#include "HAL/PlatformMath.h"
-#include "HAL/UnrealMemory.h"
 #include "Logging/LogCategory.h"
 #include "Logging/LogMacros.h"
 #include "Misc/AssertionMacros.h"
 #include "MuR/Layout.h"
 #include "MuR/Mesh.h"
 #include "MuR/Model.h"
-#include "MuR/MutableMath.h"
 #include "MuR/Operations.h"
-#include "MuR/Parameters.h"
-#include "MuR/Platform.h"
-#include "MuR/Skeleton.h"
 #include "MuR/System.h"
 #include "MuR/MutableRuntimeModule.h"
-#include "Templates/Tuple.h"
-#include "Trace/Detail/Channel.h"
 
 
 namespace UE::Mutable::Private
@@ -45,31 +36,16 @@ namespace UE::Mutable::Private
 	}
 
 
-	const void* FErrorLog::GetMessageContext(int32 Index) const
+	const TArray<const void*, TInlineAllocator<4>>* FErrorLog::GetMessageContexts(int32 Index) const
 	{
-		const void* Result = 0;
-
 		if (Messages.IsValidIndex(Index))
 		{
-			Result = Messages[Index].Context;
+    		return &Messages[Index].Contexts;
 		}
 
-		return Result;
-	}
-
-
-	const void* FErrorLog::GetMessageContext2(int32 Index) const
-	{
-		const void* Result = 0;
-
-		if (Messages.IsValidIndex(Index))
-		{
-			Result = Messages[Index].Context2;
+    	return nullptr;
 		}
 
-		return Result;
-	}
-
 
 	ErrorLogMessageType FErrorLog::GetMessageType( int32 Index ) const
 	{
@@ -127,7 +103,7 @@ namespace UE::Mutable::Private
 		Msg.Type = InType;
 		Msg.Spam = InSpamBin;
 		Msg.Text = InMessage;
-		Msg.Context = InContext;
+		Msg.Contexts.Add(InContext);
 	}
 
 
@@ -143,8 +119,23 @@ namespace UE::Mutable::Private
 		Msg.Type = InType;
 		Msg.Spam = InSpamBin;
 		Msg.Text = InMessage;
-		Msg.Context = InContext;
-		Msg.Context2 = InContext2;
+    	Msg.Contexts.Add(InContext);
+    	Msg.Contexts.Add(InContext2);
+	}
+	
+	
+	void FErrorLog::Add(const FString& InMessage,
+		ErrorLogMessageType InType,
+		TArray<const void*, TInlineAllocator<4>>& Contexts,
+		ErrorLogMessageSpamBin InSpamBin)
+    {
+    	UE::TUniqueLock Lock(MessageMutex);
+
+    	FMessage& Msg = Messages.Emplace_GetRef();
+    	Msg.Type = InType;
+    	Msg.Spam = InSpamBin;
+    	Msg.Text = InMessage;
+    	Msg.Contexts = MoveTemp(Contexts);
 	}
 
 
@@ -160,7 +151,7 @@ namespace UE::Mutable::Private
 		Msg.Type = InType;
 		Msg.Spam = InSpamBin;
 		Msg.Text = InMessage;
-		Msg.Context = InContext;
+		Msg.Contexts = {InContext};
 		Msg.Data = MakeShared<FErrorData>();
 
         if ( InDataView.UnassignedUVs && InDataView.UnassignedUVsSize > 0 )
