[Mutable] PR #13910: Make FCustomizableObjectInstanceDescriptor Load/Save functions safer.

- Now non-compact serialization does not assume the parameter type nor if it is multidimensional.

#Synced-CL 49675604

diff --git a/Engine/Plugins/Mutable/Source/CustomizableObject/Private/MuCO/CustomizableObjectInstanceDescriptor.cpp b/Engine/Plugins/Mutable/Source/CustomizableObject/Private/MuCO/CustomizableObjectInstanceDescriptor.cpp
--- a/Engine/Plugins/Mutable/Source/CustomizableObject/Private/MuCO/CustomizableObjectInstanceDescriptor.cpp
+++ b/Engine/Plugins/Mutable/Source/CustomizableObject/Private/MuCO/CustomizableObjectInstanceDescriptor.cpp
@@ -81,6 +81,7 @@ void FCustomizableObjectInstanceDescriptor::SaveDescriptor(FArchive& Ar, bool bU
 		{
 			check(Ar.IsSaving());
 			Ar << const_cast<FString &>(Name);
+			Ar << Type;
 		}
  
 		switch (Type)
@@ -164,7 +165,8 @@ void FCustomizableObjectInstanceDescriptor::SaveDescriptor(FArchive& Ar, bool bU
 			else
 			{
 				Ar << ValueName;
- 
+				Ar << bIsParamMultidimensional;
+				
 				if (bIsParamMultidimensional)
 				{
 					Ar << ValueNames;
@@ -305,6 +307,12 @@ void FCustomizableObjectInstanceDescriptor::SaveDescriptor(FArchive& Ar, bool bU
  
 			break;
 		}
+			
+		case EMutableParameterType::None:
+		{
+			UE_LOG(LogMutable, Warning, TEXT("Found unknown parameter [%s] when saving a Customizable Object Instance Descriptor"), *Name);
+			break;
+		}
  
 		default:
 			// Parameter type replication not implemented.
@@ -350,7 +358,7 @@ void FCustomizableObjectInstanceDescriptor::LoadDescriptor(FArchive& Ar)
 		else
 		{
 			Ar << Name;
-			Type = CustomizableObject->GetParameterTypeByName(Name);
+			Ar << Type;
 		}
  
 		switch (Type)
@@ -396,12 +404,13 @@ void FCustomizableObjectInstanceDescriptor::LoadDescriptor(FArchive& Ar)
 			TArray<int32> Values;
 			TArray<FString> ValueNames;
  
-			const int32 IntParameterIndex = FindTypedParameterIndex(Name, EMutableParameterType::Int);
-			const int32 ParameterIndexInObject = CustomizableObject->GetPrivate()->FindParameter(IntParameters[IntParameterIndex].ParameterName);
-			const bool bIsParamMultidimensional = CustomizableObject->GetPrivate()->IsParameterMultidimensional(ParameterIndexInObject);
  
 			if (bUseCompactDescriptor)
 			{
+				const int32 IntParameterIndex = FindTypedParameterIndex(Name, EMutableParameterType::Int);
+				const int32 ParameterIndexInObject = CustomizableObject->GetPrivate()->FindParameter(IntParameters[IntParameterIndex].ParameterName);
+				const bool bIsParamMultidimensional = CustomizableObject->GetPrivate()->IsParameterMultidimensional(ParameterIndexInObject);
+
 				Ar << Value;
  
 				if (bIsParamMultidimensional)
@@ -411,8 +420,11 @@ void FCustomizableObjectInstanceDescriptor::LoadDescriptor(FArchive& Ar)
 			}
 			else
 			{
+				bool bIsParamMultidimensional;
+
 				Ar << ValueName;
- 
+				Ar << bIsParamMultidimensional;
+				
 				if (bIsParamMultidimensional)
 				{
 					Ar << ValueNames;
@@ -550,6 +562,12 @@ void FCustomizableObjectInstanceDescriptor::LoadDescriptor(FArchive& Ar)
  
 			break;
 		}
+			
+		case EMutableParameterType::None:
+		{
+			UE_LOG(LogMutable, Warning, TEXT("Found unknown parameter [%s] when loading a Customizable Object Instance Descriptor"), *Name);
+			break;
+		}
  
 		default:
 			// Parameter type replication not implemented.
