[Mutable] Fixed problem with COI Asset User Data.

Now all animation data from a COI is is merged into the same Asset User Data.
#rnx

#Synced-CL 49675872

diff --git a/Engine/Plugins/Mutable/Source/CustomizableObject/Private/MuCO/CustomizableObjectInstance.cpp b/Engine/Plugins/Mutable/Source/CustomizableObject/Private/MuCO/CustomizableObjectInstance.cpp
--- a/Engine/Plugins/Mutable/Source/CustomizableObject/Private/MuCO/CustomizableObjectInstance.cpp
+++ b/Engine/Plugins/Mutable/Source/CustomizableObject/Private/MuCO/CustomizableObjectInstance.cpp
@@ -108,6 +108,43 @@ UTexture2D* CreateTexture(const FName& TextureName)
 }
 
 
+void AddComponentAnimBPData(FCustomizableInstanceComponentData* ComponentData, UCustomizableInstancePrivate* InstancePrivate, TArray<FSoftObjectPath>& AssetsToStream, const FCustomizableObjectAnimationSlot& AnimSlot, const TSharedRef<FUpdateContextPrivate>& Context)
+{
+	check(ComponentData && InstancePrivate);
+
+	if (!AnimSlot.AnimInstance.IsNull())
+	{
+		const TSoftClassPtr<UAnimInstance>* FoundAnimBpSlot = ComponentData->AnimSlotToBP.Find(AnimSlot.Name);
+		bool bIsSameAnimBp = FoundAnimBpSlot && AnimSlot.AnimInstance == *FoundAnimBpSlot;
+		if (!FoundAnimBpSlot)
+		{
+			ComponentData->AnimSlotToBP.Add(AnimSlot.Name, AnimSlot.AnimInstance);
+
+			if (AnimSlot.AnimInstance.Get())
+			{
+				InstancePrivate->GatheredAnimBPs.Add(AnimSlot.AnimInstance.Get());
+			}
+			else
+			{
+				AssetsToStream.AddUnique(AnimSlot.AnimInstance.ToSoftObjectPath());
+			}
+		}
+		else if (!bIsSameAnimBp)
+		{
+			// Two submeshes should not have the same animation slot index
+			Context->UpdateResult = EUpdateResult::Warning;
+
+			FString WarningMessage = FString::Printf(TEXT("Two submeshes have the same anim slot index [%s] in a Mutable Instance."), *AnimSlot.Name.ToString());
+			UE_LOG(LogMutable, Warning, TEXT("%s"), *WarningMessage);
+#if WITH_EDITOR
+			FMessageLog MessageLog("Mutable");
+			MessageLog.Notify(FText::FromString(WarningMessage), EMessageSeverity::Warning, true);
+#endif
+		}
+	}
+}
+
+
 void UCustomizableInstancePrivate::InvalidateGeneratedData()
 {
 	SkeletalMeshStatus = ESkeletalMeshStatus::NotGenerated;
@@ -5092,17 +5129,34 @@ UE::Tasks::FTask LoadAdditionalAssetsAndData(const TSharedRef<FUpdateContextPriv
 			// Add AssetUserData from Mesh Parameters
 			for (const TStrongObjectPtr<USkeletalMesh>& StrongMesh : MutableMesh->SkeletalMeshes)
 			{
-				if (USkeletalMesh* SkeletalMesh = StrongMesh.Get())
+				USkeletalMesh* SkeletalMesh = StrongMesh.Get();
+
+				if (!SkeletalMesh)
 				{
-					if (const TArray<UAssetUserData*>* AssetUserDataArray = SkeletalMesh->GetAssetUserDataArray())
+					continue;
+				}
+
+				if (const TArray<UAssetUserData*>* AssetUserDataArray = SkeletalMesh->GetAssetUserDataArray())
+				{
+					for (UAssetUserData* AssetUserData : *AssetUserDataArray)
 					{
-						for (UAssetUserData* AssetUserData : *AssetUserDataArray)
+						if (!AssetUserData)
 						{
-							// No need to load since the SKM is a StrongObjectPtr
-							if (AssetUserData)
+							continue;
+						}
+
+						if (const UCustomizableObjectInstanceUserData* COIAssetUserData = Cast<UCustomizableObjectInstanceUserData>(AssetUserData))
+						{
+							for (const FCustomizableObjectAnimationSlot& AnimSlot : COIAssetUserData->AnimationSlots)
 							{
-								ComponentData->AssetUserDataArray.Add(AssetUserData);
+								AddComponentAnimBPData(ComponentData, InstancePrivate, AssetsToStream, AnimSlot, Context);
 							}
+
+							InstancePrivate->AnimBPGameplayTags.AppendTags(COIAssetUserData->AnimationGameplayTag);
+						}
+						else
+						{
+							ComponentData->AssetUserDataArray.Add(AssetUserData);
 						}
 					}
 				}
@@ -5147,40 +5201,7 @@ UE::Tasks::FTask LoadAdditionalAssetsAndData(const TSharedRef<FUpdateContextPriv
 							continue;
 						}
 
-						FName SlotIndex = *SlotIndexString;
-
-						const TSoftClassPtr<UAnimInstance>& AnimBPAsset = ModelResources.AnimBPs[AnimBpIndex];
-
-						if (!AnimBPAsset.IsNull())
-						{
-							const TSoftClassPtr<UAnimInstance>* FoundAnimBpSlot = ComponentData->AnimSlotToBP.Find(SlotIndex);
-							bool bIsSameAnimBp = FoundAnimBpSlot && AnimBPAsset == *FoundAnimBpSlot;
-							if (!FoundAnimBpSlot)
-							{
-								ComponentData->AnimSlotToBP.Add(SlotIndex, AnimBPAsset);
-
-								if (AnimBPAsset.Get())
-								{
-									InstancePrivate->GatheredAnimBPs.Add(AnimBPAsset.Get());
-								}
-								else
-								{
-									AssetsToStream.AddUnique(AnimBPAsset.ToSoftObjectPath());
-								}
-							}
-							else if (!bIsSameAnimBp)
-							{
-								// Two submeshes should not have the same animation slot index
-								Context->UpdateResult = EUpdateResult::Warning;
-
-								FString WarningMessage = FString::Printf(TEXT("Two submeshes have the same anim slot index [%s] in a Mutable Instance."), *SlotIndex.ToString());
-								UE_LOG(LogMutable, Warning, TEXT("%s"), *WarningMessage);
-#if WITH_EDITOR
-								FMessageLog MessageLog("Mutable");
-								MessageLog.Notify(FText::FromString(WarningMessage), EMessageSeverity::Warning, true);
-#endif
-							}
-						}
+						AddComponentAnimBPData(ComponentData, InstancePrivate, AssetsToStream, FCustomizableObjectAnimationSlot{ *SlotIndexString, ModelResources.AnimBPs[AnimBpIndex] }, Context);
 					}
 				}
 				else if (Tag.RemoveFromStart("__AnimBPTag:"))
