[mutable] Mesh section pins no longer need to use the property change callback.
- Removed custom callback in favor of a more dynamic approach for the Enum and Mesh Section nodes
- Added option content validation for the Enum nodes as that was missing (we did only check for the amount of options)
- Marked the enum nodes as "needs refresh" when the amount or content of the options does not match.
- GenerateMutableSource includes purge.
- Other small fixes and tweaks.

#rnx
#Synced-CL 49674194

diff --git a/Engine/Plugins/Mutable/Source/CustomizableObject/Internal/MuCO/CustomizableObjectCustomVersion.h b/Engine/Plugins/Mutable/Source/CustomizableObject/Internal/MuCO/CustomizableObjectCustomVersion.h
--- a/Engine/Plugins/Mutable/Source/CustomizableObject/Internal/MuCO/CustomizableObjectCustomVersion.h
+++ b/Engine/Plugins/Mutable/Source/CustomizableObject/Internal/MuCO/CustomizableObjectCustomVersion.h
@@ -213,6 +213,8 @@ struct FCustomizableObjectCustomVersion
 		NodeStaticStringRemoveHiddenInputPin,
 		
 		PassthroughParameter,
+		
+		MeshSectionNodeInputPinsAsReferences,
 
 		// -----<new versions can be added above this line>-------------------------------------------------
 		VersionPlusOne,
diff --git a/Engine/Plugins/Mutable/Source/CustomizableObjectEditor/Private/MuCOE/GenerateMutableSource/GenerateMutableSource.cpp b/Engine/Plugins/Mutable/Source/CustomizableObjectEditor/Private/MuCOE/GenerateMutableSource/GenerateMutableSource.cpp
--- a/Engine/Plugins/Mutable/Source/CustomizableObjectEditor/Private/MuCOE/GenerateMutableSource/GenerateMutableSource.cpp
+++ b/Engine/Plugins/Mutable/Source/CustomizableObjectEditor/Private/MuCOE/GenerateMutableSource/GenerateMutableSource.cpp
@@ -902,6 +902,38 @@ uint32 AddUniqueSurfaceMetadata(const FMutableSourceSurfaceMetadata& SourceMetad
 }
 
 
+bool DoOptionsMatchEnum(const UCustomizableObjectNodeSwitchBase& InEditorSwitchNode,
+	const UE::Mutable::Private::NodeScalarEnumParameter& MutableScalarParameter)
+{
+	// todo: use the actual UEditor CO node once the issue UE-294591 is resolved. Currently we do need to use the core node as GenerateMutableSource is
+	// capable of navigating over the macro objects while FindInputPin is not
+	
+	const int32 EditorNodeSwitchPinsCount = InEditorSwitchNode.SwitchPins.Num();
+	if (MutableScalarParameter.Options.Num() != EditorNodeSwitchPinsCount)
+	{
+		return false;
+	}
+		
+	// Compare the name of the enum options. It expects to be able to compare the name of the switch pin with the value of the Enum Parameter
+	for (int32 ValueIndex = 0; ValueIndex < EditorNodeSwitchPinsCount; ++ValueIndex)
+	{
+		FString SwitchNodeValueName;
+		const UEdGraphPin* SwitchPin = InEditorSwitchNode.SwitchPins[ValueIndex].Get();
+		check(SwitchPin)
+
+		SwitchNodeValueName = SwitchPin->PinName.ToString();
+		const FString EnumValueName = MutableScalarParameter.Options[ValueIndex].Name;
+
+		if (EnumValueName.Compare(SwitchNodeValueName) != 0)
+		{
+			return false;
+		}
+	}
+				
+	return true;
+}
+
+
 // Convert a CustomizableObject Source Graph into a mutable source graph  
 UE::Mutable::Private::Ptr<UE::Mutable::Private::NodeObject> GenerateMutableSource(const UEdGraphPin * Pin, FMutableGraphGenerationContext & GenerationContext)
 {
diff --git a/Engine/Plugins/Mutable/Source/CustomizableObjectEditor/Private/MuCOE/GenerateMutableSource/GenerateMutableSource.h b/Engine/Plugins/Mutable/Source/CustomizableObjectEditor/Private/MuCOE/GenerateMutableSource/GenerateMutableSource.h
--- a/Engine/Plugins/Mutable/Source/CustomizableObjectEditor/Private/MuCOE/GenerateMutableSource/GenerateMutableSource.h
+++ b/Engine/Plugins/Mutable/Source/CustomizableObjectEditor/Private/MuCOE/GenerateMutableSource/GenerateMutableSource.h
@@ -44,6 +44,7 @@
 
 class UCustomizableObjectNodeGroupProjectorParameter;
 class UCONodeSkeletalMeshSectionBase;
+class UCustomizableObjectNodeSwitchBase;
 class UCustomizableObjectNodeModifierClipWithMesh;
 struct FCustomizableObjectClothingAssetData;
 
@@ -1099,3 +1100,10 @@ uint32 AddUniqueSurfaceMetadata(const FMutableSourceSurfaceMetadata& MeshData, F
 */
 UE::Tasks::FTask ConvertSkeletalMeshToMutable(TSharedRef<FMeshConversionContext> MeshConversionContext);
 
+/**
+ * Compares the Amount and names of the options exposed in the provided Enum Parameter Node with the options set in this node.
+* @param MutableScalarParameter The enum parameter node we want to compare the options against.
+ * @return true if the amount of options does match and false otherwise.
+ */
+bool DoOptionsMatchEnum(const UCustomizableObjectNodeSwitchBase& InEditorSwitchNode, const UE::Mutable::Private::NodeScalarEnumParameter& MutableScalarParameter);
+
diff --git a/Engine/Plugins/Mutable/Source/CustomizableObjectEditor/Private/MuCOE/GenerateMutableSource/GenerateMutableSourceColor.cpp b/Engine/Plugins/Mutable/Source/CustomizableObjectEditor/Private/MuCOE/GenerateMutableSource/GenerateMutableSourceColor.cpp
--- a/Engine/Plugins/Mutable/Source/CustomizableObjectEditor/Private/MuCOE/GenerateMutableSource/GenerateMutableSourceColor.cpp
+++ b/Engine/Plugins/Mutable/Source/CustomizableObjectEditor/Private/MuCOE/GenerateMutableSource/GenerateMutableSourceColor.cpp
@@ -23,6 +23,7 @@
 #include "MuCOE/Nodes/CustomizableObjectNodeColorToSRGB.h"
 #include "MuCOE/Nodes/CustomizableObjectNodeMacroInstance.h"
 #include "MuCOE/Nodes/CONodeMaterialBreak.h"
+#include "MuCOE/Nodes/CustomizableObjectNodeEnumParameter.h"
 #include "MuCOE/Nodes/CustomizableObjectNodeTable.h"
 #include "MuCOE/Nodes/CustomizableObjectNodeTextureSample.h"
 #include "MuCOE/Nodes/CustomizableObjectNodeTunnel.h"
@@ -100,7 +101,7 @@ UE::Mutable::Private::Ptr<UE::Mutable::Private::NodeColour> GenerateMutableSourc
 				NumParameters != 1)
 			{
 				const FText Message = NumParameters
-					? LOCTEXT("NoEnumParamInSwitch", "Switch nodes must have an enum switch parameter. Please connect an enum and refesh the switch node.")
+					? LOCTEXT("NoEnumParamInSwitch", "Switch nodes must have an enum switch parameter. Please connect an enum and refresh the switch node.")
 					: LOCTEXT("InvalidEnumInSwitch", "Switch nodes must have a single enum with all the options inside. Please remove all the enums but one and refresh the switch node.");
 
 				GenerationContext.Log(Message, Node);
@@ -116,7 +117,7 @@ UE::Mutable::Private::Ptr<UE::Mutable::Private::NodeColour> GenerateMutableSourc
 				// Warn about a failure.
 				if (EnumPin)
 				{
-					const FText Message = LOCTEXT("FailedToGenerateSwitchParam", "Could not generate switch enum parameter. Please refesh the switch node and connect an enum.");
+					const FText Message = LOCTEXT("FailedToGenerateSwitchParam", "Could not generate switch enum parameter. Please refresh the switch node and connect an enum.");
 					GenerationContext.Log(Message, Node);
 				}
 
@@ -130,16 +131,20 @@ UE::Mutable::Private::Ptr<UE::Mutable::Private::NodeColour> GenerateMutableSourc
 
 				return nullptr;
 			}
-
-			const int32 NumSwitchOptions = TypedNodeColorSwitch->SwitchPins.Num();
-
-			UE::Mutable::Private::NodeScalarEnumParameter* EnumParameter = static_cast<UE::Mutable::Private::NodeScalarEnumParameter*>(SwitchParam.get());
-			if (NumSwitchOptions != EnumParameter->Options.Num())
+				
 			{
-				const FText Message = LOCTEXT("MismatchedSwitch", "Switch enum and switch node have different number of options. Please refresh the switch node to make sure the outcomes are labeled properly.");
-				GenerationContext.Log(Message, Node);
+				const UE::Mutable::Private::NodeScalarEnumParameter* EnumParameter = static_cast<UE::Mutable::Private::NodeScalarEnumParameter*>(SwitchParam.get());
+				
+				if (!DoOptionsMatchEnum(*TypedNodeColorSwitch, *EnumParameter))
+				{
+					const FText Message = LOCTEXT("MismatchedSwitch", "Switch enum and switch node have different options. Please refresh the switch node to make sure the outcomes are labeled properly.");
+					GenerationContext.Log(Message, Node);
+					Node->SetRefreshNodeWarning();
+				}
 			}
 
+			const int32 NumSwitchOptions = TypedNodeColorSwitch->SwitchPins.Num();
+				
 			UE::Mutable::Private::Ptr<UE::Mutable::Private::NodeColourSwitch> SwitchNode = new UE::Mutable::Private::NodeColourSwitch;
 			SwitchNode->Parameter = SwitchParam;
 			SwitchNode->Options.SetNum(NumSwitchOptions);
diff --git a/Engine/Plugins/Mutable/Source/CustomizableObjectEditor/Private/MuCOE/GenerateMutableSource/GenerateMutableSourceComponent.cpp b/Engine/Plugins/Mutable/Source/CustomizableObjectEditor/Private/MuCOE/GenerateMutableSource/GenerateMutableSourceComponent.cpp
--- a/Engine/Plugins/Mutable/Source/CustomizableObjectEditor/Private/MuCOE/GenerateMutableSource/GenerateMutableSourceComponent.cpp
+++ b/Engine/Plugins/Mutable/Source/CustomizableObjectEditor/Private/MuCOE/GenerateMutableSource/GenerateMutableSourceComponent.cpp
@@ -3,7 +3,6 @@
 #include "GenerateMutableSourceComponent.h"
 
 #include "GenerateMutableSourceSkeletalMesh.h"
-#include "Factories/FbxMeshImportData.h"
 #include "MuCOE/GenerateMutableSource/GenerateMutableSource.h"
 
 #include "MuCOE/CustomizableObjectMacroLibrary/CustomizableObjectMacroLibrary.h"
@@ -28,7 +27,6 @@
 #include "MuT/NodeComponentVariation.h"
 #include "MuT/NodeScalar.h"
 #include "MuT/NodeSurfaceNew.h"
-#include "MuT/NodeMeshConstant.h"
 #include "MuT/NodeMaterialConstant.h"
 
 #include "Rendering/SkeletalMeshLODModel.h"
@@ -53,6 +51,11 @@ UE::Mutable::Private::Ptr<UE::Mutable::Private::NodeComponent> GenerateMutableSo
 		return static_cast<UE::Mutable::Private::NodeComponent*>(Generated->Node.get());
 	}
 
+	if (Node->IsNodeOutDatedAndNeedsRefresh())
+	{
+		Node->SetRefreshNodeWarning();
+	}
+	
 	// Bool that determines if a node can be added to the cache of nodes.
 	// Most nodes need to be added to the cache but there are some that don't. For example, MacroInstanceNodes
 	bool bCacheNode = true;
@@ -222,11 +225,8 @@ UE::Mutable::Private::Ptr<UE::Mutable::Private::NodeComponent> GenerateMutableSo
 					if (!SwitchParam)
 					{
 						// Warn about a failure.
-						if (EnumPin)
-						{
-							const FText Message = LOCTEXT("FailedToGenerateSwitchParam", "Could not generate switch enum parameter. Please refesh the switch node and connect an enum.");
-							GenerationContext.Log(Message, Node);
-						}
+						const FText Message = LOCTEXT("FailedToGenerateSwitchParam", "Could not generate switch enum parameter. Please refresh the switch node and connect an enum.");
+						GenerationContext.Log(Message, Node);
 
 						return Result;
 					}
@@ -239,14 +239,18 @@ UE::Mutable::Private::Ptr<UE::Mutable::Private::NodeComponent> GenerateMutableSo
 						return Result;
 					}
 
-					const int32 NumSwitchOptions = TypedNodeSwitch->SwitchPins.Num();
-
-					UE::Mutable::Private::NodeScalarEnumParameter* EnumParameter = static_cast<UE::Mutable::Private::NodeScalarEnumParameter*>(SwitchParam.get());
-					if (NumSwitchOptions != EnumParameter->Options.Num())
 					{
-						const FText Message = LOCTEXT("MismatchedSwitch", "Switch enum and switch node have different number of options. Please refresh the switch node to make sure the outcomes are labeled properly.");
-						GenerationContext.Log(Message, Node);
+						const UE::Mutable::Private::NodeScalarEnumParameter* EnumParameter = static_cast<UE::Mutable::Private::NodeScalarEnumParameter*>(SwitchParam.get());
+						
+						if (!DoOptionsMatchEnum(*TypedNodeSwitch, *EnumParameter))
+						{
+							const FText Message = LOCTEXT("MismatchedSwitch", "Switch enum and switch node have different options. Please refresh the switch node to make sure the outcomes are labeled properly.");
+							GenerationContext.Log(Message, Node);
+							Node->SetRefreshNodeWarning();
+						}
 					}
+					
+					const int32 NumSwitchOptions = TypedNodeSwitch->SwitchPins.Num();
 
 					UE::Mutable::Private::Ptr<UE::Mutable::Private::NodeComponentSwitch> SwitchNode = new UE::Mutable::Private::NodeComponentSwitch;
 					SwitchNode->Parameter = SwitchParam;
@@ -273,7 +277,7 @@ UE::Mutable::Private::Ptr<UE::Mutable::Private::NodeComponent> GenerateMutableSo
 				}
 				else
 				{
-					GenerationContext.Log(LOCTEXT("NoEnumParamInSwitch", "Switch nodes must have an enum switch parameter. Please connect an enum and refesh the switch node."), Node);
+					GenerationContext.Log(LOCTEXT("NoEnumParamInSwitch", "Switch nodes must have an enum switch parameter. Please connect an enum and refresh the switch node."), Node);
 					return Result;
 				}
 			}(); // invoke lambda.
diff --git a/Engine/Plugins/Mutable/Source/CustomizableObjectEditor/Private/MuCOE/GenerateMutableSource/GenerateMutableSourceFloat.cpp b/Engine/Plugins/Mutable/Source/CustomizableObjectEditor/Private/MuCOE/GenerateMutableSource/GenerateMutableSourceFloat.cpp
--- a/Engine/Plugins/Mutable/Source/CustomizableObjectEditor/Private/MuCOE/GenerateMutableSource/GenerateMutableSourceFloat.cpp
+++ b/Engine/Plugins/Mutable/Source/CustomizableObjectEditor/Private/MuCOE/GenerateMutableSource/GenerateMutableSourceFloat.cpp
@@ -26,7 +26,6 @@
 #include "MuCOE/Nodes/CustomizableObjectNodeMacroInstance.h"
 #include "MuCOE/Nodes/CustomizableObjectNodeTable.h"
 #include "MuCOE/Nodes/CustomizableObjectNodeTunnel.h"
-#include "MuT/NodeMeshExternal.h"
 #include "MuT/NodeScalarConstant.h"
 #include "MuT/NodeScalarCurve.h"
 #include "MuT/NodeScalarSwitch.h"
@@ -132,7 +131,7 @@ UE::Mutable::Private::Ptr<UE::Mutable::Private::NodeScalar> GenerateMutableSourc
 				NumParameters != 1)
 			{
 				const FText Message = NumParameters
-					? LOCTEXT("NoEnumParamInSwitch", "Switch nodes must have an enum switch parameter. Please connect an enum and refesh the switch node.")
+					? LOCTEXT("NoEnumParamInSwitch", "Switch nodes must have an enum switch parameter. Please connect an enum and refresh the switch node.")
 					: LOCTEXT("InvalidEnumInSwitch", "Switch nodes must have a single enum with all the options inside. Please remove all the enums but one and refresh the switch node.");
 
 				GenerationContext.Log(Message, Node);
@@ -148,7 +147,7 @@ UE::Mutable::Private::Ptr<UE::Mutable::Private::NodeScalar> GenerateMutableSourc
 				// Warn about a failure.
 				if (EnumPin)
 				{
-					const FText Message = LOCTEXT("FailedToGenerateSwitchParam", "Could not generate switch enum parameter. Please refesh the switch node and connect an enum.");
+					const FText Message = LOCTEXT("FailedToGenerateSwitchParam", "Could not generate switch enum parameter. Please refresh the switch node and connect an enum.");
 					GenerationContext.Log(Message, Node);
 				}
 
@@ -163,15 +162,19 @@ UE::Mutable::Private::Ptr<UE::Mutable::Private::NodeScalar> GenerateMutableSourc
 				return nullptr;
 			}
 
-			const int32 NumSwitchOptions = TypedNodeFloatSwitch->SwitchPins.Num();
-
-			UE::Mutable::Private::NodeScalarEnumParameter* EnumParameter = static_cast<UE::Mutable::Private::NodeScalarEnumParameter*>(SwitchParam.get());
-			if (NumSwitchOptions != EnumParameter->Options.Num())
 			{
-				const FText Message = LOCTEXT("MismatchedSwitch", "Switch enum and switch node have different number of options. Please refresh the switch node to make sure the outcomes are labeled properly.");
-				GenerationContext.Log(Message, Node);
+				const UE::Mutable::Private::NodeScalarEnumParameter* EnumParameter = static_cast<UE::Mutable::Private::NodeScalarEnumParameter*>(SwitchParam.get());
+				
+				if (!DoOptionsMatchEnum(*TypedNodeFloatSwitch, *EnumParameter))
+				{
+					const FText Message = LOCTEXT("MismatchedSwitch", "Switch enum and switch node have different options. Please refresh the switch node to make sure the outcomes are labeled properly.");
+					GenerationContext.Log(Message, Node);
+					Node->SetRefreshNodeWarning();
+				}
 			}
 
+			const int32 NumSwitchOptions = TypedNodeFloatSwitch->SwitchPins.Num();
+				
 			UE::Mutable::Private::Ptr<UE::Mutable::Private::NodeScalarSwitch> SwitchNode = new UE::Mutable::Private::NodeScalarSwitch;
 			SwitchNode->Parameter = SwitchParam;
 			SwitchNode->Options.SetNum(NumSwitchOptions);
diff --git a/Engine/Plugins/Mutable/Source/CustomizableObjectEditor/Private/MuCOE/GenerateMutableSource/GenerateMutableSourceImage.cpp b/Engine/Plugins/Mutable/Source/CustomizableObjectEditor/Private/MuCOE/GenerateMutableSource/GenerateMutableSourceImage.cpp
--- a/Engine/Plugins/Mutable/Source/CustomizableObjectEditor/Private/MuCOE/GenerateMutableSource/GenerateMutableSourceImage.cpp
+++ b/Engine/Plugins/Mutable/Source/CustomizableObjectEditor/Private/MuCOE/GenerateMutableSource/GenerateMutableSourceImage.cpp
@@ -66,8 +66,8 @@
 #include "MuT/NodeScalarConstant.h"
 
 #include "Interfaces/ITextureFormat.h"
-#include "Interfaces/ITextureFormatModule.h"
 #include "Interfaces/ITextureFormatManagerModule.h"
+#include "MuCOE/Nodes/CustomizableObjectNodeEnumParameter.h"
 #include "MuT/UnrealPixelFormatOverride.h"
 
 #define LOCTEXT_NAMESPACE "CustomizableObjectEditor"
@@ -358,7 +358,7 @@ UE::Mutable::Private::Ptr<UE::Mutable::Private::NodeImage> GenerateMutableSource
 				// Switch Param not generated
 				if (!SwitchParam)
 				{
-					const FText Message = LOCTEXT("FailedToGenerateSwitchParam", "Could not generate switch enum parameter. Please refesh the switch node and connect an enum.");
+					const FText Message = LOCTEXT("FailedToGenerateSwitchParam", "Could not generate switch enum parameter. Please refresh the switch node and connect an enum.");
 					GenerationContext.Log(Message, Node);
 					
 					return Result;
@@ -372,15 +372,19 @@ UE::Mutable::Private::Ptr<UE::Mutable::Private::NodeImage> GenerateMutableSource
 					return Result;
 				}
 
-				const int32 NumSwitchOptions = TypedNodeTextureSwitch->SwitchPins.Num();
-
-				UE::Mutable::Private::NodeScalarEnumParameter* EnumParameter = static_cast<UE::Mutable::Private::NodeScalarEnumParameter*>(SwitchParam.get());
-				if (NumSwitchOptions != EnumParameter->Options.Num())
 				{
-					const FText Message = LOCTEXT("MismatchedSwitch", "Switch enum and switch node have different number of options. Please refresh the switch node to make sure the outcomes are labeled properly.");
-					GenerationContext.Log(Message, Node);
+					const UE::Mutable::Private::NodeScalarEnumParameter* EnumParameter = static_cast<UE::Mutable::Private::NodeScalarEnumParameter*>(SwitchParam.get());
+
+					if (!DoOptionsMatchEnum(*TypedNodeTextureSwitch, *EnumParameter))
+					{
+						const FText Message = LOCTEXT("MismatchedSwitch", "Switch enum and switch node have different options. Please refresh the switch node to make sure the outcomes are labeled properly.");
+						GenerationContext.Log(Message, Node);
+						Node->SetRefreshNodeWarning();
+					}
 				}
 
+				const int32 NumSwitchOptions = TypedNodeTextureSwitch->SwitchPins.Num();
+				
 				UE::Mutable::Private::Ptr<UE::Mutable::Private::NodeImageSwitch> SwitchNode = new UE::Mutable::Private::NodeImageSwitch;
 				SwitchNode->Parameter = SwitchParam;
 				SwitchNode->Options.SetNum(NumSwitchOptions);
@@ -404,7 +408,7 @@ UE::Mutable::Private::Ptr<UE::Mutable::Private::NodeImage> GenerateMutableSource
 			}
 			else
 			{
-				GenerationContext.Log(LOCTEXT("NoEnumParamInSwitch", "Switch nodes must have an enum switch parameter. Please connect an enum and refesh the switch node."), Node);
+				GenerationContext.Log(LOCTEXT("NoEnumParamInSwitch", "Switch nodes must have an enum switch parameter. Please connect an enum and refresh the switch node."), Node);
 				return Result;
 			}
 		}(); // invoke lambda.
@@ -926,7 +930,7 @@ UE::Mutable::Private::Ptr<UE::Mutable::Private::NodeImage> GenerateMutableSource
 				// Switch Param not generated
 				if (!SwitchParam)
 				{
-					const FText Message = LOCTEXT("FailedToGenerateSwitchParam", "Could not generate switch enum parameter. Please refesh the switch node and connect an enum.");
+					const FText Message = LOCTEXT("FailedToGenerateSwitchParam", "Could not generate switch enum parameter. Please refresh the switch node and connect an enum.");
 					GenerationContext.Log(Message, Node);
 
 					return Result;
@@ -940,15 +944,20 @@ UE::Mutable::Private::Ptr<UE::Mutable::Private::NodeImage> GenerateMutableSource
 					return Result;
 				}
 
-				const int32 NumSwitchOptions = TypedNodePassThroughTextureSwitch->SwitchPins.Num();
-
-				UE::Mutable::Private::NodeScalarEnumParameter* EnumParameter = static_cast<UE::Mutable::Private::NodeScalarEnumParameter*>(SwitchParam.get());
-				if (NumSwitchOptions != EnumParameter->Options.Num())
 				{
-					const FText Message = LOCTEXT("MismatchedSwitch", "Switch enum and switch node have different number of options. Please refresh the switch node to make sure the outcomes are labeled properly.");
-					GenerationContext.Log(Message, Node);
+					const UE::Mutable::Private::NodeScalarEnumParameter* EnumParameter = static_cast<UE::Mutable::Private::NodeScalarEnumParameter*>(SwitchParam.get());
+
+					if (!DoOptionsMatchEnum(*TypedNodePassThroughTextureSwitch, *EnumParameter))
+					{
+						const FText Message = LOCTEXT("MismatchedSwitch", "Switch enum and switch node have different options. Please refresh the switch node to make sure the outcomes are labeled properly.");
+						GenerationContext.Log(Message, Node);
+						Node->SetRefreshNodeWarning();
+					}
 				}
 
+
+				const int32 NumSwitchOptions = TypedNodePassThroughTextureSwitch->SwitchPins.Num();
+				
 				// TODO: Implement Mutable core pass-through switch nodes
 				UE::Mutable::Private::Ptr<UE::Mutable::Private::NodeImageSwitch> SwitchNode = new UE::Mutable::Private::NodeImageSwitch;
 				SwitchNode->Parameter = SwitchParam;
@@ -974,7 +983,7 @@ UE::Mutable::Private::Ptr<UE::Mutable::Private::NodeImage> GenerateMutableSource
 			}
 			else
 			{
-				GenerationContext.Log(LOCTEXT("NoEnumParamInSwitch", "Switch nodes must have an enum switch parameter. Please connect an enum and refesh the switch node."), Node);
+				GenerationContext.Log(LOCTEXT("NoEnumParamInSwitch", "Switch nodes must have an enum switch parameter. Please connect an enum and refresh the switch node."), Node);
 				return Result;
 			}
 		}(); // invoke lambda.
diff --git a/Engine/Plugins/Mutable/Source/CustomizableObjectEditor/Private/MuCOE/GenerateMutableSource/GenerateMutableSourceMaterial.cpp b/Engine/Plugins/Mutable/Source/CustomizableObjectEditor/Private/MuCOE/GenerateMutableSource/GenerateMutableSourceMaterial.cpp
--- a/Engine/Plugins/Mutable/Source/CustomizableObjectEditor/Private/MuCOE/GenerateMutableSource/GenerateMutableSourceMaterial.cpp
+++ b/Engine/Plugins/Mutable/Source/CustomizableObjectEditor/Private/MuCOE/GenerateMutableSource/GenerateMutableSourceMaterial.cpp
@@ -4,17 +4,12 @@
 
 #include "Interfaces/ITargetPlatform.h"
 #include "Interfaces/ITextureFormat.h"
-#include "Interfaces/ITextureFormatModule.h"
-#include "Interfaces/ITextureFormatManagerModule.h"
-#include "TextureCompressorModule.h"
 
 #include "GenerateMutableSourceExternal.h"
 #include "GenerateMutableSourceMesh.h"
 #include "GenerateMutableSourceSkeletalMesh.h"
 #include "Engine/TextureLODSettings.h"
-#include "Interfaces/ITargetPlatform.h"
 #include "MuCOE/CustomizableObjectCompiler.h"
-#include "MuCOE/CustomizableObjectGraph.h"
 #include "MuCOE/EdGraphSchema_CustomizableObject.h"
 #include "MuCOE/GenerateMutableSource/GenerateMutableSourceColor.h"
 #include "MuCOE/GenerateMutableSource/GenerateMutableSourceFloat.h"
@@ -39,7 +34,6 @@
 #include "MuT/NodeImageMipmap.h"
 #include "MuT/NodeImageNormalComposite.h"
 #include "MuT/NodeImageResize.h"
-#include "MuT/NodeImageSwizzle.h"
 #include "MuT/NodeMaterialConstant.h"
 #include "MuT/NodeMaterialSwitch.h"
 #include "MuT/NodeMaterialTable.h"
@@ -52,7 +46,6 @@
 #include "MuCOE/Nodes/CONodeSkeletalMeshBreak.h"
 #include "MuT/NodeMaterialExternal.h"
 #include "MuT/NodeMaterialSkeletalMeshBreak.h"
-#include "MuT/UnrealPixelFormatOverride.h"
 
 
 #define LOCTEXT_NAMESPACE "CustomizableObjectEditor"
@@ -512,15 +505,19 @@ UE::Mutable::Private::Ptr<UE::Mutable::Private::NodeMaterial> GenerateMutableSou
 						return Result;
 					}
 
-					const int32 NumSwitchOptions = TypedNodeMaterialSwitch->SwitchPins.Num();
-
-					UE::Mutable::Private::NodeScalarEnumParameter* EnumParameter = static_cast<UE::Mutable::Private::NodeScalarEnumParameter*>(SwitchParam.get());
-					if (NumSwitchOptions != EnumParameter->Options.Num())
 					{
-						const FText Message = LOCTEXT("MismatchedSwitch", "Switch enum and switch node have different number of options. Please refresh the switch node to make sure the outcomes are labeled properly.");
-						GenerationContext.Log(Message, Node);
+						const UE::Mutable::Private::NodeScalarEnumParameter* EnumParameter = static_cast<UE::Mutable::Private::NodeScalarEnumParameter*>(SwitchParam.get());
+
+						if (!DoOptionsMatchEnum(*TypedNodeMaterialSwitch, *EnumParameter))
+						{
+							const FText Message = LOCTEXT("MismatchedSwitch", "Switch enum and switch node have different options. Please refresh the switch node to make sure the outcomes are labeled properly.");
+							GenerationContext.Log(Message, Node);
+							Node->SetRefreshNodeWarning();
+						}
 					}
 
+					const int32 NumSwitchOptions = TypedNodeMaterialSwitch->SwitchPins.Num();
+					
 					UE::Mutable::Private::Ptr<UE::Mutable::Private::NodeMaterialSwitch> SwitchNode = new UE::Mutable::Private::NodeMaterialSwitch;
 					SwitchNode->Parameter = SwitchParam;
 					SwitchNode->Options.SetNum(NumSwitchOptions);
@@ -544,7 +541,7 @@ UE::Mutable::Private::Ptr<UE::Mutable::Private::NodeMaterial> GenerateMutableSou
 				}
 				else
 				{
-					GenerationContext.Log(LOCTEXT("NoEnumParamInSwitch", "Switch nodes must have an enum switch parameter. Please connect an enum and refesh the switch node."), Node);
+					GenerationContext.Log(LOCTEXT("NoEnumParamInSwitch", "Switch nodes must have an enum switch parameter. Please connect an enum and refresh the switch node."), Node);
 					return Result;
 				}
 			}(); // invoke lambda.
diff --git a/Engine/Plugins/Mutable/Source/CustomizableObjectEditor/Private/MuCOE/GenerateMutableSource/GenerateMutableSourceMesh.cpp b/Engine/Plugins/Mutable/Source/CustomizableObjectEditor/Private/MuCOE/GenerateMutableSource/GenerateMutableSourceMesh.cpp
--- a/Engine/Plugins/Mutable/Source/CustomizableObjectEditor/Private/MuCOE/GenerateMutableSource/GenerateMutableSourceMesh.cpp
+++ b/Engine/Plugins/Mutable/Source/CustomizableObjectEditor/Private/MuCOE/GenerateMutableSource/GenerateMutableSourceMesh.cpp
@@ -12,10 +12,7 @@
 #include "GenerateMutableSourceExternal.h"
 #include "GenerateMutableSourceImage.h"
 #include "GenerateMutableSourceSkeletalMesh.h"
-#include "MeshUtilities.h"
-#include "Modules/ModuleManager.h"
 #include "MuCO/CustomizableObjectPrivate.h"
-#include "MuCO/CustomizableObjectInstance.h"
 #include "MuCO/CustomizableObjectSystem.h"
 #include "MuCO/UnrealConversionUtils.h"
 #include "MuCO/MutableMeshBufferUtils.h"
@@ -54,8 +51,6 @@
 #include "MuT/NodeMeshTable.h"
 #include "MuT/NodeMeshVariation.h"
 #include "MuT/NodeMeshSkeletalMeshBreak.h"
-#include "MuR/Types.h"
-#include "MuR/OpMeshFormat.h"
 #include "MuR/ConvertData.h"
 #include "PhysicsEngine/PhysicsAsset.h"
 #include "PhysicsEngine/SkeletalBodySetup.h"
@@ -65,11 +60,6 @@
 #include "Hash/CityHash.h"
 #include "PackedNormal.h"
 #include "MuCOE/Nodes/CONodeExternalOperation.h"
-#include "MuR/External/MeshAdapter.h"
-#include "MuR/External/Operation.h"
-#include "MuR/External/TextureAdapter.h"
-#include "MuT/NodeExternal.h"
-#include "MuT/NodeExtensionData.h"
 #include "MuT/NodeMeshExternal.h"
 
 #define LOCTEXT_NAMESPACE "CustomizableObjectEditor"
@@ -4107,11 +4097,8 @@ UE::Mutable::Private::Ptr<UE::Mutable::Private::NodeMesh> GenerateMutableSourceM
 				if (!SwitchParam)
 				{
 					// Warn about a failure.
-					if (EnumPin)
-					{
-						const FText Message = LOCTEXT("FailedToGenerateSwitchParam", "Could not generate switch enum parameter. Please refesh the switch node and connect an enum.");
+					const FText Message = LOCTEXT("FailedToGenerateSwitchParam", "Could not generate switch enum parameter. Please refresh the switch node and connect an enum.");
 						GenerationContext.Log(Message, Node);
-					}
 
 					return Result;
 				}
@@ -4124,15 +4111,19 @@ UE::Mutable::Private::Ptr<UE::Mutable::Private::NodeMesh> GenerateMutableSourceM
 					return Result;
 				}
 
-				const int32 NumSwitchOptions = TypedNodeMeshSwitch->SwitchPins.Num();
-
-				UE::Mutable::Private::NodeScalarEnumParameter* EnumParameter = static_cast<UE::Mutable::Private::NodeScalarEnumParameter*>(SwitchParam.get());
-				if (NumSwitchOptions != EnumParameter->Options.Num())
 				{
-					const FText Message = LOCTEXT("MismatchedSwitch", "Switch enum and switch node have different number of options. Please refresh the switch node to make sure the outcomes are labeled properly.");
-					GenerationContext.Log(Message, Node);
+					const UE::Mutable::Private::NodeScalarEnumParameter* EnumParameter = static_cast<UE::Mutable::Private::NodeScalarEnumParameter*>(SwitchParam.get());
+
+					if (!DoOptionsMatchEnum(*TypedNodeMeshSwitch, *EnumParameter))
+					{
+						const FText Message = LOCTEXT("MismatchedSwitch", "Switch enum and switch node have different options. Please refresh the switch node to make sure the outcomes are labeled properly.");
+						GenerationContext.Log(Message, Node);
+						Node->SetRefreshNodeWarning();
+					}
 				}
 
+				const int32 NumSwitchOptions = TypedNodeMeshSwitch->SwitchPins.Num();
+				
 				UE::Mutable::Private::Ptr<UE::Mutable::Private::NodeMeshSwitch> SwitchNode = new UE::Mutable::Private::NodeMeshSwitch;
 				SwitchNode->Parameter = SwitchParam;
 				SwitchNode->Options.SetNum(NumSwitchOptions);
@@ -4151,7 +4142,7 @@ UE::Mutable::Private::Ptr<UE::Mutable::Private::NodeMesh> GenerateMutableSourceM
 			}
 			else
 			{
-				GenerationContext.Log(LOCTEXT("NoEnumParamInSwitch", "Switch nodes must have an enum switch parameter. Please connect an enum and refesh the switch node."), Node);
+				GenerationContext.Log(LOCTEXT("NoEnumParamInSwitch", "Switch nodes must have an enum switch parameter. Please connect an input enum node."), Node);
 				return Result;
 			}
 		}(); // invoke lambda.
diff --git a/Engine/Plugins/Mutable/Source/CustomizableObjectEditor/Private/MuCOE/GenerateMutableSource/GenerateMutableSourceSurface.cpp b/Engine/Plugins/Mutable/Source/CustomizableObjectEditor/Private/MuCOE/GenerateMutableSource/GenerateMutableSourceSurface.cpp
--- a/Engine/Plugins/Mutable/Source/CustomizableObjectEditor/Private/MuCOE/GenerateMutableSource/GenerateMutableSourceSurface.cpp
+++ b/Engine/Plugins/Mutable/Source/CustomizableObjectEditor/Private/MuCOE/GenerateMutableSource/GenerateMutableSourceSurface.cpp
@@ -2,55 +2,32 @@
 
 #include "MuCOE/GenerateMutableSource/GenerateMutableSourceSurface.h"
 
+#include <MuCOE/Nodes/CONodeSkeletalMeshSection.h>
+#include <MuCOE/Nodes/CONodeSkeletalMeshSectionSwitch.h>
+#include <MuCOE/Nodes/CONodeSkeletalMeshSectionVariation.h>
+
 #include "Engine/SkinnedAssetCommon.h"
 #include "Engine/StaticMesh.h"
-#include "Engine/TextureLODSettings.h"
-#include "GPUSkinPublicDefs.h"
 #include "Modules/ModuleManager.h"
 
-#include "Materials/MaterialInstance.h"
 #include "GPUSkinVertexFactory.h"
 #include "Rendering/SkeletalMeshLODModel.h"
 
-#include "MuCO/CustomizableObjectInstance.h"
 #include "MuCO/UnrealConversionUtils.h"
 #include "MuCOE/CustomizableObjectCompiler.h"
-#include "MuCOE/CustomizableObjectLayout.h"
-#include "MuCOE/GenerateMutableSource/GenerateMutableSourceColor.h"
 #include "MuCOE/GenerateMutableSource/GenerateMutableSourceFloat.h"
 #include "MuCOE/GenerateMutableSource/GenerateMutableSourceGroupProjector.h"
-#include "MuCOE/GenerateMutableSource/GenerateMutableSourceImage.h"
 #include "MuCOE/GenerateMutableSource/GenerateMutableSourceMesh.h"
-#include "MuCOE/GenerateMutableSource/GenerateMutableSourceTable.h"
-#include "MuCOE/GenerateMutableSource/GenerateMutableSourceLayout.h"
 #include "MuCOE/GenerateMutableSource/GenerateMutableSourceMacro.h"
 #include "MuCOE/GenerateMutableSource/GenerateMutableSourceMaterial.h"
 #include "MuCOE/GraphTraversal.h"
-#include "MuCOE/MutableUtils.h"
-#include "MuCOE/Nodes/CONodeCopySkeletalMeshSection.h"
 #include "MuCOE/Nodes/CustomizableObjectNodeFloatConstant.h"
-#include "MuCOE/Nodes/CustomizableObjectNodeFloatParameter.h"
 #include "MuCOE/Nodes/CustomizableObjectNodeMacroInstance.h"
-#include "MuCOE/Nodes/CONodeSkeletalMeshSectionVariation.h"
-#include "MuCOE/Nodes/CONodeSkeletalMeshSectionSwitch.h"
-#include "MuCOE/Nodes/CONodeModifierEditSkeletalMeshSection.h"
-#include "MuCOE/Nodes/CONodeModifierExtendSkeletalMeshSection.h"
-#include "MuCOE/Nodes/CustomizableObjectNodeModifierMorphMeshSection.h"
-#include "MuCOE/Nodes/CustomizableObjectNodeModifierRemoveMesh.h"
-#include "MuCOE/Nodes/CustomizableObjectNodeModifierRemoveMeshBlocks.h"
 #include "MuCOE/Nodes/CustomizableObjectNodeSkeletalMesh.h"
-#include "MuCOE/Nodes/CustomizableObjectNodeStaticString.h"
 #include "MuCOE/Nodes/CustomizableObjectNodeTable.h"
 #include "MuCOE/Nodes/CustomizableObjectNodeTunnel.h"
 #include "MuT/NodeImageFormat.h"
-#include "MuT/NodeImageMipmap.h"
-#include "MuT/NodeImageNormalComposite.h"
-#include "MuT/NodeImageResize.h"
-#include "MuT/NodeImageSwizzle.h"
-#include "MuT/NodeMeshConstant.h"
 #include "MuT/NodeMeshFormat.h"
-#include "MuT/NodeMeshFragment.h"
-#include "MuT/NodeMaterialConstant.h"
 #include "MuT/NodeMaterialModify.h"
 #include "MuT/NodeSurfaceSwitch.h"
 #include "MuT/NodeSurfaceVariation.h"
@@ -80,6 +57,11 @@ UE::Mutable::Private::Ptr<UE::Mutable::Private::NodeSurface> GenerateMutableSour
 		return static_cast<UE::Mutable::Private::NodeSurface*>(Generated->Node.get());
 	}
 	
+	if (Node->IsNodeOutDatedAndNeedsRefresh())
+	{
+		Node->SetRefreshNodeWarning();
+	}
+	
 	UE::Mutable::Private::Ptr<UE::Mutable::Private::NodeSurface> Result;
 
 	if (UCustomizableObjectNode* CustomObjNode = Cast<UCustomizableObjectNode>(Node))
@@ -298,11 +280,8 @@ UE::Mutable::Private::Ptr<UE::Mutable::Private::NodeSurface> GenerateMutableSour
 				if (!SwitchParam)
 				{
 					// Warn about a failure.
-					if (EnumPin)
-					{
-						const FText Message = LOCTEXT("FailedToGenerateSwitchParam", "Could not generate switch enum parameter. Please refesh the switch node and connect an enum.");
+					const FText Message = LOCTEXT("FailedToGenerateSwitchParam", "Could not generate switch enum parameter. Please refresh the switch node and connect an enum.");
 						GenerationContext.Log(Message, Node);
-					}
 
 					return Result;
 				}
@@ -315,15 +294,19 @@ UE::Mutable::Private::Ptr<UE::Mutable::Private::NodeSurface> GenerateMutableSour
 					return Result;
 				}
 
-				const int32 NumSwitchOptions = TypedNodeSwitch->SwitchPins.Num();
-
-				UE::Mutable::Private::NodeScalarEnumParameter* EnumParameter = static_cast<UE::Mutable::Private::NodeScalarEnumParameter*>(SwitchParam.get());
-				if (NumSwitchOptions != EnumParameter->Options.Num())
 				{
-					const FText Message = LOCTEXT("MismatchedSwitch", "Switch enum and switch node have different number of options. Please refresh the switch node to make sure the outcomes are labeled properly.");
-					GenerationContext.Log(Message, Node);
+					const UE::Mutable::Private::NodeScalarEnumParameter* EnumParameter = static_cast<UE::Mutable::Private::NodeScalarEnumParameter*>(SwitchParam.get());
+
+					if (!DoOptionsMatchEnum(*TypedNodeSwitch, *EnumParameter))
+					{
+						const FText Message = LOCTEXT("MismatchedSwitch", "Switch enum and switch node have different options. Please refresh the switch node to make sure the outcomes are labeled properly.");
+						GenerationContext.Log(Message, Node);
+						Node->SetRefreshNodeWarning();
+					}
 				}
 
+				const int32 NumSwitchOptions = TypedNodeSwitch->SwitchPins.Num();
+				
 				UE::Mutable::Private::Ptr<UE::Mutable::Private::NodeSurfaceSwitch> SwitchNode = new UE::Mutable::Private::NodeSurfaceSwitch;
 				SwitchNode->Parameter = SwitchParam;
 				SwitchNode->Options.SetNum(NumSwitchOptions);
@@ -350,7 +333,7 @@ UE::Mutable::Private::Ptr<UE::Mutable::Private::NodeSurface> GenerateMutableSour
 			}
 			else
 			{
-				GenerationContext.Log(LOCTEXT("NoEnumParamInSwitch", "Switch nodes must have an enum switch parameter. Please connect an enum and refesh the switch node."), Node);
+				GenerationContext.Log(LOCTEXT("NoEnumParamInSwitch", "Switch nodes must have an enum switch parameter. Please connect an enum and refresh the switch node."), Node);
 				return Result;
 			}
 		}(); // invoke lambda.
diff --git a/Engine/Plugins/Mutable/Source/CustomizableObjectEditor/Private/MuCOE/Nodes/CONodeSkeletalMeshSection.cpp b/Engine/Plugins/Mutable/Source/CustomizableObjectEditor/Private/MuCOE/Nodes/CONodeSkeletalMeshSection.cpp
--- a/Engine/Plugins/Mutable/Source/CustomizableObjectEditor/Private/MuCOE/Nodes/CONodeSkeletalMeshSection.cpp
+++ b/Engine/Plugins/Mutable/Source/CustomizableObjectEditor/Private/MuCOE/Nodes/CONodeSkeletalMeshSection.cpp
@@ -325,14 +325,14 @@ void UCONodeSkeletalMeshSection::AllocateDefaultPins(UCustomizableObjectNodeRema
 	{
 		const FString PinFriendlyName = TEXT("Mesh");
 		const FString PinName = PinFriendlyName + FString(TEXT("_Input_Pin"));
-		CustomCreatePin(EGPD_Input, Schema->PC_Mesh, FName(*PinName), FText::FromString(PinFriendlyName));
+		InputMeshPin = CustomCreatePin(EGPD_Input, Schema->PC_Mesh, FName(*PinName), FText::FromString(PinFriendlyName));
 	}
 	
 	{
 		const FString PinFriendlyName = UEdGraphSchema_CustomizableObject::GetPinCategoryFriendlyName(Schema->PC_Material).ToString();
 		const FString PinName = FString(TEXT("Material_Input_Pin"));
-		UEdGraphPin* TableMaterialPin = CustomCreatePin(EGPD_Input, Schema->PC_Material, FName(*PinName), FText::FromString(PinFriendlyName));
-		TableMaterialPin->PinToolTip = "Pin for a Material from a Table Node";
+		InputMaterialPin = CustomCreatePin(EGPD_Input, Schema->PC_Material, FName(*PinName), FText::FromString(PinFriendlyName));
+		InputMaterialPin.Get()->PinToolTip = "Pin for a Material from a Table Node";
 	}
 
 	{
@@ -804,6 +804,63 @@ void UCONodeSkeletalMeshSection::BackwardsCompatibleFixup(int32 CustomizableObje
 			}
 		}
 	}
+	else if (CustomizableObjectCustomVersion == FCustomizableObjectCustomVersion::MeshSectionNodeInputPinsAsReferences)
+	{
+		if (!InputMeshPin.Get())
+		{
+			// const FString PinFriendlyName = TEXT("Mesh");
+			// const FString PinName = PinFriendlyName + FString(TEXT("_Input_Pin"));
+			// InputMeshPin = CustomCreatePin(EGPD_Input, "mesh", FName(*PinName), FText::FromString(PinFriendlyName));
+			
+			const FString PinFriendlyName = TEXT("Mesh");
+			const FString PinName = PinFriendlyName + FString(TEXT("_Input_Pin"));
+
+			if (UEdGraphPin* Pin = FindPin(PinName))
+			{
+				InputMeshPin = Pin;
+			}
+			else
+			{
+				InputMeshPin =  FindPin(PinFriendlyName);
+			}
+			
+			check(InputMeshPin.Get())
+		}
+	
+		if (!InputMaterialPin.Get())
+		{
+			// const FString PinFriendlyName = FString{TEXT("materialAsset")};
+			// const FString PinName = FString(TEXT("Material_Input_Pin"));
+			// InputMaterialPin = CustomCreatePin(EGPD_Input, "materialAsset", FName(*PinName), FText::FromString(PinFriendlyName));
+			// InputMaterialPin.Get()->PinToolTip = "Pin for a Material from a Table Node";
+			
+			const FString PinFriendlyName = UEdGraphSchema_CustomizableObject::GetPinCategoryFriendlyName(UEdGraphSchema_CustomizableObject::PC_Material).ToString();
+			const FString PinName = PinFriendlyName + FString(TEXT("_Input_Pin"));
+
+			UEdGraphPin* Pin = FindPin(PinName);
+			if (Pin)
+			{
+				InputMaterialPin = Pin;
+			}
+			else
+			{
+				// Changing the locale can change the GetPinCategoryFriendlyName() result, leading to failing to find the MaterialAsset pin.
+				// This is a provisional fix while incorrect FriendlyName usages are removed in UE-294583
+				Pin = FindPin(PinFriendlyName);
+
+				if (Pin)
+				{
+					InputMaterialPin = Pin;
+				}
+				else
+				{
+					InputMaterialPin = FindPin(FString("Material_Input_Pin"));
+				}
+			}
+			
+			check(InputMaterialPin.Get())
+		}
+	}
 }
 
 
@@ -880,48 +937,13 @@ UEdGraphPin* UCONodeSkeletalMeshSection::OutputPin() const
 
 UEdGraphPin* UCONodeSkeletalMeshSection::GetMeshPin() const
 {
-	FString PinFriendlyName = TEXT("Mesh");
-	FString PinName = PinFriendlyName + FString(TEXT("_Input_Pin"));
-
-	UEdGraphPin* Pin = FindPin(PinName);
-
-	if(Pin)
-	{
-		return Pin;
-	}
-	else
-	{
-		return FindPin(TEXT("Mesh"));
-	}
+	return InputMeshPin.Get();
 }
 
 
 UEdGraphPin* UCONodeSkeletalMeshSection::GetMaterialAssetPin() const
 {
-	const FString PinFriendlyName = UEdGraphSchema_CustomizableObject::GetPinCategoryFriendlyName(UEdGraphSchema_CustomizableObject::PC_Material).ToString();
-	const FString PinName = PinFriendlyName + FString(TEXT("_Input_Pin"));
-
-	UEdGraphPin* Pin = FindPin(PinName);
-
-	if (Pin)
-	{
-		return Pin;
-	}
-	else
-	{
-		// Changing the locale can change the GetPinCategoryFriendlyName() result, leading to failing to find the MaterialAsset pin.
-		// This is a provisional fix while incorrect FriendlyName usages are removed in UE-294583
-		Pin = FindPin(PinFriendlyName);
-
-		if (Pin)
-		{
-			return Pin;
-		}
-		else
-		{
-			return FindPin(FString("Material_Input_Pin"));
-		}
-	}
+	return InputMaterialPin.Get();
 }
 
 
@@ -1499,23 +1521,7 @@ void UCONodeSkeletalMeshSection::PinConnectionListChanged(UEdGraphPin* Pin)
 	
 	if (Pin == GetMeshPin())
 	{
-		if (LastMeshNodeConnected.IsValid())
-		{
-			LastMeshNodeConnected->PostEditChangePropertyDelegate.RemoveDynamic(this, &UCONodeSkeletalMeshSection::MeshPostEditChangeProperty);
-		}
-
-		if (const UEdGraphPin* ConnectedPin = FollowInputPin(*Pin))
-		{
-			UEdGraphNode* MeshNode = ConnectedPin->GetOwningNode();
-
-			if (MeshNode->IsA(UCustomizableObjectNodeStaticMesh::StaticClass()) || MeshNode->IsA(UCustomizableObjectNodeSkeletalMesh::StaticClass()))
-			{
-				SetDefaultMaterial();
-
-				LastMeshNodeConnected = Cast<UCustomizableObjectNode>(MeshNode);
-				LastMeshNodeConnected->PostEditChangePropertyDelegate.AddUniqueDynamic(this, &UCONodeSkeletalMeshSection::MeshPostEditChangeProperty);
-			}
-		}
+			SetDefaultMaterial();
 	}
 	else if (Cast<UCONodeSkeletalMeshSectionPinDataImage>(GetPinData(*Pin))) // Image pin
 	{
@@ -1570,21 +1576,6 @@ bool UCONodeSkeletalMeshSection::CanConnect(const UEdGraphPin* InOwnedInputPin,
 }
 
 
-void UCONodeSkeletalMeshSection::MeshPostEditChangeProperty(FPostEditChangePropertyDelegateParameters& Parameters)
-{
-	if (const UEdGraphPin* MeshPin = FindPin(TEXT("Mesh")))
-	{
-		if (const UEdGraphPin* ConnectedPin = FollowInputPin(*MeshPin);
-			ConnectedPin && ConnectedPin->GetOwningNode() == Parameters.Node)
-		{
-			SetDefaultMaterial();
-		}
-		else if (UCustomizableObjectNode* MeshNode = Cast<UCustomizableObjectNode>(Parameters.Node))
-		{
-			MeshNode->PostEditChangePropertyDelegate.RemoveDynamic(this, &UCONodeSkeletalMeshSection::MeshPostEditChangeProperty);
-		}
-	}
-}
 
 
 #undef LOCTEXT_NAMESPACE
diff --git a/Engine/Plugins/Mutable/Source/CustomizableObjectEditor/Private/MuCOE/Nodes/CONodeSkeletalMeshSection.h b/Engine/Plugins/Mutable/Source/CustomizableObjectEditor/Private/MuCOE/Nodes/CONodeSkeletalMeshSection.h
--- a/Engine/Plugins/Mutable/Source/CustomizableObjectEditor/Private/MuCOE/Nodes/CONodeSkeletalMeshSection.h
+++ b/Engine/Plugins/Mutable/Source/CustomizableObjectEditor/Private/MuCOE/Nodes/CONodeSkeletalMeshSection.h
@@ -185,6 +185,12 @@ private:
 	UPROPERTY()
 	int32 MeshComponentIndex_DEPRECATED = 0;
 
+	UPROPERTY()
+	FEdGraphPinReference InputMeshPin;
+	
+	UPROPERTY()
+	FEdGraphPinReference InputMaterialPin;
+
 public:
 	/** Selects which Mesh component of the Instance this material belongs to */
 	UPROPERTY()
@@ -215,10 +221,6 @@ private:
 	/** Set the default Material from the connected static or skeletal mesh. */
 	UE_API void SetDefaultMaterial();
 
-	/** Connected NodeStaticMesh or NodeSkeletalMesh Mesh UPROPERTY changed callback function. Sets the default material. */
-	UFUNCTION()
-	UE_API void MeshPostEditChangeProperty(FPostEditChangePropertyDelegateParameters& Parameters);
-
 	/** Format pin name. */
 	UE_API FName GetPinName(EMaterialParameterType Type, int32 ParameterIndex) const;
 
diff --git a/Engine/Plugins/Mutable/Source/CustomizableObjectEditor/Private/MuCOE/Nodes/CustomizableObjectNode.cpp b/Engine/Plugins/Mutable/Source/CustomizableObjectEditor/Private/MuCOE/Nodes/CustomizableObjectNode.cpp
--- a/Engine/Plugins/Mutable/Source/CustomizableObjectEditor/Private/MuCOE/Nodes/CustomizableObjectNode.cpp
+++ b/Engine/Plugins/Mutable/Source/CustomizableObjectEditor/Private/MuCOE/Nodes/CustomizableObjectNode.cpp
@@ -72,13 +72,7 @@ bool UCustomizableObjectNode::ShouldBreakExistingConnections(const UEdGraphPin*
 void UCustomizableObjectNode::PostEditChangeProperty(FPropertyChangedEvent& PropertyChangedEvent)
 {
 	Super::PostEditChangeProperty(PropertyChangedEvent);
-
-	FPostEditChangePropertyDelegateParameters Parameters;
-	Parameters.Node = this;
-	Parameters.FPropertyChangedEvent = &PropertyChangedEvent;
-
-	PostEditChangePropertyDelegate.Broadcast(Parameters);
-	PostEditChangePropertyRegularDelegate.Broadcast(this, PropertyChangedEvent);
+	PostEditChangePropertyDelegate.Broadcast(this, PropertyChangedEvent);
 }
 
 
diff --git a/Engine/Plugins/Mutable/Source/CustomizableObjectEditor/Private/MuCOE/Nodes/CustomizableObjectNodeSwitchBase.cpp b/Engine/Plugins/Mutable/Source/CustomizableObjectEditor/Private/MuCOE/Nodes/CustomizableObjectNodeSwitchBase.cpp
--- a/Engine/Plugins/Mutable/Source/CustomizableObjectEditor/Private/MuCOE/Nodes/CustomizableObjectNodeSwitchBase.cpp
+++ b/Engine/Plugins/Mutable/Source/CustomizableObjectEditor/Private/MuCOE/Nodes/CustomizableObjectNodeSwitchBase.cpp
@@ -3,10 +3,10 @@
 #include "MuCOE/Nodes/CustomizableObjectNodeSwitchBase.h"
 
 #include "MuCO/CustomizableObjectCustomVersion.h"
-#include "MuCOE/CustomizableObjectPin.h"
 #include "MuCOE/EdGraphSchema_CustomizableObject.h"
 #include "MuCOE/GraphTraversal.h"
 #include "MuCOE/Nodes/CustomizableObjectNodeEnumParameter.h"
+#include "MuT/NodeScalarEnumParameter.h"
 
 #include UE_INLINE_GENERATED_CPP_BY_NAME(CustomizableObjectNodeSwitchBase)
 
@@ -20,7 +20,7 @@ void UCustomizableObjectNodeSwitchBase::PinConnectionListChanged(UEdGraphPin* Pi
 {
 	if (Pin == SwitchParameterPinReference.Get())
 	{
-		LinkPostEditChangePropertyDelegate(*Pin);
+		Super::ReconstructNode();
 	}
 
 	Super::PinConnectionListChanged(Pin);
@@ -39,13 +39,15 @@ void UCustomizableObjectNodeSwitchBase::AllocateDefaultPins(UCustomizableObjectN
 	SwitchPins.Empty();	
 	if (OldEnumPin) // May be null the first time the node is created.
 	{
+		// This will fail if the node is connected to a macro node which will then prevent the generation of the input option pins
+		// todo : UE-294591
 		if (const UEdGraphPin* LinkedPin = FollowInputPin(*OldEnumPin))
 		{
 			if (UCustomizableObjectNodeEnumParameter* EnumNode = Cast<UCustomizableObjectNodeEnumParameter>(LinkedPin->GetOwningNode()))
 			{
-				for (int i = 0; i < EnumNode->Values.Num(); ++i)
+				for (int32 EnumValueIndex = 0; EnumValueIndex < EnumNode->Values.Num(); ++EnumValueIndex)
 				{
-					const FString OptionName = EnumNode->Values[i].Name;
+					const FString OptionName = EnumNode->Values[EnumValueIndex].Name;
 					SwitchPins.Add(CustomCreatePin(EGPD_Input, GetCategory(), FName(*OptionName), FText::FromString(OptionName)));
 				}
 			}
@@ -72,27 +74,6 @@ FText UCustomizableObjectNodeSwitchBase::GetTooltipText() const
 }
 
 
-void UCustomizableObjectNodeSwitchBase::EnumParameterPostEditChangeProperty(FPostEditChangePropertyDelegateParameters& Parameters)
-{
-	if (const UEdGraphPin* ConnectedPin = FollowInputPin(*SwitchParameterPinReference.Get()); ConnectedPin && ConnectedPin->GetOwningNode() == Parameters.Node)
-	{
-		Super::ReconstructNode();
-	}
-	else if (UCustomizableObjectNode* EnumNode = Cast<UCustomizableObjectNode>(Parameters.Node))
-	{
-		EnumNode->PostEditChangePropertyDelegate.RemoveDynamic(this, &UCustomizableObjectNodeSwitchBase::EnumParameterPostEditChangeProperty);
-	}
-}
-
-
-void UCustomizableObjectNodeSwitchBase::PostPasteNode()
-{
-	Super::PostPasteNode();
-
-	LinkPostEditChangePropertyDelegate(*SwitchParameterPinReference.Get());
-}
-
-
 void UCustomizableObjectNodeSwitchBase::BackwardsCompatibleFixup(int32 CustomizableObjectCustomVersion)
 {
 	Super::BackwardsCompatibleFixup(CustomizableObjectCustomVersion);
@@ -108,7 +89,7 @@ void UCustomizableObjectNodeSwitchBase::BackwardsCompatibleFixup(int32 Customiza
 		const FString PinPrefix = UEdGraphSchema_CustomizableObject::GetPinCategoryName(GetCategory()).ToString() + " ";
 
 		int32 NumOptions = 0;
-		for (UEdGraphPin* Pin : GetAllNonOrphanPins())
+		for (const UEdGraphPin* Pin : GetAllNonOrphanPins())
 		{
 			if (Pin->GetName().StartsWith(PinPrefix))
 			{
@@ -132,7 +113,7 @@ void UCustomizableObjectNodeSwitchBase::BackwardsCompatibleFixup(int32 Customiza
 		const FString PinPrefix = UEdGraphSchema_CustomizableObject::GetPinCategoryName(GetCategory()).ToString() + " ";
 
 		int32 NumOptions = 0;
-		for (UEdGraphPin* Pin : GetAllNonOrphanPins())
+		for (const UEdGraphPin* Pin : GetAllNonOrphanPins())
 		{
 			if (Pin->GetName().StartsWith(PinPrefix))
 			{
@@ -145,9 +126,9 @@ void UCustomizableObjectNodeSwitchBase::BackwardsCompatibleFixup(int32 Customiza
 		{
 			if (UCustomizableObjectNodeEnumParameter* EnumNode = Cast<UCustomizableObjectNodeEnumParameter>(LinkedPin->GetOwningNode()))
 			{
-				for (int i = 0; i < EnumNode->Values.Num(); ++i)
+				for (int32 ValueIndex = 0; ValueIndex < EnumNode->Values.Num(); ++ValueIndex)
 				{
-					OptionNames.Add(EnumNode->Values[i].Name);
+					OptionNames.Add(EnumNode->Values[ValueIndex].Name);
 				}
 			}
 		}
@@ -173,31 +154,6 @@ void UCustomizableObjectNodeSwitchBase::BackwardsCompatibleFixup(int32 Customiza
 void UCustomizableObjectNodeSwitchBase::PostBackwardsCompatibleFixup()
 {
 	Super::PostBackwardsCompatibleFixup();
-
-	if (UEdGraphPin* SwitchPin = SwitchParameterPinReference.Get()) // New nodes execute AllocateDefaultPins after PostBackwardsCompatibleFixup.
-	{
-		LinkPostEditChangePropertyDelegate(*SwitchPin);
-	}
-}
-
-
-void UCustomizableObjectNodeSwitchBase::LinkPostEditChangePropertyDelegate(const UEdGraphPin& Pin)
-{
-	if (LastNodeEnumParameterConnected.IsValid())
-	{
-		LastNodeEnumParameterConnected->PostEditChangePropertyDelegate.RemoveDynamic(this, &UCustomizableObjectNodeSwitchBase::EnumParameterPostEditChangeProperty);
-	}
-
-	if (const UEdGraphPin* ConnectedPin = FollowInputPin(Pin))
-	{
-		LastNodeEnumParameterConnected = Cast<UCustomizableObjectNode>(ConnectedPin->GetOwningNode());
-
-		if (LastNodeEnumParameterConnected.IsValid())
-		{
-			LastNodeEnumParameterConnected->PostEditChangePropertyDelegate.AddUniqueDynamic(this, &UCustomizableObjectNodeSwitchBase::EnumParameterPostEditChangeProperty);
-		}
-	}
-
 	Super::ReconstructNode();
 }
 
@@ -209,4 +165,3 @@ FString UCustomizableObjectNodeSwitchBase::GetOutputPinName() const
 
 
 #undef LOCTEXT_NAMESPACE
-
diff --git a/Engine/Plugins/Mutable/Source/CustomizableObjectEditor/Private/MuCOE/Nodes/CustomizableObjectNodeSwitchBase.h b/Engine/Plugins/Mutable/Source/CustomizableObjectEditor/Private/MuCOE/Nodes/CustomizableObjectNodeSwitchBase.h
--- a/Engine/Plugins/Mutable/Source/CustomizableObjectEditor/Private/MuCOE/Nodes/CustomizableObjectNodeSwitchBase.h
+++ b/Engine/Plugins/Mutable/Source/CustomizableObjectEditor/Private/MuCOE/Nodes/CustomizableObjectNodeSwitchBase.h
@@ -15,6 +15,10 @@ class UCustomizableObjectNodeRemapPins;
 class UObject;
 struct FFrame;
 
+namespace UE::Mutable::Private
+{
+	class NodeScalarEnumParameter;
+}
 
 UCLASS(MinimalAPI, Abstract)
 class UCustomizableObjectNodeSwitchBase : public UCustomizableObjectNode
@@ -27,7 +31,6 @@ public:
 	virtual FLinearColor GetNodeTitleColor() const override;
 	virtual FText GetTooltipText() const override;
 	virtual void PinConnectionListChanged(UEdGraphPin* Pin) override;
-	virtual void PostPasteNode() override;
 
 	// UCustomizableObjectNode interface
 	virtual void BackwardsCompatibleFixup(int32 CustomizableObjectCustomVersion) override;
@@ -36,21 +39,15 @@ public:
 
 	/** Get the output pin category. Override. */
 	virtual FName GetCategory() const PURE_VIRTUAL(UCustomizableObjectNodeSwitchBase::GetCategory, return {}; );
-	
-	/** Links the PostEditChangeProperty delegate */
-	void LinkPostEditChangePropertyDelegate(const UEdGraphPin& Pin);
 
 	/** Get the ouput pin name. Override. */
 	virtual FString GetOutputPinName() const;
-
+	
 private:
+	
 	/** Last NodeEnumParameter connected. Used to remove the callback once disconnected. */
 	TWeakObjectPtr<UCustomizableObjectNode> LastNodeEnumParameterConnected;
 
-	/** NodeEnumParameter property changed callback function. Reconstructs the node. */
-	UFUNCTION()
-	void EnumParameterPostEditChangeProperty(FPostEditChangePropertyDelegateParameters& Parameters);
-
 public:
 	UPROPERTY()
 	FEdGraphPinReference SwitchParameterPinReference;
diff --git a/Engine/Plugins/Mutable/Source/CustomizableObjectEditor/Public/MuCOE/Nodes/CustomizableObjectNode.h b/Engine/Plugins/Mutable/Source/CustomizableObjectEditor/Public/MuCOE/Nodes/CustomizableObjectNode.h
--- a/Engine/Plugins/Mutable/Source/CustomizableObjectEditor/Public/MuCOE/Nodes/CustomizableObjectNode.h
+++ b/Engine/Plugins/Mutable/Source/CustomizableObjectEditor/Public/MuCOE/Nodes/CustomizableObjectNode.h
@@ -79,21 +79,7 @@ protected:
 	virtual bool Equals(const UCustomizableObjectNodePinData& Other) const;
 };
 
-/** Encapsulation of parameters for the FPostEditChangePropertyDelegate delegate function.
- *
- * The delegate requires the function being a UFUNCTION which requires all parameters to be UCLASS, USTRUC or UENUM.
- */
-USTRUCT()
-struct FPostEditChangePropertyDelegateParameters
-{
-	GENERATED_BODY()
-
-	UObject* Node;
-	FPropertyChangedEvent* FPropertyChangedEvent;
-};
-
-DECLARE_DYNAMIC_MULTICAST_DELEGATE_OneParam(FPostEditChangePropertyDelegate, FPostEditChangePropertyDelegateParameters&, Parameters); // Deprecated, use non-dynamic version.
-DECLARE_MULTICAST_DELEGATE_TwoParams(FPostEditChangePropertyRegularDelegate, UObject*, FPropertyChangedEvent&);
+DECLARE_MULTICAST_DELEGATE_TwoParams(FPostEditChangePropertyDelegate, UObject*, FPropertyChangedEvent&);
 
 DECLARE_MULTICAST_DELEGATE(FPostReconstructNodeDelegate);
 
@@ -360,11 +346,8 @@ public:
 
 	UE_API virtual bool IsCompatibleWithGraph(UEdGraph const* Graph) const override;
 	
-	UPROPERTY()
 	FPostEditChangePropertyDelegate PostEditChangePropertyDelegate;
 
-	FPostEditChangePropertyRegularDelegate PostEditChangePropertyRegularDelegate;
-
 	FPostReconstructNodeDelegate PostReconstructNodeDelegate;
 
 	FNodeConnectionListChangedDelegate NodeConnectionListChangedDelegate;
