[Mutable] Add support to allow parametric images to load the mip tail.

#rnx
#Synced-CL 49758930

diff --git a/Engine/Plugins/Mutable/Source/CustomizableObject/Internal/MuCO/UnrealMutableImageProvider.h b/Engine/Plugins/Mutable/Source/CustomizableObject/Internal/MuCO/UnrealMutableImageProvider.h
--- a/Engine/Plugins/Mutable/Source/CustomizableObject/Internal/MuCO/UnrealMutableImageProvider.h
+++ b/Engine/Plugins/Mutable/Source/CustomizableObject/Internal/MuCO/UnrealMutableImageProvider.h
@@ -19,7 +19,7 @@ public:
 	
 	// UE::Mutable::Private::ExternalResourceProvider interface
 	// Thread: worker
-	CUSTOMIZABLEOBJECT_API virtual TTuple<UE::Tasks::FTask, TFunction<void()>> GetImageAsync(UTexture* Texture, uint8 MipmapsToSkip, TFunction<void(UE::Mutable::Private::TManagedPtr<UE::Mutable::Private::FImage>)>& ResultCallback) override;
+	CUSTOMIZABLEOBJECT_API virtual TTuple<UE::Tasks::FTask, TFunction<void()>> GetImageAsync(UTexture* Texture, uint8 MipmapsToSkip, bool bLoadMipTail, TFunction<void(UE::Mutable::Private::TManagedPtr<UE::Mutable::Private::FImage>)>& ResultCallback) override;
 	CUSTOMIZABLEOBJECT_API virtual TTuple<UE::Tasks::FTask, TFunction<void()>> GetReferencedImageAsync( int32 Id, uint8 MipmapsToSkip, TFunction<void(UE::Mutable::Private::TManagedPtr<UE::Mutable::Private::FImage>)>& ResultCallback) override;
 	CUSTOMIZABLEOBJECT_API virtual UE::Mutable::Private::FExtendedImageDesc GetImageDesc(UTexture* Texture) override;
 	CUSTOMIZABLEOBJECT_API virtual TTuple<UE::Tasks::FTask, TFunction<void()>> GetMeshAsync(USkeletalMesh* SkeletalMesh, int32 LODIndex, int32 SectionIndex, TFunction<void(UE::Mutable::Private::TManagedPtr<UE::Mutable::Private::FMesh>)>& ResultCallback) override;
diff --git a/Engine/Plugins/Mutable/Source/CustomizableObject/Internal/MuCO/UnrealToMutableTextureConversionUtils.h b/Engine/Plugins/Mutable/Source/CustomizableObject/Internal/MuCO/UnrealToMutableTextureConversionUtils.h
--- a/Engine/Plugins/Mutable/Source/CustomizableObject/Internal/MuCO/UnrealToMutableTextureConversionUtils.h
+++ b/Engine/Plugins/Mutable/Source/CustomizableObject/Internal/MuCO/UnrealToMutableTextureConversionUtils.h
@@ -151,7 +151,7 @@ enum class EUnrealToMutableConversionError
     Unknown
 };
 
-UE_API EUnrealToMutableConversionError ConvertTextureUnrealSourceToMutable(UE::Mutable::Private::FImage* OutResult, FMutableSourceTextureData&, uint8 MipmapsToSkip);
+UE_API EUnrealToMutableConversionError ConvertTextureUnrealSourceToMutable(UE::Mutable::Private::FImage* OutResult, FMutableSourceTextureData&, uint8 MipmapsToSkip, bool bLoadMipTail);
 
 /** Mesh conversion context. */
 class FMeshConversionContext
diff --git a/Engine/Plugins/Mutable/Source/CustomizableObject/Private/MuCO/UnrealMutableImageProvider.cpp b/Engine/Plugins/Mutable/Source/CustomizableObject/Private/MuCO/UnrealMutableImageProvider.cpp
--- a/Engine/Plugins/Mutable/Source/CustomizableObject/Private/MuCO/UnrealMutableImageProvider.cpp
+++ b/Engine/Plugins/Mutable/Source/CustomizableObject/Private/MuCO/UnrealMutableImageProvider.cpp
@@ -23,6 +23,7 @@
 #include "TextureResource.h"
 #include "Engine/Texture.h"
 
+#include "IO/IoBuffer.h"
 
 namespace
 {
@@ -98,19 +99,43 @@ FUnrealMutableResourceProvider::FUnrealMutableResourceProvider(const TSharedRef<
 {
 }
 
+namespace UnrealMutableImageProviderInteranl
+{
+	bool CheckBulkSizesAreValid(TNotNull<UTexture2D*> Texture, TNotNull<UE::Mutable::Private::FImage*> Image, int32 TextureMipIndexBegin, int32 TextureMipIndexEnd)
+	{
+		int32 TextureNumMips = TextureMipIndexEnd - TextureMipIndexBegin; 
+
+		if (TextureNumMips != Image->DataStorage.GetNumLODs())
+		{
+			return false;
+		}
+
+		for (int32 TextureMipIndex = TextureMipIndexBegin; TextureMipIndex < TextureMipIndexEnd; ++TextureMipIndex)
+		{
+			FByteBulkData& BulkData = Texture->GetPlatformData()->Mips[TextureMipIndex].BulkData;
+			int32 BulkDataSize = BulkData.GetBulkDataSize();
+			int32 ImageDataSize = Image->DataStorage.GetLOD(TextureMipIndex - TextureMipIndexBegin).Num(); 
+
+			if (BulkDataSize != ImageDataSize)
+			{
+				return false;
+			}
+		}
+
+		return true;
+	}
+} // namespace UnrealMutableImageProviderInteranl
 
 //-------------------------------------------------------------------------------------------------
-TTuple<UE::Tasks::FTask, TFunction<void()>> FUnrealMutableResourceProvider::GetImageAsync(UTexture* Texture, uint8 MipmapsToSkip, TFunction<void(UE::Mutable::Private::TManagedPtr<UE::Mutable::Private::FImage>)>& ResultCallback)
+TTuple<UE::Tasks::FTask, TFunction<void()>> FUnrealMutableResourceProvider::GetImageAsync(UTexture* Texture, uint8 MipmapsToSkip, bool bLoadMipTail, TFunction<void(UE::Mutable::Private::TManagedPtr<UE::Mutable::Private::FImage>)>& ResultCallback)
 {
+	using namespace UnrealMutableImageProviderInteranl;
+
 	// Thread: worker
 	MUTABLE_CPUPROFILER_SCOPE(FUnrealMutableImageProvider::GetImageAsync);
 
 	// Some data that may have to be copied from the GlobalExternalImages while it's locked
-	IBulkDataIORequest* IORequest = nullptr;
-	const int32 LODs = 1;
-
 	EPixelFormat Format = EPixelFormat::PF_Unknown;
-	int32 BulkDataSize = 0;
 
 	UE::Mutable::Private::EImageFormat MutImageFormat = UE::Mutable::Private::EImageFormat::None;
 	int32 MutImageDataSize = 0;
@@ -144,8 +169,8 @@ TTuple<UE::Tasks::FTask, TFunction<void()>> FUnrealMutableResourceProvider::GetI
 	// In the editor the src data can be directly accessed
 	UE::Mutable::Private::TManagedPtr<UE::Mutable::Private::FImage> Image = UE::Mutable::Private::MakeManaged<UE::Mutable::Private::FImage>();
 
-	FMutableSourceTextureData Tex(*Texture2D);
-	EUnrealToMutableConversionError Error = ConvertTextureUnrealSourceToMutable(Image.Get(), Tex, MipIndex);
+	FMutableSourceTextureData TextureData(*Texture2D);
+	EUnrealToMutableConversionError Error = ConvertTextureUnrealSourceToMutable(Image.Get(), TextureData, MipIndex, bLoadMipTail);
 	if (Error != EUnrealToMutableConversionError::Success)
 	{
 		// This could happen in the editor, because some source textures may have changed while there was a background compilation.
@@ -155,23 +180,11 @@ TTuple<UE::Tasks::FTask, TFunction<void()>> FUnrealMutableResourceProvider::GetI
 
 	ResultCallback(Image);
 	return Invoke(TrivialReturn);
-		
 #else
+
 	// It's safe to access TextureToLoad because ExternalImagesLock guarantees that the data in GlobalExternalImages is valid,
 	// not being modified by the game thread at the moment and the texture cannot be GCed because of the AddReferencedObjects
 	// in the FUnrealMutableImageProvider
-
-	int32 MipIndex = MipmapsToSkip < Texture2D->GetPlatformData()->Mips.Num() ? MipmapsToSkip : Texture2D->GetPlatformData()->Mips.Num() - 1;
-	check (MipIndex >= 0);
-
-	// Mips in the mip tail are inlined and can't be streamed, find the smallest mip available.
-	for (; MipIndex > 0; --MipIndex)
-	{
-		if (Texture2D->GetPlatformData()->Mips[MipIndex].BulkData.CanLoadFromDisk())
-		{
-			break;
-		}
-	}
 	
 	// Texture format and the equivalent mutable format
 	Format = Texture2D->GetPlatformData()->PixelFormat;
@@ -185,148 +198,91 @@ TTuple<UE::Tasks::FTask, TFunction<void()>> FUnrealMutableResourceProvider::GetI
 		return Invoke(TrivialReturn);
 	}
 
-	int32 SizeX = Texture2D->GetSizeX() >> MipIndex;
-	int32 SizeY = Texture2D->GetSizeY() >> MipIndex;
+	const int32 TextureNumMips = Texture2D->GetPlatformData()->Mips.Num();
 
-	check(LODs == 1);
-	UE::Mutable::Private::TManagedPtr<UE::Mutable::Private::FImage> Image = MakeManaged<UE::Mutable::Private::FImage>(SizeX, SizeY, LODs, MutImageFormat, UE::Mutable::Private::EInitializationType::NotInitialized);
-	TArrayView<uint8> MutImageDataView = Image->DataStorage.GetLOD(0);
+	int32 TextureMipIndexBegin = MipmapsToSkip < TextureNumMips ? MipmapsToSkip : TextureNumMips - 1;
+	check(TextureMipIndexBegin >= 0);
 
-	// In a packaged game the bulk data has to be loaded
-	// Get the actual file to read the mip 0 data, do not keep any reference to Texture2D because once outside of the lock
-	// it may be GCed or changed. Just keep the actual file handle and some sizes instead of the texture
-	FByteBulkData& BulkData = Texture2D->GetPlatformData()->Mips[MipIndex].BulkData;
-	BulkDataSize = BulkData.GetBulkDataSize();
-	check(BulkDataSize > 0);
+	int32 NumMipsToLoad = bLoadMipTail ? TextureMipIndexBegin - TextureNumMips : 1;
+	int32 LastMipToLoad = bLoadMipTail ? TextureNumMips - 1 : TextureMipIndexBegin;
 
-	if (BulkDataSize != MutImageDataView.Num())
+	// Find the smallest mip available.
+	for (; LastMipToLoad > 0; --LastMipToLoad)
 	{
-		UE_LOG(LogMutable, Warning, TEXT("Failed to get Image Parameter [%s]. Bulk data size is different than the expected size. BulkData size [%d]. Mutable image data size [%d]."),
-			*Texture->GetName(), BulkDataSize, MutImageDataSize);
-
-		ResultCallback(CreateDummy());
-		return Invoke(TrivialReturn);
+		if (Texture2D->GetPlatformData()->Mips[LastMipToLoad].BulkData.CanLoadFromDisk())
+		{
+			break;
+		}
 	}
 
-	// Create a streaming request if the data is not loaded or copy the mip data
-	if (!BulkData.IsBulkDataLoaded())
-	{
-		UE::Tasks::FTaskEvent IORequestCompletionEvent(TEXT("Mutable_IORequestCompletionEvent"));
-
-		TFunction<void(bool, IBulkDataIORequest*)> IOCallback =
-			[
-				MutImageDataView,
-				MutImageFormat,
-				Format,
-				Image,
-				BulkDataSize,	
-				ResultCallback, // Notice ResultCallback is captured by copy
-				IORequestCompletionEvent
-			](bool bWasCancelled, IBulkDataIORequest* IORequest)
-		{
-			ON_SCOPE_EXIT
-			{
-				UE::Tasks::FTaskEvent EventCopy = IORequestCompletionEvent;
-				EventCopy.Trigger();
-			};
-			
-			// Should we do someting different than returning a dummy image if cancelled?
-			if (bWasCancelled)
-			{
-				UE_LOG(LogMutable, Warning, TEXT("Failed to get external image. Cancelled IO Request"));
-				ResultCallback(CreateDummy());
-				return;
-			}
+	TextureMipIndexBegin = FMath::Min(LastMipToLoad, TextureMipIndexBegin);
+	int32 TextureMipIndexEnd = LastMipToLoad + 1;
 
-			uint8* Results = IORequest->GetReadResults(); // required?
+	int32 SizeX = FMath::Max(Texture2D->GetSizeX() >> TextureMipIndexBegin, 1);
+	int32 SizeY = FMath::Max(Texture2D->GetSizeY() >> TextureMipIndexBegin, 1);
 
-			if (Results && MutImageDataView.Num() == (int32)IORequest->GetSize())
-			{
-				check(BulkDataSize == (int32)IORequest->GetSize());
-				check(Results == MutImageDataView.GetData());
+	int32 NumImageLODs = TextureMipIndexEnd - TextureMipIndexBegin;
+	UE::Mutable::Private::TManagedPtr<UE::Mutable::Private::FImage> Image = MakeManaged<UE::Mutable::Private::FImage>(SizeX, SizeY, NumImageLODs, MutImageFormat, UE::Mutable::Private::EInitializationType::NotInitialized);
 
-				ResultCallback(Image);
-				return;
-			}
+	if (!CheckBulkSizesAreValid(Texture2D, Image.Get(), TextureMipIndexBegin, TextureMipIndexEnd))
+	{
+		UE_LOG(LogMutable, Warning, TEXT("Failed to get Image Parameter [%s]. Bulk data does not match expected image data sizes."), *Texture2D->GetName());
+		ResultCallback(CreateDummy());
 
-			if (!Results)
-			{
-				UE_LOG(LogMutable, Warning, TEXT("Failed to get Image Parameter. IO Request failed. Request results [%hhd]. Format: [%s]. MutableFormat: [%d]."),
-					(Results != nullptr),
-					GetPixelFormatString(Format),
-					(int32)MutImageFormat);
-			}
-			else if (MutImageDataView.Num() != (int32)IORequest->GetSize())
-			{
-				UE_LOG(LogMutable, Warning, TEXT("Failed to get Image Parameter. Requested size is different than the expected size. RequestSize: [%lld]. ExpectedSize: [%d]. Format: [%s]. MutableFormat: [%d]."),
-					IORequest->GetSize(),
-					MutImageDataView.Num(),
-					GetPixelFormatString(Format),
-					(int32)MutImageFormat);
-			}
-			else
-			{
-				UE_LOG(LogMutable, Warning, TEXT("Failed to get Image Parameter."));
-			}
+		return Invoke(TrivialReturn);
+	}
 
-			// Something failed when loading the bulk data, just return a dummy
-			ResultCallback(CreateDummy());
-		};
+	const int32 BatchCount = TextureMipIndexEnd - TextureMipIndexBegin;
+	FBulkDataBatchRequest::FBatchBuilder Batch = FBulkDataBatchRequest::NewBatch(BatchCount);
 	
-		// Is the resposability of the CreateStreamingRequest caller to delete the IORequest. 
-		// This can *not* be done in the IOCallback because it would cause a deadlock so it is deferred to the returned
-		// cleanup function. Another solution could be to spwan a new task that depends on the 
-		// IORequestComplitionEvent which deletes it.
-		TRACE_IOSTORE_METADATA_SCOPE_TAG(FName(Texture->GetPathName()));
-		IORequest = BulkData.CreateStreamingRequest(EAsyncIOPriorityAndFlags::AIOP_High, &IOCallback, MutImageDataView.GetData());
-
-		if (IORequest)
-		{
-			// Make the lambda mutable and set the IORequest pointer to null when deleted so it is safer 
-			// agains multiple calls.
-			const auto DeleteIORequest = [IORequest]() mutable -> void
-			{
-				if (IORequest)
-				{
-					delete IORequest;
-				}
-				
-				IORequest = nullptr;
-			};
-
-			return MakeTuple(IORequestCompletionEvent, DeleteIORequest);
-
-		}
-		else
+	for (int32 TextureMipIndex = TextureMipIndexBegin; TextureMipIndex < TextureMipIndexEnd; ++TextureMipIndex)
+	{
+		FByteBulkData& BulkData = Texture2D->GetPlatformData()->Mips[TextureMipIndex].BulkData;
+		int32 BulkDataSize = BulkData.GetBulkDataSize();
+		if (BulkDataSize > 0)
 		{
-			UE_LOG(LogMutable, Warning, TEXT("Failed to create an IORequest for a UTexture2D BulkData for an application-specific image parameter."));
+			TArrayView<uint8> ImageDataView = Image->DataStorage.GetLOD(TextureMipIndex - TextureMipIndexBegin);
+			check(BulkDataSize == ImageDataView.Num());
 
-			IORequestCompletionEvent.Trigger();
-			
-			ResultCallback(CreateDummy());
-			return Invoke(TrivialReturn);
+			TRACE_IOSTORE_METADATA_SCOPE_TAG(FName(Texture->GetPathName()));
+			FIoBuffer Dst(FIoBuffer::Wrap, ImageDataView.GetData(), ImageDataView.Num());
+			Batch.Read(BulkData, 0, ImageDataView.Num(), AIOP_High | AIOP_FLAG_DONTCACHE, Dst);
 		}
 	}
-	else
+		
+	UE::Tasks::FTaskEvent RequestCompletionEvent(TEXT("Mutable_ImageParamRequestCompletionEvent"));
+	TSharedPtr<FBulkDataBatchRequest> ResultBulkDataRequest = MakeShared<FBulkDataBatchRequest>();
+	Batch.Issue(
+	[
+		Image,
+		RequestCompletionEvent,
+		ResultCallback // captured by copy
+	](FBulkDataBatchRequest::EStatus Status) mutable
 	{
-		// Bulk data already loaded
-		const void* Data = (!BulkData.IsLocked()) ? BulkData.LockReadOnly() : nullptr; // TODO: Retry if it fails?
+		ON_SCOPE_EXIT
+		{
+			RequestCompletionEvent.Trigger();
+		};
 		
-		if (Data)
+		// Should we do someting different than returning a dummy image if cancelled?
+		if (Status == FBulkDataRequest::EStatus::Cancelled)
 		{
-			FMemory::Memcpy(MutImageDataView.GetData(), Data, BulkDataSize);
-
-			BulkData.Unlock();
-			ResultCallback(Image);
-			return Invoke(TrivialReturn);
+			UE_LOG(LogMutable, Warning, TEXT("Failed to get external image. Cancelled Bulk Data Request"));
+			ResultCallback(CreateDummy());
+			return;
 		}
-		else
+		
+		if (Status != FBulkDataRequest::EStatus::Ok)
 		{
-			UE_LOG(LogMutable, Warning, TEXT("Failed to get Image Parameter. Bulk data already locked or null."));
+			UE_LOG(LogMutable, Warning, TEXT("Failed to get external image parameter."));
 			ResultCallback(CreateDummy());
-			return Invoke(TrivialReturn);
+			return;
 		}
-	}
+		
+		ResultCallback(Image);
+	}, *ResultBulkDataRequest);
+
+	return MakeTuple(RequestCompletionEvent, [BulkDataRequest = MoveTemp(ResultBulkDataRequest)](){});
 #endif
 }
 
@@ -376,7 +332,9 @@ TTuple<UE::Tasks::FTask, TFunction<void()>> FUnrealMutableResourceProvider::GetR
 	UE::Mutable::Private::TManagedPtr<UE::Mutable::Private::FImage> Image = UE::Mutable::Private::MakeManaged<UE::Mutable::Private::FImage>();
 	
 	FMutableSourceTextureData Tex(*Texture2D);
-	EUnrealToMutableConversionError Error = ConvertTextureUnrealSourceToMutable(Image.Get(), Tex, MipIndex);
+
+	constexpr bool bLoadMipTail = false;
+	EUnrealToMutableConversionError Error = ConvertTextureUnrealSourceToMutable(Image.Get(), Tex, MipIndex, bLoadMipTail);
 	if (Error != EUnrealToMutableConversionError::Success)
 	{
 		// This could happen in the editor, because some source textures may have changed while updating.
diff --git a/Engine/Plugins/Mutable/Source/CustomizableObject/Private/MuCO/UnrealToMutableTextureConversionUtils.cpp b/Engine/Plugins/Mutable/Source/CustomizableObject/Private/MuCO/UnrealToMutableTextureConversionUtils.cpp
--- a/Engine/Plugins/Mutable/Source/CustomizableObject/Private/MuCO/UnrealToMutableTextureConversionUtils.cpp
+++ b/Engine/Plugins/Mutable/Source/CustomizableObject/Private/MuCO/UnrealToMutableTextureConversionUtils.cpp
@@ -114,7 +114,7 @@ EUnrealToMutableConversionError ApplyCompositeTexture(
     return EUnrealToMutableConversionError::Success;
 }
 
-void FlipGreenChannelRGBA32F(FImage& Image)
+void FlipGreenChannelRGBA32F(FImageView& Image)
 {
 	TArrayView64<FLinearColor> ImageDataView = Image.AsRGBA32F();
 	ParallelFor(ImageDataView.Num(),
@@ -124,7 +124,7 @@ void FlipGreenChannelRGBA32F(FImage& Image)
 		});
 }
 
-void FlipGreenChannelBGRA8(FImage& Image)
+void FlipGreenChannelBGRA8(FImageView& Image)
 {
 	TArrayView64<FColor> ImageDataView = Image.AsBGRA8();
 	ParallelFor(ImageDataView.Num(),
@@ -134,7 +134,7 @@ void FlipGreenChannelBGRA8(FImage& Image)
 		});
 }
 
-void Normalize(FImage& Image)
+void Normalize(FImageView& Image)
 {
 	TArrayView64<FLinearColor> ImageView = Image.AsRGBA32F();
 
@@ -148,7 +148,7 @@ void Normalize(FImage& Image)
 	}
 }
 
-void BlurNormalForComposite(FImage& Image)
+void BlurNormalForComposite(FImageView& Image)
 {
 	TArrayView64<FLinearColor> ImageView = Image.AsRGBA32F();
 
@@ -215,75 +215,107 @@ bool FMutableSourceTextureData::IsNormalComposite() const
 }
 
 
-EUnrealToMutableConversionError ConvertTextureUnrealSourceToMutable(UE::Mutable::Private::FImage* OutResult, FMutableSourceTextureData& Tex, uint8 MipmapsToSkip)
+EUnrealToMutableConversionError ConvertTextureUnrealSourceToMutable(UE::Mutable::Private::FImage* OutResult, FMutableSourceTextureData& Texture, uint8 MipmapsToSkip, bool bLoadMipTail)
 {
 	MUTABLE_CPUPROFILER_SCOPE(ConvertTextureUnrealSourceToMutable);
 
-    using namespace UnrealToMutableImageConversion_Internal;
+	using namespace UnrealToMutableImageConversion_Internal;
 
-	FTextureSource& Source = Tex.GetSource();
+	FTextureSource& Source = Texture.GetSource();
 	
 	// Correct mips to skip to fit source data
 	MipmapsToSkip = FMath::Clamp(MipmapsToSkip, 0, Source.GetNumMips()-1);
 
-    const int32 LODs = 1;
-    const int32 SizeX = Source.GetSizeX() >> MipmapsToSkip;
-    const int32 SizeY = Source.GetSizeY() >> MipmapsToSkip;
-	check(SizeX > 0 && SizeY > 0);
+	const int32 TextureNumMips = Source.GetNumMips();
 
-    ETextureSourceFormat Format = Source.GetFormat();
+	int32 TextureMipIndexBegin = MipmapsToSkip < TextureNumMips ? MipmapsToSkip : TextureNumMips - 1;
+	check(TextureMipIndexBegin >= 0);
+
+	int32 NumMipsToLoad = bLoadMipTail ? TextureMipIndexBegin - TextureNumMips : 1;
+	int32 LastMipToLoad = bLoadMipTail ? TextureNumMips - 1 : TextureMipIndexBegin;
+	int32 TextureMipIndexEnd = LastMipToLoad + 1;
+
+	const int32 SizeX = FMath::Max(Source.GetSizeX() >> MipmapsToSkip, 1);
+	const int32 SizeY = FMath::Max(Source.GetSizeY() >> MipmapsToSkip, 1);
+
+	ETextureSourceFormat Format = Source.GetFormat();
  
-    ERawImageFormat::Type RawFormat = ConvertFormatSourceToRaw(Format);
+	ERawImageFormat::Type RawFormat = ConvertFormatSourceToRaw(Format);
+
+	int32 ImageNumMips = TextureMipIndexEnd - TextureMipIndexBegin;
+	FMipMapImage TempImageStorage0;
+	FMipMapImage TempImageStorage1;
+
+	FMipMapImage* TempImage0 = &TempImageStorage0;
+	FMipMapImage* TempImage1 = &TempImageStorage1;
 
 	// What if source data is not linear?
-    FImage TempImage(SizeX, SizeY, 1, RawFormat, EGammaSpace::Linear);
-	FImage TempImage2;
+	TempImage0->Init(SizeX, SizeY, ImageNumMips, RawFormat, EGammaSpace::Linear);
 
-    if (!Source.GetMipData(TempImage.RawData, MipmapsToSkip))
-    {
-        return EUnrealToMutableConversionError::Unknown;
-    }
+	for (int32 MipLevel = 0; MipLevel < ImageNumMips; ++MipLevel)
+	{
+		FImageView MipView = TempImage0->GetMipImage(MipLevel);
+		
+		FImage OutMipImage(MipView.SizeX, MipView.SizeY, MipView.Format, MipView.GammaSpace);
+		if (!Source.GetMipImage(OutMipImage, TextureMipIndexBegin + MipLevel))
+		{
+			return EUnrealToMutableConversionError::Unknown;
+		}
+
+		check(MipView.GetImageSizeBytes() == OutMipImage.GetImageSizeBytes());
+		FMemory::Memcpy(MipView.RawData, OutMipImage.RawData.GetData(), MipView.GetImageSizeBytes());
+	}
 
-    bool bFlipGreenChannel = Tex.GetFlipGreenChannel();
+	bool bFlipGreenChannel = Texture.GetFlipGreenChannel();
 
-    // If any post processes of the image is needed, convert to RGBA32F
-    if (Tex.IsNormalComposite())
-    { 
+	// If any post processes of the image is needed, convert to RGBA32F
+	if (Texture.IsNormalComposite())
+	{ 
 		MUTABLE_CPUPROFILER_SCOPE(FlipOrComposite);
 
-        TempImage.CopyTo(TempImage2, ERawImageFormat::RGBA32F, EGammaSpace::Linear);
+		TempImage0->CopyTo(*TempImage1, ERawImageFormat::RGBA32F, EGammaSpace::Linear);
 		RawFormat = ERawImageFormat::RGBA32F;
 
-        if (bFlipGreenChannel)
-        {
-			FlipGreenChannelRGBA32F(TempImage2);
-			// Don't flip again below.
-			bFlipGreenChannel = false;
-		}
+		for (int32 MipLevel = 0; MipLevel < TempImage1->GetMipCount(); ++MipLevel)
+		{
+			FImageView MipView = TempImage1->GetMipImage(MipLevel);
 
-		// Prepare texture for use as normal composite.
-		Normalize(TempImage2);
-		BlurNormalForComposite(TempImage2);
+			if (bFlipGreenChannel)
+			{
+				FlipGreenChannelRGBA32F(MipView);
+			}
 
-        // The result is needed to TempImage
-        // Swap internals so the memory allocations is potentially reused.
-		TempImage2.Swap(TempImage);
-    }
+			// Prepare texture for use as normal composite.
+			Normalize(MipView);
+			BlurNormalForComposite(MipView);
+		}
+		
+		// Don't flip again below.
+		bFlipGreenChannel = false;
+
+		// TempImage0 is always used as the source of the operations, swap to be consistent.
+		::Swap(TempImage0, TempImage1);
+	}
    
-    const ERawImageFormat::Type MutableCompatibleFormat = Format == TSF_G8 
-            ? ERawImageFormat::G8 
-            : ERawImageFormat::BGRA8;
+	const ERawImageFormat::Type MutableCompatibleFormat = Format == TSF_G8 
+			? ERawImageFormat::G8 
+			: ERawImageFormat::BGRA8;
 
 	if (MutableCompatibleFormat != RawFormat)
 	{
 		MUTABLE_CPUPROFILER_SCOPE(ToCompatibleFormat);
-		TempImage.CopyTo(TempImage2, MutableCompatibleFormat, EGammaSpace::Linear);
-		TempImage2.Swap(TempImage);
+		TempImage0->CopyTo(*TempImage1, MutableCompatibleFormat, EGammaSpace::Linear);
+		::Swap(TempImage0, TempImage1);
 	}
 	
-	if (bFlipGreenChannel)
+	for (int32 MipLevel = 0; MipLevel < TempImage0->GetMipCount(); ++MipLevel)
 	{
-		FlipGreenChannelBGRA8(TempImage);
+		FImageView MipView = TempImage0->GetMipImage(MipLevel);
+
+		if (bFlipGreenChannel)
+		{
+			FlipGreenChannelBGRA8(MipView);
+		}
 	}
 
 	switch (MutableCompatibleFormat)
@@ -291,58 +323,73 @@ EUnrealToMutableConversionError ConvertTextureUnrealSourceToMutable(UE::Mutable:
 	case ERawImageFormat::G8:
 	{
 		MUTABLE_CPUPROFILER_SCOPE(NoConvert);
-        check(LODs == 1);
 
-		OutResult->Init(SizeX, SizeY, LODs, UE::Mutable::Private::EImageFormat::L_UByte, UE::Mutable::Private::EInitializationType::NotInitialized);
-		OutResult->DataStorage.GetInternalArray(0) = MoveTemp(TempImage.RawData);
+		OutResult->Init(SizeX, SizeY, ImageNumMips, UE::Mutable::Private::EImageFormat::L_UByte, UE::Mutable::Private::EInitializationType::NotInitialized);
+	
+		for (int32 MipLevel = 0; MipLevel < ImageNumMips; ++MipLevel)
+		{
+			TArrayView<uint8> MutableImageDataView = OutResult->DataStorage.GetLOD(MipLevel);
+			FImageView ImageCoreDataView  = TempImage0->GetMipImage(MipLevel);
+
+			check(MutableImageDataView.Num() == ImageCoreDataView.GetImageSizeBytes());
+			FMemory::Memcpy(MutableImageDataView.GetData(), ImageCoreDataView.RawData, MutableImageDataView.Num());
+		}
 		break;
 	}
 
 	case ERawImageFormat::BGRA8:
 	{
 		// Try to find out if the texture has and actually makes use of the alpha channel
-		bool bHasAlphaChannel = Tex.HasAlphaChannel()
-			&& (Tex.GetCompressionForceAlpha()
+		bool bHasAlphaChannel = Texture.HasAlphaChannel()
+			&& (Texture.GetCompressionForceAlpha()
 				||
-				FImageCore::DetectAlphaChannel(TempImage));
+				FImageCore::DetectAlphaChannel(TempImage0->GetMipImage(0)));
 
 		// TODO: If we ever manage to get Pixel Format data on cook compilation time, remove the code that sets bHasAlphaChannel and just use Texture->HasAlphaChannel() here. Currently unreliable, it always returns EPixelFormat::PF_Unknown when cooking, which returns always "false" to HasAlphaChannel().
 		if (bHasAlphaChannel)
 		{
 			MUTABLE_CPUPROFILER_SCOPE(ToRGBA);
-			OutResult->Init(SizeX, SizeY, LODs, UE::Mutable::Private::EImageFormat::RGBA_UByte, UE::Mutable::Private::EInitializationType::NotInitialized);
-            check(LODs == 1);
-            uint8* DataDest = OutResult->GetLODData(0);
-
-			// Convert to RGBA8 while copying
-			TArrayView64<FColor> ImageDataView = TempImage.AsBGRA8();
-			ParallelFor(TEXT("MutableToRGBA"), ImageDataView.Num(), 16*1024,
-				[DataDest, &ImageDataView](uint32 p)
+			OutResult->Init(SizeX, SizeY, ImageNumMips, UE::Mutable::Private::EImageFormat::RGBA_UByte, UE::Mutable::Private::EInitializationType::NotInitialized);
+		
+			for (int32 MipLevel = 0; MipLevel < ImageNumMips; ++MipLevel)
+			{
+				FImageView MipImageView = TempImage0->GetMipImage(MipLevel);
+				uint8* DataDest = OutResult->GetLODData(MipLevel);
+
+				// Convert to RGBA8 while copying
+				TArrayView64<FColor> ImageDataView = MipImageView.AsBGRA8();
+				ParallelFor(TEXT("MutableToRGBA"), ImageDataView.Num(), 16*1024,
+				[DataDest, &ImageDataView](uint32 P)
 				{
-					DataDest[4 * p + 0] = ImageDataView[p].R;
-					DataDest[4 * p + 1] = ImageDataView[p].G;
-					DataDest[4 * p + 2] = ImageDataView[p].B;
-					DataDest[4 * p + 3] = ImageDataView[p].A;
+					DataDest[4 * P + 0] = ImageDataView[P].R;
+					DataDest[4 * P + 1] = ImageDataView[P].G;
+					DataDest[4 * P + 2] = ImageDataView[P].B;
+					DataDest[4 * P + 3] = ImageDataView[P].A;
 				});
+			}   
 		}
 		else
 		{
 			MUTABLE_CPUPROFILER_SCOPE(ToRGB);
-
-			// TODO: add support for a UE::Mutable::Private::IF_RGBX_UBYTE?			
-			OutResult->Init(SizeX, SizeY, LODs, UE::Mutable::Private::EImageFormat::RGB_UByte, UE::Mutable::Private::EInitializationType::NotInitialized);
-			check(LODs == 1);
-            uint8* DataDest = OutResult->GetLODData(0);
-
-			// Convert to RGB8 while copying
-			TArrayView64<FColor> ImageDataView = TempImage.AsBGRA8();
-			ParallelFor(TEXT("MutableToRGB"), ImageDataView.Num(), 16 * 1024,
-				[DataDest, &ImageDataView](uint32 p)
+			// TODO: add support for a UE::Mutable::Private::IF_RGBX_UBYTE?
+			OutResult->Init(SizeX, SizeY, ImageNumMips, UE::Mutable::Private::EImageFormat::RGB_UByte, UE::Mutable::Private::EInitializationType::NotInitialized);
+		
+			for (int32 MipLevel = 0; MipLevel < ImageNumMips; ++MipLevel)
+			{
+				FImageView MipImageView = TempImage0->GetMipImage(MipLevel);
+				uint8* DataDest = OutResult->GetLODData(MipLevel);
+
+				// Convert to RGBA8 while copying
+				TArrayView64<FColor> ImageDataView = MipImageView.AsBGRA8();
+				ParallelFor(TEXT("MutableToRGB"), ImageDataView.Num(), 16*1024,
+				[DataDest, &ImageDataView](uint32 P)
 				{
-					DataDest[3 * p + 0] = ImageDataView[p].R;
-					DataDest[3 * p + 1] = ImageDataView[p].G;
-					DataDest[3 * p + 2] = ImageDataView[p].B;
+					DataDest[3 * P + 0] = ImageDataView[P].R;
+					DataDest[3 * P + 1] = ImageDataView[P].G;
+					DataDest[3 * P + 2] = ImageDataView[P].B;
 				});
+			}
+
 		}
 
 		//FString Msg = FString::Printf(TEXT("Alpha channel is %s for %s"), bHasAlphaChannel ? TEXT("enabled") : TEXT("disabled"), *Texture->GetName());
@@ -357,7 +404,7 @@ EUnrealToMutableConversionError ConvertTextureUnrealSourceToMutable(UE::Mutable:
 		break;
 	}
 
-    return EUnrealToMutableConversionError::Success;
+	return EUnrealToMutableConversionError::Success;
 }
 
 uint32 GetTypeHash(const FMutableSourceSurfaceMetadata& Key)
diff --git a/Engine/Plugins/Mutable/Source/CustomizableObjectEditor/Private/MuCOE/CustomizableObjectCompileRunnable.cpp b/Engine/Plugins/Mutable/Source/CustomizableObjectEditor/Private/MuCOE/CustomizableObjectCompileRunnable.cpp
--- a/Engine/Plugins/Mutable/Source/CustomizableObjectEditor/Private/MuCOE/CustomizableObjectCompileRunnable.cpp
+++ b/Engine/Plugins/Mutable/Source/CustomizableObjectEditor/Private/MuCOE/CustomizableObjectCompileRunnable.cpp
@@ -67,8 +67,9 @@ UE::Mutable::Private::TManagedPtr<UE::Mutable::Private::FImage> FCustomizableObj
 
 	// In the editor the src data can be directly accessed
 	Image = UE::Mutable::Private::MakeManaged<UE::Mutable::Private::FImage>();
-	int32 MipmapsToSkip = 0;
-	EUnrealToMutableConversionError Error = ConvertTextureUnrealSourceToMutable(Image.Get(), TextureData, MipmapsToSkip);
+	constexpr int32 MipmapsToSkip = 0;
+	constexpr bool bLoadMipTail = false;
+	EUnrealToMutableConversionError Error = ConvertTextureUnrealSourceToMutable(Image.Get(), TextureData, MipmapsToSkip, bLoadMipTail);
 
 	if (Error != EUnrealToMutableConversionError::Success)
 	{
diff --git a/Engine/Plugins/Mutable/Source/CustomizableObjectEditor/Private/MuCOE/GenerateMutableSource/GenerateMutableSourceLayout.cpp b/Engine/Plugins/Mutable/Source/CustomizableObjectEditor/Private/MuCOE/GenerateMutableSource/GenerateMutableSourceLayout.cpp
--- a/Engine/Plugins/Mutable/Source/CustomizableObjectEditor/Private/MuCOE/GenerateMutableSource/GenerateMutableSourceLayout.cpp
+++ b/Engine/Plugins/Mutable/Source/CustomizableObjectEditor/Private/MuCOE/GenerateMutableSource/GenerateMutableSourceLayout.cpp
@@ -84,8 +84,11 @@ UE::Mutable::Private::FSourceLayoutBlock ToMutable(const FCustomizableObjectLayo
 		// In the editor the src data can be directly accessed
 		UE::Mutable::Private::TManagedPtr<UE::Mutable::Private::FImage> MaskImage = UE::Mutable::Private::MakeManaged<UE::Mutable::Private::FImage>();
 
-		FMutableSourceTextureData Tex(*UnrealBlock.Mask);
-		EUnrealToMutableConversionError Error = ConvertTextureUnrealSourceToMutable(MaskImage.Get(), Tex, 0);
+		FMutableSourceTextureData Texture(*UnrealBlock.Mask);
+
+		constexpr int32 MipLevel = 0;
+		constexpr int32 bLoadMipTail = false;
+		EUnrealToMutableConversionError Error = ConvertTextureUnrealSourceToMutable(MaskImage.Get(), Texture, MipLevel, false);
 		if (Error != EUnrealToMutableConversionError::Success)
 		{
 			// This should never happen, so details are not necessary.
diff --git a/Engine/Plugins/Mutable/Source/MutableRuntime/Internal/MuR/System.h b/Engine/Plugins/Mutable/Source/MutableRuntime/Internal/MuR/System.h
--- a/Engine/Plugins/Mutable/Source/MutableRuntime/Internal/MuR/System.h
+++ b/Engine/Plugins/Mutable/Source/MutableRuntime/Internal/MuR/System.h
@@ -97,7 +97,7 @@ namespace UE::Mutable::Private
         virtual ~FExternalResourceProvider() = default;
 
         /** Returns the completion event and a cleanup function that must be called once event is completed. */
-		virtual TTuple<UE::Tasks::FTask, TFunction<void()>> GetImageAsync(UTexture* Texture, uint8 MipmapsToSkip, TFunction<void(UE::Mutable::Private::TManagedPtr<FImage>)>& ResultCallback) = 0;
+		virtual TTuple<UE::Tasks::FTask, TFunction<void()>> GetImageAsync(UTexture* Texture, uint8 MipmapsToSkip, bool bLoadMipTail, TFunction<void(UE::Mutable::Private::TManagedPtr<FImage>)>& ResultCallback) = 0;
 		virtual TTuple<UE::Tasks::FTask, TFunction<void()>> GetReferencedImageAsync(int32 Id, uint8 MipmapsToSkip, TFunction<void(UE::Mutable::Private::TManagedPtr<FImage>)>& ResultCallback) { check(false); return {}; }
 
 		virtual FExtendedImageDesc GetImageDesc(UTexture* Texture) = 0;
diff --git a/Engine/Plugins/Mutable/Source/MutableRuntime/Private/MuR/CodeRunner.cpp b/Engine/Plugins/Mutable/Source/MutableRuntime/Private/MuR/CodeRunner.cpp
--- a/Engine/Plugins/Mutable/Source/MutableRuntime/Private/MuR/CodeRunner.cpp
+++ b/Engine/Plugins/Mutable/Source/MutableRuntime/Private/MuR/CodeRunner.cpp
@@ -230,7 +230,8 @@ namespace UE::Mutable::Private
 			if (Id.ReferenceResourceId == PASSTHROUGH_ID_INVALID)
 			{
 				// It's a parameter image
-				return ExternalResourceProvider->GetImageAsync(Id.ImageParameter, MipmapsToSkip, ResultCallback);
+				constexpr bool bLoadMipTail = true;
+				return ExternalResourceProvider->GetImageAsync(Id.ImageParameter, MipmapsToSkip, bLoadMipTail, ResultCallback);
 			}
 			else
 			{
