[mutable] Fixed crash during the generation of a MutableMesh from a UE Static Mesh.
- The lambda used to identify the LOD and section of the input pin was wrongfully discarding the matching pin.

#rnx
#Synced-CL 49757097

diff --git a/Engine/Plugins/Mutable/Source/CustomizableObjectEditor/Private/MuCOE/GenerateMutableSource/GenerateMutableSourceMesh.cpp b/Engine/Plugins/Mutable/Source/CustomizableObjectEditor/Private/MuCOE/GenerateMutableSource/GenerateMutableSourceMesh.cpp
--- a/Engine/Plugins/Mutable/Source/CustomizableObjectEditor/Private/MuCOE/GenerateMutableSource/GenerateMutableSourceMesh.cpp
+++ b/Engine/Plugins/Mutable/Source/CustomizableObjectEditor/Private/MuCOE/GenerateMutableSource/GenerateMutableSourceMesh.cpp
@@ -3954,7 +3954,7 @@ UE::Mutable::Private::Ptr<UE::Mutable::Private::NodeMesh> GenerateMutableSourceM
 		{
 			for (; LODIndex < TypedNodeStatic->LODs.Num(); ++LODIndex)
 			{
-				for (; SectionIndex < TypedNodeStatic->LODs[LODIndex].Materials.Num(); ++SectionIndex)
+				for (SectionIndex = 0; SectionIndex < TypedNodeStatic->LODs[LODIndex].Materials.Num(); ++SectionIndex)
 				{
 					if (TypedNodeStatic->LODs[LODIndex].Materials[SectionIndex].MeshPinRef.Get() == Pin)
 					{
