[Mutable] Fix check with MI_SKELETALMESH_BREAK.

#Synced-CL 49759118

diff --git a/Engine/Plugins/Mutable/Source/MutableRuntime/Private/MuR/CodeRunnerBegin.cpp b/Engine/Plugins/Mutable/Source/MutableRuntime/Private/MuR/CodeRunnerBegin.cpp
--- a/Engine/Plugins/Mutable/Source/MutableRuntime/Private/MuR/CodeRunnerBegin.cpp
+++ b/Engine/Plugins/Mutable/Source/MutableRuntime/Private/MuR/CodeRunnerBegin.cpp
@@ -2937,7 +2937,6 @@ namespace UE::Mutable::Private
 			case EOpType::MI_SKELETALMESH_BREAK:
 				{
 					const OP::MaterialSkeletalMeshBreakArgs Args = Program.GetOpArgs<OP::MaterialSkeletalMeshBreakArgs>(Item.At);
-					PushOp(FScheduledOpInline(Args.SlotNameStringIndex, Item));
 					PushOp(FScheduledOpInline(Args.SkeletalMesh, Item));
 					
 					StoreMaterial(Item, MakeManaged<const FMaterial>());
