[Mutable] Included material parameters hash into the instance descriptor.

#rnx

#Synced-CL 49759118

diff --git a/Engine/Plugins/Mutable/Source/CustomizableObject/Private/MuCO/DescriptorHash.cpp b/Engine/Plugins/Mutable/Source/CustomizableObject/Private/MuCO/DescriptorHash.cpp
--- a/Engine/Plugins/Mutable/Source/CustomizableObject/Private/MuCO/DescriptorHash.cpp
+++ b/Engine/Plugins/Mutable/Source/CustomizableObject/Private/MuCO/DescriptorHash.cpp
@@ -59,6 +59,11 @@ FDescriptorHash::FDescriptorHash(const FCustomizableObjectInstanceDescriptor& De
 		Hash = HashCombine(Hash, GetTypeHash(Value));
 	}
 
+	for (const FCustomizableObjectMaterialParameterValue& Value : Descriptor.MaterialParameters)
+	{
+		Hash = HashCombine(Hash, GetTypeHash(Value));
+	}
+
 	Hash = HashCombine(Hash, GetTypeHash(Descriptor.State));
 	Hash = HashCombine(Hash, GetTypeHash(Descriptor.GetBuildParameterRelevancy()));
 	
