[mutable] Decreased dependency of the CodeRunner with the FSystem object
- Removed some "shared" objects from the FSystem as now each FLiveInstance object will hold them (like the registries and the ProgramCache). This is required to allow the execution of multiple CodeRunner instances at the same time.
- Left some dependencies to the FSystem object as those are required. The methods affected do require locks (or similar thread control methods) to allow safe access and therefore the execution of multiple CodeRunner instances.
- Fixed some typos and purged some no longer required includes.

#rnx

#Synced-CL 49801225

diff --git a/Engine/Plugins/Mutable/Source/CustomizableObject/Internal/MuCO/CustomizableObjectSystemPrivate.h b/Engine/Plugins/Mutable/Source/CustomizableObject/Internal/MuCO/CustomizableObjectSystemPrivate.h
--- a/Engine/Plugins/Mutable/Source/CustomizableObject/Internal/MuCO/CustomizableObjectSystemPrivate.h
+++ b/Engine/Plugins/Mutable/Source/CustomizableObject/Internal/MuCO/CustomizableObjectSystemPrivate.h
@@ -562,10 +562,6 @@ public:
 	TSharedPtr<UE::Mutable::Private::FParameters> Parameters; 
 	TSharedPtr<UE::Mutable::Private::FSystem> MutableSystem;
 	TSharedPtr<UE::Mutable::Private::FModel> Model;
-	TSharedPtr<UE::Mutable::Private::FMeshIdRegistry> MeshIdRegistry;
-	TSharedPtr<UE::Mutable::Private::FImageIdRegistry> ImageIdRegistry;
-	TSharedPtr<UE::Mutable::Private::FMaterialIdRegistry> MaterialIdRegistry;
-	TSharedPtr<UE::Mutable::Private::FSkeletalMeshIdRegistry> SkeletalMeshIdRegistry;
 	
 	TSharedPtr<UE::Mutable::Private::FLiveInstance> LiveInstance;
 
@@ -802,7 +798,7 @@ public:
 	UE_API int32 GetRemainingWork() const;
 	
 	UCustomizableObjectSystem* GetPublic() const;
-
+	
 	EUpdateRequired IsUpdateRequired(const UCustomizableObjectInstance& Instance, bool bOnlyUpdateIfNotGenerated, bool bOnlyUpdateIfLOD, bool bIgnoreCloseDist) const;
 
 	EQueuePriorityType GetUpdatePriority(const UCustomizableObjectInstance& Instance, bool bForceHighPriority) const;
@@ -817,7 +813,7 @@ public:
 
 	/** Updated at the beginning of each tick. */
 	int32 GetNumSkeletalMeshes() const;
-	
+
 	bool bReplaceDiscardedWithReferenceMesh = false;
 	
 	bool bSupport16BitBoneIndex = false;
@@ -856,7 +852,7 @@ public:
 	int32 NumSkeletalMeshes = 0;
 
 	bool bAutoCompileCommandletEnabled = false;
-	
+
 	bool bPendingLODUpdate = false;
 
 	UPROPERTY()
diff --git a/Engine/Plugins/Mutable/Source/CustomizableObject/Private/MuCO/CustomizableObjectMeshUpdate.cpp b/Engine/Plugins/Mutable/Source/CustomizableObject/Private/MuCO/CustomizableObjectMeshUpdate.cpp
--- a/Engine/Plugins/Mutable/Source/CustomizableObject/Private/MuCO/CustomizableObjectMeshUpdate.cpp
+++ b/Engine/Plugins/Mutable/Source/CustomizableObject/Private/MuCO/CustomizableObjectMeshUpdate.cpp
@@ -245,7 +245,14 @@ namespace impl
 #endif
 
 		// For now, we are forcing the recreation of mutable-side instances with every update.
-		TSharedRef<UE::Mutable::Private::FLiveInstance> LiveInstance = System->NewInstance(Model, OperationData->ExternalResourceProvider, nullptr);
+		TSharedRef<UE::Mutable::Private::FLiveInstance> LiveInstance = System->NewInstance(
+			Model, 
+			OperationData->MeshIdRegistry,
+			OperationData->ImageIdRegistry,
+			OperationData->MaterialIdRegistry,
+			OperationData->SkeletalMeshIdRegistry,
+			OperationData->ExternalResourceProvider, 
+			nullptr);
 
 		// LOD mask, set to all ones to build all LODs
 		const uint32 LODMask = 0xFFFFFFFF;
@@ -253,10 +260,6 @@ namespace impl
 		// Main instance generation step
 		UE::Mutable::Private::TManagedPtr<const UE::Mutable::Private::FInstance> Instance = System->BeginUpdate_MutableThread(LiveInstance,
 			OperationData->Parameters,
-			OperationData->MeshIdRegistry,
-			OperationData->ImageIdRegistry,
-			OperationData->MaterialIdRegistry,
-			OperationData->SkeletalMeshIdRegistry,
 			OperationData->State,
 			LODMask);
 		
diff --git a/Engine/Plugins/Mutable/Source/CustomizableObject/Private/MuCO/CustomizableObjectMipDataProvider.cpp b/Engine/Plugins/Mutable/Source/CustomizableObject/Private/MuCO/CustomizableObjectMipDataProvider.cpp
--- a/Engine/Plugins/Mutable/Source/CustomizableObject/Private/MuCO/CustomizableObjectMipDataProvider.cpp
+++ b/Engine/Plugins/Mutable/Source/CustomizableObject/Private/MuCO/CustomizableObjectMipDataProvider.cpp
@@ -106,7 +106,7 @@ namespace Impl
 
 		auto EndUpdateImage = [](TSharedPtr<FMutableImageOperationData>& OperationData)
 		{
-			// The request could be cancelled in parallel from CancelAsyncTasks and its value be changed
+			// The request could be canceled in parallel from CancelAsyncTasks and its value be changed
 			// between reading it and actually running Decrement() and RescheduleCallback(), so lock
 			FScopeLock Lock(&OperationData->CounterTaskLock);
 			if (!OperationData->bIsCancelled)
@@ -131,8 +131,15 @@ namespace Impl
 #endif
 
 		// For now, we are forcing the recreation of mutable-side instances with every update.
-		TSharedRef<UE::Mutable::Private::FLiveInstance> LiveInstance = System->NewInstance(Model, OperationData->UpdateContext->ExternalResourceProvider, nullptr);
-
+		TSharedRef<UE::Mutable::Private::FLiveInstance> LiveInstance = System->NewInstance(
+			Model,
+			OperationData->UpdateContext->MeshIdRegistry,
+			OperationData->UpdateContext->ImageIdRegistry,
+			OperationData->UpdateContext->MaterialIdRegistry,
+			OperationData->UpdateContext->SkeletalMeshIdRegistry,
+			OperationData->UpdateContext->ExternalResourceProvider, 
+			nullptr);
+		
 		UE::Mutable::Private::TManagedPtr<const UE::Mutable::Private::FInstance> Instance;
 
 		// Main instance generation step
@@ -142,10 +149,6 @@ namespace Impl
 
 			Instance = System->BeginUpdate_MutableThread(LiveInstance,
 				OperationData->UpdateContext->Parameters,
-				OperationData->UpdateContext->MeshIdRegistry,
-				OperationData->UpdateContext->ImageIdRegistry,
-				OperationData->UpdateContext->MaterialIdRegistry,
-				OperationData->UpdateContext->SkeletalMeshIdRegistry,
 				OperationData->UpdateContext->State,
 				LODMask);
 			
diff --git a/Engine/Plugins/Mutable/Source/CustomizableObject/Private/MuCO/CustomizableObjectSystem.cpp b/Engine/Plugins/Mutable/Source/CustomizableObject/Private/MuCO/CustomizableObjectSystem.cpp
--- a/Engine/Plugins/Mutable/Source/CustomizableObject/Private/MuCO/CustomizableObjectSystem.cpp
+++ b/Engine/Plugins/Mutable/Source/CustomizableObject/Private/MuCO/CustomizableObjectSystem.cpp
@@ -2234,7 +2234,7 @@ namespace impl
 					TSharedRef<FInstanceUpdateData::FImage>& Image = Context->InstanceUpdateData.Images.Add_GetRef(MakeShared<FInstanceUpdateData::FImage>());
 					Image->ParameterName = ImageParameter.Key.ParameterName;
 					Image->MaterialLayer = ImageParameter.Key.LayerIndex;
-					Image->ImageID = UE::Mutable::Private::FSystem::GetImageId(Context->ImageIdRegistry, Context->Model, Context->Parameters.Get(), Address);
+					Image->ImageID = UE::Mutable::Private::FSystem::GetImageId(Context->LiveInstance->ImageIdRegistry, Context->Model, Context->Parameters.Get(), Address);
 					Image->FullImageSizeX = 0;
 					Image->FullImageSizeY = 0;
 					Image->BaseMip = 0;
@@ -2274,10 +2274,6 @@ namespace impl
 		
 		UE::Mutable::Private::TManagedPtr<const UE::Mutable::Private::FInstance> Instance = OperationData->MutableSystem->BeginUpdate_MutableThread(OperationData->LiveInstance.ToSharedRef(),
 			OperationData->Parameters,
-			OperationData->MeshIdRegistry,
-			OperationData->ImageIdRegistry,
-			OperationData->MaterialIdRegistry,
-			OperationData->SkeletalMeshIdRegistry,
 			OperationData->GetCapturedDescriptor().GetState(),
 			UE::Mutable::Private::FSystem::AllLODs);
 
@@ -3035,10 +3031,6 @@ namespace impl
 		Operation->bUseResueTextureCache = Operation->bUseCaches && Operation->bNeverStream && CVarEnableMutableReuseInstanceTextures.GetValueOnGameThread();
 		
 		Operation->Model = ObjectPrivate->GetModel().ToSharedRef();
-		Operation->MeshIdRegistry = ObjectPrivate->MeshIdRegistry;
-		Operation->ImageIdRegistry = ObjectPrivate->ImageIdRegistry;
-		Operation->MaterialIdRegistry = ObjectPrivate->MaterialIdRegistry;
-		Operation->SkeletalMeshIdRegistry = ObjectPrivate->SkeletalMeshIdRegistry;
 
 		UE::Tasks::FTask CacheRuntimeTexturesEvent = UE::Tasks::MakeCompletedTask<void>();
 
@@ -3174,20 +3166,26 @@ namespace impl
 			
 			if (!Operation->LiveInstance)
 			{
-				Operation->LiveInstance = MutableSystem->NewInstance(Operation->Model, Operation->ExternalResourceProvider, Operation->PassthroughObjectLoader);
+				Operation->LiveInstance = MutableSystem->NewInstance(
+					Operation->Model, 
+					ObjectPrivate->MeshIdRegistry,
+					ObjectPrivate->ImageIdRegistry,
+					ObjectPrivate->MaterialIdRegistry,
+					ObjectPrivate->SkeletalMeshIdRegistry,
+					Operation->ExternalResourceProvider, 
+					Operation->PassthroughObjectLoader);
+			}
+			else
+			{
+				Operation->LiveInstance->ExternalResourceProvider = Operation->ExternalResourceProvider;
+				Operation->LiveInstance->PassthroughObjectLoader = Operation->PassthroughObjectLoader;
 			}
 
 			Instance->GetPrivate()->LiveInstance = Operation->bLiveUpdateMode ? Operation->LiveInstance : nullptr;
 			
-			Operation->LiveInstance->ExternalResourceProvider = Operation->ExternalResourceProvider;
-			Operation->LiveInstance->PassthroughObjectLoader = Operation->PassthroughObjectLoader;
-			
-			Operation->MutableInstance = MutableSystem->BeginUpdate_GameThread(Operation->LiveInstance.ToSharedRef(),
+			Operation->MutableInstance = MutableSystem->BeginUpdate_GameThread(
+				Operation->LiveInstance.ToSharedRef(),
 				Operation->Parameters,
-				Operation->MeshIdRegistry,
-				Operation->ImageIdRegistry,
-				Operation->MaterialIdRegistry,
-				Operation->SkeletalMeshIdRegistry,
 				Operation->GetCapturedDescriptor().GetState(),
 				UE::Mutable::Private::FSystem::AllLODs);
 		}
@@ -3343,19 +3341,19 @@ namespace impl
 void UCustomizableObjectSystemPrivate::UpdateTick(const bool bBlocking)
 {
 	NumSkeletalMeshes = 0;
-	
+
 	for (TObjectIterator<UCustomizableObjectInstance> Instance; Instance; ++Instance)
-	{
+{
 		if (!IsValid(*Instance))
-		{
+	{
 			continue;
 		}
 
 		NumSkeletalMeshes += Instance->GetPrivate()->SkeletalMeshes.Num();
-	}
+				}
 
 	FMutableUpdateCandidate* LODUpdateCandidateFound = nullptr;
-	
+
 	for (TObjectIterator<UCustomizableObjectInstanceUsage> CustomizableObjectInstanceUsageItr; CustomizableObjectInstanceUsageItr; ++CustomizableObjectInstanceUsageItr)
 	{
 		TObjectPtr<UCustomizableObjectInstanceUsage> CustomizableObjectInstanceUsage = *CustomizableObjectInstanceUsageItr;
@@ -3586,7 +3584,7 @@ void UCustomizableObjectSystemPrivate::UpdateTick(const bool bBlocking)
 		GetPublic()->DiscardInstances();
 	}
 	
-	LogBenchmarkUtil.UpdateStats(); // Must to be the last thing to perform
+	LogBenchmarkUtil.UpdateStats(); // Must be the last thing to perform
 
 	if (!bIsMutableEnabled && !CurrentMutableOperation)
 	{
@@ -3651,14 +3649,14 @@ int32 UCustomizableObjectSystemPrivate::GetRemainingWork() const
 	{
 		return 0;
 	}
-	
+
 	if (IsEngineExitRequested())
-	{	
+	{
 		return 0;
 	}
-	
+
 	int32 RemainingWork = 0;
-	
+
 	// Do not tick if the CookCommandlet is running.
 	if (!IsRunningCookCommandlet())
 	{
@@ -3666,7 +3664,7 @@ int32 UCustomizableObjectSystemPrivate::GetRemainingWork() const
 			   MutablePendingInstanceWork.Num() +
 			   bPendingLODUpdate;
 	}
-
+	
 	RemainingWork = MutableTaskGraph.GetRemainingWork();
 
 	if (ICustomizableObjectEditorModule* Module = ICustomizableObjectEditorModule::Get())
diff --git a/Engine/Plugins/Mutable/Source/CustomizableObjectEditor/Private/MuCOE/SMutableCodeViewer.cpp b/Engine/Plugins/Mutable/Source/CustomizableObjectEditor/Private/MuCOE/SMutableCodeViewer.cpp
--- a/Engine/Plugins/Mutable/Source/CustomizableObjectEditor/Private/MuCOE/SMutableCodeViewer.cpp
+++ b/Engine/Plugins/Mutable/Source/CustomizableObjectEditor/Private/MuCOE/SMutableCodeViewer.cpp
@@ -2897,14 +2897,15 @@ void SMutableCodeViewer::Tick(const FGeometry& AllottedGeometry, const double In
 	UE::Mutable::Private::FSettings Settings;
 	const TSharedPtr<UE::Mutable::Private::FSystem> System = MakeShared<UE::Mutable::Private::FSystem>(Settings);
 
-	System->BeginBuild(MutableModel, Object->GetPrivate()->MeshIdRegistry, Object->GetPrivate()->ImageIdRegistry, Object->GetPrivate()->MaterialIdRegistry, Object->GetPrivate()->SkeletalMeshIdRegistry);
+	System->BeginBuild(MutableModel);
+	TSharedRef<UE::Mutable::Private::FLiveInstance> TempLiveInstance = System->NewBuildInstance(MutableModel, ExternalResourceProvider);
 	
 	switch (OperationDataType)
 	{
 	case UE::Mutable::Private::EDataType::Layout:
 	{
 		check(PreviewLayoutViewer);
-		UE::Mutable::Private::TManagedPtr<const UE::Mutable::Private::FLayout> MutableLayout = System->BuildLayout(ExternalResourceProvider, MutableModel, PreviewParameters.Get(), SelectedOperationAddress);
+		UE::Mutable::Private::TManagedPtr<const UE::Mutable::Private::FLayout> MutableLayout = System->BuildLayout(TempLiveInstance, PreviewParameters.Get(), SelectedOperationAddress);
 		PreviewLayoutViewer->SetLayout(MutableLayout);
 		break;
 	}
@@ -2912,7 +2913,7 @@ void SMutableCodeViewer::Tick(const FGeometry& AllottedGeometry, const double In
 	case UE::Mutable::Private::EDataType::Image:
 	{
 		check(PreviewImageViewer);
-		UE::Mutable::Private::TManagedPtr<const UE::Mutable::Private::FImage> MutableImage = System->BuildImage(ExternalResourceProvider, MutableModel, PreviewParameters.Get(), SelectedOperationAddress, MipsToSkip);
+		UE::Mutable::Private::TManagedPtr<const UE::Mutable::Private::FImage> MutableImage = System->BuildImage(TempLiveInstance, PreviewParameters.Get(), SelectedOperationAddress, MipsToSkip);
 		PreviewImageViewer->SetImage(MutableImage, 0);
 		break;
 	}
@@ -2920,7 +2921,7 @@ void SMutableCodeViewer::Tick(const FGeometry& AllottedGeometry, const double In
 	case UE::Mutable::Private::EDataType::Mesh:
 	{
 		check(PreviewMeshViewer);
-		UE::Mutable::Private::TManagedPtr<const UE::Mutable::Private::FMesh> MutableMesh = System->BuildMesh(ExternalResourceProvider, MutableModel, PreviewParameters.Get(), SelectedOperationAddress, UE::Mutable::Private::EMeshContentFlags::AllFlags);
+		UE::Mutable::Private::TManagedPtr<const UE::Mutable::Private::FMesh> MutableMesh = System->BuildMesh(TempLiveInstance, PreviewParameters.Get(), SelectedOperationAddress, UE::Mutable::Private::EMeshContentFlags::AllFlags);
 		PreviewMeshViewer->SetMesh(MutableMesh);
 		break;
 	}
@@ -2928,15 +2929,15 @@ void SMutableCodeViewer::Tick(const FGeometry& AllottedGeometry, const double In
 	case UE::Mutable::Private::EDataType::Instance:
 	{
 		check(PreviewInstanceViewer);
-		UE::Mutable::Private::TManagedPtr<const UE::Mutable::Private::FInstance> MutableInstance = System->BuildInstance(ExternalResourceProvider, MutableModel, PreviewParameters.Get(), SelectedOperationAddress);
-		PreviewInstanceViewer->SetInstance(MutableInstance, MutableModel, PreviewParameters, *System);
+		UE::Mutable::Private::TManagedPtr<const UE::Mutable::Private::FInstance> MutableInstance = System->BuildInstance(TempLiveInstance, PreviewParameters.Get(), SelectedOperationAddress);
+		PreviewInstanceViewer->SetInstance(TempLiveInstance, MutableInstance, PreviewParameters, *System);
 		break;
 	}
 
 	case UE::Mutable::Private::EDataType::Bool:
 	{
 		check(PreviewBoolViewer);
-		const bool MutableBool = System->BuildBool(ExternalResourceProvider, MutableModel, PreviewParameters.Get(), SelectedOperationAddress);
+		const bool MutableBool = System->BuildBool(TempLiveInstance, PreviewParameters.Get(), SelectedOperationAddress);
 		PreviewBoolViewer->SetBool(MutableBool);
 		break;
 	}
@@ -2944,7 +2945,7 @@ void SMutableCodeViewer::Tick(const FGeometry& AllottedGeometry, const double In
 	case UE::Mutable::Private::EDataType::Int:
 	{
 		check(PreviewIntViewer);
-		const int32 MutableInt = System->BuildInt(ExternalResourceProvider, MutableModel, PreviewParameters.Get(), SelectedOperationAddress);
+		const int32 MutableInt = System->BuildInt(TempLiveInstance, PreviewParameters.Get(), SelectedOperationAddress);
 		PreviewIntViewer->SetInt(MutableInt);
 		break;
 	}
@@ -2952,7 +2953,7 @@ void SMutableCodeViewer::Tick(const FGeometry& AllottedGeometry, const double In
 	case UE::Mutable::Private::EDataType::Scalar:
 	{
 		check(PreviewScalarViewer);
-		const float MutableScalar = System->BuildScalar(ExternalResourceProvider, MutableModel, PreviewParameters.Get(), SelectedOperationAddress);
+		const float MutableScalar = System->BuildScalar(TempLiveInstance, PreviewParameters.Get(), SelectedOperationAddress);
 		PreviewScalarViewer->SetScalar(MutableScalar);
 		break;
 	}
@@ -2960,7 +2961,7 @@ void SMutableCodeViewer::Tick(const FGeometry& AllottedGeometry, const double In
 	case UE::Mutable::Private::EDataType::String:
 	{
 		check(PreviewStringViewer);
-		UE::Mutable::Private::TManagedPtr<const UE::Mutable::Private::String> MutableString = System->BuildString(ExternalResourceProvider, MutableModel, PreviewParameters.Get(), SelectedOperationAddress);
+		UE::Mutable::Private::TManagedPtr<const UE::Mutable::Private::String> MutableString = System->BuildString(TempLiveInstance, PreviewParameters.Get(), SelectedOperationAddress);
 		const FText MutableText = FText::FromString(FString(MutableString->GetValue()));
 		PreviewStringViewer->SetString(MutableText);
 		break;
@@ -2969,7 +2970,7 @@ void SMutableCodeViewer::Tick(const FGeometry& AllottedGeometry, const double In
 	case UE::Mutable::Private::EDataType::Color:
 	{
 		check(PreviewColorViewer);
-		FVector4f Color = System->BuildColour(ExternalResourceProvider, MutableModel, PreviewParameters.Get(), SelectedOperationAddress);
+		FVector4f Color = System->BuildColor(TempLiveInstance, PreviewParameters.Get(), SelectedOperationAddress);
 		PreviewColorViewer->SetColor(Color);
 		break;
 	}
@@ -2977,7 +2978,7 @@ void SMutableCodeViewer::Tick(const FGeometry& AllottedGeometry, const double In
 	case UE::Mutable::Private::EDataType::Projector:
 	{
 		check (PreviewProjectorViewer);
-		UE::Mutable::Private::FProjector Projector = System->BuildProjector(ExternalResourceProvider, MutableModel, PreviewParameters.Get(), SelectedOperationAddress);
+		UE::Mutable::Private::FProjector Projector = System->BuildProjector(TempLiveInstance, PreviewParameters.Get(), SelectedOperationAddress);
 		PreviewProjectorViewer->SetProjector(Projector);
 		break;
 	}
@@ -3007,9 +3008,10 @@ UE::Mutable::Private::TManagedPtr<const UE::Mutable::Private::FImage> SMutableCo
 {
 	UE::Mutable::Private::FSettings Settings;
 	const TSharedPtr<UE::Mutable::Private::FSystem> System = MakeShared<UE::Mutable::Private::FSystem>(Settings);
-
-	System->BeginBuild(MutableModel, Object->GetPrivate()->MeshIdRegistry, Object->GetPrivate()->ImageIdRegistry, Object->GetPrivate()->MaterialIdRegistry, Object->GetPrivate()->SkeletalMeshIdRegistry);
-	UE::Mutable::Private::TManagedPtr<const UE::Mutable::Private::FImage> MutableImage = System->BuildImage(ExternalResourceProvider, MutableModel, PreviewParameters.Get(), TargetOpAddress, 0);
+	
+	System->BeginBuild(MutableModel);
+	TSharedRef<UE::Mutable::Private::FLiveInstance> TempLiveInstance = System->NewBuildInstance(MutableModel, ExternalResourceProvider);
+	UE::Mutable::Private::TManagedPtr<const UE::Mutable::Private::FImage> MutableImage = System->BuildImage(TempLiveInstance, PreviewParameters.Get(), TargetOpAddress, 0);
 	System->EndBuild();
 	
 	return MutableImage;
@@ -3074,7 +3076,7 @@ void SMutableCodeViewer::PrepareInstanceViewer()
 {
 	if (!PreviewInstanceViewer)
 	{
-		PreviewInstanceViewer = SNew(SMutableInstanceViewer, ExternalResourceProvider);
+		PreviewInstanceViewer = SNew(SMutableInstanceViewer);
 	}
 
 	PreviewBorder->SetContent(PreviewInstanceViewer.ToSharedRef());
diff --git a/Engine/Plugins/Mutable/Source/CustomizableObjectEditor/Private/MuCOE/SMutableInstanceViewer.cpp b/Engine/Plugins/Mutable/Source/CustomizableObjectEditor/Private/MuCOE/SMutableInstanceViewer.cpp
--- a/Engine/Plugins/Mutable/Source/CustomizableObjectEditor/Private/MuCOE/SMutableInstanceViewer.cpp
+++ b/Engine/Plugins/Mutable/Source/CustomizableObjectEditor/Private/MuCOE/SMutableInstanceViewer.cpp
@@ -243,14 +243,12 @@ private:
 };
 
 
-void SMutableInstanceViewer::Construct(const FArguments& InArgs, const TSharedPtr<FUnrealMutableResourceProvider>& InExternalResourceProvider)
+void SMutableInstanceViewer::Construct(const FArguments& InArgs)
 {
 	// Splitter values
 	constexpr float TablesSplitterValue = 0.5f;
 	constexpr float ViewportSplitterValue = 0.5f;
-
-	ExternalResourceProvider = InExternalResourceProvider;
-
+	
 	ChildSlot
 	[
 		SNew(SSplitter)
@@ -285,7 +283,7 @@ FString ToString(const UE::Mutable::Private::TPassthroughObjectPtr<UObject>& Pas
 
 
 
-void SMutableInstanceViewer::SetInstance(const UE::Mutable::Private::TManagedPtr<const UE::Mutable::Private::FInstance>& InInstance, const TSharedPtr<const UE::Mutable::Private::FModel>& Model, const TSharedPtr<const UE::Mutable::Private::FParameters>& Parameters, UE::Mutable::Private::FSystem& System)
+void SMutableInstanceViewer::SetInstance(const TSharedRef<UE::Mutable::Private::FLiveInstance> InLiveInstance, const UE::Mutable::Private::TManagedPtr<const UE::Mutable::Private::FInstance>& InInstance, const TSharedPtr<const UE::Mutable::Private::FParameters>& Parameters, UE::Mutable::Private::FSystem& System)
 {
 	if (InInstance != MutableInstance)
 	{
@@ -320,10 +318,11 @@ void SMutableInstanceViewer::SetInstance(const UE::Mutable::Private::TManagedPtr
 					// Generate the mesh to get some data
 					if (MeshId)
 					{
+						check(InLiveInstance->ExternalResourceProvider)
+						
 						UE::Mutable::Private::OP::ADDRESS MeshRootAddress = MeshId.GetKey()->Address;
 						UE::Mutable::Private::TManagedPtr<const UE::Mutable::Private::FMesh> MutableMesh = System.BuildMesh(
-								ExternalResourceProvider,
-								Model, 
+								InLiveInstance,
 								Parameters.Get(), 
 								MeshRootAddress, 
 								UE::Mutable::Private::EMeshContentFlags::AllFlags);
diff --git a/Engine/Plugins/Mutable/Source/CustomizableObjectEditor/Private/MuCOE/SMutableInstanceViewer.h b/Engine/Plugins/Mutable/Source/CustomizableObjectEditor/Private/MuCOE/SMutableInstanceViewer.h
--- a/Engine/Plugins/Mutable/Source/CustomizableObjectEditor/Private/MuCOE/SMutableInstanceViewer.h
+++ b/Engine/Plugins/Mutable/Source/CustomizableObjectEditor/Private/MuCOE/SMutableInstanceViewer.h
@@ -59,17 +59,16 @@ public:
 	SLATE_END_ARGS()
 
 	/** Builds the widget */
-	void Construct(const FArguments& InArgs, const TSharedPtr<FUnrealMutableResourceProvider>& InExternalResourceProvider);
+	void Construct(const FArguments& InArgs);
 
 	/** Set the Mutable Instance to be used for this widget */
-	void SetInstance(const UE::Mutable::Private::TManagedPtr<const UE::Mutable::Private::FInstance>&, const TSharedPtr<const UE::Mutable::Private::FModel>&, const TSharedPtr<const UE::Mutable::Private::FParameters>&, UE::Mutable::Private::FSystem&);
+	void SetInstance(const TSharedRef<UE::Mutable::Private::FLiveInstance> InLiveInstance, const UE::Mutable::Private::TManagedPtr<const UE::Mutable::Private::FInstance>&, const TSharedPtr<const UE::Mutable::Private::FParameters>&, UE::Mutable::Private::FSystem&);
 	
 	/** */
 	TSharedRef<SWidget> GenerateLODsListView(const TSharedPtr<TArray<TSharedPtr<FMutableInstanceViewerLODElement>>>& InBufferChannelElements);
 	TSharedRef<SWidget> GenerateSurfaceListView(const TSharedPtr<TArray<TSharedPtr<FMutableInstanceViewerSurfaceElement>>>& InBufferChannelElements);
 
 private:
-	TSharedPtr<FUnrealMutableResourceProvider> ExternalResourceProvider;
 	
 	/** Data backend for the widget. It represents the Instance that is being "displayed" */
 	UE::Mutable::Private::TManagedPtr<const UE::Mutable::Private::FInstance> MutableInstance = nullptr;
diff --git a/Engine/Plugins/Mutable/Source/MutableRuntime/Internal/MuR/CodeVisitor.h b/Engine/Plugins/Mutable/Source/MutableRuntime/Internal/MuR/CodeVisitor.h
--- a/Engine/Plugins/Mutable/Source/MutableRuntime/Internal/MuR/CodeVisitor.h
+++ b/Engine/Plugins/Mutable/Source/MutableRuntime/Internal/MuR/CodeVisitor.h
@@ -219,15 +219,16 @@ namespace UE::Mutable::Private
         DiscreteCoveredCodeVisitorBase
             (
                 FSystem* InSystem,
-                const TSharedPtr<const FModel>& InModel,
+                const TSharedRef<FLiveInstance>& InLiveInstance,
                 const TSharedPtr<const FParameters>& InParams,
                 unsigned InLodMask,
                 bool bSkipResources=false
             )
-            : PARENT(bSkipResources)
+            : PARENT(bSkipResources) , 
+    		LiveInstance(InLiveInstance)
         {
             System = InSystem;
-            Model = InModel;
+            Model = InLiveInstance->Model;
             Params = InParams.Get();
             LODMask = InLodMask;
 
@@ -272,7 +273,7 @@ namespace UE::Mutable::Private
 
                 if (Args.condition)
                 {
-					bValue = System->BuildBool(nullptr, Model, Params, Args.condition);
+					bValue = System->BuildBool(LiveInstance, Params, Args.condition);
                 }
 
                 if (bValue)
@@ -315,7 +316,7 @@ namespace UE::Mutable::Private
 
                     PARENT::RecurseWithCurrentState( VarAddress );
 
-                    int32 Var = System->BuildInt(nullptr, Model, Params, VarAddress );
+                    int32 Var = System->BuildInt(LiveInstance, Params, VarAddress );
 
 					OP::ADDRESS ValueAt = DefAddress;
 
@@ -448,6 +449,7 @@ namespace UE::Mutable::Private
     protected:
         FSystem* System = nullptr;
 		TSharedPtr<const FModel> Model;
+    	TSharedRef<FLiveInstance> LiveInstance;
         const FParameters* Params = nullptr;
         uint32 LODMask = 0;
     };
@@ -477,11 +479,11 @@ namespace UE::Mutable::Private
         UniqueDiscreteCoveredCodeVisitor
             (
                 FSystem* InSystem,
-				const TSharedPtr<const FModel>& InModel,
+				const TSharedRef<FLiveInstance>& InLiveInstance,
                 const TSharedPtr<const FParameters>& InParams,
                 uint32 InLodMask
             )
-            : PARENT(InSystem, InModel, InParams, InLodMask )
+            : PARENT(InSystem, InLiveInstance, InParams, InLodMask )
         {
         }
 
diff --git a/Engine/Plugins/Mutable/Source/MutableRuntime/Internal/MuR/System.h b/Engine/Plugins/Mutable/Source/MutableRuntime/Internal/MuR/System.h
--- a/Engine/Plugins/Mutable/Source/MutableRuntime/Internal/MuR/System.h
+++ b/Engine/Plugins/Mutable/Source/MutableRuntime/Internal/MuR/System.h
@@ -494,9 +494,13 @@ namespace UE::Mutable::Private
 			return uint32_t(m_usedRangeIndices.Num()) - 1;
 		}
 
+		
+		void Prepare(const FProgram& InProgram, int32 InState);
+
+		
 		void Init(uint32 Size)
 		{
-			// This clear would prevent live update cache reusal
+			// This clear would prevent live update cache reuse
 			//OpExecutionData.clear();
 
 			OpExecutionData.resize(Size);
@@ -678,7 +682,7 @@ namespace UE::Mutable::Private
 			return Data->DataTypeIndex;
 		}
 
-		FVector4f GetColour(FCacheAddress Address)
+		FVector4f GetColor(FCacheAddress Address)
 		{
 			if (!Address.At) 
 			{
@@ -1492,12 +1496,13 @@ namespace UE::Mutable::Private
 		uint64 UpdatedParameters = 0;
 
 		/** Cached data for the generation of this instance. */
-		TSharedPtr<FProgramCache> Cache;
+		TSharedRef<FProgramCache> Cache = MakeShared<FProgramCache>();
 
 		TSharedPtr<FMeshIdRegistry> MeshIdRegistry;
 		TSharedPtr<FImageIdRegistry> ImageIdRegistry;
 		TSharedPtr<FMaterialIdRegistry> MaterialIdRegistry;
 		TSharedPtr<FSkeletalMeshIdRegistry> SkeletalMeshIdRegistry;
+		
 		TSharedPtr<FExternalResourceProvider> ExternalResourceProvider;
 		TSharedPtr<FPassthroughObjectLoader> PassthroughObjectLoader;
 	};
@@ -1514,10 +1519,10 @@ namespace UE::Mutable::Private
 		template<class KeyType, class ValueType>
 		using TMemoryTrackedMap = TMap<KeyType, ValueType, FDefaultMemoryTrackingSetAllocator<MemoryCounter>>;
 
-		/** Cached traking for streamed model data for one model. */
+		/** Cached tracking for streamed model data for one model. */
         struct FModelCacheEntry
         {
-            /** Model who's data is being tracked. */
+            /** Model whose data is being tracked. */
             TWeakPtr<const FModel> Model;
 
             /** For each model rom, the last time its streamed data was used. */
@@ -1526,13 +1531,11 @@ namespace UE::Mutable::Private
 			/** Count of pending operations for every rom index. */
 			TMemoryTrackedArray<uint16> PendingOpsPerRom;
         };
-
-		FLiveInstance CurrentLiveInstance;
-
+    	
 		/** Maximum working memory that mutable should be using. */
 		int64 BudgetBytes = 0;
 
-		/** Maximum excess memory reached suring the current operation. */
+		/** Maximum excess memory reached during the current operation. */
 		int64 BudgetExcessBytes = 0;
     	
 		/** This value is used to track the order of loading of roms. */
@@ -1544,39 +1547,14 @@ namespace UE::Mutable::Private
 		/** Data for each mutable instance that is being updated. */
 		TMemoryTrackedArray<TWeakPtr<FLiveInstance>> LiveInstances;
 
-		/** Temporary reference to the memory of the current instance being updated. Only valid during a mutable "atomic" operation, like a BeginUpdate or a GetImage. */
-		TSharedPtr<FProgramCache> CurrentInstanceCache;
-
-    	TSharedPtr<FMeshIdRegistry> CurrentMeshIdRegistry;
-    	
-		TSharedPtr<FImageIdRegistry> CurrentImageIdRegistry;
-
-    	TSharedPtr<FMaterialIdRegistry> CurrentMaterialIdRegistry;
-
-    	TSharedPtr<FSkeletalMeshIdRegistry> CurrentSkeletalMeshIdRegistry;
-
-    	/** Only valid in initial generation. We can not load UObjects while streaming. */
-    	TSharedPtr<FPassthroughObjectLoader> CurrentPassthroughObjectLoader;
-    	
-		/** Resources that have been used in the past, but haven't been deallocated because they still fitted the memory budget and they could be reused. */
-		TMemoryTrackedArray<TSharedPtr<FImage>> PooledImages;
-
-		/** List of intermediate resources that are not soterd anywhere yet. They are still locally referenced by code. */
-		TMemoryTrackedArray<TSharedPtr<const FImage>> TempImages;
-		TMemoryTrackedArray<TSharedPtr<const FMesh>> TempMeshes;
-
-		/** List of resources that are currently in any cache position, and the number of positions they are in. */
-		TMemoryTrackedMap<TSharedPtr<const FResource>, int32> CacheResources;
-
 		/** Given a mutable model, find or create its rom cache. */
 		FModelCacheEntry* FindModelCache(const FModel*);
 		FModelCacheEntry& FindOrAddModelCache(const TSharedPtr<const FModel>&);
 
         /** Make sure the working memory is below the internal budget, even counting with the passed additional memory. 
-		* An optional function can be passed to "block" the unload of certain roms of data.
 		* Return true if it succeeded, false otherwise.
 		*/
-		bool EnsureBudgetBelow(uint64 AdditionalMemory);
+		bool EnsureBudgetBelow(uint64 AdditionalMemory, const TSharedPtr<FProgramCache>& InCurrentCache);
 		
 		/** Return true if the memory budget is 90% full. */
 		bool IsMemoryBudgetFull() const;
@@ -1585,126 +1563,9 @@ namespace UE::Mutable::Private
 		int64 GetCurrentMemoryBytes() const;
 
         /** Register that a specific rom has been requested and update the heuristics to keep it in memory. */
-        void MarkRomUsed( int32 RomIndex, const TSharedPtr<const FModel>& );
-
-		/** */
-		[[nodiscard]] TManagedPtr<FImage> CreateImage(uint32 SizeX, uint32 SizeY, uint32 Lods, EImageFormat Format, EInitializationType Init)
-		{
-			CheckRunnerThread();
-
-			uint32 DataSize = FImage::CalculateDataSize(SizeX, SizeY, Lods, Format);
-
-			//TODO: 
-
-			// Make room in the budget
-			EnsureBudgetBelow(DataSize);
-
-			// Create it
-			TManagedPtr<FImage> Result = MakeManaged<FImage>(SizeX, SizeY, Lods, Format, Init);
-
-			return Result;
-		}
-
-		/** Ref will be nulled and relesed in any case. */
-		[[nodiscard]] TManagedPtr<FImage> CloneOrTakeOver(TManagedPtr<const FImage>&& Resource)
-		{
-			CheckRunnerThread();
-	
-			TManagedPtr<FImage> Result;
-			if (!Resource.IsUniqueReference())
-			{
-				uint32 DataSize = Resource->GetDataSize();
-				EnsureBudgetBelow(DataSize);
-
-				Result = Resource->Clone();
-			}
-			else
-			{
-				Result = ConstCastManagedPtr<FImage>(Resource);
-			}
-
-			Resource = nullptr;
-			return Result;
-		}
-
-
-		[[nodiscard]] TManagedPtr<FMesh> CreateMesh(int32 BudgetReserveSize)
-		{
-			CheckRunnerThread();
-
-			EnsureBudgetBelow(BudgetReserveSize);
-
-			TManagedPtr<FMesh> Result = MakeManaged<FMesh>();
-
-			return Result;
-		}
-
-		[[nodiscard]] TManagedPtr<FMesh> CloneOrTakeOver(TManagedPtr<const FMesh>&& Resource)
-		{
-			CheckRunnerThread();
-
-			TManagedPtr<FMesh> Result;
-			if (!Resource.IsUniqueReference())
-			{
-				const int32 ResourceDataSize = Resource->GetDataSize();
-				Result = CreateMesh(ResourceDataSize); 
-				Result->CopyFrom(*Resource);
-			}
-			else
-			{
-				Result = ConstCastManagedPtr<FMesh>(Resource);
-			}
-
-			Resource = nullptr;
-			return Result;
-		}
-
-		/** */
-		[[nodiscard]] TManagedPtr<const FMesh> LoadMesh(const FCacheAddress& From)
-		{
-			TManagedPtr<const FMesh> Result = CurrentInstanceCache->GetMesh(From);
-			return Result;
-		}
-
-    	/** */
-    	[[nodiscard]] TManagedPtr<const FSkeletalMesh> LoadSkeletalMesh(const FCacheAddress& From)
-		{
-			TManagedPtr<const FSkeletalMesh> Result = CurrentInstanceCache->GetSkeletalMesh(From);
-			return Result;
-		}
+		void MarkRomUsed( int32 RomIndex, const TSharedPtr<const FModel>& );
     	
 
-		/** */
-		[[nodiscard]] TManagedPtr<const FImage> LoadImage(const FCacheAddress& From)
-		{
-			return CurrentInstanceCache->GetImage(From);
-		}
-
-
-    	/** */
-    	[[nodiscard]] TManagedPtr<const FMaterial> LoadMaterial(const FCacheAddress& From)
-		{
-			TManagedPtr<const FMaterial> Result = CurrentInstanceCache->GetMaterial(From);
-			return Result;
-		}
-
-		/** */
-		void StoreImage(const FCacheAddress& To, const TManagedPtr<const FImage>& Resource)
-		{
-			CurrentInstanceCache->SetImage(To, Resource);
-		}
-
-		void StoreMesh(const FCacheAddress& To, const TManagedPtr<const FMesh>& Resource)
-		{
-			//Resource->CheckIntegrity();
-			CurrentInstanceCache->SetMesh(To, Resource);
-		}
-
-    	void StoreMaterial(const FCacheAddress& To, const TManagedPtr<const FMaterial>& Resource)
-		{
-			CurrentInstanceCache->SetMaterial(To, Resource);
-		}
-
 		int32 GetRomBytes() const;
 
     	/** Calculate the amount of bytes in data cached in the level 0 and 1 cache in all live instances. */
@@ -1778,13 +1639,11 @@ namespace UE::Mutable::Private
 		 * Remove all intermediate data (big and small) from the memory except for the one that has been explicitely
 		 * marked as state cache.
 		 */
-		void ClearCacheLayer0()
+		void ClearCacheLayer0(const TSharedRef<FProgramCache>& InInstanceCache)
 		{
-			check(CurrentInstanceCache);
-
 			MUTABLE_CPUPROFILER_SCOPE(ClearLayer0);
 
-			CodeContainer<FProgramCache::FOpExecutionData>::iterator it = CurrentInstanceCache->OpExecutionData.begin();
+			CodeContainer<FProgramCache::FOpExecutionData>::iterator it = InInstanceCache->OpExecutionData.begin();
 			for (; it.IsValid(); ++it)
 			{
 				FProgramCache::FOpExecutionData& Data = *it;
@@ -1797,33 +1656,33 @@ namespace UE::Mutable::Private
 				{
 				case EDataType::Image:
 				{
-					CurrentInstanceCache->ImageResults[Data.DataTypeIndex] = nullptr;
-					CurrentInstanceCache->StrongImageResults[Data.DataTypeIndex] = nullptr;
+					InInstanceCache->ImageResults[Data.DataTypeIndex] = nullptr;
+					InInstanceCache->StrongImageResults[Data.DataTypeIndex] = nullptr;
 					break;
 				}
 
 				case EDataType::Mesh:
 				{
-					CurrentInstanceCache->MeshResults[Data.DataTypeIndex] = nullptr;
-					CurrentInstanceCache->StrongMeshResults[Data.DataTypeIndex] = nullptr;
+					InInstanceCache->MeshResults[Data.DataTypeIndex] = nullptr;
+					InInstanceCache->StrongMeshResults[Data.DataTypeIndex] = nullptr;
 					break;
 				}
 
 				case EDataType::SkeletalMesh:
 				{
-					CurrentInstanceCache->SkeletalMeshResults[Data.DataTypeIndex] = nullptr;
+					InInstanceCache->SkeletalMeshResults[Data.DataTypeIndex] = nullptr;
 					break;
 				}
 
 				case EDataType::Layout:
 				{
-					CurrentInstanceCache->LayoutResults[Data.DataTypeIndex] = nullptr;
+					InInstanceCache->LayoutResults[Data.DataTypeIndex] = nullptr;
 					break;
 				}
 
 				case EDataType::Instance:
 				{
-					CurrentInstanceCache->InstanceResults[Data.DataTypeIndex] = nullptr;
+					InInstanceCache->InstanceResults[Data.DataTypeIndex] = nullptr;
 					break;
 				}
 
@@ -1840,11 +1699,11 @@ namespace UE::Mutable::Private
 		/** Remove all intermediate data (big and small) from the memory including the one that has been explicitely
 		* marked as state cache.
 		*/
-		void ClearCacheLayer1()
+		void ClearCacheLayer1(const TSharedRef<FProgramCache>& InInstanceCache)
 		{
 			MUTABLE_CPUPROFILER_SCOPE(ClearLayer1);
 
-			CodeContainer<FProgramCache::FOpExecutionData>::iterator it = CurrentInstanceCache->OpExecutionData.begin();
+			CodeContainer<FProgramCache::FOpExecutionData>::iterator it = InInstanceCache->OpExecutionData.begin();
 			for (; it.IsValid(); ++it)
 			{
 				FProgramCache::FOpExecutionData& Data = *it;
@@ -1853,34 +1712,34 @@ namespace UE::Mutable::Private
 				Data.IsValueValid = false;
 			}
 
-			for (TManagedWeakPtr<const FImage>& Value : CurrentInstanceCache->ImageResults)
+			for (TManagedWeakPtr<const FImage>& Value : InInstanceCache->ImageResults)
 			{
 				Value = nullptr;
 			}
 
-			for (TManagedPtr<const FImage>& Value : CurrentInstanceCache->StrongImageResults)
+			for (TManagedPtr<const FImage>& Value : InInstanceCache->StrongImageResults)
 			{
 				Value = nullptr;
 			}
 
 
-			for (TManagedWeakPtr<const FSkeletalMesh>& Value : CurrentInstanceCache->SkeletalMeshResults)
+			for (TManagedWeakPtr<const FSkeletalMesh>& Value : InInstanceCache->SkeletalMeshResults)
 			{
 				Value = nullptr;
 			}
 
-			for (TManagedPtr<const FMesh>& Value : CurrentInstanceCache->StrongMeshResults)
+			for (TManagedPtr<const FMesh>& Value : InInstanceCache->StrongMeshResults)
 			{
 				Value = nullptr;
 			}
 
 
-			for (TManagedWeakPtr<const FLayout>& Value : CurrentInstanceCache->LayoutResults)
+			for (TManagedWeakPtr<const FLayout>& Value : InInstanceCache->LayoutResults)
 			{
 				Value = nullptr;
 			}
 
-			for (TManagedWeakPtr<const FInstance>& Value : CurrentInstanceCache->InstanceResults)
+			for (TManagedWeakPtr<const FInstance>& Value : InInstanceCache->InstanceResults)
 			{
 				Value = nullptr;
 			}
@@ -1953,15 +1812,9 @@ namespace UE::Mutable::Private
 
 
 		/** This is a development-only check to make sure calls to resource management happen in the correct thread. */
-		FORCEINLINE void EndRunnerThread()
+		FORCEINLINE void EndRunnerThread(const TSharedRef<FLiveInstance>& InLiveInstance)
 		{
-			CurrentInstanceCache->CheckHitCountsCleared();
-
-			// If this check fails it means some operation is not correctly handling resource management and didn't release
-			// a resource it created.
-			// This should be reported and reviewed, but it is not fatal. Some unnecessary memory may be used temporarily.
-			ensure(TempImages.Num() == 0);
-			ensure(TempMeshes.Num() == 0);
+			InLiveInstance->Cache->CheckHitCountsCleared();
 
 #if !(UE_BUILD_SHIPPING || UE_BUILD_TEST)
 			// If this check fails it means have not set up correctly all the paths to debug threading for resource management.
@@ -2008,15 +1861,27 @@ namespace UE::Mutable::Private
         //! A call to NewInstance must be paired to a call to ReleasesInstance when the instance is
         //! no longer needed.
         UE_API TSharedRef<FLiveInstance> NewInstance(const TSharedPtr<const FModel>& Model,
+		const TSharedPtr<UE::Mutable::Private::FMeshIdRegistry>& InMeshIdRegistry,
+		const TSharedPtr<UE::Mutable::Private::FImageIdRegistry>& InImageIdRegistry,
+		const TSharedPtr<UE::Mutable::Private::FMaterialIdRegistry>& InMaterialIdRegistry,
+		const TSharedPtr<UE::Mutable::Private::FSkeletalMeshIdRegistry>& InSkeletalMeshIdRegistry,
         	const TSharedPtr<FExternalResourceProvider>& ExternalResourceProvider,
         	const TSharedPtr<FPassthroughObjectLoader>& PassthroughObjectLoader);
-		
+
+		/** 
+		 * Create a ready to use new FLiveInstance to be used during the mutable build operations. It will not be taken into consideration for any update operation
+		 * as the life cycle of the returned object is expected to be short and contained during the compilation operation.
+		 * It does not require further initialization as it is fully capable of operating as is.
+		 */
+		UE_API TSharedRef<FLiveInstance> NewBuildInstance(const TSharedPtr<const FModel>& InModel, 
+			const TSharedPtr<FExternalResourceProvider>& ExternalResourceProvider);
+				
         //! \brief Update an instance with a new parameter set and/or state.
         //!
         //! \warning a call to BeginUpdate must be paired with a call to EndUpdate once the returned
         //! data has been processed.
-        //! \param InstanceID The id of the instance to update, as created by a NewInstance call.
-        //! \param Params The parameters that customise this instance.
+        //! \param LiveInstance The instance to update, as created by a NewInstance call.
+        //! \param InParams The parameters that customize this instance.
         //! \param StateIndex The index of the state this instance will be set to. The states range
         //! from 0 to Model::GetStateCount-1
         //! \param LodMask Bitmask selecting the levels of detail to build (i-th bit selects i-th lod).
@@ -2026,21 +1891,13 @@ namespace UE::Mutable::Private
 		//!
 		//! Very fast. Can be called in the Game Thread.
 		UE_API TManagedPtr<const FInstance> BeginUpdate_GameThread(const TSharedRef<FLiveInstance>& LiveInstance,
-			const TSharedPtr<const FParameters>& Params,
-			const TSharedPtr<FMeshIdRegistry>& InMeshIdRegistry,
-			const TSharedPtr<FImageIdRegistry>& InImageIdRegistry,
-			const TSharedPtr<FMaterialIdRegistry>& InMaterialIdRegistry,
-			const TSharedPtr<FSkeletalMeshIdRegistry>& InSkeletalMeshIdRegistry,
+			const TSharedPtr<const FParameters>& InParams,
 			int32 StateIndex,
 			uint32 LodMask);
 
 		/** Potentially expensive. Must be called in the Mutable Thread. */
 		UE_API TManagedPtr<const FInstance> BeginUpdate_MutableThread(const TSharedRef<FLiveInstance>& LiveInstance,
 			const TSharedPtr<const FParameters>& InParams,
-			const TSharedPtr<FMeshIdRegistry>& InMeshIdRegistry,
-			const TSharedPtr<FImageIdRegistry>& InImageIdRegistry,
-			const TSharedPtr<FMaterialIdRegistry>& InMaterialIdRegistry,
-			const TSharedPtr<FSkeletalMeshIdRegistry>& InSkeletalMeshIdRegistry,
 			int32 InStateIndex,
 			uint32 InLodMask);
 		
@@ -2102,25 +1959,22 @@ namespace UE::Mutable::Private
         FSystem& operator=( FSystem&& ) = delete;
 		
 		/** This method can be used to internally prepare for code execution. */
-		MUTABLERUNTIME_API void BeginBuild(const TSharedPtr<const FModel>&,
-			const TSharedPtr<FMeshIdRegistry>&,
-			const TSharedPtr<FImageIdRegistry>&,
-			const TSharedPtr<FMaterialIdRegistry>&,
-			const TSharedPtr<FSkeletalMeshIdRegistry>&);
+		MUTABLERUNTIME_API void BeginBuild(const TSharedPtr<const FModel>&);
     	
 		MUTABLERUNTIME_API void EndBuild();
 
-		MUTABLERUNTIME_API bool BuildBool(const TSharedPtr<FExternalResourceProvider>& ExternalResourceProvider, const TSharedPtr<const FModel>&, const FParameters*, OP::ADDRESS) ;
-		MUTABLERUNTIME_API int32 BuildInt(const TSharedPtr<FExternalResourceProvider>& ExternalResourceProvider, const TSharedPtr<const FModel>&, const FParameters*, OP::ADDRESS) ;
-		MUTABLERUNTIME_API float BuildScalar(const TSharedPtr<FExternalResourceProvider>& ExternalResourceProvider, const TSharedPtr<const FModel>&, const FParameters*, OP::ADDRESS) ;
-		MUTABLERUNTIME_API FVector4f BuildColour(const TSharedPtr<FExternalResourceProvider>& ExternalResourceProvider, const TSharedPtr<const FModel>&, const FParameters*, OP::ADDRESS) ;
-		MUTABLERUNTIME_API TManagedPtr<const String> BuildString(const TSharedPtr<FExternalResourceProvider>& ExternalResourceProvider, const TSharedPtr<const FModel>&, const FParameters*, OP::ADDRESS) ;
-		MUTABLERUNTIME_API TManagedPtr<const FImage> BuildImage(const TSharedPtr<FExternalResourceProvider>& ExternalResourceProvider, const TSharedPtr<const FModel>&, const FParameters*, OP::ADDRESS, int32 MipsToSkip);
-		MUTABLERUNTIME_API TManagedPtr<const FMesh> BuildMesh(const TSharedPtr<FExternalResourceProvider>& ExternalResourceProvider, const TSharedPtr<const FModel>&, const FParameters*, OP::ADDRESS, EMeshContentFlags MeshContentFilter);
-		MUTABLERUNTIME_API TManagedPtr<const FInstance> BuildInstance(const TSharedPtr<FExternalResourceProvider>& ExternalResourceProvider, const TSharedPtr<const FModel>&, const FParameters*, OP::ADDRESS);
-		MUTABLERUNTIME_API TManagedPtr<const FLayout> BuildLayout(const TSharedPtr<FExternalResourceProvider>& ExternalResourceProvider, const TSharedPtr<const FModel>&, const FParameters*, OP::ADDRESS) ;
-    	MUTABLERUNTIME_API FProjector BuildProjector(const TSharedPtr<FExternalResourceProvider>& ExternalResourceProvider, const TSharedPtr<const FModel>&, const FParameters*, OP::ADDRESS) ;
-    	
+		MUTABLERUNTIME_API bool BuildBool(const TSharedRef<FLiveInstance>& InLiveInstance, const FParameters* InParameters, OP::ADDRESS InAddress) ;
+		MUTABLERUNTIME_API int32 BuildInt(const TSharedRef<FLiveInstance>& InLiveInstance, const FParameters* InParameters, OP::ADDRESS InAddress) ;
+		MUTABLERUNTIME_API float BuildScalar(const TSharedRef<FLiveInstance>& InLiveInstance, const FParameters* InParameters, OP::ADDRESS InAddress) ;
+		MUTABLERUNTIME_API FVector4f BuildColor(const TSharedRef<FLiveInstance>& InLiveInstance, const FParameters* InParameters, OP::ADDRESS InAddress) ;
+		MUTABLERUNTIME_API TManagedPtr<const String> BuildString(const TSharedRef<FLiveInstance>& InLiveInstance, const FParameters* InParameters, OP::ADDRESS InAddress) ;
+		MUTABLERUNTIME_API TManagedPtr<const FImage> BuildImage(const TSharedRef<FLiveInstance>& InLiveInstance, const FParameters* InParameters, OP::ADDRESS InAddress, int32 MipsToSkip);
+		MUTABLERUNTIME_API TManagedPtr<const FMesh> BuildMesh(const TSharedRef<FLiveInstance>& InLiveInstance, const FParameters* InParameters, OP::ADDRESS InRootAddress, EMeshContentFlags MeshContentFilter);
+		MUTABLERUNTIME_API TManagedPtr<const FInstance> BuildInstance(const TSharedRef<FLiveInstance>& InLiveInstance, const FParameters* InParameters, OP::ADDRESS InAddress);
+		MUTABLERUNTIME_API TManagedPtr<const FLayout> BuildLayout(const TSharedRef<FLiveInstance>& InLiveInstance, const FParameters* InParameters, OP::ADDRESS InAddress) ;
+		MUTABLERUNTIME_API FProjector BuildProjector(const TSharedRef<FLiveInstance>& InLiveInstance, const FParameters* InParameters, OP::ADDRESS InAddress) ;
+		
+		
         FSettings Settings;
 
         //! Data streaming interface, if any.
@@ -2139,15 +1993,12 @@ namespace UE::Mutable::Private
 		bool CheckUpdatedParameters(const TSharedRef<FLiveInstance>& LiveInstance, const TSharedPtr<const FParameters>&, uint64& OutUpdatedParameters );
 		
 		void RunCode(
-			const TSharedPtr<FExternalResourceProvider>& ExternalResourceProvider,
-			const TSharedPtr<const FModel>&,
+			const TSharedRef<FLiveInstance>& InLiveInstance,
 			const FParameters*,
 			OP::ADDRESS,
 			uint32 LODs = AllLODs,
 			uint8 ExecutionOptions = 0);
 
-		void PrepareCache(const FModel*, int32 State);
-
 		//! Update some mutable core unreal stats.
 		void UpdateStats();
 	};
diff --git a/Engine/Plugins/Mutable/Source/MutableRuntime/Private/MuR/CodeRunner.cpp b/Engine/Plugins/Mutable/Source/MutableRuntime/Private/MuR/CodeRunner.cpp
--- a/Engine/Plugins/Mutable/Source/MutableRuntime/Private/MuR/CodeRunner.cpp
+++ b/Engine/Plugins/Mutable/Source/MutableRuntime/Private/MuR/CodeRunner.cpp
@@ -151,42 +151,42 @@ namespace UE::Mutable::Private
 {
 
 	TSharedRef<CodeRunner> CodeRunner::Create(
+		const TSharedRef<FLiveInstance>& InLiveInstance,
 		const FSettings& InSettings,
 		FSystem* InSystem,
 		EExecutionStrategy InExecutionStrategy,
-		const TSharedPtr<const FModel>& InModel,
 		const FParameters* InParams,
 		OP::ADDRESS At,
 		uint32 InLODMask,
 		uint8 ExecutionOptions,
 		bool bForceLazy,
-		FScheduledOp::EType Type,
-		const TSharedPtr<FExternalResourceProvider>& ExternalResourceProvider)
+		FScheduledOp::EType Type)
 	{
-		return MakeShared<CodeRunner>(FPrivateToken {},
-				InSettings, InSystem, InExecutionStrategy, InModel, InParams, At, InLODMask, ExecutionOptions, bForceLazy, Type, ExternalResourceProvider);
+		return MakeShared<CodeRunner>(InLiveInstance, FPrivateToken {},
+				InSettings, InSystem, InExecutionStrategy,  InParams, At, InLODMask, ExecutionOptions, bForceLazy, Type);
 	}
 
 
-    CodeRunner::CodeRunner(FPrivateToken PrivateToken, 
+    CodeRunner::CodeRunner(
+		const TSharedRef<FLiveInstance>& InLiveInstance,
+    	FPrivateToken PrivateToken, 
 		const FSettings& InSettings,
 		FSystem* InSystem,
 		EExecutionStrategy InExecutionStrategy,
-		const TSharedPtr<const FModel>& InModel,
 		const FParameters* InParams,
 		OP::ADDRESS At,
 		uint32 InLodMask,
 		uint8 ExecutionOptions,
 		bool bForceLazy,
-		FScheduledOp::EType Type,
-		const TSharedPtr<FExternalResourceProvider>& ExternalResourceProvider)
+		FScheduledOp::EType Type)
 		: Settings(InSettings)
+		, LiveInstance(InLiveInstance)
 		, ExecutionStrategy(InExecutionStrategy)
 		, System(InSystem)
-		, Model(InModel)
+		, Model(InLiveInstance->Model)
 		, Params(InParams)
 		, LODMask(InLodMask)
-		, ExternalResourceProvider(ExternalResourceProvider)
+		, ExternalResourceProvider(InLiveInstance->ExternalResourceProvider)
 	{
 		MUTABLE_CPUPROFILER_SCOPE(CodeRunner_Create)
 		const FProgram& Program = Model->Program;
@@ -217,7 +217,7 @@ namespace UE::Mutable::Private
 
 	FProgramCache& CodeRunner::GetMemory()
     {
-		return *System->WorkingMemoryManager.CurrentInstanceCache;
+		return *LiveInstance->Cache;
 	}
 
 
@@ -529,6 +529,8 @@ namespace UE::Mutable::Private
         {
         case EOpType::IN_ADDSTRING:
         {
+        	MUTABLE_CPUPROFILER_SCOPE(IN_ADDSTRING);
+        		
 			OP::InstanceAddArgs Args = Program.GetOpArgs<OP::InstanceAddArgs>( Item.At );
             switch ( Item.Stage )
             {
@@ -574,6 +576,8 @@ namespace UE::Mutable::Private
 
         case EOpType::IN_ADDCOMPONENT:
         {
+        	MUTABLE_CPUPROFILER_SCOPE(IN_ADDCOMPONENT);
+        		
 			OP::InstanceAddComponentArgs Args = Program.GetOpArgs<OP::InstanceAddComponentArgs>(Item.At);
             switch (Item.Stage)
             {
@@ -624,6 +628,8 @@ namespace UE::Mutable::Private
 
         case EOpType::IN_ADDSURFACE:
         {
+        	MUTABLE_CPUPROFILER_SCOPE(IN_ADDSURFACE);
+        		
 			OP::InstanceAddArgs Args = Program.GetOpArgs<OP::InstanceAddArgs>(Item.At);
             switch (Item.Stage)
             {
@@ -691,6 +697,8 @@ namespace UE::Mutable::Private
 
         case EOpType::IN_ADDLOD:
         {
+        	MUTABLE_CPUPROFILER_SCOPE(IN_ADDLOD);
+        		
 			const uint8* Data = Program.GetOpArgsPointer(Item.At);
 
 			uint8 LODCount;
@@ -769,6 +777,8 @@ namespace UE::Mutable::Private
 
 		case EOpType::IN_ADDEXTENSIONDATA:
 		{
+        	MUTABLE_CPUPROFILER_SCOPE(IN_ADDEXTENSIONDATA);
+        		
 			OP::InstanceAddExtensionDataArgs Args = Program.GetOpArgs<OP::InstanceAddExtensionDataArgs>(Item.At);
 			switch (Item.Stage)
 			{
@@ -841,6 +851,8 @@ namespace UE::Mutable::Private
         {
         case EOpType::IN_ADDMESH:
         {
+        	MUTABLE_CPUPROFILER_SCOPE(IN_ADDMESH);
+        		
 			OP::InstanceAddArgs Args = Program.GetOpArgs<OP::InstanceAddArgs>(Item.At);
             switch (Item.Stage)
             {
@@ -864,7 +876,7 @@ namespace UE::Mutable::Private
 
                 if ( Args.value )
                 {
-					FMeshId MeshId = FSystem::GetMeshId(System->WorkingMemoryManager.CurrentMeshIdRegistry, InModel, InParams, Args.value);
+					FMeshId MeshId = FSystem::GetMeshId(LiveInstance->MeshIdRegistry, InModel, InParams, Args.value);
 					OP::ADDRESS NameAd = Args.name;
 					check(NameAd < (uint32)Program.ConstantStrings.Num());
 					const FString& Name = Program.ConstantStrings[NameAd];
@@ -882,6 +894,8 @@ namespace UE::Mutable::Private
 
         case EOpType::IN_ADDOVERLAYMATERIAL:
         	{
+        		MUTABLE_CPUPROFILER_SCOPE(IN_ADDOVERLAYMATERIAL);
+        		
         		OP::InstanceAddMaterialArgs Args = Program.GetOpArgs<OP::InstanceAddMaterialArgs>(Item.At);
         		switch (Item.Stage)
         		{
@@ -909,7 +923,7 @@ namespace UE::Mutable::Private
 
         				if (Args.Material)
         				{
-        					FMaterialId MaterialId = FSystem::GetMaterialId(System->WorkingMemoryManager.CurrentMaterialIdRegistry, InModel, InParams, Args.Material);
+        					FMaterialId MaterialId = FSystem::GetMaterialId(LiveInstance->MaterialIdRegistry, InModel, InParams, Args.Material);
         					Result->SetOverlayMaterialId(0, MaterialId);
         				}
 
@@ -926,6 +940,8 @@ namespace UE::Mutable::Private
 
         case EOpType::IN_ADDOVERRIDEMATERIAL:
         	{
+        		MUTABLE_CPUPROFILER_SCOPE(IN_ADDOVERRIDEMATERIAL);
+        		
         		OP::InstanceAddOverrideMaterialArgs Args = Program.GetOpArgs<OP::InstanceAddOverrideMaterialArgs>(Item.At);
         		switch (Item.Stage)
         		{
@@ -953,7 +969,7 @@ namespace UE::Mutable::Private
 
         				if (Args.Material)
         				{
-        					FMaterialId MaterialId = FSystem::GetMaterialId(System->WorkingMemoryManager.CurrentMaterialIdRegistry, InModel, InParams, Args.Material);
+        					FMaterialId MaterialId = FSystem::GetMaterialId(LiveInstance->MaterialIdRegistry, InModel, InParams, Args.Material);
         					const FString& SlotName = Program.ConstantStrings[Args.SlotName];
         					Result->AddOverrideMaterial(0, FName(SlotName), MaterialId);
         				}
@@ -971,6 +987,8 @@ namespace UE::Mutable::Private
         	
         case EOpType::IN_ADDMATERIAL:
         	{
+        		MUTABLE_CPUPROFILER_SCOPE(IN_ADDMATERIAL);
+        		
         		OP::InstanceAddMaterialArgs Args = Program.GetOpArgs<OP::InstanceAddMaterialArgs>(Item.At);
         		switch (Item.Stage)
         		{
@@ -997,7 +1015,7 @@ namespace UE::Mutable::Private
 
         				if (Args.Material)
         				{
-        					FMaterialId MaterialId = FSystem::GetMaterialId(System->WorkingMemoryManager.CurrentMaterialIdRegistry, InModel, InParams, Args.Material);
+        					FMaterialId MaterialId = FSystem::GetMaterialId(LiveInstance->MaterialIdRegistry, InModel, InParams, Args.Material);
 							Result->SetMaterialId(0, 0, 0, MaterialId);
         				}
 
@@ -1014,6 +1032,8 @@ namespace UE::Mutable::Private
 
 		case EOpType::IN_ADDSKELETALMESH:
         	{
+        		MUTABLE_CPUPROFILER_SCOPE(IN_ADDSKELETALMESH);
+        		
         		OP::InstanceAddArgs Args = Program.GetOpArgs<OP::InstanceAddArgs>(Item.At);
         		switch (Item.Stage)
         		{
@@ -1047,7 +1067,7 @@ namespace UE::Mutable::Private
 
         				if (Args.value)
         				{
-        					FSkeletalMeshId SkeletalMeshId = FSystem::GetSkeletalMeshId(System->WorkingMemoryManager.CurrentSkeletalMeshIdRegistry, InModel, InParams, Args.value);
+        					FSkeletalMeshId SkeletalMeshId = FSystem::GetSkeletalMeshId(LiveInstance->SkeletalMeshIdRegistry, InModel, InParams, Args.value);
         					Instance->SetSkeletalMeshId(0, SkeletalMeshId);
         				}
 						
@@ -1081,6 +1101,8 @@ namespace UE::Mutable::Private
 
         case EOpType::ME_CONSTANT:
         {
+        	MUTABLE_CPUPROFILER_SCOPE(ME_CONSTANT);
+        		
 			OP::MeshConstantArgs Args = Program.GetOpArgs<OP::MeshConstantArgs>(Item.At);
 			EMeshContentFlags MeshContentFlags = (EMeshContentFlags)Item.ExecutionOptions;
 
@@ -1101,6 +1123,8 @@ namespace UE::Mutable::Private
 
         case EOpType::IM_CONSTANT:
         {
+        	MUTABLE_CPUPROFILER_SCOPE(IM_CONSTANT);
+        		
 			OP::ResourceConstantArgs Args = Program.GetOpArgs<OP::ResourceConstantArgs>(Item.At);
 
 			int32 MipsToSkip = Item.ExecutionOptions;
@@ -1124,6 +1148,8 @@ namespace UE::Mutable::Private
 
 		case EOpType::ED_CONSTANT:
 		{
+        	MUTABLE_CPUPROFILER_SCOPE(ED_CONSTANT);
+        		
 			OP::ResourceConstantArgs Args = Program.GetOpArgs<OP::ResourceConstantArgs>(Item.At);
 
 			// Assume the ROM has been loaded previously
@@ -1138,6 +1164,8 @@ namespace UE::Mutable::Private
 
 		case EOpType::MI_CONSTANT:
 		{
+        	MUTABLE_CPUPROFILER_SCOPE(MI_CONSTANT);
+        		
 			OP::MaterialConstantArgs Args = Program.GetOpArgs<OP::MaterialConstantArgs>(Item.At);
 
         	TManagedPtr<const FMaterial> SourceConst = Program.ConstantMaterials[Args.Value];
@@ -6700,6 +6728,8 @@ namespace UE::Mutable::Private
 			
 		case EOpType::MI_MODIFY:
 		{
+			MUTABLE_CPUPROFILER_SCOPE(RunCodeMaterial_MI_MODIFY);
+				
 			OP::MaterialModifyArgs Args = Program.GetOpArgs<OP::MaterialModifyArgs>(Item.At);
 			int32 ParameterCount = Args.NumParameters;
 
@@ -7398,12 +7428,12 @@ namespace UE::Mutable::Private
 
 			case EDataType::Mesh:
 				RunCode_Mesh(Item, InModel.Get());
-				System->WorkingMemoryManager.EnsureBudgetBelow(0);
+				EnsureBudgetBelow(0);
 				break;
 
 			case EDataType::Image:
 				RunCode_Image(Item, pParams, InModel.Get());
-				System->WorkingMemoryManager.EnsureBudgetBelow(0);
+				EnsureBudgetBelow(0);
 				break;
 
 			case EDataType::Layout:
diff --git a/Engine/Plugins/Mutable/Source/MutableRuntime/Private/MuR/CodeRunner.h b/Engine/Plugins/Mutable/Source/MutableRuntime/Private/MuR/CodeRunner.h
--- a/Engine/Plugins/Mutable/Source/MutableRuntime/Private/MuR/CodeRunner.h
+++ b/Engine/Plugins/Mutable/Source/MutableRuntime/Private/MuR/CodeRunner.h
@@ -4,23 +4,16 @@
 
 #include "Containers/Array.h"
 #include "Containers/ContainerAllocationPolicies.h"
-#include "HAL/PlatformCrt.h"
-#include "HAL/PlatformMath.h"
 #include "MuR/Image.h"
 #include "MuR/Layout.h"
 #include "MuR/Operations.h"
-#include "MuR/Ptr.h"
 #include "MuR/RefCounted.h"
 #include "MuR/Serialisation.h"
 #include "MuR/Settings.h"
 #include "MuR/System.h"
-#include "MuR/System.h"
-#include "MuR/Types.h"
 #include "Tasks/Task.h"
-#include "Templates/RefCounting.h"
 #include "Templates/SharedPointer.h"
 #include "Templates/Tuple.h"
-#include "Containers/StaticArray.h"
 #include "Concepts/SameAs.h"
 
 namespace UE::Mutable::Private::MemoryCounters
@@ -45,32 +38,31 @@ namespace  UE::Mutable::Private
 
 	public:
 		static TSharedRef<CodeRunner> Create(
+				const TSharedRef<FLiveInstance>& InLiveInstance,
 				const FSettings&, 
 				FSystem*, 
 				EExecutionStrategy,
-				const TSharedPtr<const FModel>&, 
 				const FParameters* Params,
 				OP::ADDRESS At,
 				uint32 LodMask,
 				uint8 ExecutionOptions,
 				bool bForceLazy,
-				FScheduledOp::EType Type,
-				const TSharedPtr<FExternalResourceProvider>& ExternalResourceProvider);
+				FScheduledOp::EType Type);
 
     	// Private constructor to prevent stack allocation. In general we can not call AsShared() if the lifetime is
 		// bounded.
-		explicit CodeRunner(FPrivateToken, 
+		explicit CodeRunner(
+			const TSharedRef<FLiveInstance>& InLiveInstance,	
+			FPrivateToken, 
 			const FSettings&, 
 			FSystem*, 
 			EExecutionStrategy,
-			const TSharedPtr<const FModel>&, 
 			const FParameters* Params,
 			OP::ADDRESS At,
 			uint32 LodMask,
 			uint8 ExecutionOptions,
 			bool bForceLazy,
-			FScheduledOp::EType Type,
-			const TSharedPtr<FExternalResourceProvider>& ExternalResourceProvider);
+			FScheduledOp::EType Type);
 
     protected:
 		struct FProfileContext
@@ -138,7 +130,7 @@ namespace  UE::Mutable::Private
         void RunCode_Instance(const FScheduledOp&, const FModel*, uint32 LodMask);
         void RunCode_InstanceAddResource(const FScheduledOp&, const TSharedPtr<const FModel>& Model, const FParameters* );
 
-		/** Return false incase of failure. */
+		/** Return false in case of failure. */
         bool RunCode_ConstantResource(const FScheduledOp&, const FModel* );
 
         void RunCode_Mesh(const FScheduledOp&, const FModel* );
@@ -175,6 +167,10 @@ namespace  UE::Mutable::Private
 
 		/** Settings that may affect the execution of some operations, like image conversion quality. */
 		FSettings Settings;
+    	
+    	/** Instance whose data is being used by this runner.*/
+    	const TSharedRef<FLiveInstance> LiveInstance;
+
 	public:
         /** 
 		 * Heap of intermediate data pushed by some instructions and referred by others.
@@ -385,195 +381,229 @@ namespace  UE::Mutable::Private
 
 	public:
 
-		inline bool LoadBool(const FCacheAddress& From)
+		inline bool LoadBool(const FCacheAddress& From) const
 		{
-			return System->WorkingMemoryManager.CurrentInstanceCache->GetBool(From);
+			return LiveInstance->Cache->GetBool(From);
 		}
 
-		inline float LoadInt(const FCacheAddress& From)
+		inline float LoadInt(const FCacheAddress& From) const
 		{
-			return System->WorkingMemoryManager.CurrentInstanceCache->GetInt(From);
+			return LiveInstance->Cache->GetInt(From);
 		}
 
-		inline float LoadScalar(const FCacheAddress& From)
+		inline float LoadScalar(const FCacheAddress& From) const
 		{
-			return System->WorkingMemoryManager.CurrentInstanceCache->GetScalar(From);
+			return LiveInstance->Cache->GetScalar(From);
 		}
 
-		inline FVector4f LoadColor(const FCacheAddress& From)
+		inline FVector4f LoadColor(const FCacheAddress& From) const
 		{
-			return System->WorkingMemoryManager.CurrentInstanceCache->GetColour(From);
+			return LiveInstance->Cache->GetColor(From);
 		}
 
-    	inline FMatrix44f LoadMatrix(const FCacheAddress& From)
-		{
-			return System->WorkingMemoryManager.CurrentInstanceCache->GetMatrix(From);
+    	inline FMatrix44f LoadMatrix(const FCacheAddress& From) const
+	    {
+			return LiveInstance->Cache->GetMatrix(From);
 		}
 
-		inline TManagedPtr<const FMaterial> LoadMaterial(const FCacheAddress& From)
+		inline TManagedPtr<const FMaterial> LoadMaterial(const FCacheAddress& From) const
 		{
-			return System->WorkingMemoryManager.CurrentInstanceCache->GetMaterial(From);
+			return LiveInstance->Cache->GetMaterial(From);
 		}
 
-		inline TManagedPtr<const String> LoadString(const FCacheAddress& From)
+		inline TManagedPtr<const String> LoadString(const FCacheAddress& From) const
 		{
-			return System->WorkingMemoryManager.CurrentInstanceCache->GetString(From);
+			return LiveInstance->Cache->GetString(From);
 		}
 
-		inline FProjector LoadProjector(const FCacheAddress& From)
+		inline FProjector LoadProjector(const FCacheAddress& From) const
 		{
-			return System->WorkingMemoryManager.CurrentInstanceCache->GetProjector(From);
+			return LiveInstance->Cache->GetProjector(From);
 		}
 
-		inline TManagedPtr<const FMesh> LoadMesh(const FCacheAddress& From)
+		inline TManagedPtr<const FMesh> LoadMesh(const FCacheAddress& From) const
 		{
-			return System->WorkingMemoryManager.LoadMesh(From);
+			return LiveInstance->Cache->GetMesh(From);
 		}
 
-    	inline TManagedPtr<const FSkeletalMesh> LoadSkeletalMesh(const FCacheAddress& From)
-		{
-			return System->WorkingMemoryManager.CurrentInstanceCache->GetSkeletalMesh(From);
+    	inline TManagedPtr<const FSkeletalMesh> LoadSkeletalMesh(const FCacheAddress& From) const
+	    {
+			return LiveInstance->Cache->GetSkeletalMesh(From);
 		}
     	
-		inline TManagedPtr<const FImage> LoadImage(const FCacheAddress& From)
+		inline TManagedPtr<const FImage> LoadImage(const FCacheAddress& From) const
 		{
-			return System->WorkingMemoryManager.LoadImage(From);
+			return LiveInstance->Cache->GetImage(From);
 		}
 
-		inline TManagedPtr<const FLayout> LoadLayout(const FCacheAddress& From)
+		inline TManagedPtr<const FLayout> LoadLayout(const FCacheAddress& From) const
 		{
-			return System->WorkingMemoryManager.CurrentInstanceCache->GetLayout(From);
+			return LiveInstance->Cache->GetLayout(From);
 		}
 
-		inline TManagedPtr<const FInstance> LoadInstance(const FCacheAddress& From)
+		inline TManagedPtr<const FInstance> LoadInstance(const FCacheAddress& From) const
 		{
-			return System->WorkingMemoryManager.CurrentInstanceCache->GetInstance(From);
+			return LiveInstance->Cache->GetInstance(From);
 		}
 
-		inline TManagedPtr<const FExtensionData> LoadExtensionData(const FCacheAddress& From)
+		inline TManagedPtr<const FExtensionData> LoadExtensionData(const FCacheAddress& From) const
 		{
-			return System->WorkingMemoryManager.CurrentInstanceCache->GetExtensionData(From);
+			return LiveInstance->Cache->GetExtensionData(From);
 		}
     	
-    	inline TManagedPtr<const FInstancedStruct> LoadInstancedStruct(const FCacheAddress& From)
-		{
-			return System->WorkingMemoryManager.CurrentInstanceCache->GetInstancedStruct(From);
+    	inline TManagedPtr<const FInstancedStruct> LoadInstancedStruct(const FCacheAddress& From) const
+	    {
+			return LiveInstance->Cache->GetInstancedStruct(From);
 		}
 
 		inline void StoreValidDesc(const FCacheAddress& To)
 		{
-			System->WorkingMemoryManager.CurrentInstanceCache->SetValidDesc(To);
+			LiveInstance->Cache->SetValidDesc(To);
 		}
 
 		inline void StoreBool(const FCacheAddress& To, bool Value)
 		{
-			System->WorkingMemoryManager.CurrentInstanceCache->SetBool(To, Value);
+			LiveInstance->Cache->SetBool(To, Value);
 		}
 
 		inline void StoreInt(const FCacheAddress& To, int32 Value)
 		{
-			System->WorkingMemoryManager.CurrentInstanceCache->SetInt(To, Value);
+			LiveInstance->Cache->SetInt(To, Value);
 		}
 
 		inline void StoreScalar(const FCacheAddress& To, float Value)
 		{
-			System->WorkingMemoryManager.CurrentInstanceCache->SetScalar(To, Value);
+			LiveInstance->Cache->SetScalar(To, Value);
 		}
 
 		inline void StoreString(const FCacheAddress& To, TManagedPtr<const String> Value)
 		{
-			System->WorkingMemoryManager.CurrentInstanceCache->SetString(To, Value);
+			LiveInstance->Cache->SetString(To, Value);
 		}
 
 		inline void StoreColor(const FCacheAddress& To, const FVector4f& Value)
 		{
-			System->WorkingMemoryManager.CurrentInstanceCache->SetColour(To, Value);
+			LiveInstance->Cache->SetColour(To, Value);
 		}
 
     	inline void StoreMatrix(const FCacheAddress& To, const FMatrix44f& Value)
 		{
-			System->WorkingMemoryManager.CurrentInstanceCache->SetMatrix(To, Value);
+			LiveInstance->Cache->SetMatrix(To, Value);
 		}
 
 		inline void StoreMaterial(const FCacheAddress& To, TManagedPtr<const FMaterial> Value)
 		{
-			System->WorkingMemoryManager.CurrentInstanceCache->SetMaterial(To, Value);
+			LiveInstance->Cache->SetMaterial(To, Value);
 		}
     	
 		inline void StoreProjector(const FCacheAddress& To, const FProjector& Value)
 		{
-			System->WorkingMemoryManager.CurrentInstanceCache->SetProjector(To, Value);
+			LiveInstance->Cache->SetProjector(To, Value);
 		}
 
 		inline void StoreMesh(const FCacheAddress& To, const TManagedPtr<const FMesh>& Resource)
 		{
-			System->WorkingMemoryManager.StoreMesh(To, Resource);
+			LiveInstance->Cache->SetMesh(To, Resource);
 		}
 
     	inline void StoreSkeletalMesh(const FCacheAddress& To, TManagedPtr<const FSkeletalMesh> Resource)
 		{
-			System->WorkingMemoryManager.CurrentInstanceCache->SetSkeletalMesh(To, Resource);
+			LiveInstance->Cache->SetSkeletalMesh(To, Resource);
 		}
     	
 		inline void StoreImage(const FCacheAddress& To, const TManagedPtr<const FImage>& Resource)
 		{
-			System->WorkingMemoryManager.StoreImage(To, Resource);
+			LiveInstance->Cache->SetImage(To, Resource);
 		}
 
 		inline void StoreLayout(const FCacheAddress& To, const TManagedPtr<const FLayout>& Resource)
 		{
-			System->WorkingMemoryManager.CurrentInstanceCache->SetLayout(To, Resource);
+			LiveInstance->Cache->SetLayout(To, Resource);
 		}
 
 		inline void StoreInstance(const FCacheAddress& To, const TManagedPtr<const FInstance>& Resource)
 		{
-			System->WorkingMemoryManager.CurrentInstanceCache->SetInstance(To, Resource);
+			LiveInstance->Cache->SetInstance(To, Resource);
 		}
 
 		inline void StoreExtensionData(const FCacheAddress& To, const TManagedPtr<const FExtensionData>& Resource)
 		{
-			System->WorkingMemoryManager.CurrentInstanceCache->SetExtensionData(To, Resource.ToStrongRef());
+			LiveInstance->Cache->SetExtensionData(To, Resource.ToStrongRef());
 		}
 
 		inline void StoreInstancedStruct(const FCacheAddress& To, const TManagedPtr<const FInstancedStruct>& Resource)
 		{
-			System->WorkingMemoryManager.CurrentInstanceCache->SetInstancedStruct(To, Resource.ToStrongRef());
+			LiveInstance->Cache->SetInstancedStruct(To, Resource.ToStrongRef());
 		}
     	
 		inline TManagedPtr<FImage> CreateImage(uint32 SizeX, uint32 SizeY, uint32 Lods, EImageFormat Format, EInitializationType Init)
 		{
-			TManagedPtr<FImage> Result = System->WorkingMemoryManager.CreateImage(SizeX, SizeY, Lods, Format, Init);
+			TManagedPtr<FImage> Result = CreateImageImp(SizeX, SizeY, Lods, Format, Init);
 			return Result;
 		}
 
-		TManagedPtr<FImage> CreateImageLike(const FImage* Ref, EInitializationType Init)
+		TManagedPtr<FImage> CreateImageLike( const FImage* Ref, EInitializationType Init)
 		{
-			TManagedPtr<FImage> Result = System->WorkingMemoryManager.CreateImage(Ref->GetSizeX(), Ref->GetSizeY(), Ref->GetLODCount(), Ref->GetFormat(), Init);
+			TManagedPtr<FImage> Result = CreateImageImp(Ref->GetSizeX(), Ref->GetSizeY(), Ref->GetLODCount(), Ref->GetFormat(), Init);
 			return Result;
 		}
 
-		inline TManagedPtr<FImage> CloneOrTakeOver(TManagedPtr<const FImage>&& Ref)
+		inline TManagedPtr<FImage> CloneOrTakeOver(TManagedPtr<const FImage>&& Resource)
 		{
-			TManagedPtr<FImage> Result = System->WorkingMemoryManager.CloneOrTakeOver(MoveTemp(Ref));
+			System->WorkingMemoryManager.CheckRunnerThread();
+	
+			TManagedPtr<FImage> Result;
+			if (!Resource.IsUniqueReference())
+			{
+				uint32 DataSize = Resource->GetDataSize();
+				EnsureBudgetBelow(DataSize);
+				Result = Resource->Clone();
+			}
+			else
+			{
+				Result = ConstCastManagedPtr<FImage>(Resource);
+			}
+
+			Resource = nullptr;
 			return Result;
 		}
 
 		[[nodiscard]] inline TManagedPtr<FMesh> CreateMesh(int32 BudgetReserveSize = 0)
 		{
-			return System->WorkingMemoryManager.CreateMesh(BudgetReserveSize);
+			System->WorkingMemoryManager.CheckRunnerThread();
+
+			EnsureBudgetBelow(BudgetReserveSize);
+
+			TManagedPtr<FMesh> Result = MakeManaged<FMesh>();
+			return Result;
 		}
 
-		[[nodiscard]] inline TManagedPtr<FMesh> CloneOrTakeOver(TManagedPtr<const FMesh>&& Ref)
+		[[nodiscard]] inline TManagedPtr<FMesh> CloneOrTakeOver(TManagedPtr<const FMesh>&& Resource)
 		{
-			return System->WorkingMemoryManager.CloneOrTakeOver(MoveTemp(Ref));
+			System->WorkingMemoryManager.CheckRunnerThread();
+
+			TManagedPtr<FMesh> Result;
+			if (!Resource.IsUniqueReference())
+			{
+				const int32 ResourceDataSize = Resource->GetDataSize();
+				Result = CreateMesh(ResourceDataSize); 
+				Result->CopyFrom(*Resource);
+			}
+			else
+			{
+				Result = ConstCastManagedPtr<FMesh>(Resource);
+			}
+
+			Resource = nullptr;
+			return Result;
 		}
 
     	template<typename T>
 		TPassthroughObjectPtr<T> FindPassthrough(PASSTHROUGH_ID Id)
 		{
-			if (System->WorkingMemoryManager.CurrentPassthroughObjectLoader)
+			if (LiveInstance->PassthroughObjectLoader)
 			{
-				TPassthroughObjectPtr<T> PassthroughObject = System->WorkingMemoryManager.CurrentPassthroughObjectLoader->Find<T>(Id);
+				TPassthroughObjectPtr<T> PassthroughObject = LiveInstance->PassthroughObjectLoader->Find<T>(Id);
 				check(PassthroughObject); // Object not discovered in the CodeRunnerBegin.
 				
 				return PassthroughObject;
@@ -597,6 +627,30 @@ namespace  UE::Mutable::Private
     	TSharedPtr<FExternalResourceProvider> ExternalResourceProvider;
     
     private:
+    	
+    	inline TManagedPtr<FImage> CreateImageImp(uint32 SizeX, uint32 SizeY, uint32 Lods, EImageFormat Format, EInitializationType Init)
+    	{
+    		System->WorkingMemoryManager.CheckRunnerThread();
+			
+    		// Make room in the budget
+    		uint32 DataSize = FImage::CalculateDataSize(SizeX, SizeY, Lods, Format);
+    		EnsureBudgetBelow(DataSize);
+
+    		// Create it
+    		TManagedPtr<FImage> Result = MakeManaged<FImage>(SizeX, SizeY, Lods, Format, Init);
+    		return Result;
+    	}
+    	
+    	bool EnsureBudgetBelow( uint64 AdditionalMemory ) const
+	    {
+    		check(System)
+    		return System->WorkingMemoryManager.EnsureBudgetBelow(AdditionalMemory, LiveInstance->Cache);
+    	}
+    	
+    	void MarkRomUsed( int32 InRomIndex);
+
+    	FWorkingMemoryManager::FModelCacheEntry* GetModelCache() const;
+    	
     	struct FRomLoadOps
     	{
    		private:
diff --git a/Engine/Plugins/Mutable/Source/MutableRuntime/Private/MuR/CodeRunnerBegin.cpp b/Engine/Plugins/Mutable/Source/MutableRuntime/Private/MuR/CodeRunnerBegin.cpp
--- a/Engine/Plugins/Mutable/Source/MutableRuntime/Private/MuR/CodeRunnerBegin.cpp
+++ b/Engine/Plugins/Mutable/Source/MutableRuntime/Private/MuR/CodeRunnerBegin.cpp
@@ -13,27 +13,20 @@
 namespace UE::Mutable::Private
 {
 	CodeRunnerBegin::CodeRunnerBegin(
-			const TSharedPtr<const FModel>& InModel,
+			const TSharedRef<FLiveInstance>& InLiveInstance,
 			const TSharedPtr<const FParameters>& InParams,
-			const TSharedPtr<FProgramCache>& InProgramCache,
-			const TSharedPtr<FMeshIdRegistry>& InMeshIdRegistry,
-			const TSharedPtr<FImageIdRegistry>& InImageIdRegistry,
-			const TSharedPtr<FMaterialIdRegistry>& InMaterialIdRegistry,
-			const TSharedPtr<FSkeletalMeshIdRegistry>& InSkeletalMeshIdRegistry,
-			const TSharedPtr<FPassthroughObjectLoader>& InPassthroughObjectLoader,
-			const FScheduledOpInline& InRoot,
 			uint32 InLodMask) :
-		Model(InModel),
+		Model(InLiveInstance->Model),
 		Params(InParams),
-		Program(InModel->Program),
+		Program(InLiveInstance->Model->Program),
 		LodMask(InLodMask),
-		Executed(InModel->Program.OpAddress.Num()),
-		ProgramCache(InProgramCache),
-		MeshIdRegistry(InMeshIdRegistry),
-		ImageIdRegistry(InImageIdRegistry),
-		MaterialIdRegistry(InMaterialIdRegistry),
-		SkeletalMeshIdRegistry(InSkeletalMeshIdRegistry),
-		PassthroughObjectLoader(InPassthroughObjectLoader)
+		Executed(InLiveInstance->Model->Program.OpAddress.Num()),
+		ProgramCache(InLiveInstance->Cache),
+		MeshIdRegistry(InLiveInstance->MeshIdRegistry),
+		ImageIdRegistry(InLiveInstance->ImageIdRegistry),
+		MaterialIdRegistry(InLiveInstance->MaterialIdRegistry),
+		SkeletalMeshIdRegistry(InLiveInstance->SkeletalMeshIdRegistry),
+		PassthroughObjectLoader(InLiveInstance->PassthroughObjectLoader)
 	{
 	}
 
diff --git a/Engine/Plugins/Mutable/Source/MutableRuntime/Private/MuR/CodeRunnerBegin.h b/Engine/Plugins/Mutable/Source/MutableRuntime/Private/MuR/CodeRunnerBegin.h
--- a/Engine/Plugins/Mutable/Source/MutableRuntime/Private/MuR/CodeRunnerBegin.h
+++ b/Engine/Plugins/Mutable/Source/MutableRuntime/Private/MuR/CodeRunnerBegin.h
@@ -110,15 +110,8 @@ namespace UE::Mutable::Private
 	{
 	public:
 		MUTABLERUNTIME_API CodeRunnerBegin(
-			const TSharedPtr<const FModel>& InModel,
+			const TSharedRef<FLiveInstance>& InLiveInstance,
 			const TSharedPtr<const FParameters>& InParams,
-			const TSharedPtr<FProgramCache>& InProgramCache,
-			const TSharedPtr<FMeshIdRegistry>& InMeshIdRegistry,
-			const TSharedPtr<FImageIdRegistry>& InImageIdRegistry,
-			const TSharedPtr<FMaterialIdRegistry>& InMaterialIdRegistry,
-			const TSharedPtr<FSkeletalMeshIdRegistry>& InSkeletalMeshIdRegistry,
-			const TSharedPtr<FPassthroughObjectLoader>& InPassthroughObjectLoader,
-			const FScheduledOpInline& InRoot,
 			uint32 InLodMask);
 
 		TManagedPtr<const FInstance> RunCode(const FScheduledOpInline& Root);
diff --git a/Engine/Plugins/Mutable/Source/MutableRuntime/Private/MuR/CodeRunnerTasked.cpp b/Engine/Plugins/Mutable/Source/MutableRuntime/Private/MuR/CodeRunnerTasked.cpp
--- a/Engine/Plugins/Mutable/Source/MutableRuntime/Private/MuR/CodeRunnerTasked.cpp
+++ b/Engine/Plugins/Mutable/Source/MutableRuntime/Private/MuR/CodeRunnerTasked.cpp
@@ -72,7 +72,7 @@ namespace UE::Mutable::Private
 			// We need to increase hit count unconditionally for Full execution. 
 			if (Dep.Type == FScheduledOp::EType::Full)
 			{
-				System->WorkingMemoryManager.CurrentInstanceCache->IncreaseHitCount(At);
+				LiveInstance->Cache->IncreaseHitCount(At);
 			}
 
 			if (!Dep.At)
@@ -255,8 +255,8 @@ namespace UE::Mutable::Private
 				{
 				//case EExecutionStrategy::MinimizeMemory:
 				//{
-				//	// TODO: Prioritize operation strages that have a negative memory delta somehow (the result uses less memory than the inputs)
-				//  // An attempt to do this with an heurisitc estimation per-stage of each operation in a static table was tested in the past but maybe
+				//	// TODO: Prioritize operation stages that have a negative memory delta somehow (the result uses less memory than the inputs)
+				//  // An attempt to do this with an heuristic estimation per-stage of each operation in a static table was tested in the past but maybe
 				//  // there are better ways. 
 				//	break;
 				//}
@@ -564,7 +564,7 @@ namespace UE::Mutable::Private
 		if (IsCompressedFormat(InitialFormat))
 		{
 			EImageFormat UncompressedFormat = GetUncompressedFormat(InitialFormat);
-			TManagedPtr<FImage> Formatted = Runner->CreateImage( Base->GetSizeX(), Base->GetSizeY(), Base->GetLODCount(), UncompressedFormat, EInitializationType::NotInitialized );
+			TManagedPtr<FImage> Formatted = Runner->CreateImage(Base->GetSizeX(), Base->GetSizeY(), Base->GetLODCount(), UncompressedFormat, EInitializationType::NotInitialized );
 			bool bSuccess = false;
 			ImOp.ImagePixelFormat(bSuccess, ImageCompressionQuality, Formatted.Get(), Base.Get());
 			check(bSuccess); // Decompression cannot fail
@@ -1970,7 +1970,7 @@ namespace UE::Mutable::Private
 		bOutFailed = false;
 
 		const FProgram& Program = Runner->Model->Program;
-		FWorkingMemoryManager::FModelCacheEntry* ModelCache = Runner->System->WorkingMemoryManager.FindModelCache(Runner->Model.Get());
+		FWorkingMemoryManager::FModelCacheEntry* ModelCache = Runner->GetModelCache();
 		
 		TArray<UE::Tasks::FTask, TInlineAllocator<4>> ReadCompleteEvents;
 		ReadCompleteEvents.Reserve(4); 
@@ -2058,8 +2058,8 @@ namespace UE::Mutable::Private
 
 			// Free roms if necessary
 			{
-				Runner->System->WorkingMemoryManager.MarkRomUsed(RomIndex, Runner->Model);
-				Runner->System->WorkingMemoryManager.EnsureBudgetBelow(RomSize);
+				Runner->MarkRomUsed(RomIndex);
+				Runner->EnsureBudgetBelow(RomSize);
 			}
 
 			const int32 SizeBefore = RomLoadOp.m_streamBuffer.GetAllocatedSize();
@@ -2108,7 +2108,7 @@ namespace UE::Mutable::Private
 		}
 
 		const FProgram& Program = Runner->Model->Program;
-		FWorkingMemoryManager::FModelCacheEntry* ModelCache = Runner->System->WorkingMemoryManager.FindModelCache(Runner->Model.Get());
+		FWorkingMemoryManager::FModelCacheEntry* ModelCache = Runner->GetModelCache();
 		
 		bool bSomeMissingData = false;
 		for (const int32 RomIndex : RomIndices)
@@ -2153,7 +2153,7 @@ namespace UE::Mutable::Private
 			{
 				check(RomIndex < Program.Roms.Num());
 
-				Runner->System->WorkingMemoryManager.MarkRomUsed(RomIndex, Runner->Model);
+				Runner->MarkRomUsed(RomIndex);
 				--ModelCache->PendingOpsPerRom[RomIndex];
 
 				if (DebugRom && (DebugRomAll || RomIndex == DebugRomIndex))
@@ -2182,7 +2182,7 @@ namespace UE::Mutable::Private
 
 		const FProgram& Program = Runner->Model->Program;
 
-		FWorkingMemoryManager::FModelCacheEntry* ModelCache = Runner->System->WorkingMemoryManager.FindModelCache(Runner->Model.Get());
+		FWorkingMemoryManager::FModelCacheEntry* ModelCache = Runner->GetModelCache();
 
 		TArray<UE::Tasks::FTask> ReadCompleteEvents;
 
@@ -2228,8 +2228,8 @@ namespace UE::Mutable::Private
 
 			// Free roms if necessary
 			{
-				Runner->System->WorkingMemoryManager.MarkRomUsed(RomIndex, Runner->Model);
-				Runner->System->WorkingMemoryManager.EnsureBudgetBelow(RomSize);
+				Runner->MarkRomUsed(RomIndex);
+				Runner->EnsureBudgetBelow(RomSize);
 			}
 
 			const int32 SizeBefore = RomLoadOp.m_streamBuffer.GetAllocatedSize();
@@ -2278,7 +2278,7 @@ namespace UE::Mutable::Private
 
 		const FProgram& Program = Runner->Model->Program;
 		
-		FWorkingMemoryManager::FModelCacheEntry* ModelCache = Runner->System->WorkingMemoryManager.FindModelCache(Runner->Model.Get());
+		FWorkingMemoryManager::FModelCacheEntry* ModelCache = Runner->GetModelCache();
 
 		bool bSomeMissingData = false;
 
@@ -2330,10 +2330,10 @@ namespace UE::Mutable::Private
 					continue;
 				}
 
-				int32 RomIndex = ImageConstantResourceIndex.Index;
+				const int32 RomIndex = ImageConstantResourceIndex.Index;
 				check(RomIndex < Program.Roms.Num());
 
-				Runner->System->WorkingMemoryManager.MarkRomUsed(RomIndex, Runner->Model);
+				Runner->MarkRomUsed(RomIndex);
 				--ModelCache->PendingOpsPerRom[RomIndex];
 
 				if (DebugRom && (DebugRomAll || RomIndex == DebugRomIndex))
@@ -2761,4 +2761,17 @@ namespace UE::Mutable::Private
 		return Issued;
 	}
 
+	void CodeRunner::MarkRomUsed(int32 InRomIndex)
+	{
+		check(Model)
+		System->WorkingMemoryManager.MarkRomUsed(InRomIndex, Model);
+	}
+
+	
+	FWorkingMemoryManager::FModelCacheEntry* CodeRunner::GetModelCache() const
+	{
+		check(Model)
+		return System->WorkingMemoryManager.FindModelCache(Model.Get());
+	}
+
 } // namespace UE::Mutable::Private
diff --git a/Engine/Plugins/Mutable/Source/MutableRuntime/Private/MuR/System.cpp b/Engine/Plugins/Mutable/Source/MutableRuntime/Private/MuR/System.cpp
--- a/Engine/Plugins/Mutable/Source/MutableRuntime/Private/MuR/System.cpp
+++ b/Engine/Plugins/Mutable/Source/MutableRuntime/Private/MuR/System.cpp
@@ -6,12 +6,10 @@
 #include "BitWriterResourceKey.h"
 #include "Containers/Array.h"
 #include "Containers/Map.h"
-#include "Containers/Set.h"
 #include "Serialization/BitWriter.h"
 #include "HAL/IConsoleManager.h"
 #include "HAL/LowLevelMemTracker.h"
 #include "HAL/PlatformCrt.h"
-#include "HAL/PlatformMath.h"
 #include "Logging/LogCategory.h"
 #include "Logging/LogMacros.h"
 #include "Misc/AssertionMacros.h"
@@ -24,18 +22,14 @@
 #include "MuR/MutableTrace.h"
 #include "MuR/Operations.h"
 #include "MuR/Parameters.h"
-#include "MuR/Parameters.h"
 #include "MuR/Platform.h"
 #include "MuR/Serialisation.h"
-#include "MuR/System.h"
 #include "MuR/MutableRuntimeModule.h"
 #include "Templates/SharedPointer.h"
 #include "Templates/Tuple.h"
-#include "Trace/Detail/Channel.h"
 #include "ProfilingDebugging/CountersTrace.h"
 #include "PackedNormal.h"
 #include "Engine/SkeletalMesh.h"
-#include "MuR/System.h"
 
 #include UE_INLINE_GENERATED_CPP_BY_NAME(System)
 
@@ -110,7 +104,7 @@ namespace UE::Mutable::Private
 		MUTABLE_CPUPROFILER_SCOPE(SetWorkingMemoryBytes);
 
 		WorkingMemoryManager.BudgetBytes = InBytes;
-		WorkingMemoryManager.EnsureBudgetBelow(0);
+		WorkingMemoryManager.EnsureBudgetBelow(0, nullptr);
 
 		UpdateStats();
 	}
@@ -125,6 +119,10 @@ namespace UE::Mutable::Private
 
    
     TSharedRef<FLiveInstance> FSystem::NewInstance( const TSharedPtr<const FModel>& InModel,
+		const TSharedPtr<UE::Mutable::Private::FMeshIdRegistry>& InMeshIdRegistry,
+		const TSharedPtr<UE::Mutable::Private::FImageIdRegistry>& InImageIdRegistry,
+		const TSharedPtr<UE::Mutable::Private::FMaterialIdRegistry>& InMaterialIdRegistry,
+		const TSharedPtr<UE::Mutable::Private::FSkeletalMeshIdRegistry>& InSkeletalMeshIdRegistry,
     	const TSharedPtr<FExternalResourceProvider>& ExternalResourceProvider,
     	const TSharedPtr<FPassthroughObjectLoader>& PassthroughObjectLoader)
     {
@@ -132,14 +130,23 @@ namespace UE::Mutable::Private
 		MUTABLE_CPUPROFILER_SCOPE(NewInstance);
 
     	UE_LOG(LogMutableCore, Verbose, TEXT("Creating Mutable instance"));
-    	
+    	check(InModel)
+
 		TSharedRef<FLiveInstance> InstanceData = MakeShared<FLiveInstance>();
 		InstanceData->Instance = nullptr;
 		InstanceData->Model = InModel;
 		InstanceData->State = -1;
-		InstanceData->Cache = MakeShared<FProgramCache>();
     	InstanceData->ExternalResourceProvider = ExternalResourceProvider;
     	InstanceData->PassthroughObjectLoader = PassthroughObjectLoader;
+    	InstanceData->MeshIdRegistry = InMeshIdRegistry;
+    	InstanceData->ImageIdRegistry = InImageIdRegistry;
+    	InstanceData->MaterialIdRegistry = InMaterialIdRegistry;
+    	InstanceData->SkeletalMeshIdRegistry = InSkeletalMeshIdRegistry;
+    	
+    	check(InstanceData->MeshIdRegistry.IsValid())
+		check(InstanceData->ImageIdRegistry.IsValid())
+		check(InstanceData->MaterialIdRegistry.IsValid())
+		check(InstanceData->SkeletalMeshIdRegistry.IsValid())
     	
     	WorkingMemoryManager.LiveInstances.Add(InstanceData.ToWeakPtr());
     	WorkingMemoryManager.LiveInstances.RemoveAllSwap([](const TWeakPtr<FLiveInstance>& Instance)
@@ -149,15 +156,36 @@ namespace UE::Mutable::Private
     	
 		return InstanceData;
 	}
+
 	
+    TSharedRef<FLiveInstance> FSystem::NewBuildInstance(const TSharedPtr<const FModel>& InModel,
+	    const TSharedPtr<FExternalResourceProvider>& ExternalResourceProvider)
+    {
+    	MUTABLE_CPUPROFILER_SCOPE(NewBuildInstance);
+
+    	UE_LOG(LogMutableCore, Verbose, TEXT("Creating Mutable instance"));
+    	check(InModel)
+    	
+    	TSharedRef<FLiveInstance> InstanceData = MakeShared<FLiveInstance>();
+    	InstanceData->Instance = nullptr;
+    	InstanceData->Model = InModel;
+    	InstanceData->State = -1;
+    	InstanceData->ExternalResourceProvider = ExternalResourceProvider;
+    	InstanceData->PassthroughObjectLoader = nullptr;
+    	InstanceData->MeshIdRegistry = MakeShared<FMeshIdRegistry>();
+    	InstanceData->ImageIdRegistry = MakeShared<FImageIdRegistry>();
+    	InstanceData->MaterialIdRegistry = MakeShared<FMaterialIdRegistry>();
+    	InstanceData->SkeletalMeshIdRegistry = MakeShared<FSkeletalMeshIdRegistry>();
+
+    	InstanceData->Cache->Init(InModel->Program.OpAddress.Num());
+    	
+    	return InstanceData;
+    }
+
 
     //---------------------------------------------------------------------------------------------
     TManagedPtr<const FInstance> FSystem::BeginUpdate_GameThread(const TSharedRef<FLiveInstance>& LiveInstance,
                                      const TSharedPtr<const FParameters>& InParams,
-                                     const TSharedPtr<FMeshIdRegistry>& InMeshIdRegistry,
-									 const TSharedPtr<FImageIdRegistry>& InImageIdRegistry,
-									 const TSharedPtr<FMaterialIdRegistry>& InMaterialIdRegistry,
-									 const TSharedPtr<FSkeletalMeshIdRegistry>& InSkeletalMeshIdRegistry,
 									 int32 InStateIndex,
                                      uint32 InLodMask)
     {
@@ -166,19 +194,9 @@ namespace UE::Mutable::Private
 
 		check(InParams);
 
-    	FScheduledOpInline Root(LiveInstance->Model->Program.States[InStateIndex].Root, 0, 0, true);
-    	CodeRunnerBegin Runner(
-    		LiveInstance->Model,
-    		InParams,
-    		LiveInstance->Cache,
-    		InMeshIdRegistry,
-    		InImageIdRegistry,
-    		InMaterialIdRegistry,
-    		InSkeletalMeshIdRegistry,
-    		LiveInstance->PassthroughObjectLoader,
-    		Root,
-    		InLodMask);
+    	CodeRunnerBegin Runner(LiveInstance, InParams, InLodMask);
 
+    	FScheduledOpInline Root(LiveInstance->Model->Program.States[InStateIndex].Root, 0, 0, true);
     	TManagedPtr<const FInstance> Result = Runner.RunCode(FScheduledOpInline(Root, 0));
     	if (!Result)
     	{
@@ -192,10 +210,6 @@ namespace UE::Mutable::Private
 	
     TManagedPtr<const FInstance> FSystem::BeginUpdate_MutableThread(const TSharedRef<FLiveInstance>& LiveInstance,
 		const TSharedPtr<const FParameters>& InParams,
-		const TSharedPtr<FMeshIdRegistry>& InMeshIdRegistry,
-		const TSharedPtr<FImageIdRegistry>& InImageIdRegistry,
-		const TSharedPtr<FMaterialIdRegistry>& InMaterialIdRegistry,
-		const TSharedPtr<FSkeletalMeshIdRegistry>& InSkeletalMeshIdRegistry,
 		int32 InStateIndex,
 		uint32 InLodMask)
     {
@@ -208,18 +222,10 @@ namespace UE::Mutable::Private
 			UE_LOG(LogMutableCore, Error, TEXT("Invalid parameters in mutable update."));
 			return nullptr;
 		}
-    	
-		WorkingMemoryManager.CurrentInstanceCache = LiveInstance->Cache;
-    	WorkingMemoryManager.CurrentPassthroughObjectLoader = LiveInstance->PassthroughObjectLoader;
-
-    	WorkingMemoryManager.CurrentMeshIdRegistry = InMeshIdRegistry;
-    	WorkingMemoryManager.CurrentImageIdRegistry = InImageIdRegistry;
-    	WorkingMemoryManager.CurrentMaterialIdRegistry = InMaterialIdRegistry;
-    	WorkingMemoryManager.CurrentSkeletalMeshIdRegistry = InSkeletalMeshIdRegistry;
 
 		const FProgram& Program = LiveInstance->Model->Program;
 
-		bool bValidState = InStateIndex >= 0 && InStateIndex < (int)Program.States.Num();
+		bool bValidState = InStateIndex >= 0 && InStateIndex < (int32)Program.States.Num();
 		if (!bValidState)
 		{
 			UE_LOG(LogMutableCore, Error, TEXT("Invalid state in mutable update."));
@@ -240,24 +246,23 @@ namespace UE::Mutable::Private
 		}
 
 		// Remove cached data
-		WorkingMemoryManager.ClearCacheLayer0();
+		WorkingMemoryManager.ClearCacheLayer0(LiveInstance->Cache);
 		if (bFullBuild)
 		{
-			WorkingMemoryManager.ClearCacheLayer1();
+			WorkingMemoryManager.ClearCacheLayer1(LiveInstance->Cache);
 		}
 
 		WorkingMemoryManager.BeginRunnerThread();
-
-		OP::ADDRESS RootAt = LiveInstance->Model->Program.States[InStateIndex].Root;
-
+    	
 		// Prepare instance cache
-		PrepareCache(LiveInstance->Model.Get(), InStateIndex);
+    	LiveInstance->Cache->Prepare(Program, InStateIndex);
 		LiveInstance->OldParameters = InParams->Clone();
 
 		// Ensure the model cache has been created
 		WorkingMemoryManager.FindOrAddModelCache(LiveInstance->Model);
 
-		RunCode(LiveInstance->ExternalResourceProvider, LiveInstance->Model, InParams.Get(), RootAt, InLodMask);
+		OP::ADDRESS RootAt = Program.States[InStateIndex].Root;
+		RunCode(LiveInstance, InParams.Get(), RootAt, InLodMask);
 
 		TManagedPtr<const FInstance> Result = LiveInstance->Cache->GetInstance(FCacheAddress(RootAt, 0, 0));
 
@@ -272,16 +277,8 @@ namespace UE::Mutable::Private
 			Result = MakeManaged<FInstance>();
 		}
 
-		WorkingMemoryManager.EndRunnerThread();
-
-		WorkingMemoryManager.CurrentInstanceCache = nullptr;
-    	WorkingMemoryManager.CurrentPassthroughObjectLoader = nullptr;
-
-    	WorkingMemoryManager.CurrentMeshIdRegistry = nullptr;
-    	WorkingMemoryManager.CurrentImageIdRegistry = nullptr;
-    	WorkingMemoryManager.CurrentMaterialIdRegistry = nullptr;
-    	WorkingMemoryManager.CurrentSkeletalMeshIdRegistry = nullptr;
-
+		WorkingMemoryManager.EndRunnerThread(LiveInstance);
+    	
 		return Result;
 	}
 
@@ -301,21 +298,14 @@ namespace UE::Mutable::Private
     	{
 			return BlackImage();
     	}
-    	
-		// Find the live instance
-		WorkingMemoryManager.CurrentInstanceCache = LiveInstance->Cache;
-		WorkingMemoryManager.CurrentPassthroughObjectLoader = LiveInstance->PassthroughObjectLoader;
 
 		OP::ADDRESS RootAddress = ImageId.GetKey()->Address;
-		TManagedPtr<const FImage> Result = BuildImage(LiveInstance->ExternalResourceProvider, LiveInstance->Model, LiveInstance->OldParameters.Get(), RootAddress, MipsToSkip);
+		TManagedPtr<const FImage> Result = BuildImage(LiveInstance, LiveInstance->OldParameters.Get(), RootAddress, MipsToSkip);
 
 		if (!Result)
 		{
 			Result = BlackImage();
 		}
-
-		WorkingMemoryManager.CurrentInstanceCache = nullptr;
-    	WorkingMemoryManager.CurrentPassthroughObjectLoader = nullptr;
     	
 		return Result;
 	}
@@ -361,21 +351,18 @@ namespace UE::Mutable::Private
     	UE::Tasks::FTask RunnerTask = UE::Tasks::Launch(TEXT("FSystem::GetImage"),
 		[Context, System = this, RootAddress, LiveInstance, MipsToSkip]()
 		{
-			System->WorkingMemoryManager.CurrentInstanceCache = LiveInstance->Cache;
-			System->WorkingMemoryManager.CurrentPassthroughObjectLoader = LiveInstance->PassthroughObjectLoader;
 			System->WorkingMemoryManager.BeginRunnerThread();
 
 			Context->Runner = CodeRunner::Create(
+				LiveInstance,
 				System->Settings,
 				System,
 				EExecutionStrategy::MinimizeMemory,
-				LiveInstance->Model,
 				LiveInstance->OldParameters.Get(),
 				RootAddress, FSystem::AllLODs,
 				MipsToSkip,
 				false,
-				FScheduledOp::EType::Full,
-				LiveInstance->ExternalResourceProvider);
+				FScheduledOp::EType::Full);
 				
 			constexpr bool bForceInlineExecution = false;
 			Context->Runner->StartRun(bForceInlineExecution);
@@ -384,14 +371,14 @@ namespace UE::Mutable::Private
 		UE::Tasks::ETaskPriority::Inherit);
     	
 		return UE::Tasks::Launch(TEXT("FSystem::GetImageResultTask"),
-		[Context, System = this, RootAddress, MipsToSkip]() -> TManagedPtr<const FImage>
+		[Context, System = this, RootAddress, MipsToSkip, LiveInstance]() -> TManagedPtr<const FImage>
 		{
 			TManagedPtr<const FImage> Result;
 
 			System->bUnrecoverableError = Context->Runner->bUnrecoverableError;
 			if (!Context->Runner->bUnrecoverableError)
 			{
-				Result = System->WorkingMemoryManager.LoadImage(FCacheAddress(RootAddress, 0, MipsToSkip));
+				Result = Context->Runner->LoadImage(FCacheAddress(RootAddress, 0, MipsToSkip));
 			}
 
 			if (!Result)
@@ -399,9 +386,7 @@ namespace UE::Mutable::Private
 				Result = MakeManaged<FImage>(16, 16, 1, EImageFormat::RGBA_UByte, EInitializationType::Black);
 			}
 
-			System->WorkingMemoryManager.EndRunnerThread();
-			System->WorkingMemoryManager.CurrentInstanceCache = nullptr;
-			System->WorkingMemoryManager.CurrentPassthroughObjectLoader = nullptr;
+			System->WorkingMemoryManager.EndRunnerThread(LiveInstance);
 			
 			return Result;
 		},
@@ -432,10 +417,6 @@ namespace UE::Mutable::Private
 		
 		FExtendedImageDesc Result;
 
-		// Find the live instance
-		WorkingMemoryManager.CurrentInstanceCache = LiveInstance->Cache;
-		WorkingMemoryManager.CurrentPassthroughObjectLoader = LiveInstance->PassthroughObjectLoader;
-
 		OP::ADDRESS RootAddress = ImageId.GetKey()->Address;
 
 		const UE::Mutable::Private::FModel* Model = LiveInstance->Model.Get();
@@ -445,7 +426,7 @@ namespace UE::Mutable::Private
 		int32 VarValue = CVarClearImageDescCache.GetValueOnAnyThread();
 		if (VarValue != 0)
 		{
-			WorkingMemoryManager.CurrentInstanceCache->ClearDescCache();
+			LiveInstance->Cache->ClearDescCache();
 		}
 
 		UE::Mutable::Private::EOpType OpType = Program.GetOpType(RootAddress);
@@ -456,29 +437,25 @@ namespace UE::Mutable::Private
 					
 			int8 executionOptions = 0;
 			TSharedRef<CodeRunner> Runner = CodeRunner::Create(
+					LiveInstance,
 					Settings,
 					this,
 					EExecutionStrategy::MinimizeMemory,
-					LiveInstance->Model,
 					LiveInstance->OldParameters.Get(),
 					RootAddress,
 					FSystem::AllLODs,
 					executionOptions,
 					false,
-					FScheduledOp::EType::ImageDesc,
-					LiveInstance->ExternalResourceProvider);
+					FScheduledOp::EType::ImageDesc);
 			
 			constexpr bool bForceInlineExecution = true;
 			Runner->StartRun(bForceInlineExecution);
 
 			Result = Runner->GetImageDescResult(RootAddress);
 
-			WorkingMemoryManager.EndRunnerThread();
+			WorkingMemoryManager.EndRunnerThread(LiveInstance);
 		}
-
-		WorkingMemoryManager.CurrentInstanceCache = nullptr;
-		WorkingMemoryManager.CurrentPassthroughObjectLoader = nullptr;
-
+		
 		return Result;
 	}
 
@@ -524,14 +501,12 @@ namespace UE::Mutable::Private
 		UE::Tasks::FTask RunnerTask = UE::Tasks::Launch(TEXT("FSystem::GetImageDesc"),
 		[Context, System = this, LiveInstance, RootAddress]()
 		{
-			System->WorkingMemoryManager.CurrentInstanceCache = LiveInstance->Cache;
-			System->WorkingMemoryManager.CurrentPassthroughObjectLoader = LiveInstance->PassthroughObjectLoader;
 				
 			// TODO: It should be possible to reuse this data if cleared in the correct places only, together with HeapImageDesc.
 			int32 VarValue = CVarClearImageDescCache.GetValueOnAnyThread();
 			if (VarValue != 0)
 			{
-				System->WorkingMemoryManager.CurrentInstanceCache->ClearDescCache();
+				LiveInstance->Cache->ClearDescCache();
 			}
 
 			System->WorkingMemoryManager.BeginRunnerThread();
@@ -539,17 +514,16 @@ namespace UE::Mutable::Private
 			// GetImageDesc may call normal execution paths where meshes are computed.
 			int8 ExecutionOptions = 0;
 			Context->Runner = CodeRunner::Create(
+			LiveInstance,
 				System->Settings,
 				System,
 				EExecutionStrategy::MinimizeMemory,
-				LiveInstance->Model,
 				LiveInstance->OldParameters.Get(),
 				RootAddress,
 				FSystem::AllLODs,
 				ExecutionOptions,
 				false,
-				FScheduledOp::EType::ImageDesc,
-				LiveInstance->ExternalResourceProvider);
+				FScheduledOp::EType::ImageDesc);
 
 			constexpr bool bForceInlineExecution = false;
 			Context->Runner->StartRun(bForceInlineExecution);
@@ -558,14 +532,11 @@ namespace UE::Mutable::Private
 		UE::Tasks::ETaskPriority::Inherit);
 		
 		return UE::Tasks::Launch(TEXT("FSystem::GetImageDescResultTask"),
-		[Context, System = this, RootAddress]() -> FExtendedImageDesc
+		[Context, System = this, RootAddress, LiveInstance]() -> FExtendedImageDesc
 		{
 			FExtendedImageDesc Result = Context->Runner->GetImageDescResult(RootAddress);
 
-			System->WorkingMemoryManager.EndRunnerThread();
-			System->WorkingMemoryManager.CurrentInstanceCache = nullptr;
-			System->WorkingMemoryManager.CurrentPassthroughObjectLoader = nullptr;
-
+			System->WorkingMemoryManager.EndRunnerThread(LiveInstance);
 			return Result;
 		},
 		UE::Tasks::Prerequisites(RunnerTask),
@@ -585,13 +556,9 @@ namespace UE::Mutable::Private
 			return MakeManaged<FMesh>();
 		}
 
-		WorkingMemoryManager.CurrentInstanceCache = LiveInstance->Cache;
-		WorkingMemoryManager.CurrentPassthroughObjectLoader = LiveInstance->PassthroughObjectLoader;
-
 		OP::ADDRESS RootAddress = MeshId.GetKey()->Address;
 		TManagedPtr<const FMesh> Result = BuildMesh(
-			LiveInstance->ExternalResourceProvider,
-			LiveInstance->Model,
+			LiveInstance,
 			LiveInstance->OldParameters.Get(),
 			RootAddress,
 			MeshContentFilter);
@@ -601,9 +568,6 @@ namespace UE::Mutable::Private
 		{
 			Result = MakeManaged<FMesh>();
 		}
-
-		WorkingMemoryManager.CurrentInstanceCache = nullptr;
-		WorkingMemoryManager.CurrentPassthroughObjectLoader = nullptr;
 		
 		return Result;
 	}
@@ -650,23 +614,20 @@ namespace UE::Mutable::Private
 		UE::Tasks::FTask RunnerTask = UE::Tasks::Launch(TEXT("FSystem::GetMesh"),
 		[Context, System = this, RootAddress, LiveInstance, MeshContentFilter]()
 		{
-			System->WorkingMemoryManager.CurrentInstanceCache = LiveInstance->Cache;
-			System->WorkingMemoryManager.CurrentPassthroughObjectLoader = LiveInstance->PassthroughObjectLoader;
 			System->WorkingMemoryManager.BeginRunnerThread();
 			
 			uint8 ExecutionOptions = static_cast<uint8>(MeshContentFilter);
 			Context->Runner = CodeRunner::Create(
+			LiveInstance,
 					System->Settings,
 					System, 
 					EExecutionStrategy::MinimizeMemory, 
-					LiveInstance->Model, 
 					LiveInstance->OldParameters.Get(), 
 					RootAddress, 
 					FSystem::AllLODs, 
 					ExecutionOptions, 
 					false, 
-					FScheduledOp::EType::Full,
-					LiveInstance->ExternalResourceProvider);
+					FScheduledOp::EType::Full);
 	
 			constexpr bool bForceInlineExecution = false;
 			Context->Runner->StartRun(bForceInlineExecution);
@@ -675,19 +636,17 @@ namespace UE::Mutable::Private
 		UE::Tasks::ETaskPriority::Inherit);
 		
 		return UE::Tasks::Launch(TEXT("FSystem::GetMeshResultTask"),
-		[Context, System = this, RootAddress, MeshContentFilter]() -> TManagedPtr<const FMesh>
+		[Context, System = this, RootAddress, MeshContentFilter, LiveInstance]() -> TManagedPtr<const FMesh>
 		{
 			TManagedPtr<const FMesh> Result;
 
 			System->bUnrecoverableError = Context->Runner->bUnrecoverableError;
 			if (!Context->Runner->bUnrecoverableError)
 			{
-				Result = System->WorkingMemoryManager.LoadMesh(FCacheAddress(RootAddress, 0, static_cast<uint8>(MeshContentFilter)));
+				Result = Context->Runner->LoadMesh(FCacheAddress(RootAddress, 0, static_cast<uint8>(MeshContentFilter)));
 			}
 			
-			System->WorkingMemoryManager.EndRunnerThread();
-			System->WorkingMemoryManager.CurrentInstanceCache = nullptr;
-			System->WorkingMemoryManager.CurrentPassthroughObjectLoader = nullptr;
+			System->WorkingMemoryManager.EndRunnerThread(LiveInstance);
 			
 			return Result;
 		},
@@ -738,22 +697,19 @@ namespace UE::Mutable::Private
 		UE::Tasks::FTask RunnerTask = UE::Tasks::Launch(TEXT("FSystem::GetSkeletalMesh"),
 		[Context, System = this, RootAddress, LiveInstance, SkeletalMeshContentFilter]()
 		{
-			System->WorkingMemoryManager.CurrentInstanceCache = LiveInstance->Cache;
-			System->WorkingMemoryManager.CurrentPassthroughObjectLoader = LiveInstance->PassthroughObjectLoader;
 			System->WorkingMemoryManager.BeginRunnerThread();
 
 			Context->Runner = CodeRunner::Create(
+			LiveInstance,
 					System->Settings,
 					System, 
 					EExecutionStrategy::MinimizeMemory, 
-					LiveInstance->Model, 
 					LiveInstance->OldParameters.Get(), 
 					RootAddress, 
 					FSystem::AllLODs, 
 					static_cast<uint8>(SkeletalMeshContentFilter), 
 					false, 
-					FScheduledOp::EType::Full,
-					LiveInstance->ExternalResourceProvider);
+					FScheduledOp::EType::Full);
 
 			constexpr bool bForceInlineExecution = false;
 			Context->Runner->StartRun(bForceInlineExecution);
@@ -762,19 +718,17 @@ namespace UE::Mutable::Private
 		UE::Tasks::ETaskPriority::Inherit);
 		
 		return UE::Tasks::Launch(TEXT("FSystem::GetSkeletalMeshResultTask"),
-		[Context, System = this, RootAddress, SkeletalMeshContentFilter]() -> TManagedPtr<const FSkeletalMesh>
+		[Context, System = this, RootAddress, SkeletalMeshContentFilter, LiveInstance]() -> TManagedPtr<const FSkeletalMesh>
 		{
 			TManagedPtr<const FSkeletalMesh> Result;
 
 			System->bUnrecoverableError = Context->Runner->bUnrecoverableError;
 			if (!Context->Runner->bUnrecoverableError)
 			{
-				Result = System->WorkingMemoryManager.LoadSkeletalMesh(FCacheAddress(RootAddress, 0, static_cast<uint8>(SkeletalMeshContentFilter)));
+				Result = Context->Runner->LoadSkeletalMesh(FCacheAddress(RootAddress, 0, static_cast<uint8>(SkeletalMeshContentFilter)));
 			}
 			
-			System->WorkingMemoryManager.EndRunnerThread();
-			System->WorkingMemoryManager.CurrentInstanceCache = nullptr;
-			System->WorkingMemoryManager.CurrentPassthroughObjectLoader = nullptr;
+			System->WorkingMemoryManager.EndRunnerThread(LiveInstance);
 			
 			return Result;
 		},
@@ -825,22 +779,19 @@ namespace UE::Mutable::Private
 		UE::Tasks::FTask RunnerTask = UE::Tasks::Launch(TEXT("FSystem::GetMaterial"),
 		[Context, System = this, RootAddress, LiveInstance]()
 		{
-			System->WorkingMemoryManager.CurrentInstanceCache = LiveInstance->Cache;
-			System->WorkingMemoryManager.CurrentPassthroughObjectLoader = LiveInstance->PassthroughObjectLoader;
 			System->WorkingMemoryManager.BeginRunnerThread();
 				
 			Context->Runner = CodeRunner::Create(
+			LiveInstance,
 					System->Settings,
 					System, 
 					EExecutionStrategy::MinimizeMemory, 
-					LiveInstance->Model, 
 					LiveInstance->OldParameters.Get(), 
 					RootAddress, 
 					FSystem::AllLODs, 
 					0,
 					true, 
-					FScheduledOp::EType::Full,
-					LiveInstance->ExternalResourceProvider);
+					FScheduledOp::EType::Full);
 
 			constexpr bool bForceInlineExecution = false;
 			Context->Runner->StartRun(bForceInlineExecution);
@@ -849,19 +800,17 @@ namespace UE::Mutable::Private
 		UE::Tasks::ETaskPriority::Inherit);
 		
 		return UE::Tasks::Launch(TEXT("FSystem::GetMaterialResultTask"),
-		[Context, System = this, RootAddress]() -> TManagedPtr<const FMaterial>
+		[Context, System = this, RootAddress, LiveInstance]() -> TManagedPtr<const FMaterial>
 		{
 			TManagedPtr<const FMaterial> Result;
 
 			System->bUnrecoverableError = Context->Runner->bUnrecoverableError;
 			if (!Context->Runner->bUnrecoverableError)
 			{
-				Result = System->WorkingMemoryManager.LoadMaterial(FCacheAddress(RootAddress, 0, 0));
+				Result = Context->Runner->LoadMaterial(FCacheAddress(RootAddress, 0, 0));
 			}
 			
-			System->WorkingMemoryManager.EndRunnerThread();
-			System->WorkingMemoryManager.CurrentInstanceCache = nullptr;
-			System->WorkingMemoryManager.CurrentPassthroughObjectLoader = nullptr;
+			System->WorkingMemoryManager.EndRunnerThread(LiveInstance);
 			
 			return Result;
 		},
@@ -881,22 +830,17 @@ namespace UE::Mutable::Private
 		// Debug check to see if we managed the op-hit-counts correctly
 		LiveInstance->Cache->CheckHitCountsCleared();
 
-		WorkingMemoryManager.CurrentInstanceCache = LiveInstance->Cache;
-		WorkingMemoryManager.CurrentPassthroughObjectLoader = LiveInstance->PassthroughObjectLoader;
-
 		// We don't want to clear the cache layer 1 because it contains data that can be useful for a 
 		// future update (same states, just runtime parameters changed).
 		//WorkingMemoryManager.ClearCacheLayer1();
 
 		// We need to clear the layer 0 cache, because it contains data that is only valid for the current 
 		// parameter values (unless it is data marked as state cache)
-		WorkingMemoryManager.ClearCacheLayer0();
-
-		WorkingMemoryManager.CurrentInstanceCache = nullptr;
-		WorkingMemoryManager.CurrentPassthroughObjectLoader = nullptr;
+		WorkingMemoryManager.ClearCacheLayer0(LiveInstance->Cache);
 		
 		// Reduce the cache until it fits the limit.
-		WorkingMemoryManager.EnsureBudgetBelow(0);
+		// Not providing a cache at this point as there is no need of keeping the cache active anymore
+		WorkingMemoryManager.EnsureBudgetBelow(0, nullptr);
 
 		if (bClearRoms)
 		{
@@ -925,13 +869,13 @@ namespace UE::Mutable::Private
 
         RelevantParameterVisitor(
                 FSystem* InSystem,
-				const TSharedPtr<const FModel>& InModel,
+				const TSharedRef<FLiveInstance>& InLiveInstance,
                 const TSharedPtr<const FParameters>& InParams,
                 TArray<uint32>& OutRelevantParameters) :
-			UniqueDiscreteCoveredCodeVisitor<>(InSystem, InModel, InParams, FSystem::AllLODs),
-		    RelevantParameters(OutRelevantParameters)
+				UniqueDiscreteCoveredCodeVisitor<>(InSystem, InLiveInstance, InParams, FSystem::AllLODs),
+		    	RelevantParameters(OutRelevantParameters)
         {
-            OP::ADDRESS At = InModel->Program.States[0].Root;
+			const OP::ADDRESS At = InLiveInstance->Model->Program.States[0].Root;
 
             Run(At);
         }
@@ -973,15 +917,8 @@ namespace UE::Mutable::Private
         TArray<uint32>& OutRelevantParameters)
     {
 		LLM_SCOPE_BYNAME(TEXT("MutableRuntime"));
-
-		// Find the live instance
-		WorkingMemoryManager.CurrentInstanceCache = LiveInstance->Cache;
-		WorkingMemoryManager.CurrentPassthroughObjectLoader = LiveInstance->PassthroughObjectLoader;
-		
-		RelevantParameterVisitor visitor(this, LiveInstance->Model, FParameters, OutRelevantParameters);
-
-		WorkingMemoryManager.CurrentInstanceCache = nullptr;
-		WorkingMemoryManager.CurrentPassthroughObjectLoader = nullptr;
+    	
+		RelevantParameterVisitor visitor(this, LiveInstance, FParameters, OutRelevantParameters);
     }
 
 	
@@ -1030,309 +967,255 @@ namespace UE::Mutable::Private
     }
 
 
-	void FSystem::BeginBuild(const TSharedPtr<const FModel>& InModel,
-		const TSharedPtr<FMeshIdRegistry>& InMeshIdRegistry,
-		const TSharedPtr<FImageIdRegistry>& InImageIdRegistry,
-		const TSharedPtr<FMaterialIdRegistry>& InMaterialIdRegistry,
-		const TSharedPtr<FSkeletalMeshIdRegistry>& InSkeletalMeshIdRegistry)
+	void FSystem::BeginBuild(const TSharedPtr<const FModel>& InModel)
 	{
-		// We don't have a FLiveInstance, let's create the memory
-		// \TODO: There is no clear moment to remove this... EndBuild?
-		WorkingMemoryManager.CurrentInstanceCache = MakeShared<FProgramCache>();
-		WorkingMemoryManager.CurrentInstanceCache->Init(InModel->Program.OpAddress.Num());
-
-    	WorkingMemoryManager.CurrentPassthroughObjectLoader = nullptr;
-    	
-    	WorkingMemoryManager.CurrentMeshIdRegistry = InMeshIdRegistry;
-    	WorkingMemoryManager.CurrentImageIdRegistry = InImageIdRegistry;
-    	WorkingMemoryManager.CurrentMaterialIdRegistry = InMaterialIdRegistry;
-    	WorkingMemoryManager.CurrentSkeletalMeshIdRegistry = InSkeletalMeshIdRegistry;
-    	
 		// Ensure the model cache has been created
 		WorkingMemoryManager.FindOrAddModelCache(InModel);
-
-		PrepareCache(InModel.Get(), -1);
 	}
 
 
 	void FSystem::EndBuild()
 	{
-		WorkingMemoryManager.CurrentInstanceCache = nullptr;
-		WorkingMemoryManager.CurrentPassthroughObjectLoader = nullptr;
-    	WorkingMemoryManager.CurrentMeshIdRegistry = nullptr;
-    	WorkingMemoryManager.CurrentImageIdRegistry = nullptr;
-    	WorkingMemoryManager.CurrentMaterialIdRegistry = nullptr;
-    	WorkingMemoryManager.CurrentSkeletalMeshIdRegistry = nullptr;
+
 	}
 
 
 	void FSystem::RunCode(
-		const TSharedPtr<FExternalResourceProvider>& ExternalResourceProvider,
-		const TSharedPtr<const FModel>& InModel,
+		const TSharedRef<FLiveInstance>& InLiveInstance,
 		const FParameters* InParameters,
 		OP::ADDRESS InCodeRoot,
 		uint32 InLODs,
 		uint8 ExecutionOptions)
 	{
 		TSharedRef<CodeRunner> Runner = CodeRunner::Create(
+			InLiveInstance,
 			Settings,
 			this,
 			EExecutionStrategy::MinimizeMemory,
-			InModel,
 			InParameters,
 			InCodeRoot,
 			InLODs,
 			ExecutionOptions,
 			false,
-			FScheduledOp::EType::Full,
-			ExternalResourceProvider);
+			FScheduledOp::EType::Full);
 		
-		constexpr bool bForceInlineExecutution = true;
-		Runner->StartRun(bForceInlineExecutution);
+		constexpr bool bForceInlineExecution = true;
+		Runner->StartRun(bForceInlineExecution);
 		
 		bUnrecoverableError = Runner->bUnrecoverableError;
 	}
-
-
-	bool FSystem::BuildBool(const TSharedPtr<FExternalResourceProvider>& ExternalResourceProvider, const TSharedPtr<const FModel>& pModel, const FParameters* Params, OP::ADDRESS at)
-	{
-		WorkingMemoryManager.BeginRunnerThread();
-
-		RunCode(ExternalResourceProvider, pModel, Params, at);
-
-		bool bResult = false;
-		if (!bUnrecoverableError)
-		{
-			bResult = WorkingMemoryManager.CurrentInstanceCache->GetBool(FCacheAddress(at, 0, 0));
-		}
-
-		WorkingMemoryManager.EndRunnerThread();
-
-		return bResult;
-	}
-
-
-	float FSystem::BuildScalar(const TSharedPtr<FExternalResourceProvider>& ExternalResourceProvider, const TSharedPtr<const FModel>& pModel, const FParameters* Params, OP::ADDRESS at)
+	
+	
+	bool FSystem::BuildBool(const TSharedRef<FLiveInstance>& InLiveInstance, const FParameters* InParameters, OP::ADDRESS InAddress)
+    {
+    	WorkingMemoryManager.BeginRunnerThread();
+    	
+    	RunCode(InLiveInstance, InParameters, InAddress);
+	
+    	bool bResult = false;
+    	if (!bUnrecoverableError)
+    	{
+    		bResult = InLiveInstance->Cache->GetBool(FCacheAddress(InAddress, 0, 0));
+    	}
+	
+    	WorkingMemoryManager.EndRunnerThread(InLiveInstance);
+	
+    	return bResult;    
+    }
+	
+	
+	int32 FSystem::BuildInt(const TSharedRef<FLiveInstance>& InLiveInstance, const FParameters* InParameters, OP::ADDRESS InAddress)
 	{
 		WorkingMemoryManager.BeginRunnerThread();
+    	
+		RunCode(InLiveInstance, InParameters, InAddress);
 
-		RunCode(ExternalResourceProvider, pModel, Params, at);
-
-		float Result = 0.0f;		
+		int32 Result = 0;
 		if (!bUnrecoverableError)
 		{
-			Result = WorkingMemoryManager.CurrentInstanceCache->GetScalar(FCacheAddress(at, 0, 0));
+			Result = InLiveInstance->Cache->GetInt(FCacheAddress(InAddress, 0, 0));
 		}
 
-		WorkingMemoryManager.EndRunnerThread();
+		WorkingMemoryManager.EndRunnerThread(InLiveInstance);
 
 		return Result;
 	}
 
-
-	int32 FSystem::BuildInt(const TSharedPtr<FExternalResourceProvider>& ExternalResourceProvider, const TSharedPtr<const FModel>& pModel, const FParameters* Params, OP::ADDRESS at)
+	
+	float FSystem::BuildScalar(const TSharedRef<FLiveInstance>& InLiveInstance, const FParameters* InParameters, OP::ADDRESS InAddress)
 	{
-		WorkingMemoryManager.BeginRunnerThread();
-
-		RunCode(ExternalResourceProvider, pModel, Params, at);
+    	WorkingMemoryManager.BeginRunnerThread();
+    	
+    	RunCode(InLiveInstance, InParameters, InAddress);
 
-		int32 Result = 0;
-		if (!bUnrecoverableError)
-		{
-			Result = WorkingMemoryManager.CurrentInstanceCache->GetInt(FCacheAddress(at, 0, 0));
-		}
+    	float Result = 0.0f;		
+    	if (!bUnrecoverableError)
+    	{
+    		Result = InLiveInstance->Cache->GetScalar(FCacheAddress(InAddress, 0, 0));
+    	}
 
-		WorkingMemoryManager.EndRunnerThread();
+    	WorkingMemoryManager.EndRunnerThread(InLiveInstance);
 
-		return Result;
+    	return Result;
 	}
 
 
-	FVector4f FSystem::BuildColour(const TSharedPtr<FExternalResourceProvider>& ExternalResourceProvider, const TSharedPtr<const FModel>& pModel, const FParameters* Params, OP::ADDRESS at)
+	FVector4f FSystem::BuildColor(const TSharedRef<FLiveInstance>& InLiveInstance, const FParameters* InParameters, OP::ADDRESS InAddress)
 	{
 		WorkingMemoryManager.BeginRunnerThread();
 
 		FVector4f Result(0,0,0,1);
-
-		UE::Mutable::Private::EOpType OpType = pModel->Program.GetOpType(at);
+    	
+		UE::Mutable::Private::EOpType OpType = InLiveInstance->Model->Program.GetOpType(InAddress);
 		if (GetOpDataType(OpType) == EDataType::Color)
 		{
-			RunCode(ExternalResourceProvider, pModel, Params, at);
+			RunCode(InLiveInstance, InParameters, InAddress);
 			if (!bUnrecoverableError)
 			{
-				Result = WorkingMemoryManager.CurrentInstanceCache->GetColour(FCacheAddress(at, 0, 0));
+				Result = InLiveInstance->Cache->GetColor(FCacheAddress(InAddress, 0, 0));
 			}
 		}
 
-		WorkingMemoryManager.EndRunnerThread();
+		WorkingMemoryManager.EndRunnerThread(InLiveInstance);
 
 		return Result;
 	}
 
 	
-	FProjector FSystem::BuildProjector(const TSharedPtr<FExternalResourceProvider>& ExternalResourceProvider, const TSharedPtr<const FModel>& pModel, const FParameters* Params, OP::ADDRESS at)
+	FProjector FSystem::BuildProjector(const TSharedRef<FLiveInstance>& InLiveInstance, const FParameters* InParameters, OP::ADDRESS InAddress)
 	{
 		WorkingMemoryManager.BeginRunnerThread();
-
-    	RunCode(ExternalResourceProvider, pModel, Params, at);
+    	
+    	RunCode(InLiveInstance, InParameters, InAddress);
 
 		FProjector Result;
 		if (!bUnrecoverableError)
 		{
-			Result = WorkingMemoryManager.CurrentInstanceCache->GetProjector(FCacheAddress(at, 0, 0));
+			Result = InLiveInstance->Cache->GetProjector(FCacheAddress(InAddress, 0, 0));
 		}
 
-		WorkingMemoryManager.EndRunnerThread();
+		WorkingMemoryManager.EndRunnerThread(InLiveInstance);
 
 		return Result;
 	}
 
-	
-	TManagedPtr<const FImage> FSystem::BuildImage(const TSharedPtr<FExternalResourceProvider>& ExternalResourceProvider, const TSharedPtr<const FModel>& pModel, const FParameters* Params, OP::ADDRESS at, int32 MipsToSkip)
+
+	TManagedPtr<const FImage> FSystem::BuildImage(const TSharedRef<FLiveInstance>& InLiveInstance, const FParameters* InParameters, OP::ADDRESS InAddress, int32 MipsToSkip)
 	{
 		WorkingMemoryManager.BeginRunnerThread();
 
 		TManagedPtr<const FImage> Result;
-
-		UE::Mutable::Private::EOpType OpType = pModel->Program.GetOpType(at);
+    	
+		UE::Mutable::Private::EOpType OpType = InLiveInstance->Model->Program.GetOpType(InAddress);
 		if (GetOpDataType(OpType) == EDataType::Image)
 		{
-			RunCode(ExternalResourceProvider, pModel, Params, at, FSystem::AllLODs, uint8(MipsToSkip));
+			RunCode(InLiveInstance, InParameters, InAddress, FSystem::AllLODs, uint8(MipsToSkip));
 			if (!bUnrecoverableError)
 			{
-				Result = WorkingMemoryManager.LoadImage(FCacheAddress(at, 0, MipsToSkip));
+				Result = InLiveInstance->Cache->GetImage(FCacheAddress(InAddress, 0, MipsToSkip));
 			}			
 		}
 
-		WorkingMemoryManager.EndRunnerThread();
+		WorkingMemoryManager.EndRunnerThread(InLiveInstance);
 
 		return Result;
 	}
 
 
-	TManagedPtr<const FMesh> FSystem::BuildMesh(const TSharedPtr<FExternalResourceProvider>& ExternalResourceProvider, const TSharedPtr<const FModel>& Model, const FParameters* Params, OP::ADDRESS RootAddress, EMeshContentFlags MeshContentFilter)
+	TManagedPtr<const FMesh> FSystem::BuildMesh(const TSharedRef<FLiveInstance>& InLiveInstance, const FParameters* InParameters, OP::ADDRESS InRootAddress, EMeshContentFlags MeshContentFilter)
 	{
 		WorkingMemoryManager.BeginRunnerThread();
 
 		TManagedPtr<const FMesh> Result;
-
-		UE::Mutable::Private::EOpType OpType = Model->Program.GetOpType(RootAddress);
-
+    	
+		UE::Mutable::Private::EOpType OpType = InLiveInstance->Model->Program.GetOpType(InRootAddress);
 		if (GetOpDataType(OpType) == EDataType::Mesh)
 		{
-			uint8 ExecutionOptions = static_cast<uint8>(MeshContentFilter);
-			RunCode(ExternalResourceProvider, Model, Params, RootAddress, FSystem::AllLODs, ExecutionOptions);
+			const uint8 ExecutionOptions = static_cast<uint8>(MeshContentFilter);
+			RunCode(InLiveInstance, InParameters, InRootAddress, FSystem::AllLODs, ExecutionOptions);
 			if (!bUnrecoverableError)
 			{
-				Result = WorkingMemoryManager.LoadMesh(FCacheAddress(RootAddress, 0, ExecutionOptions));
+				Result = InLiveInstance->Cache->GetMesh(FCacheAddress(InRootAddress, 0, ExecutionOptions));
 			}
 		}
 
-		WorkingMemoryManager.EndRunnerThread();
+		WorkingMemoryManager.EndRunnerThread(InLiveInstance);
 
 		return Result;
 	}
 
 
-	TManagedPtr<const FInstance> FSystem::BuildInstance(const TSharedPtr<FExternalResourceProvider>& ExternalResourceProvider, const TSharedPtr<const FModel>& pModel, const FParameters* Params, OP::ADDRESS at)
+	TManagedPtr<const FInstance> FSystem::BuildInstance(const TSharedRef<FLiveInstance>& InLiveInstance, const FParameters* InParameters, OP::ADDRESS InAddress)
 	{
 		WorkingMemoryManager.BeginRunnerThread();
 
 		TManagedPtr<const FInstance> Result;
-
-		UE::Mutable::Private::EOpType OpType = pModel->Program.GetOpType(at);
+    	
+		UE::Mutable::Private::EOpType OpType = InLiveInstance->Model->Program.GetOpType(InAddress);
 		if (GetOpDataType(OpType) == EDataType::Instance)
 		{
-			RunCode(ExternalResourceProvider, pModel, Params, at);
+			RunCode(InLiveInstance, InParameters, InAddress);
 			if (!bUnrecoverableError)
 			{
-				Result = WorkingMemoryManager.CurrentInstanceCache->GetInstance(FCacheAddress(at, 0, 0));
+				Result = InLiveInstance->Cache->GetInstance(FCacheAddress(InAddress, 0, 0));
 			}
 		}
 
-		WorkingMemoryManager.EndRunnerThread();
+		WorkingMemoryManager.EndRunnerThread(InLiveInstance);
 
 		return Result;
 	}
 
 
-	TManagedPtr<const FLayout> FSystem::BuildLayout(const TSharedPtr<FExternalResourceProvider>& ExternalResourceProvider, const TSharedPtr<const FModel>& pModel, const FParameters* Params, OP::ADDRESS at)
+	TManagedPtr<const FLayout> FSystem::BuildLayout(const TSharedRef<FLiveInstance>& InLiveInstance, const FParameters* InParameters, OP::ADDRESS InAddress)
 	{
 		WorkingMemoryManager.BeginRunnerThread();
 
 		TManagedPtr<const FLayout> Result;
-
-		if (pModel->Program.States[0].Root)
+    	
+		if (InLiveInstance->Model->Program.States[0].Root)
 		{
-			UE::Mutable::Private::EOpType OpType = pModel->Program.GetOpType(at);
+			UE::Mutable::Private::EOpType OpType = InLiveInstance->Model->Program.GetOpType(InAddress);
 			if (GetOpDataType(OpType) == EDataType::Layout)
 			{
-				RunCode(ExternalResourceProvider, pModel, Params, at);
+				RunCode(InLiveInstance, InParameters, InAddress);
 				if (!bUnrecoverableError)
 				{
-					Result = WorkingMemoryManager.CurrentInstanceCache->GetLayout(FCacheAddress(at, 0, 0));
+					Result = InLiveInstance->Cache->GetLayout(FCacheAddress(InAddress, 0, 0));
 				}
 			}
 		}
 
-		WorkingMemoryManager.EndRunnerThread();
+		WorkingMemoryManager.EndRunnerThread(InLiveInstance);
 
 		return Result;
 	}
 
 
-	TManagedPtr<const String> FSystem::BuildString(const TSharedPtr<FExternalResourceProvider>& ExternalResourceProvider, const TSharedPtr<const FModel>& pModel, const FParameters* Params, OP::ADDRESS at)
+	TManagedPtr<const String> FSystem::BuildString(const TSharedRef<FLiveInstance>& InLiveInstance, const FParameters* InParameters, OP::ADDRESS InAddress)
 	{
 		WorkingMemoryManager.BeginRunnerThread();
 
 		TManagedPtr<const String> Result;
-
-		if (pModel->Program.States[0].Root)
+    	
+		if (InLiveInstance->Model->Program.States[0].Root)
 		{
-			UE::Mutable::Private::EOpType OpType = pModel->Program.GetOpType(at);
+			UE::Mutable::Private::EOpType OpType = InLiveInstance->Model->Program.GetOpType(InAddress);
 			if (GetOpDataType(OpType) == EDataType::String)
 			{
-				RunCode(ExternalResourceProvider, pModel, Params, at);
+				RunCode(InLiveInstance, InParameters, InAddress);
 				if (!bUnrecoverableError)
 				{
-					Result = WorkingMemoryManager.CurrentInstanceCache->GetString(FCacheAddress(at, 0, 0));
+					Result = InLiveInstance->Cache->GetString(FCacheAddress(InAddress, 0, 0));
 				}
 			}
 		}
 
-		WorkingMemoryManager.EndRunnerThread();
+		WorkingMemoryManager.EndRunnerThread(InLiveInstance);
 
 		return Result;
 	}
-
-
-	void FSystem::PrepareCache( const FModel* InModel, int32 InState)
-	{
-		MUTABLE_CPUPROFILER_SCOPE(PrepareCache);
-
-		const FProgram& Program = InModel->Program;
-		int32 OpCount = Program.OpAddress.Num();
-		WorkingMemoryManager.CurrentInstanceCache->Init(OpCount);
-
-		// Clear cache flags of existing data
-		CodeContainer<FProgramCache::FOpExecutionData>::iterator It = WorkingMemoryManager.CurrentInstanceCache->OpExecutionData.begin();
-		for (; It.IsValid(); ++It)
-		{
-			(*It).OpHitCount = 0; // This should already be 0, but just in case.
-			(*It).IsCacheLocked = false;
-		}
-
-		// Mark the resources that have to be cached to update the instance in this state
-		if (InState >= 0 && InState< Program.States.Num())
-		{
-			const FProgram::FState& State = Program.States[InState];
-			for (uint32 Address : State.m_updateCache)
-			{
-				WorkingMemoryManager.CurrentInstanceCache->SetForceCached(Address);
-			}
-		}
-	}
-
+	
 
 	void FSystem::UpdateStats()
 	{
@@ -1435,16 +1318,42 @@ namespace UE::Mutable::Private
 	}
 
 
+	void FProgramCache::Prepare(const FProgram& InProgram, int32 InState)
+	{
+    	MUTABLE_CPUPROFILER_SCOPE(PrepareCache);
+			
+    	const int32 OpCount = InProgram.OpAddress.Num();
+    	Init(OpCount);
+
+    	// Clear cache flags of existing data
+    	CodeContainer<FProgramCache::FOpExecutionData>::iterator It = OpExecutionData.begin();
+    	for (; It.IsValid(); ++It)
+    	{
+    		(*It).OpHitCount = 0; // This should already be 0, but just in case.
+    		(*It).IsCacheLocked = false;
+    	}
+
+    	// Mark the resources that have to be cached to update the instance in this state
+    	if (InState >= 0 && InState< InProgram.States.Num())
+    	{
+    		const FProgram::FState& State = InProgram.States[InState];
+    		for (uint32 Address : State.m_updateCache)
+    		{
+    			SetForceCached(Address);
+    		}
+    	}
+	}
+	
+
 	FWorkingMemoryManager::FModelCacheEntry* FWorkingMemoryManager::FindModelCache(const FModel* InModel)
 	{
 		check(InModel);
 
 		for (FModelCacheEntry& c : CachePerModel)
 		{
-			TSharedPtr<const FModel> pCandidate = c.Model.Pin();
-			if (pCandidate)
+			if (TSharedPtr<const FModel> Candidate = c.Model.Pin())
 			{
-				if (pCandidate.Get() == InModel)
+				if (Candidate.Get() == InModel)
 				{
 					return &c;
 				}
@@ -1515,7 +1424,7 @@ namespace UE::Mutable::Private
 		return CurrentBytes > BudgetThresholdBytes;
 	}
 
-	bool FWorkingMemoryManager::EnsureBudgetBelow( uint64 AdditionalMemory )
+	bool FWorkingMemoryManager::EnsureBudgetBelow( uint64 AdditionalMemory, const TSharedPtr<FProgramCache>& InCurrentCache )
     {
 		MUTABLE_CPUPROFILER_SCOPE(EnsureBudgetBelow);
 
@@ -1683,7 +1592,7 @@ namespace UE::Mutable::Private
 					continue;
 				}
 				
-				if (Instance->Cache == CurrentInstanceCache)
+				if (Instance->Cache == InCurrentCache)
 				{
 					// Ignore the current live instance.
 					continue;
@@ -1703,11 +1612,11 @@ namespace UE::Mutable::Private
 
 		// From the current live instances. It is more involved: we have to make sure any data we want to free is not also
 		// in any cache (0 or 1) position with hit-count > 0.
-		if (!bFinished && CurrentInstanceCache)
+		if (!bFinished && InCurrentCache)
 		{
 			MUTABLE_CPUPROFILER_SCOPE(EnsureBudgetBelow_FreeCached_Current);
 
-			TryRemoveFromCache(*CurrentInstanceCache);
+			TryRemoveFromCache(*InCurrentCache);
 		}
 
 
@@ -1720,7 +1629,7 @@ namespace UE::Mutable::Private
 				BudgetExcessBytes = ExcessBytes;
 
 				// We failed to free enough memory. Log this, but try to continue anyway.
-				// This is a good place to insert a brakpoint to detect callstacks with memory peaks
+				// This is a good place to insert a breakpoint to detect callstacks with memory peaks
 				UE_LOG(LogMutableCore, Log, TEXT("Failed to keep memory budget. Budget: %" INT64_FMT ", Current: %" INT64_FMT ", New: %" UINT64_FMT),
 					BudgetBytes / 1024, (TotalBytes - AdditionalMemory) / 1024, AdditionalMemory / 1024);
 				
@@ -2078,6 +1987,8 @@ namespace UE::Mutable::Private
 
 	FMeshId FSystem::GetMeshId(const TSharedPtr<FMeshIdRegistry>& Registry, const TSharedPtr<const FModel>& Model, const FParameters* Params, OP::ADDRESS RootAt)
     {
+    	check(Registry)
+    	
     	FGeneratedMeshKey NewKey;
     	NewKey.Address = RootAt;
 
@@ -2091,6 +2002,8 @@ namespace UE::Mutable::Private
     	
 	FImageId FSystem::GetImageId(const TSharedPtr<FImageIdRegistry>& Registry, const TSharedPtr<const FModel>& Model, const FParameters* Params, OP::ADDRESS RootAt)
     {
+    	check(Registry)
+    
     	FGeneratedImageKey NewKey;
     	NewKey.Address = RootAt;
 
@@ -2104,6 +2017,8 @@ namespace UE::Mutable::Private
 	
 	FMaterialId FSystem::GetMaterialId(const TSharedPtr<FMaterialIdRegistry>& Registry, const TSharedPtr<const FModel>& Model, const FParameters* Params, OP::ADDRESS RootAt)
     {
+    	check(Registry)
+    
     	FGeneratedMaterialKey NewKey;
     	NewKey.Address = RootAt;
 
@@ -2117,6 +2032,8 @@ namespace UE::Mutable::Private
 	
 	FSkeletalMeshId FSystem::GetSkeletalMeshId(const TSharedPtr<FSkeletalMeshIdRegistry>& Registry, const TSharedPtr<const FModel>& Model, const FParameters* Params, OP::ADDRESS RootAt)
 	{
+		check(Registry)
+	
     	FGeneratedSkeletalMeshKey NewKey;
     	NewKey.Address = RootAt;
 
diff --git a/Engine/Plugins/Mutable/Source/MutableTools/Private/MuT/CodeOptimiser.cpp b/Engine/Plugins/Mutable/Source/MutableTools/Private/MuT/CodeOptimiser.cpp
--- a/Engine/Plugins/Mutable/Source/MutableTools/Private/MuT/CodeOptimiser.cpp
+++ b/Engine/Plugins/Mutable/Source/MutableTools/Private/MuT/CodeOptimiser.cpp
@@ -194,8 +194,11 @@ namespace UE::Mutable::Private
 			Model->Program.States.Add(state);
 
 			TSharedPtr<FParameters> LocalParams = FModel::NewParameters(Model);
-			System->BeginBuild(Model, nullptr, nullptr, nullptr, nullptr);
 
+			TSharedRef<FLiveInstance> TempLiveInstance = System->NewBuildInstance(Model, nullptr);
+			
+			System->BeginBuild(Model);
+			
 			// Calculate the value and replace this op by a constant
 			switch( DataType )
 			{
@@ -203,7 +206,7 @@ namespace UE::Mutable::Private
 			{
 				MUTABLE_CPUPROFILER_SCOPE(ConstantMesh);
 
-				TManagedPtr<const FMesh> MeshBuild = System->BuildMesh(nullptr, Model, LocalParams.Get(), at, EMeshContentFlags::AllFlags);
+				TManagedPtr<const FMesh> MeshBuild = System->BuildMesh(TempLiveInstance, LocalParams.Get(), at, EMeshContentFlags::AllFlags);
 
 				if (MeshBuild)
 				{
@@ -220,7 +223,7 @@ namespace UE::Mutable::Private
 			{
 				MUTABLE_CPUPROFILER_SCOPE(ConstantImage);
 
-				TManagedPtr<const FImage> ImageResult = System->BuildImage(nullptr, Model, LocalParams.Get(), at, 0);
+				TManagedPtr<const FImage> ImageResult = System->BuildImage(TempLiveInstance, LocalParams.Get(), at, 0);
 
 				if (ImageResult)
 				{
@@ -237,7 +240,7 @@ namespace UE::Mutable::Private
 			{
 				MUTABLE_CPUPROFILER_SCOPE(ConstantLayout);
 
-				TManagedPtr<const FLayout> LayoutResult = System->BuildLayout(nullptr, Model, LocalParams.Get(), at);
+				TManagedPtr<const FLayout> LayoutResult = System->BuildLayout(TempLiveInstance, LocalParams.Get(), at);
 
 				if (LayoutResult)
 				{
@@ -253,7 +256,7 @@ namespace UE::Mutable::Private
 			{
 				MUTABLE_CPUPROFILER_SCOPE(ConstantBool);
 
-				bool value = System->BuildBool(nullptr, Model, LocalParams.Get(), at);
+				bool value = System->BuildBool(TempLiveInstance, LocalParams.Get(), at);
 				Result = new ASTOpConstantBool(value);
 				break;
 			}
@@ -263,7 +266,7 @@ namespace UE::Mutable::Private
 				MUTABLE_CPUPROFILER_SCOPE(ConstantBool);
 
 				FVector4f ResultColor(0, 0, 0, 0);
-				ResultColor = System->BuildColour(nullptr, Model, LocalParams.Get(), at);
+				ResultColor = System->BuildColor(TempLiveInstance, LocalParams.Get(), at);
 
 				{
 					UE::Mutable::Private::Ptr<ASTOpConstantColor> ConstantOp = new ASTOpConstantColor();
@@ -1099,7 +1102,7 @@ namespace UE::Mutable::Private
 			ASTOp::LogHistogram(roots);
 		}
 
-		// Minimal optimisation of constant subtrees
+		// Minimal optimization of constant subtrees
 		else if ( Options->GetPrivate()->OptimisationOptions.bConstReduction )
 		{
 			// The first duplicated data remover has the special mission of removing
