[Mutable] PR #14185: Fix Mutable compilation crash

#Synced-CL 49839474

diff --git a/Engine/Plugins/Mutable/Source/CustomizableObjectEditor/Private/MuCOE/CustomizableObjectCompiler.cpp b/Engine/Plugins/Mutable/Source/CustomizableObjectEditor/Private/MuCOE/CustomizableObjectCompiler.cpp
--- a/Engine/Plugins/Mutable/Source/CustomizableObjectEditor/Private/MuCOE/CustomizableObjectCompiler.cpp
+++ b/Engine/Plugins/Mutable/Source/CustomizableObjectEditor/Private/MuCOE/CustomizableObjectCompiler.cpp
@@ -1243,13 +1243,14 @@ bool FCustomizableObjectCompiler::TryPopCompileRequest()
 		UCustomizableObject* Object = CompileRequest->GetCustomizableObject();
 		if (!Object)
 		{
+			TSharedRef<FCompilationRequest> SelectedCompileRequest = CompileRequest;
 			It.RemoveCurrent();
-
+			
 			FCompileCallbackParams Params;
 			Params.bRequestFailed = true;
     	
-			CompileRequest->Callback.ExecuteIfBound(Params);
-			CompileRequest->CallbackNative.ExecuteIfBound(Params);
+			SelectedCompileRequest->Callback.ExecuteIfBound(Params);
+			SelectedCompileRequest->CallbackNative.ExecuteIfBound(Params);
 			continue;
 		}
 		
