[Mutable] Add support for iterative multicompose operations.

- Added Cvar mutable.MaxWorkerThreadsForIterativeOpExec to control the maximum number of worker threads a platform can have to use iterative operations. Defaulted to 4, -1 forces iterative always.
- Fix crash and data corruption when multicompose base image is not compressed but the source image is compressed.

#rnx

#Synced-CL 49840162

diff --git a/Engine/Plugins/Mutable/Source/MutableRuntime/Internal/MuR/SystemTypes.h b/Engine/Plugins/Mutable/Source/MutableRuntime/Internal/MuR/SystemTypes.h
--- a/Engine/Plugins/Mutable/Source/MutableRuntime/Internal/MuR/SystemTypes.h
+++ b/Engine/Plugins/Mutable/Source/MutableRuntime/Internal/MuR/SystemTypes.h
@@ -13,6 +13,9 @@ namespace UE::Mutable::Private
 	/** This structure stores the data about an ongoing mutable operation that needs to be executed. */
 	struct FScheduledOp
 	{
+		static constexpr int32 NumStageBits = 6;
+		static constexpr int32 MaxNumStages = (1 << NumStageBits) - 1;
+
 		inline FScheduledOp() = default;
 
 		inline FScheduledOp(OP::ADDRESS InAt, const FScheduledOp& InOpTemplate, uint8 InStage = 0, uint32 InCustomState = 0)
@@ -23,7 +26,7 @@ namespace UE::Mutable::Private
 			, Stage(InStage)
 			, Type(InOpTemplate.Type)
 		{
-			check(InStage < ((1 << 6) - 1));
+			check(InStage <= MaxNumStages);
 		}
 
 		static inline FScheduledOp FromOpAndOptions(OP::ADDRESS InAt, const FScheduledOp& InOpTemplate, uint8 InExecutionOptions)
@@ -60,7 +63,7 @@ namespace UE::Mutable::Private
 		//! Internal stage of the operation.
 		//! Stage 0 is usually scheduling of children, and 1 is execution. Some instructions
 		//! may have more steges to schedule children that are optional for execution, etc.
-		uint8 Stage      : 6 = 0;
+		uint8 Stage      : NumStageBits = 0;
 
 		uint8 bForceLazy : 1 = 0;
 
diff --git a/Engine/Plugins/Mutable/Source/MutableRuntime/Private/MuR/CodeRunner.cpp b/Engine/Plugins/Mutable/Source/MutableRuntime/Private/MuR/CodeRunner.cpp
--- a/Engine/Plugins/Mutable/Source/MutableRuntime/Private/MuR/CodeRunner.cpp
+++ b/Engine/Plugins/Mutable/Source/MutableRuntime/Private/MuR/CodeRunner.cpp
@@ -109,8 +109,17 @@ static FAutoConsoleVariableRef CVarUseImageTransformVectorImpl (
 	bUseImageTransformVectorImpl,
 	TEXT("If set to true, enables the vectorized implementation of the image transform pixel processing."),
 	ECVF_Default);
+
+int32 MaxWorkerThreadsForIterativeOpExec = 4;
+static FAutoConsoleVariableRef CVarMaxWorkerThreadsForIterativeOpExec(
+	TEXT("mutable.MaxWorkerThreadsForIterativeOpExec"),
+	MaxWorkerThreadsForIterativeOpExec,
+	TEXT("Sets the maximum number of worker threads to consider executing the iterative variant of an operation.")
+	TEXT("A value of -1 forces the iterative op execution, a value 0 prevents the iterative variant to be used."),
+	ECVF_Default);
 }
 
+
 #if MUTABLE_DEBUG_CODERUNNER_TASK_SCHEDULE_CALLSTACK
 
 #define AddOp(...) Invoke([&](){ AddOp(__VA_ARGS__); });
@@ -4125,6 +4134,8 @@ namespace UE::Mutable::Private
 			}
 			case 1:
 			{
+				MUTABLE_CPUPROFILER_SCOPE(RunCode_ImageMultiCompose_1);
+
 				TManagedPtr<const FLayout> ComposeLayout =
 					LoadLayout(FCacheAddress(Args.Layout, FScheduledOp::FromOpAndOptions(Args.Layout, Item, 0)));
 
@@ -4142,33 +4153,40 @@ namespace UE::Mutable::Private
 					}
 				}
 
-				FScheduledOpData SourceData;
-				SourceData.Resource = SourceLayout;
+				int32 HeapIndex = HeapData.Emplace();
+				HeapData.Emplace();
+				HeapData.Emplace();
+				HeapData.Emplace();
 
-				FScheduledOpData Data;
-				Data.Resource = ComposeLayout;
-				Data.MultiCompose.NextHeapIndex = HeapData.Add(SourceData);
-				Data.MultiCompose.bHasBlocksToCompose = bHasBlocksToCompose;
-				int32 DataPos = HeapData.Add(Data);
+				HeapData[HeapIndex + 0].MultiCompose.bSplitInMultipleIterations = 
+						MaxWorkerThreadsForIterativeOpExec < 0 || 
+						(int32)LowLevelTasks::FScheduler::Get().GetNumWorkers() < MaxWorkerThreadsForIterativeOpExec;
+
+				HeapData[HeapIndex + 0].MultiCompose.bHasBlocksToCompose = bHasBlocksToCompose;
+				HeapData[HeapIndex + 0].Resource = ComposeLayout;
+
+				HeapData[HeapIndex + 1].Resource = SourceLayout;
 
 				if (bHasBlocksToCompose)
 				{
-					AddOp(FScheduledOp(Item.At, Item, 2, DataPos), FScheduledOp(Args.Base, Item));
+					AddOp(FScheduledOp(Item.At, Item, 2, HeapIndex), FScheduledOp(Args.Base, Item));
 				}
 				else
 				{
 					// Jump directly to stage 3, no need to load the source image.
-					AddOp(FScheduledOp(Item.At, Item, 3, DataPos), FScheduledOp(Args.Base, Item));
+					AddOp(FScheduledOp(Item.At, Item, 3, HeapIndex), FScheduledOp(Args.Base, Item));
 				}
 
 				break;
 			}
 			case 2:
 			{
-				FScheduledOpData& Data = HeapData[Item.CustomState];
+				MUTABLE_CPUPROFILER_SCOPE(RunCode_ImageMultiCompose_2);
+
+				FScheduledOpData& Data = HeapData[Item.CustomState + 0];
 				TManagedPtr<const FLayout> ComposeLayout = StaticCastManagedPtr<const FLayout>(Data.Resource);
 
-				FScheduledOpData& SourceData = HeapData[Data.MultiCompose.NextHeapIndex];
+				FScheduledOpData& SourceData = HeapData[Item.CustomState + 1];
 				TManagedPtr<const FLayout> SourceLayout = StaticCastManagedPtr<const FLayout>(SourceData.Resource);
 
 				FIntVector2 SourceLayoutGridSize = SourceLayout->GetGridSize();
@@ -4235,18 +4253,18 @@ namespace UE::Mutable::Private
 			}
 			case 3:
 			{
-				MUTABLE_CPUPROFILER_SCOPE(RunCode_ImageMultiCompose);
+				MUTABLE_CPUPROFILER_SCOPE(RunCode_ImageMultiCompose_3);
+
+				FScheduledOpData& Data = HeapData[Item.CustomState + 0];
+				TManagedPtr<const FLayout> ComposeLayout = StaticCastManagedPtr<const FLayout>(Data.Resource);
+
+				FScheduledOpData& SourceData = HeapData[Item.CustomState + 1];
+				TManagedPtr<const FLayout> SourceLayout = StaticCastManagedPtr<const FLayout>(SourceData.Resource);
 
 				TManagedPtr<const FImage> BaseImage = LoadImage(FCacheAddress(Args.Base, Item));
 
-				FScheduledOpData& Data = HeapData[Item.CustomState];
 				if (BaseImage && Data.MultiCompose.bHasBlocksToCompose)
 				{
-					TManagedPtr<const FLayout> ComposeLayout = StaticCastManagedPtr<const FLayout>(Data.Resource);
-
-					FScheduledOpData& SourceData = HeapData[Data.MultiCompose.NextHeapIndex];
-					TManagedPtr<const FLayout> SourceLayout = StaticCastManagedPtr<const FLayout>(SourceData.Resource);
-
 					int32 SourceImageLOD = Data.MultiCompose.SourceLOD;
 
 					TManagedPtr<const FImage> SourceImage = 
@@ -4268,8 +4286,6 @@ namespace UE::Mutable::Private
 							SourceImage = Resized;	
 						}
 
-						TManagedPtr<FImage> Result = CloneOrTakeOver(MoveTemp(BaseImage));
-
 						FIntVector2 SourceLayoutGridSize = SourceLayout->GetGridSize();
 						FIntVector2 ComposeLayoutGridSize = ComposeLayout->GetGridSize();
 
@@ -4278,13 +4294,10 @@ namespace UE::Mutable::Private
 								SourceImage->GetSizeY() / SourceLayoutGridSize.Y);
 
 						FIntVector2 ComposeLayoutGridSizeInPixels = FIntVector2(
-								Result->GetSizeX() / ComposeLayoutGridSize.X,
-								Result->GetSizeY() / ComposeLayoutGridSize.Y);
-
-						TArray<FIntRect, TInlineAllocator<16>> SrcRects;
-						TArray<FIntRect, TInlineAllocator<16>> DstRects;
-						TArray<int32,    TInlineAllocator<16>> SelectedLODForRect;
+								BaseImage->GetSizeX() / ComposeLayoutGridSize.X,
+								BaseImage->GetSizeY() / ComposeLayoutGridSize.Y);
 
+						Data.MultiCompose.NumElemsToProcess = 0;
 						int32 NumSrcLODsNeeded = 1;
 						for (const FLayoutBlock& SrcBlock : SourceLayout->Blocks)
 						{
@@ -4293,13 +4306,13 @@ namespace UE::Mutable::Private
 							if (FoundDstBlockIndex >= 0)
 							{
 								const FLayoutBlock& DstBlock = ComposeLayout->Blocks[FoundDstBlockIndex];
-								FIntRect& SrcRectInPixels = SrcRects.Emplace_GetRef(
+								FIntRect SrcRectInPixels = FIntRect(
 										SrcBlock.Min.X * SourceLayoutGridSizeInPixels.X, 
 										SrcBlock.Min.Y * SourceLayoutGridSizeInPixels.Y, 
 										(SrcBlock.Min.X + SrcBlock.Size.X) * SourceLayoutGridSizeInPixels.X, 
 										(SrcBlock.Min.Y + SrcBlock.Size.Y) * SourceLayoutGridSizeInPixels.Y);
 
-								FIntRect& DstRectInPixels = DstRects.Emplace_GetRef(
+								FIntRect DstRectInPixels = FIntRect(
 										DstBlock.Min.X * ComposeLayoutGridSizeInPixels.X, 
 										DstBlock.Min.Y * ComposeLayoutGridSizeInPixels.Y, 
 										(DstBlock.Min.X + DstBlock.Size.X) * ComposeLayoutGridSizeInPixels.X, 
@@ -4311,24 +4324,13 @@ namespace UE::Mutable::Private
 								int32 SelectedLOD = FMath::FloorToInt32(BestLODForRect + LODRoundTolerance);
 
 								NumSrcLODsNeeded = FMath::Max(SelectedLOD + 1, NumSrcLODsNeeded);
-								SelectedLODForRect.Emplace(SelectedLOD);
-							}
-						}
-
-						// Fix src rect sizes for the selected LOD.
-						for (int32 RectIndex = SrcRects.Num() - 1; RectIndex >= 0; --RectIndex)
-						{
-							int32 SelectedLOD = SelectedLODForRect[RectIndex];
-
-							for (int32 LODIndex = 0; LODIndex < SelectedLOD; ++LODIndex)
-							{
-								SrcRects[RectIndex] = FIntRect::DivideAndRoundUp(SrcRects[RectIndex], 2);
-							}
+								++Data.MultiCompose.NumElemsToProcess;		
+							}	
 						}
 
 						if (SourceImage->GetLODCount() < NumSrcLODsNeeded)
 						{
-							MUTABLE_CPUPROFILER_SCOPE(RunCode_ImageTransform_BilinearMipGen);
+							MUTABLE_CPUPROFILER_SCOPE(RunCode_ImageMultiCompose_BilinearMipGen);
 
 							TManagedPtr<FImage> OwnedSource = CloneOrTakeOver(MoveTemp(SourceImage));
 							OwnedSource->DataStorage.SetNumLODs(NumSrcLODsNeeded);
@@ -4338,28 +4340,153 @@ namespace UE::Mutable::Private
 							//ImageMipmapInPlace(0, OwnedSource.Get(), FMipmapGenerationSettings{});
 
 							SourceImage = OwnedSource;
-						}	
+						}
 
-						ImageMultiCompose(*Result, *SourceImage, DstRects, SrcRects, SelectedLODForRect, Settings.ImageCompressionQuality);
-						Result->DataStorage.SetNumLODs(1);
+						HeapData[Item.CustomState + 2].Resource = SourceImage; 
+						HeapData[Item.CustomState + 3].Resource = CloneOrTakeOver(MoveTemp(BaseImage)); 
 
-						StoreImage(Item, Result);
+						AddOp(FScheduledOp(Item.At, Item, 4, Item.CustomState));
 					}
 					else
 					{
 						StoreImage(Item, BaseImage);
+
+						HeapData[Item.CustomState + 0].Resource = nullptr;
+						HeapData[Item.CustomState + 1].Resource = nullptr;
+						HeapData[Item.CustomState + 2].Resource = nullptr;
+						HeapData[Item.CustomState + 3].Resource = nullptr;
 					}
 				}
 				else
 				{
 					StoreImage(Item, BaseImage);
+
+					HeapData[Item.CustomState + 0].Resource = nullptr;
+					HeapData[Item.CustomState + 1].Resource = nullptr;
+					HeapData[Item.CustomState + 2].Resource = nullptr;
+					HeapData[Item.CustomState + 3].Resource = nullptr;
 				}
 
 				break;
 			}
 
 			default:
-				check(false);
+			{
+				check(Item.Stage >= 4);
+
+				MUTABLE_CPUPROFILER_SCOPE(RunCode_ImageMultiCompose_Iters);
+
+				FScheduledOpData& Data = HeapData[Item.CustomState + 0];
+				FScheduledOpData& ComposeLayoutData = Data;
+				FScheduledOpData& SourceLayoutData  = HeapData[Item.CustomState + 1];
+				FScheduledOpData& SourceImageData   = HeapData[Item.CustomState + 2];
+				FScheduledOpData& ResultImageData   = HeapData[Item.CustomState + 3];
+				
+				TManagedPtr<const FLayout> ComposeLayout = StaticCastManagedPtr<const FLayout>(ComposeLayoutData.Resource);
+				TManagedPtr<const FLayout> SourceLayout = StaticCastManagedPtr<const FLayout>(SourceLayoutData.Resource);
+				TManagedPtr<const FImage> SourceImage = StaticCastManagedPtr<const FImage>(SourceImageData.Resource);
+				TManagedPtr<FImage> ResultImage = ConstCastManagedPtr<FImage>(StaticCastManagedPtr<const FImage>(ResultImageData.Resource));
+
+				check(SourceImage);
+				check(ResultImage);
+
+				int32 IterationBegin = 0;
+				int32 IterationEnd = IterationBegin + Data.MultiCompose.NumElemsToProcess;
+
+				if (Data.MultiCompose.bSplitInMultipleIterations)
+				{
+					const int32 CurrentIteration = Item.Stage - 4;
+					constexpr int32 MaxIterations = FScheduledOp::MaxNumStages - 4;
+					const int32 ElemsPerIter = FMath::DivideAndRoundUp<int32>(Data.MultiCompose.NumElemsToProcess, MaxIterations);
+
+					IterationBegin = CurrentIteration*ElemsPerIter; 
+					IterationEnd = FMath::Min<int32>(IterationBegin + ElemsPerIter, Data.MultiCompose.NumElemsToProcess);
+				}
+
+				int32 SourceImageLOD = Data.MultiCompose.SourceLOD;
+
+				FIntVector2 SourceLayoutGridSize = SourceLayout->GetGridSize();
+				FIntVector2 ComposeLayoutGridSize = ComposeLayout->GetGridSize();
+
+				FIntVector2 SourceLayoutGridSizeInPixels = FIntVector2(
+						SourceImage->GetSizeX() / SourceLayoutGridSize.X,
+						SourceImage->GetSizeY() / SourceLayoutGridSize.Y);
+
+				FIntVector2 ComposeLayoutGridSizeInPixels = FIntVector2(
+						ResultImage->GetSizeX() / ComposeLayoutGridSize.X,
+						ResultImage->GetSizeY() / ComposeLayoutGridSize.Y);
+
+				TArray<FIntRect, TInlineAllocator<16>> SrcRects;
+				TArray<FIntRect, TInlineAllocator<16>> DstRects;
+				TArray<int32,    TInlineAllocator<16>> SelectedLODForRect;
+
+				int32 ElemIndex = 0;
+				int32 NumSrcLODsNeeded = 1;
+				for (const FLayoutBlock& SrcBlock : SourceLayout->Blocks)
+				{
+					int32 FoundDstBlockIndex = ComposeLayout->FindBlock(SrcBlock.Id);
+
+					if (FoundDstBlockIndex >= 0)
+					{
+						if (ElemIndex >= IterationBegin && ElemIndex < IterationEnd)
+						{
+							const FLayoutBlock& DstBlock = ComposeLayout->Blocks[FoundDstBlockIndex];
+							FIntRect& SrcRectInPixels = SrcRects.Emplace_GetRef(
+									SrcBlock.Min.X * SourceLayoutGridSizeInPixels.X, 
+									SrcBlock.Min.Y * SourceLayoutGridSizeInPixels.Y, 
+									(SrcBlock.Min.X + SrcBlock.Size.X) * SourceLayoutGridSizeInPixels.X, 
+									(SrcBlock.Min.Y + SrcBlock.Size.Y) * SourceLayoutGridSizeInPixels.Y);
+
+							FIntRect& DstRectInPixels = DstRects.Emplace_GetRef(
+									DstBlock.Min.X * ComposeLayoutGridSizeInPixels.X, 
+									DstBlock.Min.Y * ComposeLayoutGridSizeInPixels.Y, 
+									(DstBlock.Min.X + DstBlock.Size.X) * ComposeLayoutGridSizeInPixels.X, 
+									(DstBlock.Min.Y + DstBlock.Size.Y) * ComposeLayoutGridSizeInPixels.Y);
+
+							float BestLODForRect = FMath::Max(0.0f,
+									ImageMultiComposeComputeBestLODForSrcRect(DstRectInPixels, SrcRectInPixels));
+							
+							int32 SelectedLOD = FMath::FloorToInt32(BestLODForRect + LODRoundTolerance);
+
+							NumSrcLODsNeeded = FMath::Max(SelectedLOD + 1, NumSrcLODsNeeded);
+							SelectedLODForRect.Emplace(SelectedLOD);
+						}
+
+						if (++ElemIndex >= IterationEnd)
+						{
+							break;
+						}
+					}
+				}
+
+				// Fix src rect sizes for the selected LOD.
+				for (int32 RectIndex = SrcRects.Num() - 1; RectIndex >= 0; --RectIndex)
+				{
+					int32 SelectedLOD = SelectedLODForRect[RectIndex];
+
+					for (int32 LODIndex = 0; LODIndex < SelectedLOD; ++LODIndex)
+					{
+						SrcRects[RectIndex] = FIntRect::DivideAndRoundUp(SrcRects[RectIndex], 2);
+					}
+				}
+
+				ImageMultiCompose(*ResultImage, *SourceImage, DstRects, SrcRects, SelectedLODForRect, Settings.ImageCompressionQuality);
+				ResultImage->DataStorage.SetNumLODs(1);
+
+				if (ElemIndex >= Data.MultiCompose.NumElemsToProcess)
+				{
+					HeapData[Item.CustomState + 0].Resource = nullptr;
+					HeapData[Item.CustomState + 1].Resource = nullptr;
+					HeapData[Item.CustomState + 2].Resource = nullptr;
+					HeapData[Item.CustomState + 3].Resource = nullptr;
+
+					StoreImage(Item, ResultImage);
+				}
+				else
+				{
+					AddOp(FScheduledOp(Item.At, Item, Item.Stage + 1, Item.CustomState));
+				}
+			}
 			}	
 
 			break;
diff --git a/Engine/Plugins/Mutable/Source/MutableRuntime/Private/MuR/CodeRunner.h b/Engine/Plugins/Mutable/Source/MutableRuntime/Private/MuR/CodeRunner.h
--- a/Engine/Plugins/Mutable/Source/MutableRuntime/Private/MuR/CodeRunner.h
+++ b/Engine/Plugins/Mutable/Source/MutableRuntime/Private/MuR/CodeRunner.h
@@ -107,10 +107,12 @@ namespace  UE::Mutable::Private
 
 				struct
 				{
-					int32 NextHeapIndex;
-					bool bHasBlocksToCompose;
-					uint8 SourceLOD;
-					uint8 SourceNumLODs;
+					uint16 NumElemsToProcess;
+					uint8 SourceLOD : 7;
+					uint8 bHasBlocksToCompose : 1;
+					
+					uint8 SourceNumLODs : 7;
+					uint8 bSplitInMultipleIterations : 1;
 				} MultiCompose;
 
 			};
diff --git a/Engine/Plugins/Mutable/Source/MutableRuntime/Private/MuR/OpImageMultiCompose.cpp b/Engine/Plugins/Mutable/Source/MutableRuntime/Private/MuR/OpImageMultiCompose.cpp
--- a/Engine/Plugins/Mutable/Source/MutableRuntime/Private/MuR/OpImageMultiCompose.cpp
+++ b/Engine/Plugins/Mutable/Source/MutableRuntime/Private/MuR/OpImageMultiCompose.cpp
@@ -133,17 +133,23 @@ namespace OpImageMultiComposeInternal
 				((DstRect.Max.Y % FormatInfo.Dst.PixelsPerBlockY) == 0);
 
 		const bool bUncompressedFormatMissmatch =
-				FormatInfo.SrcDecompressedFormat != FormatInfo.DstDecompressedFormat && 
-				FormatInfo.SrcDecompressedFormat == Source.GetFormat() && FormatInfo.DstDecompressedFormat == Base.GetFormat();
+				(FormatInfo.SrcDecompressedFormat != FormatInfo.DstDecompressedFormat) &
+				((FormatInfo.SrcDecompressedFormat == Source.GetFormat()) & (FormatInfo.DstDecompressedFormat == Base.GetFormat()));
 
 		const bool bNeedsResize = 
 				(DstRect.Width()  != SrcRect.Width()) |
 				(DstRect.Height() != SrcRect.Height());
 
+		const bool bDirectDecompression = !bNeedsResize                                &
+				(FormatInfo.SrcDecompressedFormat == FormatInfo.DstDecompressedFormat) &
+				(FormatInfo.DstDecompressedFormat == Base.GetFormat())                 &
+				(FormatInfo.SrcDecompressedFormat != Source.GetFormat());
+
 		const bool bNeedsSrcDecompression = 
 				(bNeedsResize & (FormatInfo.SrcDecompressedFormat != Source.GetFormat())) |
 				(!(bIsSrcCompressionRectAligned & bIsDstCompressionRectAligned))          |
-				bUncompressedFormatMissmatch;
+				(bUncompressedFormatMissmatch)                                            |
+				(bDirectDecompression);
 
 		const bool bNeedsDstDecompression =
 				(bNeedsResize & (FormatInfo.DstDecompressedFormat != Base.GetFormat())) |
@@ -564,6 +570,7 @@ namespace OpImageMultiComposeInternal
 							FormatInfo.DstDecompressionFunc(FromSize, ToSize, SubSize, FromPtr, ToPtr);	
 						}
 
+						bool bDstDataAlreadySet = false;
 						if (EnumHasAnyFlags(WorkingSubImages.ActiveStages, EActiveStages::SrcDecompression))
 						{
 							MUTABLE_CPUPROFILER_SCOPE(ImageMultiCompose_SrcDecompression);
@@ -588,8 +595,11 @@ namespace OpImageMultiComposeInternal
 
 								FormatInfo.SrcDecompressionFunc(FromSize, ToSize, SubSize, FromPtr, ToPtr);
 							}
+
+							bDstDataAlreadySet = !EnumHasAnyFlags(ActiveStages, EActiveStages::DstCompression | EActiveStages::Resize);
 						}
 
+						if (!bDstDataAlreadySet)
 						{
 							TConstSubImage<uint8> ConstDecSrc = WorkingSubImages.DecSrc.ToConst();
 
@@ -616,7 +626,7 @@ namespace OpImageMultiComposeInternal
 								FormatInfo.ResizeFunc(OneOverScalingFactor, WorkingSubImages.SrcSubPixelOffset, 
 										StageDst.Size, StageDst.SubSize, StageSrc.Size, StageSrc.SubSize, ToPtr, FromPtr);	
 							}
-							else 
+							else
 							{
 								MUTABLE_CPUPROFILER_SCOPE(ImageMultiCompose_Copy);
 								CopySubImage(StageDst, StageSrc);
