Fix static analysis error.

#rnx

#Synced-CL 49840162

diff --git a/Engine/Plugins/Mutable/Source/MutableRuntime/Private/MuR/OpImageMultiCompose.cpp b/Engine/Plugins/Mutable/Source/MutableRuntime/Private/MuR/OpImageMultiCompose.cpp
--- a/Engine/Plugins/Mutable/Source/MutableRuntime/Private/MuR/OpImageMultiCompose.cpp
+++ b/Engine/Plugins/Mutable/Source/MutableRuntime/Private/MuR/OpImageMultiCompose.cpp
@@ -140,7 +140,7 @@ namespace OpImageMultiComposeInternal
 				(DstRect.Width()  != SrcRect.Width()) |
 				(DstRect.Height() != SrcRect.Height());
 
-		const bool bDirectDecompression = !bNeedsResize                                &
+		const bool bDirectDecompression = (!bNeedsResize)                              &
 				(FormatInfo.SrcDecompressedFormat == FormatInfo.DstDecompressedFormat) &
 				(FormatInfo.DstDecompressedFormat == Base.GetFormat())                 &
 				(FormatInfo.SrcDecompressedFormat != Source.GetFormat());
