[Mutable] CodeRunner Removed Parameter and Model references outside LiveInstance.

#Synced-CL 49839474

diff --git a/Engine/Plugins/Mutable/Source/CustomizableObject/Private/MuCO/CustomizableObjectSystem.cpp b/Engine/Plugins/Mutable/Source/CustomizableObject/Private/MuCO/CustomizableObjectSystem.cpp
--- a/Engine/Plugins/Mutable/Source/CustomizableObject/Private/MuCO/CustomizableObjectSystem.cpp
+++ b/Engine/Plugins/Mutable/Source/CustomizableObject/Private/MuCO/CustomizableObjectSystem.cpp
@@ -1841,7 +1841,7 @@ namespace impl
 		// Update the parameter relevancy.
 		{
 			MUTABLE_CPUPROFILER_SCOPE(ParameterRelevancy)
-			OperationData->MutableSystem->GetParameterRelevancy(OperationData->LiveInstance.ToSharedRef(), OperationData->Parameters, OperationData->RelevantParameters);
+			OperationData->MutableSystem->GetParameterRelevancy(OperationData->LiveInstance.ToSharedRef(), OperationData->RelevantParameters);
 		}
 
 		OperationData->TaskGetMeshesTime = FPlatformTime::Seconds() - OperationData->TaskGetMeshesTimeStart;
@@ -2234,7 +2234,7 @@ namespace impl
 					TSharedRef<FInstanceUpdateData::FImage>& Image = Context->InstanceUpdateData.Images.Add_GetRef(MakeShared<FInstanceUpdateData::FImage>());
 					Image->ParameterName = ImageParameter.Key.ParameterName;
 					Image->MaterialLayer = ImageParameter.Key.LayerIndex;
-					Image->ImageID = UE::Mutable::Private::FSystem::GetImageId(Context->LiveInstance->ImageIdRegistry, Context->Model, Context->Parameters.Get(), Address);
+					Image->ImageID = UE::Mutable::Private::FSystem::GetImageId(Context->LiveInstance->ImageIdRegistry, *Context->Model.Get(), *Context->Parameters.Get(), Address);
 					Image->FullImageSizeX = 0;
 					Image->FullImageSizeY = 0;
 					Image->BaseMip = 0;
diff --git a/Engine/Plugins/Mutable/Source/CustomizableObjectEditor/Private/MuCOE/SMutableCodeViewer.cpp b/Engine/Plugins/Mutable/Source/CustomizableObjectEditor/Private/MuCOE/SMutableCodeViewer.cpp
--- a/Engine/Plugins/Mutable/Source/CustomizableObjectEditor/Private/MuCOE/SMutableCodeViewer.cpp
+++ b/Engine/Plugins/Mutable/Source/CustomizableObjectEditor/Private/MuCOE/SMutableCodeViewer.cpp
@@ -2899,13 +2899,14 @@ void SMutableCodeViewer::Tick(const FGeometry& AllottedGeometry, const double In
 
 	System->BeginBuild(MutableModel);
 	TSharedRef<UE::Mutable::Private::FLiveInstance> TempLiveInstance = System->NewBuildInstance(MutableModel, ExternalResourceProvider);
+	TempLiveInstance->OldParameters = PreviewParameters;
 	
 	switch (OperationDataType)
 	{
 	case UE::Mutable::Private::EDataType::Layout:
 	{
 		check(PreviewLayoutViewer);
-		UE::Mutable::Private::TManagedPtr<const UE::Mutable::Private::FLayout> MutableLayout = System->BuildLayout(TempLiveInstance, PreviewParameters.Get(), SelectedOperationAddress);
+		UE::Mutable::Private::TManagedPtr<const UE::Mutable::Private::FLayout> MutableLayout = System->BuildLayout(TempLiveInstance, SelectedOperationAddress);
 		PreviewLayoutViewer->SetLayout(MutableLayout);
 		break;
 	}
@@ -2913,7 +2914,7 @@ void SMutableCodeViewer::Tick(const FGeometry& AllottedGeometry, const double In
 	case UE::Mutable::Private::EDataType::Image:
 	{
 		check(PreviewImageViewer);
-		UE::Mutable::Private::TManagedPtr<const UE::Mutable::Private::FImage> MutableImage = System->BuildImage(TempLiveInstance, PreviewParameters.Get(), SelectedOperationAddress, MipsToSkip);
+		UE::Mutable::Private::TManagedPtr<const UE::Mutable::Private::FImage> MutableImage = System->BuildImage(TempLiveInstance, SelectedOperationAddress, MipsToSkip);
 		PreviewImageViewer->SetImage(MutableImage, 0);
 		break;
 	}
@@ -2921,7 +2922,7 @@ void SMutableCodeViewer::Tick(const FGeometry& AllottedGeometry, const double In
 	case UE::Mutable::Private::EDataType::Mesh:
 	{
 		check(PreviewMeshViewer);
-		UE::Mutable::Private::TManagedPtr<const UE::Mutable::Private::FMesh> MutableMesh = System->BuildMesh(TempLiveInstance, PreviewParameters.Get(), SelectedOperationAddress, UE::Mutable::Private::EMeshContentFlags::AllFlags);
+		UE::Mutable::Private::TManagedPtr<const UE::Mutable::Private::FMesh> MutableMesh = System->BuildMesh(TempLiveInstance, SelectedOperationAddress, UE::Mutable::Private::EMeshContentFlags::AllFlags);
 		PreviewMeshViewer->SetMesh(MutableMesh);
 		break;
 	}
@@ -2929,7 +2930,7 @@ void SMutableCodeViewer::Tick(const FGeometry& AllottedGeometry, const double In
 	case UE::Mutable::Private::EDataType::Instance:
 	{
 		check(PreviewInstanceViewer);
-		UE::Mutable::Private::TManagedPtr<const UE::Mutable::Private::FInstance> MutableInstance = System->BuildInstance(TempLiveInstance, PreviewParameters.Get(), SelectedOperationAddress);
+		UE::Mutable::Private::TManagedPtr<const UE::Mutable::Private::FInstance> MutableInstance = System->BuildInstance(TempLiveInstance, SelectedOperationAddress);
 		PreviewInstanceViewer->SetInstance(TempLiveInstance, MutableInstance, PreviewParameters, *System);
 		break;
 	}
@@ -2937,7 +2938,7 @@ void SMutableCodeViewer::Tick(const FGeometry& AllottedGeometry, const double In
 	case UE::Mutable::Private::EDataType::Bool:
 	{
 		check(PreviewBoolViewer);
-		const bool MutableBool = System->BuildBool(TempLiveInstance, PreviewParameters.Get(), SelectedOperationAddress);
+		const bool MutableBool = System->BuildBool(TempLiveInstance, SelectedOperationAddress);
 		PreviewBoolViewer->SetBool(MutableBool);
 		break;
 	}
@@ -2945,7 +2946,7 @@ void SMutableCodeViewer::Tick(const FGeometry& AllottedGeometry, const double In
 	case UE::Mutable::Private::EDataType::Int:
 	{
 		check(PreviewIntViewer);
-		const int32 MutableInt = System->BuildInt(TempLiveInstance, PreviewParameters.Get(), SelectedOperationAddress);
+		const int32 MutableInt = System->BuildInt(TempLiveInstance, SelectedOperationAddress);
 		PreviewIntViewer->SetInt(MutableInt);
 		break;
 	}
@@ -2953,7 +2954,7 @@ void SMutableCodeViewer::Tick(const FGeometry& AllottedGeometry, const double In
 	case UE::Mutable::Private::EDataType::Scalar:
 	{
 		check(PreviewScalarViewer);
-		const float MutableScalar = System->BuildScalar(TempLiveInstance, PreviewParameters.Get(), SelectedOperationAddress);
+		const float MutableScalar = System->BuildScalar(TempLiveInstance, SelectedOperationAddress);
 		PreviewScalarViewer->SetScalar(MutableScalar);
 		break;
 	}
@@ -2961,7 +2962,7 @@ void SMutableCodeViewer::Tick(const FGeometry& AllottedGeometry, const double In
 	case UE::Mutable::Private::EDataType::String:
 	{
 		check(PreviewStringViewer);
-		UE::Mutable::Private::TManagedPtr<const UE::Mutable::Private::String> MutableString = System->BuildString(TempLiveInstance, PreviewParameters.Get(), SelectedOperationAddress);
+		UE::Mutable::Private::TManagedPtr<const UE::Mutable::Private::String> MutableString = System->BuildString(TempLiveInstance, SelectedOperationAddress);
 		const FText MutableText = FText::FromString(FString(MutableString->GetValue()));
 		PreviewStringViewer->SetString(MutableText);
 		break;
@@ -2970,7 +2971,7 @@ void SMutableCodeViewer::Tick(const FGeometry& AllottedGeometry, const double In
 	case UE::Mutable::Private::EDataType::Color:
 	{
 		check(PreviewColorViewer);
-		FVector4f Color = System->BuildColor(TempLiveInstance, PreviewParameters.Get(), SelectedOperationAddress);
+		FVector4f Color = System->BuildColor(TempLiveInstance, SelectedOperationAddress);
 		PreviewColorViewer->SetColor(Color);
 		break;
 	}
@@ -2978,7 +2979,7 @@ void SMutableCodeViewer::Tick(const FGeometry& AllottedGeometry, const double In
 	case UE::Mutable::Private::EDataType::Projector:
 	{
 		check (PreviewProjectorViewer);
-		UE::Mutable::Private::FProjector Projector = System->BuildProjector(TempLiveInstance, PreviewParameters.Get(), SelectedOperationAddress);
+		UE::Mutable::Private::FProjector Projector = System->BuildProjector(TempLiveInstance, SelectedOperationAddress);
 		PreviewProjectorViewer->SetProjector(Projector);
 		break;
 	}
@@ -3010,8 +3011,11 @@ UE::Mutable::Private::TManagedPtr<const UE::Mutable::Private::FImage> SMutableCo
 	const TSharedPtr<UE::Mutable::Private::FSystem> System = MakeShared<UE::Mutable::Private::FSystem>(Settings);
 	
 	System->BeginBuild(MutableModel);
+	
 	TSharedRef<UE::Mutable::Private::FLiveInstance> TempLiveInstance = System->NewBuildInstance(MutableModel, ExternalResourceProvider);
-	UE::Mutable::Private::TManagedPtr<const UE::Mutable::Private::FImage> MutableImage = System->BuildImage(TempLiveInstance, PreviewParameters.Get(), TargetOpAddress, 0);
+	TempLiveInstance->OldParameters = PreviewParameters;
+	
+	UE::Mutable::Private::TManagedPtr<const UE::Mutable::Private::FImage> MutableImage = System->BuildImage(TempLiveInstance, TargetOpAddress, 0);
 	System->EndBuild();
 	
 	return MutableImage;
diff --git a/Engine/Plugins/Mutable/Source/CustomizableObjectEditor/Private/MuCOE/SMutableInstanceViewer.cpp b/Engine/Plugins/Mutable/Source/CustomizableObjectEditor/Private/MuCOE/SMutableInstanceViewer.cpp
--- a/Engine/Plugins/Mutable/Source/CustomizableObjectEditor/Private/MuCOE/SMutableInstanceViewer.cpp
+++ b/Engine/Plugins/Mutable/Source/CustomizableObjectEditor/Private/MuCOE/SMutableInstanceViewer.cpp
@@ -323,7 +323,6 @@ void SMutableInstanceViewer::SetInstance(const TSharedRef<UE::Mutable::Private::
 						UE::Mutable::Private::OP::ADDRESS MeshRootAddress = MeshId.GetKey()->Address;
 						UE::Mutable::Private::TManagedPtr<const UE::Mutable::Private::FMesh> MutableMesh = System.BuildMesh(
 								InLiveInstance,
-								Parameters.Get(), 
 								MeshRootAddress, 
 								UE::Mutable::Private::EMeshContentFlags::AllFlags);
 
diff --git a/Engine/Plugins/Mutable/Source/MutableRuntime/Internal/MuR/CodeVisitor.h b/Engine/Plugins/Mutable/Source/MutableRuntime/Internal/MuR/CodeVisitor.h
--- a/Engine/Plugins/Mutable/Source/MutableRuntime/Internal/MuR/CodeVisitor.h
+++ b/Engine/Plugins/Mutable/Source/MutableRuntime/Internal/MuR/CodeVisitor.h
@@ -220,7 +220,6 @@ namespace UE::Mutable::Private
             (
                 FSystem* InSystem,
                 const TSharedRef<FLiveInstance>& InLiveInstance,
-                const TSharedPtr<const FParameters>& InParams,
                 unsigned InLodMask,
                 bool bSkipResources=false
             )
@@ -229,7 +228,6 @@ namespace UE::Mutable::Private
         {
             System = InSystem;
             Model = InLiveInstance->Model;
-            Params = InParams.Get();
             LODMask = InLodMask;
 
             // Visiting state
@@ -273,7 +271,7 @@ namespace UE::Mutable::Private
 
                 if (Args.condition)
                 {
-					bValue = System->BuildBool(LiveInstance, Params, Args.condition);
+					bValue = System->BuildBool(LiveInstance, Args.condition);
                 }
 
                 if (bValue)
@@ -316,7 +314,7 @@ namespace UE::Mutable::Private
 
                     PARENT::RecurseWithCurrentState( VarAddress );
 
-                    int32 Var = System->BuildInt(LiveInstance, Params, VarAddress );
+                    int32 Var = System->BuildInt(LiveInstance, VarAddress );
 
 					OP::ADDRESS ValueAt = DefAddress;
 
@@ -450,7 +448,6 @@ namespace UE::Mutable::Private
         FSystem* System = nullptr;
 		TSharedPtr<const FModel> Model;
     	TSharedRef<FLiveInstance> LiveInstance;
-        const FParameters* Params = nullptr;
         uint32 LODMask = 0;
     };
 
@@ -480,10 +477,9 @@ namespace UE::Mutable::Private
             (
                 FSystem* InSystem,
 				const TSharedRef<FLiveInstance>& InLiveInstance,
-                const TSharedPtr<const FParameters>& InParams,
                 uint32 InLodMask
             )
-            : PARENT(InSystem, InLiveInstance, InParams, InLodMask )
+            : PARENT(InSystem, InLiveInstance, InLodMask)
         {
         }
 
diff --git a/Engine/Plugins/Mutable/Source/MutableRuntime/Internal/MuR/System.h b/Engine/Plugins/Mutable/Source/MutableRuntime/Internal/MuR/System.h
--- a/Engine/Plugins/Mutable/Source/MutableRuntime/Internal/MuR/System.h
+++ b/Engine/Plugins/Mutable/Source/MutableRuntime/Internal/MuR/System.h
@@ -1940,17 +1940,15 @@ namespace UE::Mutable::Private
 		//! Calculate the relevancy of every parameter. Some parameters may be unused depending on
 		//! the values of other parameters. This is useful to hide irrelevant parameters in
 		//! dynamic user interfaces.
-        UE_API void GetParameterRelevancy(const TSharedRef<FLiveInstance>& LiveInstance,
-        	const TSharedPtr<const FParameters>& FParameters,
-			TArray<uint32>& OutRelevantParameters);
+        UE_API void GetParameterRelevancy(const TSharedRef<FLiveInstance>& LiveInstance, TArray<uint32>& OutRelevantParameters);
 
-		UE_API static FMeshId GetMeshId(const TSharedPtr<FMeshIdRegistry>& Registry, const TSharedPtr<const FModel>& Model, const FParameters* Params, OP::ADDRESS RootAt);
+		UE_API static FMeshId GetMeshId(const TSharedPtr<FMeshIdRegistry>& Registry, const FModel& Model, const FParameters& Params, OP::ADDRESS RootAt);
     	
-		UE_API static FImageId GetImageId(const TSharedPtr<FImageIdRegistry>& Registry, const TSharedPtr<const FModel>& Model, const FParameters* Params, OP::ADDRESS RootAt);
+		UE_API static FImageId GetImageId(const TSharedPtr<FImageIdRegistry>& Registry, const FModel& Model, const FParameters& Params, OP::ADDRESS RootAt);
 
-		UE_API static FMaterialId GetMaterialId(const TSharedPtr<FMaterialIdRegistry>& Registry, const TSharedPtr<const FModel>& Model, const FParameters* Params, OP::ADDRESS RootAt);
+		UE_API static FMaterialId GetMaterialId(const TSharedPtr<FMaterialIdRegistry>& Registry, const FModel& Model, const FParameters& Params, OP::ADDRESS RootAt);
 
-		UE_API static FSkeletalMeshId GetSkeletalMeshId(const TSharedPtr<FSkeletalMeshIdRegistry>& Registry, const TSharedPtr<const FModel>& Model, const FParameters* Params, OP::ADDRESS RootAt);
+		UE_API static FSkeletalMeshId GetSkeletalMeshId(const TSharedPtr<FSkeletalMeshIdRegistry>& Registry, const FModel& Model, const FParameters& Params, OP::ADDRESS RootAt);
 		
         // Prevent copy, move and assignment.
         FSystem( const FSystem& ) = delete;
@@ -1963,16 +1961,16 @@ namespace UE::Mutable::Private
     	
 		MUTABLERUNTIME_API void EndBuild();
 
-		MUTABLERUNTIME_API bool BuildBool(const TSharedRef<FLiveInstance>& InLiveInstance, const FParameters* InParameters, OP::ADDRESS InAddress) ;
-		MUTABLERUNTIME_API int32 BuildInt(const TSharedRef<FLiveInstance>& InLiveInstance, const FParameters* InParameters, OP::ADDRESS InAddress) ;
-		MUTABLERUNTIME_API float BuildScalar(const TSharedRef<FLiveInstance>& InLiveInstance, const FParameters* InParameters, OP::ADDRESS InAddress) ;
-		MUTABLERUNTIME_API FVector4f BuildColor(const TSharedRef<FLiveInstance>& InLiveInstance, const FParameters* InParameters, OP::ADDRESS InAddress) ;
-		MUTABLERUNTIME_API TManagedPtr<const String> BuildString(const TSharedRef<FLiveInstance>& InLiveInstance, const FParameters* InParameters, OP::ADDRESS InAddress) ;
-		MUTABLERUNTIME_API TManagedPtr<const FImage> BuildImage(const TSharedRef<FLiveInstance>& InLiveInstance, const FParameters* InParameters, OP::ADDRESS InAddress, int32 MipsToSkip);
-		MUTABLERUNTIME_API TManagedPtr<const FMesh> BuildMesh(const TSharedRef<FLiveInstance>& InLiveInstance, const FParameters* InParameters, OP::ADDRESS InRootAddress, EMeshContentFlags MeshContentFilter);
-		MUTABLERUNTIME_API TManagedPtr<const FInstance> BuildInstance(const TSharedRef<FLiveInstance>& InLiveInstance, const FParameters* InParameters, OP::ADDRESS InAddress);
-		MUTABLERUNTIME_API TManagedPtr<const FLayout> BuildLayout(const TSharedRef<FLiveInstance>& InLiveInstance, const FParameters* InParameters, OP::ADDRESS InAddress) ;
-		MUTABLERUNTIME_API FProjector BuildProjector(const TSharedRef<FLiveInstance>& InLiveInstance, const FParameters* InParameters, OP::ADDRESS InAddress) ;
+		MUTABLERUNTIME_API bool BuildBool(const TSharedRef<FLiveInstance>& InLiveInstance, OP::ADDRESS InAddress) ;
+		MUTABLERUNTIME_API int32 BuildInt(const TSharedRef<FLiveInstance>& InLiveInstance, OP::ADDRESS InAddress) ;
+		MUTABLERUNTIME_API float BuildScalar(const TSharedRef<FLiveInstance>& InLiveInstance, OP::ADDRESS InAddress) ;
+		MUTABLERUNTIME_API FVector4f BuildColor(const TSharedRef<FLiveInstance>& InLiveInstance, OP::ADDRESS InAddress) ;
+		MUTABLERUNTIME_API TManagedPtr<const String> BuildString(const TSharedRef<FLiveInstance>& InLiveInstance, OP::ADDRESS InAddress) ;
+		MUTABLERUNTIME_API TManagedPtr<const FImage> BuildImage(const TSharedRef<FLiveInstance>& InLiveInstance, OP::ADDRESS InAddress, int32 MipsToSkip);
+		MUTABLERUNTIME_API TManagedPtr<const FMesh> BuildMesh(const TSharedRef<FLiveInstance>& InLiveInstance, OP::ADDRESS InRootAddress, EMeshContentFlags MeshContentFilter);
+		MUTABLERUNTIME_API TManagedPtr<const FInstance> BuildInstance(const TSharedRef<FLiveInstance>& InLiveInstance, OP::ADDRESS InAddress);
+		MUTABLERUNTIME_API TManagedPtr<const FLayout> BuildLayout(const TSharedRef<FLiveInstance>& InLiveInstance, OP::ADDRESS InAddress) ;
+		MUTABLERUNTIME_API FProjector BuildProjector(const TSharedRef<FLiveInstance>& InLiveInstance, OP::ADDRESS InAddress) ;
 		
 		
         FSettings Settings;
@@ -1994,7 +1992,6 @@ namespace UE::Mutable::Private
 		
 		void RunCode(
 			const TSharedRef<FLiveInstance>& InLiveInstance,
-			const FParameters*,
 			OP::ADDRESS,
 			uint32 LODs = AllLODs,
 			uint8 ExecutionOptions = 0);
diff --git a/Engine/Plugins/Mutable/Source/MutableRuntime/Private/MuR/CodeRunner.cpp b/Engine/Plugins/Mutable/Source/MutableRuntime/Private/MuR/CodeRunner.cpp
--- a/Engine/Plugins/Mutable/Source/MutableRuntime/Private/MuR/CodeRunner.cpp
+++ b/Engine/Plugins/Mutable/Source/MutableRuntime/Private/MuR/CodeRunner.cpp
@@ -162,9 +162,8 @@ namespace UE::Mutable::Private
 	TSharedRef<CodeRunner> CodeRunner::Create(
 		const TSharedRef<FLiveInstance>& InLiveInstance,
 		const FSettings& InSettings,
-		FSystem* InSystem,
+		FSystem& InSystem,
 		EExecutionStrategy InExecutionStrategy,
-		const FParameters* InParams,
 		OP::ADDRESS At,
 		uint32 InLODMask,
 		uint8 ExecutionOptions,
@@ -172,7 +171,7 @@ namespace UE::Mutable::Private
 		FScheduledOp::EType Type)
 	{
 		return MakeShared<CodeRunner>(InLiveInstance, FPrivateToken {},
-				InSettings, InSystem, InExecutionStrategy,  InParams, At, InLODMask, ExecutionOptions, bForceLazy, Type);
+				InSettings, InSystem, InExecutionStrategy, At, InLODMask, ExecutionOptions, bForceLazy, Type);
 	}
 
 
@@ -180,9 +179,8 @@ namespace UE::Mutable::Private
 		const TSharedRef<FLiveInstance>& InLiveInstance,
     	FPrivateToken PrivateToken, 
 		const FSettings& InSettings,
-		FSystem* InSystem,
+		FSystem& InSystem,
 		EExecutionStrategy InExecutionStrategy,
-		const FParameters* InParams,
 		OP::ADDRESS At,
 		uint32 InLodMask,
 		uint8 ExecutionOptions,
@@ -192,15 +190,14 @@ namespace UE::Mutable::Private
 		, LiveInstance(InLiveInstance)
 		, ExecutionStrategy(InExecutionStrategy)
 		, System(InSystem)
-		, Model(InLiveInstance->Model)
-		, Params(InParams)
+		, Model(*InLiveInstance->Model.Get())
+		, Program(InLiveInstance->Model->Program)
+		, Parameters(*InLiveInstance->OldParameters)
 		, LODMask(InLodMask)
 		, ExternalResourceProvider(InLiveInstance->ExternalResourceProvider)
 	{
 		MUTABLE_CPUPROFILER_SCOPE(CodeRunner_Create)
-		const FProgram& Program = Model->Program;
 		ScheduledStagePerOp.resize(Program.OpAddress.Num());
-
 		
 		if (Type == FScheduledOp::EType::ImageDesc)
 		{
@@ -289,8 +286,6 @@ namespace UE::Mutable::Private
 	{
 		MUTABLE_CPUPROFILER_SCOPE(GetExternalImageDesc);
 
-		check(System);
-
 		if (ExternalResourceProvider)
 		{
 			return ExternalResourceProvider->GetImageDesc(Image);
@@ -306,12 +301,10 @@ namespace UE::Mutable::Private
 
 	
     //---------------------------------------------------------------------------------------------
-    void CodeRunner::RunCode_Conditional( const FScheduledOp& Item, const FModel* InModel )
+    void CodeRunner::RunCode_Conditional( const FScheduledOp& Item)
     {
         MUTABLE_CPUPROFILER_SCOPE(RunCode_Conditional);
-
-		const FProgram& Program = Model->Program;
-
+		
 		EOpType type = Program.GetOpType(Item.At);
 		OP::ConditionalArgs Args = Program.GetOpArgs<OP::ConditionalArgs>(Item.At);
 
@@ -381,10 +374,8 @@ namespace UE::Mutable::Private
 
 
 	//---------------------------------------------------------------------------------------------
-	void CodeRunner::RunCode_Switch(const FScheduledOp& Item, const FModel* InModel )
+	void CodeRunner::RunCode_Switch(const FScheduledOp& Item)
 	{
-		const FProgram& Program = Model->Program;
-
 		EOpType Type = Program.GetOpType(Item.At);
 
 		const uint8* Data = Program.GetOpArgsPointer(Item.At);
@@ -527,12 +518,10 @@ namespace UE::Mutable::Private
 
 
     //---------------------------------------------------------------------------------------------
-    void CodeRunner::RunCode_Instance(const FScheduledOp& Item, const FModel* InModel, uint32 lodMask )
+    void CodeRunner::RunCode_Instance(const FScheduledOp& Item, uint32 lodMask)
     {
         MUTABLE_CPUPROFILER_SCOPE(RunCode_Instance);
-
-		const FProgram& Program = Model->Program;
-
+		
 		EOpType type = Program.GetOpType(Item.At);
         switch (type)
         {
@@ -844,17 +833,10 @@ namespace UE::Mutable::Private
 
 
     //---------------------------------------------------------------------------------------------
-    void CodeRunner::RunCode_InstanceAddResource(const FScheduledOp& Item, const TSharedPtr<const FModel>& InModel, const FParameters* InParams )
+    void CodeRunner::RunCode_InstanceAddResource(const FScheduledOp& Item)
     {
 		MUTABLE_CPUPROFILER_SCOPE(RunCode_InstanceAddResource);
-
-		if (!InModel || !System)
-		{
-			return;
-		}
-
-		const FProgram& Program = Model->Program;
-
+		
 		EOpType type = Program.GetOpType(Item.At);
         switch (type)
         {
@@ -885,7 +867,7 @@ namespace UE::Mutable::Private
 
                 if ( Args.value )
                 {
-					FMeshId MeshId = FSystem::GetMeshId(LiveInstance->MeshIdRegistry, InModel, InParams, Args.value);
+					FMeshId MeshId = FSystem::GetMeshId(LiveInstance->MeshIdRegistry, Model, Parameters, Args.value);
 					OP::ADDRESS NameAd = Args.name;
 					check(NameAd < (uint32)Program.ConstantStrings.Num());
 					const FString& Name = Program.ConstantStrings[NameAd];
@@ -932,7 +914,7 @@ namespace UE::Mutable::Private
 
         				if (Args.Material)
         				{
-        					FMaterialId MaterialId = FSystem::GetMaterialId(LiveInstance->MaterialIdRegistry, InModel, InParams, Args.Material);
+        					FMaterialId MaterialId = FSystem::GetMaterialId(LiveInstance->MaterialIdRegistry, Model, Parameters, Args.Material);
         					Result->SetOverlayMaterialId(0, MaterialId);
         				}
 
@@ -978,7 +960,7 @@ namespace UE::Mutable::Private
 
         				if (Args.Material)
         				{
-        					FMaterialId MaterialId = FSystem::GetMaterialId(LiveInstance->MaterialIdRegistry, InModel, InParams, Args.Material);
+        					FMaterialId MaterialId = FSystem::GetMaterialId(LiveInstance->MaterialIdRegistry, Model, Parameters, Args.Material);
         					const FString& SlotName = Program.ConstantStrings[Args.SlotName];
         					Result->AddOverrideMaterial(0, FName(SlotName), MaterialId);
         				}
@@ -1024,7 +1006,7 @@ namespace UE::Mutable::Private
 
         				if (Args.Material)
         				{
-        					FMaterialId MaterialId = FSystem::GetMaterialId(LiveInstance->MaterialIdRegistry, InModel, InParams, Args.Material);
+        					FMaterialId MaterialId = FSystem::GetMaterialId(LiveInstance->MaterialIdRegistry, Model, Parameters, Args.Material);
 							Result->SetMaterialId(0, 0, 0, MaterialId);
         				}
 
@@ -1076,7 +1058,7 @@ namespace UE::Mutable::Private
 
         				if (Args.value)
         				{
-        					FSkeletalMeshId SkeletalMeshId = FSystem::GetSkeletalMeshId(LiveInstance->SkeletalMeshIdRegistry, InModel, InParams, Args.value);
+        					FSkeletalMeshId SkeletalMeshId = FSystem::GetSkeletalMeshId(LiveInstance->SkeletalMeshIdRegistry, Model, Parameters, Args.value);
         					Instance->SetSkeletalMeshId(0, SkeletalMeshId);
         				}
 						
@@ -1098,12 +1080,10 @@ namespace UE::Mutable::Private
 	
 
     //---------------------------------------------------------------------------------------------
-    bool CodeRunner::RunCode_ConstantResource(const FScheduledOp& Item, const FModel* InModel)
+    bool CodeRunner::RunCode_ConstantResource(const FScheduledOp& Item)
     {
 		MUTABLE_CPUPROFILER_SCOPE(RunCode_Constant);
-
-		const FProgram& Program = Model->Program;
-
+		
 		EOpType type = Program.GetOpType(Item.At);
         switch (type)
         {
@@ -1200,13 +1180,12 @@ namespace UE::Mutable::Private
 		return true;
     }
 
-	void CodeRunner::RunCode_Mesh(const FScheduledOp& Item, const FModel* InModel )
+	void CodeRunner::RunCode_Mesh(const FScheduledOp& Item)
     {
 		MUTABLE_CPUPROFILER_SCOPE(RunCode_Mesh);
 
 		check((EMeshContentFlags)Item.ExecutionOptions != EMeshContentFlags::None); 
 
-		const FProgram& Program = Model->Program;
 
 		EOpType type = Program.GetOpType(Item.At);
 
@@ -2108,7 +2087,7 @@ namespace UE::Mutable::Private
 				{
 					TManagedPtr<const FMesh> Source = LoadMesh(FCacheAddress(Args.Source, Item));
 					
-					check(Args.MorphShape < (uint32)InModel->Program.ConstantShapes.Num());
+					check(Args.MorphShape < (uint32)Program.ConstantShapes.Num());
 
 					// Should be an ellipse
 					const FShape& MorphShape = Program.ConstantShapes[Args.MorphShape];
@@ -2120,7 +2099,7 @@ namespace UE::Mutable::Private
 
 					if (Args.VertexSelectionType == EClipVertexSelectionType::Shape)
 					{
-						check(Args.VertexSelectionShapeOrBone < (uint32)InModel->Program.ConstantShapes.Num());
+						check(Args.VertexSelectionShapeOrBone < (uint32)Program.ConstantShapes.Num());
 
 						// Should be None or an axis aligned box
 						const FShape& SelectionShape = Program.ConstantShapes[Args.VertexSelectionShapeOrBone];
@@ -3357,14 +3336,12 @@ namespace UE::Mutable::Private
 		}
 	}
 
-    void CodeRunner::RunCode_Image(const FScheduledOp& Item, const FParameters* pParams, const FModel* InModel )
+    void CodeRunner::RunCode_Image(const FScheduledOp& Item)
     {
 		MUTABLE_CPUPROFILER_SCOPE(RunCode_Image);
 
 		FImageOperator ImOp = MakeImageOperator(this);
-
-		const FProgram& Program = Model->Program;
-
+		
 		EOpType type = Program.GetOpType(Item.At);
 		switch (type)
         {
@@ -3469,7 +3446,7 @@ namespace UE::Mutable::Private
 
                     // We support both integers and scalars here, which is not common.
                     // \todo: review if this is necessary or we can enforce it at compile time.
-                    EDataType RangeSizeType = GetOpDataType( InModel->Program.GetOpType(Args.rangeSize) );
+                    EDataType RangeSizeType = GetOpDataType(Program.GetOpType(Args.rangeSize) );
                     if (RangeSizeType == EDataType::Int)
                     {
 						Iterations = LoadInt(RangeAddress);
@@ -4930,7 +4907,7 @@ namespace UE::Mutable::Private
         {
 			OP::ParameterArgs Args = Program.GetOpArgs<OP::ParameterArgs>(Item.At);
 			TSharedPtr<FRangeIndex> Index = BuildCurrentOpRangeIndex(Item, Args.variable);
-			UTexture* Image = pParams->GetTextureValue(Args.variable, Index.Get());
+			UTexture* Image = Parameters.GetTextureValue(Args.variable, Index.Get());
 
        		TPassthroughObjectPtr<UTexture> ExternalObject(Image);
 			TManagedPtr<FImage> Result = FImage::CreateAsReference(ExternalObject, FImageDesc(), false);
@@ -5910,15 +5887,14 @@ namespace UE::Mutable::Private
 	
 	TSharedPtr<FRangeIndex> CodeRunner::BuildCurrentOpRangeIndex(const FScheduledOp& Item, int32 ParameterIndex)
 	{
-		return UE::Mutable::Private::BuildCurrentOpRangeIndex(Item, *Params, *Model.Get(), GetMemory(), ParameterIndex);
+		return UE::Mutable::Private::BuildCurrentOpRangeIndex(Item, Parameters, Model, GetMemory(), ParameterIndex);
 	}
 
     //---------------------------------------------------------------------------------------------
-    void CodeRunner::RunCode_Bool(const FScheduledOp& Item, const FParameters* pParams, const FModel* InModel )
+    void CodeRunner::RunCode_Bool(const FScheduledOp& Item)
     {
         MUTABLE_CPUPROFILER_SCOPE(RunCode_Bool);
 
-		const FProgram& Program = Model->Program;
 		EOpType type = Program.GetOpType(Item.At);
         switch (type)
         {
@@ -5936,7 +5912,7 @@ namespace UE::Mutable::Private
 			OP::ParameterArgs Args = Program.GetOpArgs<OP::ParameterArgs>(Item.At);
             bool result = false;
 			TSharedPtr<FRangeIndex> Index = BuildCurrentOpRangeIndex( Item, Args.variable );
-            result = pParams->GetBoolValue( Args.variable, Index.Get() );
+            result = Parameters.GetBoolValue( Args.variable, Index.Get() );
             StoreBool( Item, result );
             break;
         }
@@ -6127,12 +6103,10 @@ namespace UE::Mutable::Private
     }
 	
     //---------------------------------------------------------------------------------------------
-    void CodeRunner::RunCode_Int(const FScheduledOp& Item, const FParameters* pParams, const FModel* InModel )
+    void CodeRunner::RunCode_Int(const FScheduledOp& Item)
     {
         MUTABLE_CPUPROFILER_SCOPE(RunCode_Int);
 
-		const FProgram& Program = Model->Program;
-
 		EOpType type = Program.GetOpType(Item.At);
         switch (type)
         {
@@ -6149,17 +6123,15 @@ namespace UE::Mutable::Private
         {
 			OP::ParameterArgs Args = Program.GetOpArgs<OP::ParameterArgs>(Item.At);
 			TSharedPtr<FRangeIndex> Index = BuildCurrentOpRangeIndex( Item, Args.variable );
-            int32 result = pParams->GetIntValue( Args.variable, Index.Get());
+            int32 result = Parameters.GetIntValue( Args.variable, Index.Get());
 
             // Check that the value is actually valid. Otherwise set the default.
-            if ( pParams->GetIntPossibleValueCount( Args.variable ) )
+            if ( Parameters.GetIntPossibleValueCount( Args.variable ) )
             {
                 bool valid = false;
-                for ( int32 i=0;
-                      (!valid) && i<pParams->GetIntPossibleValueCount( Args.variable );
-                      ++i )
+                for (int32 ValueIndex = 0; !valid && ValueIndex < Parameters.GetIntPossibleValueCount(Args.variable); ++ValueIndex)
                 {
-                    if ( result == pParams->GetIntPossibleValue( Args.variable, i ) )
+                    if ( result == Parameters.GetIntPossibleValue( Args.variable, ValueIndex ))
                     {
                         valid = true;
                     }
@@ -6167,7 +6139,7 @@ namespace UE::Mutable::Private
 
                 if (!valid)
                 {
-                    result = pParams->GetIntPossibleValue( Args.variable, 0 );
+                    result = Parameters.GetIntPossibleValue( Args.variable, 0 );
                 }
             }
 
@@ -6182,12 +6154,10 @@ namespace UE::Mutable::Private
     }
 
     //---------------------------------------------------------------------------------------------
-    void CodeRunner::RunCode_Scalar(const FScheduledOp& Item, const FParameters* pParams, const FModel* InModel )
+    void CodeRunner::RunCode_Scalar(const FScheduledOp& Item)
     {
         MUTABLE_CPUPROFILER_SCOPE(RunCode_Scalar);
 
-		const FProgram& Program = Model->Program;
-
 		EOpType type = Program.GetOpType(Item.At);
         switch (type)
         {
@@ -6204,7 +6174,7 @@ namespace UE::Mutable::Private
         {
 			OP::ParameterArgs Args = Program.GetOpArgs<OP::ParameterArgs>(Item.At);
 			TSharedPtr<FRangeIndex> Index = BuildCurrentOpRangeIndex( Item, Args.variable );
-            float result = pParams->GetFloatValue( Args.variable, Index.Get());
+            float result = Parameters.GetFloatValue( Args.variable, Index.Get());
             StoreScalar( Item, result );
             break;
         }
@@ -6360,12 +6330,10 @@ namespace UE::Mutable::Private
     }
 
     //---------------------------------------------------------------------------------------------
-    void CodeRunner::RunCode_String(const FScheduledOp& Item, const FParameters* pParams, const FModel* InModel )
+    void CodeRunner::RunCode_String(const FScheduledOp& Item)
     {
         MUTABLE_CPUPROFILER_SCOPE(RunCode_String );
 
-		const FProgram& Program = Model->Program;
-
 		EOpType type = Program.GetOpType( Item.At );
         switch ( type )
         {
@@ -6373,7 +6341,7 @@ namespace UE::Mutable::Private
         case EOpType::ST_CONSTANT:
         {
 			OP::ResourceConstantArgs Args = Program.GetOpArgs<OP::ResourceConstantArgs>( Item.At );
-            check( Args.value < (uint32)InModel->Program.ConstantStrings.Num() );
+            check( Args.value < (uint32)Program.ConstantStrings.Num() );
 
             const FString& Result = Program.ConstantStrings[Args.value];
             StoreString( Item, MakeManaged<String>(Result) );
@@ -6386,7 +6354,7 @@ namespace UE::Mutable::Private
 			OP::ParameterArgs Args = Program.GetOpArgs<OP::ParameterArgs>( Item.At );
 			TSharedPtr<FRangeIndex> Index = BuildCurrentOpRangeIndex( Item, Args.variable );
 			FString Result;
-			pParams->GetStringValue(Args.variable, Result, Index.Get());
+			Parameters.GetStringValue(Args.variable, Result, Index.Get());
             StoreString( Item, MakeManaged<String>(Result) );
             break;
         }
@@ -6399,12 +6367,10 @@ namespace UE::Mutable::Private
 
 
     //---------------------------------------------------------------------------------------------
-    void CodeRunner::RunCode_Colour(const FScheduledOp& Item, const FParameters* pParams, const FModel* InModel )
+    void CodeRunner::RunCode_Colour(const FScheduledOp& Item)
     {
 		MUTABLE_CPUPROFILER_SCOPE(RunCode_Colour);
 
-		const FProgram& Program = Model->Program;
-
 		EOpType type = Program.GetOpType(Item.At);
 
         switch ( type )
@@ -6422,7 +6388,7 @@ namespace UE::Mutable::Private
 			OP::ParameterArgs Args = Program.GetOpArgs<OP::ParameterArgs>(Item.At);
 			TSharedPtr<FRangeIndex> Index = BuildCurrentOpRangeIndex( Item, Args.variable );
             FVector4f V;
-            pParams->GetColourValue( Args.variable, V, Index.Get());
+            Parameters.GetColourValue( Args.variable, V, Index.Get());
             StoreColor( Item, V );
             break;
         }
@@ -6710,12 +6676,10 @@ namespace UE::Mutable::Private
 
 
     //---------------------------------------------------------------------------------------------
-    void CodeRunner::RunCode_Projector(const FScheduledOp& Item, const FParameters* pParams, const FModel* InModel )
+    void CodeRunner::RunCode_Projector(const FScheduledOp& Item)
     {
         MUTABLE_CPUPROFILER_SCOPE(RunCode_Projector);
 
-		const FProgram& Program = Model->Program;
-
 		EOpType type = Program.GetOpType(Item.At);
         switch (type)
         {
@@ -6732,7 +6696,7 @@ namespace UE::Mutable::Private
         {
 			OP::ParameterArgs Args = Program.GetOpArgs<OP::ParameterArgs>(Item.At);
 			TSharedPtr<FRangeIndex> Index = BuildCurrentOpRangeIndex( Item, Args.variable );
-            FProjector Result = pParams->GetProjectorValue(Args.variable,Index.Get());
+            FProjector Result = Parameters.GetProjectorValue(Args.variable,Index.Get());
 
             // The type cannot be changed, take it from the default value
             const FProjector& def = Program.Parameters[Args.variable].DefaultValue.Get<FParamProjectorType>();
@@ -6749,12 +6713,10 @@ namespace UE::Mutable::Private
     }
 
     //---------------------------------------------------------------------------------------------
-    void CodeRunner::RunCode_Matrix(const FScheduledOp& Item, const FParameters* pParams, const FModel* InModel )
+    void CodeRunner::RunCode_Matrix(const FScheduledOp& Item)
     {
 		MUTABLE_CPUPROFILER_SCOPE(RunCode_Transform);
 
-		const FProgram& Program = Model->Program;
-
 		EOpType type = Program.GetOpType(Item.At);
 
 		switch ( type )
@@ -6771,7 +6733,7 @@ namespace UE::Mutable::Private
 				OP::ParameterArgs Args = Program.GetOpArgs<OP::ParameterArgs>(Item.At);
 				TSharedPtr<FRangeIndex> Index = BuildCurrentOpRangeIndex( Item, Args.variable );
 				FMatrix44f Value;
-				pParams->GetMatrixValue( Args.variable, Value, Index.Get());
+				Parameters.GetMatrixValue( Args.variable, Value, Index.Get());
 				StoreMatrix( Item, Value );
 				break;
 			}
@@ -6779,12 +6741,10 @@ namespace UE::Mutable::Private
     }
 
 	//---------------------------------------------------------------------------------------------
-	void CodeRunner::RunCode_Material(const FScheduledOp& Item, const FParameters* InParams, const FModel* InModel)
+	void CodeRunner::RunCode_Material(const FScheduledOp& Item)
 	{
 		MUTABLE_CPUPROFILER_SCOPE(RunCode_Material);
-
-		const FProgram& Program = Model->Program;
-
+		
 		EOpType type = Program.GetOpType(Item.At);
 
 		switch (type)
@@ -6797,7 +6757,7 @@ namespace UE::Mutable::Private
 			OP::MaterialParameterArgs Args = Program.GetOpArgs<OP::MaterialParameterArgs >(Item.At);
 			TSharedPtr<FRangeIndex> Index = BuildCurrentOpRangeIndex(Item, Args.Variable);
 
-			UMaterialInterface* MaterialInterface = InParams->GetMaterialValue(Args.Variable, Index.Get());
+			UMaterialInterface* MaterialInterface = Parameters.GetMaterialValue(Args.Variable, Index.Get());
 			TManagedPtr<FMaterial> Material = MakeManaged<FMaterial>();
 			
 			Material->PassthroughObject = TPassthroughObjectPtr(MaterialInterface);
@@ -7090,11 +7050,9 @@ namespace UE::Mutable::Private
 
 
     //---------------------------------------------------------------------------------------------
-    void CodeRunner::RunCode_Layout(const FScheduledOp& Item, const FModel* InModel )
+    void CodeRunner::RunCode_Layout(const FScheduledOp& Item)
     {
-        //MUTABLE_CPUPROFILER_SCOPE(RunCode_Layout);
-
-		const FProgram& Program = Model->Program;
+        MUTABLE_CPUPROFILER_SCOPE(RunCode_Layout);
 
 		EOpType type = Program.GetOpType(Item.At);
         switch (type)
@@ -7103,7 +7061,7 @@ namespace UE::Mutable::Private
         case EOpType::LA_CONSTANT:
         {
 			OP::ResourceConstantArgs Args = Program.GetOpArgs<OP::ResourceConstantArgs>(Item.At);
-            check( Args.value < (uint32)InModel->Program.ConstantLayouts.Num() );
+            check( Args.value < (uint32)Program.ConstantLayouts.Num() );
 
             TManagedPtr<const FLayout> Result = Program.ConstantLayouts[ Args.value ];
             StoreLayout(Item, Result);
@@ -7262,12 +7220,10 @@ namespace UE::Mutable::Private
 
 
     //---------------------------------------------------------------------------------------------
-    void CodeRunner::RunCode( const FScheduledOp& Item, const FParameters* pParams, const TSharedPtr<const FModel>& InModel, uint32 LodMask, bool& bOutFailed)
+    void CodeRunner::RunCode(const FScheduledOp& Item, uint32 LodMask, bool& bOutFailed)
     {
 		//UE_LOG( LogMutableCore, Log, TEXT("Running :%5d , %d "), Item.At, Item.Stage );
 		check( Item.Type == FScheduledOp::EType::Full );
-		
-		const FProgram& Program = InModel->Program;
 
 		EOpType type = Program.GetOpType(Item.At);
 		//UE_LOG(LogMutableCore, Log, TEXT("Running :%5d , %d, of type %d "), Item.At, Item.Stage, type);
@@ -7275,7 +7231,7 @@ namespace UE::Mutable::Private
 		// Very spammy, for debugging purposes.
 		//if (System)
 		//{
-		//	System->WorkingMemoryManager.LogWorkingMemory( this );
+		//	System.WorkingMemoryManager.LogWorkingMemory( this );
 		//}
 
 		switch ( type )
@@ -7292,14 +7248,14 @@ namespace UE::Mutable::Private
         case EOpType::IN_CONDITIONAL:
 		case EOpType::ED_CONDITIONAL:
 		case EOpType::MI_CONDITIONAL:
-            RunCode_Conditional(Item, InModel.Get());
+            RunCode_Conditional(Item);
             break;
 
         case EOpType::ME_CONSTANT:
 		case EOpType::IM_CONSTANT:
 		case EOpType::ED_CONSTANT:
 		case EOpType::MI_CONSTANT:
-            RunCode_ConstantResource(Item, InModel.Get());
+            RunCode_ConstantResource(Item);
             break;
 
         case EOpType::NU_SWITCH:
@@ -7311,7 +7267,7 @@ namespace UE::Mutable::Private
         case EOpType::IN_SWITCH:
 		case EOpType::ED_SWITCH:
 		case EOpType::MI_SWITCH:
-            RunCode_Switch(Item, InModel.Get());
+            RunCode_Switch(Item);
             break;
 
         case EOpType::IN_ADDMESH:
@@ -7319,7 +7275,7 @@ namespace UE::Mutable::Private
         case EOpType::IN_ADDOVERLAYMATERIAL:
         case EOpType::IN_ADDOVERRIDEMATERIAL:
     	case EOpType::IN_ADDSKELETALMESH:
-            RunCode_InstanceAddResource(Item, InModel, pParams);
+            RunCode_InstanceAddResource(Item);
             break;
 
 		case EOpType::SC_EXTERNAL:
@@ -7525,7 +7481,7 @@ namespace UE::Mutable::Private
 			{
 				OP::ParameterArgs Args = Program.GetOpArgs<OP::ParameterArgs>(Item.At);
 				TSharedPtr<FRangeIndex> Index = BuildCurrentOpRangeIndex(Item, Args.variable);
-				TManagedPtr<const FInstancedStruct> Result = pParams->GetInstancedStructValue(Args.variable, Index.Get());
+				TManagedPtr<const FInstancedStruct> Result = Parameters.GetInstancedStructValue(Args.variable, Index.Get());
 				StoreInstancedStruct(Item, Result);
 
 				break;
@@ -7535,7 +7491,7 @@ namespace UE::Mutable::Private
 			{
 				OP::ParameterArgs Args = Program.GetOpArgs<OP::ParameterArgs>(Item.At);
 				TSharedPtr<FRangeIndex> Index = BuildCurrentOpRangeIndex(Item, Args.variable);
-				USkeletalMesh* SkeletalMesh = pParams->GetSkeletalMeshValue(Args.variable, Index.Get());
+				USkeletalMesh* SkeletalMesh = Parameters.GetSkeletalMeshValue(Args.variable, Index.Get());
 
 				TManagedPtr<FSkeletalMesh> Result = MakeManaged<FSkeletalMesh>();
 				Result->PassthroughObject = TPassthroughObjectPtr<USkeletalMesh>(SkeletalMesh);
@@ -7550,53 +7506,53 @@ namespace UE::Mutable::Private
 			switch (DataType)
 			{
 			case EDataType::Instance:
-				RunCode_Instance(Item, InModel.Get(), LodMask);
+				RunCode_Instance(Item, LodMask);
 				break;
 
 			case EDataType::Mesh:
-				RunCode_Mesh(Item, InModel.Get());
+				RunCode_Mesh(Item);
 				EnsureBudgetBelow(0);
 				break;
 
 			case EDataType::Image:
-				RunCode_Image(Item, pParams, InModel.Get());
+				RunCode_Image(Item);
 				EnsureBudgetBelow(0);
 				break;
 
 			case EDataType::Layout:
-				RunCode_Layout(Item, InModel.Get());
+				RunCode_Layout(Item);
 				break;
 
 			case EDataType::Bool:
-				RunCode_Bool(Item, pParams, InModel.Get());
+				RunCode_Bool(Item);
 				break;
 
 			case EDataType::Scalar:
-				RunCode_Scalar(Item, pParams, InModel.Get());
+				RunCode_Scalar(Item);
 				break;
 
 			case EDataType::String:
-				RunCode_String(Item, pParams, InModel.Get());
+				RunCode_String(Item);
 				break;
 
 			case EDataType::Int:
-				RunCode_Int(Item, pParams, InModel.Get());
+				RunCode_Int(Item);
 				break;
 
 			case EDataType::Projector:
-				RunCode_Projector(Item, pParams, InModel.Get());
+				RunCode_Projector(Item);
 				break;
 
 			case EDataType::Color:
-				RunCode_Colour(Item, pParams, InModel.Get());
+				RunCode_Colour(Item);
 				break;
 
 			case EDataType::Matrix:
-				RunCode_Matrix(Item, pParams, InModel.Get());
+				RunCode_Matrix(Item);
 				break;
 
 			case EDataType::Material:
-				RunCode_Material(Item, pParams, InModel.Get());
+				RunCode_Material(Item);
 				break;
 				
 			default:
@@ -7611,13 +7567,11 @@ namespace UE::Mutable::Private
     }
 
 	
-	void CodeRunner::RunCodeImageDesc(const FScheduledOp& Item, const FParameters* pParams)
+	void CodeRunner::RunCodeImageDesc(const FScheduledOp& Item)
 	{
 		MUTABLE_CPUPROFILER_SCOPE(RunCodeImageDesc);
 		check(Item.Type == FScheduledOp::EType::ImageDesc);
 
-		const FProgram& Program = Model->Program;
-
 		EOpType Type = Program.GetOpType(Item.At);
 		switch (Type)
 		{
@@ -7653,7 +7607,7 @@ namespace UE::Mutable::Private
 					else
 					{
 						uint32 RomId = CurrentIndex.Index;
-						bIsLODAvailable = System->StreamInterface->DoesBlockExist(Model.Get(), RomId);
+						bIsLODAvailable = System.StreamInterface->DoesBlockExist(&Model, RomId);
 					}
 
 					if (bIsLODAvailable)
@@ -7722,7 +7676,7 @@ namespace UE::Mutable::Private
 			check(Item.Stage == 0);
 
 			OP::ParameterArgs Args = Program.GetOpArgs<OP::ParameterArgs>(Item.At);
-			UTexture* Image = pParams->GetTextureValue(Args.variable);
+			UTexture* Image = Parameters.GetTextureValue(Args.variable);
 
 			FExtendedImageDesc& Result = ImageDescResults.FindOrAdd(Item.At);
 			Result = GetExternalImageDesc(Image);
@@ -8821,7 +8775,7 @@ namespace UE::Mutable::Private
 			TSharedPtr<FRangeIndex> Index = BuildCurrentOpRangeIndex(Item, Args.MaterialParameter);
 
 			// Get material parameter from the array of parameters
-			UMaterialInterface* MaterialInterface = pParams->GetMaterialValue(Args.MaterialParameter, Index.Get());
+			UMaterialInterface* MaterialInterface = Parameters.GetMaterialValue(Args.MaterialParameter, Index.Get());
 
 			// Get the texture parameter name
 			check(Args.ParameterName < (uint32)Program.ConstantStrings.Num());
@@ -8872,7 +8826,7 @@ namespace UE::Mutable::Private
 				return New;
 			},
 
-			Runner->System->ImagePixelFormatOverride
+			Runner->System.ImagePixelFormatOverride
 		);
 	}
 }
diff --git a/Engine/Plugins/Mutable/Source/MutableRuntime/Private/MuR/CodeRunner.h b/Engine/Plugins/Mutable/Source/MutableRuntime/Private/MuR/CodeRunner.h
--- a/Engine/Plugins/Mutable/Source/MutableRuntime/Private/MuR/CodeRunner.h
+++ b/Engine/Plugins/Mutable/Source/MutableRuntime/Private/MuR/CodeRunner.h
@@ -39,10 +39,9 @@ namespace  UE::Mutable::Private
 	public:
 		static TSharedRef<CodeRunner> Create(
 				const TSharedRef<FLiveInstance>& InLiveInstance,
-				const FSettings&, 
-				FSystem*, 
+				const FSettings& InSettings, 
+				FSystem& InSystem, 
 				EExecutionStrategy,
-				const FParameters* Params,
 				OP::ADDRESS At,
 				uint32 LodMask,
 				uint8 ExecutionOptions,
@@ -54,10 +53,9 @@ namespace  UE::Mutable::Private
 		explicit CodeRunner(
 			const TSharedRef<FLiveInstance>& InLiveInstance,	
 			FPrivateToken, 
-			const FSettings&, 
-			FSystem*, 
+			const FSettings& InSettings, 
+			FSystem& InSystem, 
 			EExecutionStrategy,
-			const FParameters* Params,
 			OP::ADDRESS At,
 			uint32 LodMask,
 			uint8 ExecutionOptions,
@@ -125,29 +123,29 @@ namespace  UE::Mutable::Private
     	
         TSharedPtr<FRangeIndex> BuildCurrentOpRangeIndex( const FScheduledOp&, int32 ParameterIndex);
 
-        void RunCode( const FScheduledOp&, const FParameters*, const TSharedPtr<const FModel>&, uint32 LodMask, bool& bOutFailed);
+        void RunCode( const FScheduledOp&, uint32 LodMask, bool& bOutFailed);
 
-        void RunCode_Conditional(const FScheduledOp&, const FModel* );
-        void RunCode_Switch(const FScheduledOp&, const FModel* );
-        void RunCode_Instance(const FScheduledOp&, const FModel*, uint32 LodMask);
-        void RunCode_InstanceAddResource(const FScheduledOp&, const TSharedPtr<const FModel>& Model, const FParameters* );
+        void RunCode_Conditional(const FScheduledOp&);
+        void RunCode_Switch(const FScheduledOp&);
+        void RunCode_Instance(const FScheduledOp&, uint32 LodMask);
+        void RunCode_InstanceAddResource(const FScheduledOp&);
 
 		/** Return false in case of failure. */
-        bool RunCode_ConstantResource(const FScheduledOp&, const FModel* );
-
-        void RunCode_Mesh(const FScheduledOp&, const FModel* );
-        void RunCode_Image(const FScheduledOp&, const FParameters*, const FModel* );
-        void RunCode_Layout(const FScheduledOp&, const FModel* );
-        void RunCode_Bool(const FScheduledOp&, const FParameters*, const FModel* );
-        void RunCode_Int(const FScheduledOp&, const FParameters*, const FModel* );
-        void RunCode_Scalar(const FScheduledOp&, const FParameters*, const FModel* );
-        void RunCode_String(const FScheduledOp&, const FParameters*, const FModel* );
-        void RunCode_Colour(const FScheduledOp&, const FParameters*, const FModel* );
-        void RunCode_Projector(const FScheduledOp&, const FParameters*, const FModel* );
-        void RunCode_Matrix(const FScheduledOp&, const FParameters*, const FModel* );
-    	void RunCode_Material(const FScheduledOp&, const FParameters*, const FModel* );
-
-    	void RunCodeImageDesc(const FScheduledOp&, const FParameters*);
+        bool RunCode_ConstantResource(const FScheduledOp&);
+
+        void RunCode_Mesh(const FScheduledOp&);
+        void RunCode_Image(const FScheduledOp&);
+        void RunCode_Layout(const FScheduledOp&);
+        void RunCode_Bool(const FScheduledOp&);
+        void RunCode_Int(const FScheduledOp&);
+        void RunCode_Scalar(const FScheduledOp&);
+        void RunCode_String(const FScheduledOp&);
+        void RunCode_Colour(const FScheduledOp&);
+        void RunCode_Projector(const FScheduledOp&);
+        void RunCode_Matrix(const FScheduledOp&);
+    	void RunCode_Material(const FScheduledOp&);
+
+    	void RunCodeImageDesc(const FScheduledOp&);
     
 	public:
 		struct FExternalResourceId
@@ -552,7 +550,7 @@ namespace  UE::Mutable::Private
 
 		inline TManagedPtr<FImage> CloneOrTakeOver(TManagedPtr<const FImage>&& Resource)
 		{
-			System->WorkingMemoryManager.CheckRunnerThread();
+			System.WorkingMemoryManager.CheckRunnerThread();
 	
 			TManagedPtr<FImage> Result;
 			if (!Resource.IsUniqueReference())
@@ -572,7 +570,7 @@ namespace  UE::Mutable::Private
 
 		[[nodiscard]] inline TManagedPtr<FMesh> CreateMesh(int32 BudgetReserveSize = 0)
 		{
-			System->WorkingMemoryManager.CheckRunnerThread();
+			System.WorkingMemoryManager.CheckRunnerThread();
 
 			EnsureBudgetBelow(BudgetReserveSize);
 
@@ -582,7 +580,7 @@ namespace  UE::Mutable::Private
 
 		[[nodiscard]] inline TManagedPtr<FMesh> CloneOrTakeOver(TManagedPtr<const FMesh>&& Resource)
 		{
-			System->WorkingMemoryManager.CheckRunnerThread();
+			System.WorkingMemoryManager.CheckRunnerThread();
 
 			TManagedPtr<FMesh> Result;
 			if (!Resource.IsUniqueReference())
@@ -621,9 +619,14 @@ namespace  UE::Mutable::Private
 		*/
 		bool bUnrecoverableError = false;
 
-		FSystem* System = nullptr;
-		TSharedPtr<const FModel> Model = nullptr;
-		const FParameters* Params = nullptr;
+		FSystem& System;
+    	
+    	// Dereference optimization. LiveInstance keeps them alive.
+		const FModel& Model;
+		const FProgram& Program;
+		const FParameters& Parameters;
+    	
+    public:
 		uint32 LODMask = 0;
 
     	TSharedPtr<FExternalResourceProvider> ExternalResourceProvider;
@@ -632,7 +635,7 @@ namespace  UE::Mutable::Private
     	
     	inline TManagedPtr<FImage> CreateImageImp(uint32 SizeX, uint32 SizeY, uint32 Lods, EImageFormat Format, EInitializationType Init)
     	{
-    		System->WorkingMemoryManager.CheckRunnerThread();
+    		System.WorkingMemoryManager.CheckRunnerThread();
 			
     		// Make room in the budget
     		uint32 DataSize = FImage::CalculateDataSize(SizeX, SizeY, Lods, Format);
@@ -645,8 +648,7 @@ namespace  UE::Mutable::Private
     	
     	bool EnsureBudgetBelow( uint64 AdditionalMemory ) const
 	    {
-    		check(System)
-    		return System->WorkingMemoryManager.EnsureBudgetBelow(AdditionalMemory, LiveInstance->Cache);
+    		return System.WorkingMemoryManager.EnsureBudgetBelow(AdditionalMemory, LiveInstance->Cache);
     	}
     	
     	void MarkRomUsed( int32 InRomIndex);
diff --git a/Engine/Plugins/Mutable/Source/MutableRuntime/Private/MuR/CodeRunnerBegin.cpp b/Engine/Plugins/Mutable/Source/MutableRuntime/Private/MuR/CodeRunnerBegin.cpp
--- a/Engine/Plugins/Mutable/Source/MutableRuntime/Private/MuR/CodeRunnerBegin.cpp
+++ b/Engine/Plugins/Mutable/Source/MutableRuntime/Private/MuR/CodeRunnerBegin.cpp
@@ -853,7 +853,7 @@ namespace UE::Mutable::Private
 							
 							if (Args.value)
 							{
-								FMeshId MeshId = FSystem::GetMeshId(MeshIdRegistry, Model, Params.Get(), Args.value);
+								FMeshId MeshId = FSystem::GetMeshId(MeshIdRegistry, *Model.Get(), *Params.Get(), Args.value);
 								OP::ADDRESS NameAd = Args.name;
 								check(NameAd < (uint32)Program.ConstantStrings.Num());
 								const FString& Name = Program.ConstantStrings[NameAd];
@@ -908,7 +908,7 @@ namespace UE::Mutable::Private
 
 							if (Args.Material)
 							{
-        						FMaterialId MaterialId = FSystem::GetMaterialId(MaterialIdRegistry, Model, Params.Get(), Args.Material);
+        						FMaterialId MaterialId = FSystem::GetMaterialId(MaterialIdRegistry, *Model.Get(), *Params.Get(), Args.Material);
 								Instance->SetMaterialId(0, 0, 0, MaterialId);
 
 								LoadMaterial(Args.Material);
@@ -961,7 +961,7 @@ namespace UE::Mutable::Private
 
 							if (Args.Material)
 							{
-        						FMaterialId MaterialId = FSystem::GetMaterialId(MaterialIdRegistry, Model, Params.Get(), Args.Material);
+        						FMaterialId MaterialId = FSystem::GetMaterialId(MaterialIdRegistry, *Model.Get(), *Params.Get(), Args.Material);
 								Instance->SetOverlayMaterialId(0, MaterialId);
 
 								// Pop Material
@@ -1013,7 +1013,7 @@ namespace UE::Mutable::Private
 
 							if (Args.Material)
 							{
-								FMaterialId MaterialId = FSystem::GetMaterialId(MaterialIdRegistry, Model, Params.Get(), Args.Material);
+								FMaterialId MaterialId = FSystem::GetMaterialId(MaterialIdRegistry, *Model.Get(), *Params.Get(), Args.Material);
 								const FString& SlotName = Program.ConstantStrings[Args.SlotName];
 								Instance->AddOverrideMaterial(0, FName(SlotName), MaterialId);
 
@@ -1327,7 +1327,7 @@ namespace UE::Mutable::Private
 							
 							if (Args.value)
 							{
-        						FSkeletalMeshId SkeletalMeshId = FSystem::GetSkeletalMeshId(SkeletalMeshIdRegistry, Model, Params.Get(), Args.value);
+        						FSkeletalMeshId SkeletalMeshId = FSystem::GetSkeletalMeshId(SkeletalMeshIdRegistry, *Model.Get(), *Params.Get(), Args.value);
 								Instance->SetSkeletalMeshId(0, SkeletalMeshId);
 							}
 							
diff --git a/Engine/Plugins/Mutable/Source/MutableRuntime/Private/MuR/CodeRunnerTasked.cpp b/Engine/Plugins/Mutable/Source/MutableRuntime/Private/MuR/CodeRunnerTasked.cpp
--- a/Engine/Plugins/Mutable/Source/MutableRuntime/Private/MuR/CodeRunnerTasked.cpp
+++ b/Engine/Plugins/Mutable/Source/MutableRuntime/Private/MuR/CodeRunnerTasked.cpp
@@ -110,7 +110,7 @@ namespace UE::Mutable::Private
 		else
 		{
 			// We could wait. Let's see if we have enough memory to issue tasks anyway.
-			bool bHaveEnoughMemory = !System->WorkingMemoryManager.IsMemoryBudgetFull();
+			bool bHaveEnoughMemory = !System.WorkingMemoryManager.IsMemoryBudgetFull();
 			if (bHaveEnoughMemory)
 			{
 				return  true;
@@ -272,7 +272,7 @@ namespace UE::Mutable::Private
 				// Special processing in case it is an ImageDesc operation
 				if (Item.Type == FScheduledOp::EType::ImageDesc)
 				{
-					RunCodeImageDesc(Item, Params);
+					RunCodeImageDesc(Item);
 					continue;
 				}
 
@@ -304,7 +304,7 @@ namespace UE::Mutable::Private
 				{
 					// Run immediately
 					bool bFailed = false;
-					RunCode(Item, Params, Model, LODMask, bFailed);
+					RunCode(Item, LODMask, bFailed);
 					if (bFailed)
 					{
 						return AbortRun();
@@ -321,7 +321,7 @@ namespace UE::Mutable::Private
 				if (ProfileContext)
 				{
 					++ProfileContext->NumRunOps;
-					++ProfileContext->RunOpsPerType[int32(Model->Program.GetOpType(Item.At))];
+					++ProfileContext->RunOpsPerType[int32(LiveInstance->Model->Program.GetOpType(Item.At))];
 				}
 			}
 
@@ -399,12 +399,12 @@ namespace UE::Mutable::Private
 						}
 					}
 
-					System->WorkingMemoryManager.InvalidateRunnerThread();
+					System.WorkingMemoryManager.InvalidateRunnerThread();
 
 					UE::Tasks::AddNested(UE::Tasks::Launch(TEXT("CodeRunnerFromIssuedTasksTask"),
 						[Runner = AsShared(), ProfileContext = MoveTemp(ProfileContext)]() mutable
 						{
-							Runner->System->WorkingMemoryManager.ResetRunnerThread();
+							Runner->System.WorkingMemoryManager.ResetRunnerThread();
 
 							constexpr bool bForceInlineExecution = false;
 							Runner->Run(MoveTemp(ProfileContext), bForceInlineExecution);
@@ -433,12 +433,12 @@ namespace UE::Mutable::Private
 			{
 				if (FTimespan::FromSeconds(FPlatformTime::Seconds()) > TimeOut)
 				{
-					System->WorkingMemoryManager.InvalidateRunnerThread();
+					System.WorkingMemoryManager.InvalidateRunnerThread();
 
 					UE::Tasks::AddNested(UE::Tasks::Launch(TEXT("CodeRunnerFromTimeoutTask"),
 						[Runner = AsShared(), ProfileContext = MoveTemp(ProfileContext)]() mutable
 						{
-							Runner->System->WorkingMemoryManager.ResetRunnerThread();
+							Runner->System.WorkingMemoryManager.ResetRunnerThread();
 
 							constexpr bool bForceInlineExecution = false;
 							Runner->Run(MoveTemp(ProfileContext), bForceInlineExecution);
@@ -1116,7 +1116,7 @@ namespace UE::Mutable::Private
 			return false;
 		}
 
-		ImagePixelFormatFunc = Runner->System->ImagePixelFormatOverride;
+		ImagePixelFormatFunc = Runner->System.ImagePixelFormatOverride;
 
 		// Create destination data
 		Result = Runner->CreateImage(Base->GetSizeX(), Base->GetSizeY(), Base->GetLODCount(), TargetFormat, EInitializationType::NotInitialized);
@@ -1233,7 +1233,7 @@ namespace UE::Mutable::Private
 		FImageOperator ImOp = MakeImageOperator(Runner);
 		ImOp.ImageMipmap_PrepareScratch(Result.Get(), StartLevel, LevelCount, Scratch);
 
-		ImagePixelFormatFunc = Runner->System->ImagePixelFormatOverride;
+		ImagePixelFormatFunc = Runner->System.ImagePixelFormatOverride;
 
 		return true;
 	}
@@ -1560,7 +1560,7 @@ namespace UE::Mutable::Private
 
 		Result = Runner->CreateImage( destSize[0], destSize[1], Lods, Base->GetFormat(), EInitializationType::NotInitialized );
 
-		ImagePixelFormatFunc = Runner->System->ImagePixelFormatOverride;
+		ImagePixelFormatFunc = Runner->System.ImagePixelFormatOverride;
 
 		return true;
 	}
@@ -1681,7 +1681,7 @@ namespace UE::Mutable::Private
 
 		Result = Runner->CreateImage(DestSize[0], DestSize[1], Lods, Base->GetFormat(), EInitializationType::NotInitialized);
 
-		ImagePixelFormatFunc = Runner->System->ImagePixelFormatOverride;
+		ImagePixelFormatFunc = Runner->System.ImagePixelFormatOverride;
 
 		return true;
 	}
@@ -1907,7 +1907,7 @@ namespace UE::Mutable::Private
 			}
 		}
 
-		ImagePixelFormatFunc = Runner->System->ImagePixelFormatOverride;
+		ImagePixelFormatFunc = Runner->System.ImagePixelFormatOverride;
 
 		return true;
 	}
@@ -1962,14 +1962,14 @@ namespace UE::Mutable::Private
 		// This runs in the mutable Runner thread
 		MUTABLE_CPUPROFILER_SCOPE(FLoadMeshRomTask_Prepare);
 
-		if (!Runner || !Runner->System)
+		if (!Runner)
 		{
 			return false;
 		}
 
 		bOutFailed = false;
 
-		const FProgram& Program = Runner->Model->Program;
+		const FProgram& LocalProgram = Runner->Program;
 		FWorkingMemoryManager::FModelCacheEntry* ModelCache = Runner->GetModelCache();
 		
 		TArray<UE::Tasks::FTask, TInlineAllocator<4>> ReadCompleteEvents;
@@ -2017,7 +2017,7 @@ namespace UE::Mutable::Private
 		for (const int32 MeshContentRomIndex : RomsToLoad)
 		{
 			FConstantResourceIndex CurrentIndex = 
-					Program.ConstantMeshContentIndices[MeshContentRomIndex + FirstIndex];
+					LocalProgram.ConstantMeshContentIndices[MeshContentRomIndex + FirstIndex];
 
 			if (!CurrentIndex.Streamable)
 			{
@@ -2027,7 +2027,7 @@ namespace UE::Mutable::Private
 
 			int32 RomIndex = CurrentIndex.Index;
 
-			check(RomIndex < Program.Roms.Num());
+			check(RomIndex < LocalProgram.Roms.Num());
 
 			++ModelCache->PendingOpsPerRom[RomIndex];
 
@@ -2036,7 +2036,7 @@ namespace UE::Mutable::Private
 				UE_LOG(LogMutableCore, Log, TEXT("Preparing rom %d, now peding ops is %d."), RomIndex, ModelCache->PendingOpsPerRom[RomIndex]);
 			}
 
-			if (Program.IsRomLoaded(RomIndex))
+			if (LocalProgram.IsRomLoaded(RomIndex))
 			{
 				continue;
 			}
@@ -2051,9 +2051,9 @@ namespace UE::Mutable::Private
 
 			FRomLoadOp& RomLoadOp = Runner->RomLoadOps.Create(RomIndex);
 			
-			check(Runner->System->StreamInterface);
+			check(Runner->System.StreamInterface);
 
-			const uint32 RomSize = Program.Roms[RomIndex].Size;
+			const uint32 RomSize = LocalProgram.Roms[RomIndex].Size;
 			check(RomSize > 0);
 
 			// Free roms if necessary
@@ -2077,7 +2077,7 @@ namespace UE::Mutable::Private
 
 		
 			const uint32 RomId = RomIndex;
-			RomLoadOp.m_streamID = Runner->System->StreamInterface->BeginReadBlock(Runner->Model.Get(), RomId, RomLoadOp.m_streamBuffer.GetData(), RomSize, EDataType::Mesh, &Callback);
+			RomLoadOp.m_streamID = Runner->System.StreamInterface->BeginReadBlock(&Runner->Model, RomId, RomLoadOp.m_streamBuffer.GetData(), RomSize, EDataType::Mesh, &Callback);
 			if (RomLoadOp.m_streamID < 0)
 			{
 				bOutFailed = true;
@@ -2102,12 +2102,12 @@ namespace UE::Mutable::Private
 		// This runs in the Runner thread
 		MUTABLE_CPUPROFILER_SCOPE(FLoadMeshRomTask_Complete);
 
-		if (!Runner || !Runner->System)
+		if (!Runner)
 		{
 			return false;
 		}
 
-		const FProgram& Program = Runner->Model->Program;
+		const FProgram& LocalProgram = Runner->Program;
 		FWorkingMemoryManager::FModelCacheEntry* ModelCache = Runner->GetModelCache();
 		
 		bool bSomeMissingData = false;
@@ -2116,7 +2116,7 @@ namespace UE::Mutable::Private
 			// Since task could be reordered, we need to make sure we end the rom read before continuing
 			if (FRomLoadOp* RomLoadOp = Runner->RomLoadOps.Find(RomIndex))
 			{				
-				bool bSuccess = Runner->System->StreamInterface->EndRead(RomLoadOp->m_streamID);
+				bool bSuccess = Runner->System.StreamInterface->EndRead(RomLoadOp->m_streamID);
 
 				MUTABLE_CPUPROFILER_SCOPE(Unserialise);
 
@@ -2127,7 +2127,7 @@ namespace UE::Mutable::Private
 
 					TManagedPtr<FMesh> Value = FMesh::StaticUnserialise(Arch);
 
-					Program.SetMeshRomValue(RomIndex, Value);
+					LocalProgram.SetMeshRomValue(RomIndex, Value);
 				}
 				else
 				{
@@ -2145,13 +2145,13 @@ namespace UE::Mutable::Private
 		}
 
 		// Process the constant op normally, now that the rom is loaded.	
-		bool bSuccess = Runner->RunCode_ConstantResource(Op, Runner->Model.Get());
+		bool bSuccess = Runner->RunCode_ConstantResource(Op);
 
 		if (bSuccess)
 		{
 			for (int32 RomIndex : RomIndices)
 			{
-				check(RomIndex < Program.Roms.Num());
+				check(RomIndex < LocalProgram.Roms.Num());
 
 				Runner->MarkRomUsed(RomIndex);
 				--ModelCache->PendingOpsPerRom[RomIndex];
@@ -2171,7 +2171,7 @@ namespace UE::Mutable::Private
 	//---------------------------------------------------------------------------------------------
 	bool CodeRunner::FLoadImageRomsTask::Prepare(CodeRunner* Runner, bool& bOutFailed )
 	{
-		if (!Runner || !Runner->System)
+		if (!Runner)
 		{
 			return false;
 		}
@@ -2180,7 +2180,7 @@ namespace UE::Mutable::Private
 		MUTABLE_CPUPROFILER_SCOPE(FLoadImageRomsTask_Prepare);
 		bOutFailed = false;
 
-		const FProgram& Program = Runner->Model->Program;
+		const FProgram& LocalProgram = Runner->Program;
 
 		FWorkingMemoryManager::FModelCacheEntry* ModelCache = Runner->GetModelCache();
 
@@ -2190,7 +2190,7 @@ namespace UE::Mutable::Private
 
 		for (int32 ImageIndex = ImageIndexBegin; ImageIndex < ImageIndexEnd; ++ImageIndex)
 		{
-			const FConstantResourceIndex ImageConstantResourceIndex = Program.ConstantImageLODIndices[ImageIndex];
+			const FConstantResourceIndex ImageConstantResourceIndex = LocalProgram.ConstantImageLODIndices[ImageIndex];
 
 			if (!ImageConstantResourceIndex.Streamable)
 			{
@@ -2199,14 +2199,14 @@ namespace UE::Mutable::Private
 			}
 
 			int32 RomIndex = ImageConstantResourceIndex.Index;
-			check(RomIndex < Program.Roms.Num());
+			check(RomIndex < LocalProgram.Roms.Num());
 
 			++ModelCache->PendingOpsPerRom[RomIndex];
 
 			if (DebugRom && (DebugRomAll || RomIndex == DebugRomIndex))
 				UE_LOG(LogMutableCore, Log, TEXT("Preparing rom %d, now peding ops is %d."), RomIndex, ModelCache->PendingOpsPerRom[RomIndex]);
 
-			if (Program.IsRomLoaded(RomIndex))
+			if (LocalProgram.IsRomLoaded(RomIndex))
 			{
 				continue;
 			}
@@ -2221,9 +2221,9 @@ namespace UE::Mutable::Private
 
 			FRomLoadOp& RomLoadOp = Runner->RomLoadOps.Create(RomIndex);
 			
-			check(Runner->System->StreamInterface);
+			check(Runner->System.StreamInterface);
 
-			const uint32 RomSize = Program.Roms[RomIndex].Size;
+			const uint32 RomSize = LocalProgram.Roms[RomIndex].Size;
 			check(RomSize > 0);
 
 			// Free roms if necessary
@@ -2246,7 +2246,7 @@ namespace UE::Mutable::Private
 			};
 			
 			const uint32 RomId = RomIndex;
-			RomLoadOp.m_streamID = Runner->System->StreamInterface->BeginReadBlock(Runner->Model.Get(), RomId, RomLoadOp.m_streamBuffer.GetData(), RomSize, EDataType::Image, &Callback);
+			RomLoadOp.m_streamID = Runner->System.StreamInterface->BeginReadBlock(&Runner->Model, RomId, RomLoadOp.m_streamBuffer.GetData(), RomSize, EDataType::Image, &Callback);
 			if (RomLoadOp.m_streamID < 0)
 			{
 				bOutFailed = true;
@@ -2271,12 +2271,12 @@ namespace UE::Mutable::Private
 		// This runs in the Runner thread
 		MUTABLE_CPUPROFILER_SCOPE(FLoadImageRomsTask_Complete);
 
-		if (!Runner || !Runner->System)
+		if (!Runner)
 		{
 			return false;
 		}
 
-		const FProgram& Program = Runner->Model->Program;
+		const FProgram& LocalProgram = Runner->Program;
 		
 		FWorkingMemoryManager::FModelCacheEntry* ModelCache = Runner->GetModelCache();
 
@@ -2287,7 +2287,7 @@ namespace UE::Mutable::Private
 			// Since task could be reordered, we need to make sure we end the rom read before continuing
 			if (FRomLoadOp* RomLoadOp = Runner->RomLoadOps.Find(RomIndex))
 			{				
-				bool bSuccess = Runner->System->StreamInterface->EndRead(RomLoadOp->m_streamID);
+				bool bSuccess = Runner->System.StreamInterface->EndRead(RomLoadOp->m_streamID);
 
 				MUTABLE_CPUPROFILER_SCOPE(Unserialise);
 
@@ -2298,7 +2298,7 @@ namespace UE::Mutable::Private
 
 					TManagedPtr<FImage> Value = FImage::StaticUnserialise(Arch);
 
-					Program.SetImageRomValue(RomIndex, Value);
+					LocalProgram.SetImageRomValue(RomIndex, Value);
 				}
 				else
 				{
@@ -2316,13 +2316,13 @@ namespace UE::Mutable::Private
 		}
 
 		// Process the constant op normally, now that the rom is loaded.	
-		bool bSuccess = Runner->RunCode_ConstantResource(Op, Runner->Model.Get());
+		bool bSuccess = Runner->RunCode_ConstantResource(Op);
 
 		if (bSuccess)
 		{
 			for (int32 ImageIndex = ImageIndexBegin; ImageIndex < ImageIndexEnd; ++ImageIndex)
 			{
-				FConstantResourceIndex ImageConstantResourceIndex = Program.ConstantImageLODIndices[ImageIndex];
+				FConstantResourceIndex ImageConstantResourceIndex = LocalProgram.ConstantImageLODIndices[ImageIndex];
 
 				if (!ImageConstantResourceIndex.Streamable)
 				{
@@ -2331,7 +2331,7 @@ namespace UE::Mutable::Private
 				}
 
 				const int32 RomIndex = ImageConstantResourceIndex.Index;
-				check(RomIndex < Program.Roms.Num());
+				check(RomIndex < LocalProgram.Roms.Num());
 
 				Runner->MarkRomUsed(RomIndex);
 				--ModelCache->PendingOpsPerRom[RomIndex];
@@ -2495,9 +2495,7 @@ namespace UE::Mutable::Private
 	TSharedPtr<CodeRunner::FIssuedTask> CodeRunner::IssueOp(FScheduledOp Item)
 	{
 		TSharedPtr<FIssuedTask> Issued;
-
-		const FProgram& Program = Model->Program;
-
+		
 		EOpType type = Program.GetOpType(Item.At);
 
 		switch (type)
@@ -2586,7 +2584,7 @@ namespace UE::Mutable::Private
 			TSharedPtr<FRangeIndex> Index = BuildCurrentOpRangeIndex(Item, Args.MaterialParameter);
 
 			// Get material parameter from the array of parameters
-			UMaterialInterface* Material = Params->GetMaterialValue(Args.MaterialParameter, Index.Get());
+			UMaterialInterface* Material = Parameters.GetMaterialValue(Args.MaterialParameter, Index.Get());
 
 			// Get the texture parameter name
 			check(Args.ParameterName < (uint32)Program.ConstantStrings.Num());
@@ -2613,7 +2611,7 @@ namespace UE::Mutable::Private
 
 		case EOpType::IM_REFERENCE:
 		{
-			OP::ResourceReferenceArgs Args = Model->Program.GetOpArgs<OP::ResourceReferenceArgs>(Item.At);
+			OP::ResourceReferenceArgs Args = Program.GetOpArgs<OP::ResourceReferenceArgs>(Item.At);
 
 			// We only convert references to images if indicated in the operation.
 			if (Args.ForceLoad)
@@ -2763,15 +2761,13 @@ namespace UE::Mutable::Private
 
 	void CodeRunner::MarkRomUsed(int32 InRomIndex)
 	{
-		check(Model)
-		System->WorkingMemoryManager.MarkRomUsed(InRomIndex, Model);
+		System.WorkingMemoryManager.MarkRomUsed(InRomIndex, LiveInstance->Model);
 	}
 
 	
 	FWorkingMemoryManager::FModelCacheEntry* CodeRunner::GetModelCache() const
 	{
-		check(Model)
-		return System->WorkingMemoryManager.FindModelCache(Model.Get());
+		return System.WorkingMemoryManager.FindModelCache(&Model);
 	}
 
 } // namespace UE::Mutable::Private
diff --git a/Engine/Plugins/Mutable/Source/MutableRuntime/Private/MuR/System.cpp b/Engine/Plugins/Mutable/Source/MutableRuntime/Private/MuR/System.cpp
--- a/Engine/Plugins/Mutable/Source/MutableRuntime/Private/MuR/System.cpp
+++ b/Engine/Plugins/Mutable/Source/MutableRuntime/Private/MuR/System.cpp
@@ -262,7 +262,7 @@ namespace UE::Mutable::Private
 		WorkingMemoryManager.FindOrAddModelCache(LiveInstance->Model);
 
 		OP::ADDRESS RootAt = Program.States[InStateIndex].Root;
-		RunCode(LiveInstance, InParams.Get(), RootAt, InLodMask);
+		RunCode(LiveInstance, RootAt, InLodMask);
 
 		TManagedPtr<const FInstance> Result = LiveInstance->Cache->GetInstance(FCacheAddress(RootAt, 0, 0));
 
@@ -300,7 +300,7 @@ namespace UE::Mutable::Private
     	}
 
 		OP::ADDRESS RootAddress = ImageId.GetKey()->Address;
-		TManagedPtr<const FImage> Result = BuildImage(LiveInstance, LiveInstance->OldParameters.Get(), RootAddress, MipsToSkip);
+		TManagedPtr<const FImage> Result = BuildImage(LiveInstance, RootAddress, MipsToSkip);
 
 		if (!Result)
 		{
@@ -356,9 +356,8 @@ namespace UE::Mutable::Private
 			Context->Runner = CodeRunner::Create(
 				LiveInstance,
 				System->Settings,
-				System,
+				*System,
 				EExecutionStrategy::MinimizeMemory,
-				LiveInstance->OldParameters.Get(),
 				RootAddress, FSystem::AllLODs,
 				MipsToSkip,
 				false,
@@ -439,9 +438,8 @@ namespace UE::Mutable::Private
 			TSharedRef<CodeRunner> Runner = CodeRunner::Create(
 					LiveInstance,
 					Settings,
-					this,
+					*this,
 					EExecutionStrategy::MinimizeMemory,
-					LiveInstance->OldParameters.Get(),
 					RootAddress,
 					FSystem::AllLODs,
 					executionOptions,
@@ -516,9 +514,8 @@ namespace UE::Mutable::Private
 			Context->Runner = CodeRunner::Create(
 			LiveInstance,
 				System->Settings,
-				System,
+				*System,
 				EExecutionStrategy::MinimizeMemory,
-				LiveInstance->OldParameters.Get(),
 				RootAddress,
 				FSystem::AllLODs,
 				ExecutionOptions,
@@ -559,7 +556,6 @@ namespace UE::Mutable::Private
 		OP::ADDRESS RootAddress = MeshId.GetKey()->Address;
 		TManagedPtr<const FMesh> Result = BuildMesh(
 			LiveInstance,
-			LiveInstance->OldParameters.Get(),
 			RootAddress,
 			MeshContentFilter);
 
@@ -620,9 +616,8 @@ namespace UE::Mutable::Private
 			Context->Runner = CodeRunner::Create(
 			LiveInstance,
 					System->Settings,
-					System, 
+					*System, 
 					EExecutionStrategy::MinimizeMemory, 
-					LiveInstance->OldParameters.Get(), 
 					RootAddress, 
 					FSystem::AllLODs, 
 					ExecutionOptions, 
@@ -702,9 +697,8 @@ namespace UE::Mutable::Private
 			Context->Runner = CodeRunner::Create(
 			LiveInstance,
 					System->Settings,
-					System, 
+					*System, 
 					EExecutionStrategy::MinimizeMemory, 
-					LiveInstance->OldParameters.Get(), 
 					RootAddress, 
 					FSystem::AllLODs, 
 					static_cast<uint8>(SkeletalMeshContentFilter), 
@@ -784,9 +778,8 @@ namespace UE::Mutable::Private
 			Context->Runner = CodeRunner::Create(
 			LiveInstance,
 					System->Settings,
-					System, 
+					*System, 
 					EExecutionStrategy::MinimizeMemory, 
-					LiveInstance->OldParameters.Get(), 
 					RootAddress, 
 					FSystem::AllLODs, 
 					0,
@@ -870,9 +863,8 @@ namespace UE::Mutable::Private
         RelevantParameterVisitor(
                 FSystem* InSystem,
 				const TSharedRef<FLiveInstance>& InLiveInstance,
-                const TSharedPtr<const FParameters>& InParams,
                 TArray<uint32>& OutRelevantParameters) :
-				UniqueDiscreteCoveredCodeVisitor<>(InSystem, InLiveInstance, InParams, FSystem::AllLODs),
+				UniqueDiscreteCoveredCodeVisitor<>(InSystem, InLiveInstance, FSystem::AllLODs),
 		    	RelevantParameters(OutRelevantParameters)
         {
 			const OP::ADDRESS At = InLiveInstance->Model->Program.States[0].Root;
@@ -912,13 +904,11 @@ namespace UE::Mutable::Private
     };
 
 	
-    void FSystem::GetParameterRelevancy(const TSharedRef<FLiveInstance>& LiveInstance,
-    	const TSharedPtr<const FParameters>& FParameters,
-        TArray<uint32>& OutRelevantParameters)
+    void FSystem::GetParameterRelevancy(const TSharedRef<FLiveInstance>& LiveInstance, TArray<uint32>& OutRelevantParameters)
     {
 		LLM_SCOPE_BYNAME(TEXT("MutableRuntime"));
     	
-		RelevantParameterVisitor visitor(this, LiveInstance, FParameters, OutRelevantParameters);
+		RelevantParameterVisitor visitor(this, LiveInstance, OutRelevantParameters);
     }
 
 	
@@ -982,7 +972,6 @@ namespace UE::Mutable::Private
 
 	void FSystem::RunCode(
 		const TSharedRef<FLiveInstance>& InLiveInstance,
-		const FParameters* InParameters,
 		OP::ADDRESS InCodeRoot,
 		uint32 InLODs,
 		uint8 ExecutionOptions)
@@ -990,9 +979,8 @@ namespace UE::Mutable::Private
 		TSharedRef<CodeRunner> Runner = CodeRunner::Create(
 			InLiveInstance,
 			Settings,
-			this,
+			*this,
 			EExecutionStrategy::MinimizeMemory,
-			InParameters,
 			InCodeRoot,
 			InLODs,
 			ExecutionOptions,
@@ -1006,11 +994,11 @@ namespace UE::Mutable::Private
 	}
 	
 	
-	bool FSystem::BuildBool(const TSharedRef<FLiveInstance>& InLiveInstance, const FParameters* InParameters, OP::ADDRESS InAddress)
+	bool FSystem::BuildBool(const TSharedRef<FLiveInstance>& InLiveInstance, OP::ADDRESS InAddress)
     {
     	WorkingMemoryManager.BeginRunnerThread();
     	
-    	RunCode(InLiveInstance, InParameters, InAddress);
+    	RunCode(InLiveInstance, InAddress);
 	
     	bool bResult = false;
     	if (!bUnrecoverableError)
@@ -1024,11 +1012,11 @@ namespace UE::Mutable::Private
     }
 	
 	
-	int32 FSystem::BuildInt(const TSharedRef<FLiveInstance>& InLiveInstance, const FParameters* InParameters, OP::ADDRESS InAddress)
+	int32 FSystem::BuildInt(const TSharedRef<FLiveInstance>& InLiveInstance, OP::ADDRESS InAddress)
 	{
 		WorkingMemoryManager.BeginRunnerThread();
     	
-		RunCode(InLiveInstance, InParameters, InAddress);
+		RunCode(InLiveInstance, InAddress);
 
 		int32 Result = 0;
 		if (!bUnrecoverableError)
@@ -1042,11 +1030,11 @@ namespace UE::Mutable::Private
 	}
 
 	
-	float FSystem::BuildScalar(const TSharedRef<FLiveInstance>& InLiveInstance, const FParameters* InParameters, OP::ADDRESS InAddress)
+	float FSystem::BuildScalar(const TSharedRef<FLiveInstance>& InLiveInstance, OP::ADDRESS InAddress)
 	{
     	WorkingMemoryManager.BeginRunnerThread();
     	
-    	RunCode(InLiveInstance, InParameters, InAddress);
+    	RunCode(InLiveInstance, InAddress);
 
     	float Result = 0.0f;		
     	if (!bUnrecoverableError)
@@ -1060,7 +1048,7 @@ namespace UE::Mutable::Private
 	}
 
 
-	FVector4f FSystem::BuildColor(const TSharedRef<FLiveInstance>& InLiveInstance, const FParameters* InParameters, OP::ADDRESS InAddress)
+	FVector4f FSystem::BuildColor(const TSharedRef<FLiveInstance>& InLiveInstance, OP::ADDRESS InAddress)
 	{
 		WorkingMemoryManager.BeginRunnerThread();
 
@@ -1069,7 +1057,7 @@ namespace UE::Mutable::Private
 		UE::Mutable::Private::EOpType OpType = InLiveInstance->Model->Program.GetOpType(InAddress);
 		if (GetOpDataType(OpType) == EDataType::Color)
 		{
-			RunCode(InLiveInstance, InParameters, InAddress);
+			RunCode(InLiveInstance, InAddress);
 			if (!bUnrecoverableError)
 			{
 				Result = InLiveInstance->Cache->GetColor(FCacheAddress(InAddress, 0, 0));
@@ -1082,11 +1070,11 @@ namespace UE::Mutable::Private
 	}
 
 	
-	FProjector FSystem::BuildProjector(const TSharedRef<FLiveInstance>& InLiveInstance, const FParameters* InParameters, OP::ADDRESS InAddress)
+	FProjector FSystem::BuildProjector(const TSharedRef<FLiveInstance>& InLiveInstance, OP::ADDRESS InAddress)
 	{
 		WorkingMemoryManager.BeginRunnerThread();
     	
-    	RunCode(InLiveInstance, InParameters, InAddress);
+    	RunCode(InLiveInstance, InAddress);
 
 		FProjector Result;
 		if (!bUnrecoverableError)
@@ -1100,7 +1088,7 @@ namespace UE::Mutable::Private
 	}
 
 
-	TManagedPtr<const FImage> FSystem::BuildImage(const TSharedRef<FLiveInstance>& InLiveInstance, const FParameters* InParameters, OP::ADDRESS InAddress, int32 MipsToSkip)
+	TManagedPtr<const FImage> FSystem::BuildImage(const TSharedRef<FLiveInstance>& InLiveInstance, OP::ADDRESS InAddress, int32 MipsToSkip)
 	{
 		WorkingMemoryManager.BeginRunnerThread();
 
@@ -1109,7 +1097,7 @@ namespace UE::Mutable::Private
 		UE::Mutable::Private::EOpType OpType = InLiveInstance->Model->Program.GetOpType(InAddress);
 		if (GetOpDataType(OpType) == EDataType::Image)
 		{
-			RunCode(InLiveInstance, InParameters, InAddress, FSystem::AllLODs, uint8(MipsToSkip));
+			RunCode(InLiveInstance, InAddress, FSystem::AllLODs, uint8(MipsToSkip));
 			if (!bUnrecoverableError)
 			{
 				Result = InLiveInstance->Cache->GetImage(FCacheAddress(InAddress, 0, MipsToSkip));
@@ -1122,7 +1110,7 @@ namespace UE::Mutable::Private
 	}
 
 
-	TManagedPtr<const FMesh> FSystem::BuildMesh(const TSharedRef<FLiveInstance>& InLiveInstance, const FParameters* InParameters, OP::ADDRESS InRootAddress, EMeshContentFlags MeshContentFilter)
+	TManagedPtr<const FMesh> FSystem::BuildMesh(const TSharedRef<FLiveInstance>& InLiveInstance, OP::ADDRESS InRootAddress, EMeshContentFlags MeshContentFilter)
 	{
 		WorkingMemoryManager.BeginRunnerThread();
 
@@ -1132,7 +1120,7 @@ namespace UE::Mutable::Private
 		if (GetOpDataType(OpType) == EDataType::Mesh)
 		{
 			const uint8 ExecutionOptions = static_cast<uint8>(MeshContentFilter);
-			RunCode(InLiveInstance, InParameters, InRootAddress, FSystem::AllLODs, ExecutionOptions);
+			RunCode(InLiveInstance, InRootAddress, FSystem::AllLODs, ExecutionOptions);
 			if (!bUnrecoverableError)
 			{
 				Result = InLiveInstance->Cache->GetMesh(FCacheAddress(InRootAddress, 0, ExecutionOptions));
@@ -1145,7 +1133,7 @@ namespace UE::Mutable::Private
 	}
 
 
-	TManagedPtr<const FInstance> FSystem::BuildInstance(const TSharedRef<FLiveInstance>& InLiveInstance, const FParameters* InParameters, OP::ADDRESS InAddress)
+	TManagedPtr<const FInstance> FSystem::BuildInstance(const TSharedRef<FLiveInstance>& InLiveInstance, OP::ADDRESS InAddress)
 	{
 		WorkingMemoryManager.BeginRunnerThread();
 
@@ -1154,7 +1142,7 @@ namespace UE::Mutable::Private
 		UE::Mutable::Private::EOpType OpType = InLiveInstance->Model->Program.GetOpType(InAddress);
 		if (GetOpDataType(OpType) == EDataType::Instance)
 		{
-			RunCode(InLiveInstance, InParameters, InAddress);
+			RunCode(InLiveInstance, InAddress);
 			if (!bUnrecoverableError)
 			{
 				Result = InLiveInstance->Cache->GetInstance(FCacheAddress(InAddress, 0, 0));
@@ -1167,7 +1155,7 @@ namespace UE::Mutable::Private
 	}
 
 
-	TManagedPtr<const FLayout> FSystem::BuildLayout(const TSharedRef<FLiveInstance>& InLiveInstance, const FParameters* InParameters, OP::ADDRESS InAddress)
+	TManagedPtr<const FLayout> FSystem::BuildLayout(const TSharedRef<FLiveInstance>& InLiveInstance, OP::ADDRESS InAddress)
 	{
 		WorkingMemoryManager.BeginRunnerThread();
 
@@ -1178,7 +1166,7 @@ namespace UE::Mutable::Private
 			UE::Mutable::Private::EOpType OpType = InLiveInstance->Model->Program.GetOpType(InAddress);
 			if (GetOpDataType(OpType) == EDataType::Layout)
 			{
-				RunCode(InLiveInstance, InParameters, InAddress);
+				RunCode(InLiveInstance, InAddress);
 				if (!bUnrecoverableError)
 				{
 					Result = InLiveInstance->Cache->GetLayout(FCacheAddress(InAddress, 0, 0));
@@ -1192,7 +1180,7 @@ namespace UE::Mutable::Private
 	}
 
 
-	TManagedPtr<const String> FSystem::BuildString(const TSharedRef<FLiveInstance>& InLiveInstance, const FParameters* InParameters, OP::ADDRESS InAddress)
+	TManagedPtr<const String> FSystem::BuildString(const TSharedRef<FLiveInstance>& InLiveInstance, OP::ADDRESS InAddress)
 	{
 		WorkingMemoryManager.BeginRunnerThread();
 
@@ -1203,7 +1191,7 @@ namespace UE::Mutable::Private
 			UE::Mutable::Private::EOpType OpType = InLiveInstance->Model->Program.GetOpType(InAddress);
 			if (GetOpDataType(OpType) == EDataType::String)
 			{
-				RunCode(InLiveInstance, InParameters, InAddress);
+				RunCode(InLiveInstance, InAddress);
 				if (!bUnrecoverableError)
 				{
 					Result = InLiveInstance->Cache->GetString(FCacheAddress(InAddress, 0, 0));
@@ -1681,16 +1669,11 @@ namespace UE::Mutable::Private
 	}
 
 
-	void GetResourceKey(const TSharedPtr<const FModel>& Model, const FParameters* Params, uint32 ParamListIndex, OP::ADDRESS RootAt, TMemoryTrackedArray<uint8>& ParameterValuesBlob)
+	void GetResourceKey(const FModel& Model, const FParameters& Params, uint32 ParamListIndex, OP::ADDRESS RootAt, TMemoryTrackedArray<uint8>& ParameterValuesBlob)
 	{
 		MUTABLE_CPUPROFILER_SCOPE(GetResourceKey);
     	
-		if (!Model)
-		{
-			return;
-		}
-
-		const FProgram& Program = Model->Program;
+		const FProgram& Program = Model.Program;
 
 		// Find the list of relevant parameters
 		const TArray<uint16>* RelevantParams = nullptr;
@@ -1707,7 +1690,7 @@ namespace UE::Mutable::Private
 		// Generate the relevant parameters blob
 		FBitWriterResourceKey Blob(2048*8, true);
 
-		const TArray<FParameterDesc>& ParamDescs = Params->Model->Program.Parameters;
+		const TArray<FParameterDesc>& ParamDescs = Params.Model->Program.Parameters;
 
 		// First make a mask with a bit for each relevant parameter. It will be on for parameters included in the blob.
 		// A parameter will be excluded from the blob if it has the default value, and no multivalues.
@@ -1717,11 +1700,11 @@ namespace UE::Mutable::Private
 			for (int32 IndexIndex = 0; IndexIndex < RelevantParams->Num(); ++IndexIndex)
 			{
 				int32 ParamIndex = (*RelevantParams)[IndexIndex];
-				bool bInclude = Params->HasMultipleValues(ParamIndex);
+				bool bInclude = Params.HasMultipleValues(ParamIndex);
 				if (!bInclude)
 				{
 					bInclude = !(
-						Params->Values[ParamIndex]
+						Params.Values[ParamIndex]
 						==
 						ParamDescs[ParamIndex].DefaultValue
 						);
@@ -1746,12 +1729,12 @@ namespace UE::Mutable::Private
 			switch (Program.Parameters[ParamIndex].Type)
 			{
 			case EParameterType::Bool:
-				Blob.WriteBit(Params->Values[ParamIndex].Get<FParamBoolType>() ? 1 : 0);
+				Blob.WriteBit(Params.Values[ParamIndex].Get<FParamBoolType>() ? 1 : 0);
 
 				// Multi-values
-				if (Params->HasMultipleValues(ParamIndex))
+				if (Params.HasMultipleValues(ParamIndex))
 				{
-					const TMap< TArray<int32>, FParameterValue >& Multi = Params->MultiValues[ParamIndex];
+					const TMap< TArray<int32>, FParameterValue >& Multi = Params.MultiValues[ParamIndex];
 					AddMultiValueKeys(Blob, Multi);
 					for (const TPair< TArray<int32>, FParameterValue >& v : Multi)
 					{
@@ -1763,11 +1746,11 @@ namespace UE::Mutable::Private
 			case EParameterType::Int:
 			{
 				int32 MaxValue = ParamDescs[ParamIndex].PossibleValues.Num();
-				int32 Value = Params->Values[ParamIndex].Get<FParamIntType>();
+				int32 Value = Params.Values[ParamIndex].Get<FParamIntType>();
 				if (MaxValue)
 				{
 					// It is an enum
-					uint32 LimitedValue = FMath::Clamp( Params->GetIntValueIndex(ParamIndex,Value), 0, MaxValue-1 );
+					uint32 LimitedValue = FMath::Clamp( Params.GetIntValueIndex(ParamIndex,Value), 0, MaxValue-1 );
 					Blob.SerializeInt(LimitedValue, uint32(MaxValue));
 				}
 				else
@@ -1778,9 +1761,9 @@ namespace UE::Mutable::Private
 				}
 
 				// Multi-values
-				if (Params->HasMultipleValues(ParamIndex))
+				if (Params.HasMultipleValues(ParamIndex))
 				{
-					const TMap< TArray<int32>, FParameterValue >& Multi = Params->MultiValues[ParamIndex];
+					const TMap< TArray<int32>, FParameterValue >& Multi = Params.MultiValues[ParamIndex];
 					AddMultiValueKeys(Blob, Multi);
 					for (const TPair< TArray<int32>, FParameterValue >& v : Multi)
 					{
@@ -1804,12 +1787,12 @@ namespace UE::Mutable::Private
 
 			case EParameterType::Float:
 				DataSize = sizeof(FParamFloatType);
-				Blob.Serialize(&const_cast<FParameters*>(Params)->Values[ParamIndex].Get<FParamFloatType>(), DataSize);
+				Blob.Serialize(&const_cast<FParameters*>(&Params)->Values[ParamIndex].Get<FParamFloatType>(), DataSize);
 
 				// Multi-values
-				if (Params->HasMultipleValues(ParamIndex))
+				if (Params.HasMultipleValues(ParamIndex))
 				{
-					const TMap< TArray<int32>, FParameterValue >& Multi = Params->MultiValues[ParamIndex];
+					const TMap< TArray<int32>, FParameterValue >& Multi = Params.MultiValues[ParamIndex];
 					AddMultiValueKeys(Blob, Multi);
 					for (const TPair< TArray<int32>, FParameterValue >& v : Multi)
 					{
@@ -1820,12 +1803,12 @@ namespace UE::Mutable::Private
 
 			case EParameterType::Color:
 				DataSize = sizeof(FParamColorType);
-				Blob.Serialize(&const_cast<FParameters*>(Params)->Values[ParamIndex].Get<FParamColorType>(), DataSize);
+				Blob.Serialize(&const_cast<FParameters*>(&Params)->Values[ParamIndex].Get<FParamColorType>(), DataSize);
 
 				// Multi-values
-				if (Params->HasMultipleValues(ParamIndex))
+				if (Params.HasMultipleValues(ParamIndex))
 				{
-					const TMap< TArray<int32>, FParameterValue >& Multi = Params->MultiValues[ParamIndex];
+					const TMap< TArray<int32>, FParameterValue >& Multi = Params.MultiValues[ParamIndex];
 					AddMultiValueKeys(Blob, Multi);
 					for (const TPair< TArray<int32>, FParameterValue >& v : Multi)
 					{
@@ -1837,7 +1820,7 @@ namespace UE::Mutable::Private
 			case EParameterType::Projector:
 			{
 				FPackedNormal TempVec;
-				const FProjector& Value = Params->Values[ParamIndex].Get<FParamProjectorType>();
+				const FProjector& Value = Params.Values[ParamIndex].Get<FParamProjectorType>();
 				Blob.Serialize((void*)&Value.position, sizeof(FVector3f));
 				TempVec = FPackedNormal(Value.direction);
 				Blob.Serialize(&TempVec, sizeof(FPackedNormal));
@@ -1847,9 +1830,9 @@ namespace UE::Mutable::Private
 				Blob.Serialize((void*)&Value.projectionAngle, sizeof(float));
 
 				// Multi-values
-				if (Params->HasMultipleValues(ParamIndex))
+				if (Params.HasMultipleValues(ParamIndex))
 				{
-					const TMap< TArray<int32>, FParameterValue >& Multi = Params->MultiValues[ParamIndex];
+					const TMap< TArray<int32>, FParameterValue >& Multi = Params.MultiValues[ParamIndex];
 					AddMultiValueKeys(Blob, Multi);
 					for (const TPair< TArray<int32>, FParameterValue >& v : Multi)
 					{
@@ -1868,14 +1851,14 @@ namespace UE::Mutable::Private
 
 			case EParameterType::Texture:
 			{
-				TStrongObjectPtr<UTexture> Object = Params->Values[ParamIndex].Get<FParamTextureType>();
+				TStrongObjectPtr<UTexture> Object = Params.Values[ParamIndex].Get<FParamTextureType>();
 				FString Path = Object ? Object->GetPathName() : FString();
 				Blob.Serialize(GetData(Path), Path.Len() * sizeof(FString::ElementType));
 
 				// Multi-values
-				if (Params->HasMultipleValues(ParamIndex))
+				if (Params.HasMultipleValues(ParamIndex))
 				{
-					const TMap< TArray<int32>, FParameterValue >& Multi = Params->MultiValues[ParamIndex];
+					const TMap< TArray<int32>, FParameterValue >& Multi = Params.MultiValues[ParamIndex];
 					AddMultiValueKeys(Blob, Multi);
 					for (const TPair< TArray<int32>, FParameterValue >& v : Multi)
 					{
@@ -1889,14 +1872,14 @@ namespace UE::Mutable::Private
 
 			case EParameterType::SkeletalMesh:
 			{
-				TStrongObjectPtr<USkeletalMesh> Object = Params->Values[ParamIndex].Get<FParamSkeletalMeshType>();
+				TStrongObjectPtr<USkeletalMesh> Object = Params.Values[ParamIndex].Get<FParamSkeletalMeshType>();
 				FString Path = Object ? Object->GetPathName() : FString();
 				Blob.Serialize(GetData(Path), Path.Len() * sizeof(FString::ElementType));
 
 				// Multi-values
-				if (Params->HasMultipleValues(ParamIndex))
+				if (Params.HasMultipleValues(ParamIndex))
 				{
-					const TMap< TArray<int32>, FParameterValue >& Multi = Params->MultiValues[ParamIndex];
+					const TMap< TArray<int32>, FParameterValue >& Multi = Params.MultiValues[ParamIndex];
 					AddMultiValueKeys(Blob, Multi);
 					for (const TPair< TArray<int32>, FParameterValue >& v : Multi)
 					{
@@ -1910,14 +1893,14 @@ namespace UE::Mutable::Private
 
 			case EParameterType::Material:
 			{
-				TStrongObjectPtr<UMaterialInterface> Object = Params->Values[ParamIndex].Get<FParamMaterialType>();
+				TStrongObjectPtr<UMaterialInterface> Object = Params.Values[ParamIndex].Get<FParamMaterialType>();
 				FString Path = Object ? Object->GetPathName() : FString();
 				Blob.Serialize(GetData(Path), Path.Len() * sizeof(FString::ElementType));
 
 				// Multi-values
-				if (Params->HasMultipleValues(ParamIndex))
+				if (Params.HasMultipleValues(ParamIndex))
 				{
-					const TMap< TArray<int32>, FParameterValue >& Multi = Params->MultiValues[ParamIndex];
+					const TMap< TArray<int32>, FParameterValue >& Multi = Params.MultiValues[ParamIndex];
 					AddMultiValueKeys(Blob, Multi);
 					for (const TPair< TArray<int32>, FParameterValue >& v : Multi)
 					{
@@ -1933,13 +1916,13 @@ namespace UE::Mutable::Private
 			{
 				DataSize = sizeof(FMatrix44f);
 
-				const FMatrix44f& Value = Params->Values[ParamIndex].Get<FParamMatrixType>();
+				const FMatrix44f& Value = Params.Values[ParamIndex].Get<FParamMatrixType>();
 				Blob.Serialize((void*)&Value, DataSize);
 
 				// Multi-values
-				if (Params->HasMultipleValues(ParamIndex))
+				if (Params.HasMultipleValues(ParamIndex))
 				{
-					const TMap<TArray<int32>, FParameterValue>& Multi = Params->MultiValues[ParamIndex];
+					const TMap<TArray<int32>, FParameterValue>& Multi = Params.MultiValues[ParamIndex];
 					AddMultiValueKeys(Blob, Multi);
 					for (const TPair<TArray<int32>, FParameterValue>& MultiValuePair : Multi)
 					{
@@ -1952,15 +1935,15 @@ namespace UE::Mutable::Private
 
 			case EParameterType::InstancedStruct:
 			{
-				const FParamInstancedStructType& Value = Params->Values[ParamIndex].Get<FParamInstancedStructType>();
+				const FParamInstancedStructType& Value = Params.Values[ParamIndex].Get<FParamInstancedStructType>();
 
 				FInstancedStruct* ValueInstancedStruct = const_cast<FInstancedStruct*>(Value.Get());
 				ValueInstancedStruct->Serialize(Blob);
 
 				// Multi-values
-				if (Params->HasMultipleValues(ParamIndex))
+				if (Params.HasMultipleValues(ParamIndex))
 				{
-					const TMap<TArray<int32>, FParameterValue>& Multi = Params->MultiValues[ParamIndex];
+					const TMap<TArray<int32>, FParameterValue>& Multi = Params.MultiValues[ParamIndex];
 					AddMultiValueKeys(Blob, Multi);
 					for (const TPair<TArray<int32>, FParameterValue>& MultiValuePair : Multi)
 					{
@@ -1985,14 +1968,14 @@ namespace UE::Mutable::Private
 	}
 
 
-	FMeshId FSystem::GetMeshId(const TSharedPtr<FMeshIdRegistry>& Registry, const TSharedPtr<const FModel>& Model, const FParameters* Params, OP::ADDRESS RootAt)
+	FMeshId FSystem::GetMeshId(const TSharedPtr<FMeshIdRegistry>& Registry, const FModel& Model, const FParameters& Params, OP::ADDRESS RootAt)
     {
     	check(Registry)
     	
     	FGeneratedMeshKey NewKey;
     	NewKey.Address = RootAt;
 
-    	const int32* Result = Model->Program.RelevantParameterList.Find(RootAt);
+    	const int32* Result = Model.Program.RelevantParameterList.Find(RootAt);
     	const int32 ParamListIndex = Result ? *Result : INDEX_NONE; 
 		GetResourceKey(Model, Params, ParamListIndex, RootAt, NewKey.ParameterValuesBlob);
 
@@ -2000,14 +1983,14 @@ namespace UE::Mutable::Private
     }
 	
     	
-	FImageId FSystem::GetImageId(const TSharedPtr<FImageIdRegistry>& Registry, const TSharedPtr<const FModel>& Model, const FParameters* Params, OP::ADDRESS RootAt)
+	FImageId FSystem::GetImageId(const TSharedPtr<FImageIdRegistry>& Registry, const FModel& Model, const FParameters& Params, OP::ADDRESS RootAt)
     {
     	check(Registry)
     
     	FGeneratedImageKey NewKey;
     	NewKey.Address = RootAt;
 
-    	const int32* Result = Model->Program.RelevantParameterList.Find(RootAt);
+    	const int32* Result = Model.Program.RelevantParameterList.Find(RootAt);
     	const int32 ParamListIndex = Result ? *Result : INDEX_NONE; 
     	GetResourceKey(Model, Params, ParamListIndex, RootAt, NewKey.ParameterValuesBlob);
 
@@ -2015,14 +1998,14 @@ namespace UE::Mutable::Private
     }
 
 	
-	FMaterialId FSystem::GetMaterialId(const TSharedPtr<FMaterialIdRegistry>& Registry, const TSharedPtr<const FModel>& Model, const FParameters* Params, OP::ADDRESS RootAt)
+	FMaterialId FSystem::GetMaterialId(const TSharedPtr<FMaterialIdRegistry>& Registry, const FModel& Model, const FParameters& Params, OP::ADDRESS RootAt)
     {
     	check(Registry)
     
     	FGeneratedMaterialKey NewKey;
     	NewKey.Address = RootAt;
 
-    	const int32* Result = Model->Program.RelevantParameterList.Find(RootAt);
+    	const int32* Result = Model.Program.RelevantParameterList.Find(RootAt);
     	const int32 ParamListIndex = Result ? *Result : INDEX_NONE; 
     	GetResourceKey(Model, Params, ParamListIndex, RootAt, NewKey.ParameterValuesBlob);
 
@@ -2030,14 +2013,14 @@ namespace UE::Mutable::Private
     }
 
 	
-	FSkeletalMeshId FSystem::GetSkeletalMeshId(const TSharedPtr<FSkeletalMeshIdRegistry>& Registry, const TSharedPtr<const FModel>& Model, const FParameters* Params, OP::ADDRESS RootAt)
+	FSkeletalMeshId FSystem::GetSkeletalMeshId(const TSharedPtr<FSkeletalMeshIdRegistry>& Registry, const FModel& Model, const FParameters& Params, OP::ADDRESS RootAt)
 	{
 		check(Registry)
 	
     	FGeneratedSkeletalMeshKey NewKey;
     	NewKey.Address = RootAt;
 
-    	const int32* Result = Model->Program.RelevantParameterList.Find(RootAt);
+    	const int32* Result = Model.Program.RelevantParameterList.Find(RootAt);
     	const int32 ParamListIndex = Result ? *Result : INDEX_NONE; 
     	GetResourceKey(Model, Params, ParamListIndex, RootAt, NewKey.ParameterValuesBlob);
 
diff --git a/Engine/Plugins/Mutable/Source/MutableTools/Private/MuT/CodeOptimiser.cpp b/Engine/Plugins/Mutable/Source/MutableTools/Private/MuT/CodeOptimiser.cpp
--- a/Engine/Plugins/Mutable/Source/MutableTools/Private/MuT/CodeOptimiser.cpp
+++ b/Engine/Plugins/Mutable/Source/MutableTools/Private/MuT/CodeOptimiser.cpp
@@ -196,6 +196,7 @@ namespace UE::Mutable::Private
 			TSharedPtr<FParameters> LocalParams = FModel::NewParameters(Model);
 
 			TSharedRef<FLiveInstance> TempLiveInstance = System->NewBuildInstance(Model, nullptr);
+			TempLiveInstance->OldParameters = LocalParams;
 			
 			System->BeginBuild(Model);
 			
@@ -206,7 +207,7 @@ namespace UE::Mutable::Private
 			{
 				MUTABLE_CPUPROFILER_SCOPE(ConstantMesh);
 
-				TManagedPtr<const FMesh> MeshBuild = System->BuildMesh(TempLiveInstance, LocalParams.Get(), at, EMeshContentFlags::AllFlags);
+				TManagedPtr<const FMesh> MeshBuild = System->BuildMesh(TempLiveInstance, at, EMeshContentFlags::AllFlags);
 
 				if (MeshBuild)
 				{
@@ -223,7 +224,7 @@ namespace UE::Mutable::Private
 			{
 				MUTABLE_CPUPROFILER_SCOPE(ConstantImage);
 
-				TManagedPtr<const FImage> ImageResult = System->BuildImage(TempLiveInstance, LocalParams.Get(), at, 0);
+				TManagedPtr<const FImage> ImageResult = System->BuildImage(TempLiveInstance, at, 0);
 
 				if (ImageResult)
 				{
@@ -240,7 +241,7 @@ namespace UE::Mutable::Private
 			{
 				MUTABLE_CPUPROFILER_SCOPE(ConstantLayout);
 
-				TManagedPtr<const FLayout> LayoutResult = System->BuildLayout(TempLiveInstance, LocalParams.Get(), at);
+				TManagedPtr<const FLayout> LayoutResult = System->BuildLayout(TempLiveInstance, at);
 
 				if (LayoutResult)
 				{
@@ -256,7 +257,7 @@ namespace UE::Mutable::Private
 			{
 				MUTABLE_CPUPROFILER_SCOPE(ConstantBool);
 
-				bool value = System->BuildBool(TempLiveInstance, LocalParams.Get(), at);
+				bool value = System->BuildBool(TempLiveInstance, at);
 				Result = new ASTOpConstantBool(value);
 				break;
 			}
@@ -266,7 +267,7 @@ namespace UE::Mutable::Private
 				MUTABLE_CPUPROFILER_SCOPE(ConstantBool);
 
 				FVector4f ResultColor(0, 0, 0, 0);
-				ResultColor = System->BuildColor(TempLiveInstance, LocalParams.Get(), at);
+				ResultColor = System->BuildColor(TempLiveInstance, at);
 
 				{
 					UE::Mutable::Private::Ptr<ASTOpConstantColor> ConstantOp = new ASTOpConstantColor();
