[Mutable] Fixed crash of unexpected parameter name in a Break Material Node.

#rnx

#Synced-CL 49838890

diff --git a/Engine/Plugins/Mutable/Source/MutableRuntime/Internal/MuR/System.h b/Engine/Plugins/Mutable/Source/MutableRuntime/Internal/MuR/System.h
--- a/Engine/Plugins/Mutable/Source/MutableRuntime/Internal/MuR/System.h
+++ b/Engine/Plugins/Mutable/Source/MutableRuntime/Internal/MuR/System.h
@@ -764,7 +764,12 @@ namespace UE::Mutable::Private
 				return nullptr;
 			}
 
-			check(Data->DataType == EDataType::Instance);
+			// Aborted Runs may not have stored the data type of the Address operation
+			if (Data->DataType != EDataType::Instance)
+			{
+				return nullptr;
+			}
+
 			TManagedPtr<const FInstance> Result = InstanceResults[Data->DataTypeIndex].Pin();
 
 			// We need to decrease the hit-count even if the result is null.
diff --git a/Engine/Plugins/Mutable/Source/MutableRuntime/Private/MuR/CodeRunner.cpp b/Engine/Plugins/Mutable/Source/MutableRuntime/Private/MuR/CodeRunner.cpp
--- a/Engine/Plugins/Mutable/Source/MutableRuntime/Private/MuR/CodeRunner.cpp
+++ b/Engine/Plugins/Mutable/Source/MutableRuntime/Private/MuR/CodeRunner.cpp
@@ -5737,9 +5737,11 @@ namespace UE::Mutable::Private
 				const FString& ParameterName = Program.ConstantStrings[Args.ParameterName];
 				const FParameterKey ParameterKey = { FName(ParameterName), Args.LayerIndex };
 				const FMaterial::FImageParameterData* ImageParameterData = Material->ImageParameters.Find(ParameterKey);
+				bool bParameterFound = false;
 
 				if (ImageParameterData)
 				{
+					bParameterFound = true;
 					const TVariant<OP::ADDRESS, TManagedPtr<const FImage>> ImageParameter = ImageParameterData->ImageParameter;
 
 					if (ImageParameter.IsType<OP::ADDRESS>())
@@ -5755,38 +5757,30 @@ namespace UE::Mutable::Private
 						StoreImage(Item, ImageParameter.Get<TManagedPtr<const FImage>>());
 					}
 				}
-				else if (Material->PassthroughObject.IsResolved())
+				else if (const UMaterialInterface* MaterialInterface = Material->PassthroughObject.Get())
 				{
 					UTexture* ParameterValue = nullptr;
+					FHashedMaterialParameterInfo ParameterInfo;
+					ParameterInfo.Name = FScriptName(FName(ParameterName));
+					ParameterInfo.Index = ParameterKey.LayerIndex;
+					ParameterInfo.Association = ParameterInfo.Index == INDEX_NONE ? GlobalParameter : LayerParameter;
 
-					if (const UMaterialInterface* MaterialInterface = Material->PassthroughObject.Get())
-					{
-						FHashedMaterialParameterInfo ParameterInfo;
-						ParameterInfo.Name = FScriptName(FName(ParameterName));
-						ParameterInfo.Index = ParameterKey.LayerIndex;
-						ParameterInfo.Association = ParameterInfo.Index == INDEX_NONE ? GlobalParameter : LayerParameter;
+					// Find image using the parameter info
+					bParameterFound = MaterialInterface->GetTextureParameterValue(ParameterInfo, ParameterValue);
 
-						// Find image using the parameter info
-						MaterialInterface->GetTextureParameterValue(ParameterInfo, ParameterValue);
-					}
-
-					if (ParameterValue)
+					if (bParameterFound)
 					{
 						TPassthroughObjectPtr<UTexture> ExternalObject(ParameterValue);
 						TManagedPtr<FImage> Result = FImage::CreateAsReference(ExternalObject, FImageDesc(), false);
 
 						StoreImage(Item, Result);
 					}
-					else
-					{
-						AbortRun();
-						UE_LOG(LogMutableCore, Warning, TEXT("Expected Passthrough Texture parameter [%s] not found in Break Material node."), *ParameterName);
-					}
 				}
-				else
+				
+				if(!bParameterFound)
 				{
-					check(false);
-					StoreImage(Item, nullptr);
+					UE_LOG(LogMutableCore, Error, TEXT("Expected Texture parameter [%s] not found in Break Material node."), *ParameterName);
+					AbortRun();
 				}
 
 				break;
@@ -6276,41 +6270,36 @@ namespace UE::Mutable::Private
 				const FString& ParameterName = Program.ConstantStrings[Args.ParameterName];
 				const FParameterKey ParameterKey = { FName(ParameterName), Args.LayerIndex };
 				const float* Value = Material->ScalarParameters.Find(ParameterKey);
+				bool bParameterFound = false;
 
 				if (Value)
 				{
+					bParameterFound = true;
+
 					StoreScalar(Item, *Value);
 				}
-				else if (Material->PassthroughObject.IsResolved())
+				else if (const UMaterialInterface* MaterialInterface = Material->PassthroughObject.Get())
 				{
-					bool bParameterFound = false;
 					float ParameterValue = 0.0f;
 
-					if (const UMaterialInterface* MaterialInterface = Material->PassthroughObject.Get())
-					{
-						FHashedMaterialParameterInfo ParameterInfo;
-						ParameterInfo.Name = FScriptName(FName(ParameterName));
-						ParameterInfo.Index = ParameterKey.LayerIndex;
-						ParameterInfo.Association = ParameterInfo.Index == INDEX_NONE ? GlobalParameter : LayerParameter;
+					FHashedMaterialParameterInfo ParameterInfo;
+					ParameterInfo.Name = FScriptName(FName(ParameterName));
+					ParameterInfo.Index = ParameterKey.LayerIndex;
+					ParameterInfo.Association = ParameterInfo.Index == INDEX_NONE ? GlobalParameter : LayerParameter;
 
-						// Find float using the parameter info
-						bParameterFound = MaterialInterface->GetScalarParameterValue(ParameterInfo, ParameterValue);
-					}
+					// Find float using the parameter info
+					bParameterFound = MaterialInterface->GetScalarParameterValue(ParameterInfo, ParameterValue);
 
 					if (bParameterFound)
 					{
 						StoreScalar(Item, ParameterValue);
 					}
-					else
-					{
-						AbortRun();
-						UE_LOG(LogMutableCore, Warning, TEXT("Expected Scalar parameter [%s] not found in Break Material node."), *ParameterName);
-					}
 				}
-				else
+
+				if(!bParameterFound)
 				{
-					check(false);
-					StoreScalar(Item, 0.0f);
+					UE_LOG(LogMutableCore, Error, TEXT("Expected Scalar parameter [%s] not found in Break Material node."), *ParameterName);
+					AbortRun();
 				}
 
 				break;
@@ -6621,41 +6610,36 @@ namespace UE::Mutable::Private
 				const FString& ParameterName = Program.ConstantStrings[Args.ParameterName];
 				const FParameterKey ParameterKey = { FName(ParameterName), Args.LayerIndex };
 				const FVector4f* Value = Material->ColorParameters.Find(ParameterKey);
+				bool bParameterFound = false;
 
 				if (Value)
 				{
+					bParameterFound = true;
+
 					StoreColor(Item, *Value);
 				}
-				else if (Material->PassthroughObject.IsResolved())
+				else if (const UMaterialInterface* MaterialInterface = Material->PassthroughObject.Get())
 				{
-					bool bParameterFound = false;
 					FLinearColor ParameterValue;
 
-					if (const UMaterialInterface* MaterialInterface = Material->PassthroughObject.Get())
-					{
-						FHashedMaterialParameterInfo ParameterInfo;
-						ParameterInfo.Name = FScriptName(FName(ParameterName));
-						ParameterInfo.Index = ParameterKey.LayerIndex;
-						ParameterInfo.Association = ParameterInfo.Index == INDEX_NONE ? GlobalParameter : LayerParameter;
+					FHashedMaterialParameterInfo ParameterInfo;
+					ParameterInfo.Name = FScriptName(FName(ParameterName));
+					ParameterInfo.Index = ParameterKey.LayerIndex;
+					ParameterInfo.Association = ParameterInfo.Index == INDEX_NONE ? GlobalParameter : LayerParameter;
 
-						// Find float using the parameter info
-						bParameterFound = MaterialInterface->GetVectorParameterValue(ParameterInfo, ParameterValue);
-					}
+					// Find float using the parameter info
+					bParameterFound = MaterialInterface->GetVectorParameterValue(ParameterInfo, ParameterValue);
 
 					if (bParameterFound)
 					{
 						StoreColor(Item, ParameterValue);
 					}
-					else
-					{
-						AbortRun();
-						UE_LOG(LogMutableCore, Warning, TEXT("Expected Color parameter [%s] not found in Break Material node."), *ParameterName);
-					}
 				}
-				else
+
+				if(!bParameterFound)
 				{
-					check(false);
-					StoreColor(Item, FVector4f::Zero());
+					UE_LOG(LogMutableCore, Error, TEXT("Expected Color parameter [%s] not found in Break Material node."), *ParameterName);
+					AbortRun();
 				}
 
 				break;
@@ -7033,7 +7017,7 @@ namespace UE::Mutable::Private
 					MutableMaterial->PassthroughObject = TPassthroughObjectPtr(FoundMaterialInterface.Get());
 
 					StoreMaterial(Item, MutableMaterial);
-					
+
 					break;
 				}
 			
@@ -7220,7 +7204,7 @@ namespace UE::Mutable::Private
 
 
     //---------------------------------------------------------------------------------------------
-    void CodeRunner::RunCode(const FScheduledOp& Item, uint32 LodMask, bool& bOutFailed)
+    void CodeRunner::RunCode(const FScheduledOp& Item, uint32 LodMask)
     {
 		//UE_LOG( LogMutableCore, Log, TEXT("Running :%5d , %d "), Item.At, Item.Stage );
 		check( Item.Type == FScheduledOp::EType::Full );
@@ -7294,7 +7278,8 @@ namespace UE::Mutable::Private
 				if (!OperationInstancedStruct.IsValid()) // Plugin not loaded
 				{
                     UE_LOG(LogMutableCore, Error, TEXT("Trying to evaluate an External Operation which has not been loaded."));
-					bOutFailed = true;
+					AbortRun();
+
 					break;
 				}
 					
@@ -7304,7 +7289,8 @@ namespace UE::Mutable::Private
 				if (Inputs.Num() != Args.NumOperants)
 				{
 					UE_LOG(LogMutableCore, Error, TEXT("External Operation inputs do not match compiled data: %s"), *OperationInstancedStruct.GetScriptStruct()->GetName());
-					bOutFailed = true;
+					AbortRun();
+
 					return;
 				}
 				
@@ -7428,7 +7414,8 @@ namespace UE::Mutable::Private
 						if (!Context.Output.Ptr)
 						{
 							UE_LOG(LogMutableCore, Error, TEXT("External Operation output not set: %s"), *OperationInstancedStruct.GetScriptStruct()->GetName());
-							bOutFailed = true;
+							AbortRun();
+
 							return;
 						}
 						
@@ -8727,6 +8714,10 @@ namespace UE::Mutable::Private
 						StoreValidDesc(Item);
 					}
 				}
+				else if (Material->PassthroughObject.IsResolved())
+				{
+					//TODO(Max)
+				}
 				else
 				{
 					ImageDescResults.FindOrAdd(Item.At);
diff --git a/Engine/Plugins/Mutable/Source/MutableRuntime/Private/MuR/CodeRunner.h b/Engine/Plugins/Mutable/Source/MutableRuntime/Private/MuR/CodeRunner.h
--- a/Engine/Plugins/Mutable/Source/MutableRuntime/Private/MuR/CodeRunner.h
+++ b/Engine/Plugins/Mutable/Source/MutableRuntime/Private/MuR/CodeRunner.h
@@ -123,7 +123,7 @@ namespace  UE::Mutable::Private
     	
         TSharedPtr<FRangeIndex> BuildCurrentOpRangeIndex( const FScheduledOp&, int32 ParameterIndex);
 
-        void RunCode( const FScheduledOp&, uint32 LodMask, bool& bOutFailed);
+        void RunCode( const FScheduledOp&, uint32 LodMask);
 
         void RunCode_Conditional(const FScheduledOp&);
         void RunCode_Switch(const FScheduledOp&);
diff --git a/Engine/Plugins/Mutable/Source/MutableRuntime/Private/MuR/CodeRunnerTasked.cpp b/Engine/Plugins/Mutable/Source/MutableRuntime/Private/MuR/CodeRunnerTasked.cpp
--- a/Engine/Plugins/Mutable/Source/MutableRuntime/Private/MuR/CodeRunnerTasked.cpp
+++ b/Engine/Plugins/Mutable/Source/MutableRuntime/Private/MuR/CodeRunnerTasked.cpp
@@ -303,11 +303,10 @@ namespace UE::Mutable::Private
 				else
 				{
 					// Run immediately
-					bool bFailed = false;
-					RunCode(Item, LODMask, bFailed);
-					if (bFailed)
+					RunCode(Item, LODMask);
+					if (bUnrecoverableError)
 					{
-						return AbortRun();
+						return;
 					}
 
 					if (ScheduledStagePerOp[Item] == Item.Stage + 1)
