[Mutable] Fix CO loading incorrect hierarchy.

#Synced-CL 49839474

diff --git a/Engine/Plugins/Mutable/Source/CustomizableObjectEditor/Private/MuCOE/GraphTraversal.cpp b/Engine/Plugins/Mutable/Source/CustomizableObjectEditor/Private/MuCOE/GraphTraversal.cpp
--- a/Engine/Plugins/Mutable/Source/CustomizableObjectEditor/Private/MuCOE/GraphTraversal.cpp
+++ b/Engine/Plugins/Mutable/Source/CustomizableObjectEditor/Private/MuCOE/GraphTraversal.cpp
@@ -808,6 +808,11 @@ void GetNodeGroupObjectNodeMappingImmersive(UCustomizableObject* Object, IAssetR
 		}
 	}
 	
+	if (ArrayReferenceNames.IsEmpty())
+	{
+		return;
+	}
+
 	Filter.ClassPaths.Add(UCustomizableObject::StaticClass()->GetClassPathName());
 	
 	TArray<FAssetData> ArrayAssetData;
