[Mutable] Fix CO loading hang when a Child Object is not properly connected to a group.

#Synced-CL 49839474

diff --git a/Engine/Plugins/Mutable/Source/CustomizableObjectEditor/Private/MuCOE/CustomizableObjectEditorModule.cpp b/Engine/Plugins/Mutable/Source/CustomizableObjectEditor/Private/MuCOE/CustomizableObjectEditorModule.cpp
--- a/Engine/Plugins/Mutable/Source/CustomizableObjectEditor/Private/MuCOE/CustomizableObjectEditorModule.cpp
+++ b/Engine/Plugins/Mutable/Source/CustomizableObjectEditor/Private/MuCOE/CustomizableObjectEditorModule.cpp
@@ -1238,7 +1238,11 @@ void FCustomizableObjectEditorModule::Tick(bool bBlocking)
 			}
 			
 			if (bHierarchyLoaded)
-			{
+			{	
+				// Traversing the graph upwards does not take into account Group Object Nodes. 
+				// A Child Object may reference a parent but not a Group Object. This can cause the loaded CO to not appear in the Related objects. Added manually to fix this.
+				RelatedObjects.Add(Object);
+				
 				for (UCustomizableObject* RelatedObject : RelatedObjects)
 				{
 					if (ObjectsPendingLoad.Remove(TStrongObjectPtr(RelatedObject)))
