{Mutable] Reverted GetInstance function if case that could hide a critical failure.

#rnx
#Synced-CL 49838890

diff --git a/Engine/Plugins/Mutable/Source/MutableRuntime/Internal/MuR/System.h b/Engine/Plugins/Mutable/Source/MutableRuntime/Internal/MuR/System.h
--- a/Engine/Plugins/Mutable/Source/MutableRuntime/Internal/MuR/System.h
+++ b/Engine/Plugins/Mutable/Source/MutableRuntime/Internal/MuR/System.h
@@ -764,12 +764,7 @@ namespace UE::Mutable::Private
 				return nullptr;
 			}
 
-			// Aborted Runs may not have stored the data type of the Address operation
-			if (Data->DataType != EDataType::Instance)
-			{
-				return nullptr;
-			}
-
+			check(Data->DataType == EDataType::Instance)
 			TManagedPtr<const FInstance> Result = InstanceResults[Data->DataTypeIndex].Pin();
 
 			// We need to decrease the hit-count even if the result is null.
diff --git a/Engine/Plugins/Mutable/Source/MutableRuntime/Private/MuR/System.cpp b/Engine/Plugins/Mutable/Source/MutableRuntime/Private/MuR/System.cpp
--- a/Engine/Plugins/Mutable/Source/MutableRuntime/Private/MuR/System.cpp
+++ b/Engine/Plugins/Mutable/Source/MutableRuntime/Private/MuR/System.cpp
@@ -264,11 +264,15 @@ namespace UE::Mutable::Private
 		OP::ADDRESS RootAt = Program.States[InStateIndex].Root;
 		RunCode(LiveInstance, RootAt, InLodMask);
 
-		TManagedPtr<const FInstance> Result = LiveInstance->Cache->GetInstance(FCacheAddress(RootAt, 0, 0));
+		TManagedPtr<const FInstance> Result;
+
+		if (!bUnrecoverableError)
+		{
+			Result = LiveInstance->Cache->GetInstance(FCacheAddress(RootAt, 0, 0));
+		}
 
 		// Debug check to see if we managed the op-hit-counts correctly
 		LiveInstance->Cache->CheckHitCountsCleared();
-
 		LiveInstance->Instance = Result;
 
     	if (!Result)
