[Mutable] Move all Rom related member/functions to FProgram.

- Moved Model.h implementations to the source file.

#Synced-CL 49839474

#Synced-CL 49839474

diff --git a/Engine/Plugins/Mutable/Source/CustomizableObject/Private/MuCO/CustomizableObject.cpp b/Engine/Plugins/Mutable/Source/CustomizableObject/Private/MuCO/CustomizableObject.cpp
--- a/Engine/Plugins/Mutable/Source/CustomizableObject/Private/MuCO/CustomizableObject.cpp
+++ b/Engine/Plugins/Mutable/Source/CustomizableObject/Private/MuCO/CustomizableObject.cpp
@@ -2616,12 +2616,12 @@ namespace UE::Mutable::Private
 			// Add High-res textures to the Optional bucket
 			OptionalBucket.Categories.Add({ FFileCategoryID(EStreamableDataType::Model, uint16(UE::Mutable::Private::EDataType::Image), uint16(EMutableFileFlags::HighRes)), 0, 0, 0});
 
-			const int32 NumRoms = Model->GetRomCount();
+			const int32 NumRoms = Model->Program.GetRomCount();
 			for (int32 RomIndex = 0; RomIndex < NumRoms; ++RomIndex)
 			{
-				const uint32 BlockSize = Model->GetRomSize(RomIndex);
-				const uint16 BlockResourceType = Model->IsMeshData(RomIndex) ? uint16(UE::Mutable::Private::EDataType::Mesh) : uint16(UE::Mutable::Private::EDataType::Image);
-				const EMutableFileFlags BlockFlags = Model->IsRomHighRes(RomIndex) ? EMutableFileFlags::HighRes : EMutableFileFlags::None;
+				const uint32 BlockSize = Model->Program.GetRomSize(RomIndex);
+				const uint16 BlockResourceType = Model->Program.GetRomDataType(RomIndex) == ERomDataType::Mesh ? uint16(UE::Mutable::Private::EDataType::Mesh) : uint16(UE::Mutable::Private::EDataType::Image);
+				const EMutableFileFlags BlockFlags = Model->Program.IsRomHighRes(RomIndex) ? EMutableFileFlags::HighRes : EMutableFileFlags::None;
 
 				FFileCategoryID CategoryID = { EStreamableDataType::Model, BlockResourceType, uint16(BlockFlags) };
 				TPair<FFileBucket&, FFileCategory&> It = FindOrAddCategory(FileBuckets, DefaultBucket, CategoryID); // Add block to an existing or new category
@@ -2979,23 +2979,23 @@ namespace UE::Mutable::Private
 		}
 
 		// Root nodes by flags.
-		const uint32 NumRoms = Model->GetRomCount();
+		const uint32 NumRoms = Model->Program.GetRomCount();
 		TMap<FFileCategoryID, FClassifyNode> RootNode;
 
 		// Create blocks data.
 		{
 			for (uint32 RomIndex = 0; RomIndex < NumRoms; ++RomIndex)
 			{
-				uint32 SourceBlockId = Model->GetRomSourceId(RomIndex);
-				const uint32 BlockSize = Model->GetRomSize(RomIndex);
+				uint32 SourceBlockId = Model->Program.GetRomSourceId(RomIndex);
+				const uint32 BlockSize = Model->Program.GetRomSize(RomIndex);
 				
 				uint16 BlockResourceType = 0;
 				EMutableFileFlags BlockFlags = EMutableFileFlags::None;
 				// TODO: Temp. Remove after unifying generated output files code between editor an package. UE-222777 
 				if (bUseRomTypeAndFlagsToFilter)
 				{
-					BlockResourceType = Model->IsMeshData(RomIndex) ? uint16(UE::Mutable::Private::EDataType::Mesh) : uint16(UE::Mutable::Private::EDataType::Image);
-					BlockFlags = Model->IsRomHighRes(RomIndex) ? EMutableFileFlags::HighRes : EMutableFileFlags::None;
+					BlockResourceType = Model->Program.GetRomDataType(RomIndex) == ERomDataType::Mesh ? uint16(UE::Mutable::Private::EDataType::Mesh) : uint16(UE::Mutable::Private::EDataType::Image);
+					BlockFlags = Model->Program.IsRomHighRes(RomIndex) ? EMutableFileFlags::HighRes : EMutableFileFlags::None;
 				}
 
 				FFileCategoryID CurrentCategory = { EStreamableDataType::Model, BlockResourceType, uint16(BlockFlags) };
diff --git a/Engine/Plugins/Mutable/Source/CustomizableObject/Private/MuCO/CustomizableObjectMipDataProvider.cpp b/Engine/Plugins/Mutable/Source/CustomizableObject/Private/MuCO/CustomizableObjectMipDataProvider.cpp
--- a/Engine/Plugins/Mutable/Source/CustomizableObject/Private/MuCO/CustomizableObjectMipDataProvider.cpp
+++ b/Engine/Plugins/Mutable/Source/CustomizableObject/Private/MuCO/CustomizableObjectMipDataProvider.cpp
@@ -366,16 +366,16 @@ int32 FMutableTextureMipDataProvider::GetMips(const FTextureUpdateContext& Conte
 
 				for (int32 ImageToLoadIndex : ImageRef.ConstantImagesNeededToGenerate)
 				{
-					const int32 RomId = Model->GetConstantImageRomId(ImageToLoadIndex, LODIndex);
+					const int32 RomId = Model->Program.GetConstantImageRomId(ImageToLoadIndex, LODIndex);
 
-					if (RomId >= 0 && Model->IsRomHighRes(RomId))
+					if (RomId >= 0 && Model->Program.IsRomHighRes(RomId))
 					{
 						FMutableStreamableBlock& Block = ModelStreamableBulkData->ModelStreamables[RomId];
 						if (!Block.IsPrefetched)
 						{
 							Block.IsPrefetched = true;
 
-							MaxBlockSize = FMath::Max(MaxBlockSize, Model->GetRomSize(RomId));
+							MaxBlockSize = FMath::Max(MaxBlockSize, Model->Program.GetRomSize(RomId));
 							RomsToPrefetch.Add(RomId);
 						}
 					}
@@ -395,7 +395,7 @@ int32 FMutableTextureMipDataProvider::GetMips(const FTextureUpdateContext& Conte
 					ModelStreamableBulkData->ModelStreamables[RomId],
 					UE::Mutable::Private::EStreamableDataType::Model,
 					(uint16)UE::Mutable::Private::EDataType::Image,
-					TArrayView<uint8>(AllocatedMemory, Model->GetRomSize(RomId)));
+					TArrayView<uint8>(AllocatedMemory, Model->Program.GetRomSize(RomId)));
 			}
 		}
 
diff --git a/Engine/Plugins/Mutable/Source/CustomizableObjectEditor/Private/MuCOE/CustomizableObjectCompiler.cpp b/Engine/Plugins/Mutable/Source/CustomizableObjectEditor/Private/MuCOE/CustomizableObjectCompiler.cpp
--- a/Engine/Plugins/Mutable/Source/CustomizableObjectEditor/Private/MuCOE/CustomizableObjectCompiler.cpp
+++ b/Engine/Plugins/Mutable/Source/CustomizableObjectEditor/Private/MuCOE/CustomizableObjectCompiler.cpp
@@ -1106,7 +1106,7 @@ void FCustomizableObjectCompiler::CompleteRequest(ECompilationStatePrivate State
 
 		if (PlatformData->Model)
 		{
-			PlatformData->Model->UnloadRoms();
+			PlatformData->Model->Program.UnloadRoms();
 		}
 
 		if (!CurrentOptions.bIsCooking)
@@ -1628,7 +1628,7 @@ void FCustomizableObjectCompiler::FinishCompilationTask()
 	// At this point it is assumed that all data goes into a single file.
 	if (Model)
 	{
-		const int32 NumStreamingFiles = Model->GetRomCount();
+		const int32 NumStreamingFiles = Model->Program.GetRomCount();
 
 		TMap<uint32, FMutableStreamableBlock>& ModelStreamables = PlatformData->ModelStreamableBulkData->ModelStreamables;
 		ModelStreamables.Empty(NumStreamingFiles);
@@ -1636,8 +1636,8 @@ void FCustomizableObjectCompiler::FinishCompilationTask()
 		uint64 Offset = 0;
 		for (int32 FileIndex = 0; FileIndex < NumStreamingFiles; ++FileIndex)
 		{
-			const uint32 ResourceSize = Model->GetRomSize(FileIndex);
-			EMutableFileFlags Flags = Model->IsRomHighRes(FileIndex) ? EMutableFileFlags::HighRes : EMutableFileFlags::None;
+			const uint32 ResourceSize = Model->Program.GetRomSize(FileIndex);
+			EMutableFileFlags Flags = Model->Program.IsRomHighRes(FileIndex) ? EMutableFileFlags::HighRes : EMutableFileFlags::None;
 			ModelStreamables.Add(FileIndex, FMutableStreamableBlock{ .FileId = 0, .Flags = uint16(Flags), .Offset = Offset });
 			Offset += ResourceSize;
 		}
diff --git a/Engine/Plugins/Mutable/Source/CustomizableObjectEditor/Private/MuCOE/SMutableCodeViewer.cpp b/Engine/Plugins/Mutable/Source/CustomizableObjectEditor/Private/MuCOE/SMutableCodeViewer.cpp
--- a/Engine/Plugins/Mutable/Source/CustomizableObjectEditor/Private/MuCOE/SMutableCodeViewer.cpp
+++ b/Engine/Plugins/Mutable/Source/CustomizableObjectEditor/Private/MuCOE/SMutableCodeViewer.cpp
@@ -2897,7 +2897,6 @@ void SMutableCodeViewer::Tick(const FGeometry& AllottedGeometry, const double In
 	UE::Mutable::Private::FSettings Settings;
 	const TSharedPtr<UE::Mutable::Private::FSystem> System = MakeShared<UE::Mutable::Private::FSystem>(Settings);
 
-	System->BeginBuild(MutableModel);
 	TSharedRef<UE::Mutable::Private::FLiveInstance> TempLiveInstance = System->NewBuildInstance(MutableModel, ExternalResourceProvider);
 	TempLiveInstance->OldParameters = PreviewParameters;
 	
@@ -3000,8 +2999,6 @@ void SMutableCodeViewer::Tick(const FGeometry& AllottedGeometry, const double In
 		// There is no viewer for this type.
 		break;
 	}
-
-	System->EndBuild();
 }
 
 
@@ -3010,14 +3007,11 @@ UE::Mutable::Private::TManagedPtr<const UE::Mutable::Private::FImage> SMutableCo
 	UE::Mutable::Private::FSettings Settings;
 	const TSharedPtr<UE::Mutable::Private::FSystem> System = MakeShared<UE::Mutable::Private::FSystem>(Settings);
 	
-	System->BeginBuild(MutableModel);
 	
 	TSharedRef<UE::Mutable::Private::FLiveInstance> TempLiveInstance = System->NewBuildInstance(MutableModel, ExternalResourceProvider);
 	TempLiveInstance->OldParameters = PreviewParameters;
 	
-	UE::Mutable::Private::TManagedPtr<const UE::Mutable::Private::FImage> MutableImage = System->BuildImage(TempLiveInstance, TargetOpAddress, 0);
-	System->EndBuild();
-	
+	UE::Mutable::Private::TManagedPtr<const UE::Mutable::Private::FImage> MutableImage = System->BuildImage(TempLiveInstance, TargetOpAddress, 0);	
 	return MutableImage;
 }
 
diff --git a/Engine/Plugins/Mutable/Source/MutableRuntime/Internal/MuR/Model.h b/Engine/Plugins/Mutable/Source/MutableRuntime/Internal/MuR/Model.h
--- a/Engine/Plugins/Mutable/Source/MutableRuntime/Internal/MuR/Model.h
+++ b/Engine/Plugins/Mutable/Source/MutableRuntime/Internal/MuR/Model.h
@@ -30,7 +30,6 @@ namespace UE::Mutable::Private
 	class FInputArchive;
 	class FModelWriter;
 	class FOutputArchive;
-	class FSystem;
     class FModelParametersGenerator;
 
 	
@@ -189,13 +188,7 @@ namespace UE::Mutable::Private
     //!
     struct FProgram
     {
-        FProgram()
-        {
-            // Add the null instruction at address 0.
-            // TODO: Will do it in the linker
-            AppendCode( ByteCode, EOpType::NONE );
-            OpAddress.Add(0);
-        }
+        UE_API FProgram();
 
         struct FState
         {
@@ -223,68 +216,21 @@ namespace UE::Mutable::Private
 			TArray< TPair<OP::ADDRESS,uint64> > m_dynamicResources;
 
             /** */
-            inline void Serialise( FOutputArchive& arch ) const
-            {
-                arch << Name;
-                arch << Root;
-                arch << m_runtimeParameters;
-                arch << m_updateCache;
-                arch << m_dynamicResources;
-            }
+            void Serialise(FOutputArchive& Arch) const;
 
             /** */
-            inline void Unserialise( FInputArchive& arch )
-            {
-                arch >> Name;
-                arch >> Root;
-                arch >> m_runtimeParameters;
-                arch >> m_updateCache;
-                arch >> m_dynamicResources;
-            }
+            void Unserialise(FInputArchive& Arch);
 
             /** Returns the mask of parameters (from the runtime parameter list of this state) including the parameters that 
 			* are relevant for the dynamic resource at the given address.
 			*/
-            uint64 IsDynamic( OP::ADDRESS at ) const
-            {
-                uint64 res = 0;
-
-                for ( int32 i=0; i<m_dynamicResources.Num(); ++i )
-                {
-                    if ( m_dynamicResources[i].Key==at )
-                    {
-                        res = m_dynamicResources[i].Value;
-						break;
-                    }
-                }
-
-                return res;
-            }
+            uint64 IsDynamic(OP::ADDRESS At) const;
 
             /** */
-            bool IsUpdateCache( OP::ADDRESS at ) const
-            {
-                bool res = false;
-
-                for (int32 i=0; !res && i<m_updateCache.Num(); ++i )
-                {
-                    if ( m_updateCache[i]==at )
-                    {
-                        res = true;
-                    }
-                }
-
-                return res;
-            }
+            bool IsUpdateCache(OP::ADDRESS at) const;
 
             /** */ 
-            void AddUpdateCache( OP::ADDRESS at )
-            {
-                if ( !IsUpdateCache(at) )
-                {
-                    m_updateCache.Add( at );
-                }
-            }
+            void AddUpdateCache(OP::ADDRESS At);
         };
 
         /** Location in the ByteCode of the beginning of each operation */
@@ -402,370 +348,81 @@ namespace UE::Mutable::Private
 		 */
 		bool bIsValid = true;
 #endif
-        /** */
-        void Serialise(FOutputArchive& Arch) const
-        {
-            Arch << OpAddress;
-            Arch << ByteCode;
-            Arch << States;
-			Arch << Roms;
-			Arch << RomsCompileData;
-			Arch << ConstantImageLODsPermanent;
-			Arch << ConstantImageLODIndices;
-			Arch << ConstantImages;
-			Arch << ConstantMeshesPermanent;
-			Arch << ConstantMeshContentIndices;
-			Arch << ConstantMeshes;
-			Arch << ConstantExtensionData;
-			Arch << ConstantStrings;
-			Arch << ConstantUInt32Lists;
-			Arch << ConstantUInt64Lists;
-            Arch << ConstantLayouts;
-            Arch << ConstantProjectors;
-			Arch << ConstantMatrices;
-			Arch << ConstantShapes;
-            Arch << ConstantCurves;
-            Arch << ConstantSkeletons;
-            Arch << Parameters;
-            Arch << Ranges;
-            Arch << ParameterLists;
-            Arch << RelevantParameterList;
-			Arch << ConstantMaterials;
-        }
+    	
+    	/** For each model rom, the last time its streamed data was used. */
+    	mutable TMemoryTrackedArray<TPair<uint64, uint64>> RomWeights;
 
-        /** */
-        void Unserialise(FInputArchive& Arch)
-        {
-            Arch >> OpAddress;
-            Arch >> ByteCode;
-            Arch >> States;
-			Arch >> Roms;
-			Arch >> RomsCompileData;
-			Arch >> ConstantImageLODsPermanent;
-			Arch >> ConstantImageLODIndices;
-			Arch >> ConstantImages;
-			Arch >> ConstantMeshesPermanent;
-			Arch >> ConstantMeshContentIndices;
-			Arch >> ConstantMeshes;
-			Arch >> ConstantExtensionData;
-			Arch >> ConstantStrings;
-			Arch >> ConstantUInt32Lists;
-			Arch >> ConstantUInt64Lists;
-            Arch >> ConstantLayouts;
-            Arch >> ConstantProjectors;
-			Arch >> ConstantMatrices;
-			Arch >> ConstantShapes;
-            Arch >> ConstantCurves;
-            Arch >> ConstantSkeletons;
-            Arch >> Parameters;
-            Arch >> Ranges;
-            Arch >> ParameterLists;
-            Arch >> RelevantParameterList;
-            Arch >> ConstantMaterials;
-        }
+    	/** Count of pending operations for every rom index. */
+    	mutable TMemoryTrackedArray<uint16> PendingOpsPerRom;
+    	
+        UE_API void Serialise(FOutputArchive& Arch) const;
+
+        void Unserialise(FInputArchive& Arch);
 
         /** Debug method that logs the top used instruction types. */
         void LogHistogram() const;
 
+    	UE_API uint32 GetRomSize(int32 Index) const;
+    	
 		/** Return true if the given ROM is loaded. */
-		FORCEINLINE bool IsRomLoaded(int32 RomIndex) const
-		{
-			switch (Roms[RomIndex].ResourceType)
-			{
-				case uint32(ERomDataType::Image):
-					return ConstantImageLODsStreamed.Contains(RomIndex);
-				case uint32(ERomDataType::Mesh):
-					return ConstantMeshesStreamed.Contains(RomIndex);
-				//TODO(Max): Make materials loadable and unloadable from a rom
-				default:
-					check(false);
-					break;
-			}
+		bool IsRomLoaded(int32 RomIndex) const;
 
-			return false;
-		}
+        UE_API int32 GetRomCount() const;
+
+#if WITH_EDITOR
+    	UE_API uint32 GetRomSourceId(int32 Index) const;
+#endif
 
+    	UE_API ERomDataType GetRomDataType(int32 Index) const;
+
+    	UE_API bool IsRomHighRes(int32 Index) const;
+   
+    	/** 
+		 * Returns the ConstantImageIndex LODIndex rom id. In case ConstantImageIndex does not have LODIndex 
+		 * returns -1.
+		 */
+    	UE_API int32 GetConstantImageRomId(int32 ConstantImageIndex, int32 LODIndex) const;
+    	
+    	UE_API void UnloadRoms() const;
+    	
 		/** Unload a rom resource. 
 		* If a pointer to a size is passed in OutDataSize, the value will be set with the size of the unloaded rom.
 		*/
-		FORCEINLINE void UnloadRom(int32 RomIndex, int32* OutDataSize=nullptr) const
-		{
-			if (DebugRom && (DebugRomAll||RomIndex == DebugRomIndex))
-				UE_LOG(LogMutableCore, Log, TEXT("Unloading rom %d."), RomIndex);
-
-			if (LoadedMemTrackedRoms.IsValidIndex(RomIndex))
-			{
-				LoadedMemTrackedRoms.RemoveAt(RomIndex);
-			}
-
-			switch (Roms[RomIndex].ResourceType)
-			{
-			case uint32(ERomDataType::Image):
-			{
-				TManagedPtr<const FImage> Data;
-				ConstantImageLODsStreamed.RemoveAndCopyValue(RomIndex, Data);
-				if (OutDataSize && Data)
-				{
-					*OutDataSize = Data->GetDataSize();
-				}
-				break;
-			}
-			case uint32(ERomDataType::Mesh):
-			{
-				TManagedPtr<const FMesh> Data;
-				ConstantMeshesStreamed.RemoveAndCopyValue(RomIndex, Data);
-				if (OutDataSize && Data)
-				{
-					*OutDataSize = Data->GetDataSize();
-				}
-				break;
-			}
-			//TODO(Max): Make materials loadable and unloadable from a rom
-			default:
-				check(false);
-				break;
-			}
-		}
-
-		FORCEINLINE SSIZE_T GetRomBytes(int32 RomIndex) const
-		{
-			SSIZE_T Result = 0;
-			switch (Roms[RomIndex].ResourceType)
-			{
-			case uint32(ERomDataType::Image):
-			{
-				const TManagedPtr<const FImage>* Data = ConstantImageLODsStreamed.Find(RomIndex);
-				if (Data)
-				{
-					Result = (*Data)->GetDataSize();
-				}
-				break;
-			}
-			case uint32(ERomDataType::Mesh):
-			{
-				const TManagedPtr<const FMesh>* Data = ConstantMeshesStreamed.Find(RomIndex);
-				if (Data)
-				{
-					Result = (*Data)->GetDataSize();
-				}
-				break;
-			}
-			//TODO(Max): Make materials loadable and unloadable from a rom
-			default:
-				check(false);
-				break;
-			}
-
-			return Result;
-		}
-
-		FORCEINLINE void SetMeshRomValue(uint32 RomIndex, const TManagedPtr<FMesh>& Value) const
-		{
-			check(Roms[RomIndex].ResourceType == uint32(ERomDataType::Mesh));
-			check(!ConstantMeshesStreamed.Contains(RomIndex));
-
-			LoadedMemTrackedRoms.EmplaceAt(RomIndex, (uint8)Roms[RomIndex].ResourceType);
-			ConstantMeshesStreamed.Add(RomIndex, Value);
-		}
-
-		FORCEINLINE void SetImageRomValue(uint32 RomIndex, const TManagedPtr<FImage>& Value) const
-		{
-			check(Roms[RomIndex].ResourceType == uint32(ERomDataType::Image));
-			check(!ConstantImageLODsStreamed.Contains(RomIndex));
-
-			LoadedMemTrackedRoms.EmplaceAt(RomIndex, (uint8)Roms[RomIndex].ResourceType);
-			ConstantImageLODsStreamed.Add(RomIndex, Value);
-		}
-
-		OP::ADDRESS AddConstant(TManagedPtr<const FExtensionData> Data)
-		{
-			// Ensure unique
-			for (int32 Index = 0; Index < ConstantExtensionData.Num(); Index++)
-			{
-				const FExtensionData* Candidate = ConstantExtensionData[Index].Data.Get();
-				if (*Candidate == *Data)
-				{
-					return Index;
-				}
-			}
-
-			FExtensionDataConstant& NewConstant = ConstantExtensionData.AddDefaulted_GetRef();
-			NewConstant.Data = Data;
-
-			return ConstantExtensionData.Num() - 1;
-		}
-
-		OP::ADDRESS AddConstant( TManagedPtr<const FLayout> pLayout )
-        {
-            // Ensure unique
-            for (int32 i=0; i<ConstantLayouts.Num(); ++i)
-            {
-                if (ConstantLayouts[i]==pLayout)
-                {
-                    return (OP::ADDRESS)i;
-                }
-            }
-
-            OP::ADDRESS index = OP::ADDRESS( ConstantLayouts.Num() );
-            ConstantLayouts.Add( pLayout );
-            return index;
-        }
-
-        OP::ADDRESS AddConstant(const TManagedPtr<const FSkeleton>& Skeleton )
-        {
-            // Ensure unique
-            for ( int32 i=0; i<ConstantSkeletons.Num(); ++i)
-            {
-                if ( ConstantSkeletons[i]== Skeleton
-                     ||
-                     *ConstantSkeletons[i]==*Skeleton
-                     )
-                {
-                    return (OP::ADDRESS)i;
-                }
-            }
-
-            OP::ADDRESS index = OP::ADDRESS( ConstantSkeletons.Num() );
-            ConstantSkeletons.Add(Skeleton);
-            return index;
-        }
-
-        OP::ADDRESS AddConstant(const TManagedPtr<const FPhysicsBody>& PhysicsBody)
-        {
-            // Ensure unique
-            for (int32 i = 0; i < ConstantPhysicsBodies.Num(); ++i)
-            {
-                if (ConstantPhysicsBodies[i] == PhysicsBody || *ConstantPhysicsBodies[i] == *PhysicsBody )
-                {
-                    return (OP::ADDRESS)i;
-                }
-            }
-
-            OP::ADDRESS index = OP::ADDRESS(ConstantPhysicsBodies.Num());
-            ConstantPhysicsBodies.Add(PhysicsBody);
-            return index;
-        }
-
-        OP::CONSTANT_STRING_INDEX AddConstant(const FString& Str)
-        {     
-			return (OP::CONSTANT_STRING_INDEX)ConstantStrings.AddUnique(Str);
-        }
-		
-		OP::ADDRESS AddConstant(const TArray<uint32>& UInt32List)
-		{
-			return (OP::ADDRESS)ConstantUInt32Lists.AddUnique(UInt32List);
-		}
-
-		OP::ADDRESS AddConstant(const TArray<uint64>& UInt64List)
-		{
-			return (OP::ADDRESS)ConstantUInt64Lists.AddUnique(UInt64List);
-		}
-
-		OP::ADDRESS AddConstant(const TArray<FString>& StringList)
-		{
-			const int32 NumStrings = StringList.Num();
+        void UnloadRom(int32 RomIndex, int32* OutDataSize=nullptr) const;
 
-			TArray<uint32> StringAddrs;
-			StringAddrs.SetNum(NumStrings);
+        int32 GetRomBytes(int32 RomIndex) const;
 
-			for (int32 StringIndex = 0; StringIndex < NumStrings; ++StringIndex)
-			{
-				StringAddrs[StringIndex] = AddConstant(StringList[StringIndex]);
-			}
-
-			return (OP::ADDRESS)AddConstant(StringAddrs);
-		}
-
-        OP::ADDRESS AddConstant( const FMatrix44f& m )
-        {
-            // Ensure unique
-            for (int32 i=0; i<ConstantMatrices.Num(); ++i)
-            {
-                if (ConstantMatrices[i]==m)
-                {
-                    return (OP::ADDRESS)i;
-                }
-            }
-
-            OP::ADDRESS index = OP::ADDRESS( ConstantMatrices.Num() );
-            ConstantMatrices.Add( m );
-            return index;
-        }
-
-        OP::ADDRESS AddConstant( const FShape& m )
-        {
-            // Ensure unique
-            for (int32 i=0; i<ConstantShapes.Num(); ++i)
-            {
-                if (ConstantShapes[i]==m)
-                {
-                    return (OP::ADDRESS)i;
-                }
-            }
-
-            OP::ADDRESS index = OP::ADDRESS( ConstantShapes.Num() );
-            ConstantShapes.Add( m );
-            return index;
-        }
-
-        OP::ADDRESS AddConstant( const FProjector& m )
-        {
-            // Ensure unique
-            for (int32 i=0; i<ConstantProjectors.Num(); ++i)
-            {
-                if (ConstantProjectors[i]==m)
-                {
-                    return (OP::ADDRESS)i;
-                }
-            }
-
-            OP::ADDRESS index = OP::ADDRESS( ConstantProjectors.Num() );
-            ConstantProjectors.Add( m );
-            return index;
-        }
-
-        OP::ADDRESS AddConstant( const FRichCurve& m )
-        {
-            OP::ADDRESS index = OP::ADDRESS( ConstantCurves.AddUnique( m ) );
-            return index;
-        }
+        void SetMeshRomValue(uint32 RomIndex, const TManagedPtr<FMesh>& Value) const;
 
-		OP::ADDRESS AddConstant(const TManagedPtr<const FMaterial>& Material)
-		{
-			for (int32 Index=0; Index<ConstantMaterials.Num(); ++Index)
-			{
-				if (ConstantMaterials[Index] == Material || *ConstantMaterials[Index] == *Material)
-				{
-					return static_cast<OP::ADDRESS>(Index);
-				}
-			}
-
-			OP::ADDRESS Index = static_cast<OP::ADDRESS>(ConstantMaterials.Num());
-			ConstantMaterials.Add(Material);
-			return Index;
-		}
-
-		inline TManagedPtr<const FImage> GetImageLOD(FConstantResourceIndex Index) const
-		{
-			TManagedPtr<const FImage> Mip;
-			if (!Index.Streamable)
-			{
-				if (ConstantImageLODsPermanent.IsValidIndex(Index.Index))
-				{
-					Mip = ConstantImageLODsPermanent[Index.Index];
-				}
-			}
-			else
-			{
-				const TManagedPtr<const FImage>* Found = ConstantImageLODsStreamed.Find(Index.Index);
-				if (Found)
-				{
-					Mip = *Found;
-				}
-			}
+        void SetImageRomValue(uint32 RomIndex, const TManagedPtr<FImage>& Value) const;
+    	
+        UE_API OP::CONSTANT_EXTENSIONDATA_INDEX AddConstant(const TManagedPtr<const FExtensionData>& Data);
+    	
+        UE_API OP::ADDRESS AddConstant(const TManagedPtr<const FLayout>& Layout);
+    	
+        UE_API OP::ADDRESS AddConstant(const TManagedPtr<const FSkeleton>& Skeleton);
+    	
+        UE_API OP::ADDRESS AddConstant(const TManagedPtr<const FPhysicsBody>& PhysicsBody);
+    	
+        UE_API OP::CONSTANT_STRING_INDEX AddConstant(const FString& Str);
+    	
+        UE_API OP::ADDRESS AddConstant(const TArray<uint32>& UInt32List);
+    	
+        UE_API  OP::ADDRESS AddConstant(const TArray<uint64>& UInt64List);
+    	
+        UE_API OP::ADDRESS AddConstant(const TArray<FString>& StringList);
 
-			return Mip;
-		}
+        UE_API OP::CONSTANT_MATRIX_INDEX AddConstant(const FMatrix44f& Constant);
+    	
+        UE_API OP::CONSTANT_SHAPE_INDEX AddConstant(const FShape& Constant);
+    	
+        UE_API OP::ADDRESS AddConstant(const FProjector& Constant);
+    	
+        UE_API OP::ADDRESS AddConstant(const FRichCurve& Constant);
+    	
+        UE_API OP::ADDRESS AddConstant(const TManagedPtr<const FMaterial>& Material);
+    	
+        UE_API TManagedPtr<const FImage> GetImageLOD(FConstantResourceIndex Index) const;
 
         /** Get a constant image, assuming at least some mips are loaded. The image constant will be composed with lodaded mips if necessary. */
 		template <typename CreateImageFunc>
@@ -1111,28 +768,12 @@ namespace UE::Mutable::Private
 			OutMesh = Result;
 		}
 
-		void GetExtensionDataConstant(int32 ConstantIndex, TManagedPtr<const FExtensionData>& Result) const
-		{
-			const FExtensionDataConstant& Constant = ConstantExtensionData[ConstantIndex];
-
-			check(Constant.Data);
-
-			Result = Constant.Data;
-		}
+		void GetExtensionDataConstant(int32 ConstantIndex, TManagedPtr<const FExtensionData>& Result) const;
 
-        FORCEINLINE EOpType GetOpType( OP::ADDRESS at ) const
-        {
-            if (at>=OP::ADDRESS(OpAddress.Num())) return EOpType::NONE;
-
-            EOpType result;
-            uint64 byteCodeAddress = OpAddress[at];
-			FMemory::Memcpy( &result, &ByteCode[byteCodeAddress], sizeof(EOpType) );
-            check( result<EOpType::COUNT);
-            return result;
-        }
-
-		template<typename ARGS>
-		inline const ARGS GetOpArgs(OP::ADDRESS at) const
+        UE_API EOpType GetOpType(OP::ADDRESS At) const;
+    	
+        template<typename ARGS>
+		ARGS GetOpArgs(OP::ADDRESS at) const
 		{
 			ARGS result;
 			uint64 byteCodeAddress = OpAddress[at];
@@ -1142,28 +783,18 @@ namespace UE::Mutable::Private
 		}
 
 		template<typename ARGS>
-		inline void SetOpArgs(OP::ADDRESS at, const ARGS& Args)
+		void SetOpArgs(OP::ADDRESS At, const ARGS& Args)
 		{
-			uint64 byteCodeAddress = OpAddress[at];
+			uint64 byteCodeAddress = OpAddress[At];
 			byteCodeAddress += sizeof(EOpType);
 			FMemory::Memcpy(&ByteCode[byteCodeAddress], &Args, sizeof(ARGS));
 		}
 
-        inline const uint8* GetOpArgsPointer( OP::ADDRESS at ) const
-        {
-            uint64 byteCodeAddress = OpAddress[at];
-            byteCodeAddress += sizeof(EOpType);
-            const uint8* pData = (const uint8*)&ByteCode[byteCodeAddress];
-            return pData;
-        }
+        UE_API const uint8* GetOpArgsPointer(OP::ADDRESS At) const;
+    	
+        uint8* GetOpArgsPointer(OP::ADDRESS At);
 
-        inline uint8* GetOpArgsPointer( OP::ADDRESS at )
-        {
-            uint64 byteCodeAddress = OpAddress[at];
-            byteCodeAddress += sizeof(EOpType);
-            uint8* pData = (uint8*)&ByteCode[byteCodeAddress];
-            return pData;
-        }
+        void MarkRomUsed(int32 RomIndex) const;
     };
 
 	
@@ -1177,22 +808,16 @@ namespace UE::Mutable::Private
 		static UE_API void Serialise(const FModel*, FOutputArchive&);
 		static UE_API TSharedPtr<FModel> StaticUnserialise(FInputArchive&);
     	
-    	void Serialise(FOutputArchive& arch) const
-    	{
-    		arch << Program;
-    	}
+    	void Serialise(FOutputArchive& Arch) const;
 
-    	void Unserialise(FInputArchive& arch)
-    	{
-    		arch >> Program;
-    	}
+		void Unserialise(FInputArchive& Arch);
 
 		/** Special serialise operation that serialises the data in separate "files". An object
          * with the ModelStreamer interface is responsible of storing this data and providing
 		 * the "file" concept. 
 		 * If bDropData is set to true, the rom data will be freed as it is serialized.
 		 */
-        static UE_API void Serialise( FModel*, FModelWriter&, bool bDropData=false );
+        static UE_API void Serialise(FModel* Model, FModelWriter& Streamer, bool bDropData = false);
 
 		//! Return true if the model has external data in other files. This kind of models will
 		//! require data streaming when used.
@@ -1277,26 +902,6 @@ namespace UE::Mutable::Private
         //! \param OutScale Pointer to the projector-space scaling of the projector.
     	UE_API void GetProjectorDefaultValue(int32 Index, EProjectorType* OutProjectionType, FVector3f* OutPos,
     	 	FVector3f* OutDir, FVector3f* OutUp, FVector3f* OutScale, float* OutProjectionAngle) const;
-
-    	UE_API int32 GetRomCount() const;
-
-#if WITH_EDITOR
-		UE_API uint32 GetRomSourceId(int32 Index) const;
-#endif
-
-		UE_API uint32 GetRomSize(int32 Index) const;
-
-		UE_API bool IsMeshData(int32 Index) const;
-
-		UE_API bool IsRomHighRes(int32 Index) const;
-   
-		/** 
-		 * Returns the ConstantImageIndex LODIndex rom id. In case ConstantImageIndex does not have LODIndex 
-		 * returns -1.
-		 */
-    	UE_API int32 GetConstantImageRomId(int32 ConstantImageIndex, int32 LODIndex) const;
-    	
-    	UE_API void UnloadRoms();
     	
     	FProgram Program;
 	};
diff --git a/Engine/Plugins/Mutable/Source/MutableRuntime/Internal/MuR/System.h b/Engine/Plugins/Mutable/Source/MutableRuntime/Internal/MuR/System.h
--- a/Engine/Plugins/Mutable/Source/MutableRuntime/Internal/MuR/System.h
+++ b/Engine/Plugins/Mutable/Source/MutableRuntime/Internal/MuR/System.h
@@ -12,7 +12,6 @@
 #include "MuR/MemoryCounters.h"
 #include "MuR/Operations.h"
 #include "MuR/MutableString.h"
-#include "MuR/Model.h"
 #include "MuR/Material.h"
 #include "MuR/Parameters.h"
 #include "MuR/MutableTrace.h"
@@ -1517,40 +1516,17 @@ namespace UE::Mutable::Private
 		using TMemoryTrackedArray = TArray<Type, FDefaultMemoryTrackingAllocator<MemoryCounter>>;
 
 		template<class KeyType, class ValueType>
-		using TMemoryTrackedMap = TMap<KeyType, ValueType, FDefaultMemoryTrackingSetAllocator<MemoryCounter>>;
+		using TMemoryTrackedMap = TMap<KeyType, ValueType, FDefaultMemoryTrackingSetAllocator<MemoryCounter>>;    	
 
-		/** Cached tracking for streamed model data for one model. */
-        struct FModelCacheEntry
-        {
-            /** Model whose data is being tracked. */
-            TWeakPtr<const FModel> Model;
-
-            /** For each model rom, the last time its streamed data was used. */
-            TMemoryTrackedArray<TPair<uint64, uint64>> RomWeights;
-
-			/** Count of pending operations for every rom index. */
-			TMemoryTrackedArray<uint16> PendingOpsPerRom;
-        };
-    	
 		/** Maximum working memory that mutable should be using. */
 		int64 BudgetBytes = 0;
 
 		/** Maximum excess memory reached during the current operation. */
 		int64 BudgetExcessBytes = 0;
     	
-		/** This value is used to track the order of loading of roms. */
-        uint64 RomTick = 0;
-
-		/** Control info for the per-model cache of streamed data. */
-        TMemoryTrackedArray<FModelCacheEntry> CachePerModel;
-
 		/** Data for each mutable instance that is being updated. */
 		TMemoryTrackedArray<TWeakPtr<FLiveInstance>> LiveInstances;
 
-		/** Given a mutable model, find or create its rom cache. */
-		FModelCacheEntry* FindModelCache(const FModel*);
-		FModelCacheEntry& FindOrAddModelCache(const TSharedPtr<const FModel>&);
-
         /** Make sure the working memory is below the internal budget, even counting with the passed additional memory. 
 		* Return true if it succeeded, false otherwise.
 		*/
@@ -1562,10 +1538,6 @@ namespace UE::Mutable::Private
 		/** Calculate the current usage of memory as used to calculate the budget. */
 		int64 GetCurrentMemoryBytes() const;
 
-        /** Register that a specific rom has been requested and update the heuristics to keep it in memory. */
-		void MarkRomUsed( int32 RomIndex, const TSharedPtr<const FModel>& );
-    	
-
 		int32 GetRomBytes() const;
 
     	/** Calculate the amount of bytes in data cached in the level 0 and 1 cache in all live instances. */
@@ -1956,11 +1928,6 @@ namespace UE::Mutable::Private
         FSystem( FSystem&& ) = delete;
         FSystem& operator=( FSystem&& ) = delete;
 		
-		/** This method can be used to internally prepare for code execution. */
-		MUTABLERUNTIME_API void BeginBuild(const TSharedPtr<const FModel>&);
-    	
-		MUTABLERUNTIME_API void EndBuild();
-
 		MUTABLERUNTIME_API bool BuildBool(const TSharedRef<FLiveInstance>& InLiveInstance, OP::ADDRESS InAddress) ;
 		MUTABLERUNTIME_API int32 BuildInt(const TSharedRef<FLiveInstance>& InLiveInstance, OP::ADDRESS InAddress) ;
 		MUTABLERUNTIME_API float BuildScalar(const TSharedRef<FLiveInstance>& InLiveInstance, OP::ADDRESS InAddress) ;
@@ -1988,6 +1955,9 @@ namespace UE::Mutable::Private
 
 		FWorkingMemoryManager WorkingMemoryManager;
 
+		/** Increased each time a rom is loaded. Used to keep track of which roms to evict. */
+		static uint64 RomTick;
+		
 		bool CheckUpdatedParameters(const TSharedRef<FLiveInstance>& LiveInstance, const TSharedPtr<const FParameters>&, uint64& OutUpdatedParameters );
 		
 		void RunCode(
diff --git a/Engine/Plugins/Mutable/Source/MutableRuntime/Private/MuR/CodeRunner.h b/Engine/Plugins/Mutable/Source/MutableRuntime/Private/MuR/CodeRunner.h
--- a/Engine/Plugins/Mutable/Source/MutableRuntime/Private/MuR/CodeRunner.h
+++ b/Engine/Plugins/Mutable/Source/MutableRuntime/Private/MuR/CodeRunner.h
@@ -651,9 +651,6 @@ namespace  UE::Mutable::Private
     		return System.WorkingMemoryManager.EnsureBudgetBelow(AdditionalMemory, LiveInstance->Cache);
     	}
     	
-    	void MarkRomUsed( int32 InRomIndex);
-
-    	FWorkingMemoryManager::FModelCacheEntry* GetModelCache() const;
     	
     	struct FRomLoadOps
     	{
diff --git a/Engine/Plugins/Mutable/Source/MutableRuntime/Private/MuR/CodeRunnerTasked.cpp b/Engine/Plugins/Mutable/Source/MutableRuntime/Private/MuR/CodeRunnerTasked.cpp
--- a/Engine/Plugins/Mutable/Source/MutableRuntime/Private/MuR/CodeRunnerTasked.cpp
+++ b/Engine/Plugins/Mutable/Source/MutableRuntime/Private/MuR/CodeRunnerTasked.cpp
@@ -1969,7 +1969,6 @@ namespace UE::Mutable::Private
 		bOutFailed = false;
 
 		const FProgram& LocalProgram = Runner->Program;
-		FWorkingMemoryManager::FModelCacheEntry* ModelCache = Runner->GetModelCache();
 		
 		TArray<UE::Tasks::FTask, TInlineAllocator<4>> ReadCompleteEvents;
 		ReadCompleteEvents.Reserve(4); 
@@ -2028,11 +2027,11 @@ namespace UE::Mutable::Private
 
 			check(RomIndex < LocalProgram.Roms.Num());
 
-			++ModelCache->PendingOpsPerRom[RomIndex];
+			++LocalProgram.PendingOpsPerRom[RomIndex];
 
 			if (DebugRom && (DebugRomAll || RomIndex == DebugRomIndex))
 			{
-				UE_LOG(LogMutableCore, Log, TEXT("Preparing rom %d, now peding ops is %d."), RomIndex, ModelCache->PendingOpsPerRom[RomIndex]);
+				UE_LOG(LogMutableCore, Log, TEXT("Preparing rom %d, now peding ops is %d."), RomIndex, LocalProgram.PendingOpsPerRom[RomIndex]);
 			}
 
 			if (LocalProgram.IsRomLoaded(RomIndex))
@@ -2057,7 +2056,7 @@ namespace UE::Mutable::Private
 
 			// Free roms if necessary
 			{
-				Runner->MarkRomUsed(RomIndex);
+				LocalProgram.MarkRomUsed(RomIndex);
 				Runner->EnsureBudgetBelow(RomSize);
 			}
 
@@ -2101,13 +2100,12 @@ namespace UE::Mutable::Private
 		// This runs in the Runner thread
 		MUTABLE_CPUPROFILER_SCOPE(FLoadMeshRomTask_Complete);
 
+		const FProgram& LocalProgram = Runner->Program;
+		
 		if (!Runner)
 		{
 			return false;
 		}
-
-		const FProgram& LocalProgram = Runner->Program;
-		FWorkingMemoryManager::FModelCacheEntry* ModelCache = Runner->GetModelCache();
 		
 		bool bSomeMissingData = false;
 		for (const int32 RomIndex : RomIndices)
@@ -2152,12 +2150,12 @@ namespace UE::Mutable::Private
 			{
 				check(RomIndex < LocalProgram.Roms.Num());
 
-				Runner->MarkRomUsed(RomIndex);
-				--ModelCache->PendingOpsPerRom[RomIndex];
+				Runner->Model.Program.MarkRomUsed(RomIndex);
+				--LocalProgram.PendingOpsPerRom[RomIndex];
 
 				if (DebugRom && (DebugRomAll || RomIndex == DebugRomIndex))
 				{
-					UE_LOG(LogMutableCore, Log, TEXT("FLoadMeshRomsTask::Complete rom %d, now peding ops is %d."), RomIndex, ModelCache->PendingOpsPerRom[RomIndex]);
+					UE_LOG(LogMutableCore, Log, TEXT("FLoadMeshRomsTask::Complete rom %d, now peding ops is %d."), RomIndex, LocalProgram.PendingOpsPerRom[RomIndex]);
 				}
 			}
 		}
@@ -2180,9 +2178,7 @@ namespace UE::Mutable::Private
 		bOutFailed = false;
 
 		const FProgram& LocalProgram = Runner->Program;
-
-		FWorkingMemoryManager::FModelCacheEntry* ModelCache = Runner->GetModelCache();
-
+		
 		TArray<UE::Tasks::FTask> ReadCompleteEvents;
 
 		ReadCompleteEvents.Reserve(FMath::Max(0, ImageIndexEnd - ImageIndexBegin));
@@ -2200,10 +2196,10 @@ namespace UE::Mutable::Private
 			int32 RomIndex = ImageConstantResourceIndex.Index;
 			check(RomIndex < LocalProgram.Roms.Num());
 
-			++ModelCache->PendingOpsPerRom[RomIndex];
+			++LocalProgram.PendingOpsPerRom[RomIndex];
 
 			if (DebugRom && (DebugRomAll || RomIndex == DebugRomIndex))
-				UE_LOG(LogMutableCore, Log, TEXT("Preparing rom %d, now peding ops is %d."), RomIndex, ModelCache->PendingOpsPerRom[RomIndex]);
+				UE_LOG(LogMutableCore, Log, TEXT("Preparing rom %d, now peding ops is %d."), RomIndex, LocalProgram.PendingOpsPerRom[RomIndex]);
 
 			if (LocalProgram.IsRomLoaded(RomIndex))
 			{
@@ -2227,7 +2223,7 @@ namespace UE::Mutable::Private
 
 			// Free roms if necessary
 			{
-				Runner->MarkRomUsed(RomIndex);
+				Runner->Model.Program.MarkRomUsed(RomIndex);
 				Runner->EnsureBudgetBelow(RomSize);
 			}
 
@@ -2277,8 +2273,6 @@ namespace UE::Mutable::Private
 
 		const FProgram& LocalProgram = Runner->Program;
 		
-		FWorkingMemoryManager::FModelCacheEntry* ModelCache = Runner->GetModelCache();
-
 		bool bSomeMissingData = false;
 
 		for (const int32 RomIndex : RomIndices)
@@ -2332,12 +2326,12 @@ namespace UE::Mutable::Private
 				const int32 RomIndex = ImageConstantResourceIndex.Index;
 				check(RomIndex < LocalProgram.Roms.Num());
 
-				Runner->MarkRomUsed(RomIndex);
-				--ModelCache->PendingOpsPerRom[RomIndex];
+				Runner->Model.Program.MarkRomUsed(RomIndex);
+				--LocalProgram.PendingOpsPerRom[RomIndex];
 
 				if (DebugRom && (DebugRomAll || RomIndex == DebugRomIndex))
 				{
-					UE_LOG(LogMutableCore, Log, TEXT("FLoadImageRomsTask::Complete rom %d, now peding ops is %d."), RomIndex, ModelCache->PendingOpsPerRom[RomIndex]);
+					UE_LOG(LogMutableCore, Log, TEXT("FLoadImageRomsTask::Complete rom %d, now peding ops is %d."), RomIndex, LocalProgram.PendingOpsPerRom[RomIndex]);
 				}
 			}
 		}
@@ -2757,16 +2751,4 @@ namespace UE::Mutable::Private
 
 		return Issued;
 	}
-
-	void CodeRunner::MarkRomUsed(int32 InRomIndex)
-	{
-		System.WorkingMemoryManager.MarkRomUsed(InRomIndex, LiveInstance->Model);
-	}
-
-	
-	FWorkingMemoryManager::FModelCacheEntry* CodeRunner::GetModelCache() const
-	{
-		return System.WorkingMemoryManager.FindModelCache(&Model);
-	}
-
 } // namespace UE::Mutable::Private
diff --git a/Engine/Plugins/Mutable/Source/MutableRuntime/Private/MuR/Model.cpp b/Engine/Plugins/Mutable/Source/MutableRuntime/Private/MuR/Model.cpp
--- a/Engine/Plugins/Mutable/Source/MutableRuntime/Private/MuR/Model.cpp
+++ b/Engine/Plugins/Mutable/Source/MutableRuntime/Private/MuR/Model.cpp
@@ -32,7 +32,143 @@ namespace UE::Mutable::Private
 	MUTABLE_IMPLEMENT_POD_SERIALISABLE(FConstantResourceIndex);
 
 	
-    void FProgram::LogHistogram() const
+	FProgram::FProgram()
+	{
+		// Add the null instruction at address 0.
+		// TODO: Will do it in the linker
+		AppendCode( ByteCode, EOpType::NONE );
+		OpAddress.Add(0);
+	}
+
+
+	void FProgram::FState::Serialise(FOutputArchive& Arch) const
+	{
+		Arch << Name;
+		Arch << Root;
+		Arch << m_runtimeParameters;
+		Arch << m_updateCache;
+		Arch << m_dynamicResources;
+	}
+
+
+	void FProgram::FState::Unserialise(FInputArchive& Arch)
+	{
+		Arch >> Name;
+		Arch >> Root;
+		Arch >> m_runtimeParameters;
+		Arch >> m_updateCache;
+		Arch >> m_dynamicResources;
+	}
+
+
+	uint64 FProgram::FState::IsDynamic(OP::ADDRESS At) const
+	{
+		uint64 Result = 0;
+
+		for (int32 Index = 0; Index < m_dynamicResources.Num(); ++Index)
+		{
+			if (m_dynamicResources[Index].Key == At)
+			{
+				Result = m_dynamicResources[Index].Value;
+				break;
+			}
+		}
+
+		return Result;
+	}
+
+
+	bool FProgram::FState::IsUpdateCache(OP::ADDRESS At) const
+	{
+		bool Result = false;
+
+		for (int32 i = 0; !Result && i < m_updateCache.Num(); ++i)
+		{
+			if (m_updateCache[i] == At)
+			{
+				Result = true;
+			}
+		}
+
+		return Result;
+	}
+
+
+	void FProgram::FState::AddUpdateCache(OP::ADDRESS At)
+	{
+		if (!IsUpdateCache(At))
+		{
+			m_updateCache.Add(At);
+		}
+	}
+
+
+	void FProgram::Serialise(FOutputArchive& Arch) const
+	{
+		Arch << OpAddress;
+		Arch << ByteCode;
+		Arch << States;
+		Arch << Roms;
+		Arch << RomsCompileData;
+		Arch << ConstantImageLODsPermanent;
+		Arch << ConstantImageLODIndices;
+		Arch << ConstantImages;
+		Arch << ConstantMeshesPermanent;
+		Arch << ConstantMeshContentIndices;
+		Arch << ConstantMeshes;
+		Arch << ConstantExtensionData;
+		Arch << ConstantStrings;
+		Arch << ConstantUInt32Lists;
+		Arch << ConstantUInt64Lists;
+		Arch << ConstantLayouts;
+		Arch << ConstantProjectors;
+		Arch << ConstantMatrices;
+		Arch << ConstantShapes;
+		Arch << ConstantCurves;
+		Arch << ConstantSkeletons;
+		Arch << Parameters;
+		Arch << Ranges;
+		Arch << ParameterLists;
+		Arch << RelevantParameterList;
+		Arch << ConstantMaterials;
+	}
+
+
+	void FProgram::Unserialise(FInputArchive& Arch)
+	{
+		Arch >> OpAddress;
+		Arch >> ByteCode;
+		Arch >> States;
+		Arch >> Roms;
+		Arch >> RomsCompileData;
+		Arch >> ConstantImageLODsPermanent;
+		Arch >> ConstantImageLODIndices;
+		Arch >> ConstantImages;
+		Arch >> ConstantMeshesPermanent;
+		Arch >> ConstantMeshContentIndices;
+		Arch >> ConstantMeshes;
+		Arch >> ConstantExtensionData;
+		Arch >> ConstantStrings;
+		Arch >> ConstantUInt32Lists;
+		Arch >> ConstantUInt64Lists;
+		Arch >> ConstantLayouts;
+		Arch >> ConstantProjectors;
+		Arch >> ConstantMatrices;
+		Arch >> ConstantShapes;
+		Arch >> ConstantCurves;
+		Arch >> ConstantSkeletons;
+		Arch >> Parameters;
+		Arch >> Ranges;
+		Arch >> ParameterLists;
+		Arch >> RelevantParameterList;
+		Arch >> ConstantMaterials;
+		
+		PendingOpsPerRom.SetNum(Roms.Num());
+		RomWeights.SetNumZeroed(Roms.Num());
+	}
+	
+
+	void FProgram::LogHistogram() const
     {
 #if 0
         uint64 countPerType[(int32)EOpType::COUNT];
@@ -65,21 +201,368 @@ namespace UE::Mutable::Private
 #endif
     }
 
-	
-    void FModel::UnloadRoms()
+
+    void FProgram::UnloadRoms() const
+    {
+    	for (int32 RomIndex = 0; RomIndex < Roms.Num(); ++RomIndex)
+    	{
+    		UnloadRom(RomIndex);
+    	}
+    }
+
+
+    void FProgram::UnloadRom(int32 RomIndex, int32* OutDataSize) const
+    {
+	    if (DebugRom && (DebugRomAll||RomIndex == DebugRomIndex))
+		    UE_LOG(LogMutableCore, Log, TEXT("Unloading rom %d."), RomIndex);
+
+	    if (LoadedMemTrackedRoms.IsValidIndex(RomIndex))
+	    {
+		    LoadedMemTrackedRoms.RemoveAt(RomIndex);
+	    }
+
+	    switch (Roms[RomIndex].ResourceType)
+	    {
+	    case uint32(ERomDataType::Image):
+		    {
+			    TManagedPtr<const FImage> Data;
+			    ConstantImageLODsStreamed.RemoveAndCopyValue(RomIndex, Data);
+			    if (OutDataSize && Data)
+			    {
+				    *OutDataSize = Data->GetDataSize();
+			    }
+			    break;
+		    }
+	    case uint32(ERomDataType::Mesh):
+		    {
+			    TManagedPtr<const FMesh> Data;
+			    ConstantMeshesStreamed.RemoveAndCopyValue(RomIndex, Data);
+			    if (OutDataSize && Data)
+			    {
+				    *OutDataSize = Data->GetDataSize();
+			    }
+			    break;
+		    }
+	    //TODO(Max): Make materials loadable and unloadable from a rom
+	    default:
+		    check(false);
+		    break;
+	    }
+    }
+
+
+    int32 FProgram::GetRomBytes(int32 RomIndex) const
+    {
+	    int32 Result = 0;
+	    switch (Roms[RomIndex].ResourceType)
+	    {
+	    case uint32(ERomDataType::Image):
+		    {
+			    const TManagedPtr<const FImage>* Data = ConstantImageLODsStreamed.Find(RomIndex);
+			    if (Data)
+			    {
+				    Result = (*Data)->GetDataSize();
+			    }
+			    break;
+		    }
+	    case uint32(ERomDataType::Mesh):
+		    {
+			    const TManagedPtr<const FMesh>* Data = ConstantMeshesStreamed.Find(RomIndex);
+			    if (Data)
+			    {
+				    Result = (*Data)->GetDataSize();
+			    }
+			    break;
+		    }
+	    //TODO(Max): Make materials loadable and unloadable from a rom
+	    default:
+		    check(false);
+		    break;
+	    }
+
+	    return Result;
+    }
+
+
+    void FProgram::SetMeshRomValue(uint32 RomIndex, const TManagedPtr<FMesh>& Value) const
+    {
+	    check(Roms[RomIndex].ResourceType == uint32(ERomDataType::Mesh));
+	    check(!ConstantMeshesStreamed.Contains(RomIndex));
+
+	    LoadedMemTrackedRoms.EmplaceAt(RomIndex, (uint8)Roms[RomIndex].ResourceType);
+	    ConstantMeshesStreamed.Add(RomIndex, Value);
+    }
+
+
+    void FProgram::SetImageRomValue(uint32 RomIndex, const TManagedPtr<FImage>& Value) const
+    {
+	    check(Roms[RomIndex].ResourceType == uint32(ERomDataType::Image));
+	    check(!ConstantImageLODsStreamed.Contains(RomIndex));
+
+	    LoadedMemTrackedRoms.EmplaceAt(RomIndex, (uint8)Roms[RomIndex].ResourceType);
+	    ConstantImageLODsStreamed.Add(RomIndex, Value);
+    }
+
+
+    OP::CONSTANT_EXTENSIONDATA_INDEX FProgram::AddConstant(const TManagedPtr<const FExtensionData>& Data)
+    {
+	    // Ensure unique
+	    for (int32 Index = 0; Index < ConstantExtensionData.Num(); Index++)
+	    {
+		    const FExtensionData* Candidate = ConstantExtensionData[Index].Data.Get();
+		    if (*Candidate == *Data)
+		    {
+			    return Index;
+		    }
+	    }
+
+	    FExtensionDataConstant& NewConstant = ConstantExtensionData.AddDefaulted_GetRef();
+	    NewConstant.Data = Data;
+
+	    return ConstantExtensionData.Num() - 1;
+    }
+
+
+    OP::ADDRESS FProgram::AddConstant(const TManagedPtr<const FLayout>& Layout)
+    {
+	    // Ensure unique
+	    for (int32 Index = 0; Index < ConstantLayouts.Num(); ++Index)
+	    {
+		    if (ConstantLayouts[Index] == Layout)
+		    {
+			    return (OP::ADDRESS)Index;
+		    }
+	    }
+
+	    OP::ADDRESS Index = OP::ADDRESS(ConstantLayouts.Num());
+	    ConstantLayouts.Add(Layout);
+	    return Index;
+    }
+
+
+    OP::ADDRESS FProgram::AddConstant(const TManagedPtr<const FSkeleton>& Skeleton)
+    {
+	    // Ensure unique
+	    for (int32 Index = 0; Index < ConstantSkeletons.Num(); ++Index)
+	    {
+		    if (ConstantSkeletons[Index] == Skeleton ||
+			    *ConstantSkeletons[Index] == *Skeleton
+		    )
+		    {
+			    return (OP::ADDRESS)Index;
+		    }
+	    }
+
+	    OP::ADDRESS Index = OP::ADDRESS(ConstantSkeletons.Num());
+	    ConstantSkeletons.Add(Skeleton);
+		
+	    return Index;
+    }
+
+
+    OP::ADDRESS FProgram::AddConstant(const TManagedPtr<const FPhysicsBody>& PhysicsBody)
+    {
+	    // Ensure unique
+	    for (int32 Index = 0; Index < ConstantPhysicsBodies.Num(); ++Index)
+	    {
+		    if (ConstantPhysicsBodies[Index] == PhysicsBody || *ConstantPhysicsBodies[Index] == *PhysicsBody )
+		    {
+			    return (OP::ADDRESS)Index;
+		    }
+	    }
+
+	    OP::ADDRESS index = OP::ADDRESS(ConstantPhysicsBodies.Num());
+	    ConstantPhysicsBodies.Add(PhysicsBody);
+	    return index;
+    }
+
+
+    OP::CONSTANT_STRING_INDEX FProgram::AddConstant(const FString& Str)
+    {     
+	    return (OP::CONSTANT_STRING_INDEX)ConstantStrings.AddUnique(Str);
+    }
+
+
+    OP::ADDRESS FProgram::AddConstant(const TArray<uint32>& UInt32List)
+    {
+	    return (OP::ADDRESS)ConstantUInt32Lists.AddUnique(UInt32List);
+    }
+
+
+    OP::ADDRESS FProgram::AddConstant(const TArray<uint64>& UInt64List)
+    {
+	    return (OP::ADDRESS)ConstantUInt64Lists.AddUnique(UInt64List);
+    }
+
+
+    OP::ADDRESS FProgram::AddConstant(const TArray<FString>& StringList)
+    {
+	    const int32 NumStrings = StringList.Num();
+
+	    TArray<uint32> StringAddrs;
+	    StringAddrs.SetNum(NumStrings);
+
+	    for (int32 StringIndex = 0; StringIndex < NumStrings; ++StringIndex)
+	    {
+		    StringAddrs[StringIndex] = AddConstant(StringList[StringIndex]);
+	    }
+
+	    return (OP::ADDRESS)AddConstant(StringAddrs);
+    }
+
+
+    OP::CONSTANT_MATRIX_INDEX FProgram::AddConstant(const FMatrix44f& Constant)
+    {
+	    // Ensure unique
+	    for (int32 i = 0; i < ConstantMatrices.Num(); ++i)
+	    {
+		    if (ConstantMatrices[i] == Constant)
+		    {
+			    return (OP::CONSTANT_MATRIX_INDEX)i;
+		    }
+	    }
+
+	    OP::CONSTANT_MATRIX_INDEX Index = OP::CONSTANT_MATRIX_INDEX(ConstantMatrices.Num());
+	    ConstantMatrices.Add(Constant);
+	    return Index;
+    }
+
+
+    OP::CONSTANT_SHAPE_INDEX FProgram::AddConstant(const FShape& Constant)
+    {
+	    // Ensure unique
+	    for (int32 Index = 0; Index<ConstantShapes.Num(); ++Index)
+	    {
+		    if (ConstantShapes[Index] == Constant)
+		    {
+			    return (OP::CONSTANT_SHAPE_INDEX)Index;
+		    }
+	    }
+
+	    OP::CONSTANT_SHAPE_INDEX Index = OP::CONSTANT_SHAPE_INDEX(ConstantShapes.Num());
+	    ConstantShapes.Add(Constant);
+	    return Index;
+    }
+
+
+    OP::ADDRESS FProgram::AddConstant(const FProjector& Constant)
     {
-	    for (int32 RomIndex = 0; RomIndex < Program.Roms.Num(); ++RomIndex)
+	    // Ensure unique
+	    for (int32 Index = 0; Index < ConstantProjectors.Num(); ++Index)
 	    {
-		    Program.UnloadRom(RomIndex);
+		    if (ConstantProjectors[Index] == Constant)
+		    {
+			    return (OP::ADDRESS)Index;
+		    }
 	    }
+
+	    OP::ADDRESS Index = OP::ADDRESS(ConstantProjectors.Num());
+	    ConstantProjectors.Add(Constant);
+	    return Index;
     }
 
 
-    void FModel::Serialise( const FModel* p, FOutputArchive& arch )
+    OP::ADDRESS FProgram::AddConstant(const FRichCurve& Constant)
+    {
+	    OP::ADDRESS Index = OP::ADDRESS(ConstantCurves.AddUnique(Constant ));
+	    return Index;
+    }
+
+
+    OP::ADDRESS FProgram::AddConstant(const TManagedPtr<const FMaterial>& Material)
+    {
+	    for (int32 Index=0; Index<ConstantMaterials.Num(); ++Index)
+	    {
+		    if (ConstantMaterials[Index] == Material || *ConstantMaterials[Index] == *Material)
+		    {
+			    return static_cast<OP::ADDRESS>(Index);
+		    }
+	    }
+
+	    OP::ADDRESS Index = static_cast<OP::ADDRESS>(ConstantMaterials.Num());
+	    ConstantMaterials.Add(Material);
+	    return Index;
+    }
+
+
+    TManagedPtr<const FImage> FProgram::GetImageLOD(FConstantResourceIndex Index) const
+    {
+	    TManagedPtr<const FImage> Mip;
+	    if (!Index.Streamable)
+	    {
+		    if (ConstantImageLODsPermanent.IsValidIndex(Index.Index))
+		    {
+			    Mip = ConstantImageLODsPermanent[Index.Index];
+		    }
+	    }
+	    else
+	    {
+		    const TManagedPtr<const FImage>* Found = ConstantImageLODsStreamed.Find(Index.Index);
+		    if (Found)
+		    {
+			    Mip = *Found;
+		    }
+	    }
+
+	    return Mip;
+    }
+
+
+    void FProgram::GetExtensionDataConstant(int32 ConstantIndex, TManagedPtr<const FExtensionData>& Result) const
+    {
+	    const FExtensionDataConstant& Constant = ConstantExtensionData[ConstantIndex];
+
+	    check(Constant.Data);
+
+	    Result = Constant.Data;
+    }
+
+
+    EOpType FProgram::GetOpType(OP::ADDRESS At) const
+    {
+	    if (At >= OP::ADDRESS(OpAddress.Num()))
+	    {
+		    return EOpType::NONE;
+	    }
+		
+	    EOpType Result;
+	    uint64 byteCodeAddress = OpAddress[At];
+	    FMemory::Memcpy(&Result, &ByteCode[byteCodeAddress], sizeof(EOpType));
+	    check(Result < EOpType::COUNT);
+	    return Result;
+    }
+
+
+    const uint8* FProgram::GetOpArgsPointer(OP::ADDRESS At) const
+    {
+	    uint64 byteCodeAddress = OpAddress[At];
+	    byteCodeAddress += sizeof(EOpType);
+	    const uint8* pData = (const uint8*)&ByteCode[byteCodeAddress];
+	    return pData;
+    }
+
+
+    uint8* FProgram::GetOpArgsPointer(OP::ADDRESS At)
+    {
+	    uint64 byteCodeAddress = OpAddress[At];
+	    byteCodeAddress += sizeof(EOpType);
+	    uint8* pData = (uint8*)&ByteCode[byteCodeAddress];
+	    return pData;
+    }
+
+
+    void FProgram::MarkRomUsed(int32 RomIndex) const
+    {
+	    RomWeights[RomIndex].Key++;
+	    RomWeights[RomIndex].Value = ++FSystem::RomTick;
+    }
+
+
+    void FModel::Serialise(const FModel* Model, FOutputArchive& Arch)
     {
 		LLM_SCOPE_BYNAME(TEXT("MutableRuntime"));
 
-        arch << *p;
+        Arch << *Model;
     }
 
    
@@ -88,28 +571,27 @@ namespace UE::Mutable::Private
     public:
 
         // Life cycle
-        FOutputModelStream(FModelWriter* InStreamer )
-            : Streamer( InStreamer )
+        FOutputModelStream(FModelWriter* InStreamer)
+            : Streamer(InStreamer)
         {
         }
 
         // FInputStream interface
-        void Write( const void* Data, uint64 Size ) override
+        void Write(const void* Data, uint64 Size) override
         {
-            Streamer->Write( Data, Size );
+            Streamer->Write(Data, Size);
         }
 
     private:
-
         FModelWriter* Streamer;
     };
 
 
-    void FModel::Serialise( FModel* p, FModelWriter& streamer, bool bDropData )
+    void FModel::Serialise(FModel* Model, FModelWriter& Streamer, bool bDropData)
     {
 		LLM_SCOPE_BYNAME(TEXT("MutableRuntime"));
 
-		UE::Mutable::Private::FProgram& Program = p->Program;
+		UE::Mutable::Private::FProgram& Program = Model->Program;
 
 		FOutputMemoryStream MemStream(16 * 1024 * 1024);
 
@@ -126,9 +608,9 @@ namespace UE::Mutable::Private
 			FImage::Serialise(Entry.Value.Get(), MemoryArch);
 			check(RomData.Size == MemStream.GetBufferSize());
 
-			streamer.OpenWriteFile(RomIndex, true);
-			streamer.Write(MemStream.GetBuffer(), MemStream.GetBufferSize());
-			streamer.CloseWriteFile();
+			Streamer.OpenWriteFile(RomIndex, true);
+			Streamer.Write(MemStream.GetBuffer(), MemStream.GetBufferSize());
+			Streamer.CloseWriteFile();
 
 			// Do this progressively to avoid duplicating all data in memory.
 			if (bDropData)
@@ -155,9 +637,9 @@ namespace UE::Mutable::Private
 			FMesh::Serialise(ResData.Value.Get(), MemoryArch);
 			check(RomData.Size == MemStream.GetBufferSize());
 
-			streamer.OpenWriteFile(RomIndex, true);
-			streamer.Write(MemStream.GetBuffer(), MemStream.GetBufferSize());
-			streamer.CloseWriteFile();
+			Streamer.OpenWriteFile(RomIndex, true);
+			Streamer.Write(MemStream.GetBuffer(), MemStream.GetBufferSize());
+			Streamer.CloseWriteFile();
 
 			// Do this progressively to avoid duplicating all data in memory.
 			if (bDropData)
@@ -173,13 +655,13 @@ namespace UE::Mutable::Private
 
 		// Store the main data of the model
 		{
-			streamer.OpenWriteFile(0, false);
-			FOutputModelStream stream(&streamer);
+			Streamer.OpenWriteFile(0, false);
+			FOutputModelStream stream(&Streamer);
 			FOutputArchive arch(&stream);
 
-			arch << *p;
+			arch << *Model;
 
-			streamer.CloseWriteFile();
+			Streamer.CloseWriteFile();
 		}
 	}
 
@@ -213,17 +695,29 @@ namespace UE::Mutable::Private
 	}
 
 
-	TSharedPtr<FModel> FModel::StaticUnserialise( FInputArchive& arch )
+	TSharedPtr<FModel> FModel::StaticUnserialise(FInputArchive& Arch)
     {
 		LLM_SCOPE_BYNAME(TEXT("MutableRuntime"));
 
 		TSharedPtr<FModel> pResult = MakeShared<FModel>();
-        arch >> *pResult;
+        Arch >> *pResult;
         return pResult;
     }
-	
-	
-    bool FModel::GetBoolDefaultValue(int32 Index) const
+
+
+	void FModel::Serialise(FOutputArchive& Arch) const
+	{
+		Arch << Program;
+	}
+
+
+	void FModel::Unserialise(FInputArchive& Arch)
+	{
+		Arch >> Program;
+	}
+
+
+	bool FModel::GetBoolDefaultValue(int32 Index) const
     {
     	check(Program.Parameters.IsValidIndex(Index));
 		check(Program.Parameters[Index].Type == EParameterType::Bool);
@@ -328,50 +822,68 @@ namespace UE::Mutable::Private
 	}
 
 	
-    int32 FModel::GetRomCount() const
+    int32 FProgram::GetRomCount() const
     {
-    	return Program.Roms.Num();
+    	return Roms.Num();
     }
 
 	
 #if WITH_EDITOR
-	uint32 FModel::GetRomSourceId(int32 Index) const
+	uint32 FProgram::GetRomSourceId(int32 Index) const
 	{
-		return Program.RomsCompileData[Index].SourceId;
+		return RomsCompileData[Index].SourceId;
 	}
 #endif
 
 	
-    uint32 FModel::GetRomSize(int32 Index) const
+    uint32 FProgram::GetRomSize(int32 Index) const
     {
-    	return Program.Roms[Index].Size;
+    	return Roms[Index].Size;
     }
 
-	
-	bool FModel::IsMeshData(int32 Index) const
+
+    bool FProgram::IsRomLoaded(int32 RomIndex) const
+    {
+	    switch (Roms[RomIndex].ResourceType)
+	    {
+	    case uint32(ERomDataType::Image):
+		    return ConstantImageLODsStreamed.Contains(RomIndex);
+	    case uint32(ERomDataType::Mesh):
+		    return ConstantMeshesStreamed.Contains(RomIndex);
+	    //TODO(Max): Make materials loadable and unloadable from a rom
+	    default:
+		    check(false);
+		    break;
+	    }
+
+	    return false;
+    }
+
+
+    ERomDataType FProgram::GetRomDataType(int32 Index) const
 	{
-		return Program.Roms[Index].ResourceType==uint32(ERomDataType::Mesh);
+		return static_cast<ERomDataType>(Roms[Index].ResourceType);
 	}
 
 	
-	bool FModel::IsRomHighRes(int32 Index) const
+	bool FProgram::IsRomHighRes(int32 Index) const
 	{
-		return Program.Roms[Index].IsHighRes==1;
+		return Roms[Index].IsHighRes == 1;
 	}
 
 	
-    int32 FModel::GetConstantImageRomId(int32 ConstantImageIndex, int32 LODIndex) const
+    int32 FProgram::GetConstantImageRomId(int32 ConstantImageIndex, int32 LODIndex) const
     {
-        check(Program.ConstantImages.IsValidIndex(ConstantImageIndex));
+        check(ConstantImages.IsValidIndex(ConstantImageIndex));
         
-        FImageLODRange LODRange = Program.ConstantImages[ConstantImageIndex];
+        FImageLODRange LODRange = ConstantImages[ConstantImageIndex];
 
         if (LODIndex >= LODRange.LODCount) 
         {
             return -1; 
         }
 
-        FConstantResourceIndex ResourceIndex = Program.ConstantImageLODIndices[LODRange.FirstIndex + LODIndex];
+        FConstantResourceIndex ResourceIndex = ConstantImageLODIndices[LODRange.FirstIndex + LODIndex];
         if (!ResourceIndex.Streamable)
         {
             return -1;
diff --git a/Engine/Plugins/Mutable/Source/MutableRuntime/Private/MuR/System.cpp b/Engine/Plugins/Mutable/Source/MutableRuntime/Private/MuR/System.cpp
--- a/Engine/Plugins/Mutable/Source/MutableRuntime/Private/MuR/System.cpp
+++ b/Engine/Plugins/Mutable/Source/MutableRuntime/Private/MuR/System.cpp
@@ -77,7 +77,10 @@ namespace UE::Mutable::Private
 	TRACE_DECLARE_INT_COUNTER(MutableRuntime_MemTotal,		TEXT("MutableRuntime/MemTotal"));
 	TRACE_DECLARE_INT_COUNTER(MutableRuntime_MemBudget,		TEXT("MutableRuntime/MemBudget"));
 
-   
+
+	uint64 FSystem::RomTick = 0;
+
+	
     FSystem::FSystem(const FSettings& InSettings)
     {
     	LLM_SCOPE_BYNAME(TEXT("MutableRuntime"));
@@ -257,10 +260,7 @@ namespace UE::Mutable::Private
 		// Prepare instance cache
     	LiveInstance->Cache->Prepare(Program, InStateIndex);
 		LiveInstance->OldParameters = InParams->Clone();
-
-		// Ensure the model cache has been created
-		WorkingMemoryManager.FindOrAddModelCache(LiveInstance->Model);
-
+    	
 		OP::ADDRESS RootAt = Program.States[InStateIndex].Root;
 		RunCode(LiveInstance, RootAt, InLodMask);
 
@@ -841,18 +841,12 @@ namespace UE::Mutable::Private
 
 		if (bClearRoms)
 		{
-			for (FWorkingMemoryManager::FModelCacheEntry& ModelCache : WorkingMemoryManager.CachePerModel)
-			{
-				if (const TSharedPtr<const FModel> CacheModel = ModelCache.Model.Pin())
-				{
-					const FProgram& Program = CacheModel->Program;
+			const FProgram& Program = LiveInstance->Model->Program;
 
-					for (int32 RomIndex = 0; RomIndex < Program.Roms.Num(); ++RomIndex)
-					{
-						Program.UnloadRom(RomIndex);
-					}
-				}
-			}	
+			for (int32 RomIndex = 0; RomIndex < Program.Roms.Num(); ++RomIndex)
+			{
+				Program.UnloadRom(RomIndex);
+			}
 		}
 		
 		UpdateStats();
@@ -961,19 +955,6 @@ namespace UE::Mutable::Private
     }
 
 
-	void FSystem::BeginBuild(const TSharedPtr<const FModel>& InModel)
-	{
-		// Ensure the model cache has been created
-		WorkingMemoryManager.FindOrAddModelCache(InModel);
-	}
-
-
-	void FSystem::EndBuild()
-	{
-
-	}
-
-
 	void FSystem::RunCode(
 		const TSharedRef<FLiveInstance>& InLiveInstance,
 		OP::ADDRESS InCodeRoot,
@@ -1337,58 +1318,6 @@ namespace UE::Mutable::Private
 	}
 	
 
-	FWorkingMemoryManager::FModelCacheEntry* FWorkingMemoryManager::FindModelCache(const FModel* InModel)
-	{
-		check(InModel);
-
-		for (FModelCacheEntry& c : CachePerModel)
-		{
-			if (TSharedPtr<const FModel> Candidate = c.Model.Pin())
-			{
-				if (Candidate.Get() == InModel)
-				{
-					return &c;
-				}
-			}
-		}
-
-		return nullptr;
-	}
-
-   
-	FWorkingMemoryManager::FModelCacheEntry& FWorkingMemoryManager::FindOrAddModelCache(const TSharedPtr<const FModel>& InModel)
-    {
-        check(InModel);
-
-		// First clean stray data for models that may have been unloaded.
-		for (int32 CacheIndex=0; CacheIndex<CachePerModel.Num(); )
-		{
-			if (!CachePerModel[CacheIndex].Model.Pin())
-			{
-				CachePerModel.RemoveAtSwap(CacheIndex);
-			}
-			else
-			{
-				++CacheIndex;
-			}
-		}
-
-
-		FModelCacheEntry* ExistingCache = FindModelCache(InModel.Get());
-		if (ExistingCache)
-		{
-			return *ExistingCache;
-		}        
-
-        // Not found. Add new
-		FModelCacheEntry n;
-        n.Model = TWeakPtr<const FModel>(InModel);
-		n.PendingOpsPerRom.SetNum(InModel->Program.Roms.Num());
-    	n.RomWeights.SetNumZeroed(InModel->Program.Roms.Num());
-    	CachePerModel.Add(n);
-        return CachePerModel.Last();
-    }
-
 
 	int64 FWorkingMemoryManager::GetCurrentMemoryBytes() const
 	{
@@ -1448,38 +1377,39 @@ namespace UE::Mutable::Private
 			Candidates.Reserve(512);
 
 			SSIZE_T RomBytes = 0;
-			for (FModelCacheEntry& ModelCache : CachePerModel)
+			for (TWeakPtr<FLiveInstance>& WeakLiveInstance : LiveInstances)
 			{
-				TSharedPtr<const FModel> CacheModel = ModelCache.Model.Pin();
-				if (CacheModel)
+				TSharedPtr<FLiveInstance> LiveInstance = WeakLiveInstance.Pin();
+				if (!LiveInstance)
 				{
-					const FProgram& Program = CacheModel->Program;
-					check(ModelCache.RomWeights.Num() == Program.Roms.Num());
+					continue;
+				}
+				
+				const FProgram& Program = LiveInstance->Model->Program;
 	
-					check(Program.LoadedMemTrackedRoms.GetMaxIndex() <= Program.Roms.Num());
+				check(Program.LoadedMemTrackedRoms.GetMaxIndex() <= Program.Roms.Num());
 
-					for (TSparseArray<uint8>::TConstIterator Iter = Program.LoadedMemTrackedRoms.CreateConstIterator(); Iter; ++Iter)
-					{
-						const int32 RomIndex = Iter.GetIndex();
+				for (TSparseArray<uint8>::TConstIterator Iter = Program.LoadedMemTrackedRoms.CreateConstIterator(); Iter; ++Iter)
+				{
+					const int32 RomIndex = Iter.GetIndex();
 
-						const FRomDataRuntime& Rom = Program.Roms[RomIndex];
-						check(Program.IsRomLoaded(RomIndex));
-						check(Rom.ResourceType == (uint32)*Iter);
+					const FRomDataRuntime& Rom = Program.Roms[RomIndex];
+					check(Program.IsRomLoaded(RomIndex));
+					check(Rom.ResourceType == (uint32)*Iter);
 
-						RomBytes += CacheModel->Program.GetRomBytes(RomIndex);
+					RomBytes += Program.GetRomBytes(RomIndex);
 
-						// We cannot unload a rom if some operation is expecting it.
-						const bool bIsRomLocked = ModelCache.PendingOpsPerRom.IsValidIndex(RomIndex) && 
-												  ModelCache.PendingOpsPerRom[RomIndex] > 0;
-						if (!bIsRomLocked)
-						{
-							constexpr float FactorWeight = 100.0f;
-							constexpr float FactorTime = -1.0f;
-							float Priority = FactorWeight * float(ModelCache.RomWeights[RomIndex].Get<0>()) +
-											 FactorTime * float((RomTick - ModelCache.RomWeights[RomIndex].Get<1>()));
+					// We cannot unload a rom if some operation is expecting it.
+					const bool bIsRomLocked = Program.PendingOpsPerRom.IsValidIndex(RomIndex) && 
+											  Program.PendingOpsPerRom[RomIndex] > 0;
+					if (!bIsRomLocked)
+					{
+						constexpr float FactorWeight = 100.0f;
+						constexpr float FactorTime = -1.0f;
+						float Priority = FactorWeight * float(Program.RomWeights[RomIndex].Get<0>()) +
+										 FactorTime * float(FSystem::RomTick - Program.RomWeights[RomIndex].Get<1>());
 
-							Candidates.Emplace(Priority, FRomRef{ CacheModel.Get(), RomIndex });
-						}
+						Candidates.Emplace(Priority, FRomRef{ LiveInstance->Model.Get(), RomIndex });
 					}
 				}
 			}
@@ -1635,29 +1565,7 @@ namespace UE::Mutable::Private
 
         return bFinished;
     }
-
-   
-    void FWorkingMemoryManager::MarkRomUsed( int32 romIndex, const TSharedPtr<const FModel>& pModel )
-    {
-        check(pModel);
-
-        // If budget is zero, we don't unload anything here, and we assume it is managed somewhere else.
-        if (!BudgetBytes)
-        {
-            return;
-        }
-
-        ++RomTick;
-
-        // Update current cache
-        {
-			FModelCacheEntry* ModelCache = FindModelCache(pModel.Get());
-        	
-            ModelCache->RomWeights[romIndex].Key++;
-            ModelCache->RomWeights[romIndex].Value = RomTick;
-        }
-    }
-
+	
 
 	static void AddMultiValueKeys(FBitWriter& Blob, const TMap< TArray<int32>, FParameterValue >& Multi)
 	{
diff --git a/Engine/Plugins/Mutable/Source/MutableTools/Private/MuT/CodeOptimiser.cpp b/Engine/Plugins/Mutable/Source/MutableTools/Private/MuT/CodeOptimiser.cpp
--- a/Engine/Plugins/Mutable/Source/MutableTools/Private/MuT/CodeOptimiser.cpp
+++ b/Engine/Plugins/Mutable/Source/MutableTools/Private/MuT/CodeOptimiser.cpp
@@ -198,8 +198,6 @@ namespace UE::Mutable::Private
 			TSharedRef<FLiveInstance> TempLiveInstance = System->NewBuildInstance(Model, nullptr);
 			TempLiveInstance->OldParameters = LocalParams;
 			
-			System->BeginBuild(Model);
-			
 			// Calculate the value and replace this op by a constant
 			switch( DataType )
 			{
@@ -287,10 +285,7 @@ namespace UE::Mutable::Private
 			default:
 				break;
 			}
-
-			System->EndBuild();
 		}
-
 	};
 
 
diff --git a/Engine/Plugins/Mutable/Source/MutableTools/Private/MuT/Compiler.cpp b/Engine/Plugins/Mutable/Source/MutableTools/Private/MuT/Compiler.cpp
--- a/Engine/Plugins/Mutable/Source/MutableTools/Private/MuT/Compiler.cpp
+++ b/Engine/Plugins/Mutable/Source/MutableTools/Private/MuT/Compiler.cpp
@@ -665,6 +665,9 @@ namespace UE::Mutable::Private
 		// Pack data
 		m_pD->GenerateRoms(Result.Get(), LinkerOptions.AdditionalData );
 
+		Program.PendingOpsPerRom.SetNum(Program.Roms.Num());
+		Program.RomWeights.SetNumZeroed(Program.Roms.Num());
+		
 		// We are not touching the program anymore. Ensure we are not wasting memory.
 		Program.Roms.Shrink();
 		Program.ConstantImageLODsPermanent.Shrink();
diff --git a/Engine/Plugins/Mutable/Source/MutableValidation/Private/MuV/ValidationUtils.cpp b/Engine/Plugins/Mutable/Source/MutableValidation/Private/MuV/ValidationUtils.cpp
--- a/Engine/Plugins/Mutable/Source/MutableValidation/Private/MuV/ValidationUtils.cpp
+++ b/Engine/Plugins/Mutable/Source/MutableValidation/Private/MuV/ValidationUtils.cpp
@@ -210,11 +210,11 @@ bool TestCustomizableObject(UCustomizableObject& InTargetCustomizableObject, con
 
 		// Roms ---------------------- //
 		{
-			const int32 RomCount =  MutableModel->GetRomCount();
+			const int32 RomCount =  MutableModel->Program.GetRomCount();
 			int64 TotalRomSizeBytes = 0;
 			for (int32 RomIndex = 0; RomIndex < RomCount; RomIndex++)
 			{
-				const uint32 RomByteSize = MutableModel->GetRomSize(RomIndex);
+				const uint32 RomByteSize = MutableModel->Program.GetRomSize(RomIndex);
 				TotalRomSizeBytes += RomByteSize;
 			}
 
