[Mutable] Fix Runner static analysis.

#Synced-CL 49839474

#Synced-CL 49839474

diff --git a/Engine/Plugins/Mutable/Source/MutableRuntime/Private/MuR/CodeRunnerTasked.cpp b/Engine/Plugins/Mutable/Source/MutableRuntime/Private/MuR/CodeRunnerTasked.cpp
--- a/Engine/Plugins/Mutable/Source/MutableRuntime/Private/MuR/CodeRunnerTasked.cpp
+++ b/Engine/Plugins/Mutable/Source/MutableRuntime/Private/MuR/CodeRunnerTasked.cpp
@@ -2099,13 +2099,8 @@ namespace UE::Mutable::Private
 	{
 		// This runs in the Runner thread
 		MUTABLE_CPUPROFILER_SCOPE(FLoadMeshRomTask_Complete);
-
-		const FProgram& LocalProgram = Runner->Program;
 		
-		if (!Runner)
-		{
-			return false;
-		}
+		const FProgram& LocalProgram = Runner->Program;
 		
 		bool bSomeMissingData = false;
 		for (const int32 RomIndex : RomIndices)
