[Mutable] Fix CO always automatically compiling.

#Synced-CL 49839474

diff --git a/Engine/Plugins/Mutable/Source/CustomizableObject/Private/MuCO/CustomizableObjectInstance.cpp b/Engine/Plugins/Mutable/Source/CustomizableObject/Private/MuCO/CustomizableObjectInstance.cpp
--- a/Engine/Plugins/Mutable/Source/CustomizableObject/Private/MuCO/CustomizableObjectInstance.cpp
+++ b/Engine/Plugins/Mutable/Source/CustomizableObject/Private/MuCO/CustomizableObjectInstance.cpp
@@ -870,7 +870,10 @@ void UCustomizableInstancePrivate::TickUpdateCloseCustomizableObjects(UCustomiza
 	{
 		if (ICustomizableObjectEditorModule* EditorModule = ICustomizableObjectEditorModule::Get())
 		{
-			EditorModule->CompileCustomizableObject(*Object, nullptr, true, false);
+			FCompileParams CompileParams;
+			CompileParams.bSkipIfCompiled = true;
+			
+			EditorModule->CompileCustomizableObject(*Object, &CompileParams, true, false);
 		}
 	}
 #endif
