[mutable] Made the LiveInsance be reused during all the operations of a single COI update (Mesh and image streaming included)
- Forced the live update to be disabled also in the case we have the streaming of meshes enabled.
- Mesh and image streaming tasks will now reuse the LiveInstance generated during the initial generation.
- A Cvar has been added: "mutable.ReuseLiveInstanceCacheBetweenUpdates" to enable the sharing of the cached data from the InitGenration with the streaming operations. It is disabled at the moment as that functionality is not yet ready.
- Removed "State" and "Paramneters" from the "OpeationData" structure, as that data should already be part of the LiveInstance being used during the execution of the update.
- Made the ProgramCache "Clear" method also empty the arrays containing pointers to the generated mutable data structures (like FLayout, FInstance or FMesh)
- LiveInstance OldParamters are now named Parameters
- Other small tweaks

#rnx

#Synced-CL 49920339

diff --git a/Engine/Plugins/Mutable/Source/CustomizableObject/Internal/MuCO/CustomizableObjectMipDataProvider.h b/Engine/Plugins/Mutable/Source/CustomizableObject/Internal/MuCO/CustomizableObjectMipDataProvider.h
--- a/Engine/Plugins/Mutable/Source/CustomizableObject/Internal/MuCO/CustomizableObjectMipDataProvider.h
+++ b/Engine/Plugins/Mutable/Source/CustomizableObject/Internal/MuCO/CustomizableObjectMipDataProvider.h
@@ -6,7 +6,6 @@
 #include "MuCO/FMutableTaskGraph.h"
 #include "Async/TaskGraphInterfaces.h"
 #include "Engine/TextureMipDataProviderFactory.h"
-#include "MuCO/CustomizableObjectSystem.h"
 #include "MuR/Image.h"
 #include "Streaming/TextureMipDataProvider.h"
 
@@ -32,6 +31,12 @@ class FMutableUpdateImageContext
 {
 public:
 
+	FMutableUpdateImageContext(TSharedRef<UE::Mutable::Private::FLiveInstance> InLiveInstance)
+		: LiveInstance(InLiveInstance)
+	{
+		
+	}
+	
 	// Benchmarking Utility data (it may not always be present)
 	FString CapturedDescriptor;
 	bool bLevelBegunPlay = false;
@@ -41,22 +46,12 @@ public:
 
 	FName CustomizableObjectName;
 
+	TSharedRef<UE::Mutable::Private::FLiveInstance> LiveInstance;
+	
 	TSharedPtr<UE::Mutable::Private::FSystem> System;
-	TSharedPtr<UE::Mutable::Private::FModel> Model;
-
-	TSharedPtr<UE::Mutable::Private::FMeshIdRegistry> MeshIdRegistry;
-	TSharedPtr<UE::Mutable::Private::FImageIdRegistry> ImageIdRegistry;
-	TSharedPtr<UE::Mutable::Private::FMaterialIdRegistry> MaterialIdRegistry;
-	TSharedPtr<UE::Mutable::Private::FSkeletalMeshIdRegistry> SkeletalMeshIdRegistry;
-
-	TSharedPtr<FUnrealMutableResourceProvider> ExternalResourceProvider;
 
 	TSharedPtr<FModelStreamableBulkData> ModelStreamableBulkData;
 	
-	TSharedPtr<const UE::Mutable::Private::FParameters> Parameters;
-	int32 State = -1;
-
-	TArray<TSharedPtr<const UE::Mutable::Private::FImage>> ImageParameterValues;
 };
 
 
diff --git a/Engine/Plugins/Mutable/Source/CustomizableObject/Internal/MuCO/CustomizableObjectSystemPrivate.h b/Engine/Plugins/Mutable/Source/CustomizableObject/Internal/MuCO/CustomizableObjectSystemPrivate.h
--- a/Engine/Plugins/Mutable/Source/CustomizableObject/Internal/MuCO/CustomizableObjectSystemPrivate.h
+++ b/Engine/Plugins/Mutable/Source/CustomizableObject/Internal/MuCO/CustomizableObjectSystemPrivate.h
@@ -37,6 +37,7 @@ class FUnrealMutableResourceProvider;
 class FMutableStreamRequest;
 class UEditorImageProvider;
 class UCustomizableObjectSystem;
+struct FParameters;
 namespace LowLevelTasks { enum class ETaskPriority : int8; }
 struct FTexturePlatformData;
 
@@ -574,6 +575,9 @@ public:
 	/** Manage loading Passthrough Objects discovered on the StartUpdate. */
 	TSharedPtr<UE::Mutable::Private::FPassthroughObjectLoader> PassthroughObjectLoader;
 	
+	/** The parameters the Live Instance had before being reused. */
+	TSharedPtr<const UE::Mutable::Private::FParameters> OldParameters;
+	
 	/** Weak reference to the instance we are operating on.
 	 *It is weak because we don't want to lock it in case it becomes irrelevant in the game while operations are pending and it needs to be destroyed. */
 	TWeakObjectPtr<UCustomizableObjectInstance> Instance;
@@ -606,12 +610,15 @@ public:
 	/** This option comes from the operation request. It is used to reduce the number of mipmaps that mutable must generate for images.  */
 	int32 MipsToSkip = 0;
 	
+	/** The state the Live Instance had before being reused. */
+	int32 OldState = -1;
+	
 	EQueuePriorityType PriorityType = EQueuePriorityType::Low;
 
 	EUpdateResult UpdateResult = EUpdateResult::Success;
 	
 	/** This option comes from the operation request */
-	uint8 bNeverStream : 1 = false;
+	uint8 bNeverStreamMips : 1 = false;
 	
 	uint8 bOnlyUpdateIfNotGenerated : 1 = false;
 
diff --git a/Engine/Plugins/Mutable/Source/CustomizableObject/Private/MuCO/CustomizableObjectInstance.cpp b/Engine/Plugins/Mutable/Source/CustomizableObject/Private/MuCO/CustomizableObjectInstance.cpp
--- a/Engine/Plugins/Mutable/Source/CustomizableObject/Private/MuCO/CustomizableObjectInstance.cpp
+++ b/Engine/Plugins/Mutable/Source/CustomizableObject/Private/MuCO/CustomizableObjectInstance.cpp
@@ -5491,7 +5491,7 @@ UTexture* BuildTexture(const TSharedRef<FUpdateContextPrivate>& Context, const i
 	FTextureCache& TextureCache = Context->Object->GetPrivate()->TextureCache;
 
 	FString CurrentState = Context->Instance->GetCurrentState();
-	bool bNeverStream = Context->bNeverStream;
+	bool bNeverStream = Context->bNeverStreamMips;
 
 	check((bNeverStream && Context->MipsToSkip == 0) ||
 		(!bNeverStream && Context->MipsToSkip >= 0));
@@ -5609,7 +5609,7 @@ UTexture* BuildTexture(const TSharedRef<FUpdateContextPrivate>& Context, const i
 				}
 			}
 				
-			//if (!bNeverStream) // No need to check bNeverStream. In that case, the texture won't use 
+			//if (!bNeverStreamMips) // No need to check bNeverStreamMips. In that case, the texture won't use 
 			// the MutableMipDataProviderFactory anyway and it's needed for detecting Mutable textures elsewhere
 			UMutableTextureMipDataProviderFactory* MutableMipDataProviderFactory = Cast<UMutableTextureMipDataProviderFactory>(Texture->GetAssetUserDataOfClass(UMutableTextureMipDataProviderFactory::StaticClass()));
 			if (!MutableMipDataProviderFactory)
@@ -5828,7 +5828,7 @@ FBuildMaterialOut BuildMaterial(const TSharedRef<FUpdateContextPrivate>& Context
 		}
 		else
 		{
-			// Add it to the CO cache. Currently we do not have one for Materials.
+			// Add it to the CO cache. Currently, we do not have one for Materials.
 		}
 	}
 
@@ -5849,19 +5849,12 @@ void UCustomizableInstancePrivate::BuildResources(const TSharedRef<FUpdateContex
 	const UModelResources& ModelResources = *CustomizableObject->GetPrivate()->GetModelResources();
 	
 	// Prepare the data to store in order to regenerate resources for this instance (usually texture mips).
-	TSharedRef<FMutableUpdateImageContext> UpdateImageContext = MakeShared<FMutableUpdateImageContext>();
+	check(Context->LiveInstance)
+	TSharedRef<FMutableUpdateImageContext> UpdateImageContext = MakeShared<FMutableUpdateImageContext>(Context->LiveInstance.ToSharedRef());
+	UpdateImageContext->System = UCustomizableObjectSystem::GetInstance()->GetPrivate()->MutableSystem;
 	UpdateImageContext->CustomizableObjectPathName = CustomizableObject->GetPathName();
 	UpdateImageContext->InstancePathName = Public->GetPathName();
-	UpdateImageContext->System = UCustomizableObjectSystem::GetInstance()->GetPrivate()->MutableSystem;
-	UpdateImageContext->Model = Context->Model;
-	UpdateImageContext->MeshIdRegistry = CustomizableObject->GetPrivate()->MeshIdRegistry;
-	UpdateImageContext->ImageIdRegistry = CustomizableObject->GetPrivate()->ImageIdRegistry;
-    UpdateImageContext->MaterialIdRegistry = CustomizableObject->GetPrivate()->MaterialIdRegistry,
-    UpdateImageContext->SkeletalMeshIdRegistry = CustomizableObject->GetPrivate()->SkeletalMeshIdRegistry,
-	UpdateImageContext->ExternalResourceProvider = Context->ExternalResourceProvider;
 	UpdateImageContext->ModelStreamableBulkData = CustomizableObject->GetPrivate()->GetModelStreamableBulkData();
-	UpdateImageContext->Parameters = Context->Parameters;
-	UpdateImageContext->State = Context->GetCapturedDescriptor().GetState();
 
 	// Cache the descriptor as a string if we want to later report it using our benchmark utility. 
 	if (FLogBenchmarkUtil::IsBenchmarkingReportingEnabled())
diff --git a/Engine/Plugins/Mutable/Source/CustomizableObject/Private/MuCO/CustomizableObjectMeshUpdate.cpp b/Engine/Plugins/Mutable/Source/CustomizableObject/Private/MuCO/CustomizableObjectMeshUpdate.cpp
--- a/Engine/Plugins/Mutable/Source/CustomizableObject/Private/MuCO/CustomizableObjectMeshUpdate.cpp
+++ b/Engine/Plugins/Mutable/Source/CustomizableObject/Private/MuCO/CustomizableObjectMeshUpdate.cpp
@@ -14,7 +14,6 @@ CustomizableObjectMeshUpdate.cpp: Helpers to stream in CustomizableObject skelet
 #include "MuCO/CustomizableObjectSystemPrivate.h"
 #include "MuCO/CustomizableObjectSkeletalMesh.h"
 #include "MuCO/UnrealConversionUtils.h"
-#include "MuCO/UnrealMutableImageProvider.h"
 
 #include "MuR/Model.h"
 #include "MuR/MeshBufferSet.h"
@@ -45,36 +44,22 @@ FAutoConsoleVariableRef CVarMutableEnableCPUMorphTargetStreaming(
 
 FCustomizableObjectMeshStreamIn::FCustomizableObjectMeshStreamIn(
 	const UCustomizableObjectSkeletalMesh* InMesh,
-	EThreadType CreateResourcesThread,
-	const TSharedPtr<FModelStreamableBulkData>& ModelStreamableBulkData) :
+	EThreadType CreateResourcesThread) :
 	FSkeletalMeshStreamIn(InMesh, CreateResourcesThread)
 {
-	OperationData = MakeShared<FMutableUpdateMeshContext>();
-
-	// This must run in the mutable thread.
-	check(UCustomizableObjectSystem::IsCreated());
-	OperationData->System = UCustomizableObjectSystem::GetInstance()->GetPrivate()->MutableSystem;
-
-	OperationData->Model = InMesh->Model;
-	OperationData->Parameters = InMesh->Parameters;
-	OperationData->MeshIdRegistry = InMesh->MeshIdRegistry;
-	OperationData->ImageIdRegistry = InMesh->ImageIdRegistry;
-	OperationData->MaterialIdRegistry = InMesh->MaterialIdRegistry;
-	OperationData->SkeletalMeshIdRegistry = InMesh->SkeletalMeshIdRegistry;
-	OperationData->ExternalResourceProvider = InMesh->ExternalResourceProvider;
-	OperationData->State = InMesh->State;
-
-	OperationData->ModelStreamableBulkData = ModelStreamableBulkData;
+	OperationData = InMesh->UpdateMeshContext;
+	check(OperationData)
+	
+	// Set up the event that will indicate that the Task_Mutable_UpdateMesh_End was run
+	OperationData->MeshStreamingPrerequisite =  MakeShared<UE::Tasks::FTaskEvent>(TEXT("Task_Mutable_UpdateMesh_End"));
 	
-	OperationData->MeshIDs = InMesh->MeshIDs;
-	OperationData->Meshes.SetNum(InMesh->MeshIDs.Num());
-
 	OperationData->CurrentFirstLODIdx = CurrentFirstLODIdx;
 	OperationData->PendingFirstLODIdx = PendingFirstLODIdx;
 
 	PushTask(FContext(InMesh, TT_None), TT_Async, SRA_UPDATE_CALLBACK(DoInitiate), TT_None, nullptr);
 }
 
+
 void FCustomizableObjectMeshStreamIn::OnUpdateMeshFinished()
 {
 	if (!IsCancelled() || !bEnableGCHangFix)
@@ -90,6 +75,7 @@ void FCustomizableObjectMeshStreamIn::OnUpdateMeshFinished()
 	}
 }
 
+
 void FCustomizableObjectMeshStreamIn::Abort()
 {
 	if (!IsCancelled() && !IsCompleted() && bEnableGCHangFix)
@@ -114,6 +100,7 @@ void FCustomizableObjectMeshStreamIn::Abort()
 	}
 }
 
+
 void FCustomizableObjectMeshStreamIn::DoInitiate(const FContext& Context)
 {
 	check(Context.CurrentThread == TT_Async);
@@ -133,6 +120,7 @@ void FCustomizableObjectMeshStreamIn::DoInitiate(const FContext& Context)
 	}
 }
 
+
 void FCustomizableObjectMeshStreamIn::DoConvertResources(const FContext& Context)
 {
 	check(Context.CurrentThread == TT_Async);
@@ -152,6 +140,7 @@ void FCustomizableObjectMeshStreamIn::DoConvertResources(const FContext& Context
 	}
 }
 
+
 void FCustomizableObjectMeshStreamIn::DoCreateBuffers(const FContext& Context)
 {
 	MUTABLE_CPUPROFILER_SCOPE(FCustomizableObjectMeshStreamIn::DoCreateBuffers)
@@ -167,6 +156,7 @@ void FCustomizableObjectMeshStreamIn::DoCreateBuffers(const FContext& Context)
 		TT_Render, SRA_UPDATE_CALLBACK(DoFinishUpdate));
 }
 
+
 void FCustomizableObjectMeshStreamIn::DoCancelMeshUpdate(const FContext& Context)
 {
 	MUTABLE_CPUPROFILER_SCOPE(FCustomizableObjectMeshStreamIn::DoCancelMeshUpdate)
@@ -186,11 +176,13 @@ namespace impl
 		const bool bClearRoms = UCustomizableObjectSystem::ShouldClearWorkingMemoryOnUpdateEnd();
 		OperationData->System->EndUpdate(LiveInstance, bClearRoms);
 
-		OperationData->Event.Trigger();
+		check(OperationData->MeshStreamingPrerequisite.IsValid())
+		OperationData->MeshStreamingPrerequisite->Trigger();
 		
 		TRACE_END_REGION(UE_MUTABLE_UPDATE_MESH_REGION);
 	}
 
+	
 	void Task_Mutable_UpdateMesh_Loop(
 		const TSharedPtr<FMutableUpdateMeshContext> OperationData,
 		TRefCountPtr<FCustomizableObjectMeshStreamIn>& Task,
@@ -220,6 +212,7 @@ namespace impl
 			LowLevelTasks::ETaskPriority::Inherit));
 	}
 
+	
 	void Task_Mutable_UpdateMesh(const TSharedPtr<FMutableUpdateMeshContext> OperationData, TRefCountPtr<FCustomizableObjectMeshStreamIn>& Task)
 	{
 		MUTABLE_CPUPROFILER_SCOPE(Task_Mutable_UpdateMesh);
@@ -232,9 +225,10 @@ namespace impl
 		TRACE_BEGIN_REGION(UE_MUTABLE_UPDATE_MESH_REGION);
 
 		TSharedPtr<UE::Mutable::Private::FSystem> System = OperationData->System;
-		const TSharedPtr<UE::Mutable::Private::FModel, ESPMode::ThreadSafe> Model = OperationData->Model;
-
+		TSharedRef<UE::Mutable::Private::FLiveInstance> LiveInstance = OperationData->LiveInstance;
+		
 #if WITH_EDITOR
+		const TSharedPtr<const UE::Mutable::Private::FModel, ESPMode::ThreadSafe> Model = LiveInstance->Model;
 		// Recompiling a CO in the editor will invalidate the previously generated Model. Check that it is valid before accessing the streamed data.
 		if (!Model || !Model->IsValid())
 		{
@@ -244,23 +238,20 @@ namespace impl
 		}
 #endif
 
-		// For now, we are forcing the recreation of mutable-side instances with every update.
-		TSharedRef<UE::Mutable::Private::FLiveInstance> LiveInstance = System->NewInstance(
-			Model, 
-			OperationData->MeshIdRegistry,
-			OperationData->ImageIdRegistry,
-			OperationData->MaterialIdRegistry,
-			OperationData->SkeletalMeshIdRegistry,
-			OperationData->ExternalResourceProvider, 
-			nullptr);
-
 		// LOD mask, set to all ones to build all LODs
-		const uint32 LODMask = 0xFFFFFFFF;
-
+		constexpr uint32 LODMask = 0xFFFFFFFF;
+		
+		// Prevent the usage of the already cached elements by fully clearing the cache.
+		if (!CVarReuseLiveInstanceCacheBetweenUpdates.GetValueOnAnyThread())
+		{
+			LiveInstance->Cache->Clear();
+		}
+		
 		// Main instance generation step
+		// note: as this method has already been run we provide the param and stats of the instance to prevent the full clearing of its contents
 		UE::Mutable::Private::TManagedPtr<const UE::Mutable::Private::FInstance> Instance = System->BeginUpdate_MutableThread(LiveInstance,
-			OperationData->Parameters,
-			OperationData->State,
+			OperationData->LiveInstance->Parameters,
+			OperationData->LiveInstance->State,
 			LODMask);
 		
 		check(Instance);
@@ -293,7 +284,8 @@ void FCustomizableObjectMeshStreamIn::RequestMeshUpdate(const FContext& Context)
 
 #if WITH_EDITOR
 	// Recompiling a CO in the editor will invalidate the previously generated Model. Check that it is valid before accessing the streamed data.
-	if (!OperationData->Model || !OperationData->Model->IsValid())
+	const TSharedPtr<const UE::Mutable::Private::FModel> Model = OperationData->LiveInstance->Model;
+	if (!Model || !Model->IsValid())
 	{
 		Abort();
 		return;
@@ -314,6 +306,8 @@ void FCustomizableObjectMeshStreamIn::RequestMeshUpdate(const FContext& Context)
 		});
 
 	// Stream data (morphs, clothing...).
+	check(OperationData->MeshStreamingPrerequisite.IsValid())
+	
 	UE::Tasks::TTask StreamData = UE::Tasks::Launch(TEXT("StreamData"),
 		[RefThis = TRefCountPtr(this)]()
 		{
@@ -336,7 +330,7 @@ void FCustomizableObjectMeshStreamIn::RequestMeshUpdate(const FContext& Context)
 			// Stream data.
 			UE::Tasks::AddNested(StreamRequest.Stream());
 		},
-		OperationData->Event,
+		*OperationData->MeshStreamingPrerequisite,
 		UE::Tasks::ETaskPriority::Inherit);
 
 	// Go to next step.
diff --git a/Engine/Plugins/Mutable/Source/CustomizableObject/Private/MuCO/CustomizableObjectMeshUpdate.h b/Engine/Plugins/Mutable/Source/CustomizableObject/Private/MuCO/CustomizableObjectMeshUpdate.h
--- a/Engine/Plugins/Mutable/Source/CustomizableObject/Private/MuCO/CustomizableObjectMeshUpdate.h
+++ b/Engine/Plugins/Mutable/Source/CustomizableObject/Private/MuCO/CustomizableObjectMeshUpdate.h
@@ -24,19 +24,17 @@ namespace UE::Mutable::Private
 
 struct FMutableUpdateMeshContext
 {
+	FMutableUpdateMeshContext(const TSharedRef<UE::Mutable::Private::FLiveInstance>& InLiveInstance)
+		: LiveInstance(InLiveInstance)
+	{
+		
+	}
+	
 	TSharedPtr<UE::Mutable::Private::FSystem> System;
-	TSharedPtr<UE::Mutable::Private::FModel> Model;
-	TSharedPtr<UE::Mutable::Private::FMeshIdRegistry> MeshIdRegistry;
-	TSharedPtr<UE::Mutable::Private::FImageIdRegistry> ImageIdRegistry;
-	TSharedPtr<UE::Mutable::Private::FMaterialIdRegistry> MaterialIdRegistry;
-	TSharedPtr<UE::Mutable::Private::FSkeletalMeshIdRegistry> SkeletalMeshIdRegistry;
-	TSharedPtr<FUnrealMutableResourceProvider> ExternalResourceProvider;
-	TSharedPtr<const UE::Mutable::Private::FParameters> Parameters;
+	TSharedRef<UE::Mutable::Private::FLiveInstance> LiveInstance;
 
 	TSharedPtr<FModelStreamableBulkData> ModelStreamableBulkData;
 	
-	int32 State = -1;
-
 	// IDs of the meshes to generate per LOD.
 	TArray<UE::Mutable::Private::FMeshId> MeshIDs;
 
@@ -49,7 +47,8 @@ struct FMutableUpdateMeshContext
 	/** Serialized data per block. */
 	TMap<uint32, FClothingMeshData> ClothingMeshData;
 
-	UE::Tasks::FTaskEvent Event = UE::Tasks::FTaskEvent(TEXT("Task_Mutable_UpdateMesh_End"));
+	/** Task setup during the nominal operation of the mesh streaming system. It will be null in non-streaming operation. */
+	TSharedPtr<UE::Tasks::FTaskEvent> MeshStreamingPrerequisite;
 
 	/** SkeletalMesh */
 	/** The resident first LOD resource index.With domain = [0, ResourceState.NumLODs[.NOT THE ASSET LOD INDEX! */ 
@@ -73,8 +72,7 @@ class FCustomizableObjectMeshStreamIn : public FSkeletalMeshStreamIn
 public:
 	FCustomizableObjectMeshStreamIn(
 		const UCustomizableObjectSkeletalMesh* InMesh,
-		EThreadType CreateResourcesThread,
-		const TSharedPtr<FModelStreamableBulkData>& ModelStreamableBulkData);
+		EThreadType CreateResourcesThread);
 
 	void OnUpdateMeshFinished();
 
diff --git a/Engine/Plugins/Mutable/Source/CustomizableObject/Private/MuCO/CustomizableObjectMipDataProvider.cpp b/Engine/Plugins/Mutable/Source/CustomizableObject/Private/MuCO/CustomizableObjectMipDataProvider.cpp
--- a/Engine/Plugins/Mutable/Source/CustomizableObject/Private/MuCO/CustomizableObjectMipDataProvider.cpp
+++ b/Engine/Plugins/Mutable/Source/CustomizableObject/Private/MuCO/CustomizableObjectMipDataProvider.cpp
@@ -4,6 +4,7 @@
 
 #include "MuCO/CustomizableObjectPrivate.h"
 #include "MuCO/CustomizableObjectInstance.h"
+#include "MuCO/CustomizableObjectSystem.h"
 #include "MuCO/CustomizableObjectSystemPrivate.h"
 #include "MuCO/MutableStreamRequest.h"
 #include "MuCO/UnrealMutableImageProvider.h"
@@ -91,8 +92,8 @@ namespace Impl
 		// This runs in a worker thread.
 		check(OperationData.IsValid());
 		check(OperationData->UpdateContext->System);
-		check(OperationData->UpdateContext->Model);
-		check(OperationData->UpdateContext->Parameters);
+		check(OperationData->UpdateContext->LiveInstance->Model);
+		check(OperationData->UpdateContext->LiveInstance->Parameters);
 
 		if (!OperationData.IsValid())
 		{
@@ -100,9 +101,11 @@ namespace Impl
 		}
 
 		TRACE_BEGIN_REGION(UE_MUTABLE_MIPDATA_PROVIDER_UPDATE_IMAGE_REGION);
+
+		const TSharedRef<UE::Mutable::Private::FLiveInstance> LiveInstance = OperationData->UpdateContext->LiveInstance;
 		
 		TSharedPtr<UE::Mutable::Private::FSystem> System = OperationData->UpdateContext->System;
-		const TSharedPtr<UE::Mutable::Private::FModel, ESPMode::ThreadSafe> Model = OperationData->UpdateContext->Model;
+		const TSharedPtr<const UE::Mutable::Private::FModel, ESPMode::ThreadSafe> Model = LiveInstance->Model;
 
 		auto EndUpdateImage = [](TSharedPtr<FMutableImageOperationData>& OperationData)
 		{
@@ -129,27 +132,23 @@ namespace Impl
 			return;
 		}
 #endif
-
-		// For now, we are forcing the recreation of mutable-side instances with every update.
-		TSharedRef<UE::Mutable::Private::FLiveInstance> LiveInstance = System->NewInstance(
-			Model,
-			OperationData->UpdateContext->MeshIdRegistry,
-			OperationData->UpdateContext->ImageIdRegistry,
-			OperationData->UpdateContext->MaterialIdRegistry,
-			OperationData->UpdateContext->SkeletalMeshIdRegistry,
-			OperationData->UpdateContext->ExternalResourceProvider, 
-			nullptr);
 		
 		UE::Mutable::Private::TManagedPtr<const UE::Mutable::Private::FInstance> Instance;
 
 		// Main instance generation step
 		{
 			// LOD mask, set to all ones to build all LODs
-			uint32 LODMask = 0xFFFFFFFF;
-
+			constexpr uint32 LODMask = 0xFFFFFFFF;
+			
+			// Prevent the usage of the already cached elements by fully clearing the cache.
+			if (!CVarReuseLiveInstanceCacheBetweenUpdates.GetValueOnAnyThread())
+			{
+				LiveInstance->Cache->Clear();
+			}
+			
 			Instance = System->BeginUpdate_MutableThread(LiveInstance,
-				OperationData->UpdateContext->Parameters,
-				OperationData->UpdateContext->State,
+				LiveInstance->Parameters,
+				LiveInstance->State,
 				LODMask);
 			
 			check(Instance);
@@ -160,7 +159,7 @@ namespace Impl
 		UE::Tasks::TTask<UE::Mutable::Private::TManagedPtr<const UE::Mutable::Private::FImage>> GetImageTask = System->GetImage({}, LiveInstance, ImageRef.ImageID, ImageRef.BaseMip + OperationData->MipsToSkip);
 
 		UE::Tasks::AddNested(UE::Tasks::Launch(TEXT("MipDataProvider_EndUpdateImagesTask"),
-				[System, OperationData, LiveInstance, StartTime, GetImageTask, EndUpdateImage]() mutable
+				[System, OperationData, StartTime, GetImageTask, EndUpdateImage]() mutable
 				{
 					check(GetImageTask.IsCompleted());
 
@@ -168,8 +167,8 @@ namespace Impl
 
 					check(ResultImage);
 
-					int32 FullMipCount = ResultImage->GetMipmapCount(ResultImage->GetSizeX(), ResultImage->GetSizeY());
-					int32 RealMipCount = ResultImage->GetLODCount();
+					const int32 FullMipCount = ResultImage->GetMipmapCount(ResultImage->GetSizeX(), ResultImage->GetSizeY());
+					const int32 RealMipCount = ResultImage->GetLODCount();
 
 					// Did we fail to generate the entire mipchain (if we have mips at all)?
 					bool bForceMipchain = (RealMipCount != 1) && (RealMipCount != FullMipCount);
@@ -205,7 +204,7 @@ namespace Impl
 					{
 						MUTABLE_CPUPROFILER_SCOPE(EndUpdate);
 						const bool bClearRoms = UCustomizableObjectSystem::ShouldClearWorkingMemoryOnUpdateEnd();
-						System->EndUpdate(LiveInstance, bClearRoms);
+						System->EndUpdate(OperationData->UpdateContext->LiveInstance, bClearRoms);
 					}
 
 					if (FLogBenchmarkUtil::IsBenchmarkingReportingEnabled())
@@ -357,7 +356,7 @@ int32 FMutableTextureMipDataProvider::GetMips(const FTextureUpdateContext& Conte
 
 			uint32 MaxBlockSize = 0;
 
-			TSharedPtr<UE::Mutable::Private::FModel, ESPMode::ThreadSafe> Model = UpdateContext->Model;
+			TSharedPtr<const UE::Mutable::Private::FModel, ESPMode::ThreadSafe> Model = UpdateContext->LiveInstance->Model;
 			check(Model);
 
 			if (!OperationData->Levels.IsEmpty())
diff --git a/Engine/Plugins/Mutable/Source/CustomizableObject/Private/MuCO/CustomizableObjectSkeletalMesh.cpp b/Engine/Plugins/Mutable/Source/CustomizableObject/Private/MuCO/CustomizableObjectSkeletalMesh.cpp
--- a/Engine/Plugins/Mutable/Source/CustomizableObject/Private/MuCO/CustomizableObjectSkeletalMesh.cpp
+++ b/Engine/Plugins/Mutable/Source/CustomizableObject/Private/MuCO/CustomizableObjectSkeletalMesh.cpp
@@ -25,17 +25,6 @@ void UCustomizableObjectSkeletalMesh::InitMutableStreamingData(const TSharedRef<
 	// Debug info
 	CustomizableObjectPathName = GetNameSafe(&CustomizableObject);
 
-	// Init properties
-	Model = CustomizableObject.GetPrivate()->GetModel();
-	MeshIdRegistry = CustomizableObject.GetPrivate()->MeshIdRegistry;
-	ImageIdRegistry = CustomizableObject.GetPrivate()->ImageIdRegistry;
-	MaterialIdRegistry = CustomizableObject.GetPrivate()->MaterialIdRegistry;
-	SkeletalMeshIdRegistry = CustomizableObject.GetPrivate()->SkeletalMeshIdRegistry;
-	ExternalResourceProvider = InOperationData->ExternalResourceProvider;
-
-	Parameters = InOperationData->Parameters;
-	State = InOperationData->GetCapturedDescriptor().GetState();
-
 	MeshIDs.Init({}, MAX_MESH_LOD_COUNT);
 	SurfaceIDs.SetNum(MAX_MESH_LOD_COUNT);
 
@@ -59,6 +48,20 @@ void UCustomizableObjectSkeletalMesh::InitMutableStreamingData(const TSharedRef<
 
 	const UModelResources& ModelResources = CustomizableObject.GetPrivate()->GetModelResourcesChecked();
 
+	{
+		// Init properties
+		check(InOperationData->LiveInstance.IsValid());
+		UpdateMeshContext = MakeShared<FMutableUpdateMeshContext>(InOperationData->LiveInstance.ToSharedRef());
+
+		// This must run in the mutable thread.
+		check(UCustomizableObjectSystem::IsCreated());
+		UpdateMeshContext->System = UCustomizableObjectSystem::GetInstance()->GetPrivate()->MutableSystem;
+		UpdateMeshContext->ModelStreamableBulkData = ModelStreamableBulkData;
+		
+		UpdateMeshContext->MeshIDs = MeshIDs;
+		UpdateMeshContext->Meshes.SetNum(MeshIDs.Num());
+	}
+	
 	for (const FSkinWeightProfileInfo& SkinWeightProfile : GetSkinWeightProfiles())
 	{
 		const FMutableSkinWeightProfileInfo* ProfileInfo = ModelResources.SkinWeightProfilesInfo.FindByPredicate(
@@ -88,7 +91,7 @@ bool UCustomizableObjectSkeletalMesh::StreamIn(int32 NewMipCount, bool bHighPrio
 			? FRenderAssetUpdate::TT_Async
 			: FRenderAssetUpdate::TT_Render;
 
-		PendingUpdate = new FCustomizableObjectMeshStreamIn(this, CreateResourcesThread, ModelStreamableBulkData);
+		PendingUpdate = new FCustomizableObjectMeshStreamIn(this, CreateResourcesThread);
 
 		return !PendingUpdate->IsCancelled();
 	}
diff --git a/Engine/Plugins/Mutable/Source/CustomizableObject/Private/MuCO/CustomizableObjectSystem.cpp b/Engine/Plugins/Mutable/Source/CustomizableObject/Private/MuCO/CustomizableObjectSystem.cpp
--- a/Engine/Plugins/Mutable/Source/CustomizableObject/Private/MuCO/CustomizableObjectSystem.cpp
+++ b/Engine/Plugins/Mutable/Source/CustomizableObject/Private/MuCO/CustomizableObjectSystem.cpp
@@ -163,6 +163,12 @@ TAutoConsoleVariable<bool> CVarMutableHighPriorityLoading(
 	true,
 	TEXT("If enabled, the request to load additional assets will have high priority."));
 
+TAutoConsoleVariable<bool> CVarReuseLiveInstanceCacheBetweenUpdates (
+	TEXT("mutable.ReuseLiveInstanceCacheBetweenUpdates"),
+	false,
+	TEXT("If set to true, the program cache will be kept untouched between calls to the BeginUpdate_mutableThread. If set to false, then the LiveInstance program cache will be fully cleared between calls to the mentioned method."),	
+	ECVF_Default
+);
 
 int32 UCustomizableObjectSystemPrivate::SkeletalMeshMinLodQualityLevel = -1;
 
@@ -2272,11 +2278,21 @@ namespace impl
 			OperationData->MutableSystem->SetImagePixelConversionOverride( OperationData->PixelFormatOverride );
 		}
 		
+		// Prevent the usage of the already cached elements by fully clearing the cache.
+		if (!CVarReuseLiveInstanceCacheBetweenUpdates.GetValueOnAnyThread())
+		{
+			OperationData->LiveInstance->Cache->Clear();
+		}
+		
 		UE::Mutable::Private::TManagedPtr<const UE::Mutable::Private::FInstance> Instance = OperationData->MutableSystem->BeginUpdate_MutableThread(OperationData->LiveInstance.ToSharedRef(),
-			OperationData->Parameters,
-			OperationData->GetCapturedDescriptor().GetState(),
+			OperationData->OldParameters,
+			OperationData->OldState,
 			UE::Mutable::Private::FSystem::AllLODs);
 
+		// This data is no longer useful
+		OperationData->OldParameters = nullptr;
+		OperationData->OldState = -1;
+		
 		OperationData->MutableInstance = Instance;
 		
 		const bool bForceGeometryGeneration = CVarForceGeometryOnFirstGeneration.GetValueOnAnyThread();
@@ -3014,21 +3030,21 @@ namespace impl
 #endif
 			)
 		{
-			Operation->bNeverStream = true;
+			Operation->bNeverStreamMips = true;
 		}
 		else
 		{
-			Operation->bNeverStream = StateData ? StateData->bDisableTextureStreaming : false;
+			Operation->bNeverStreamMips = StateData ? StateData->bDisableTextureStreaming : false;
 		}
 #else
-		Operation->bNeverStream = true;
+		Operation->bNeverStreamMips = true;
 #endif
 	
-		Operation->MipsToSkip = !Operation->bNeverStream ? 255 : 0; // // 0 means generate all mips. 255 means skip all possible mips until only UTexture::GetStaticMinTextureResidentMipCount() are left
+		Operation->MipsToSkip = !Operation->bNeverStreamMips ? 255 : 0; // // 0 means generate all mips. 255 means skip all possible mips until only UTexture::GetStaticMinTextureResidentMipCount() are left
 		
 		Operation->bUseCaches = Operation->bKeepOwnershipOfGeneratedResources && !Operation->bBake && !FLogBenchmarkUtil::IsBenchmarkingReportingEnabled();
 		Operation->bUseMeshCache = Operation->bUseCaches && Object->bEnableMeshCache && !Operation->bLiveUpdateMode && CVarEnableMeshCache.GetValueOnAnyThread(true);
-		Operation->bUseResueTextureCache = Operation->bUseCaches && Operation->bNeverStream && CVarEnableMutableReuseInstanceTextures.GetValueOnGameThread();
+		Operation->bUseResueTextureCache = Operation->bUseCaches && Operation->bNeverStreamMips && CVarEnableMutableReuseInstanceTextures.GetValueOnGameThread();
 		
 		Operation->Model = ObjectPrivate->GetModel().ToSharedRef();
 
@@ -3156,37 +3172,47 @@ namespace impl
 
 			Operation->bLiveUpdateMode = SystemPrivate->EnableMutableLiveUpdate && (StateData ? StateData->bLiveUpdateMode : false);
 		
-			if (Operation->bLiveUpdateMode && (!Operation->bNeverStream || Operation->MipsToSkip > 0))
+			// todo : make the stream settings in Operation respond to the settings in Operation->bLiveUpdateMode
+			if (Operation->bLiveUpdateMode && 
+				(Operation->bStreamMeshLODs || !Operation->bNeverStreamMips || Operation->MipsToSkip > 0))
 			{
-				UE_LOG(LogMutable, Warning, TEXT("Instance LiveUpdateMode does not yet support progressive streaming of Mutable textures. Disabling LiveUpdateMode for this update."));
+				UE_LOG(LogMutable, Warning, TEXT("Instance LiveUpdateMode does not yet support progressive streaming of Mutable textures or meshes. Disabling LiveUpdateMode for this update."));
 				Operation->bLiveUpdateMode = false;
 			}
+
+			Operation->LiveInstance = MutableSystem->NewInstance(
+				Operation->Model,
+				ObjectPrivate->MeshIdRegistry,
+				ObjectPrivate->ImageIdRegistry,
+				ObjectPrivate->MaterialIdRegistry,
+				ObjectPrivate->SkeletalMeshIdRegistry,
+				Operation->ExternalResourceProvider,
+				Operation->PassthroughObjectLoader);
 			
-			Operation->LiveInstance = Operation->bLiveUpdateMode ? Instance->GetPrivate()->LiveInstance : nullptr;
+			Operation->LiveInstance->Parameters = Operation->Parameters;
+			Operation->LiveInstance->State = Operation->GetCapturedDescriptor().GetState();
 			
-			if (!Operation->LiveInstance)
+			if (Operation->bLiveUpdateMode)
 			{
-				Operation->LiveInstance = MutableSystem->NewInstance(
-					Operation->Model, 
-					ObjectPrivate->MeshIdRegistry,
-					ObjectPrivate->ImageIdRegistry,
-					ObjectPrivate->MaterialIdRegistry,
-					ObjectPrivate->SkeletalMeshIdRegistry,
-					Operation->ExternalResourceProvider, 
-					Operation->PassthroughObjectLoader);
+				if (Instance->GetPrivate()->LiveInstance.IsValid())
+				{
+					// Reuse the cache of the live instance that was added during a previous update
+					Operation->LiveInstance->Cache = Instance->GetPrivate()->LiveInstance->Cache;
+				
+					// Keep the old params and state to be able to later determine if the reused cache requires clearing before being used
+					Operation->OldParameters = Instance->GetPrivate()->LiveInstance->Parameters;
+					Operation->OldState = Instance->GetPrivate()->LiveInstance->State;
+				}
+				
+				Instance->GetPrivate()->LiveInstance = Operation->LiveInstance;
 			}
 			else
 			{
-				Operation->LiveInstance->ExternalResourceProvider = Operation->ExternalResourceProvider;
-				Operation->LiveInstance->PassthroughObjectLoader = Operation->PassthroughObjectLoader;
+				Instance->GetPrivate()->LiveInstance = nullptr;
 			}
-
-			Instance->GetPrivate()->LiveInstance = Operation->bLiveUpdateMode ? Operation->LiveInstance : nullptr;
 			
 			Operation->MutableInstance = MutableSystem->BeginUpdate_GameThread(
 				Operation->LiveInstance.ToSharedRef(),
-				Operation->Parameters,
-				Operation->GetCapturedDescriptor().GetState(),
 				UE::Mutable::Private::FSystem::AllLODs);
 		}
 		
@@ -4248,7 +4274,7 @@ void LoadClothing(FMutableStreamRequest& MutableDataStreamer, const UE::Mutable:
 			}
 			else
 			{
-				UE_LOG(LogMutable, Error, TEXT("Invalid streamed clothing data block [%d] found."), TypedResourceId.Id);
+				UE_LOG(LogMutable, Error, TEXT("Invalid streamed clothing data block [%llu] found."), TypedResourceId.Id);
 			}
 		}
 	}
diff --git a/Engine/Plugins/Mutable/Source/CustomizableObject/Public/MuCO/CustomizableObjectSkeletalMesh.h b/Engine/Plugins/Mutable/Source/CustomizableObject/Public/MuCO/CustomizableObjectSkeletalMesh.h
--- a/Engine/Plugins/Mutable/Source/CustomizableObject/Public/MuCO/CustomizableObjectSkeletalMesh.h
+++ b/Engine/Plugins/Mutable/Source/CustomizableObject/Public/MuCO/CustomizableObjectSkeletalMesh.h
@@ -20,6 +20,7 @@ class UCustomizableInstancePrivate;
 class FUpdateContextPrivate;
 class FCustomizableObjectMeshStreamIn;
 struct FModelStreamableBulkData;
+struct FMutableUpdateMeshContext;
 class FUnrealMutableResourceProvider;
 
 
@@ -46,17 +47,8 @@ class UCustomizableObjectSkeletalMesh : public USkeletalMesh
 
 	TSharedPtr<FModelStreamableBulkData> ModelStreamableBulkData;
 
-	TSharedPtr<UE::Mutable::Private::FModel> Model;
-	TSharedPtr<UE::Mutable::Private::FMeshIdRegistry> MeshIdRegistry;
-	TSharedPtr<UE::Mutable::Private::FImageIdRegistry> ImageIdRegistry;
-	TSharedPtr<UE::Mutable::Private::FMaterialIdRegistry> MaterialIdRegistry;
-	TSharedPtr<UE::Mutable::Private::FSkeletalMeshIdRegistry> SkeletalMeshIdRegistry;
+	TSharedPtr<FMutableUpdateMeshContext> UpdateMeshContext;
 
-	TSharedPtr<FUnrealMutableResourceProvider> ExternalResourceProvider;
-
-	/** Instance parameters at the time of creation. */
-	TSharedPtr<UE::Mutable::Private::FParameters> Parameters;
-	int32 State = -1;
 
 	/** Mesh IDs per LOD */
 	TArray<UE::Mutable::Private::FMeshId> MeshIDs;
diff --git a/Engine/Plugins/Mutable/Source/CustomizableObject/Public/MuCO/CustomizableObjectSystem.h b/Engine/Plugins/Mutable/Source/CustomizableObject/Public/MuCO/CustomizableObjectSystem.h
--- a/Engine/Plugins/Mutable/Source/CustomizableObject/Public/MuCO/CustomizableObjectSystem.h
+++ b/Engine/Plugins/Mutable/Source/CustomizableObject/Public/MuCO/CustomizableObjectSystem.h
@@ -43,6 +43,8 @@ extern TAutoConsoleVariable<bool> CVarEnableReleaseMeshResources;
 
 extern TAutoConsoleVariable<bool> CVarFixLowPriorityTasksOverlap;
 
+extern TAutoConsoleVariable<bool> CVarReuseLiveInstanceCacheBetweenUpdates;
+
 
 UENUM(BlueprintType)
 enum class ECustomizableObjectOptimizationLevel : uint8
diff --git a/Engine/Plugins/Mutable/Source/CustomizableObjectEditor/Private/MuCOE/SMutableCodeViewer.cpp b/Engine/Plugins/Mutable/Source/CustomizableObjectEditor/Private/MuCOE/SMutableCodeViewer.cpp
--- a/Engine/Plugins/Mutable/Source/CustomizableObjectEditor/Private/MuCOE/SMutableCodeViewer.cpp
+++ b/Engine/Plugins/Mutable/Source/CustomizableObjectEditor/Private/MuCOE/SMutableCodeViewer.cpp
@@ -2897,8 +2897,7 @@ void SMutableCodeViewer::Tick(const FGeometry& AllottedGeometry, const double In
 	UE::Mutable::Private::FSettings Settings;
 	const TSharedPtr<UE::Mutable::Private::FSystem> System = MakeShared<UE::Mutable::Private::FSystem>(Settings);
 
-	TSharedRef<UE::Mutable::Private::FLiveInstance> TempLiveInstance = System->NewBuildInstance(MutableModel, ExternalResourceProvider);
-	TempLiveInstance->OldParameters = PreviewParameters;
+	TSharedRef<UE::Mutable::Private::FLiveInstance> TempLiveInstance = System->NewBuildInstance(MutableModel, PreviewParameters, ExternalResourceProvider);
 	
 	switch (OperationDataType)
 	{
@@ -3008,8 +3007,7 @@ UE::Mutable::Private::TManagedPtr<const UE::Mutable::Private::FImage> SMutableCo
 	const TSharedPtr<UE::Mutable::Private::FSystem> System = MakeShared<UE::Mutable::Private::FSystem>(Settings);
 	
 	
-	TSharedRef<UE::Mutable::Private::FLiveInstance> TempLiveInstance = System->NewBuildInstance(MutableModel, ExternalResourceProvider);
-	TempLiveInstance->OldParameters = PreviewParameters;
+	TSharedRef<UE::Mutable::Private::FLiveInstance> TempLiveInstance = System->NewBuildInstance(MutableModel, PreviewParameters, ExternalResourceProvider);
 	
 	UE::Mutable::Private::TManagedPtr<const UE::Mutable::Private::FImage> MutableImage = System->BuildImage(TempLiveInstance, TargetOpAddress, 0);	
 	return MutableImage;
diff --git a/Engine/Plugins/Mutable/Source/MutableRuntime/Internal/MuR/System.h b/Engine/Plugins/Mutable/Source/MutableRuntime/Internal/MuR/System.h
--- a/Engine/Plugins/Mutable/Source/MutableRuntime/Internal/MuR/System.h
+++ b/Engine/Plugins/Mutable/Source/MutableRuntime/Internal/MuR/System.h
@@ -614,6 +614,25 @@ namespace UE::Mutable::Private
 			uint32 CodeSize = OpExecutionData.size_code();
 			OpExecutionData.clear();
 			OpExecutionData.resize(CodeSize);
+			
+			LayoutResults.Empty();
+			InstanceResults.Empty();
+			StringResults.Empty();
+			ExtensionDataResults.Empty();
+			InstancedStructResults.Empty();
+			MaterialResults.Empty();
+			SkeletalMeshResults.Empty();
+
+			ImageResults.Empty();
+			MeshResults.Empty();
+		
+			/** Storage for strong references. Size of Strong Results is the same as the weak. */
+			StrongImageResults.Empty();
+			StrongMeshResults.Empty();
+
+			ColorResults.Empty();
+			ProjectorResults.Empty();
+			MatrixResults.Empty();
 		}
 
 
@@ -629,6 +648,122 @@ namespace UE::Mutable::Private
 		}
 
 
+		/** 
+		 * Remove all intermediate data (big and small) from the memory except for the one that has been explicitly
+		 * marked as state cache.
+		 */
+		void ClearLayer0()
+		{
+			MUTABLE_CPUPROFILER_SCOPE(ClearLayer0);
+
+			CodeContainer<FProgramCache::FOpExecutionData>::iterator it = OpExecutionData.begin();
+			for (; it.IsValid(); ++it)
+			{
+				FProgramCache::FOpExecutionData& Data = *it;
+				if (!Data.DataTypeIndex || Data.IsCacheLocked)
+				{
+					continue;
+				}
+
+				switch (Data.DataType)
+				{
+				case EDataType::Image:
+				{
+					check(ImageResults.Num() && StrongImageResults.Num());
+					ImageResults[Data.DataTypeIndex] = nullptr;
+					StrongImageResults[Data.DataTypeIndex] = nullptr;
+					break;
+				}
+
+				case EDataType::Mesh:
+				{
+					check(MeshResults.Num() && StrongMeshResults.Num());
+					MeshResults[Data.DataTypeIndex] = nullptr;
+					StrongMeshResults[Data.DataTypeIndex] = nullptr;
+					break;
+				}
+
+				case EDataType::SkeletalMesh:
+				{
+					check(SkeletalMeshResults.Num());
+					SkeletalMeshResults[Data.DataTypeIndex] = nullptr;
+					break;
+				}
+
+				case EDataType::Layout:
+				{
+					check(LayoutResults.Num());
+					LayoutResults[Data.DataTypeIndex] = nullptr;
+					break;
+				}
+
+				case EDataType::Instance:
+				{
+					check(InstanceResults.Num());
+					InstanceResults[Data.DataTypeIndex] = nullptr;
+					break;
+				}
+
+				default:
+					break;
+				}
+
+				Data.OpHitCount = 0;
+				Data.IsValueValid = false;
+			}
+		}
+
+
+		/** Remove all intermediate data (big and small) from the memory including the one that has been explicitly
+		* marked as state cache.
+		*/
+		void ClearLayer1()
+		{
+			MUTABLE_CPUPROFILER_SCOPE(ClearLayer1);
+
+			CodeContainer<FProgramCache::FOpExecutionData>::iterator it = OpExecutionData.begin();
+			for (; it.IsValid(); ++it)
+			{
+				FProgramCache::FOpExecutionData& Data = *it;
+				
+				Data.OpHitCount = 0;
+				Data.IsValueValid = false;
+			}
+
+			for (TManagedWeakPtr<const FImage>& Value : ImageResults)
+			{
+				Value = nullptr;
+			}
+
+			for (TManagedPtr<const FImage>& Value : StrongImageResults)
+			{
+				Value = nullptr;
+			}
+			
+			
+			for (TManagedWeakPtr<const FSkeletalMesh>& Value : SkeletalMeshResults)
+			{
+				Value = nullptr;
+			}
+
+			for (TManagedPtr<const FMesh>& Value : StrongMeshResults)
+			{
+				Value = nullptr;
+			}
+			
+
+			for (TManagedWeakPtr<const FLayout>& Value : LayoutResults)
+			{
+				Value = nullptr;
+			}
+
+			for (TManagedWeakPtr<const FInstance>& Value : InstanceResults)
+			{
+				Value = nullptr;
+			}
+		}
+		
+		
 		bool GetBool(FCacheAddress Address)
 		{
 			if (!Address.At) 
@@ -1487,13 +1622,8 @@ namespace UE::Mutable::Private
 		TManagedPtr<const FInstance> Instance;
 		TSharedPtr<const FModel> Model;
 
-		TSharedPtr<const FParameters> OldParameters;
+		TSharedPtr<const FParameters> Parameters;
 		
-		/** Mask of the parameters that have changed since the last update.
-		* Every bit represents a state parameter.
-		*/
-		uint64 UpdatedParameters = 0;
-
 		/** Cached data for the generation of this instance. */
 		TSharedRef<FProgramCache> Cache = MakeShared<FProgramCache>();
 
@@ -1516,7 +1646,7 @@ namespace UE::Mutable::Private
 		using TMemoryTrackedArray = TArray<Type, FDefaultMemoryTrackingAllocator<MemoryCounter>>;
 
 		template<class KeyType, class ValueType>
-		using TMemoryTrackedMap = TMap<KeyType, ValueType, FDefaultMemoryTrackingSetAllocator<MemoryCounter>>;    	
+		using TMemoryTrackedMap = TMap<KeyType, ValueType, FDefaultMemoryTrackingSetAllocator<MemoryCounter>>;
 
 		/** Maximum working memory that mutable should be using. */
 		int64 BudgetBytes = 0;
@@ -1607,116 +1737,6 @@ namespace UE::Mutable::Private
     		return Result;
     	}
 
-		/** 
-		 * Remove all intermediate data (big and small) from the memory except for the one that has been explicitely
-		 * marked as state cache.
-		 */
-		void ClearCacheLayer0(const TSharedRef<FProgramCache>& InInstanceCache)
-		{
-			MUTABLE_CPUPROFILER_SCOPE(ClearLayer0);
-
-			CodeContainer<FProgramCache::FOpExecutionData>::iterator it = InInstanceCache->OpExecutionData.begin();
-			for (; it.IsValid(); ++it)
-			{
-				FProgramCache::FOpExecutionData& Data = *it;
-				if (!Data.DataTypeIndex || Data.IsCacheLocked)
-				{
-					continue;
-				}
-
-				switch (Data.DataType)
-				{
-				case EDataType::Image:
-				{
-					InInstanceCache->ImageResults[Data.DataTypeIndex] = nullptr;
-					InInstanceCache->StrongImageResults[Data.DataTypeIndex] = nullptr;
-					break;
-				}
-
-				case EDataType::Mesh:
-				{
-					InInstanceCache->MeshResults[Data.DataTypeIndex] = nullptr;
-					InInstanceCache->StrongMeshResults[Data.DataTypeIndex] = nullptr;
-					break;
-				}
-
-				case EDataType::SkeletalMesh:
-				{
-					InInstanceCache->SkeletalMeshResults[Data.DataTypeIndex] = nullptr;
-					break;
-				}
-
-				case EDataType::Layout:
-				{
-					InInstanceCache->LayoutResults[Data.DataTypeIndex] = nullptr;
-					break;
-				}
-
-				case EDataType::Instance:
-				{
-					InInstanceCache->InstanceResults[Data.DataTypeIndex] = nullptr;
-					break;
-				}
-
-				default:
-					break;
-				}
-
-				Data.OpHitCount = 0;
-				Data.IsValueValid = false;
-			}
-		}
-
-
-		/** Remove all intermediate data (big and small) from the memory including the one that has been explicitely
-		* marked as state cache.
-		*/
-		void ClearCacheLayer1(const TSharedRef<FProgramCache>& InInstanceCache)
-		{
-			MUTABLE_CPUPROFILER_SCOPE(ClearLayer1);
-
-			CodeContainer<FProgramCache::FOpExecutionData>::iterator it = InInstanceCache->OpExecutionData.begin();
-			for (; it.IsValid(); ++it)
-			{
-				FProgramCache::FOpExecutionData& Data = *it;
-				
-				Data.OpHitCount = 0;
-				Data.IsValueValid = false;
-			}
-
-			for (TManagedWeakPtr<const FImage>& Value : InInstanceCache->ImageResults)
-			{
-				Value = nullptr;
-			}
-
-			for (TManagedPtr<const FImage>& Value : InInstanceCache->StrongImageResults)
-			{
-				Value = nullptr;
-			}
-
-
-			for (TManagedWeakPtr<const FSkeletalMesh>& Value : InInstanceCache->SkeletalMeshResults)
-			{
-				Value = nullptr;
-			}
-
-			for (TManagedPtr<const FMesh>& Value : InInstanceCache->StrongMeshResults)
-			{
-				Value = nullptr;
-			}
-
-
-			for (TManagedWeakPtr<const FLayout>& Value : InInstanceCache->LayoutResults)
-			{
-				Value = nullptr;
-			}
-
-			for (TManagedWeakPtr<const FInstance>& Value : InInstanceCache->InstanceResults)
-			{
-				Value = nullptr;
-			}
-		}
-
 
 		void LogWorkingMemory(const class CodeRunner* CurrentRunner) const;
 
@@ -1846,6 +1866,7 @@ namespace UE::Mutable::Private
 		 * It does not require further initialization as it is fully capable of operating as is.
 		 */
 		UE_API TSharedRef<FLiveInstance> NewBuildInstance(const TSharedPtr<const FModel>& InModel, 
+			const TSharedPtr<const FParameters>& InParameters,
 			const TSharedPtr<FExternalResourceProvider>& ExternalResourceProvider);
 				
         //! \brief Update an instance with a new parameter set and/or state.
@@ -1862,15 +1883,12 @@ namespace UE::Mutable::Private
         //! the same instanceID parameter.
 		//!
 		//! Very fast. Can be called in the Game Thread.
-		UE_API TManagedPtr<const FInstance> BeginUpdate_GameThread(const TSharedRef<FLiveInstance>& LiveInstance,
-			const TSharedPtr<const FParameters>& InParams,
-			int32 StateIndex,
-			uint32 LodMask);
+		UE_API TManagedPtr<const FInstance> BeginUpdate_GameThread(const TSharedRef<FLiveInstance>& LiveInstance, uint32 LodMask);
 
 		/** Potentially expensive. Must be called in the Mutable Thread. */
 		UE_API TManagedPtr<const FInstance> BeginUpdate_MutableThread(const TSharedRef<FLiveInstance>& LiveInstance,
-			const TSharedPtr<const FParameters>& InParams,
-			int32 InStateIndex,
+			const TSharedPtr<const FParameters>& InOldParameters,
+			int32 InOldState,
 			uint32 InLodMask);
 		
 		//! Only valid between BeginUpdate and EndUpdate
@@ -1958,7 +1976,7 @@ namespace UE::Mutable::Private
 		/** Increased each time a rom is loaded. Used to keep track of which roms to evict. */
 		static uint64 RomTick;
 		
-		bool CheckUpdatedParameters(const TSharedRef<FLiveInstance>& LiveInstance, const TSharedPtr<const FParameters>&, uint64& OutUpdatedParameters );
+		bool CheckUpdatedParameters(const TSharedRef<FLiveInstance>& LiveInstance, const TSharedPtr<const FParameters>&);
 		
 		void RunCode(
 			const TSharedRef<FLiveInstance>& InLiveInstance,
diff --git a/Engine/Plugins/Mutable/Source/MutableRuntime/Private/MuR/CodeRunner.cpp b/Engine/Plugins/Mutable/Source/MutableRuntime/Private/MuR/CodeRunner.cpp
--- a/Engine/Plugins/Mutable/Source/MutableRuntime/Private/MuR/CodeRunner.cpp
+++ b/Engine/Plugins/Mutable/Source/MutableRuntime/Private/MuR/CodeRunner.cpp
@@ -192,7 +192,7 @@ namespace UE::Mutable::Private
 		, System(InSystem)
 		, Model(*InLiveInstance->Model.Get())
 		, Program(InLiveInstance->Model->Program)
-		, Parameters(*InLiveInstance->OldParameters)
+		, Parameters(*InLiveInstance->Parameters)
 		, LODMask(InLodMask)
 		, ExternalResourceProvider(InLiveInstance->ExternalResourceProvider)
 	{
diff --git a/Engine/Plugins/Mutable/Source/MutableRuntime/Private/MuR/CodeRunnerBegin.cpp b/Engine/Plugins/Mutable/Source/MutableRuntime/Private/MuR/CodeRunnerBegin.cpp
--- a/Engine/Plugins/Mutable/Source/MutableRuntime/Private/MuR/CodeRunnerBegin.cpp
+++ b/Engine/Plugins/Mutable/Source/MutableRuntime/Private/MuR/CodeRunnerBegin.cpp
@@ -14,10 +14,9 @@ namespace UE::Mutable::Private
 {
 	CodeRunnerBegin::CodeRunnerBegin(
 			const TSharedRef<FLiveInstance>& InLiveInstance,
-			const TSharedPtr<const FParameters>& InParams,
 			uint32 InLodMask) :
 		Model(InLiveInstance->Model),
-		Params(InParams),
+		Params(InLiveInstance->Parameters),
 		Program(InLiveInstance->Model->Program),
 		LodMask(InLodMask),
 		Executed(InLiveInstance->Model->Program.OpAddress.Num()),
diff --git a/Engine/Plugins/Mutable/Source/MutableRuntime/Private/MuR/CodeRunnerBegin.h b/Engine/Plugins/Mutable/Source/MutableRuntime/Private/MuR/CodeRunnerBegin.h
--- a/Engine/Plugins/Mutable/Source/MutableRuntime/Private/MuR/CodeRunnerBegin.h
+++ b/Engine/Plugins/Mutable/Source/MutableRuntime/Private/MuR/CodeRunnerBegin.h
@@ -111,7 +111,6 @@ namespace UE::Mutable::Private
 	public:
 		MUTABLERUNTIME_API CodeRunnerBegin(
 			const TSharedRef<FLiveInstance>& InLiveInstance,
-			const TSharedPtr<const FParameters>& InParams,
 			uint32 InLodMask);
 
 		TManagedPtr<const FInstance> RunCode(const FScheduledOpInline& Root);
diff --git a/Engine/Plugins/Mutable/Source/MutableRuntime/Private/MuR/Model.cpp b/Engine/Plugins/Mutable/Source/MutableRuntime/Private/MuR/Model.cpp
--- a/Engine/Plugins/Mutable/Source/MutableRuntime/Private/MuR/Model.cpp
+++ b/Engine/Plugins/Mutable/Source/MutableRuntime/Private/MuR/Model.cpp
@@ -944,7 +944,7 @@ namespace UE::Mutable::Private
                         case EParameterType::Projector:
                         {
 //							float m[16];
-//							OldParameters->GetProjectorValue( p, m );
+//							Parameters->GetProjectorValue( p, m );
                             pRes->Values[thisP].Set<FParamProjectorType>(OldParameters->Values[p].Get<FParamProjectorType>());
                             break;
                         }
diff --git a/Engine/Plugins/Mutable/Source/MutableRuntime/Private/MuR/System.cpp b/Engine/Plugins/Mutable/Source/MutableRuntime/Private/MuR/System.cpp
--- a/Engine/Plugins/Mutable/Source/MutableRuntime/Private/MuR/System.cpp
+++ b/Engine/Plugins/Mutable/Source/MutableRuntime/Private/MuR/System.cpp
@@ -162,6 +162,7 @@ namespace UE::Mutable::Private
 
 	
     TSharedRef<FLiveInstance> FSystem::NewBuildInstance(const TSharedPtr<const FModel>& InModel,
+		const TSharedPtr<const FParameters>& InParameters,
 	    const TSharedPtr<FExternalResourceProvider>& ExternalResourceProvider)
     {
     	MUTABLE_CPUPROFILER_SCOPE(NewBuildInstance);
@@ -171,6 +172,7 @@ namespace UE::Mutable::Private
     	
     	TSharedRef<FLiveInstance> InstanceData = MakeShared<FLiveInstance>();
     	InstanceData->Instance = nullptr;
+    	InstanceData->Parameters = InParameters;
     	InstanceData->Model = InModel;
     	InstanceData->State = -1;
     	InstanceData->ExternalResourceProvider = ExternalResourceProvider;
@@ -187,19 +189,17 @@ namespace UE::Mutable::Private
 
 
     //---------------------------------------------------------------------------------------------
-    TManagedPtr<const FInstance> FSystem::BeginUpdate_GameThread(const TSharedRef<FLiveInstance>& LiveInstance,
-                                     const TSharedPtr<const FParameters>& InParams,
-									 int32 InStateIndex,
-                                     uint32 InLodMask)
+    TManagedPtr<const FInstance> FSystem::BeginUpdate_GameThread(const TSharedRef<FLiveInstance>& LiveInstance,uint32 InLodMask)
     {
 		LLM_SCOPE_BYNAME(TEXT("MutableRuntime"));
 		MUTABLE_CPUPROFILER_SCOPE(BeginUpdate_GameThread);
 
-		check(InParams);
+		check(LiveInstance->Parameters);
 
-    	CodeRunnerBegin Runner(LiveInstance, InParams, InLodMask);
+    	CodeRunnerBegin Runner(LiveInstance, InLodMask);
 
-    	FScheduledOpInline Root(LiveInstance->Model->Program.States[InStateIndex].Root, 0, 0, true);
+    	check(LiveInstance->State >= 0);
+    	FScheduledOpInline Root(LiveInstance->Model->Program.States[LiveInstance->State].Root, 0, 0, true);
     	TManagedPtr<const FInstance> Result = Runner.RunCode(FScheduledOpInline(Root, 0));
     	if (!Result)
     	{
@@ -212,15 +212,15 @@ namespace UE::Mutable::Private
 
 	
     TManagedPtr<const FInstance> FSystem::BeginUpdate_MutableThread(const TSharedRef<FLiveInstance>& LiveInstance,
-		const TSharedPtr<const FParameters>& InParams,
-		int32 InStateIndex,
+		const TSharedPtr<const FParameters>& InOldParameters,
+		int32 InOldState,
 		uint32 InLodMask)
     {
 		LLM_SCOPE_BYNAME(TEXT("MutableRuntime"));
 		MUTABLE_CPUPROFILER_SCOPE(BeginUpdate_MutableThread);
 		TRACE_COUNTER_INCREMENT(MutableRuntime_Updates);
 
-		if (!InParams)
+		if (!LiveInstance->Parameters)
 		{
 			UE_LOG(LogMutableCore, Error, TEXT("Invalid parameters in mutable update."));
 			return nullptr;
@@ -228,7 +228,9 @@ namespace UE::Mutable::Private
 
 		const FProgram& Program = LiveInstance->Model->Program;
 
-		bool bValidState = InStateIndex >= 0 && InStateIndex < (int32)Program.States.Num();
+    	const int32 NewState = LiveInstance->State;
+    	
+		bool bValidState = NewState >= 0 && NewState < (int32)Program.States.Num();
 		if (!bValidState)
 		{
 			UE_LOG(LogMutableCore, Error, TEXT("Invalid state in mutable update."));
@@ -237,31 +239,24 @@ namespace UE::Mutable::Private
 
 		// This may free resources that allow us to use less memory.
 		LiveInstance->Instance = nullptr;
-
-		bool bFullBuild = (InStateIndex != LiveInstance->State);
-
-		LiveInstance->State = InStateIndex;
-
-		// If we changed parameters that are not in this state, we need to rebuild all.
-		if (!bFullBuild)
-		{
-			bFullBuild = CheckUpdatedParameters(LiveInstance, InParams, LiveInstance->UpdatedParameters);
-		}
-
-		// Remove cached data
-		WorkingMemoryManager.ClearCacheLayer0(LiveInstance->Cache);
-		if (bFullBuild)
+    	
+    	// Remove cached data
 		{
-			WorkingMemoryManager.ClearCacheLayer1(LiveInstance->Cache);
+			LiveInstance->Cache->ClearLayer0();
+    	
+			// If we changed the state or the parameters that are not in this state, we need to rebuild all.
+			if (InOldState != NewState || CheckUpdatedParameters(LiveInstance, InOldParameters))
+			{
+				LiveInstance->Cache->ClearLayer1();
+			}
 		}
 
 		WorkingMemoryManager.BeginRunnerThread();
     	
 		// Prepare instance cache
-    	LiveInstance->Cache->Prepare(Program, InStateIndex);
-		LiveInstance->OldParameters = InParams->Clone();
+    	LiveInstance->Cache->Prepare(Program, NewState);
     	
-		OP::ADDRESS RootAt = Program.States[InStateIndex].Root;
+		OP::ADDRESS RootAt = Program.States[NewState].Root;
 		RunCode(LiveInstance, RootAt, InLodMask);
 
 		TManagedPtr<const FInstance> Result;
@@ -829,11 +824,11 @@ namespace UE::Mutable::Private
 
 		// We don't want to clear the cache layer 1 because it contains data that can be useful for a 
 		// future update (same states, just runtime parameters changed).
-		//WorkingMemoryManager.ClearCacheLayer1();
+		//LiveInstance->Cache->ClearLayer1();
 
 		// We need to clear the layer 0 cache, because it contains data that is only valid for the current 
 		// parameter values (unless it is data marked as state cache)
-		WorkingMemoryManager.ClearCacheLayer0(LiveInstance->Cache);
+		LiveInstance->Cache->ClearLayer0();
 		
 		// Reduce the cache until it fits the limit.
 		// Not providing a cache at this point as there is no need of keeping the cache active anymore
@@ -910,45 +905,34 @@ namespace UE::Mutable::Private
     }
 
 	
-	bool FSystem::CheckUpdatedParameters(const TSharedRef<FLiveInstance>& LiveInstance,
-                                 const TSharedPtr<const FParameters>& Params,
-                                 uint64& UpdatedParameters)
+	bool FSystem::CheckUpdatedParameters(const TSharedRef<FLiveInstance>& LiveInstance, const TSharedPtr<const FParameters>& InOldParameters)
     {
         bool bFullBuild = false;
 
-		if (!LiveInstance->OldParameters)
+		if (!LiveInstance->Parameters || !InOldParameters)
 		{
-			UpdatedParameters = AllParametersMask;
 			return true;
 		}
 
+    	const TSharedPtr<const FParameters> NewParameters = LiveInstance->Parameters.ToSharedRef();
+
         // check what parameters have changed
-		UpdatedParameters = 0;
         const FProgram& Program = LiveInstance->Model->Program;
-        const TArray<int>& runtimeParams = Program.States[ LiveInstance->State ].m_runtimeParameters;
+        const TArray<int>& RuntimeParameters = Program.States[ LiveInstance->State ].m_runtimeParameters;
 
-        check( Params->GetCount() == (int)Program.Parameters.Num() );
-        check( !LiveInstance->OldParameters
-			||
-			Params->GetCount() == LiveInstance->OldParameters->GetCount() );
+        check( NewParameters->GetCount() == (int32)Program.Parameters.Num() );
+        check( NewParameters->GetCount() == NewParameters->GetCount() );
 
-        for ( int32 p=0; p<Program.Parameters.Num() && !bFullBuild; ++p )
+        for ( int32 ParameterIndex=0; ParameterIndex<Program.Parameters.Num() ; ++ParameterIndex )
         {
-            bool isRuntime = runtimeParams.Contains( p );
-            bool changed = !Params->HasSameValue( p, LiveInstance->OldParameters, p );
-
-            if (changed && isRuntime)
-            {
-				uint64 runtimeIndex = runtimeParams.IndexOfByKey(p);
-				UpdatedParameters |= uint64(1) << runtimeIndex;
-            }
-            else if (changed)
-            {
-                // A non-runtime parameter has changed, we need a full build.
-                // TODO: report, or log somehow.
-				bFullBuild = true;
-                UpdatedParameters = AllParametersMask;
-            }
+            const bool bIsRuntimeParameter = RuntimeParameters.Contains( ParameterIndex );
+            const bool bHasChanged = !NewParameters->HasSameValue( ParameterIndex, InOldParameters, ParameterIndex );
+
+        	if (bHasChanged && !bIsRuntimeParameter)
+        	{
+        		bFullBuild = true;
+        		break;
+        	}
         }
 
         return bFullBuild;
diff --git a/Engine/Plugins/Mutable/Source/MutableTools/Private/MuT/CodeOptimiser.cpp b/Engine/Plugins/Mutable/Source/MutableTools/Private/MuT/CodeOptimiser.cpp
--- a/Engine/Plugins/Mutable/Source/MutableTools/Private/MuT/CodeOptimiser.cpp
+++ b/Engine/Plugins/Mutable/Source/MutableTools/Private/MuT/CodeOptimiser.cpp
@@ -195,8 +195,7 @@ namespace UE::Mutable::Private
 
 			TSharedPtr<FParameters> LocalParams = FModel::NewParameters(Model);
 
-			TSharedRef<FLiveInstance> TempLiveInstance = System->NewBuildInstance(Model, nullptr);
-			TempLiveInstance->OldParameters = LocalParams;
+			TSharedRef<FLiveInstance> TempLiveInstance = System->NewBuildInstance(Model, LocalParams,  nullptr);
 			
 			// Calculate the value and replace this op by a constant
 			switch( DataType )
