[Mutable] Remove the option to use .mut data files in packaged builds and all related code.

#rnx

#Synced-CL 49944882

#Synced-CL 49944882

diff --git a/Engine/Plugins/Mutable/Source/CustomizableObject/Internal/MuCO/CustomizableObjectCompilerTypes.h b/Engine/Plugins/Mutable/Source/CustomizableObject/Internal/MuCO/CustomizableObjectCompilerTypes.h
--- a/Engine/Plugins/Mutable/Source/CustomizableObject/Internal/MuCO/CustomizableObjectCompilerTypes.h
+++ b/Engine/Plugins/Mutable/Source/CustomizableObject/Internal/MuCO/CustomizableObjectCompilerTypes.h
@@ -25,13 +25,6 @@ struct FCompilationOptions
 	// From 0 to UE_MUTABLE_MAX_OPTIMIZATION
 	int32 OptimizationLevel = UE_MUTABLE_MAX_OPTIMIZATION;
 
-	/** High limit of the size in bytes of the packaged data when cooking this object.
-	 * This limit is before any pak or filesystem compression. This limit will be broken if a single piece of data is bigger because data is not fragmented for packaging purposes.
-	 *
-	 * This option does not modify the DDC Key.
-	 */
-	uint64 PackagedDataBytesLimit = 256 * 1024 * 1024;
-
 	/** High (inclusive) limit of the size in bytes of a data block to be included into the compiled object directly instead of stored in a streamable file.
 	 *
 	 * This option does not modify the DDC Key.
@@ -97,12 +90,6 @@ struct FCompilationOptions
 	
 	/** Stores the only option of an Int Param that should be compiled. */
 	TMap<FString, FString> ParamNamesToSelectedOptions;
-
-	/** Used to know if the object will be using FByteBulkData.
-	 *
-	 * This option does not modify the DDC Key. 
-	 */
-	bool bUseBulkData = false;
 };
 
 
diff --git a/Engine/Plugins/Mutable/Source/CustomizableObject/Internal/MuCO/CustomizableObjectPrivate.h b/Engine/Plugins/Mutable/Source/CustomizableObject/Internal/MuCO/CustomizableObjectPrivate.h
--- a/Engine/Plugins/Mutable/Source/CustomizableObject/Internal/MuCO/CustomizableObjectPrivate.h
+++ b/Engine/Plugins/Mutable/Source/CustomizableObject/Internal/MuCO/CustomizableObjectPrivate.h
@@ -1115,15 +1115,6 @@ namespace UE::Mutable::Private
 		const ITargetPlatform& TargetPlatform,
 		TArray<FFile>& OutBulkDataFiles);
 
-	/** Generate the list of BulkData files with a soft restriction to the size of the files.
-	 */
-	void CUSTOMIZABLEOBJECT_API GenerateBulkDataFilesListWithSizeLimit(
-		TSharedPtr<const UE::Mutable::Private::FModel, ESPMode::ThreadSafe> Model,
-		FModelStreamableBulkData& ModelStreamableBulkData,
-		const ITargetPlatform* TargetPlatform,
-		uint64 TargetBulkDataFileBytes,
-		TArray<FFile>& OutBulkDataFiles);
-
 	/** Compute the number of files and sizes the BulkData will be split into and update
 	 * the streamable's FileIds and Offsets.
 	 */
@@ -1415,37 +1406,6 @@ public:
 };
 
 
-UCLASS(MinimalAPI, config = Engine)
-class UCustomizableObjectBulk : public UObject
-{
-public:
-	GENERATED_BODY()
-
-	//~ Begin UObject Interface
-	CUSTOMIZABLEOBJECT_API virtual void PostLoad() override;
-	//~ End UObject Interface
-
-	/**  */
-	const FString& GetBulkFilePrefix() const { return BulkFilePrefix; }
-
-	CUSTOMIZABLEOBJECT_API TUniquePtr<IAsyncReadFileHandle> OpenFileAsyncRead(uint32 FileId, uint32 Flags) const;
-
-#if WITH_EDITOR
-
-	//~ Begin UObject Interface
-	CUSTOMIZABLEOBJECT_API virtual void CookAdditionalFilesOverride(const TCHAR*, const ITargetPlatform*, TFunctionRef<void(const TCHAR*, void*, int64)>) override;
-	//~ End UObject Interface
-#endif
-
-#if WITH_EDITOR
-private:
-#endif
-
-	/** Prefix to locate bulkfiles for loading, using the file ids in each FMutableStreamableBlock. */
-	FString BulkFilePrefix;
-};
-
-
 USTRUCT()
 struct FMutableMeshComponentData
 {
@@ -1815,9 +1775,6 @@ public:
 	CUSTOMIZABLEOBJECT_API void PostCompile();
 #endif
 
-	/** Returns a pointer to the BulkData subobject, only valid in packaged builds. */
-	CUSTOMIZABLEOBJECT_API const UCustomizableObjectBulk* GetStreamableBulkData() const;
-
 	CUSTOMIZABLEOBJECT_API UCustomizableObject* GetPublic() const;
 
 #if WITH_EDITOR
@@ -1999,10 +1956,6 @@ public:
 	UPROPERTY()
 	int32 OptimizationLevel = UE_MUTABLE_MAX_OPTIMIZATION;
 
-	/** High limit of the size in bytes of the packaged data when cooking this object.
-	 * This limit is before any pak or filesystem compression. This limit will be broken if a single piece of data is bigger because data is not fragmented for packaging purposes.	*/
-	UPROPERTY()
-	uint64 PackagedDataBytesLimit = 256 * 1024 * 1024;
 	
 	/** High (inclusive) limit of the size in bytes of a data block to be included into the compiled object directly instead of stored in a streamable file. */
 	UPROPERTY()
diff --git a/Engine/Plugins/Mutable/Source/CustomizableObject/Private/MuCO/CustomizableObject.cpp b/Engine/Plugins/Mutable/Source/CustomizableObject/Private/MuCO/CustomizableObject.cpp
--- a/Engine/Plugins/Mutable/Source/CustomizableObject/Private/MuCO/CustomizableObject.cpp
+++ b/Engine/Plugins/Mutable/Source/CustomizableObject/Private/MuCO/CustomizableObject.cpp
@@ -61,22 +61,6 @@ DEFINE_LOG_CATEGORY(LogMutable);
 
 #if WITH_EDITOR
 
-TAutoConsoleVariable<int32> CVarPackagedDataBytesLimitOverride(
-	TEXT("mutable.PackagedDataBytesLimitOverride"),
-	-1,
-	TEXT("Defines the value to be used as 'PackagedDataBytesLimitOverride' for the compilation of all COs.\n")
-	TEXT(" <0 : Use value defined in the CO\n")
-	TEXT(" >=0  : Use this value instead\n"));
-
-
-TAutoConsoleVariable<bool> CVarMutableUseBulkData(
-	TEXT("Mutable.UseBulkData"),
-	true,
-	TEXT("Switch between .utoc/.ucas (FBulkData) and .mut files (CookAdditionalFiles).\n")
-	TEXT("True - Use FBulkData to store streamable data.\n")
-	TEXT("False - Use Mut files to store streamable data\n"));
-
-
 TAutoConsoleVariable<bool> CVarMutableAsyncCook(
 	TEXT("Mutable.CookAsync"),
 	true,
@@ -87,29 +71,6 @@ TAutoConsoleVariable<bool> CVarMutableAsyncCook(
 
 static const FString s_EmptyString;
 
-#if WITH_EDITORONLY_DATA
-
-namespace UE::Mutable::Private
-{
-
-template <typename T>
-T* MoveOldObjectAndCreateNew(UClass* Class, UObject* InOuter)
-{
-	FName ObjectFName = Class->GetFName();
-	FString ObjectNameStr = ObjectFName.ToString();
-	UObject* Existing = FindObject<UAssetUserData>(InOuter, *ObjectNameStr);
-	if (Existing)
-	{
-		// Move the old object out of the way
-		Existing->Rename(nullptr /* Rename will pick a free name*/, GetTransientPackage(), REN_DontCreateRedirectors);
-	}
-	return NewObject<T>(InOuter, Class, *ObjectNameStr);
-}
-
-}
-
-#endif
-
 //-------------------------------------------------------------------------------------------------
 
 UCustomizableObject::UCustomizableObject()
@@ -189,50 +150,37 @@ void UCustomizableObject::PreSave(FObjectPreSaveContext ObjectSaveContext)
 			TObjectPtr<UModelResources> ModelResources = GetPrivate()->GetModelResources(true);
 			ModelResources->Rename(nullptr, this, REN_DontCreateRedirectors);
 
-			const bool bUseBulkData = CVarMutableUseBulkData.GetValueOnAnyThread();
-			if (bUseBulkData)
-			{
-				UE::Mutable::Private::FMutableCachedPlatformData& CachedPlatformData = *GetPrivate()->CachedPlatformsData.Find(TargetPlatform->PlatformName());
-				TSharedPtr<FModelStreamableBulkData> ModelStreamableBulkData = GetPrivate()->GetModelStreamableBulkData(true);
+			UE::Mutable::Private::FMutableCachedPlatformData& CachedPlatformData = *GetPrivate()->CachedPlatformsData.Find(TargetPlatform->PlatformName());
+			TSharedPtr<FModelStreamableBulkData> ModelStreamableBulkData = GetPrivate()->GetModelStreamableBulkData(true);
 
-				const int32 NumBulkDataFiles = CachedPlatformData.BulkDataFiles.Num();
+			const int32 NumBulkDataFiles = CachedPlatformData.BulkDataFiles.Num();
 
-				ModelStreamableBulkData->StreamableBulkData.SetNum(NumBulkDataFiles);
+			ModelStreamableBulkData->StreamableBulkData.SetNum(NumBulkDataFiles);
 
-				const auto WriteBulkData = [ModelStreamableBulkData](UE::Mutable::Private::FFile& File, TArray64<uint8>& FileBulkData, uint32 FileIndex)
-					{
-						MUTABLE_CPUPROFILER_SCOPE(WriteBulkData);
+			const auto WriteBulkData = [ModelStreamableBulkData](UE::Mutable::Private::FFile& File, TArray64<uint8>& FileBulkData, uint32 FileIndex)
+				{
+					MUTABLE_CPUPROFILER_SCOPE(WriteBulkData);
 
-						FByteBulkData& ByteBulkData = ModelStreamableBulkData->StreamableBulkData[FileIndex];
+					FByteBulkData& ByteBulkData = ModelStreamableBulkData->StreamableBulkData[FileIndex];
 
-						// BulkData file to store the file to. CookedIndex 0 is used as a default for backwards compatibility, +1 to skip it.
-						ByteBulkData.SetCookedIndex(FBulkDataCookedIndex((File.Id % MAX_uint8) + 1));
+					// BulkData file to store the file to. CookedIndex 0 is used as a default for backwards compatibility, +1 to skip it.
+					ByteBulkData.SetCookedIndex(FBulkDataCookedIndex((File.Id % MAX_uint8) + 1));
 
-						ByteBulkData.Lock(LOCK_READ_WRITE);
-						uint8* Ptr = (uint8*)ByteBulkData.Realloc(FileBulkData.Num());
-						FMemory::Memcpy(Ptr, FileBulkData.GetData(), FileBulkData.Num());
-						ByteBulkData.Unlock();
+					ByteBulkData.Lock(LOCK_READ_WRITE);
+					uint8* Ptr = (uint8*)ByteBulkData.Realloc(FileBulkData.Num());
+					FMemory::Memcpy(Ptr, FileBulkData.GetData(), FileBulkData.Num());
+					ByteBulkData.Unlock();
 
-						uint32 BulkDataFlags = BULKDATA_PayloadInSeparateFile | BULKDATA_Force_NOT_InlinePayload;
-						if (File.Flags == uint16(EMutableFileFlags::HighRes))
-						{
-							BulkDataFlags |= BULKDATA_OptionalPayload;
-						}
-						ByteBulkData.SetBulkDataFlags(BulkDataFlags);
-					};
+					uint32 BulkDataFlags = BULKDATA_PayloadInSeparateFile | BULKDATA_Force_NOT_InlinePayload;
+					if (File.Flags == uint16(EMutableFileFlags::HighRes))
+					{
+						BulkDataFlags |= BULKDATA_OptionalPayload;
+					}
+					ByteBulkData.SetBulkDataFlags(BulkDataFlags);
+				};
 
-				bool bDropData = true;
-				UE::Mutable::Private::SerializeBulkDataFiles(CachedPlatformData, CachedPlatformData.BulkDataFiles, WriteBulkData, bDropData);
-			}
-			else 
-			{
-				// Create an export object to manage the streamable data
-				if (!BulkData)
-				{
-					BulkData = UE::Mutable::Private::MoveOldObjectAndCreateNew<UCustomizableObjectBulk>(UCustomizableObjectBulk::StaticClass(), this);
-				}
-				BulkData->Mark(OBJECTMARK_TagExp);
-			}
+			bool bDropData = true;
+			UE::Mutable::Private::SerializeBulkDataFiles(CachedPlatformData, CachedPlatformData.BulkDataFiles, WriteBulkData, bDropData);
 		}
 		else
 		{
@@ -396,7 +344,6 @@ void UCustomizableObjectPrivate::BackwardsCompatibleFixup(int32 CustomizableObje
 		OptimizationLevel = GetPublic()->CompileOptions_DEPRECATED.OptimizationLevel;
 		TextureCompression = GetPublic()->CompileOptions_DEPRECATED.TextureCompression;
 		EmbeddedDataBytesLimit = GetPublic()->CompileOptions_DEPRECATED.EmbeddedDataBytesLimit;
-		PackagedDataBytesLimit = GetPublic()->CompileOptions_DEPRECATED.PackagedDataBytesLimit;
 	}
 
 	if (CustomizableObjectCustomVersion == FCustomizableObjectCustomVersion::NewComponentOptions)
@@ -1104,12 +1051,6 @@ void UCustomizableObjectPrivate::PostCompile()
 #endif
 
 
-const UCustomizableObjectBulk* UCustomizableObjectPrivate::GetStreamableBulkData() const
-{
-	return GetPublic()->BulkData;
-}
-
-
 UCustomizableObject* UCustomizableObjectPrivate::GetPublic() const
 {
 	UCustomizableObject* Public = StaticCast<UCustomizableObject*>(GetOuter());
@@ -2438,20 +2379,7 @@ FCompilationOptions UCustomizableObjectPrivate::GetCompileOptions() const
 	FCompilationOptions Options;
 	Options.TextureCompression = TextureCompression;
 	Options.OptimizationLevel = OptimizationLevel;
-
 	Options.TargetPlatform = GetTargetPlatformManagerRef().GetRunningTargetPlatform();
-
-	const int32 TargetBulkDataFileBytesOverride = CVarPackagedDataBytesLimitOverride.GetValueOnAnyThread();
-	if ( TargetBulkDataFileBytesOverride >= 0)
-	{
-		Options.PackagedDataBytesLimit = TargetBulkDataFileBytesOverride;
-		UE_LOG(LogMutable,Display, TEXT("Ignoring CO PackagedDataBytesLimit value in favour of overriding CVar value : mutable.PackagedDataBytesLimitOverride %llu"), Options.PackagedDataBytesLimit);
-	}
-	else
-	{
-		Options.PackagedDataBytesLimit =  PackagedDataBytesLimit;
-	}
-
 	Options.EmbeddedDataBytesLimit = EmbeddedDataBytesLimit;
 	Options.CustomizableObjectNumBoneInfluences = ICustomizableObjectModule::Get().GetNumBoneInfluences();
 	Options.bRealTimeMorphTargetsEnabled = IsRealTimeMorphTargetsEnabled();
@@ -2875,93 +2803,6 @@ namespace UE::Mutable::Private
 	}
 	
 
-	void GenerateBulkDataFilesListWithSizeLimit(
-		TSharedPtr<const UE::Mutable::Private::FModel, ESPMode::ThreadSafe> Model,
-		FModelStreamableBulkData& ModelStreamableBulkData,
-		const ITargetPlatform* TargetPlatform,
-		uint64 TargetBulkDataFileBytes,
-		TArray<FFile>& OutBulkDataFiles)
-	{
-		MUTABLE_CPUPROFILER_SCOPE(GenerateBulkDataFilesListWithSizeLimit);
-
-		if (!Model)
-		{
-			return;
-		}
-
-		const uint64 MaxChunkSize = UCustomizableObjectSystem::GetInstance()->GetMaxChunkSizeForPlatform(TargetPlatform);
-		TargetBulkDataFileBytes = FMath::Min(TargetBulkDataFileBytes, MaxChunkSize);
-
-		// Unlimited number of files, limited file size. Add blocks to the file if the size limit won't be surpassed. Add at least one block to each file. 
-		const auto CreateFileList = [TargetBulkDataFileBytes](const FFileCategoryID& CategoryID, const FClassifyNode& Node, TArray<FFile>& OutBulkDataFiles)
-			{
-				// Temp: Group by order in the array
-				for (int32 BlockIndex = 0; BlockIndex < Node.Blocks.Num(); )
-				{
-					FFile File;
-					File.DataType = CategoryID.DataType;
-					File.ResourceType = CategoryID.ResourceType;
-					File.Flags = CategoryID.Flags;
-
-					uint64 FileSize = 0;
-					uint32 FileId = uint32(CategoryID.DataType);
-
-					while (BlockIndex < Node.Blocks.Num())
-					{
-						const FBlock& CurrentBlock = Node.Blocks[BlockIndex];
-
-						if (FileSize > 0 &&
-							FileSize + CurrentBlock.Size > TargetBulkDataFileBytes &&
-							TargetBulkDataFileBytes > 0)
-						{
-							break;
-						}
-
-						// Block added to file. Set offset and increase file size.
-						FileSize += CurrentBlock.Size;
-
-						// Generate cumulative id for this file
-						FileId = HashCombine(FileId, CurrentBlock.Id);
-
-						// Add the block to the current file
-						File.Blocks.Add(CurrentBlock);
-
-						// Next block
-						++BlockIndex;
-					}
-
-					const int32 NumFiles = OutBulkDataFiles.Num();
-
-					// Ensure the FileId is unique
-					bool bUnique = false;
-					while (!bUnique)
-					{
-						bUnique = true;
-						for (int32 PreviousFileIndex = 0; PreviousFileIndex < NumFiles; ++PreviousFileIndex)
-						{
-							if (OutBulkDataFiles[PreviousFileIndex].Id == FileId)
-							{
-								bUnique = false;
-								++FileId;
-								break;
-							}
-						}
-					}
-
-					// Set it to the editor-only file descriptor
-					File.Id = FileId;
-
-					OutBulkDataFiles.Add(MoveTemp(File));
-				}
-			};
-
-		// TODO: Temp. Remove after unifying generated output files code between editor an package. UE-222777
-		const bool bUseRomTypeAndFlagsToFilter = TargetPlatform->RequiresCookedData();
-
-		GenerateBulkDataFilesList(Model, ModelStreamableBulkData, bUseRomTypeAndFlagsToFilter, CreateFileList, OutBulkDataFiles);
-	}
-
-
 	void GenerateBulkDataFilesList(
 		TSharedPtr<const UE::Mutable::Private::FModel, ESPMode::ThreadSafe> Model,
 		FModelStreamableBulkData& ModelStreamableBulkData,
@@ -3508,82 +3349,6 @@ void UCustomizableObjectPrivate::LogMemory()
 
 #endif // WITH_EDITOR
 
-//-------------------------------------------------------------------------------------------------
-//-------------------------------------------------------------------------------------------------
-//-------------------------------------------------------------------------------------------------
-void UCustomizableObjectBulk::PostLoad()
-{
-	UObject::PostLoad();
-
-	const FString OutermostName = GetOutermost()->GetName();
-	FString PackageFilename = FPackageName::LongPackageNameToFilename(OutermostName);
-	FPaths::MakeStandardFilename(PackageFilename);
-	BulkFilePrefix = PackageFilename;
-}
-
-TUniquePtr<IAsyncReadFileHandle> UCustomizableObjectBulk::OpenFileAsyncRead(uint32 FileId, uint32 Flags) const
-{
-	check(IsInGameThread());
-
-	FString FilePath = FString::Printf(TEXT("%s-%08x.mut"), *BulkFilePrefix, FileId);
-	if (Flags == uint16(EMutableFileFlags::HighRes))
-	{
-		FilePath += TEXT(".high");
-	}
-
-	IAsyncReadFileHandle* Result = FPlatformFileManager::Get().GetPlatformFile().OpenAsyncRead(*FilePath);
-	
-	// Result being null does not mean the file does not exist. A request has to be made. Let the callee deal with it.
-	//UE_CLOG(!Result, LogMutable, Warning, TEXT("CustomizableObjectBulkData: Failed to open file [%s]."), *FilePath);
-
-	return TUniquePtr<IAsyncReadFileHandle>(Result);
-}
-
-#if WITH_EDITOR
-
-void UCustomizableObjectBulk::CookAdditionalFilesOverride(const TCHAR* PackageFilename,
-	const ITargetPlatform* TargetPlatform,
-	TFunctionRef<void(const TCHAR* Filename, void* Data, int64 Size)> WriteAdditionalFile)
-{
-	// Don't save streamed data on server builds since it won't be used anyway.
-	if (TargetPlatform->IsServerOnly())
-	{
-		return;
-	}
-
-	UCustomizableObject* CustomizableObject = CastChecked<UCustomizableObject>(GetOutermostObject());
-
-	UE::Mutable::Private::FMutableCachedPlatformData* PlatformData = CustomizableObject->GetPrivate()->CachedPlatformsData.Find(TargetPlatform->PlatformName());
-	if (!PlatformData)
-	{
-		UE_LOG(LogMutable, Warning, TEXT("CookAdditionalFilesOverride: Customizable Object [%s] is missing [%s] platform data."),
-			*CustomizableObject->GetName(),
-			*TargetPlatform->PlatformName());
-
-		return;
-	}
-
-	const FString CookedBulkFileName = FString::Printf(TEXT("%s/%s"), *FPaths::GetPath(PackageFilename), *CustomizableObject->GetName());
-
-	const auto WriteFile = [WriteAdditionalFile, CookedBulkFileName](UE::Mutable::Private::FFile& File, TArray64<uint8>& FileBulkData, uint32 FileIndex)
-		{
-			FString FileName = CookedBulkFileName + FString::Printf(TEXT("-%08x.mut"), File.Id);
-
-			if (File.Flags == uint16(EMutableFileFlags::HighRes))
-			{
-				// We can do something different here for high-res data.
-				// For example: change the file name. We also need to detect it when generating the file name for loading.
-				FileName += TEXT(".high");
-			}
-
-			WriteAdditionalFile(*FileName, FileBulkData.GetData(), FileBulkData.Num());
-		};
-
-	bool bDropData = true;
-	UE::Mutable::Private::SerializeBulkDataFiles(*PlatformData, PlatformData->BulkDataFiles, WriteFile, bDropData);
-}
-#endif // WITH_EDITOR
-
 
 //-------------------------------------------------------------------------------------------------
 //-------------------------------------------------------------------------------------------------
diff --git a/Engine/Plugins/Mutable/Source/CustomizableObject/Private/MuCO/UnrealMutableModelDiskStreamer.cpp b/Engine/Plugins/Mutable/Source/CustomizableObject/Private/MuCO/UnrealMutableModelDiskStreamer.cpp
--- a/Engine/Plugins/Mutable/Source/CustomizableObject/Private/MuCO/UnrealMutableModelDiskStreamer.cpp
+++ b/Engine/Plugins/Mutable/Source/CustomizableObject/Private/MuCO/UnrealMutableModelDiskStreamer.cpp
@@ -81,46 +81,32 @@ bool FUnrealMutableModelBulkReader::PrepareStreamingForObject(UCustomizableObjec
 		}
 	}
 
+	const TSharedPtr<FModelStreamableBulkData> ModelStreamableBulkData = CustomizableObject->GetPrivate()->GetModelStreamableBulkData();
+
+	// There's no need to prepare the object for streaming if there's no data to stream (Dataless COs).
+	if (!ModelStreamableBulkData || ModelStreamableBulkData->ModelStreamables.IsEmpty()) //-V1051
+	{
+		return true;
+	}
+
+	const TSharedPtr<const UE::Mutable::Private::FModel>& Model = CustomizableObject->GetPrivate()->GetModel();
+
 	// Is the object already prepared for streaming?
 	bool bAlreadyStreaming = Objects.FindByPredicate(
-		[CustomizableObject](const FObjectData& d)
-		{ return d.Model.Pin().Get() == CustomizableObject->GetPrivate()->GetModel().Get(); })
+		[Model](const FObjectData& d)
+		{ return d.Model.Pin().Get() == Model.Get(); })
 		!=
 		nullptr;
 
 	if (!bAlreadyStreaming)
 	{
 		FObjectData NewData;
-		NewData.Model = TWeakPtr<const UE::Mutable::Private::FModel>(CustomizableObject->GetPrivate()->GetModel());
-
-		const TSharedPtr<FModelStreamableBulkData> ModelStreamableBulkData = CustomizableObject->GetPrivate()->GetModelStreamableBulkData();
+		NewData.Model = TWeakPtr<const UE::Mutable::Private::FModel>(Model);
 		NewData.ModelStreamableBulkData = TWeakPtr<FModelStreamableBulkData>(ModelStreamableBulkData);
 
-		if (!ModelStreamableBulkData || ModelStreamableBulkData->ModelStreamables.IsEmpty()) //-V1051
-		{
-			UE_LOG(LogMutable, Warning, TEXT("Streaming: Customizable Object %s has no data to stream."), *CustomizableObject->GetName());
-
-#if !WITH_EDITOR
-			check(false);
-#endif
-			return false;
-		}
-
 #if WITH_EDITOR
 		FString FullFileName = CustomizableObject->GetPrivate()->GetCompiledDataFileName(nullptr, true);
 		NewData.BulkFilePrefix = *FullFileName + GetDataTypeExtension(UE::Mutable::Private::EStreamableDataType::Model);
-#else
-		if (ModelStreamableBulkData->StreamableBulkData.IsEmpty())
-		{
-			const UCustomizableObjectBulk* BulkData = CustomizableObject->GetPrivate()->GetStreamableBulkData();
-			if (!BulkData)
-			{
-				UE_LOG(LogMutable, Warning, TEXT("Streaming: Customizable Object %s is missing the BulkData export."), *CustomizableObject->GetName());
-				return false;
-			}
-
-			NewData.BulkFilePrefix = BulkData->GetBulkFilePrefix();
-		}
 #endif
 		Objects.Add(MoveTemp(NewData));
 	}
@@ -251,7 +237,8 @@ bool FUnrealMutableModelBulkReader::DoesBlockExist(const UE::Mutable::Private::F
 		return false;
 	}
 
-	if (PinnedModelStreamableBulkData->StreamableBulkData.IsValidIndex(Block->FileId))
+#if !WITH_EDITOR
+	if (ensure(PinnedModelStreamableBulkData->StreamableBulkData.IsValidIndex(Block->FileId)))
 	{	
 		MUTABLE_CPUPROFILER_SCOPE(FUnrealMutableModelBulkReader::DoesBlockExist_Bulk)
 
@@ -260,21 +247,21 @@ bool FUnrealMutableModelBulkReader::DoesBlockExist(const UE::Mutable::Private::F
 	}
 	else
 	{
-#if WITH_EDITOR
-		if (PinnedModelStreamableBulkData->bIsStoredInDDC)
-		{
-			MUTABLE_CPUPROFILER_SCOPE(FUnrealMutableModelBulkReader::DoesBlockExist_DDC)
-			return true;
-		}
-		else
+		return false;
+	}
+#else
+	if (PinnedModelStreamableBulkData->bIsStoredInDDC)
+	{
+		MUTABLE_CPUPROFILER_SCOPE(FUnrealMutableModelBulkReader::DoesBlockExist_DDC)
+		return true;
+	}
+	else
+	{
+		MUTABLE_CPUPROFILER_SCOPE(FUnrealMutableModelBulkReader::DoesBlockExist_File)
+		const FString FilePath = Private::MakeBlockFilePath(ObjectData->BulkFilePrefix, *Block);
+		return FPlatformFileManager::Get().GetPlatformFile().FileExists(*FilePath);
+	}
 #endif
-		{
-			MUTABLE_CPUPROFILER_SCOPE(FUnrealMutableModelBulkReader::DoesBlockExist_File)
-
-			const FString FilePath = Private::MakeBlockFilePath(ObjectData->BulkFilePrefix, *Block);
-			return FPlatformFileManager::Get().GetPlatformFile().FileExists(*FilePath);
-		}
-	}	
 }
 
 
@@ -327,7 +314,8 @@ UE::Mutable::Private::FModelReader::FOperationID FUnrealMutableModelBulkReader::
 
 	FOperationID Result = ++LastOperationID;
 
-	if (PinnedModelStreamableBulkData->StreamableBulkData.IsValidIndex(Block->FileId))
+#if !WITH_EDITOR
+	if (ensure(PinnedModelStreamableBulkData->StreamableBulkData.IsValidIndex(Block->FileId)))
 	{
 		FByteBulkData& BulkData = PinnedModelStreamableBulkData->StreamableBulkData[Block->FileId];
 
@@ -352,105 +340,100 @@ UE::Mutable::Private::FModelReader::FOperationID FUnrealMutableModelBulkReader::
 	}
 	else
 	{
-#if WITH_EDITOR
-		if (PinnedModelStreamableBulkData->bIsStoredInDDC)
+		return INDEX_NONE;
+	}
+#else
+	if (PinnedModelStreamableBulkData->bIsStoredInDDC)
+	{
+		using namespace UE::DerivedData;
+
+		UE::FSharedString SharedName = UE::FSharedString(TEXT("UnrealMutableModelBulkReader"));
+
+		FCacheGetValueRequest Request;
+		Request.Name = SharedName;
+		Request.Key = PinnedModelStreamableBulkData->DDCKey;
+		Request.Key.Hash = PinnedModelStreamableBulkData->DDCValues[Block->FileId];
+		Request.Policy = PinnedModelStreamableBulkData->DDCDefaultPolicy;
+		
+		FReadRequest ReadRequest;
+		ReadRequest.DDCReadRequest = MakeShared<FDDCReadRequest>();
+		GetCache().GetValue(MakeArrayView(&Request,1), ReadRequest.DDCReadRequest->RequestOwner,
+			[Request = ReadRequest.DDCReadRequest, CompletionCallbackCapture = *CompletionCallback, Block, Buffer, Size](FCacheGetValueResponse&& Response)
 		{
-			using namespace UE::DerivedData;
+			MUTABLE_CPUPROFILER_SCOPE(FUnrealMutableModelBulkStreamer::OpenReadFile_DDC_Callback);
+			bool bSuccess = Response.Status == EStatus::Ok;
 
-			UE::FSharedString SharedName = UE::FSharedString(TEXT("UnrealMutableModelBulkReader"));
-
-			FCacheGetValueRequest Request;
-			Request.Name = SharedName;
-			Request.Key = PinnedModelStreamableBulkData->DDCKey;
-			Request.Key.Hash = PinnedModelStreamableBulkData->DDCValues[Block->FileId];
-			Request.Policy = PinnedModelStreamableBulkData->DDCDefaultPolicy;
-			
-			FReadRequest ReadRequest;
-			ReadRequest.DDCReadRequest = MakeShared<FDDCReadRequest>();
-			GetCache().GetValue(MakeArrayView(&Request,1), ReadRequest.DDCReadRequest->RequestOwner,
-				[Request = ReadRequest.DDCReadRequest, CompletionCallbackCapture = *CompletionCallback, Block, Buffer, Size](FCacheGetValueResponse&& Response)
+			if (bSuccess)
 			{
-				MUTABLE_CPUPROFILER_SCOPE(FUnrealMutableModelBulkStreamer::OpenReadFile_DDC_Callback);
-				bool bSuccess = Response.Status == EStatus::Ok;
-
-				if (bSuccess)
+				const FCompressedBuffer& CompressedBuffer = Response.Value.GetData();
+				if (ensure(!CompressedBuffer.IsNull()))
 				{
-					const FCompressedBuffer& CompressedBuffer = Response.Value.GetData();
-					if (ensure(!CompressedBuffer.IsNull()))
+					if (Size == CompressedBuffer.GetRawSize())
+					{
+						bSuccess = CompressedBuffer.TryDecompressTo(MakeMemoryView(Buffer, Size));
+					}
+					else if (Block->Offset + Size <= CompressedBuffer.GetRawSize())
 					{
-						if (Size == CompressedBuffer.GetRawSize())
-						{
-							bSuccess = CompressedBuffer.TryDecompressTo(MakeMemoryView(Buffer, Size));
-						}
-						else if (Block->Offset + Size <= CompressedBuffer.GetRawSize())
-						{
-							FCompressedBufferReader BufferReader(CompressedBuffer);
-							bSuccess = BufferReader.TryDecompressTo(MakeMemoryView(Buffer, Size), Block->Offset);
-						}
-						else 
-						{
-							bSuccess = false;
-						}
+						FCompressedBufferReader BufferReader(CompressedBuffer);
+						bSuccess = BufferReader.TryDecompressTo(MakeMemoryView(Buffer, Size), Block->Offset);
+					}
+					else 
+					{
+						bSuccess = false;
 					}
 				}
+			}
 
-				Request->bSuccess = bSuccess;
-				CompletionCallbackCapture(bSuccess);
-			});
-
-			ObjectData->CurrentReadRequests.Add(Result, ReadRequest);
-		}
-		else
-#endif
-		{
-
-			int32 BulkDataOffsetInFile = 0;
+			Request->bSuccess = bSuccess;
+			CompletionCallbackCapture(bSuccess);
+		});
 
-#if WITH_EDITOR
-			BulkDataOffsetInFile = sizeof(MutableCompiledDataStreamHeader);
-#endif
+		ObjectData->CurrentReadRequests.Add(Result, ReadRequest);
+	}
+	else
+	{
+		FReadRequest ReadRequest;
 
-			FReadRequest ReadRequest;
+		if (CompletionCallback)
+		{
+			ReadRequest.FileCallback = MakeShared<TFunction<void(bool, IAsyncReadRequest*)>>([CompletionCallbackCapture = *CompletionCallback](bool bWasCancelled, IAsyncReadRequest*) -> void
+				{
+					CompletionCallbackCapture(!bWasCancelled);
+				});
+		}
 
-			if (CompletionCallback)
-			{
-				ReadRequest.FileCallback = MakeShared<TFunction<void(bool, IAsyncReadRequest*)>>([CompletionCallbackCapture = *CompletionCallback](bool bWasCancelled, IAsyncReadRequest*) -> void
-					{
-						CompletionCallbackCapture(!bWasCancelled);
-					});
-			}
+		TSharedPtr<IAsyncReadFileHandle> FileHandle;
+		{
+			FScopeLock Lock(&FileHandlesCritical);
 
-			TSharedPtr<IAsyncReadFileHandle> FileHandle;
+			TSharedPtr<IAsyncReadFileHandle>& Found = ObjectData->ReadFileHandles.FindOrAdd(Block->FileId);
+			if (!Found)
 			{
-				FScopeLock Lock(&FileHandlesCritical);
+				FString FilePath = Private::MakeBlockFilePath(ObjectData->BulkFilePrefix, *Block);
+				Found = MakeShareable(FPlatformFileManager::Get().GetPlatformFile().OpenAsyncRead(*FilePath));
 
-				TSharedPtr<IAsyncReadFileHandle>& Found = ObjectData->ReadFileHandles.FindOrAdd(Block->FileId);
 				if (!Found)
 				{
-					FString FilePath = Private::MakeBlockFilePath(ObjectData->BulkFilePrefix, *Block);
-					Found = MakeShareable(FPlatformFileManager::Get().GetPlatformFile().OpenAsyncRead(*FilePath));
-
-					if (!Found)
-					{
-						UE_LOG(LogMutable, Error, TEXT("Failed to create AsyncReadFileHandle. File Path [%s]."), *FilePath);
-						check(false);
-						return -1;
-					}
+					UE_LOG(LogMutable, Error, TEXT("Failed to create AsyncReadFileHandle. File Path [%s]."), *FilePath);
+					check(false);
+					return -1;
 				}
-
-				FileHandle = Found;
 			}
 
-			ReadRequest.FileReadRequest = MakeShareable(FileHandle->ReadRequest(
-				BulkDataOffsetInFile + Block->Offset,
-				Size,
-				(EAsyncIOPriorityAndFlags)StreamPriority,
-				ReadRequest.FileCallback.Get(),
-				reinterpret_cast<uint8*>(Buffer)));
-
-			ObjectData->CurrentReadRequests.Add(Result, ReadRequest);
+			FileHandle = Found;
 		}
+
+		ReadRequest.FileReadRequest = MakeShareable(FileHandle->ReadRequest(
+			sizeof(MutableCompiledDataStreamHeader) + Block->Offset,
+			Size,
+			(EAsyncIOPriorityAndFlags)StreamPriority,
+			ReadRequest.FileCallback.Get(),
+			reinterpret_cast<uint8*>(Buffer)));
+
+		ObjectData->CurrentReadRequests.Add(Result, ReadRequest);
 	}
+#endif
+	
 
 	INC_DWORD_STAT(STAT_MutableStreamingOps);
 
diff --git a/Engine/Plugins/Mutable/Source/CustomizableObject/Public/MuCO/CustomizableObject.h b/Engine/Plugins/Mutable/Source/CustomizableObject/Public/MuCO/CustomizableObject.h
--- a/Engine/Plugins/Mutable/Source/CustomizableObject/Public/MuCO/CustomizableObject.h
+++ b/Engine/Plugins/Mutable/Source/CustomizableObject/Public/MuCO/CustomizableObject.h
@@ -23,7 +23,6 @@ class UCustomizableObject;
 class UEdGraph;
 class USkeletalMesh;
 class UCustomizableObjectPrivate;
-class UCustomizableObjectBulk;
 struct FFrame;
 struct FStreamableHandle;
 template <typename FuncType> class TFunctionRef;
@@ -35,8 +34,6 @@ CUSTOMIZABLEOBJECT_API DECLARE_LOG_CATEGORY_EXTERN(LogMutable, Log, All);
 DECLARE_MULTICAST_DELEGATE(FPostCompileDelegate)
 
 
-CUSTOMIZABLEOBJECT_API extern TAutoConsoleVariable<bool> CVarMutableUseBulkData;
-
 CUSTOMIZABLEOBJECT_API extern TAutoConsoleVariable<bool> CVarMutableAsyncCook;
 
 
@@ -561,10 +558,6 @@ public:
 	CUSTOMIZABLEOBJECT_API UCustomizableObjectPrivate* GetPrivate();
 	
 private:
-	/** BulkData that stores all in-game resources used by Mutable when generating instances.
-	  * Only valid in packaged builds */
-	UPROPERTY()
-	TObjectPtr<UCustomizableObjectBulk> BulkData;
 
 	UPROPERTY()
 	TObjectPtr<UCustomizableObjectPrivate> Private;
diff --git a/Engine/Plugins/Mutable/Source/CustomizableObject/Public/MuCO/CustomizableObject_Deprecated.h b/Engine/Plugins/Mutable/Source/CustomizableObject/Public/MuCO/CustomizableObject_Deprecated.h
--- a/Engine/Plugins/Mutable/Source/CustomizableObject/Public/MuCO/CustomizableObject_Deprecated.h
+++ b/Engine/Plugins/Mutable/Source/CustomizableObject/Public/MuCO/CustomizableObject_Deprecated.h
@@ -18,9 +18,6 @@ struct FCompilationOptions_DEPRECATED
 	UPROPERTY()
 	int32 OptimizationLevel = 2; // UE_MUTABLE_MAX_OPTIMIZATION;
 
-	UPROPERTY()
-	uint64 PackagedDataBytesLimit = 256 * 1024 * 1024;
-
 	UPROPERTY()
 	uint64 EmbeddedDataBytesLimit = 1024;
 };
\ No newline at end of file
diff --git a/Engine/Plugins/Mutable/Source/CustomizableObjectEditor/Private/MuCOE/CustomizableObjectCompileRunnable.cpp b/Engine/Plugins/Mutable/Source/CustomizableObjectEditor/Private/MuCOE/CustomizableObjectCompileRunnable.cpp
--- a/Engine/Plugins/Mutable/Source/CustomizableObjectEditor/Private/MuCOE/CustomizableObjectCompileRunnable.cpp
+++ b/Engine/Plugins/Mutable/Source/CustomizableObjectEditor/Private/MuCOE/CustomizableObjectCompileRunnable.cpp
@@ -309,7 +309,7 @@ uint32 FCustomizableObjectCompileRunnable::Run()
 	);
 
 	const int32 MinResidentMips = UTexture::GetStaticMinTextureResidentMipCount();
-	CompilerOptions->SetDataPackingStrategy( MinResidentMips, Options.EmbeddedDataBytesLimit, Options.PackagedDataBytesLimit );
+	CompilerOptions->SetDataPackingStrategy( MinResidentMips, Options.EmbeddedDataBytesLimit );
 
 	// We always compile for progressive image generation.
 	CompilerOptions->SetEnableProgressiveImages(true);
@@ -489,16 +489,8 @@ void FCustomizableObjectSaveDDRunnable::CachePlatformData()
 	// Cache ModelStreamables
 	{
 		// Generate list of files and update streamable blocks ids and offsets
-		if (Options.bUseBulkData)
-		{
-			UE::Mutable::Private::GenerateBulkDataFilesListWithFileLimit(CustomizableObjectId, PlatformData->Model, *PlatformData->ModelStreamableBulkData.Get(), MAX_uint8, 
+		UE::Mutable::Private::GenerateBulkDataFilesListWithFileLimit(CustomizableObjectId, PlatformData->Model, *PlatformData->ModelStreamableBulkData.Get(), MAX_uint8, 
 				Options.bIsCooking, false, *Options.TargetPlatform, PlatformData->BulkDataFiles);
-		}
-		else
-		{
-			const uint64 PackageDataBytesLimit = Options.bIsCooking ? Options.PackagedDataBytesLimit : MAX_uint64;
-			UE::Mutable::Private::GenerateBulkDataFilesListWithSizeLimit(PlatformData->Model, *PlatformData->ModelStreamableBulkData.Get(), Options.TargetPlatform, PackageDataBytesLimit, PlatformData->BulkDataFiles);
-		}
 	}
 
 	// Cache Model and Model Roms
diff --git a/Engine/Plugins/Mutable/Source/CustomizableObjectEditor/Private/MuCOE/CustomizableObjectCompiler.cpp b/Engine/Plugins/Mutable/Source/CustomizableObjectEditor/Private/MuCOE/CustomizableObjectCompiler.cpp
--- a/Engine/Plugins/Mutable/Source/CustomizableObjectEditor/Private/MuCOE/CustomizableObjectCompiler.cpp
+++ b/Engine/Plugins/Mutable/Source/CustomizableObjectEditor/Private/MuCOE/CustomizableObjectCompiler.cpp
@@ -1495,17 +1495,10 @@ void FCustomizableObjectCompiler::FinishLoadingModelDataFromDDC()
 					}
 
 					// Generate list of files and update streamable blocks ids and offsets
-					if (CVarMutableUseBulkData.GetValueOnAnyThread())
-					{
-						const uint32 NumBulkDataFilesPerBucket = MAX_uint8;
-						UE::Mutable::Private::GenerateBulkDataFilesListWithFileLimit(GenerateDataDistributionIdentifier(*CustomizableObject), PlatformData->Model, *PlatformData->ModelStreamableBulkData.Get(),
+					const uint32 NumBulkDataFilesPerBucket = MAX_uint8;
+					UE::Mutable::Private::GenerateBulkDataFilesListWithFileLimit(GenerateDataDistributionIdentifier(*CustomizableObject), PlatformData->Model, *PlatformData->ModelStreamableBulkData.Get(),
 							NumBulkDataFilesPerBucket, Options.bIsCooking, false, *Options.TargetPlatform, PlatformData->BulkDataFiles);
-					}
-					else
-					{
-						UE::Mutable::Private::GenerateBulkDataFilesListWithSizeLimit(PlatformData->Model, *PlatformData->ModelStreamableBulkData.Get(),
-							Options.TargetPlatform, Options.PackagedDataBytesLimit, PlatformData->BulkDataFiles);
-					}
+					
 				}
 				else
 				{
diff --git a/Engine/Plugins/Mutable/Source/CustomizableObjectEditor/Private/MuCOE/CustomizableObjectEditor.cpp b/Engine/Plugins/Mutable/Source/CustomizableObjectEditor/Private/MuCOE/CustomizableObjectEditor.cpp
--- a/Engine/Plugins/Mutable/Source/CustomizableObjectEditor/Private/MuCOE/CustomizableObjectEditor.cpp
+++ b/Engine/Plugins/Mutable/Source/CustomizableObjectEditor/Private/MuCOE/CustomizableObjectEditor.cpp
@@ -1139,32 +1139,6 @@ TSharedRef<SWidget> FCustomizableObjectEditor::GenerateCompileOptionsMenuContent
 					}
 				});
 			MenuBuilder.AddWidget(EmbeddedDataLimitCombo.ToSharedRef(), LOCTEXT("MutableCompileEmbeddedLimit", "Embedded Data Limit (Bytes)"));
-
-		// Packaging file size control.
-		TArray<SNumericDropDown<float>::FNamedValue> PackagedOptions;
-		PackagedOptions.Add(SNumericDropDown<float>::FNamedValue(0, FText::FromString(TEXT("0")), FText::FromString(TEXT("Split All"))));
-		PackagedOptions.Add(SNumericDropDown<float>::FNamedValue(16   * 1024, FText::FromString(TEXT("16 KB")), FText::FromString(TEXT("16 KB"))));
-		PackagedOptions.Add(SNumericDropDown<float>::FNamedValue(64   * 1024, FText::FromString(TEXT("64 KB")), FText::FromString(TEXT("64 KB"))));
-		PackagedOptions.Add(SNumericDropDown<float>::FNamedValue(1024 * 1024, FText::FromString(TEXT("1 MB")), FText::FromString(TEXT("1 MB"))));
-		PackagedOptions.Add(SNumericDropDown<float>::FNamedValue(64   * 1024 * 1024, FText::FromString(TEXT("64 MB")), FText::FromString(TEXT("64 MB"))));
-		PackagedOptions.Add(SNumericDropDown<float>::FNamedValue(256  * 1024 * 1024, FText::FromString(TEXT("256 MB")), FText::FromString(TEXT("256 MB"))));
-		PackagedOptions.Add(SNumericDropDown<float>::FNamedValue(1024 * 1024 * 1024, FText::FromString(TEXT("1 GB")), FText::FromString(TEXT("1 GB"))));
-
-		PackagedDataLimitCombo = SNew(SNumericDropDown<float>)
-			.DropDownValues(PackagedOptions)
-			.Value_Lambda([this]()
-				{
-					return CustomizableObject ? float(CustomizableObject->GetPrivate()->PackagedDataBytesLimit) : 0.0f;
-				})
-			.OnValueChanged_Lambda([this](float Value)
-				{
-					if (CustomizableObject)
-					{
-						CustomizableObject->GetPrivate()->PackagedDataBytesLimit = uint64(Value);
-						CustomizableObject->Modify();
-					}
-				});
-			MenuBuilder.AddWidget(PackagedDataLimitCombo.ToSharedRef(), LOCTEXT("MutableCompilePackagedLimit", "Packaged Data File Max Limit (Bytes)"));
 	}
 	MenuBuilder.EndSection();
 
@@ -1525,7 +1499,6 @@ void FCustomizableObjectEditor::ResetCompileOptions()
 	CustomizableObject->GetPrivate()->OptimizationLevel = DefaultObject->OptimizationLevel;
 	CustomizableObject->GetPrivate()->TextureCompression = DefaultObject->TextureCompression;
 	CustomizableObject->GetPrivate()->EmbeddedDataBytesLimit = DefaultObject->EmbeddedDataBytesLimit;
-	CustomizableObject->GetPrivate()->PackagedDataBytesLimit = DefaultObject->PackagedDataBytesLimit;
 }
 
 void FCustomizableObjectEditor::OnChangeCompileOptimizationLevel(TSharedPtr<FString> NewSelection, ESelectInfo::Type SelectInfo)
diff --git a/Engine/Plugins/Mutable/Source/CustomizableObjectEditor/Private/MuCOE/CustomizableObjectEditor.h b/Engine/Plugins/Mutable/Source/CustomizableObjectEditor/Private/MuCOE/CustomizableObjectEditor.h
--- a/Engine/Plugins/Mutable/Source/CustomizableObjectEditor/Private/MuCOE/CustomizableObjectEditor.h
+++ b/Engine/Plugins/Mutable/Source/CustomizableObjectEditor/Private/MuCOE/CustomizableObjectEditor.h
@@ -214,7 +214,6 @@ private:
 	TSharedPtr<STextComboBox> CompileTextureCompressionCombo;
 	TArray< TSharedPtr<FString> > CompileTextureCompressionStrings;
 	TSharedPtr<SNumericDropDown<float>> EmbeddedDataLimitCombo;
-	TSharedPtr<SNumericDropDown<float>> PackagedDataLimitCombo;
 
 	void OnChangeCompileOptimizationLevel(TSharedPtr<FString> NewSelection, ESelectInfo::Type);
 	void OnChangeCompileTextureCompressionType(TSharedPtr<FString> NewSelection, ESelectInfo::Type);
diff --git a/Engine/Plugins/Mutable/Source/CustomizableObjectEditor/Private/MuCOE/CustomizableObjectEditorModule.cpp b/Engine/Plugins/Mutable/Source/CustomizableObjectEditor/Private/MuCOE/CustomizableObjectEditorModule.cpp
--- a/Engine/Plugins/Mutable/Source/CustomizableObjectEditor/Private/MuCOE/CustomizableObjectEditorModule.cpp
+++ b/Engine/Plugins/Mutable/Source/CustomizableObjectEditor/Private/MuCOE/CustomizableObjectEditorModule.cpp
@@ -1297,7 +1297,6 @@ void FCustomizableObjectEditorModule::BeginCacheForCookedPlatformData(UCustomiza
 	CompileRequest->Options.OptimizationLevel = UE_MUTABLE_MAX_OPTIMIZATION; // Force max optimization when packaging.
 	CompileRequest->Options.TextureCompression = ECustomizableObjectTextureCompression::HighQuality;
 	CompileRequest->Options.bIsCooking = true;
-	CompileRequest->Options.bUseBulkData = CVarMutableUseBulkData.GetValueOnAnyThread();
 	CompileRequest->Options.TargetPlatform = TargetPlatform;
 
 	if (CVarMutableDerivedDataCacheUsage.GetValueOnAnyThread())
diff --git a/Engine/Plugins/Mutable/Source/MutableTools/Internal/MuT/Compiler.h b/Engine/Plugins/Mutable/Source/MutableTools/Internal/MuT/Compiler.h
--- a/Engine/Plugins/Mutable/Source/MutableTools/Internal/MuT/Compiler.h
+++ b/Engine/Plugins/Mutable/Source/MutableTools/Internal/MuT/Compiler.h
@@ -66,7 +66,7 @@ namespace UE::Mutable::Private
 		UE_API void SetImageCompressionQuality(int32 Quality);
 
         //! 
-        UE_API void SetDataPackingStrategy( int32 MinTextureResidentMipCount, uint64 EmbeddedDataBytesLimit, uint64 PackagedDataBytesLimit);
+        UE_API void SetDataPackingStrategy( int32 MinTextureResidentMipCount, uint64 EmbeddedDataBytesLimit);
 
 		/** If enabled it will make sure that the object is compile to generate smaller mips of the images. */
 		UE_API void SetEnableProgressiveImages(bool bEnabled);
diff --git a/Engine/Plugins/Mutable/Source/MutableTools/Private/MuT/Compiler.cpp b/Engine/Plugins/Mutable/Source/MutableTools/Private/MuT/Compiler.cpp
--- a/Engine/Plugins/Mutable/Source/MutableTools/Private/MuT/Compiler.cpp
+++ b/Engine/Plugins/Mutable/Source/MutableTools/Private/MuT/Compiler.cpp
@@ -285,10 +285,9 @@ namespace UE::Mutable::Private
 
 
 	//---------------------------------------------------------------------------------------------
-	void CompilerOptions::SetDataPackingStrategy(int32 MinTextureResidentMipCount, uint64 EmbeddedDataBytesLimit, uint64 PackagedDataBytesLimit)
+	void CompilerOptions::SetDataPackingStrategy(int32 MinTextureResidentMipCount, uint64 EmbeddedDataBytesLimit)
 	{
 		m_pD->EmbeddedDataBytesLimit = EmbeddedDataBytesLimit;
-		m_pD->PackagedDataBytesLimit = PackagedDataBytesLimit;
 		m_pD->MinTextureResidentMipCount = MinTextureResidentMipCount;
 	}
 
diff --git a/Engine/Plugins/Mutable/Source/MutableTools/Private/MuT/CompilerPrivate.h b/Engine/Plugins/Mutable/Source/MutableTools/Private/MuT/CompilerPrivate.h
--- a/Engine/Plugins/Mutable/Source/MutableTools/Private/MuT/CompilerPrivate.h
+++ b/Engine/Plugins/Mutable/Source/MutableTools/Private/MuT/CompilerPrivate.h
@@ -34,7 +34,6 @@ namespace UE::Mutable::Private
         FModelOptimizationOptions OptimisationOptions;
 
 		uint64 EmbeddedDataBytesLimit = 1024;
-		uint64 PackagedDataBytesLimit = 1024*1024*64;
 
 		// \TODO: Unused?
 		int32 MinTextureResidentMipCount = 3;
diff --git a/Engine/Plugins/Mutable/Source/MutableValidation/Private/MuV/CustomizableObjectCompilationUtility.cpp b/Engine/Plugins/Mutable/Source/MutableValidation/Private/MuV/CustomizableObjectCompilationUtility.cpp
--- a/Engine/Plugins/Mutable/Source/MutableValidation/Private/MuV/CustomizableObjectCompilationUtility.cpp
+++ b/Engine/Plugins/Mutable/Source/MutableValidation/Private/MuV/CustomizableObjectCompilationUtility.cpp
@@ -56,7 +56,6 @@ bool FCustomizableObjectCompilationUtility::CompileCustomizableObject(UCustomiza
 		UE_LOG(LogMutable, Log, TEXT("(int) model_optimization_level : %d "), CompilationOptions.OptimizationLevel);
 		UE_LOG(LogMutable, Log, TEXT("(string) model_texture_compression : %s "), *UEnum::GetValueAsString(CompilationOptions.TextureCompression));
 		UE_LOG(LogMutable, Log, TEXT("(string) model_compile_platform_name : %s "), *CompilationOptions.TargetPlatform->PlatformName());
-		UE_LOG(LogMutable, Log, TEXT("(int) model_package_data_bytes_limit_bytes : %llu "), CompilationOptions.PackagedDataBytesLimit);
 	}
 
 	// Delete data from previous compilations just to be sure it does not affect the new compilation operation
