[Mutable] Update checks used to validate if LODRenderData from a SkeletalMesh Parameter can be accessed.

#rnx

#Synced-CL 49944882

diff --git a/Engine/Plugins/Mutable/Source/CustomizableObject/Private/MuCO/UnrealConversionUtils.cpp b/Engine/Plugins/Mutable/Source/CustomizableObject/Private/MuCO/UnrealConversionUtils.cpp
--- a/Engine/Plugins/Mutable/Source/CustomizableObject/Private/MuCO/UnrealConversionUtils.cpp
+++ b/Engine/Plugins/Mutable/Source/CustomizableObject/Private/MuCO/UnrealConversionUtils.cpp
@@ -1097,10 +1097,17 @@ namespace UnrealConversionUtils
 				}
 					
 				Context->OriginalLOD = LOD;
-				Context->DataLOD = LOD;
+#if WITH_EDITOR
+				const bool bNeecCPUData = true;
+#else
+				const bool bNeecCPUData = SkeletalMesh->NeedCPUData(LODIndex);
+#endif
 
-				const bool bUseStreaming = !LOD->bStreamedDataInlined;
-				if (bUseStreaming)
+				if(LOD->bStreamedDataInlined && bNeecCPUData)
+				{
+					Context->DataLOD = LOD;
+				}
+				else if (LOD->StreamingBulkData.CanLoadFromDisk())
 				{
 					MUTABLE_CPUPROFILER_SCOPE(MeshStreaming);
 
@@ -1147,14 +1154,21 @@ namespace UnrealConversionUtils
 			{
 				MUTABLE_CPUPROFILER_SCOPE(MeshConversion);
 
-				if (!SkeletalMesh->NeedCPUData(LODIndex))
+				if (!Context->OriginalLOD)
 				{
-					UE_LOG(LogMutable, Error, TEXT("Unable to read LOD %d Data from the CPU. Enable bAllowCPUAccess on the Skeletal Mesh %s."), LODIndex, *SkeletalMesh->GetName());
 					return;
 				}
-					
-				if (!Context->OriginalLOD)
+
+				if (!Context->DataLOD)
 				{
+					if (Context->OriginalLOD->bStreamedDataInlined && !SkeletalMesh->NeedCPUData(LODIndex))
+					{
+						UE_LOG(LogMutable, Error, TEXT("Unable to read LOD %d Data from the CPU. Enable bAllowCPUAccess on the Skeletal Mesh %s."), LODIndex, *SkeletalMesh->GetName());
+					}
+					else
+					{
+						UE_LOG(LogMutable, Error, TEXT("Unable to read LOD %d Data from Skeletal Mesh %s."), LODIndex, *SkeletalMesh->GetName());
+					}
 					return;
 				}
 					
