[Mutable] Remove LOD limiters.

#Synced-CL 49839474

#Synced-CL 49839474

#Synced-CL 49839474

diff --git a/Engine/Plugins/Mutable/Source/CustomizableObject/Internal/MuCO/CustomizableObjectSystemPrivate.h b/Engine/Plugins/Mutable/Source/CustomizableObject/Internal/MuCO/CustomizableObjectSystemPrivate.h
--- a/Engine/Plugins/Mutable/Source/CustomizableObject/Internal/MuCO/CustomizableObjectSystemPrivate.h
+++ b/Engine/Plugins/Mutable/Source/CustomizableObject/Internal/MuCO/CustomizableObjectSystemPrivate.h
@@ -855,9 +855,7 @@ public:
 	UE::Tasks::FTask LastUpdateMutableTask = UE::Tasks::MakeCompletedTask<void>();
 	
 	FLogBenchmarkUtil LogBenchmarkUtil;
-
-	int32 NumSkeletalMeshes = 0;
-
+	
 	bool bAutoCompileCommandletEnabled = false;
 
 	bool bPendingLODUpdate = false;
diff --git a/Engine/Plugins/Mutable/Source/CustomizableObject/Private/MuCO/CustomizableInstanceLODManagement.cpp b/Engine/Plugins/Mutable/Source/CustomizableObject/Private/MuCO/CustomizableInstanceLODManagement.cpp
--- a/Engine/Plugins/Mutable/Source/CustomizableObject/Private/MuCO/CustomizableInstanceLODManagement.cpp
+++ b/Engine/Plugins/Mutable/Source/CustomizableObject/Private/MuCO/CustomizableInstanceLODManagement.cpp
@@ -25,30 +25,6 @@
 #include "LevelEditorViewport.h"
 #endif
 
-static TAutoConsoleVariable<int32> CVarNumGeneratedInstancesLimit(
-	TEXT("b.NumGeneratedInstancesLimit"),
-	0,
-	TEXT("If different than 0, limit the number of mutable instances with full LODs to have at a given time."),
-	ECVF_Scalability);
-
-static TAutoConsoleVariable<int32> CVarNumGeneratedInstancesLimitLOD1(
-	TEXT("b.NumGeneratedInstancesLimitLOD1"),
-	0,
-	TEXT("If different than 0, limit the number of mutable instances with LOD 1 to have at a given time."),
-	ECVF_Scalability);
-
-static TAutoConsoleVariable<int32> CVarNumGeneratedInstancesLimitLOD2(
-	TEXT("b.NumGeneratedInstancesLimitLOD2"),
-	0,
-	TEXT("If different than 0, limit the number of mutable instances with LOD 2 to have at a given time."),
-	ECVF_Scalability);
-
-static TAutoConsoleVariable<int32> CVarDistanceForFixedLOD2(
-	TEXT("b.DistanceForFixedLOD2"),
-	50000,
-	TEXT("If NumGeneratedInstancesLimit is different than 0, sets the distance at which the system will fix the LOD of an instance to the lowest res one (LOD2) to prevent unnecessary LOD changes and memory consumption"),
-	ECVF_Scalability);
-
 static TAutoConsoleVariable<bool> CVarOnlyUpdateCloseCustomizableObjects(
 	TEXT("b.OnlyUpdateCloseCustomizableObjects"),
 	false,
@@ -117,8 +93,7 @@ void UpdateCameraToInstancesDistance(const FVector CameraPosition)
 
 void UCustomizableInstanceLODManagement::UpdateInstanceDistsAndLODs(FMutableInstanceUpdateMap& InOutRequestedUpdates)
 {
-	int32 NumGeneratedInstancesLimitLODs = GetNumGeneratedInstancesLimitFullLODs();
-	if (NumGeneratedInstancesLimitLODs > 0 || (IsOnlyUpdateCloseCustomizableObjectsEnabled() && IsOnlyGenerateRequestedLODLevelsEnabled()))
+	if (IsOnlyUpdateCloseCustomizableObjectsEnabled() && IsOnlyGenerateRequestedLODLevelsEnabled())
 	{
 		UCustomizableObjectSystem* CustomizableObjectSystem = UCustomizableObjectSystem::GetInstance();
 
@@ -193,142 +168,7 @@ void UCustomizableInstanceLODManagement::UpdateInstanceDistsAndLODs(FMutableInst
 		}
 	}
 
-	if (NumGeneratedInstancesLimitLODs > 0)
-	{
-		int32 NumGeneratedInstancesLimitLOD1 = GetNumGeneratedInstancesLimitLOD1();
-		int32 NumGeneratedInstancesLimitLOD2 = GetNumGeneratedInstancesLimitLOD2();
-
-		TArray<UCustomizableObjectInstance*> SortedInstances;
-
-		for (TObjectIterator<UCustomizableObjectInstance> CustomizableObjectInstance; CustomizableObjectInstance; ++CustomizableObjectInstance)
-		{
-			if (IsValid(*CustomizableObjectInstance) &&
-				CustomizableObjectInstance->GetPrivate() &&
-				CustomizableObjectInstance->GetIsBeingUsedByComponentInPlay())
-			{
-				if (const UCustomizableObject* CustomizableObject = CustomizableObjectInstance->GetCustomizableObject();
-					CustomizableObject &&
-					!CustomizableObject->GetPrivate()->IsLocked())
-				{
-					UCustomizableObjectInstance* Ptr = *CustomizableObjectInstance;
-					Ptr->SetIsDiscardedBecauseOfTooManyInstances(false);
-					SortedInstances.Add(Ptr);
-				}
-			}
-		}
-
-		for (int32 i = 0; i < GetNumberOfPriorityUpdateInstances() && i < SortedInstances.Num(); ++i)
-		{
-			SortedInstances[i]->SetIsPlayerOrNearIt(true);
-		}
-
-		TMap<FName, uint8> MinLODs;
-		TMap<FName, uint8> RequestedLODs;
-		for (int32 i = NumGeneratedInstancesLimitLODs + NumGeneratedInstancesLimitLOD1 + NumGeneratedInstancesLimitLOD2; i < SortedInstances.Num(); ++i)
-		{
-			SortedInstances[i]->SetIsDiscardedBecauseOfTooManyInstances(true);
-		}
-
-		if (SortedInstances.Num() > NumGeneratedInstancesLimitLODs)
-		{
-			SortedInstances.Sort([](UCustomizableObjectInstance& A, UCustomizableObjectInstance& B)
-			{
-				return A.GetMinSquareDistToPlayer() < B.GetMinSquareDistToPlayer();
-			});
-
-			bool bAlreadyReachedFixedLOD = false;
-			const float DistanceForFixedLOD = float(CVarDistanceForFixedLOD2.GetValueOnGameThread());
-			const float DistanceForFixedLODSquared = FMath::Square(DistanceForFixedLOD);
-
-
-			for (int32 i = 0; i < NumGeneratedInstancesLimitLODs && i < SortedInstances.Num(); ++i)
-			{
-				const UCustomizableObject* CO = SortedInstances[i]->GetCustomizableObject();
-				
-				if (!bAlreadyReachedFixedLOD && SortedInstances[i]->GetMinSquareDistToPlayer() < DistanceForFixedLODSquared)
-				{
-					for (int32 ComponentIndex = 0; ComponentIndex < CO->GetComponentCount(); ++ComponentIndex)
-					{
-						MinLODs.Add(CO->GetComponentName(ComponentIndex), 0);
-						RequestedLODs.Add(CO->GetComponentName(ComponentIndex), 0);
-					}
-
-					SortedInstances[i]->SetRequestedLODs(MinLODs, RequestedLODs, InOutRequestedUpdates);
-				}
-				else
-				{
-					for (int32 ComponentIndex = 0; ComponentIndex < CO->GetComponentCount(); ++ComponentIndex)
-					{
-						MinLODs.Add(CO->GetComponentName(ComponentIndex), UINT8_MAX);
-						RequestedLODs.Add(CO->GetComponentName(ComponentIndex), 0);
-					}
-
-					SortedInstances[i]->SetRequestedLODs(MinLODs, RequestedLODs, InOutRequestedUpdates);
-					bAlreadyReachedFixedLOD = true;
-				}
-			}
-
-			for (int32 i = NumGeneratedInstancesLimitLODs; i < NumGeneratedInstancesLimitLODs + NumGeneratedInstancesLimitLOD1 && i < SortedInstances.Num(); ++i)
-			{
-				const UCustomizableObject* CO = SortedInstances[i]->GetCustomizableObject();
-
-				if (!bAlreadyReachedFixedLOD && SortedInstances[i]->GetMinSquareDistToPlayer() < DistanceForFixedLODSquared)
-				{
-					for (int32 ComponentIndex = 0; ComponentIndex < CO->GetComponentCount(); ++ComponentIndex)
-					{
-						MinLODs.Add(CO->GetComponentName(ComponentIndex), 1);
-						RequestedLODs.Add(CO->GetComponentName(ComponentIndex), 0);
-					}
-					
-					SortedInstances[i]->SetRequestedLODs(MinLODs, RequestedLODs, InOutRequestedUpdates);
-				}
-				else
-				{
-					for (int32 ComponentIndex = 0; ComponentIndex < CO->GetComponentCount(); ++ComponentIndex)
-					{
-						MinLODs.Add(CO->GetComponentName(ComponentIndex), UINT8_MAX);
-						RequestedLODs.Add(CO->GetComponentName(ComponentIndex), 0);
-					}
-					
-					SortedInstances[i]->SetRequestedLODs(MinLODs, RequestedLODs, InOutRequestedUpdates);
-					bAlreadyReachedFixedLOD = true;
-				}
-			}
-
-			for (int32 i = NumGeneratedInstancesLimitLODs + NumGeneratedInstancesLimitLOD1; i < NumGeneratedInstancesLimitLODs + NumGeneratedInstancesLimitLOD1 + NumGeneratedInstancesLimitLOD2 && i < SortedInstances.Num(); ++i)
-			{
-				const UCustomizableObject* CO = SortedInstances[i]->GetCustomizableObject();
-				for (int32 ComponentIndex = 0; ComponentIndex < CO->GetComponentCount(); ++ComponentIndex)
-				{
-					MinLODs.Add(CO->GetComponentName(ComponentIndex), 2);
-					RequestedLODs.Add(CO->GetComponentName(ComponentIndex), 0);
-				}
-
-				SortedInstances[i]->SetRequestedLODs(MinLODs, RequestedLODs, InOutRequestedUpdates);
-			}
-
-			for (int32 i = NumGeneratedInstancesLimitLODs + NumGeneratedInstancesLimitLOD1 + NumGeneratedInstancesLimitLOD2; i < SortedInstances.Num(); ++i)
-			{
-				SortedInstances[i]->SetIsDiscardedBecauseOfTooManyInstances(true);
-			}
-		}
-		else
-		{
-			// No limit surpassed, set all instances to have all LODs, there will be an UpdateSkeletalMesh only if there's a change in LOD state
-			for (int32 i = 0; i < SortedInstances.Num(); ++i)
-			{
-				const UCustomizableObject* CO = SortedInstances[i]->GetCustomizableObject();
-				for (int32 ComponentIndex = 0; ComponentIndex < CO->GetComponentCount(); ++ComponentIndex)
-				{
-					MinLODs.Add(CO->GetComponentName(ComponentIndex), 2);
-					RequestedLODs.Add(CO->GetComponentName(ComponentIndex), 0);
-				}
-				
-				SortedInstances[i]->SetRequestedLODs(MinLODs, RequestedLODs, InOutRequestedUpdates);
-			}
-		}
-	}
-	else if (IsOnlyGenerateRequestedLODLevelsEnabled())
+	if (IsOnlyGenerateRequestedLODLevelsEnabled())
 	{
 		struct FLODTracker
 		{
@@ -464,24 +304,6 @@ void UCustomizableInstanceLODManagement::UpdateInstanceDistsAndLODs(FMutableInst
 }
 
 
-int32 UCustomizableInstanceLODManagement::GetNumGeneratedInstancesLimitFullLODs() const
-{
-	return CVarNumGeneratedInstancesLimit.GetValueOnGameThread();
-}
-
-
-int32 UCustomizableInstanceLODManagement::GetNumGeneratedInstancesLimitLOD1() const
-{
-	return CVarNumGeneratedInstancesLimitLOD1.GetValueOnGameThread();
-}
-
-
-int32 UCustomizableInstanceLODManagement::GetNumGeneratedInstancesLimitLOD2() const
-{
-	return CVarNumGeneratedInstancesLimitLOD2.GetValueOnGameThread();
-}
-
-
 void UCustomizableInstanceLODManagement::SetNumberOfPriorityUpdateInstances(int32 InNumPriorityUpdateInstances)
 {
 	NumPriorityUpdateInstances = InNumPriorityUpdateInstances;
diff --git a/Engine/Plugins/Mutable/Source/CustomizableObject/Private/MuCO/CustomizableObjectSystem.cpp b/Engine/Plugins/Mutable/Source/CustomizableObject/Private/MuCO/CustomizableObjectSystem.cpp
--- a/Engine/Plugins/Mutable/Source/CustomizableObject/Private/MuCO/CustomizableObjectSystem.cpp
+++ b/Engine/Plugins/Mutable/Source/CustomizableObject/Private/MuCO/CustomizableObjectSystem.cpp
@@ -1395,12 +1395,6 @@ void UCustomizableObjectSystemPrivate::SetReplaceDiscardedWithReferenceMeshEnabl
 }
 
 
-int32 UCustomizableObjectSystemPrivate::GetNumSkeletalMeshes() const
-{
-	return NumSkeletalMeshes;
-}
-
-
 EUpdateRequired UCustomizableObjectSystemPrivate::IsUpdateRequired(const UCustomizableObjectInstance& Instance, bool bOnlyUpdateIfNotGenerated, bool bOnlyUpdateIfLODs, bool bIgnoreCloseDist) const
 {
 	UCustomizableObjectSystem* System = UCustomizableObjectSystem::GetInstance();
@@ -1412,16 +1406,6 @@ EUpdateRequired UCustomizableObjectSystemPrivate::IsUpdateRequired(const UCustom
 	}
 
 	const bool bIsGenerated = Private->SkeletalMeshStatus != ESkeletalMeshStatus::NotGenerated;
-	const int32 NumGeneratedInstancesLimit = System->GetInstanceLODManagement()->GetNumGeneratedInstancesLimitFullLODs();
-	const int32 NumGeneratedInstancesLimitLOD1 = System->GetInstanceLODManagement()->GetNumGeneratedInstancesLimitLOD1();
-	const int32 NumGeneratedInstancesLimitLOD2 = System->GetInstanceLODManagement()->GetNumGeneratedInstancesLimitLOD2();
-
-	if (!bIsGenerated && // Prevent generating more instances than the limit, but let updates to existing instances run normally
-		NumGeneratedInstancesLimit > 0 &&
-		System->GetPrivate()->GetNumSkeletalMeshes() > NumGeneratedInstancesLimit + NumGeneratedInstancesLimitLOD1 + NumGeneratedInstancesLimitLOD2)
-	{
-		return EUpdateRequired::NoUpdate;
-	}
 
 	const bool bDiscardByDistance = Private->LastMinSquareDistFromComponentToPlayer > FMath::Square(System->GetInstanceLODManagement()->GetOnlyUpdateCloseCustomizableObjectsDist());
 	const bool bLODManagementDiscard = System->GetInstanceLODManagement()->IsOnlyUpdateCloseCustomizableObjectsEnabled() &&
@@ -3366,18 +3350,6 @@ namespace impl
 
 void UCustomizableObjectSystemPrivate::UpdateTick(const bool bBlocking)
 {
-	NumSkeletalMeshes = 0;
-
-	for (TObjectIterator<UCustomizableObjectInstance> Instance; Instance; ++Instance)
-{
-		if (!IsValid(*Instance))
-	{
-			continue;
-		}
-
-		NumSkeletalMeshes += Instance->GetPrivate()->SkeletalMeshes.Num();
-				}
-
 	FMutableUpdateCandidate* LODUpdateCandidateFound = nullptr;
 
 	for (TObjectIterator<UCustomizableObjectInstanceUsage> CustomizableObjectInstanceUsageItr; CustomizableObjectInstanceUsageItr; ++CustomizableObjectInstanceUsageItr)
diff --git a/Engine/Plugins/Mutable/Source/CustomizableObject/Public/MuCO/CustomizableInstanceLODManagement.h b/Engine/Plugins/Mutable/Source/CustomizableObject/Public/MuCO/CustomizableInstanceLODManagement.h
--- a/Engine/Plugins/Mutable/Source/CustomizableObject/Public/MuCO/CustomizableInstanceLODManagement.h
+++ b/Engine/Plugins/Mutable/Source/CustomizableObject/Public/MuCO/CustomizableInstanceLODManagement.h
@@ -21,9 +21,6 @@ public:
 
 	// WARNING! The following methods must be overriden in derived classes
 	virtual void UpdateInstanceDistsAndLODs(FMutableInstanceUpdateMap& InOutRequestedUpdates) PURE_VIRTUAL(UCustomizableInstanceLODManagementBase::UpdateInstanceDistsAndLODs, return;);
-	virtual int32 GetNumGeneratedInstancesLimitFullLODs() const PURE_VIRTUAL(UCustomizableInstanceLODManagementBase::GetNumGeneratedInstancesLimitFullLODs, return 0;);
-	virtual int32 GetNumGeneratedInstancesLimitLOD1() const PURE_VIRTUAL(UCustomizableInstanceLODManagementBase::GetNumGeneratedInstancesLimitLOD1, return 0;);
-	virtual int32 GetNumGeneratedInstancesLimitLOD2() const PURE_VIRTUAL(UCustomizableInstanceLODManagementBase::GetNumGeneratedInstancesLimitLOD2, return 0;);
 	virtual float GetOnlyUpdateCloseCustomizableObjectsDist() const PURE_VIRTUAL(UCustomizableInstanceLODManagementBase::GetOnlyUpdateCloseCustomizableObjectsDist, return 0.f;);
 	virtual bool IsOnlyUpdateCloseCustomizableObjectsEnabled() const PURE_VIRTUAL(UCustomizableInstanceLODManagementBase::IsOnlyUpdateCloseCustomizableObjectsEnabled, return false;);
 	virtual bool IsOnlyGenerateRequestedLODLevelsEnabled() const PURE_VIRTUAL(UCustomizableInstanceLODManagementBase::IsOnlyGenerateRequestedLODLevelsEnabled, return false;);
@@ -43,11 +40,7 @@ public:
 	UE_API virtual ~UCustomizableInstanceLODManagement();
 
 	UE_API virtual void UpdateInstanceDistsAndLODs(FMutableInstanceUpdateMap& InOutRequestedUpdates) override;
-
-	UE_API virtual int32 GetNumGeneratedInstancesLimitFullLODs() const override;
-	UE_API virtual int32 GetNumGeneratedInstancesLimitLOD1() const override;
-	UE_API virtual int32 GetNumGeneratedInstancesLimitLOD2() const override;
-
+	
 	// Used to define the actors that will be considered the centers of LOD distance calculations. If left empty it will revert to the default, the first player controller actor
 	void AddViewCenter(const AActor* const InCentralActor) { ViewCenters.Add(TWeakObjectPtr<const AActor>(InCentralActor)); };
 	void RemoveViewCenter(const AActor* const InCentralActor) { ViewCenters.Remove(TWeakObjectPtr<const AActor>(InCentralActor)); };
