[Mutable] Fix FMesh serialization crash due to incorrect file offsets

#rnx

#Synced-CL 49944882

diff --git a/Engine/Plugins/Mutable/Source/CustomizableObject/Internal/MuCO/CustomizableObjectPrivate.h b/Engine/Plugins/Mutable/Source/CustomizableObject/Internal/MuCO/CustomizableObjectPrivate.h
--- a/Engine/Plugins/Mutable/Source/CustomizableObject/Internal/MuCO/CustomizableObjectPrivate.h
+++ b/Engine/Plugins/Mutable/Source/CustomizableObject/Internal/MuCO/CustomizableObjectPrivate.h
@@ -1115,6 +1115,15 @@ namespace UE::Mutable::Private
 		const ITargetPlatform& TargetPlatform,
 		TArray<FFile>& OutBulkDataFiles);
 
+	/** Generate the list of BulkData files with a soft restriction to the size of the files.
+	 */
+	void CUSTOMIZABLEOBJECT_API GenerateBulkDataFilesListWithSizeLimit(
+		TSharedPtr<const UE::Mutable::Private::FModel, ESPMode::ThreadSafe> Model,
+		FModelStreamableBulkData& ModelStreamableBulkData,
+		const ITargetPlatform* TargetPlatform,
+		uint64 TargetBulkDataFileBytes,
+		TArray<FFile>& OutBulkDataFiles);
+
 	/** Compute the number of files and sizes the BulkData will be split into and update
 	 * the streamable's FileIds and Offsets.
 	 */
@@ -1695,6 +1704,8 @@ enum class ECustomizableObjectVersions : int32
 	VersionBumpBackout,
 	
 	ComponentId2,
+
+	BackoutGenerateBulkFileWithDataLimit,
 	
 	// -----<new versions can be added above this line>--------
 	LastCustomizableObjectVersion
diff --git a/Engine/Plugins/Mutable/Source/CustomizableObject/Private/MuCO/CustomizableObject.cpp b/Engine/Plugins/Mutable/Source/CustomizableObject/Private/MuCO/CustomizableObject.cpp
--- a/Engine/Plugins/Mutable/Source/CustomizableObject/Private/MuCO/CustomizableObject.cpp
+++ b/Engine/Plugins/Mutable/Source/CustomizableObject/Private/MuCO/CustomizableObject.cpp
@@ -2803,6 +2803,84 @@ namespace UE::Mutable::Private
 	}
 	
 
+	void GenerateBulkDataFilesListWithSizeLimit(
+		TSharedPtr<const UE::Mutable::Private::FModel, ESPMode::ThreadSafe> Model,
+		FModelStreamableBulkData& ModelStreamableBulkData,
+		const ITargetPlatform* TargetPlatform,
+		uint64 ,
+		TArray<FFile>& OutBulkDataFiles)
+	{
+		MUTABLE_CPUPROFILER_SCOPE(GenerateBulkDataFilesListWithSizeLimit);
+
+		if (!Model)
+		{
+			return;
+		}
+
+
+		// Unlimited number of files, limited file size. Add blocks to the file if the size limit won't be surpassed. Add at least one block to each file. 
+		const auto CreateFileList = [](const FFileCategoryID& CategoryID, const FClassifyNode& Node, TArray<FFile>& OutBulkDataFiles)
+			{
+				// Temp: Group by order in the array
+				for (int32 BlockIndex = 0; BlockIndex < Node.Blocks.Num(); )
+				{
+					FFile File;
+					File.DataType = CategoryID.DataType;
+					File.ResourceType = CategoryID.ResourceType;
+					File.Flags = CategoryID.Flags;
+
+					uint64 FileSize = 0;
+					uint32 FileId = uint32(CategoryID.DataType);
+
+					while (BlockIndex < Node.Blocks.Num())
+					{
+						const FBlock& CurrentBlock = Node.Blocks[BlockIndex];
+
+						// Block added to file. Set offset and increase file size.
+						FileSize += CurrentBlock.Size;
+
+						// Generate cumulative id for this file
+						FileId = HashCombine(FileId, CurrentBlock.Id);
+
+						// Add the block to the current file
+						File.Blocks.Add(CurrentBlock);
+
+						// Next block
+						++BlockIndex;
+					}
+
+					const int32 NumFiles = OutBulkDataFiles.Num();
+
+					// Ensure the FileId is unique
+					bool bUnique = false;
+					while (!bUnique)
+					{
+						bUnique = true;
+						for (int32 PreviousFileIndex = 0; PreviousFileIndex < NumFiles; ++PreviousFileIndex)
+						{
+							if (OutBulkDataFiles[PreviousFileIndex].Id == FileId)
+							{
+								bUnique = false;
+								++FileId;
+								break;
+							}
+						}
+					}
+
+					// Set it to the editor-only file descriptor
+					File.Id = FileId;
+
+					OutBulkDataFiles.Add(MoveTemp(File));
+				}
+			};
+
+		// TODO: Temp. Remove after unifying generated output files code between editor an package. UE-222777
+		const bool bUseRomTypeAndFlagsToFilter = TargetPlatform->RequiresCookedData();
+
+		GenerateBulkDataFilesList(Model, ModelStreamableBulkData, bUseRomTypeAndFlagsToFilter, CreateFileList, OutBulkDataFiles);
+	}
+
+
 	void GenerateBulkDataFilesList(
 		TSharedPtr<const UE::Mutable::Private::FModel, ESPMode::ThreadSafe> Model,
 		FModelStreamableBulkData& ModelStreamableBulkData,
diff --git a/Engine/Plugins/Mutable/Source/CustomizableObjectEditor/Private/MuCOE/CustomizableObjectCompileRunnable.cpp b/Engine/Plugins/Mutable/Source/CustomizableObjectEditor/Private/MuCOE/CustomizableObjectCompileRunnable.cpp
--- a/Engine/Plugins/Mutable/Source/CustomizableObjectEditor/Private/MuCOE/CustomizableObjectCompileRunnable.cpp
+++ b/Engine/Plugins/Mutable/Source/CustomizableObjectEditor/Private/MuCOE/CustomizableObjectCompileRunnable.cpp
@@ -489,8 +489,15 @@ void FCustomizableObjectSaveDDRunnable::CachePlatformData()
 	// Cache ModelStreamables
 	{
 		// Generate list of files and update streamable blocks ids and offsets
-		UE::Mutable::Private::GenerateBulkDataFilesListWithFileLimit(CustomizableObjectId, PlatformData->Model, *PlatformData->ModelStreamableBulkData.Get(), MAX_uint8, 
+		if (Options.bIsCooking)
+		{
+			UE::Mutable::Private::GenerateBulkDataFilesListWithFileLimit(CustomizableObjectId, PlatformData->Model, *PlatformData->ModelStreamableBulkData.Get(), MAX_uint8, 
 				Options.bIsCooking, false, *Options.TargetPlatform, PlatformData->BulkDataFiles);
+		}
+		else
+		{
+			UE::Mutable::Private::GenerateBulkDataFilesListWithSizeLimit(PlatformData->Model, *PlatformData->ModelStreamableBulkData.Get(), Options.TargetPlatform, MAX_uint64, PlatformData->BulkDataFiles);
+		}
 	}
 
 	// Cache Model and Model Roms
diff --git a/Engine/Plugins/Mutable/Source/CustomizableObjectEditor/Private/MuCOE/CustomizableObjectCompiler.cpp b/Engine/Plugins/Mutable/Source/CustomizableObjectEditor/Private/MuCOE/CustomizableObjectCompiler.cpp
--- a/Engine/Plugins/Mutable/Source/CustomizableObjectEditor/Private/MuCOE/CustomizableObjectCompiler.cpp
+++ b/Engine/Plugins/Mutable/Source/CustomizableObjectEditor/Private/MuCOE/CustomizableObjectCompiler.cpp
@@ -1495,10 +1495,17 @@ void FCustomizableObjectCompiler::FinishLoadingModelDataFromDDC()
 					}
 
 					// Generate list of files and update streamable blocks ids and offsets
-					const uint32 NumBulkDataFilesPerBucket = MAX_uint8;
-					UE::Mutable::Private::GenerateBulkDataFilesListWithFileLimit(GenerateDataDistributionIdentifier(*CustomizableObject), PlatformData->Model, *PlatformData->ModelStreamableBulkData.Get(),
+					if (Options.bIsCooking)
+					{
+						const uint32 NumBulkDataFilesPerBucket = MAX_uint8;
+						UE::Mutable::Private::GenerateBulkDataFilesListWithFileLimit(GenerateDataDistributionIdentifier(*CustomizableObject), PlatformData->Model, *PlatformData->ModelStreamableBulkData.Get(),
 							NumBulkDataFilesPerBucket, Options.bIsCooking, false, *Options.TargetPlatform, PlatformData->BulkDataFiles);
-					
+					}
+					else
+					{
+						UE::Mutable::Private::GenerateBulkDataFilesListWithSizeLimit(PlatformData->Model, *PlatformData->ModelStreamableBulkData.Get(),
+							Options.TargetPlatform, 0 /* unused */, PlatformData->BulkDataFiles);
+					}
 				}
 				else
 				{
