[mutable] Fixed issue with the FProgramCache that caused by the wrong interpretation of the FOpExecutionData::DataTypeIndex.
- The FOpExecutionData can be used to hold values instead of indices to the array of data part of the Cache. The data held there was being used, for all cases (the ones where it points to an index and the ones where it directly contains data) to determine what entries should be set as invalid (ClearLayer0). It was preventing the ones that held data from being set as "IsValueValid = false", which, indirectly, made their data be relevant during the update (preventing the reevaluation of some branches of the runtime graph).
- This issue could be reproduced when "mutable.ReuseLiveInstanceCacheBetweenUpdates" is set to true. the LiveUpdate option is enabled and when the parameter we are working with is set as a parameter of the current state (runtime parameter).

#rnx
#Synced-CL 50038443

diff --git a/Engine/Plugins/Mutable/Source/MutableRuntime/Internal/MuR/System.h b/Engine/Plugins/Mutable/Source/MutableRuntime/Internal/MuR/System.h
--- a/Engine/Plugins/Mutable/Source/MutableRuntime/Internal/MuR/System.h
+++ b/Engine/Plugins/Mutable/Source/MutableRuntime/Internal/MuR/System.h
@@ -660,7 +660,7 @@ namespace UE::Mutable::Private
 			for (; it.IsValid(); ++it)
 			{
 				FProgramCache::FOpExecutionData& Data = *it;
-				if (!Data.DataTypeIndex || Data.IsCacheLocked)
+				if (Data.IsCacheLocked)
 				{
 					continue;
 				}
@@ -668,41 +668,51 @@ namespace UE::Mutable::Private
 				switch (Data.DataType)
 				{
 				case EDataType::Image:
-				{
-					check(ImageResults.Num() && StrongImageResults.Num());
-					ImageResults[Data.DataTypeIndex] = nullptr;
-					StrongImageResults[Data.DataTypeIndex] = nullptr;
-					break;
-				}
+					{
+						if (Data.DataTypeIndex)
+						{
+							ImageResults[Data.DataTypeIndex] = nullptr;
+							StrongImageResults[Data.DataTypeIndex] = nullptr;
+						}
+						break;
+					}
 
 				case EDataType::Mesh:
-				{
-					check(MeshResults.Num() && StrongMeshResults.Num());
-					MeshResults[Data.DataTypeIndex] = nullptr;
-					StrongMeshResults[Data.DataTypeIndex] = nullptr;
-					break;
-				}
+					{
+						if (Data.DataTypeIndex)
+						{
+							MeshResults[Data.DataTypeIndex] = nullptr;
+							StrongMeshResults[Data.DataTypeIndex] = nullptr;
+						}
+						break;
+					}
 
 				case EDataType::SkeletalMesh:
-				{
-					check(SkeletalMeshResults.Num());
-					SkeletalMeshResults[Data.DataTypeIndex] = nullptr;
-					break;
-				}
+					{
+						if (Data.DataTypeIndex)
+						{
+							SkeletalMeshResults[Data.DataTypeIndex] = nullptr;
+						}
+						break;
+					}
 
 				case EDataType::Layout:
-				{
-					check(LayoutResults.Num());
-					LayoutResults[Data.DataTypeIndex] = nullptr;
-					break;
-				}
+					{
+						if (Data.DataTypeIndex)
+						{
+							LayoutResults[Data.DataTypeIndex] = nullptr;
+						}
+						break;
+					}
 
 				case EDataType::Instance:
-				{
-					check(InstanceResults.Num());
-					InstanceResults[Data.DataTypeIndex] = nullptr;
-					break;
-				}
+					{
+						if (Data.DataTypeIndex)
+						{
+							InstanceResults[Data.DataTypeIndex] = nullptr;
+						}
+						break;
+					}
 
 				default:
 					break;
