[Mutable] Replace BoneIds with BoneNames for runtime operations.
- Resources serialize BoneIds to favour deterministic cooks.
- Resources convert BoneIds to FNames or indices upon load.
- Optimize mesh conversion at runtime.

#rnx

#Synced-CL 49999995

diff --git a/Engine/Plugins/Mutable/Source/CustomizableObject/Internal/MuCO/BoneNames.h b/Engine/Plugins/Mutable/Source/CustomizableObject/Internal/MuCO/BoneNames.h
deleted file mode 100644
diff --git a/Engine/Plugins/Mutable/Source/CustomizableObject/Internal/MuCO/CustomizableObjectPrivate.h b/Engine/Plugins/Mutable/Source/CustomizableObject/Internal/MuCO/CustomizableObjectPrivate.h
--- a/Engine/Plugins/Mutable/Source/CustomizableObject/Internal/MuCO/CustomizableObjectPrivate.h
+++ b/Engine/Plugins/Mutable/Source/CustomizableObject/Internal/MuCO/CustomizableObjectPrivate.h
@@ -26,7 +26,6 @@ class UCustomizableObjectSkeletalMesh;
 namespace UE::Mutable::Private 
 { 
 	class FModel;
-	struct FBoneName;
 }
 
 #if WITH_EDITOR
@@ -266,6 +265,8 @@ public:
 
 	void Add(const TArray<UE::Mutable::Private::FMeshId>& Key, UCustomizableObjectSkeletalMesh* Value);
 
+	void Clear();
+
 private:
 	using FMeshCacheRegistry = TRegistry<TArray<UE::Mutable::Private::FMeshId>, TWeakObjectPtr<UCustomizableObjectSkeletalMesh>>;
 	
@@ -299,6 +300,8 @@ public:
 
 	void Remove(const UTexture2D& Value);
 
+	void Clear();
+
 private:
 	TMap<FId, TWeakObjectPtr<UTexture2D>> GeneratedTextures;
 };
@@ -311,6 +314,8 @@ public:
 
 	void Add(const TArray<uint16>& Key, USkeleton* Value);
 
+	void Clear();
+
 private:
 	TMap<TArray<uint16>, TWeakObjectPtr<USkeleton>> MergedSkeletons;
 };
@@ -1258,9 +1263,6 @@ public:
 	UPROPERTY()
 	TArray<FName> MaterialSlotNames;
 
-	UPROPERTY()
-	TMap<FString, uint32> BoneNamesMap;
-
 	/** Mesh sockets provided by the part skeletal meshes, to be merged in the generated meshes */
 	UPROPERTY()
 	TMap<uint32, FMutableRefSocket> Sockets;
@@ -1707,6 +1709,8 @@ enum class ECustomizableObjectVersions : int32
 
 	BackoutGenerateBulkFileWithDataLimit,
 	
+	ConvertBoneIdsToBoneNames,
+
 	// -----<new versions can be added above this line>--------
 	LastCustomizableObjectVersion
 };
diff --git a/Engine/Plugins/Mutable/Source/CustomizableObject/Internal/MuCO/CustomizableObjectSystemPrivate.h b/Engine/Plugins/Mutable/Source/CustomizableObject/Internal/MuCO/CustomizableObjectSystemPrivate.h
--- a/Engine/Plugins/Mutable/Source/CustomizableObject/Internal/MuCO/CustomizableObjectSystemPrivate.h
+++ b/Engine/Plugins/Mutable/Source/CustomizableObject/Internal/MuCO/CustomizableObjectSystemPrivate.h
@@ -32,7 +32,6 @@
 #define UE_API CUSTOMIZABLEOBJECT_API
 
 class FMutableUpdateImageContext;
-class FBoneNames;
 class FUnrealMutableResourceProvider;
 class FMutableStreamRequest;
 class UEditorImageProvider;
@@ -396,15 +395,15 @@ struct FInstanceUpdateData
 	TArray<FVector> Vectors;
 	TArray<FScalar> Scalars;
 
-	TArray<UE::Mutable::Private::FBoneName> ActiveBones;
-	TArray<UE::Mutable::Private::FBoneName> BoneMaps;
+	TArray<FName> ActiveBones;
+	TArray<FName> BoneMaps;
 	
 	struct FBone
 	{
-		UE::Mutable::Private::FBoneName Name;
+		FName BoneName;
 		FMatrix44f MatrixWithScale;
 
-		bool operator==(const UE::Mutable::Private::FBoneName& OtherName) const { return Name == OtherName; };
+		bool operator==(const FName& OtherBoneName) const { return BoneName == OtherBoneName; };
 	};
 	
 	/** Key is the Block Id. Value is the LoadAdditionalAssetsAndData read destination. */
@@ -423,8 +422,6 @@ struct FInstanceUpdateData
 		TArray<uint16> SkeletonIds;
 
 		TArray<FBone> BonePose;
-
-		TMap<UE::Mutable::Private::FBoneName, TPair<FName, uint16>> BoneInfoMap;
 	};
 
 	// Access by instance component index
@@ -568,8 +565,6 @@ public:
 
 	TSharedPtr<FUnrealMutableResourceProvider> ExternalResourceProvider;
 
-	TSharedPtr<FBoneNames> BoneNames;
-
 	TSharedPtr<FMutableUpdateImageContext> UpdateImageContext;
 	
 	/** Manage loading Passthrough Objects discovered on the StartUpdate. */
diff --git a/Engine/Plugins/Mutable/Source/CustomizableObject/Internal/MuCO/UnrealConversionUtils.h b/Engine/Plugins/Mutable/Source/CustomizableObject/Internal/MuCO/UnrealConversionUtils.h
--- a/Engine/Plugins/Mutable/Source/CustomizableObject/Internal/MuCO/UnrealConversionUtils.h
+++ b/Engine/Plugins/Mutable/Source/CustomizableObject/Internal/MuCO/UnrealConversionUtils.h
@@ -15,7 +15,6 @@
 #include "Tasks/Task.h"
 
 
-class FBoneNames;
 struct FClothBufferIndexMapping;
 struct FMeshToMeshVertData;
 struct FCustomizableObjectMeshToMeshVertData;
@@ -26,7 +25,6 @@ namespace UE::Mutable::Private
 {
 	class FMesh;
 	class FMeshBufferSet;
-	struct FBoneName;
 }
 struct FReferenceSkeleton;
 struct FMutableSurfaceMetadata;
@@ -66,8 +64,8 @@ namespace UnrealConversionUtils
 	CUSTOMIZABLEOBJECT_API void SetupRenderSections(
 		FSkeletalMeshLODRenderData& LODResource,
 		const UE::Mutable::Private::FMesh* InMutableMesh,
-		const TArray<UE::Mutable::Private::FBoneName>& InBoneMap,
-		const TMap<UE::Mutable::Private::FBoneName, TPair<FName, uint16>>& BoneInfoMap,
+		const FReferenceSkeleton& InRefSkeleton,
+		const TArray<FName>& InBoneMap,
 		const int32 InFirstBoneMapIndex,
 		const TArray<const FMutableSurfaceMetadata*>& SectionMetadata);
 
@@ -158,7 +156,7 @@ namespace UnrealConversionUtils
 
 	void MorphTargetVertexInfoBuffers(FSkeletalMeshLODRenderData& LODResource, const USkeletalMesh& Owner, const UE::Mutable::Private::FMesh& MutableMesh, const TMap<uint32, FMorphTargetMeshData>& MorphTargetMeshData, int32 LODIndex, bool bGenerateCPUMorphTargets);
 
-	CUSTOMIZABLEOBJECT_API UE::Tasks::FTask ConvertSkeletalMeshFromRuntimeData(USkeletalMesh* SkeletalMesh, int32 LODIndex, int32 SectionIndex, UE::Mutable::Private::TManagedPtr<UE::Mutable::Private::FMesh> Result, const TSharedRef<FBoneNames>& BoneNames);
+	CUSTOMIZABLEOBJECT_API UE::Tasks::FTask ConvertSkeletalMeshFromRuntimeData(USkeletalMesh* SkeletalMesh, int32 LODIndex, int32 SectionIndex, UE::Mutable::Private::TManagedPtr<UE::Mutable::Private::FMesh> Result);
 
 	void GetSectionClothData(const UE::Mutable::Private::FMesh& MutableMesh, int32 LODIndex, const TMap<uint32, FClothingMeshData>& ClothingMeshData, TArray<FSectionClothData>& SectionsClothData, int32& NumClothingDataNotFound);
 	
diff --git a/Engine/Plugins/Mutable/Source/CustomizableObject/Internal/MuCO/UnrealMutableImageProvider.h b/Engine/Plugins/Mutable/Source/CustomizableObject/Internal/MuCO/UnrealMutableImageProvider.h
--- a/Engine/Plugins/Mutable/Source/CustomizableObject/Internal/MuCO/UnrealMutableImageProvider.h
+++ b/Engine/Plugins/Mutable/Source/CustomizableObject/Internal/MuCO/UnrealMutableImageProvider.h
@@ -2,7 +2,6 @@
 
 #pragma once
 
-#include "MuCO/BoneNames.h"
 #include "MuR/Image.h"
 #include "MuR/System.h"
 #include "Tasks/Task.h"
@@ -15,7 +14,6 @@ class USkeletalMesh;
 class FUnrealMutableResourceProvider : public UE::Mutable::Private::FExternalResourceProvider
 {
 public:
-	CUSTOMIZABLEOBJECT_API FUnrealMutableResourceProvider(const TSharedRef<FBoneNames>& InBoneNames);
 	
 	// UE::Mutable::Private::ExternalResourceProvider interface
 	// Thread: worker
@@ -43,6 +41,4 @@ private:
 #if WITH_EDITOR
 	FCriticalSection RuntimeReferencedLock;
 #endif
-
-	TSharedRef<FBoneNames> BoneNames;
 };
diff --git a/Engine/Plugins/Mutable/Source/CustomizableObject/Private/MuCO/BoneNames.cpp b/Engine/Plugins/Mutable/Source/CustomizableObject/Private/MuCO/BoneNames.cpp
deleted file mode 100644
diff --git a/Engine/Plugins/Mutable/Source/CustomizableObject/Private/MuCO/CustomizableObject.cpp b/Engine/Plugins/Mutable/Source/CustomizableObject/Private/MuCO/CustomizableObject.cpp
--- a/Engine/Plugins/Mutable/Source/CustomizableObject/Private/MuCO/CustomizableObject.cpp
+++ b/Engine/Plugins/Mutable/Source/CustomizableObject/Private/MuCO/CustomizableObject.cpp
@@ -1036,6 +1036,10 @@ void UCustomizableObjectPrivate::PostCompile()
 	MeshIdRegistry = MakeShared<UE::Mutable::Private::FMeshIdRegistry>();
 	ImageIdRegistry = MakeShared<UE::Mutable::Private::FImageIdRegistry>();
 	
+	MeshCache.Clear();
+	SkeletonCache.Clear();
+	TextureCache.Clear();
+
 	for (TObjectIterator<UCustomizableObjectInstance> It; It; ++It)
 	{
 		if (It->GetCustomizableObject() == this->GetPublic())
@@ -1978,6 +1982,12 @@ void FMeshCache::Add(const TArray<UE::Mutable::Private::FMeshId>& Key, UCustomiz
 }
 
 
+void FMeshCache::Clear()
+{
+	GeneratedMeshes = MakeShared<FMeshCacheRegistry>();
+}
+
+
 USkeleton* FSkeletonCache::Get(const TArray<uint16>& Key)
 {
 	const TWeakObjectPtr<USkeleton>* Result = MergedSkeletons.Find(Key);
@@ -2005,6 +2015,11 @@ void FSkeletonCache::Add(const TArray<uint16>& Key, USkeleton* Value)
 }
 
 
+void FSkeletonCache::Clear()
+{
+	MergedSkeletons.Empty();
+}
+
 UTexture2D* FTextureCache::Get(const FId& Key)
 {
 	const TWeakObjectPtr<UTexture2D>* Result = GeneratedTextures.Find(Key);
@@ -2045,6 +2060,12 @@ void FTextureCache::Remove(const UTexture2D& Value)
 }
 
 
+void FTextureCache::Clear()
+{
+	GeneratedTextures.Empty();
+}
+
+
 FArchive& operator<<(FArchive& Ar, FIntegerParameterUIData& Struct)
 {
 	Ar << Struct.ParamUIMetadata;
@@ -3318,16 +3339,18 @@ void UCustomizableObjectPrivate::LogMemory()
 			StringsTotal += Entry.GetAllocatedSize();
 		}
 
+		int32 ConstantNames = Program.ConstantNames.GetAllocatedSize();
+
 		int32 SkeletonsTotal = Program.ConstantSkeletons.GetAllocatedSize();
 		for (const UE::Mutable::Private::TManagedPtr<const UE::Mutable::Private::FSkeleton>& Entry : Program.ConstantSkeletons)
 		{
-			SkeletonsTotal += sizeof(UE::Mutable::Private::FSkeleton) + Entry->DebugBoneNames.GetAllocatedSize() + Entry->BoneIds.GetAllocatedSize() + Entry->BoneParents.GetAllocatedSize();
+			SkeletonsTotal += sizeof(UE::Mutable::Private::FSkeleton) + Entry->BoneNames.GetAllocatedSize() + Entry->BoneIds.GetAllocatedSize() + Entry->BoneParents.GetAllocatedSize();
 		}
 
 		int32 PhysicsTotal = Program.ConstantPhysicsBodies.GetAllocatedSize();
 		for (const UE::Mutable::Private::TManagedPtr<const UE::Mutable::Private::FPhysicsBody>& Entry : Program.ConstantPhysicsBodies)
 		{
-			PhysicsTotal += sizeof(UE::Mutable::Private::FPhysicsBody) + Entry->BoneIds.GetAllocatedSize() + Entry->Bodies.GetAllocatedSize() + Entry->BodiesCustomIds.GetAllocatedSize();
+			PhysicsTotal += sizeof(UE::Mutable::Private::FPhysicsBody) + Entry->BodiesBoneNames.GetAllocatedSize() + Entry->Bodies.GetAllocatedSize() + Entry->BodiesCustomIds.GetAllocatedSize();
 			for (const UE::Mutable::Private::FPhysicsBodyAggregate& Body : Entry->Bodies)
 			{
 				PhysicsTotal += Body.Spheres.GetAllocatedSize() + Body.Boxes.GetAllocatedSize() + Body.Convex.GetAllocatedSize() + Body.Sphyls.GetAllocatedSize() + Body.TaperedCapsules.GetAllocatedSize();
@@ -3346,7 +3369,8 @@ void UCustomizableObjectPrivate::LogMemory()
 
 		int32 ModelTotal = ByteCode + OpAddress + Roms 
 			+ ImageTotal + MeshesTotal 
-			+ StringsTotal + SkeletonsTotal + PhysicsTotal
+			+ StringsTotal + ConstantNames
+			+ SkeletonsTotal + PhysicsTotal
 			+ ParametersTotal;
 
 
@@ -3374,7 +3398,6 @@ void UCustomizableObjectPrivate::LogMemory()
 				Arch << ModelResources->AnimBPs;
 				Arch << ModelResources->AnimBpOverridePhysiscAssetsInfo;
 				Arch << ModelResources->MaterialSlotNames;
-				Arch << ModelResources->BoneNamesMap;
 				Arch << ModelResources->Sockets;
 				Arch << ModelResources->SkinWeightProfilesInfo;
 				Arch << ModelResources->MeshMetadata;
diff --git a/Engine/Plugins/Mutable/Source/CustomizableObject/Private/MuCO/CustomizableObjectInstance.cpp b/Engine/Plugins/Mutable/Source/CustomizableObject/Private/MuCO/CustomizableObjectInstance.cpp
--- a/Engine/Plugins/Mutable/Source/CustomizableObject/Private/MuCO/CustomizableObjectInstance.cpp
+++ b/Engine/Plugins/Mutable/Source/CustomizableObject/Private/MuCO/CustomizableObjectInstance.cpp
@@ -32,7 +32,6 @@
 #include "MuCO/UnrealConversionUtils.h"
 #include "MuCO/UnrealPortabilityHelpers.h"
 #include "MuCO/LogBenchmarkUtil.h"
-#include "MuCO/BoneNames.h"
 
 #include "PhysicsEngine/PhysicsAsset.h"
 #include "Rendering/Texture2DResource.h"
@@ -1345,18 +1344,7 @@ namespace
 
 		Result->bNotForDedicatedServer = TemplateAsset->bNotForDedicatedServer;
 
-		const TMap<UE::Mutable::Private::FBoneName, TPair<FName, uint16>>& BoneInfoMap = Context->InstanceUpdateData.SkeletonsPerInstanceComponent[ComponentIndex].BoneInfoMap;
-		TMap<FName, int32> BonesInUse;
-
-		const int32 MutablePhysicsBodyCount = MutablePhysics->GetBodyCount();
-		BonesInUse.Reserve(MutablePhysicsBodyCount);
-		for ( int32 I = 0; I < MutablePhysicsBodyCount; ++I )
-		{
-			if (const TPair<FName, uint16>* BoneInfo = BoneInfoMap.Find(MutablePhysics->GetBodyBoneId(I)))
-			{
-				BonesInUse.Add(BoneInfo->Key, I);
-			}
-		}
+		const TArray<FName>& BonesInUse = MutablePhysics->BodiesBoneNames;
 
 		const int32 PhysicsAssetBodySetupNum = TemplateAsset->SkeletalBodySetups.Num();
 		bool bTemplateBodyNotUsedFound = false;
@@ -1368,8 +1356,8 @@ namespace
 		{
 			const TObjectPtr<USkeletalBodySetup>& BodySetup = TemplateAsset->SkeletalBodySetups[BodySetupIndex];
 
-			int32* MutableBodyIndex = BonesInUse.Find(BodySetup->BoneName);
-			if (!MutableBodyIndex)
+			int32 MutableBodyIndex = BonesInUse.Find(BodySetup->BoneName);
+			if (MutableBodyIndex == INDEX_NONE)
 			{
 				bTemplateBodyNotUsedFound = true;
 				UsageMap[BodySetupIndex] = 0;
@@ -1397,10 +1385,10 @@ namespace
 			NewBodySetup->bSkipScaleFromAnimation = BodySetup->bSkipScaleFromAnimation;
 
 			// PhysicalAnimationProfiles can't be added with the current UPhysicsAsset API outside the editor.
-			// Don't pouplate them for now.	
+			// Don't populate them for now.	
 			//NewBodySetup->PhysicalAnimationData = BodySetup->PhysicalAnimationData;
 
-			NewBodySetup->AggGeom = MakeAggGeomFromMutablePhysics(*MutableBodyIndex, MutablePhysics);
+			NewBodySetup->AggGeom = MakeAggGeomFromMutablePhysics(MutableBodyIndex, MutablePhysics);
 
 			Result->SkeletalBodySetups.Add(NewBodySetup);
 		}
@@ -1512,7 +1500,7 @@ UPhysicsAsset* GetOrBuildMainPhysicsAsset(
 		return Result;
 	}
 
-	// Just get the referenced asset if no recontrution or merge is needed.
+	// Just get the referenced asset if no reconstruction or merge is needed.
 	if (ValidAssets.Num() == 1 && !MutablePhysics->bBodiesModified)
 	{
 		return ValidAssets[0];
@@ -1533,18 +1521,7 @@ UPhysicsAsset* GetOrBuildMainPhysicsAsset(
 
 	Result->bNotForDedicatedServer = TemplateAsset->bNotForDedicatedServer;
 
-	const TMap<UE::Mutable::Private::FBoneName, TPair<FName, uint16>>& BoneInfoMap = Context->InstanceUpdateData.SkeletonsPerInstanceComponent[ComponentIndex].BoneInfoMap;
-	TMap<FName, int32> BonesInUse;
-
-	const int32 MutablePhysicsBodyCount = MutablePhysics->GetBodyCount();
-	BonesInUse.Reserve(MutablePhysicsBodyCount);
-	for ( int32 I = 0; I < MutablePhysicsBodyCount; ++I )
-	{
-		if (const TPair<FName, uint16>* BoneInfo = BoneInfoMap.Find(MutablePhysics->GetBodyBoneId(I)))
-		{
-			BonesInUse.Add(BoneInfo->Key, I);
-		}
-	}
+	const TArray<FName>& BonesInUse = MutablePhysics->BodiesBoneNames;
 
 	// Each array is a set of elements that can collide  
 	TArray<TArray<int32, TInlineAllocator<8>>> CollisionSets;
@@ -1554,7 +1531,7 @@ UPhysicsAsset* GetOrBuildMainPhysicsAsset(
 	// Map from BodyName/BoneName to set and index in set.
 	TMap<FName, CollisionSetEntryType> BodySetupSetMap;
 	
-	// Only for elements that belong to two or more differnet sets. 
+	// Only for elements that belong to two or more different sets. 
 	// Contains in which set the elements belong.
 	using MultiSetArrayType = TArray<int32, TInlineAllocator<4>>;
 	TMap<int32, MultiSetArrayType> MultiCollisionSets;
@@ -1576,8 +1553,8 @@ UPhysicsAsset* GetOrBuildMainPhysicsAsset(
 		{
 			const TObjectPtr<USkeletalBodySetup>& BodySetup = ValidAssets[CollisionSetIndex]->SkeletalBodySetups[BodySetupIndex];
 			
-			int32* MutableBodyIndex = BonesInUse.Find(BodySetup->BoneName);
-			if (!MutableBodyIndex)
+			int32 MutableBodyIndex = BonesInUse.Find(BodySetup->BoneName);
+			if (MutableBodyIndex == INDEX_NONE)
 			{
 				continue;
 			}
@@ -1607,10 +1584,10 @@ UPhysicsAsset* GetOrBuildMainPhysicsAsset(
 				NewBodySetup->bSkipScaleFromAnimation = BodySetup->bSkipScaleFromAnimation;
 				
 				// PhysicalAnimationProfiles can't be added with the current UPhysicsAsset API outside the editor.
-				// Don't poulate them for now.	
+				// Don't populate them for now.	
 				//NewBodySetup->PhysicalAnimationData = BodySetup->PhysicalAnimationData;
 
-				NewBodySetup->AggGeom = MakeAggGeomFromMutablePhysics(*MutableBodyIndex, MutablePhysics);
+				NewBodySetup->AggGeom = MakeAggGeomFromMutablePhysics(MutableBodyIndex, MutablePhysics);
 
 
 				Result->SkeletalBodySetups.Add(NewBodySetup);
@@ -1628,8 +1605,8 @@ UPhysicsAsset* GetOrBuildMainPhysicsAsset(
 				int32 FoundBodyIndex = Found->Get<2>();
 				
 				// No need to add the body again. Volumes that come form mutable are already merged.
-				// here we only need to merge properies.
-				// TODO: check if there is other properties worth merging. In case of conflict select the more restrivtive one? 
+				// here we only need to merge properties.
+				// TODO: check if there is other properties worth merging. In case of conflict select the more restrictive one? 
 				Result->SkeletalBodySetups[FoundBodyIndex]->bConsiderForBounds |= BodySetup->bConsiderForBounds;
 
 				// Mark as removed so no indices are invalidated.
@@ -1659,9 +1636,9 @@ UPhysicsAsset* GetOrBuildMainPhysicsAsset(
 			int32 MappedIdx1 = SetsIndexMap[CollisionSetIndex][DisabledCollision.Key.Indices[1]];
 
 			// This will generate correct disables for the case when two shapes from different sets
-			// are meged to the same setup. Will introduce repeated pairs, but this is not a problem.
+			// are merged to the same setup. Will introduce repeated pairs, but this is not a problem.
 
-			// Currenly if two bodies / bones have disabled collision in one of the merged assets, the collision
+			// Currently if two bodies / bones have disabled collision in one of the merged assets, the collision
 			// will remain disabled even if other merges allow it.   
 			if ( MappedIdx0 != INDEX_NONE && MappedIdx1 != INDEX_NONE )
 			{
@@ -1694,7 +1671,7 @@ UPhysicsAsset* GetOrBuildMainPhysicsAsset(
 
 	if (bDisableCollisionsBetweenDifferentAssets)
 	{
-		// Compute collision disable table size upperbound to reduce number of alloactions.
+		// Compute collision disable table size upperbound to reduce number of allocations.
 		int32 CollisionDisableTableSize = 0;
 		for (int32 S0 = 1; S0 < CollisionSets.Num(); ++S0)
 		{
@@ -3693,44 +3670,12 @@ bool BuildSkeletonData(const TSharedRef<FUpdateContextPrivate>& Context, USkelet
 	const TArray<FMeshBoneInfo>& RawRefBoneInfo = ReferenceSkeleton.GetRawRefBoneInfo();
 	const int32 RawRefBoneCount = ReferenceSkeleton.GetRawBoneNum();
 
-	const TArray<FInstanceUpdateData::FBone>& BonePose = Context->InstanceUpdateData.SkeletonsPerInstanceComponent[ComponentIndex].BonePose;
-	TMap<UE::Mutable::Private::FBoneName, TPair<FName,uint16>>& BoneInfoMap = Context->InstanceUpdateData.SkeletonsPerInstanceComponent[ComponentIndex].BoneInfoMap;
-
-	{
-		MUTABLE_CPUPROFILER_SCOPE(BuildSkeletonData_BuildBoneInfoMap);
-
-		BoneInfoMap.Reserve(RawRefBoneCount);
-		
-		for (int32 Index = 0; Index < RawRefBoneCount; ++Index)
-		{
-			const FName BoneName = RawRefBoneInfo[Index].Name;
-			if(const UE::Mutable::Private::FBoneName* Bone = Context->BoneNames->Find(BoneName))
-			{
-				TPair<FName, uint16>& BoneInfo = BoneInfoMap.Add(*Bone);
-				BoneInfo.Key = BoneName;
-				BoneInfo.Value = Index;
-			}
-		}
-	}
-
-	{
-		MUTABLE_CPUPROFILER_SCOPE(BuildSkeletonData_EnsureBonesExist);
-
-		// Ensure all required bones are present in the skeleton
-		for (const FInstanceUpdateData::FBone& Bone : BonePose)
-		{
-			if (!BoneInfoMap.Find(Bone.Name))
-			{
-				UE_LOG(LogMutable, Warning, TEXT("The skeleton of skeletal mesh [%s] is missing a bone with ID [%d], which the mesh requires."),
-					*SkeletalMesh.GetName(), Bone.Name.Id);
-				return false;
-			}
-		}
-	}
 
 	{
 		MUTABLE_CPUPROFILER_SCOPE(BuildSkeletonData_ApplyPose);
 		
+		const TArray<FInstanceUpdateData::FBone>& BonePose = Context->InstanceUpdateData.SkeletonsPerInstanceComponent[ComponentIndex].BonePose;
+		
 		TArray<FMatrix44f>& RefBasesInvMatrix = SkeletalMesh.GetRefBasesInvMatrix();
 		RefBasesInvMatrix.Empty(RawRefBoneCount);
 
@@ -3740,8 +3685,15 @@ bool BuildSkeletonData(const TSharedRef<FUpdateContextPrivate>& Context, USkelet
 		// First step is to update the RefBasesInvMatrix for the bones.
 		for (const FInstanceUpdateData::FBone& Bone : BonePose)
 		{
-			const int32 BoneIndex = BoneInfoMap[Bone.Name].Value;
-			RefBasesInvMatrix[BoneIndex] = Bone.MatrixWithScale;
+			const int32 BoneIndex = ReferenceSkeleton.FindBoneIndex(Bone.BoneName);
+			if (RefBasesInvMatrix.IsValidIndex(BoneIndex))
+			{
+				RefBasesInvMatrix[BoneIndex] = Bone.MatrixWithScale;
+			}
+			else 
+			{
+				return false;
+			}
 		}
 
 		// The second step is to update the pose transforms in the ref skeleton from the BasesInvMatrix
@@ -4536,16 +4488,15 @@ bool BuildOrCopyRenderData(const TSharedRef<FUpdateContextPrivate>& Context, USk
 		
 		FSkeletalMeshLODRenderData& LODResource = RenderData->LODRenderData[LODIndex];
 
-		const TMap<UE::Mutable::Private::FBoneName, TPair<FName, uint16>>& BoneInfoMap = Context->InstanceUpdateData.SkeletonsPerInstanceComponent[ComponentIndex].BoneInfoMap;
-		
 		// Set active and required bones
 		{
-			const TArray<UE::Mutable::Private::FBoneName>& ActiveBones = Context->InstanceUpdateData.ActiveBones;
+			const FReferenceSkeleton& RefSkeleton = SkeletalMesh->GetRefSkeleton();
+			const TArray<FName>& ActiveBones = Context->InstanceUpdateData.ActiveBones;
 			LODResource.ActiveBoneIndices.Reserve(LOD->ActiveBoneCount);
 
 			for (uint32 Index = 0; Index < LOD->ActiveBoneCount; ++Index)
 			{
-				const uint16 ActiveBoneIndex = BoneInfoMap[ActiveBones[LOD->FirstActiveBone + Index]].Value;
+				const uint16 ActiveBoneIndex = RefSkeleton.FindBoneIndex(ActiveBones[LOD->FirstActiveBone + Index]);
 				LODResource.ActiveBoneIndices.Add(ActiveBoneIndex);
 			}
 
@@ -4603,8 +4554,8 @@ bool BuildOrCopyRenderData(const TSharedRef<FUpdateContextPrivate>& Context, USk
 		UnrealConversionUtils::SetupRenderSections(
 			LODResource,
 			LOD->Mesh.Get(),
+			SkeletalMesh->GetRefSkeleton(),
 			Context->InstanceUpdateData.BoneMaps,
-			BoneInfoMap,
 			LOD->FirstBoneMap,
 			MeshSurfacesMetadata);
 
diff --git a/Engine/Plugins/Mutable/Source/CustomizableObject/Private/MuCO/CustomizableObjectSystem.cpp b/Engine/Plugins/Mutable/Source/CustomizableObject/Private/MuCO/CustomizableObjectSystem.cpp
--- a/Engine/Plugins/Mutable/Source/CustomizableObject/Private/MuCO/CustomizableObjectSystem.cpp
+++ b/Engine/Plugins/Mutable/Source/CustomizableObject/Private/MuCO/CustomizableObjectSystem.cpp
@@ -1897,6 +1897,11 @@ namespace impl
 					continue;
 				}
 
+				if (!Mesh->GetSkeleton())
+				{
+					continue;
+				}
+
 				FInstanceUpdateData::FSkeletonData& SkeletonData = OperationData->InstanceUpdateData.SkeletonsPerInstanceComponent[ComponentIndex];
 
 				// Add SkeletonIds 
@@ -1916,26 +1921,32 @@ namespace impl
 					}
 				}
 
+				const TArray<FName>& BoneNames = Mesh->GetSkeleton()->BoneNames;
+
 				// Append BoneMap to the array of BoneMaps
-				const TArray<UE::Mutable::Private::FBoneName>& BoneMap = Mesh->GetBoneMap();
+				const TArray<UE::Mutable::Private::FBoneIdOrIndex>& BoneMap = Mesh->GetBoneMap();
 				LOD->FirstBoneMap = OperationData->InstanceUpdateData.BoneMaps.Num();
 				LOD->BoneMapCount = BoneMap.Num();
-				OperationData->InstanceUpdateData.BoneMaps.Append(BoneMap);
+				
+				for (UE::Mutable::Private::FBoneIdOrIndex Bone : BoneMap)
+				{
+					OperationData->InstanceUpdateData.BoneMaps.Add(BoneNames[Bone.Index]);
+				}
 
 				// Add active bone indices and poses
 				LOD->FirstActiveBone = OperationData->InstanceUpdateData.ActiveBones.Num();
 				LOD->ActiveBoneCount = Mesh->GetBonePoseCount();
 				for (uint32 BoneIndex = 0; BoneIndex < LOD->ActiveBoneCount; ++BoneIndex)
 				{
-					const UE::Mutable::Private::FBoneName& BoneId = Mesh->GetBonePoseId(BoneIndex);
+					FName BoneName = BoneNames[Mesh->GetBonePoseId(BoneIndex).Index];
 
-					OperationData->InstanceUpdateData.ActiveBones.Add(BoneId);
+					OperationData->InstanceUpdateData.ActiveBones.Add(BoneName);
 
-					if (!SkeletonData.BonePose.FindByKey(BoneId))
+					if (!SkeletonData.BonePose.FindByKey(BoneName))
 					{
 						FTransform3f Transform;
 						Mesh->GetBonePoseTransform(BoneIndex, Transform);
-						SkeletonData.BonePose.Add({ BoneId,Transform.Inverse().ToMatrixWithScale() });
+						SkeletonData.BonePose.Add({ BoneName,Transform.Inverse().ToMatrixWithScale() });
 					}
 				}
 			}
@@ -3147,11 +3158,8 @@ namespace impl
 			SystemPrivate->MutableTaskGraph.AllowLaunchingMutableTaskLowPriority(false, false);
 
 			const TSharedPtr<UE::Mutable::Private::FSystem> MutableSystem = SystemPrivate->MutableSystem;
-
-			const TSharedRef<FBoneNames> BoneNames = MakeShared<FBoneNames>(*Object->GetPrivate()->GetModelResources());
-			Operation->BoneNames = BoneNames;
 			
-			Operation->ExternalResourceProvider = MakeShared<FUnrealMutableResourceProvider>(BoneNames);
+			Operation->ExternalResourceProvider = MakeShared<FUnrealMutableResourceProvider>();
 			Operation->PassthroughObjectLoader = MakeShared<UE::Mutable::Private::FPassthroughObjectLoader>();
 
 			Operation->bLiveUpdateMode = SystemPrivate->EnableMutableLiveUpdate && (StateData ? StateData->bLiveUpdateMode : false);
diff --git a/Engine/Plugins/Mutable/Source/CustomizableObject/Private/MuCO/UnrealConversionUtils.cpp b/Engine/Plugins/Mutable/Source/CustomizableObject/Private/MuCO/UnrealConversionUtils.cpp
--- a/Engine/Plugins/Mutable/Source/CustomizableObject/Private/MuCO/UnrealConversionUtils.cpp
+++ b/Engine/Plugins/Mutable/Source/CustomizableObject/Private/MuCO/UnrealConversionUtils.cpp
@@ -18,7 +18,6 @@
 #include "UObject/StrongObjectPtr.h"
 #include "GPUSkinVertexFactory.h"
 #include "SkinnedAssetCompiler.h"
-#include "MuCO/BoneNames.h"
 #include "Rendering/SkeletalMeshRenderData.h"
 #include "Rendering/MorphTargetVertexCodec.h"
 
@@ -220,8 +219,8 @@ namespace UnrealConversionUtils
 	void SetupRenderSections(
 		FSkeletalMeshLODRenderData& LODResource,
 		const UE::Mutable::Private::FMesh* InMutableMesh,
-		const TArray<UE::Mutable::Private::FBoneName>& InBoneMap,
-		const TMap<UE::Mutable::Private::FBoneName, TPair<FName, uint16>>& BoneInfoMap,
+		const FReferenceSkeleton& InRefSkeleton,
+		const TArray<FName>& InBoneMap,
 		const int32 InFirstBoneMapIndex,
 		const TArray<const FMutableSurfaceMetadata*>& SurfacesMetadata)
 	{
@@ -274,20 +273,13 @@ namespace UnrealConversionUtils
 				Section.bCastShadow = SurfacesMetadata[SurfaceIndex]->bCastShadow;
 			}
 
-			// InBoneMaps may contain bonemaps from other sections. Copy the bones belonging to this mesh.
+			// InBoneMaps may contain BoneMaps from other sections. Copy the bones belonging to this mesh.
 			FirstBone += InFirstBoneMapIndex;
 			
 			Section.BoneMap.Reserve(BoneCount);
 			for (int32 BoneMapIndex = 0; BoneMapIndex < BoneCount; ++BoneMapIndex, ++FirstBone)
 			{
-				if (InBoneMap.IsValidIndex(FirstBone))
-				{
-					UE::Mutable::Private::FBoneName FirstBoneName = InBoneMap[FirstBone];
-					if (const TPair<FName, uint16>* Found = BoneInfoMap.Find(FirstBoneName))
-					{
-						Section.BoneMap.Add(Found->Value);
-					}
-				}
+				Section.BoneMap.Add(InRefSkeleton.FindBoneIndex(InBoneMap[FirstBone]));
 			}
 		}
 	}
@@ -1032,7 +1024,7 @@ namespace UnrealConversionUtils
 	}
 
 	
-	UE::Tasks::FTask ConvertSkeletalMeshFromRuntimeData(USkeletalMesh* SkeletalMesh, int32 LODIndex, int32 SectionIndex, UE::Mutable::Private::TManagedPtr<UE::Mutable::Private::FMesh> Result, const TSharedRef<FBoneNames>& BoneNames)
+	UE::Tasks::FTask ConvertSkeletalMeshFromRuntimeData(USkeletalMesh* SkeletalMesh, int32 LODIndex, int32 SectionIndex, UE::Mutable::Private::TManagedPtr<UE::Mutable::Private::FMesh> Result)
 	{
 		MUTABLE_CPUPROFILER_SCOPE(ConvertSkeletalMeshFromRuntimeData);
 
@@ -1071,7 +1063,7 @@ namespace UnrealConversionUtils
 #endif
 
 		UE::Tasks::FTask ConversionTask = UE::Tasks::Launch(TEXT("ConversionTask"),
-			[Context, SkeletalMesh = TStrongObjectPtr(SkeletalMesh), LODIndex, SectionIndex, BoneNames]()
+			[Context, SkeletalMesh = TStrongObjectPtr(SkeletalMesh), LODIndex, SectionIndex]()
 			{
 				FSkeletalMeshRenderData* Data = SkeletalMesh->GetResourceForRendering();
 
@@ -1150,7 +1142,7 @@ namespace UnrealConversionUtils
 
 		
 		UE::Tasks::FTask FinalConversionTask = UE::Tasks::Launch( TEXT("MutableRuntimeMeshConversionFinal"),
-			[Context, SkeletalMesh = TStrongObjectPtr(SkeletalMesh), LODIndex, SectionIndex, Result, BoneNames]()
+			[Context, SkeletalMesh = TStrongObjectPtr(SkeletalMesh), LODIndex, SectionIndex, Result]()
 			{
 				MUTABLE_CPUPROFILER_SCOPE(MeshConversion);
 
@@ -1176,14 +1168,10 @@ namespace UnrealConversionUtils
 
 				const FSkelMeshRenderSection& Section = Context->OriginalLOD->RenderSections[SectionIndex];
 					
-				// Skeleton data
-				bool bBoneMapModified = false;
-				TArray<FBoneIndexType> BoneMap;
-				TArray<FBoneIndexType> RemappedBoneMapIndices;
-
 				bool bIgnoreSkeleton = false;
 				if (!bIgnoreSkeleton)
 				{
+					MUTABLE_CPUPROFILER_SCOPE(MeshConversion_Skeleton)
 					USkeleton* Skeleton = SkeletalMesh->GetSkeleton();
 					if (!Skeleton)
 					{
@@ -1191,43 +1179,23 @@ namespace UnrealConversionUtils
 						return;
 					}
 
-					const TArray<uint16>& SourceRequiredBones = Context->OriginalLOD->RequiredBones;
-
-					// Build RequiredBones array
-					TArray<FBoneIndexType> RequiredBones;
-					RequiredBones.Reserve(SourceRequiredBones.Num());
-
-					for (const FBoneIndexType& RequiredBoneIndex : SourceRequiredBones)
-					{
-						RequiredBones.AddUnique(RequiredBoneIndex);
-					}
-
-					// Rebuild BoneMap
-					const TArray<uint16>& SourceBoneMap = Section.BoneMap;
-					const int32 NumBonesInBoneMap = SourceBoneMap.Num();
-
-					for (int32 BoneIndex = 0; BoneIndex < NumBonesInBoneMap; ++BoneIndex)
-					{
-						const FBoneIndexType FinalBoneIndex = SourceBoneMap[BoneIndex];
-
-						const int32 BoneMapIndex = BoneMap.AddUnique(FinalBoneIndex);
-						RemappedBoneMapIndices.Add(BoneMapIndex);
-
-						bBoneMapModified = bBoneMapModified || SourceBoneMap[BoneIndex] != FinalBoneIndex;
-					}
-
 					// Create the skeleton, poses, and BoneMap for this mesh
 					UE::Mutable::Private::TManagedPtr<UE::Mutable::Private::FSkeleton> MutableSkeleton = 
 							UE::Mutable::Private::MakeManaged<UE::Mutable::Private::FSkeleton>();
 					Result->SetSkeleton(MutableSkeleton);
 
+					const TArray<uint16>& RequiredBones = Context->OriginalLOD->RequiredBones;
 					const int32 NumRequiredBones = RequiredBones.Num();
 					Result->SetBonePoseCount(NumRequiredBones);
-					MutableSkeleton->SetBoneCount(NumRequiredBones);
 
-					// MutableBoneMap will not keep an index to the Skeleton, but to the BoneName
-					TArray<UE::Mutable::Private::FBoneName> MutableBoneMap;
-					MutableBoneMap.SetNum(BoneMap.Num());
+					MutableSkeleton->BoneNames.Reserve(NumRequiredBones);
+					MutableSkeleton->BoneParents.Reserve(NumRequiredBones);
+
+					const TArray<uint16>& BoneMap = Section.BoneMap;
+					const int32 NumBonesInBoneMap = BoneMap.Num();
+
+					TArray<UE::Mutable::Private::FBoneIdOrIndex> MutableBoneMap;
+					MutableBoneMap.SetNum(NumBonesInBoneMap);
 
 					TArray<FMatrix> ComposedRefPoseMatrices;
 					ComposedRefPoseMatrices.SetNum(NumRequiredBones);
@@ -1240,22 +1208,14 @@ namespace UnrealConversionUtils
 						const FMeshBoneInfo& BoneInfo = RefBoneInfo[RefSkeletonBoneIndex];
 						const int32 ParentBoneIndex = RequiredBones.Find(BoneInfo.ParentIndex);
 
-						// Set bone hierarchy
-
-						// Get the bone id
-						UE::Mutable::Private::FBoneName BoneName = BoneNames->FindOrAdd(BoneInfo.Name);
-						
-						MutableSkeleton->SetBoneName(BoneIndex, BoneName);
-						MutableSkeleton->SetBoneParent(BoneIndex, ParentBoneIndex);
-
-						// Debug. Will not be serialized
-						MutableSkeleton->SetDebugName(BoneIndex, BoneInfo.Name);
+						UE::Mutable::Private::FBoneIdOrIndex MutableBoneIndex;
+						MutableBoneIndex.Index = MutableSkeleton->AddBone(BoneInfo.Name, ParentBoneIndex);
 
 						// BoneMap: Convert RefSkeletonBoneIndex to BoneId
 						const int32 BoneMapIndex = BoneMap.Find(RefSkeletonBoneIndex);
 						if (BoneMapIndex != INDEX_NONE)
 						{
-							MutableBoneMap[BoneMapIndex] = BoneName;
+							MutableBoneMap[BoneMapIndex] = MutableBoneIndex;
 						}
 
 						if (ParentBoneIndex >= 0)
@@ -1275,7 +1235,7 @@ namespace UnrealConversionUtils
 						EnumAddFlags(BoneUsageFlags, BoneMapIndex != INDEX_NONE ? UE::Mutable::Private::EBoneUsageFlags::Skinning : UE::Mutable::Private::EBoneUsageFlags::None);
 						EnumAddFlags(BoneUsageFlags, ParentBoneIndex == INDEX_NONE ? UE::Mutable::Private::EBoneUsageFlags::Root : UE::Mutable::Private::EBoneUsageFlags::None);
 
-						Result->SetBonePose(BoneIndex, BoneName, BoneTransform, BoneUsageFlags);
+						Result->SetBonePose(BoneIndex, MutableBoneIndex, BoneTransform, BoneUsageFlags);
 					}
 
 					Result->SetBoneMap(MutableBoneMap);
@@ -1293,6 +1253,7 @@ namespace UnrealConversionUtils
 				// Position buffer
 				int32 CurrentVertexBuffer = 0;
 				{
+					MUTABLE_CPUPROFILER_SCOPE(MeshConversion_Positions)
 					MutableMeshBufferUtils::SetupVertexPositionsBuffer(CurrentVertexBuffer, Vertices);
 
 					int32 ElementSize = Vertices.Buffers[CurrentVertexBuffer].ElementSize;
@@ -1306,6 +1267,8 @@ namespace UnrealConversionUtils
 
 				// Tangent buffer
 				{
+					MUTABLE_CPUPROFILER_SCOPE(MeshConversion_Tangent)
+
 					bool bIsHighPrecision = Context->DataLOD->StaticVertexBuffers.StaticMeshVertexBuffer.GetUseHighPrecisionTangentBasis();
 
 					const UE::Mutable::Private::EMeshBufferSemantic Semantics[2] = { UE::Mutable::Private::EMeshBufferSemantic::Tangent, UE::Mutable::Private::EMeshBufferSemantic::Normal };
@@ -1336,6 +1299,8 @@ namespace UnrealConversionUtils
 
 				// Texture coords buffer
 				{
+					MUTABLE_CPUPROFILER_SCOPE(MeshConversion_TexCoords)
+
 					int32 NumTexCoords = Context->OriginalLOD->GetNumTexCoords();
 
 					constexpr int32 MaxChannelCount = 4;
@@ -1369,6 +1334,8 @@ namespace UnrealConversionUtils
 				// Skin buffer
 				if (Section.MaxBoneInfluences > 0)
 				{
+					MUTABLE_CPUPROFILER_SCOPE(MeshConversion_Skin)
+
 					const FSkinWeightVertexBuffer* SkinBuffer = Context->DataLOD->GetSkinWeightVertexBuffer();
 					const int32 MaxBoneIndexTypeSizeBytes = SkinBuffer->GetBoneIndexByteSize();
 					const int32 MaxBoneWeightTypeSizeBytes = SkinBuffer->GetBoneWeightByteSize();
@@ -1388,6 +1355,8 @@ namespace UnrealConversionUtils
 				// Colour buffer
 				if (Context->DataLOD->StaticVertexBuffers.ColorVertexBuffer.GetNumVertices() > 0)
 				{
+					MUTABLE_CPUPROFILER_SCOPE(MeshConversion_Color)
+
 					const FColorVertexBuffer& ColorBuffer = Context->DataLOD->StaticVertexBuffers.ColorVertexBuffer;
 
 					MutableMeshBufferUtils::SetupVertexColorBuffer(CurrentVertexBuffer, Vertices);
@@ -1406,6 +1375,8 @@ namespace UnrealConversionUtils
 
 				// Indices
 				{
+					MUTABLE_CPUPROFILER_SCOPE(MeshConversion_Indices)
+
 					int32 FirstIndexIndex = Section.BaseIndex;
 					int32 IndexCount = Section.NumTriangles * 3;
 					int32 ElementSize = Context->DataLOD->MultiSizeIndexContainer.GetDataTypeSize();
diff --git a/Engine/Plugins/Mutable/Source/CustomizableObject/Private/MuCO/UnrealMutableImageProvider.cpp b/Engine/Plugins/Mutable/Source/CustomizableObject/Private/MuCO/UnrealMutableImageProvider.cpp
--- a/Engine/Plugins/Mutable/Source/CustomizableObject/Private/MuCO/UnrealMutableImageProvider.cpp
+++ b/Engine/Plugins/Mutable/Source/CustomizableObject/Private/MuCO/UnrealMutableImageProvider.cpp
@@ -95,10 +95,6 @@ UE::Mutable::Private::EImageFormat GetMutablePixelFormat(EPixelFormat InTextureF
 }
 
 
-FUnrealMutableResourceProvider::FUnrealMutableResourceProvider(const TSharedRef<FBoneNames>& InBoneNames) : BoneNames(InBoneNames)
-{
-}
-
 namespace UnrealMutableImageProviderInteranl
 {
 	bool CheckBulkSizesAreValid(TNotNull<UTexture2D*> Texture, TNotNull<UE::Mutable::Private::FImage*> Image, int32 TextureMipIndexBegin, int32 TextureMipIndexEnd)
@@ -439,7 +435,7 @@ TTuple<UE::Tasks::FTask, TFunction<void()>> FUnrealMutableResourceProvider::GetM
 		return Invoke(TrivialReturn);
 	}
 	
-	UE::Tasks::FTask ConversionTask = UnrealConversionUtils::ConvertSkeletalMeshFromRuntimeData(SkeletalMesh, LODIndex, SectionIndex, Result, BoneNames);
+	UE::Tasks::FTask ConversionTask = UnrealConversionUtils::ConvertSkeletalMeshFromRuntimeData(SkeletalMesh, LODIndex, SectionIndex, Result);
 
 	return MakeTuple(
 		// Some post-game conversion stuff can happen here in a worker thread
diff --git a/Engine/Plugins/Mutable/Source/CustomizableObjectEditor/Private/MuCOE/CustomizableObjectCompiler.cpp b/Engine/Plugins/Mutable/Source/CustomizableObjectEditor/Private/MuCOE/CustomizableObjectCompiler.cpp
--- a/Engine/Plugins/Mutable/Source/CustomizableObjectEditor/Private/MuCOE/CustomizableObjectCompiler.cpp
+++ b/Engine/Plugins/Mutable/Source/CustomizableObjectEditor/Private/MuCOE/CustomizableObjectCompiler.cpp
@@ -1668,19 +1668,6 @@ void FCustomizableObjectCompiler::FinishCompilationTask()
 		
 		ModelResources->Sockets = MoveTemp(CompilationContext->Sockets);
 
-		const int32 NumBones = CompilationContext->UniqueBoneNames.Num() + CompilationContext->RemappedBoneNames.Num();
-		ModelResources->BoneNamesMap.Reserve(NumBones);
-
-		for (auto& It : CompilationContext->UniqueBoneNames)
-		{
-			ModelResources->BoneNamesMap.Add(It.Value, It.Key.Id);
-		}
-
-		for (auto& It : CompilationContext->RemappedBoneNames)
-		{
-			ModelResources->BoneNamesMap.Add(It.Key, It.Value.Id);
-		}
-
 		ModelResources->SkinWeightProfilesInfo = MoveTemp(CompilationContext->SkinWeightProfilesInfo);
 
 		ModelResources->StreamedResourceDataEditor = MoveTemp(CompilationContext->StreamedResourceData);
diff --git a/Engine/Plugins/Mutable/Source/CustomizableObjectEditor/Private/MuCOE/CustomizableObjectEditor.cpp b/Engine/Plugins/Mutable/Source/CustomizableObjectEditor/Private/MuCOE/CustomizableObjectEditor.cpp
--- a/Engine/Plugins/Mutable/Source/CustomizableObjectEditor/Private/MuCOE/CustomizableObjectEditor.cpp
+++ b/Engine/Plugins/Mutable/Source/CustomizableObjectEditor/Private/MuCOE/CustomizableObjectEditor.cpp
@@ -62,7 +62,6 @@
 #include "Widgets/Input/STextComboBox.h"
 #include "Widgets/Layout/SScrollBox.h"
 #include "Widgets/Notifications/SNotificationList.h"
-#include "MuCO/BoneNames.h"
 #include "MuCO/CustomizableObjectCustomVersion.h"
 #include "MuCO/UnrealMutableImageProvider.h"
 #include "MuCOE/CompileRequest.h"
@@ -1461,8 +1460,7 @@ void FCustomizableObjectEditor::CodeViewer() const
 		return;
 	}
 	
-	const TSharedRef<FBoneNames> BoneNames = MakeShared<FBoneNames>(); // TODO DebuggerCompile does not create a UModelResource
-	TSharedPtr<FUnrealMutableResourceProvider> ExternalResourceProvider = MakeShared<FUnrealMutableResourceProvider>(BoneNames);
+	TSharedPtr<FUnrealMutableResourceProvider> ExternalResourceProvider = MakeShared<FUnrealMutableResourceProvider>();
 	
 	const TSharedPtr<SDockTab> NewMutableObjectTab = SNew(SDockTab)
 	.Label(LOCTEXT("CodeViewer", "Code Viewer"))
diff --git a/Engine/Plugins/Mutable/Source/CustomizableObjectEditor/Private/MuCOE/GenerateMutableSource/GenerateMutableSource.cpp b/Engine/Plugins/Mutable/Source/CustomizableObjectEditor/Private/MuCOE/GenerateMutableSource/GenerateMutableSource.cpp
--- a/Engine/Plugins/Mutable/Source/CustomizableObjectEditor/Private/MuCOE/GenerateMutableSource/GenerateMutableSource.cpp
+++ b/Engine/Plugins/Mutable/Source/CustomizableObjectEditor/Private/MuCOE/GenerateMutableSource/GenerateMutableSource.cpp
@@ -147,39 +147,6 @@ void FMutableCompilationContext::Log(const FText& Message, const UObject* Contex
 }
 
 
-UE::Mutable::Private::FBoneName FMutableCompilationContext::GetBoneUnique(const FName& InBoneName)
-{
-	const FString BoneNameString = InBoneName.ToString().ToLower();
-	UE::Mutable::Private::FBoneName Bone(CityHash32(reinterpret_cast<const char*>(*BoneNameString), BoneNameString.Len() * sizeof(FString::ElementType)));
-
-	bool bNewRemappedBoneName = false;
-
-	FString& BoneName = UniqueBoneNames.FindOrAdd(Bone, BoneNameString);
-	while (BoneName != InBoneName)
-	{
-		if (UE::Mutable::Private::FBoneName* RemappedBoneName = RemappedBoneNames.Find(BoneNameString))
-		{
-			Bone.Id = RemappedBoneName->Id;
-			break;
-		}
-
-		// Id collision detected
-		bNewRemappedBoneName = true;
-
-		// Increase Id in an attempt to make it unique again.
-		++Bone.Id;
-		BoneName = UniqueBoneNames.FindOrAdd(Bone, BoneNameString);
-	}
-
-	if (bNewRemappedBoneName)
-	{
-		RemappedBoneNames.Add(BoneNameString, Bone);
-	}
-
-	return Bone;
-}
-
-
 uint32 FMutableCompilationContext::GetSkinWeightProfileIdUnique(const FName InProfileName)
 {
 	const FString ProfileNameString = InProfileName.ToString();
@@ -246,25 +213,6 @@ FString FMutableCompilationContext::GetObjectName() const
 }
 
 
-bool FMutableCompilationContext::FindBone(const FName& InBoneName, UE::Mutable::Private::FBoneName& OutBoneName) const
-{
-	const FString BoneNameString = InBoneName.ToString().ToLower();
-	OutBoneName.Id = CityHash32(reinterpret_cast<const char*>(*BoneNameString), BoneNameString.Len() * sizeof(FString::ElementType));
-	if (UniqueBoneNames.Contains(OutBoneName))
-	{
-		return true;
-	}
-
-	if (const UE::Mutable::Private::FBoneName* BoneName = RemappedBoneNames.Find(BoneNameString))
-	{
-		OutBoneName = *BoneName;
-		return true;
-	}
-
-	return false;
-}
-
-
 FMutableGraphGenerationContext::FMutableGraphGenerationContext(const TSharedPtr<FMutableCompilationContext>& InCompilationContext)
 	: CompilationContext(InCompilationContext)
 	, ExtensionDataCompilerInterface(*this)
@@ -782,18 +730,21 @@ UE::Mutable::Private::Ptr<UE::Mutable::Private::NodeMeshApplyPose> CreateNodeMes
 
 	MutableMesh->SetSkeleton(MutableSkeleton);
 	MutableMesh->SetBonePoseCount(ArrayBoneName.Num());
-	MutableSkeleton->SetBoneCount(ArrayBoneName.Num());
+	MutableSkeleton->BoneNames.Reserve(ArrayBoneName.Num());
+	MutableSkeleton->BoneParents.Reserve(ArrayBoneName.Num());
 
 	for (int32 i = 0; i < ArrayBoneName.Num(); ++i)
 	{
-		const FName BoneName = ArrayBoneName[i];
-		const UE::Mutable::Private::FBoneName& BoneId = GenerationContext.CompilationContext->GetBoneUnique(BoneName);
+		constexpr int16 ParentIndex = -1; // Skeletons from poses don't need bone parents?
 
-		MutableSkeleton->SetDebugName(i, BoneName);
-		MutableSkeleton->SetBoneName(i, { BoneId });
-		MutableMesh->SetBonePose(i, BoneId, (FTransform3f)ArrayTransform[i], UE::Mutable::Private::EBoneUsageFlags::Skinning);
+		UE::Mutable::Private::FBoneIdOrIndex MutableBoneIndex;
+		MutableBoneIndex.Index = MutableSkeleton->AddBone(ArrayBoneName[i], ParentIndex);
+		MutableMesh->SetBonePose(i, MutableBoneIndex, (FTransform3f)ArrayTransform[i], UE::Mutable::Private::EBoneUsageFlags::Skinning);
 	}
 
+	// TODO: Review why Skeletons from poses don't need bone parents
+	MutableSkeleton->BoneParents.Empty();
+
 	UE::Mutable::Private::Ptr<UE::Mutable::Private::NodeMeshApplyPose> NodeMeshApplyPose = new UE::Mutable::Private::NodeMeshApplyPose;
 	NodeMeshApplyPose->Base = InputMeshNode;
 	NodeMeshApplyPose->Pose = PoseNodeMesh;
diff --git a/Engine/Plugins/Mutable/Source/CustomizableObjectEditor/Private/MuCOE/GenerateMutableSource/GenerateMutableSource.h b/Engine/Plugins/Mutable/Source/CustomizableObjectEditor/Private/MuCOE/GenerateMutableSource/GenerateMutableSource.h
--- a/Engine/Plugins/Mutable/Source/CustomizableObjectEditor/Private/MuCOE/GenerateMutableSource/GenerateMutableSource.h
+++ b/Engine/Plugins/Mutable/Source/CustomizableObjectEditor/Private/MuCOE/GenerateMutableSource/GenerateMutableSource.h
@@ -577,12 +577,6 @@ struct FMutableCompilationContext
 	/** Only Mesh Components (no passthrough). */
 	TArray<FSkeletalMeshComponentInfo> ComponentInfos;
 
-	/** Array of unique Bone identifiers. */
-	TMap<UE::Mutable::Private::FBoneName, FString> UniqueBoneNames;
-
-	/** Bone identifiers that had a collision. */
-	TMap<FString, UE::Mutable::Private::FBoneName> RemappedBoneNames;
-	
 	UE::Mutable::Private::FPassthroughObjectFactory PassthroughObjectFactory;
 
 private:
@@ -610,12 +604,6 @@ public:
 
 	FSkeletalMeshComponentInfo* GetComponentInfo(UE::Mutable::Private::FComponentId ComponentName);
 
-	/** Get unique identifier for BoneName built from its FString. */
-	UE::Mutable::Private::FBoneName GetBoneUnique(const FName& BoneName);
-
-	/** */
-	bool FindBone(const FName& BoneName, UE::Mutable::Private::FBoneName& OutBone) const;
-
 	/** */
 	uint32 GetSkinWeightProfileIdUnique(const FName ProfileName);
 
diff --git a/Engine/Plugins/Mutable/Source/CustomizableObjectEditor/Private/MuCOE/GenerateMutableSource/GenerateMutableSourceMesh.cpp b/Engine/Plugins/Mutable/Source/CustomizableObjectEditor/Private/MuCOE/GenerateMutableSource/GenerateMutableSourceMesh.cpp
--- a/Engine/Plugins/Mutable/Source/CustomizableObjectEditor/Private/MuCOE/GenerateMutableSource/GenerateMutableSourceMesh.cpp
+++ b/Engine/Plugins/Mutable/Source/CustomizableObjectEditor/Private/MuCOE/GenerateMutableSource/GenerateMutableSourceMesh.cpp
@@ -663,39 +663,6 @@ bool IsSkeletalMeshCompatibleWithRefSkeleton(FSkeletalMeshComponentInfo& Compone
 	return true;
 }
 
-void SetAndPropagatePoseBoneUsage(UE::Mutable::Private::FMesh& MutableMesh, int32 PoseIndex, UE::Mutable::Private::EBoneUsageFlags Usage)
-{
-	if (!MutableMesh.GetSkeleton())
-	{
-		return;
-	}
-
-	const UE::Mutable::Private::FSkeleton& MutableSkeleton = *MutableMesh.GetSkeleton();
-
-	if (!MutableMesh.BonePoses.IsValidIndex(PoseIndex))
-	{
-		check(false);
-		return;
-	}
-
-	int32 BoneIndex = MutableSkeleton.FindBone(MutableMesh.BonePoses[PoseIndex].BoneId);
-
-	while (BoneIndex != INDEX_NONE)
-	{
-		PoseIndex = MutableMesh.FindBonePose(MutableSkeleton.GetBoneName(BoneIndex));
-
-		if (PoseIndex == INDEX_NONE)
-		{
-			check(false);
-			return;
-		}
-
-		EnumAddFlags(MutableMesh.BonePoses[PoseIndex].BoneUsageFlags, Usage);
-
-		BoneIndex = MutableSkeleton.GetBoneParent(BoneIndex);
-	}
-
-}
 
 TArray<TTuple<UPhysicsAsset*, int32>> GetPhysicsAssetsFromAnimInstance(const TSoftClassPtr<UAnimInstance>& AnimInstance)
 {
@@ -762,6 +729,7 @@ TArray<TTuple<UPhysicsAsset*, int32>> GetPhysicsAssetsFromAnimInstance(FMutableG
 	return Result;
 }
 
+// Could be removed but may produce patch size
 TArray<uint8> MakePhysicsAssetBodySetupRelevancyMap(const FMutableCompilationContext& Context, UPhysicsAsset* Asset, const UE::Mutable::Private::TManagedPtr<UE::Mutable::Private::FMesh>& Mesh)
 {
 	const int32 BodySetupsNum = Asset->SkeletalBodySetups.Num();
@@ -769,21 +737,21 @@ TArray<uint8> MakePhysicsAssetBodySetupRelevancyMap(const FMutableCompilationCon
 	TArray<uint8> RelevancyMap;
 	RelevancyMap.Init(0, BodySetupsNum);
 
-	if (!Mesh->GetSkeleton())
+	UE::Mutable::Private::TManagedPtr<const UE::Mutable::Private::FSkeleton> Skeleton = Mesh->GetSkeleton();
+	if (!Skeleton)
 	{
 		return RelevancyMap;
 	}
 
 	for (int32 BodyIndex = 0; BodyIndex < BodySetupsNum; ++BodyIndex)
 	{
-		UE::Mutable::Private::FBoneName Bone; 
-		RelevancyMap[BodyIndex] = Context.FindBone(Asset->SkeletalBodySetups[BodyIndex]->BoneName, Bone);
+		RelevancyMap[BodyIndex] = Skeleton->BoneNames.Find(Asset->SkeletalBodySetups[BodyIndex]->BoneName) != INDEX_NONE;
 	}
 
 	return RelevancyMap;
 }
 
-UE::Mutable::Private::TManagedPtr<UE::Mutable::Private::FPhysicsBody> MakePhysicsBodyFromAsset(FMutableCompilationContext& Context, UPhysicsAsset* Asset, const TArray<uint8>& BodySetupRelevancyMap)
+UE::Mutable::Private::TManagedPtr<UE::Mutable::Private::FPhysicsBody> MakePhysicsBodyFromAsset(UPhysicsAsset* Asset, const UE::Mutable::Private::FSkeleton& MutableSkeleton, const TArray<uint8>& BodySetupRelevancyMap)
 {
 	check(Asset);
 	check(Asset->SkeletalBodySetups.Num() == BodySetupRelevancyMap.Num());
@@ -816,8 +784,7 @@ UE::Mutable::Private::TManagedPtr<UE::Mutable::Private::FPhysicsBody> MakePhysic
 
 		TObjectPtr<USkeletalBodySetup>& BodySetup = SkeletalBodySetups[SourceBodyIndex++];
 
-		const UE::Mutable::Private::FBoneName& BodyBoneId = Context.GetBoneUnique(BodySetup->BoneName);
-		PhysicsBody->SetBodyBoneId(B, BodyBoneId);
+		PhysicsBody->SetBody(B, BodySetup->BoneName, INDEX_NONE);
 		
 		const int32 NumSpheres = BodySetup->AggGeom.SphereElems.Num();
 		PhysicsBody->SetSphereCount(B, NumSpheres);
@@ -1408,10 +1375,12 @@ UE::Tasks::FTask ConvertSkeletalMeshToMutable(TSharedRef<FMeshConversionContext>
 
 		const int32 NumRequiredBones = RequiredBones.Num();
 		MutableMesh->SetBonePoseCount(NumRequiredBones);
-		MutableSkeleton->SetBoneCount(NumRequiredBones);
+
+		MutableSkeleton->BoneNames.Reserve(NumRequiredBones);
+		MutableSkeleton->BoneParents.Reserve(NumRequiredBones);
 
 		// MutableBoneMap will not keep an index to the Skeleton, but to the BoneName
-		TArray<UE::Mutable::Private::FBoneName> MutableBoneMap;
+		TArray<UE::Mutable::Private::FBoneIdOrIndex> MutableBoneMap;
 		MutableBoneMap.SetNum(BoneMap.Num());
 
 		TArray<FMatrix> ComposedRefPoseMatrices;
@@ -1425,20 +1394,14 @@ UE::Tasks::FTask ConvertSkeletalMeshToMutable(TSharedRef<FMeshConversionContext>
 			const FMeshBoneInfo& BoneInfo = RefBoneInfo[RefSkeletonBoneIndex];
 			const int32 ParentBoneIndex = RequiredBones.Find(BoneInfo.ParentIndex);
 
-			// Set bone hierarchy
-			const UE::Mutable::Private::FBoneName& BoneName = Context.GetBoneUnique(BoneInfo.Name);
-
-			MutableSkeleton->SetBoneName(BoneIndex, BoneName);
-			MutableSkeleton->SetBoneParent(BoneIndex, ParentBoneIndex);
-
-			// Debug. Will not be serialized
-			MutableSkeleton->SetDebugName(BoneIndex, BoneInfo.Name);
+			UE::Mutable::Private::FBoneIdOrIndex MutableBoneIndex;
+			MutableBoneIndex.Index = MutableSkeleton->AddBone(BoneInfo.Name, ParentBoneIndex);
 
 			// BoneMap: Convert RefSkeletonBoneIndex to BoneId
 			const int32 BoneMapIndex = BoneMap.Find(RefSkeletonBoneIndex);
 			if (BoneMapIndex != INDEX_NONE)
 			{
-				MutableBoneMap[BoneMapIndex] = BoneName;
+				MutableBoneMap[BoneMapIndex] = MutableBoneIndex;
 			}
 
 			if (ParentBoneIndex >= 0)
@@ -1458,7 +1421,7 @@ UE::Tasks::FTask ConvertSkeletalMeshToMutable(TSharedRef<FMeshConversionContext>
 			EnumAddFlags(BoneUsageFlags, BoneMapIndex != INDEX_NONE ? UE::Mutable::Private::EBoneUsageFlags::Skinning : UE::Mutable::Private::EBoneUsageFlags::None);
 			EnumAddFlags(BoneUsageFlags, ParentBoneIndex == INDEX_NONE ? UE::Mutable::Private::EBoneUsageFlags::Root : UE::Mutable::Private::EBoneUsageFlags::None);
 
-			MutableMesh->SetBonePose(BoneIndex, BoneName, BoneTransform, BoneUsageFlags);
+			MutableMesh->SetBonePose(BoneIndex, MutableBoneIndex, BoneTransform, BoneUsageFlags);
 		}
 
 		MutableMesh->SetBoneMap(MutableBoneMap);
@@ -2330,6 +2293,8 @@ UE::Tasks::FTask ConvertSkeletalMeshToMutable(TSharedRef<FMeshConversionContext>
 		TArray<uint8> DiscardedBodySetups;
 		DiscardedBodySetups.Init(1, SkeletalBodySetups.Num());
 
+		const UE::Mutable::Private::FSkeleton& MutableSkeleton = *MutableMesh->GetSkeleton();
+
 		for (int32 BodySetupIndex = 0; BodySetupIndex < SkeletalBodySetups.Num(); ++BodySetupIndex)
 		{
 			TObjectPtr<USkeletalBodySetup>& BodySetup = SkeletalBodySetups[BodySetupIndex];
@@ -2338,8 +2303,7 @@ UE::Tasks::FTask ConvertSkeletalMeshToMutable(TSharedRef<FMeshConversionContext>
 				continue;
 			}
 
-			const UE::Mutable::Private::FBoneName& BoneName = Context.GetBoneUnique(BodySetup->BoneName);
-			const int32 BonePoseIndex = MutableMesh->FindBonePose(BoneName);
+			const int32 BonePoseIndex = MutableMesh->FindBonePoseByBoneIndex(MutableSkeleton.BoneNames.Find(BodySetup->BoneName));
 
 			if (BonePoseIndex == INDEX_NONE)
 			{
@@ -2402,8 +2366,7 @@ UE::Tasks::FTask ConvertSkeletalMeshToMutable(TSharedRef<FMeshConversionContext>
 		{
 			TObjectPtr<USkeletalBodySetup>& BodySetup = RelevantBodySetups[B];
 
-			const UE::Mutable::Private::FBoneName& BoneId = Context.GetBoneUnique(BodySetup->BoneName);
-			PhysicsBody->SetBodyBoneId(B, BoneId);
+			PhysicsBody->SetBody(B, BodySetup->BoneName);
 
 			const int32 NumSpheres = BodySetup->AggGeom.SphereElems.Num();
 			PhysicsBody->SetSphereCount(B, NumSpheres);
@@ -2507,38 +2470,48 @@ UE::Tasks::FTask ConvertSkeletalMeshToMutable(TSharedRef<FMeshConversionContext>
 
 	}
 
-	// Set Bone Parenting usages. This has to be done after all primary usages are set.
-	for (int32 I = MutableMesh->GetBonePoseCount() - 1; I >= 0; --I)
+	if (MutableMesh->GetSkeleton()) 
 	{
-		UE::Mutable::Private::FMesh::FBonePose& BonePose = MutableMesh->BonePoses[I];
+		const UE::Mutable::Private::FSkeleton& MutableSkeleton = *MutableMesh->GetSkeleton();
 
-		constexpr UE::Mutable::Private::EBoneUsageFlags FlagsToPropagate =
-			UE::Mutable::Private::EBoneUsageFlags::Skinning | UE::Mutable::Private::EBoneUsageFlags::Physics | UE::Mutable::Private::EBoneUsageFlags::Deform;
-		if (EnumHasAnyFlags(BonePose.BoneUsageFlags, FlagsToPropagate))
+		// Set Bone Parenting usages. This has to be done after all primary usages are set.
+		for (int32 I = MutableMesh->GetBonePoseCount() - 1; I >= 0; --I)
 		{
-			const int32 Index = MutableMesh->GetSkeleton()->FindBone(BonePose.BoneId);
+			UE::Mutable::Private::FMesh::FBonePose& BonePose = MutableMesh->BonePoses[I];
 
-			if (Index == INDEX_NONE)
+			constexpr UE::Mutable::Private::EBoneUsageFlags FlagsToPropagate =
+				UE::Mutable::Private::EBoneUsageFlags::Skinning | UE::Mutable::Private::EBoneUsageFlags::Physics | UE::Mutable::Private::EBoneUsageFlags::Deform;
+			if (EnumHasAnyFlags(BonePose.BoneUsageFlags, FlagsToPropagate))
 			{
-				continue;
-			}
+				const int32 ParentIndex = MutableSkeleton.BoneParents[BonePose.BoneId.Index];
 
-			const int32 ParentIndex = MutableMesh->GetSkeleton()->GetBoneParent(Index);
+				if (ParentIndex == INDEX_NONE)
+				{
+					continue;
+				}
 
-			if (ParentIndex == INDEX_NONE)
-			{
-				continue;
-			}
+				int32 PoseIndex = MutableMesh->FindBonePoseByBoneIndex(ParentIndex);
+				if (PoseIndex == INDEX_NONE)
+				{
+					continue;
+				}
+
+				const UE::Mutable::Private::EBoneUsageFlags ParentPropagationFlags =
+					(EnumHasAnyFlags(BonePose.BoneUsageFlags, UE::Mutable::Private::EBoneUsageFlags::Skinning)
+						? UE::Mutable::Private::EBoneUsageFlags::SkinningParent : UE::Mutable::Private::EBoneUsageFlags::None) |
+					(EnumHasAnyFlags(BonePose.BoneUsageFlags, UE::Mutable::Private::EBoneUsageFlags::Physics)
+						? UE::Mutable::Private::EBoneUsageFlags::PhysicsParent : UE::Mutable::Private::EBoneUsageFlags::None) |
+					(EnumHasAnyFlags(BonePose.BoneUsageFlags, UE::Mutable::Private::EBoneUsageFlags::Deform)
+						? UE::Mutable::Private::EBoneUsageFlags::DeformParent : UE::Mutable::Private::EBoneUsageFlags::None);
+
+				do
+				{
+					EnumAddFlags(MutableMesh->BonePoses[PoseIndex].BoneUsageFlags, ParentPropagationFlags);
 
-			const UE::Mutable::Private::EBoneUsageFlags ParentPropagationFlags =
-				(EnumHasAnyFlags(BonePose.BoneUsageFlags, UE::Mutable::Private::EBoneUsageFlags::Skinning)
-					? UE::Mutable::Private::EBoneUsageFlags::SkinningParent : UE::Mutable::Private::EBoneUsageFlags::None) |
-				(EnumHasAnyFlags(BonePose.BoneUsageFlags, UE::Mutable::Private::EBoneUsageFlags::Physics)
-					? UE::Mutable::Private::EBoneUsageFlags::PhysicsParent : UE::Mutable::Private::EBoneUsageFlags::None) |
-				(EnumHasAnyFlags(BonePose.BoneUsageFlags, UE::Mutable::Private::EBoneUsageFlags::Deform)
-					? UE::Mutable::Private::EBoneUsageFlags::DeformParent : UE::Mutable::Private::EBoneUsageFlags::None);
+					PoseIndex = MutableMesh->FindBonePoseByBoneIndex(MutableSkeleton.BoneParents[MutableMesh->BonePoses[PoseIndex].BoneId.Index]);
 
-			SetAndPropagatePoseBoneUsage(*MutableMesh, ParentIndex, ParentPropagationFlags);
+				} while (PoseIndex != INDEX_NONE);
+			}
 		}
 	}
 
@@ -2563,8 +2536,7 @@ UE::Tasks::FTask ConvertSkeletalMeshToMutable(TSharedRef<FMeshConversionContext>
 
 			const int32 PhysicsAssetId = Context.AnimBpOverridePhysicsAssetsInfo.AddUnique(Info);
 
-			UE::Mutable::Private::TManagedPtr<UE::Mutable::Private::FPhysicsBody> MutableBody = MakePhysicsBodyFromAsset(Context,
-				PropertyAsset,
+			UE::Mutable::Private::TManagedPtr<UE::Mutable::Private::FPhysicsBody> MutableBody = MakePhysicsBodyFromAsset(PropertyAsset, *MutableMesh->GetSkeleton(),
 				MakePhysicsAssetBodySetupRelevancyMap(Context, PropertyAsset, MutableMesh));
 			MutableBody->CustomId = PhysicsAssetId;
 
@@ -3644,30 +3616,18 @@ UE::Mutable::Private::NodeMeshPtr GenerateMorphMesh(const UEdGraphPin* Pin,
 					bool bWarningFound = false;
 					if (TypedMorphNode->bReshapeSkeleton)
 					{
-						TArray<FName> BonesToDeform;
 						bWarningFound = GetAndValidateReshapeBonesToDeform(
-							BonesToDeform, TypedMorphNode->BonesToDeform, SkeletalMeshesToDeform, TypedMorphNode, TypedMorphNode->SelectionMethod, GenerationContext);
-
-						for (const FName BoneName : BonesToDeform)
-						{
-							Result->BonesToDeform.Add(GenerationContext.CompilationContext->GetBoneUnique(BoneName));
-						}
+							Result->BonesToDeform, TypedMorphNode->BonesToDeform, SkeletalMeshesToDeform, TypedMorphNode, TypedMorphNode->SelectionMethod, GenerationContext);
 					}
 
 					if (TypedMorphNode->bReshapePhysicsVolumes)
 					{
-						TArray<FName> PhysicsToDeform;
-
-						const EBoneDeformSelectionMethod SelectionMethod = TypedMorphNode->PhysicsSelectionMethod;
 						bWarningFound = bWarningFound || GetAndValidateReshapePhysicsToDeform(
-							PhysicsToDeform, 
-							TypedMorphNode->PhysicsBodiesToDeform, SkeletalMeshesToDeform, SelectionMethod, 
+							Result->PhysicsToDeform, 
+							TypedMorphNode->PhysicsBodiesToDeform, 
+							SkeletalMeshesToDeform,
+							TypedMorphNode->PhysicsSelectionMethod,
 							TypedMorphNode, GenerationContext);
-	
-						for (const FName& PhysicsBoneName : PhysicsToDeform)
-						{
-							Result->PhysicsToDeform.Add(GenerationContext.CompilationContext->GetBoneUnique(PhysicsBoneName));
-						}
 					}
 						
 					if (bWarningFound)
@@ -4278,29 +4238,19 @@ UE::Mutable::Private::Ptr<UE::Mutable::Private::NodeMesh> GenerateMutableSourceM
 			bool bWarningFound = false;
 			if (TypedNodeReshape->bReshapeSkeleton)
 			{
-				TArray<FName> BonesToDeform;
 				bWarningFound = GetAndValidateReshapeBonesToDeform(
-					BonesToDeform, TypedNodeReshape->BonesToDeform, SkeletalMeshesToDeform, TypedNodeReshape, TypedNodeReshape->SelectionMethod, GenerationContext);
-				
-				for (const FName& BoneName : BonesToDeform)
-				{
-					MeshNode->BonesToDeform.Add(GenerationContext.CompilationContext->GetBoneUnique(BoneName));
-				}
+					MeshNode->BonesToDeform, TypedNodeReshape->BonesToDeform, SkeletalMeshesToDeform, TypedNodeReshape, TypedNodeReshape->SelectionMethod, GenerationContext);
 			}
 
 			if (TypedNodeReshape->bReshapePhysicsVolumes)
 			{
 				EBoneDeformSelectionMethod SelectionMethod = TypedNodeReshape->PhysicsSelectionMethod;
-				TArray<FName> PhysicsToDeform;
 				bWarningFound = bWarningFound || GetAndValidateReshapePhysicsToDeform(
-					PhysicsToDeform, 
-					TypedNodeReshape->PhysicsBodiesToDeform, SkeletalMeshesToDeform, SelectionMethod, 
+					MeshNode->PhysicsToDeform,
+					TypedNodeReshape->PhysicsBodiesToDeform,
+					SkeletalMeshesToDeform,
+					TypedNodeReshape->PhysicsSelectionMethod,
 					TypedNodeReshape, GenerationContext);
-
-				for (const FName& PhysicsBoneName : PhysicsToDeform)
-				{
-					MeshNode->PhysicsToDeform.Add(GenerationContext.CompilationContext->GetBoneUnique(PhysicsBoneName));
-				}	
 			}
 			
 			if (bWarningFound)
diff --git a/Engine/Plugins/Mutable/Source/CustomizableObjectEditor/Private/MuCOE/GenerateMutableSource/GenerateMutableSourceModifier.cpp b/Engine/Plugins/Mutable/Source/CustomizableObjectEditor/Private/MuCOE/GenerateMutableSource/GenerateMutableSourceModifier.cpp
--- a/Engine/Plugins/Mutable/Source/CustomizableObjectEditor/Private/MuCOE/GenerateMutableSource/GenerateMutableSourceModifier.cpp
+++ b/Engine/Plugins/Mutable/Source/CustomizableObjectEditor/Private/MuCOE/GenerateMutableSource/GenerateMutableSourceModifier.cpp
@@ -97,7 +97,7 @@ UE::Mutable::Private::Ptr<UE::Mutable::Private::NodeModifier> GenerateMutableSou
 		ClipNode->SetParams(TypedNodeClip->B, TypedNodeClip->Exponent);
 		ClipNode->SetMorphEllipse(TypedNodeClip->Radius, TypedNodeClip->Radius2, TypedNodeClip->RotationAngle);
 
-		ClipNode->SetVertexSelectionBone(GenerationContext.CompilationContext->GetBoneUnique(TypedNodeClip->BoneName), TypedNodeClip->MaxEffectRadius);
+		ClipNode->SetVertexSelectionBone(TypedNodeClip->BoneName, TypedNodeClip->MaxEffectRadius);
 
 		ClipNode->MultipleTagsPolicy = TypedNodeClip->MultipleTagPolicy;
 		ClipNode->RequiredTags = TypedNodeClip->GetNodeRequiredTags(&GenerationContext.MacroNodesStack);
@@ -653,7 +653,7 @@ UE::Mutable::Private::Ptr<UE::Mutable::Private::NodeModifier> GenerateMutableSou
 		UE::Mutable::Private::Ptr<UE::Mutable::Private::NodeModifierMeshTransformWithBone> TransformNode = new UE::Mutable::Private::NodeModifierMeshTransformWithBone();
 		Result = TransformNode;
 
-		TransformNode->BoneName = GenerationContext.CompilationContext->GetBoneUnique(*TypedNodeTransformWithBone->BoneName);
+		TransformNode->BoneName = *TypedNodeTransformWithBone->BoneName;
 		TransformNode->ThresholdFactor = TypedNodeTransformWithBone->ThresholdFactor;
 
 		if (const UEdGraphPin* ConnectedPin = FollowInputPin(*TypedNodeTransformWithBone->GetTransformPin()))
diff --git a/Engine/Plugins/Mutable/Source/CustomizableObjectEditor/Private/MuCOE/GenerateMutableSource/GenerateMutableSourceTable.cpp b/Engine/Plugins/Mutable/Source/CustomizableObjectEditor/Private/MuCOE/GenerateMutableSource/GenerateMutableSourceTable.cpp
--- a/Engine/Plugins/Mutable/Source/CustomizableObjectEditor/Private/MuCOE/GenerateMutableSource/GenerateMutableSourceTable.cpp
+++ b/Engine/Plugins/Mutable/Source/CustomizableObjectEditor/Private/MuCOE/GenerateMutableSource/GenerateMutableSourceTable.cpp
@@ -494,18 +494,22 @@ bool FillTableColumn(const UCustomizableObjectNodeTable* TableNode, UE::Mutable:
 
 				MutableMesh->SetSkeleton(MutableSkeleton);
 				MutableMesh->SetBonePoseCount(ArrayBoneName.Num());
-				MutableSkeleton->SetBoneCount(ArrayBoneName.Num());
+
+				MutableSkeleton->BoneNames.Reserve(ArrayBoneName.Num());
+				MutableSkeleton->BoneParents.Reserve(ArrayBoneName.Num());
 
 				for (int32 i = 0; i < ArrayBoneName.Num(); ++i)
 				{
-					const FName BoneName = ArrayBoneName[i];
-					const UE::Mutable::Private::FBoneName& BoneId = GenerationContext.CompilationContext->GetBoneUnique(BoneName);
+					constexpr int16 ParentIndex = -1; // Skeletons from poses don't need bone parents?
 
-					MutableSkeleton->SetDebugName(i, BoneName);
-					MutableSkeleton->SetBoneName(i, { BoneId });
-					MutableMesh->SetBonePose(i, BoneId, (FTransform3f)ArrayTransform[i], UE::Mutable::Private::EBoneUsageFlags::Skinning);
+					UE::Mutable::Private::FBoneIdOrIndex MutableBoneIndex;
+					MutableBoneIndex.Index = MutableSkeleton->AddBone(ArrayBoneName[i], ParentIndex);
+					MutableMesh->SetBonePose(i, MutableBoneIndex, (FTransform3f)ArrayTransform[i], UE::Mutable::Private::EBoneUsageFlags::Skinning);
 				}
 
+				// TODO: Review why Skeletons from poses don't need bone parents
+				MutableSkeleton->BoneParents.Empty();
+
 				MutableTable->SetCell(CurrentColumn, RowId, MutableMesh);
 			}
 		}
diff --git a/Engine/Plugins/Mutable/Source/CustomizableObjectEditor/Private/MuCOE/MutableMeshPreviewUtils.cpp b/Engine/Plugins/Mutable/Source/CustomizableObjectEditor/Private/MuCOE/MutableMeshPreviewUtils.cpp
--- a/Engine/Plugins/Mutable/Source/CustomizableObjectEditor/Private/MuCOE/MutableMeshPreviewUtils.cpp
+++ b/Engine/Plugins/Mutable/Source/CustomizableObjectEditor/Private/MuCOE/MutableMeshPreviewUtils.cpp
@@ -384,15 +384,12 @@ namespace MutableMeshPreviewUtils
 		{
 			OutSkeletalMesh->SetHasVertexColors(false);
 
-			// Find how many bones the bonemap could have
+			// Find how many bones the BoneMap could have
 			const int32 NumBonesInBoneMap = !InMutableMesh->GetBoneMap().IsEmpty() ? InMutableMesh->GetBoneMap().Num() : InMutableMesh->GetBonePoseCount();
 			
-			// Fill the bonemap with zeros
-			TArray<UE::Mutable::Private::FBoneName> BoneMap;
-			BoneMap.SetNumZeroed(FMath::Max(NumBonesInBoneMap, 1));
-
-			TMap <UE::Mutable::Private::FBoneName, TPair<FName, uint16>> BoneInfoMap;
-			BoneInfoMap.Add(UE::Mutable::Private::FBoneName(0), {NAME_None, 0});
+			// Fill the BoneMap with zeros
+			TArray<FName> BoneMap;
+			BoneMap.Init(FName("Root"), FMath::Max(NumBonesInBoneMap, 1));
 
 			FSkeletalMeshLODRenderData& LODResource = OutSkeletalMesh->GetResourceForRendering()->LODRenderData[0];
 
@@ -404,8 +401,8 @@ namespace MutableMeshPreviewUtils
 			UnrealConversionUtils::SetupRenderSections(
 				LODResource,
 				InMutableMesh,
+				OutSkeletalMesh->GetRefSkeleton(),
 				BoneMap,
-				BoneInfoMap,
 				0,
 				MeshSurfacesMetadata);
 
@@ -415,7 +412,7 @@ namespace MutableMeshPreviewUtils
 				false); 
 			
 
-			// Ensure there's at least one bone in the bonemap of each render section, the root.
+			// Ensure there's at least one bone in the BoneMap of each render section, the root.
 			for (FSkelMeshRenderSection& RenderSection : LODResource.RenderSections)
 			{
 				if (RenderSection.BoneMap.IsEmpty())
diff --git a/Engine/Plugins/Mutable/Source/CustomizableObjectEditor/Private/MuCOE/SMutableCodeViewer.cpp b/Engine/Plugins/Mutable/Source/CustomizableObjectEditor/Private/MuCOE/SMutableCodeViewer.cpp
--- a/Engine/Plugins/Mutable/Source/CustomizableObjectEditor/Private/MuCOE/SMutableCodeViewer.cpp
+++ b/Engine/Plugins/Mutable/Source/CustomizableObjectEditor/Private/MuCOE/SMutableCodeViewer.cpp
@@ -2664,31 +2664,6 @@ bool SMutableCodeViewer::IsConstantResourceUsedByOperation(const int32 IndexOnCo
 				{
 					bResourceLocated = IndexOnConstantsArray == InProgram.GetOpArgs<UE::Mutable::Private::OP::InstanceAddExtensionDataArgs>(OperationAddress).ExtensionDataName;
 				}
-				else if (OperationType == UE::Mutable::Private::EOpType::ME_BINDSHAPE)
-				{
-					UE::Mutable::Private::OP::MeshBindShapeArgs Arguments = InProgram.GetOpArgs<UE::Mutable::Private::OP::MeshBindShapeArgs>(OperationAddress);
-					const uint8_t* Data = InProgram.GetOpArgsPointer(OperationAddress);
-
-					// Bones are stored after the args
-					Data += sizeof(Arguments);
-
-					// Iterate over the bones and check if they point to the same index on the string constants array
-					int32 NumBones;
-					FMemory::Memcpy(&NumBones, Data, sizeof(int32)); 
-					Data += sizeof(int32);
-				
-					for (int32 Bone = 0; Bone < NumBones; ++Bone)
-					{
-						// Exit once we know that the data is pointing to the index provided
-						if (*Data == IndexOnConstantsArray)
-						{
-							bResourceLocated = true;
-							continue;
-						}
-						
-						Data += sizeof(int32);
-					}
-				}
 				else if (OperationType == UE::Mutable::Private::EOpType::ME_ADDMETADATA)
 				{
 					 const uint8* OpData = InProgram.GetOpArgsPointer(OperationAddress);
@@ -2831,10 +2806,6 @@ bool SMutableCodeViewer::IsConstantResourceUsedByOperation(const int32 IndexOnCo
 				{
 					bResourceLocated = IndexOnConstantsArray == InProgram.GetOpArgs<UE::Mutable::Private::OP::MeshConstantArgs>(OperationAddress).Skeleton;
 				}
-				else if (OperationType == UE::Mutable::Private::EOpType::ME_SETSKELETON)
-				{
-					bResourceLocated = IndexOnConstantsArray == InProgram.GetOpArgs<UE::Mutable::Private::OP::MeshSetSkeletonArgs>(OperationAddress).Skeleton;
-				}
 				break;
 			}
 		case UE::Mutable::Private::EDataType::PhysicsAsset:
diff --git a/Engine/Plugins/Mutable/Source/CustomizableObjectEditor/Private/MuCOE/SMutableMeshViewer.cpp b/Engine/Plugins/Mutable/Source/CustomizableObjectEditor/Private/MuCOE/SMutableMeshViewer.cpp
--- a/Engine/Plugins/Mutable/Source/CustomizableObjectEditor/Private/MuCOE/SMutableMeshViewer.cpp
+++ b/Engine/Plugins/Mutable/Source/CustomizableObjectEditor/Private/MuCOE/SMutableMeshViewer.cpp
@@ -800,7 +800,7 @@ FText SMutableMeshViewer::GetFaceCount() const
 
 FText SMutableMeshViewer::GetBoneCount() const
 {
-	return FText::AsNumber(MutableMesh && MutableMesh->GetSkeleton() ? MutableMesh->GetSkeleton()->GetBoneCount() : 0);
+	return FText::AsNumber(MutableMesh && MutableMesh->GetSkeleton() ? MutableMesh->GetSkeleton()->GetNumBones() : 0);
 }
 
 FText SMutableMeshViewer::GetMeshIdPrefix() const
diff --git a/Engine/Plugins/Mutable/Source/CustomizableObjectEditor/Private/MuCOE/SMutableSkeletonViewer.cpp b/Engine/Plugins/Mutable/Source/CustomizableObjectEditor/Private/MuCOE/SMutableSkeletonViewer.cpp
--- a/Engine/Plugins/Mutable/Source/CustomizableObjectEditor/Private/MuCOE/SMutableSkeletonViewer.cpp
+++ b/Engine/Plugins/Mutable/Source/CustomizableObjectEditor/Private/MuCOE/SMutableSkeletonViewer.cpp
@@ -115,7 +115,7 @@ void SMutableSkeletonViewer::FillBoneDefinitionArrays()
 		return;
 	}
 
-	const int32 BoneCount = Skeleton->GetBoneCount();
+	const int32 BoneCount = Skeleton->GetNumBones();
 	for (int32 BoneIndex = 0; BoneIndex < BoneCount; BoneIndex++)
 	{
 		// Generate a new definition object
@@ -123,14 +123,14 @@ void SMutableSkeletonViewer::FillBoneDefinitionArrays()
 
 		// Save the bone data
 		NewBoneDefinition->BoneIndex = BoneIndex;
-		NewBoneDefinition->BoneName = Skeleton->GetDebugName(BoneIndex);
+		NewBoneDefinition->BoneName = Skeleton->BoneNames[BoneIndex];
 
 		// Get the parent node if any
-		const int32 ParentBoneIndex = Skeleton->GetBoneParent(BoneIndex);
+		const int32 ParentBoneIndex = Skeleton->BoneParents[BoneIndex];
 		if (ParentBoneIndex != INDEX_NONE)
 		{
 			NewBoneDefinition->ParentBoneIndex = ParentBoneIndex;
-			NewBoneDefinition->ParentBoneName = Skeleton->GetDebugName(ParentBoneIndex);
+			NewBoneDefinition->ParentBoneName = Skeleton->BoneNames[ParentBoneIndex];
 
 			// Store child bones
 			ChildBones.Add(NewBoneDefinition);
diff --git a/Engine/Plugins/Mutable/Source/MutableRuntime/Internal/MuR/Mesh.h b/Engine/Plugins/Mutable/Source/MutableRuntime/Internal/MuR/Mesh.h
--- a/Engine/Plugins/Mutable/Source/MutableRuntime/Internal/MuR/Mesh.h
+++ b/Engine/Plugins/Mutable/Source/MutableRuntime/Internal/MuR/Mesh.h
@@ -342,6 +342,8 @@ namespace UE::Mutable::Private
         UE_API void SetSkeleton(const TManagedPtr<const FSkeleton>&);
 		UE_API TManagedPtr<const FSkeleton> GetSkeleton() const;
 
+		void ConvertBoneIdsToIndices();
+
         //! \}
 
         //! \name PhysicsBody information
@@ -378,7 +380,7 @@ namespace UE::Mutable::Private
 		UE_API const TArray<uint64>& GetStreamedResources() const;
 
 		//!
-		UE_API int32 FindBonePose(const FBoneName& BoneName) const;
+		UE_API int32 FindBonePoseByBoneIndex(int32 Index) const;
 		
 		//!
 		UE_API void SetBonePoseCount(int32 Count);
@@ -387,10 +389,10 @@ namespace UE::Mutable::Private
 		UE_API int32 GetBonePoseCount() const;
 
 		//!
-		UE_API void SetBonePose(int32 Index, const FBoneName& BoneName, FTransform3f Transform, EBoneUsageFlags BoneUsageFlags);
+		UE_API void SetBonePose(int32 Index, FBoneIdOrIndex BoneName, FTransform3f Transform, EBoneUsageFlags BoneUsageFlags);
 
 		//! @return - Bone identifier of the pose at 'Index'.
-		UE_API const FBoneName& GetBonePoseId(int32 BoneIndex) const;
+		UE_API FBoneIdOrIndex GetBonePoseId(int32 BoneIndex) const;
 
 		//! Return a matrix stored per bone. It is a set of 16-float values.
 		UE_API void GetBonePoseTransform(int32 BoneIndex, FTransform3f& Transform) const;
@@ -399,10 +401,10 @@ namespace UE::Mutable::Private
 		UE_API EBoneUsageFlags GetBoneUsageFlags(int32 BoneIndex) const;
 
 		//! Set the bonemap of this mesh
-		UE_API void SetBoneMap(const TArray<FBoneName>& InBoneMap);
+		UE_API void SetBoneMap(const TArray<FBoneIdOrIndex>& InBoneMap);
 
 		//! Return an array containing the bonemaps of all surfaces in the mesh.
-		UE_API const TArray<FBoneName>& GetBoneMap() const;
+		UE_API const TArray<FBoneIdOrIndex>& GetBoneMap() const;
 
 		//!
 		UE_API int32 GetSkeletonIDsCount() const;
@@ -492,7 +494,7 @@ namespace UE::Mutable::Private
 		struct FBonePose
 		{
 			// Identifier built from the bone FName.
-			FBoneName BoneId;
+			FBoneIdOrIndex BoneId;
 
 			EBoneUsageFlags BoneUsageFlags = EBoneUsageFlags::None;
 			
@@ -523,7 +525,7 @@ namespace UE::Mutable::Private
 		int16 BonePosePriority = 0;
 
 		/** Array containing the bonemaps of all surfaces in the mesh. */
-		TArray<FBoneName> BoneMap;
+		TArray<FBoneIdOrIndex> BoneMap;
     	
 		inline void Serialise(FOutputArchive& arch) const;
 
diff --git a/Engine/Plugins/Mutable/Source/MutableRuntime/Internal/MuR/Model.h b/Engine/Plugins/Mutable/Source/MutableRuntime/Internal/MuR/Model.h
--- a/Engine/Plugins/Mutable/Source/MutableRuntime/Internal/MuR/Model.h
+++ b/Engine/Plugins/Mutable/Source/MutableRuntime/Internal/MuR/Model.h
@@ -339,6 +339,9 @@ namespace UE::Mutable::Private
 		/** Constant Material Data */
 		TArray<TManagedPtr<const FMaterial>> ConstantMaterials;
 
+		/** Constant FNames */
+		TMap<uint32, FName> ConstantNames;
+
     	TSharedPtr<const IExternalOperationProvider> ExternalOperationProvider;
     	
 #if WITH_EDITOR
@@ -359,6 +362,8 @@ namespace UE::Mutable::Private
 
         void Unserialise(FInputArchive& Arch);
 
+		void PostLoad();
+
         /** Debug method that logs the top used instruction types. */
         void LogHistogram() const;
 
@@ -412,6 +417,8 @@ namespace UE::Mutable::Private
     	
         UE_API OP::ADDRESS AddConstant(const TArray<FString>& StringList);
 
+		UE_API uint32 AddConstant(const FName& InName);
+
         UE_API OP::CONSTANT_MATRIX_INDEX AddConstant(const FMatrix44f& Constant);
     	
         UE_API OP::CONSTANT_SHAPE_INDEX AddConstant(const FShape& Constant);
@@ -726,6 +733,11 @@ namespace UE::Mutable::Private
 						}
 					}
 				}
+
+				if (TManagedPtr<FPhysicsBody> PhysicsBody = ConstCastManagedPtr<FPhysicsBody>(Result->PhysicsBody))
+				{
+					PhysicsBody->ConvertBoneIdsToNames(ConstantNames);
+				}
 			}
 
 			if (MetaDataMesh)
@@ -739,7 +751,7 @@ namespace UE::Mutable::Private
 				}
 				else
 				{
-					// Only in case the geoemtry has been filtered, add the meta data descriptors.
+					// Only in case the geometry has been filtered, add the meta data descriptors.
 					if (EnumHasAnyFlags(MeshContentRange.GetContentFlags(), EMeshContentFlags::GeometryData) &&
 						!EnumHasAnyFlags(FilterContentFlags, EMeshContentFlags::GeometryData))
 					{
@@ -763,6 +775,8 @@ namespace UE::Mutable::Private
 			if (ConstantSkeletons.IsValidIndex(SkeletonConstantIndex))
 			{
 				Result->Skeleton = ConstantSkeletons[SkeletonConstantIndex];
+
+				Result->ConvertBoneIdsToIndices();
 			}
 
 			OutMesh = Result;
diff --git a/Engine/Plugins/Mutable/Source/MutableRuntime/Internal/MuR/OpMeshMerge.h b/Engine/Plugins/Mutable/Source/MutableRuntime/Internal/MuR/OpMeshMerge.h
--- a/Engine/Plugins/Mutable/Source/MutableRuntime/Internal/MuR/OpMeshMerge.h
+++ b/Engine/Plugins/Mutable/Source/MutableRuntime/Internal/MuR/OpMeshMerge.h
@@ -260,6 +260,29 @@ namespace UE::Mutable::Private
 		}
 	}
 
+	inline void MergeSkeletons(FSkeleton& Base, const FSkeleton& Other, TArray<int32>& OutRemappedOtherBoneIndices)
+	{
+		const int32 NumBonesBaseBeforeMerge = Base.GetNumBones();
+		const int32 NumBonesOther = Other.GetNumBones();
+		OutRemappedOtherBoneIndices.SetNumUninitialized(NumBonesOther);
+
+		// Merge pSecond and build the remap array 
+		for (int32 Index = 0; Index < NumBonesOther; ++Index)
+		{
+			OutRemappedOtherBoneIndices[Index] = Base.FindOrAddBone(Other.BoneNames[Index], Other.BoneParents[Index]);
+		}
+
+		// Remap Parent bone indices to new Skeleton
+		for (int32 Index = NumBonesBaseBeforeMerge; Index < Base.BoneParents.Num(); ++Index)
+		{
+			int16 OtherBoneIndex = Base.BoneParents[Index];
+			if (OtherBoneIndex != INDEX_NONE)
+			{
+				Base.BoneParents[Index] = OutRemappedOtherBoneIndices[OtherBoneIndex];
+			}
+		}
+	}
+
 	//---------------------------------------------------------------------------------------------
 	//! Merge two meshes into one new mesh
 	//---------------------------------------------------------------------------------------------
@@ -623,70 +646,39 @@ namespace UE::Mutable::Private
 			Result->SkeletonIDs.AddUnique(SkeletonID);
 		}
 
+		TManagedPtr<const FSkeleton> FirstSkeleton = pFirst->GetSkeleton();
+		TManagedPtr<const FSkeleton> SecondSkeleton = pSecond->GetSkeleton();
+
 		// Do they have the same skeleton?
-		bool bMergeSkeletons = pFirst->GetSkeleton() != pSecond->GetSkeleton();
+		bool bMergeSkeletons = FirstSkeleton != SecondSkeleton;
 
 		// Are they different skeletons but with the same data?
-		if (bMergeSkeletons && pFirst->GetSkeleton() && pSecond->GetSkeleton())
+		if (bMergeSkeletons && FirstSkeleton && SecondSkeleton)
 		{
-			bMergeSkeletons = !(*pFirst->GetSkeleton() == *pSecond->GetSkeleton());
+			bMergeSkeletons = !(*FirstSkeleton == *SecondSkeleton);
 		}
 
+		TArray<int32> RemappedSecondBoneIndices;
 
 		if (bMergeSkeletons)
 		{
 			MUTABLE_CPUPROFILER_SCOPE(MergeSkeleton);
-			TManagedPtr<FSkeleton> ResultSkeleton;
 
-			TManagedPtr<const FSkeleton> FirstSkeleton = pFirst->GetSkeleton();
-			TManagedPtr<const FSkeleton> SecondSkeleton = pSecond->GetSkeleton();
-
-			const int32 NumBonesFirst = FirstSkeleton ? FirstSkeleton->GetBoneCount() : 0;
-			const int32 NumBonesSecond = SecondSkeleton ? SecondSkeleton->GetBoneCount() : 0;
-
-			ResultSkeleton = FirstSkeleton ? FirstSkeleton->Clone() : MakeManaged<FSkeleton>();
+			TManagedPtr<FSkeleton> ResultSkeleton = FirstSkeleton ? FirstSkeleton->Clone() : MakeManaged<FSkeleton>();
 			Result->SetSkeleton(ResultSkeleton);
 
-			TArray<uint16> SecondToResultBoneIndices;
-			SecondToResultBoneIndices.SetNumUninitialized(NumBonesSecond);
-
-			// Merge pSecond and build the remap array 
-			for (int32 SecondBoneIndex = 0; SecondBoneIndex < NumBonesSecond; ++SecondBoneIndex)
-			{
-				const FBoneName& BoneNameId = SecondSkeleton->BoneIds[SecondBoneIndex];
-				int32 Index = ResultSkeleton->FindBone(BoneNameId);
-
-				// Add a new bone
-				if (Index == INDEX_NONE)
-				{
-					Index = ResultSkeleton->BoneIds.Add(BoneNameId);
-
-					// Add an incorrect index, to be fixed below in case the parent index is later in the bone array.
-					ResultSkeleton->BoneParents.Add(SecondSkeleton->BoneParents[SecondBoneIndex]);
-#if WITH_EDITOR
-					if (SecondSkeleton->DebugBoneNames.IsValidIndex(SecondBoneIndex))
-					{
-						ResultSkeleton->DebugBoneNames.Add(SecondSkeleton->DebugBoneNames[SecondBoneIndex]);
-					}
-#endif
-				}
-
-				SecondToResultBoneIndices[SecondBoneIndex] = (uint16)Index;
-			}
-
-			// Fix second mesh bone parents
-			for (int32 ob = NumBonesFirst; ob < ResultSkeleton->BoneParents.Num(); ++ob)
-			{
-				int16 secondMeshIndex = ResultSkeleton->BoneParents[ob];
-				if (secondMeshIndex != INDEX_NONE)
-				{
-					ResultSkeleton->BoneParents[ob] = SecondToResultBoneIndices[secondMeshIndex];
-				}
-			}
+			MergeSkeletons(*ResultSkeleton.Get(), *SecondSkeleton.Get(), RemappedSecondBoneIndices);
 		}
 		else
 		{
-			Result->SetSkeleton(pFirst->GetSkeleton());
+			Result->SetSkeleton(FirstSkeleton);
+
+			const int32 NumBonesFirst = FirstSkeleton ? FirstSkeleton->GetNumBones() : 0;
+			RemappedSecondBoneIndices.SetNumUninitialized(NumBonesFirst);
+			for (int32 Index = 0; Index < NumBonesFirst; ++Index)
+			{
+				RemappedSecondBoneIndices[Index] = Index;
+			}
 		}
 
 
@@ -697,9 +689,10 @@ namespace UE::Mutable::Private
 		bool bRemapBoneIndices = false;
 		TArray<uint16> RemappedBoneMapIndices;
 
+		const int32 NumBonesInSecondBoneMap = pSecond->BoneMap.Num();
+
 		// Used to know the format of the bone index buffer
 		uint32 MaxNumBonesInBoneMaps = 0;
-		const int32 NumSecondBonesInBoneMap = pSecond->BoneMap.Num();
 
 		{
 			MUTABLE_CPUPROFILER_SCOPE(Surfaces);
@@ -707,17 +700,29 @@ namespace UE::Mutable::Private
 			const int32 NumFirstBonesInBoneMap = pFirst->BoneMap.Num();
 			Result->BoneMap = pFirst->BoneMap;
 
+			// Copy BoneMap to remap bones when merging skeletons
+			TArray<FBoneIdOrIndex> SecondBoneMap = pSecond->BoneMap;
+
+			if (bMergeSkeletons)
+			{
+				// Remap BoneMap to new skeleton
+				for (int32 BoneIndex = 0; BoneIndex < NumBonesInSecondBoneMap; ++BoneIndex)
+				{
+					SecondBoneMap[BoneIndex].Index = RemappedSecondBoneIndices[SecondBoneMap[BoneIndex].Index];
+				}
+			}
+
 			if (bMergeSurfaces)
 			{
 				// Merge BoneMaps
-				RemappedBoneMapIndices.SetNumUninitialized(NumSecondBonesInBoneMap);
+				RemappedBoneMapIndices.Reserve(NumBonesInSecondBoneMap);
 
-				for (uint16 SecondBoneMapIndex = 0; SecondBoneMapIndex < NumSecondBonesInBoneMap; ++SecondBoneMapIndex)
+				for (uint16 Index = 0; Index < NumBonesInSecondBoneMap; ++Index)
 				{
-					const int32 BoneMapIndex = Result->BoneMap.AddUnique(pSecond->BoneMap[SecondBoneMapIndex]);
-					RemappedBoneMapIndices[SecondBoneMapIndex] = BoneMapIndex;
+					const int32 BoneMapIndex = Result->BoneMap.AddUnique(SecondBoneMap[Index]);
+					RemappedBoneMapIndices.Add(BoneMapIndex);
 
-					bRemapBoneIndices = bRemapBoneIndices || BoneMapIndex != SecondBoneMapIndex;
+					bRemapBoneIndices = bRemapBoneIndices || BoneMapIndex != Index;
 				}
 
 				FMeshSurface& NewSurface = Result->Surfaces.AddDefaulted_GetRef();
@@ -787,14 +792,11 @@ namespace UE::Mutable::Private
 
 		// Pose
 		//---------------------------
-		if (Result->GetSkeleton())
+		if (TManagedPtr<const FSkeleton> ResultSkeleton = Result->GetSkeleton())
 		{
 			MUTABLE_CPUPROFILER_SCOPE(Pose);
 
-			if (Result->GetSkeleton())
-			{
-				Result->BonePoses.Reserve(Result->GetSkeleton()->GetBoneCount());
-			}
+			Result->BonePoses.Reserve(ResultSkeleton->BoneNames.Num());
 
 			// Copy poses from the first mesh
 			Result->BonePoses = pFirst->BonePoses;
@@ -802,11 +804,12 @@ namespace UE::Mutable::Private
 			// Add or override bone poses
 			for (const FMesh::FBonePose& SecondBonePose : pSecond->BonePoses)
 			{
-				const int32 ResultBoneIndex = Result->FindBonePose(SecondBonePose.BoneId);
+				const int32 RemappedBoneIndex = RemappedSecondBoneIndices[SecondBonePose.BoneId.Index];
+				const int32 ResultBonePoseIndex = Result->FindBonePoseByBoneIndex(RemappedBoneIndex);
 
-				if (ResultBoneIndex != INDEX_NONE)
+				if (ResultBonePoseIndex != INDEX_NONE)
 				{
-					FMesh::FBonePose& ResultBonePose = Result->BonePoses[ResultBoneIndex];
+					FMesh::FBonePose& ResultBonePose = Result->BonePoses[ResultBonePoseIndex];
 
 					// TODO: Not sure how to tune this priority, review it.
 					// For now use a similar strategy as before. 
@@ -832,7 +835,9 @@ namespace UE::Mutable::Private
 				}
 				else
 				{
-					Result->BonePoses.Add(SecondBonePose);
+					FMesh::FBonePose NewPose = SecondBonePose;
+					NewPose.BoneId.Index = RemappedBoneIndex;
+					Result->BonePoses.Add(NewPose);
 				}
 			}
 
@@ -848,11 +853,11 @@ namespace UE::Mutable::Private
 			// Appends InPhysicsBody to OutPhysicsBody removing Bodies that are equal, have same bone and customId and its properies are identical.	
 			auto AppendPhysicsBodiesUnique = [](FPhysicsBody& OutPhysicsBody, const FPhysicsBody& InPhysicsBody) -> bool
 			{
-				TArray<FBoneName>& OutBones = OutPhysicsBody.BoneIds;
+				TArray<FName>& OutBones = OutPhysicsBody.BodiesBoneNames;
 				TArray<int32>& OutCustomIds = OutPhysicsBody.BodiesCustomIds;
 				TArray<FPhysicsBodyAggregate>& OutBodies = OutPhysicsBody.Bodies;
 
-				const TArray<FBoneName>& InBones = InPhysicsBody.BoneIds;
+				const TArray<FName>& InBones = InPhysicsBody.BodiesBoneNames;
 				const TArray<int32>& InCustomIds = InPhysicsBody.BodiesCustomIds;
 				const TArray<FPhysicsBodyAggregate>& InBodies = InPhysicsBody.Bodies;
 
@@ -1280,7 +1285,7 @@ namespace UE::Mutable::Private
 										uint8 BoneMapIndex = pD[ComponentIndex];
 
 										// be defensive
-										if (BoneMapIndex < NumSecondBonesInBoneMap)
+										if (BoneMapIndex < NumBonesInSecondBoneMap)
 										{
 											pD[ComponentIndex] = (uint8)RemappedBoneMapIndices[BoneMapIndex];
 										}
@@ -1304,7 +1309,7 @@ namespace UE::Mutable::Private
 										uint16 BoneMapIndex = pD[ComponentIndex];
 
 										// be defensive
-										if (BoneMapIndex < NumSecondBonesInBoneMap)
+										if (BoneMapIndex < NumBonesInSecondBoneMap)
 										{
 											pD[ComponentIndex] = (uint16)RemappedBoneMapIndices[BoneMapIndex];
 										}
@@ -1382,69 +1387,5 @@ namespace UE::Mutable::Private
 		{
 			MergeMorphs(*Result, *pFirst, *pSecond);	
 		}
-	}
-	
-
-    //---------------------------------------------------------------------------------------------
-    //!
-    //---------------------------------------------------------------------------------------------
-    inline void ExtendSkeleton( FSkeleton* pBase, const FSkeleton* pOther )
-    {
-        TMap<int32,int32> otherToResult;
-
-        int32 initialBones = pBase->GetBoneCount();
-		for (int32 b = 0; pOther && b < pOther->GetBoneCount(); ++b)
-		{
-			int32 resultBoneIndex = pBase->FindBone(pOther->GetBoneName(b));
-			if (resultBoneIndex < 0)
-			{
-				int32 newIndex = pBase->BoneIds.Num();
-				otherToResult.Add(b, newIndex);
-				pBase->BoneIds.Add(pOther->BoneIds[b]);
-
-				// Will be remapped below
-				pBase->BoneParents.Add(pOther->BoneParents[b]);
-
-#if WITH_EDITOR
-				if (pOther->DebugBoneNames.IsValidIndex(b))
-				{
-					pBase->DebugBoneNames.Add(pOther->DebugBoneNames[b]);
-				}
-#endif
-			}
-			else
-			{
-				otherToResult.Add(b, resultBoneIndex);
-			}
-		}
-
-        // Fix bone parent indices of the bones added from pOther
-        for ( int32 b=initialBones;b<pBase->GetBoneCount(); ++b)
-        {
-            int16_t sourceIndex = pBase->BoneParents[b];
-            if (sourceIndex>=0)
-            {
-                pBase->BoneParents[b] = (int16_t)otherToResult[ sourceIndex ];
-            }
-        }
-    }
-
-
-    //---------------------------------------------------------------------------------------------
-    //! Return null if there is no need to remap the mesh
-    //---------------------------------------------------------------------------------------------
-    inline void MeshRemapSkeleton(FMesh* Result, const FMesh* SourceMesh, TManagedPtr<const FSkeleton> Skeleton, bool& bOutSuccess)
-    {
-		bOutSuccess = true;
-
-        if (SourceMesh->GetVertexCount() == 0 || !SourceMesh->GetSkeleton() || SourceMesh->GetSkeleton()->GetBoneCount() == 0)
-        {
-			bOutSuccess = false;
-            return;
-        }
-
-		Result->CopyFrom(*SourceMesh);
-		Result->SetSkeleton(Skeleton);
-    }
-	
+	}	
 }
diff --git a/Engine/Plugins/Mutable/Source/MutableRuntime/Internal/MuR/Operations.h b/Engine/Plugins/Mutable/Source/MutableRuntime/Internal/MuR/Operations.h
--- a/Engine/Plugins/Mutable/Source/MutableRuntime/Internal/MuR/Operations.h
+++ b/Engine/Plugins/Mutable/Source/MutableRuntime/Internal/MuR/Operations.h
@@ -290,9 +290,6 @@ namespace UE::Mutable::Private
         //! Clip the mesh with another mesh.
         ME_CLIPWITHMESH,
 
-        //! Replace the skeleton data from a mesh with another one.
-        ME_SETSKELETON,
-
         //! Project a mesh using a projector and clipping the irrelevant faces
         ME_PROJECT,
 
@@ -1017,12 +1014,6 @@ namespace UE::Mutable::Private
             ADDRESS ClipMesh;
         };
 
-        struct MeshSetSkeletonArgs
-        {
-            ADDRESS Source;
-            ADDRESS Skeleton;
-        };
-
         struct MeshProjectArgs
         {
             ADDRESS Mesh;
diff --git a/Engine/Plugins/Mutable/Source/MutableRuntime/Internal/MuR/PhysicsBody.h b/Engine/Plugins/Mutable/Source/MutableRuntime/Internal/MuR/PhysicsBody.h
--- a/Engine/Plugins/Mutable/Source/MutableRuntime/Internal/MuR/PhysicsBody.h
+++ b/Engine/Plugins/Mutable/Source/MutableRuntime/Internal/MuR/PhysicsBody.h
@@ -165,12 +165,10 @@ namespace UE::Mutable::Private
 		UE_API void SetBodyCount(int32 B);
 		UE_API int32 GetBodyCount() const;
 	
-		UE_API void SetBodyBoneId(int32 B, const FBoneName& BoneId);
-		UE_API const FBoneName& GetBodyBoneId(int32 B) const;
-		
-		UE_API void SetBodyCustomId(int32 B, int32 BodyCustomId);
-		UE_API int32 GetBodyCustomId(int32 B) const;
-		
+		UE_API void SetBody(int32 BodyIndex, FName BoneName, int32 CustomId = INDEX_NONE);
+
+		UE_API FName GetBodyBoneName(int32 Index) const;
+
 		UE_API int32 GetSphereCount(int32 B) const;
 		UE_API int32 GetBoxCount(int32 B) const;
 		UE_API int32 GetConvexCount(int32 B) const;
@@ -263,11 +261,15 @@ namespace UE::Mutable::Private
 		UE_API const FString& GetSphylName(int32 B, int32 I) const;
 		UE_API const FString& GetTaperedCapsuleName(int32 B, int32 I) const;
 
+		void ConvertBoneIdsToNames(const TMap<uint32, FName>& ConstantNames);
+		void ConvertNamesToBoneIds(const TMap<uint32, FName>& ConstantNames);
+
 	public:
 		int32 CustomId = -1;
 
 		// Bone name the physics volume aggregate is bound to. 
-		TArray<FBoneName> BoneIds;
+		TArray<FName> BodiesBoneNames; // Runtime
+		TArray<uint32> BoneIds; // Serialized
 		TArray<FPhysicsBodyAggregate> Bodies;
 		TArray<int32> BodiesCustomIds;
 
diff --git a/Engine/Plugins/Mutable/Source/MutableRuntime/Internal/MuR/Skeleton.h b/Engine/Plugins/Mutable/Source/MutableRuntime/Internal/MuR/Skeleton.h
--- a/Engine/Plugins/Mutable/Source/MutableRuntime/Internal/MuR/Skeleton.h
+++ b/Engine/Plugins/Mutable/Source/MutableRuntime/Internal/MuR/Skeleton.h
@@ -16,14 +16,10 @@
 
 namespace UE::Mutable::Private
 {	
-	// Bone name identifier
-	struct FBoneName
+	union FBoneIdOrIndex
 	{
-		FBoneName() {};
-		FBoneName(uint32 InID) : Id(InID) {};
-
-		// Hash built from the bone name (FString)
-		uint32 Id = 0;
+		uint32 Id;
+		int32 Index;
 
 		inline void Serialise(FOutputArchive& Arch) const
 		{
@@ -35,19 +31,12 @@ namespace UE::Mutable::Private
 			Arch >> Id;
 		}
 
-		//!
-		inline bool operator==(const FBoneName& Other) const
+		inline bool operator==(const FBoneIdOrIndex& Other) const
 		{
 			return Id == Other.Id;
 		}
 	};
 
-	inline uint32 GetTypeHash(const FBoneName& Bone)
-	{
-		return Bone.Id;
-	}
-
-
     /** Skeleton object. */
 	class FSkeleton
 	{
@@ -69,33 +58,18 @@ namespace UE::Mutable::Private
 		// Own interface
 		//-----------------------------------------------------------------------------------------
 
-        //! @return - Number of bones in the Skeleton
-		UE_API int32 GetBoneCount() const;
-        UE_API void SetBoneCount(int32 Count);
-
-		//! @return - FName of the bone at 'Index'. Only valid in the editor
-		UE_API const FName GetDebugName(int32 Index) const;
-		UE_API void SetDebugName(const int32 Index, const FName BoneName);
-
-        //! Get and set the parent bone of each bone. The parent can be -1 if the bone is a root.
-        UE_API int32 GetBoneParent(int32 BoneIndex) const;
-        UE_API void SetBoneParent(int32 BoneIndex, int32 ParentBoneIndex);
-
-		//! @return - BoneName of the Bone at 'Index'.
-		UE_API const FBoneName& GetBoneName(int32 Index) const;
-		UE_API void SetBoneName(int32 Index, const FBoneName& BoneName);
+		UE_API int32 AddBone(const FName InBoneName, const int16 InParentIndex);
+		UE_API int32 FindOrAddBone(const FName InBoneName, const int16 InParentIndex);
 
-		//! @return - Index in the Skeleton. INDEX_NONE if not found.
-		UE_API int32 FindBone(const FBoneName& BoneName) const;
-		
+		UE_API int32 GetNumBones() const;
 
 	public:
 
-		//! DEBUG. FNames of the bones. Only valid in the editor. Do not serialize.
-		TArray<FName> DebugBoneNames;
+		//! FNames of the bones. Serialized in the program as ConstantNames.
+		TArray<FName> BoneNames;
 
-		//! Array of bone identifiers. 
-		TArray<FBoneName> BoneIds;
+		//! FName Id in the ConstantNames map. 
+		TArray<uint32> BoneIds;
 
 		//! For each bone, index of the parent bone in the bone vectors. -1 means no parent.
 		//! This array must have the same size than the m_bones array.
@@ -110,7 +84,8 @@ namespace UE::Mutable::Private
 		//!
 		inline bool operator==(const FSkeleton& Other) const
 		{
-			return BoneIds == Other.BoneIds
+			return BoneNames == Other.BoneNames
+				// && BoneIds == Other.BoneIds // Don't use boneIds to compare skeletons since they are added when creating the constants.
 				&& BoneParents == Other.BoneParents;
 		}
 	};
diff --git a/Engine/Plugins/Mutable/Source/MutableRuntime/Private/MuR/CodeRunner.cpp b/Engine/Plugins/Mutable/Source/MutableRuntime/Private/MuR/CodeRunner.cpp
--- a/Engine/Plugins/Mutable/Source/MutableRuntime/Private/MuR/CodeRunner.cpp
+++ b/Engine/Plugins/Mutable/Source/MutableRuntime/Private/MuR/CodeRunner.cpp
@@ -2107,7 +2107,7 @@ namespace UE::Mutable::Private
 						TManagedPtr<FMesh> Result = CreateMesh(Source ? Source->GetDataSize() : 0);
 						
 						bool bOutSuccess = false;
-						MeshClipMorphPlane(Result.Get(), Source.Get(), Origin, Normal, Args.Dist, Args.Factor, MorphShape.size[0], MorphShape.size[1], MorphShape.size[2], SelectionShape, bRemoveFaceIfAllVerticesCulled, bOutSuccess, nullptr, -1);
+						MeshClipMorphPlane(Result.Get(), Source.Get(), Origin, Normal, Args.Dist, Args.Factor, MorphShape.size[0], MorphShape.size[1], MorphShape.size[2], SelectionShape, bRemoveFaceIfAllVerticesCulled, bOutSuccess, NAME_None, -1);
 						if (!bOutSuccess)
 						{
 							StoreMesh(Item, Source);
@@ -2126,10 +2126,10 @@ namespace UE::Mutable::Private
 						TManagedPtr<FMesh> Result = CreateMesh(Source ? Source->GetDataSize() : 0);
 						
 						check(Args.VertexSelectionShapeOrBone <= MAX_uint32);
-						const FBoneName Bone(Args.VertexSelectionShapeOrBone);
-						
+						FName BoneName = Program.ConstantNames[Args.VertexSelectionShapeOrBone];
+
 						bool bOutSuccess = false;
-						MeshClipMorphPlane(Result.Get(), Source.Get(), Origin, Normal, Args.Dist, Args.Factor, MorphShape.size[0], MorphShape.size[1], MorphShape.size[2], SelectionShape, bRemoveFaceIfAllVerticesCulled, bOutSuccess, &Bone, Args.MaxBoneRadius);
+						MeshClipMorphPlane(Result.Get(), Source.Get(), Origin, Normal, Args.Dist, Args.Factor, MorphShape.size[0], MorphShape.size[1], MorphShape.size[2], SelectionShape, bRemoveFaceIfAllVerticesCulled, bOutSuccess, BoneName, Args.MaxBoneRadius);
 
 						if (!bOutSuccess)
 						{
@@ -2149,7 +2149,7 @@ namespace UE::Mutable::Private
 						TManagedPtr<FMesh> Result = CreateMesh(Source ? Source->GetDataSize() : 0);
 						
 						bool bOutSuccess = false;
-						MeshClipMorphPlane(Result.Get(), Source.Get(), Origin, Normal, Args.Dist, Args.Factor, MorphShape.size[0], MorphShape.size[1], MorphShape.size[2], SelectionShape, bRemoveFaceIfAllVerticesCulled, bOutSuccess, nullptr, -1.0f);
+						MeshClipMorphPlane(Result.Get(), Source.Get(), Origin, Normal, Args.Dist, Args.Factor, MorphShape.size[0], MorphShape.size[1], MorphShape.size[2], SelectionShape, bRemoveFaceIfAllVerticesCulled, bOutSuccess, NAME_None, -1.0f);
 						
 						if (!bOutSuccess)
 						{
@@ -2432,19 +2432,27 @@ namespace UE::Mutable::Private
 					FMemory::Memcpy(&NumBones, Data, sizeof(int32)); 
 					Data += sizeof(int32);
 					
-					TArray<FBoneName> BonesToDeform;
-					BonesToDeform.SetNumUninitialized(NumBones);
-					FMemory::Memcpy(BonesToDeform.GetData(), Data, NumBones * sizeof(FBoneName));
-					Data += NumBones * sizeof(FBoneName);
+					TArrayView<const uint32> BoneIds(reinterpret_cast<const uint32*>(Data), NumBones);
+					TArray<FName> BonesToDeform;
+					BonesToDeform.Reserve(NumBones);
+					for (uint32 BoneId : BoneIds)
+					{
+						BonesToDeform.Add(Program.ConstantNames[BoneId]);
+					}
+					Data += NumBones * sizeof(uint32);
 
 					int32 NumPhysicsBodies;
 					FMemory::Memcpy(&NumPhysicsBodies, Data, sizeof(int32)); 
 					Data += sizeof(int32);
 
-					TArray<FBoneName> PhysicsToDeform;
-					PhysicsToDeform.SetNumUninitialized(NumPhysicsBodies);
-					FMemory::Memcpy(PhysicsToDeform.GetData(), Data, NumPhysicsBodies * sizeof(FBoneName));
-					Data += NumPhysicsBodies * sizeof(FBoneName);
+					TArrayView<const uint32> PhysicsBoneIds(reinterpret_cast<const uint32*>(Data), NumPhysicsBodies);
+					TArray<FName> PhysicsToDeform;
+					PhysicsToDeform.Reserve(NumPhysicsBodies);
+					for (uint32 PhysicsBoneId : PhysicsBoneIds)
+					{
+						PhysicsToDeform.Add(Program.ConstantNames[PhysicsBoneId]);
+					}
+					Data += NumPhysicsBodies * sizeof(uint32);
 
 					EMeshBindShapeFlags BindFlags = static_cast<EMeshBindShapeFlags>(Args.flags);
 					EMeshContentFlags MeshContentFilter = (EMeshContentFlags)Item.ExecutionOptions;
@@ -2690,83 +2698,6 @@ namespace UE::Mutable::Private
 			break;
 		}
 
-        case EOpType::ME_SETSKELETON:
-        {
-			OP::MeshSetSkeletonArgs Args = Program.GetOpArgs<OP::MeshSetSkeletonArgs>(Item.At);
-            switch (Item.Stage)
-            {
-            case 0:
-			{
-				if (Args.Source)
-				{
-					AddOp(FScheduledOp(Item.At, Item, 1),
-						FScheduledOp(Args.Source, Item),
-						FScheduledOp(Args.Skeleton, Item));
-				}
-				else
-				{
-					StoreMesh(Item, nullptr);
-				}
-				break;
-			}
-            case 1:
-            {
-            	MUTABLE_CPUPROFILER_SCOPE(ME_SETSKELETON_1)
-            		
-                TManagedPtr<const FMesh> Source = LoadMesh(FCacheAddress(Args.Source, Item));
-                TManagedPtr<const FMesh> Skeleton = LoadMesh(FCacheAddress(Args.Skeleton, Item));
-
-                // Only if both are valid.
-                if (Source && Skeleton)
-                {
-                    if (Source->GetSkeleton() && Source->GetSkeleton()->GetBoneCount() > 0)
-                    {
-                        // For some reason we already have bone data, so we can't just overwrite it
-                        // or the skinning may break. This may happen because of a problem in the
-                        // optimiser that needs investigation.
-                        // \TODO Be defensive, for now.
-                        UE_LOG(LogMutableCore, Warning, TEXT("Performing a MeshRemapSkeleton, instead of MeshSetSkeletonData because source mesh already has some skeleton."));
-
-						TManagedPtr<FMesh> Result = CreateMesh(Source->GetDataSize());
-
-						bool bOutSuccess = false;
-                        MeshRemapSkeleton(Result.Get(), Source.Get(), Skeleton->GetSkeleton(), bOutSuccess);
-
-						if (!bOutSuccess)
-						{
-							StoreMesh(Item, Source);
-						}
-						else
-						{
-							//Result->GetPrivate()->CheckIntegrity();
-							StoreMesh(Item, Result);
-						}
-                    }
-                    else
-                    {
-						TManagedPtr<FMesh> Result = CloneOrTakeOver(MoveTemp(Source));
-
-                        Result->SetSkeleton(Skeleton->GetSkeleton());
-
-						//Result->GetPrivate()->CheckIntegrity();
-						StoreMesh(Item, Result);
-                    }
-                }
-				else
-				{
-					StoreMesh(Item, Source);
-				}
-
-                break;
-            }
-
-            default:
-                check(false);
-            }
-
-            break;
-        }
-
         case EOpType::ME_REMOVEMASK:
         {
        		MUTABLE_CPUPROFILER_SCOPE(ME_REMOVEMASK)
@@ -3281,13 +3212,13 @@ namespace UE::Mutable::Private
 				{
 					TManagedPtr<const FMesh> SourceMesh = LoadMesh(FCacheAddress(Args.SourceMesh, Item));
 					const FMatrix44f& Transform = LoadMatrix(FCacheAddress(Args.Matrix, Item));
-					const FBoneName BoneName = FBoneName(Args.BoneId);
+					FName BoneName = Program.ConstantNames[Args.BoneId];
 
 					if (SourceMesh)
 					{
 						TManagedPtr<FMesh> Result = CloneOrTakeOver(MoveTemp(SourceMesh));
 
-						MeshTransformWithBoneInline(Result.Get(), Transform, FBoneName(Args.BoneId), Args.ThresholdFactor);
+						MeshTransformWithBoneInline(Result.Get(), Transform, BoneName, Args.ThresholdFactor);
 						StoreMesh(Item, Result);
 					}
 					else
diff --git a/Engine/Plugins/Mutable/Source/MutableRuntime/Private/MuR/CodeRunnerBegin.cpp b/Engine/Plugins/Mutable/Source/MutableRuntime/Private/MuR/CodeRunnerBegin.cpp
--- a/Engine/Plugins/Mutable/Source/MutableRuntime/Private/MuR/CodeRunnerBegin.cpp
+++ b/Engine/Plugins/Mutable/Source/MutableRuntime/Private/MuR/CodeRunnerBegin.cpp
@@ -2145,18 +2145,6 @@ namespace UE::Mutable::Private
 					break;
 				}
 	
-			case EOpType::ME_SETSKELETON:
-				{
-					check(!Item.bEvaluate);
-
-					OP::MeshSetSkeletonArgs Args = Program.GetOpArgs<OP::MeshSetSkeletonArgs>(Item.At);
-					PushOp(FScheduledOpInline(Args.Source, Item));
-					PushOp(FScheduledOpInline(Args.Skeleton, Item));
-
-					StoreNone(Item);
-					break;
-				}
-	
 			case EOpType::ME_REMOVEMASK:
 				{
 					check(!Item.bEvaluate);
diff --git a/Engine/Plugins/Mutable/Source/MutableRuntime/Private/MuR/Mesh.cpp b/Engine/Plugins/Mutable/Source/MutableRuntime/Private/MuR/Mesh.cpp
--- a/Engine/Plugins/Mutable/Source/MutableRuntime/Private/MuR/Mesh.cpp
+++ b/Engine/Plugins/Mutable/Source/MutableRuntime/Private/MuR/Mesh.cpp
@@ -344,6 +344,42 @@ void FMesh::SetSkeleton(const TManagedPtr<const FSkeleton>& InSkeleton)
 }
 
 
+void FMesh::ConvertBoneIdsToIndices()
+{
+	if (ensure(Skeleton))
+	{
+		// Convert BoneMaps' BoneIds to BoneIndices
+		const int32 NumBonesBoneMap = BoneMap.Num();
+		for (int32 Index = 0; Index < NumBonesBoneMap; ++Index)
+		{
+			BoneMap[Index].Index = Skeleton->BoneIds.Find(BoneMap[Index].Id);
+		}
+
+		// Convert BonePoses' BoneIds to BoneIndices
+		const int32 NumBones = Skeleton->GetNumBones();
+		const int32 NumBonePoses = BonePoses.Num();
+
+		int32 AuxIndex = 0;
+		for (int32 Index = 0; Index < NumBonePoses; ++Index)
+		{
+			FBoneIdOrIndex& BoneId = BonePoses[Index].BoneId;
+
+			// Bone Poses are sorted to match the skeleton hierarchy. 
+			for (; AuxIndex < NumBones; ++AuxIndex)
+			{
+				if (Skeleton->BoneIds[AuxIndex] == BoneId.Id)
+				{
+					BoneId.Index = AuxIndex;
+					break;
+				}
+			}
+
+			check(Skeleton->BoneIds.IsValidIndex(BoneId.Index));
+		}
+	}
+}
+
+
 TManagedPtr<const FPhysicsBody> FMesh::GetPhysicsBody() const
 {
     return PhysicsBody;
@@ -546,9 +582,9 @@ const TArray<uint64>& FMesh::GetStreamedResources() const
 }
 
 
-int32 FMesh::FindBonePose(const FBoneName& BoneId) const
+int32 FMesh::FindBonePoseByBoneIndex(int32 Index) const
 {
-	return BonePoses.IndexOfByPredicate([BoneId](const FBonePose& Pose) { return Pose.BoneId == BoneId; });
+	return BonePoses.IndexOfByPredicate([Index](const FBonePose& Pose) { return Pose.BoneId.Index == Index; });
 }
 
 
@@ -565,7 +601,7 @@ int32 UE::Mutable::Private::FMesh::GetBonePoseCount() const
 }
 
 
-void UE::Mutable::Private::FMesh::SetBonePose(int32 Index, const FBoneName& BoneId, FTransform3f Transform, EBoneUsageFlags BoneUsageFlags)
+void UE::Mutable::Private::FMesh::SetBonePose(int32 Index, FBoneIdOrIndex BoneId, FTransform3f Transform, EBoneUsageFlags BoneUsageFlags)
 {
 	check(BonePoses.IsValidIndex(Index));
 	if (BonePoses.IsValidIndex(Index))
@@ -575,7 +611,7 @@ void UE::Mutable::Private::FMesh::SetBonePose(int32 Index, const FBoneName& Bone
 }
 
 
-const FBoneName& FMesh::GetBonePoseId(int32 Index) const
+FBoneIdOrIndex FMesh::GetBonePoseId(int32 Index) const
 {
 	check(BonePoses.IsValidIndex(Index));
 	return BonePoses[Index].BoneId;
@@ -596,13 +632,13 @@ EBoneUsageFlags FMesh::GetBoneUsageFlags(int32 BoneIndex) const
 }
 
 
-void FMesh::SetBoneMap(const TArray<FBoneName>& InBoneMap)
+void FMesh::SetBoneMap(const TArray<FBoneIdOrIndex>& InBoneMap)
 {
 	BoneMap = InBoneMap;
 }
 
 
-const TArray<FBoneName>& FMesh::GetBoneMap() const
+const TArray<FBoneIdOrIndex>& FMesh::GetBoneMap() const
 {
 	return BoneMap;
 }
@@ -911,7 +947,7 @@ void FMesh::Serialise(FOutputArchive& Arch) const
 
 	Arch << SkeletonIDs;
 
-	Arch << Skeleton;
+	//Arch << Skeleton;
 	Arch << PhysicsBody;
 
 	Arch << uint32(Flags);
@@ -938,7 +974,7 @@ void FMesh::Unserialise(FInputArchive& Arch)
 
 	Arch >> SkeletonIDs;
 
-	Arch >> Skeleton;
+	//Arch >> Skeleton;
 	Arch >> PhysicsBody;
 
 	uint32 Temp;
diff --git a/Engine/Plugins/Mutable/Source/MutableRuntime/Private/MuR/Model.cpp b/Engine/Plugins/Mutable/Source/MutableRuntime/Private/MuR/Model.cpp
--- a/Engine/Plugins/Mutable/Source/MutableRuntime/Private/MuR/Model.cpp
+++ b/Engine/Plugins/Mutable/Source/MutableRuntime/Private/MuR/Model.cpp
@@ -131,6 +131,7 @@ namespace UE::Mutable::Private
 		Arch << ParameterLists;
 		Arch << RelevantParameterList;
 		Arch << ConstantMaterials;
+		Arch << ConstantNames;
 	}
 
 
@@ -162,11 +163,33 @@ namespace UE::Mutable::Private
 		Arch >> ParameterLists;
 		Arch >> RelevantParameterList;
 		Arch >> ConstantMaterials;
+		Arch >> ConstantNames;
 		
 		PendingOpsPerRom.SetNum(Roms.Num());
 		RomWeights.SetNumZeroed(Roms.Num());
+
+		PostLoad();
 	}
 	
+	void FProgram::PostLoad()
+	{
+		// Convert Skeletons' BoneIds to FNames for runtime operations
+		const int32 NumSkeletons = ConstantSkeletons.Num();
+		for (int32 Index = 0; Index < NumSkeletons; ++Index)
+		{
+			TManagedPtr<FSkeleton> Skeleton = ConstCastManagedPtr<FSkeleton>(ConstantSkeletons[Index]);
+
+			const TArray<uint32>& BoneIds = Skeleton->BoneIds;
+
+			TArray<FName>& BoneNames = Skeleton->BoneNames;
+			BoneNames.Reserve(BoneIds.Num());
+
+			for (uint32 BoneId : BoneIds)
+			{
+				BoneNames.Add(ConstantNames[BoneId]);
+			}
+		}
+	}
 
 	void FProgram::LogHistogram() const
     {
@@ -340,24 +363,36 @@ namespace UE::Mutable::Private
     }
 
 
-    OP::ADDRESS FProgram::AddConstant(const TManagedPtr<const FSkeleton>& Skeleton)
-    {
-	    // Ensure unique
-	    for (int32 Index = 0; Index < ConstantSkeletons.Num(); ++Index)
-	    {
-		    if (ConstantSkeletons[Index] == Skeleton ||
-			    *ConstantSkeletons[Index] == *Skeleton
-		    )
-		    {
-			    return (OP::ADDRESS)Index;
-		    }
-	    }
+	OP::ADDRESS FProgram::AddConstant(const TManagedPtr<const FSkeleton>& InSkeleton)
+	{
+		// Ensure unique
+		for (int32 i = 0; i < ConstantSkeletons.Num(); ++i)
+		{
+			if (ConstantSkeletons[i] == InSkeleton
+				||
+				*ConstantSkeletons[i] == *InSkeleton // Does not compare BoneIds
+				)
+			{
+				return (OP::ADDRESS)i;
+			}
+		}
 
-	    OP::ADDRESS Index = OP::ADDRESS(ConstantSkeletons.Num());
-	    ConstantSkeletons.Add(Skeleton);
-		
-	    return Index;
-    }
+		// Convert FNames to FBoneIds
+		TManagedPtr<FSkeleton> Skeleton = InSkeleton->Clone();
+
+		const TArray<FName>& BoneNames = Skeleton->BoneNames;
+		const int32 NumBoneNames = BoneNames.Num();
+
+		TArray<uint32>& BoneIds = Skeleton->BoneIds;
+		BoneIds.SetNumUninitialized(NumBoneNames);
+
+		for (int32 Index = 0; Index < NumBoneNames; ++Index)
+		{
+			BoneIds[Index] = AddConstant(BoneNames[Index]);
+		}
+
+		return OP::ADDRESS(ConstantSkeletons.Add(Skeleton));
+	}
 
 
     OP::ADDRESS FProgram::AddConstant(const TManagedPtr<const FPhysicsBody>& PhysicsBody)
@@ -382,6 +417,21 @@ namespace UE::Mutable::Private
 	    return (OP::CONSTANT_STRING_INDEX)ConstantStrings.AddUnique(Str);
     }
 
+	uint32 FProgram::AddConstant(const FName& InName)
+	{
+		const FString NameString = InName.ToString().ToLower();
+		uint32 NameId = CityHash32(reinterpret_cast<const char*>(*NameString), NameString.Len() * sizeof(FString::ElementType));
+
+		FName& ConstantName = ConstantNames.FindOrAdd(NameId, InName);
+		while (ConstantName != InName)
+		{
+			// Increase Id in an attempt to make it unique again.
+			++NameId;
+			ConstantName = ConstantNames.FindOrAdd(NameId, InName);
+		}
+
+		return NameId;
+	}
 
     OP::ADDRESS FProgram::AddConstant(const TArray<uint32>& UInt32List)
     {
diff --git a/Engine/Plugins/Mutable/Source/MutableRuntime/Private/MuR/OpMeshApplyPose.h b/Engine/Plugins/Mutable/Source/MutableRuntime/Private/MuR/OpMeshApplyPose.h
--- a/Engine/Plugins/Mutable/Source/MutableRuntime/Private/MuR/OpMeshApplyPose.h
+++ b/Engine/Plugins/Mutable/Source/MutableRuntime/Private/MuR/OpMeshApplyPose.h
@@ -11,7 +11,7 @@
 
 namespace UE::Mutable::Private
 {
-	inline void SetPoseAsReference(FMesh* InOutResult, const FMesh* PoseMesh)
+	inline void SetPoseAsReference(FMesh* InOutResult, const FMesh* PoseMesh, const TArray<int32>& BoneToPoseIndex)
 	{
 		// Set the pose as the Result reference pose.
 		// PoseMesh PoseXform can be decomposed to PoseXform = ModelPoseXform * ModelRefXform^-1, 
@@ -19,17 +19,17 @@ namespace UE::Mutable::Private
 		// to multiply PoseXform * ModelRefXform. 
 		// (Note, transform multiplication application is from right to left for this comment, but TTransform::operator* is reversed).
 		const int32 ResultNumBones = InOutResult->GetBonePoseCount();
-		for (int32 BoneIndex = 0; BoneIndex < ResultNumBones; ++BoneIndex)
+		for (int32 PoseIndex = 0; PoseIndex < ResultNumBones; ++PoseIndex)
 		{
-			const FBoneName BoneName = InOutResult->GetBonePoseId(BoneIndex);
-			const int32 FoundPoseIndex = PoseMesh->FindBonePose(BoneName);
+			const FBoneIdOrIndex BoneId = InOutResult->BonePoses[PoseIndex].BoneId;
+			const int32 FoundPoseIndex = BoneToPoseIndex[BoneId.Index];
 			if (FoundPoseIndex == INDEX_NONE)
 			{
 				continue;
 			}
 
 			FTransform3f ModelRefTransform;
-			InOutResult->GetBonePoseTransform(BoneIndex, ModelRefTransform);
+			InOutResult->GetBonePoseTransform(PoseIndex, ModelRefTransform);
 
 			FTransform3f PoseTransform;
 			PoseMesh->GetBonePoseTransform(FoundPoseIndex, PoseTransform); 
@@ -37,10 +37,10 @@ namespace UE::Mutable::Private
 			// Adding the Reshaped flag will prioritize this bone over other bones without the flag
 			// when merging poses with the same bone.
 			// TODO: Add another flag to indicate this case or generalize the Reshaped flag. 
-			EBoneUsageFlags UsageFlags = InOutResult->GetBoneUsageFlags(BoneIndex);
+			EBoneUsageFlags UsageFlags = InOutResult->GetBoneUsageFlags(PoseIndex);
 			EnumAddFlags(UsageFlags, EBoneUsageFlags::Reshaped);
 			
-			InOutResult->SetBonePose(BoneIndex, BoneName, ModelRefTransform * PoseTransform, UsageFlags);
+			InOutResult->SetBonePose(PoseIndex, BoneId, ModelRefTransform * PoseTransform, UsageFlags);
 		}
 	}
 
@@ -53,8 +53,9 @@ namespace UE::Mutable::Private
 
 		bOutSuccess = true;
 
-		TManagedPtr<const FSkeleton> Skeleton = BaseMesh->GetSkeleton();
-		if (!Skeleton)
+		TManagedPtr<const FSkeleton> BaseSkeleton = BaseMesh->GetSkeleton();
+		TManagedPtr<const FSkeleton> PoseSkeleton = PoseMesh->GetSkeleton();
+		if (!BaseSkeleton || !PoseSkeleton)
 		{
 			bOutSuccess = false;
 			return;
@@ -63,13 +64,19 @@ namespace UE::Mutable::Private
         // We assume the matrices are transforms from the binding pose bone to the new pose
 
 		// Find closest bone affected by the pose for each bone in the skeleton.
-		const int32 NumBones = Skeleton->GetBoneCount();
+		const int32 NumBonesBase = BaseSkeleton->GetNumBones();
 		TArray<int32> BoneToPoseIndex;
-		BoneToPoseIndex.Init(INDEX_NONE, NumBones);
+		BoneToPoseIndex.Init(INDEX_NONE, NumBonesBase);
 
-		for (int32 BoneIndex = 0; BoneIndex < NumBones; ++BoneIndex)
+		for (int32 BoneIndex = 0; BoneIndex < NumBonesBase; ++BoneIndex)
 		{
-			const int32 BonePoseIndex = PoseMesh->FindBonePose(Skeleton->GetBoneName(BoneIndex));
+			const int32 BoneNameIndex = PoseSkeleton->BoneNames.Find(BaseSkeleton->BoneNames[BoneIndex]);
+			if (BoneNameIndex == INDEX_NONE)
+			{
+				continue; // BaseSkeleton bone not found in PoseSkeleton.
+			}
+
+			const int32 BonePoseIndex = PoseMesh->FindBonePoseByBoneIndex(BoneNameIndex);
 			if (BonePoseIndex != INDEX_NONE)
 			{
 				BoneToPoseIndex[BoneIndex] = BonePoseIndex;
@@ -77,14 +84,14 @@ namespace UE::Mutable::Private
 			}
 
 			// Parent bones are in a strictly increassing order. Set the pose index from the parent bone.
-			const int32 ParentIndex = Skeleton->GetBoneParent(BoneIndex);
+			const int32 ParentIndex = BaseSkeleton->BoneParents[BoneIndex];
 			if (ParentIndex != INDEX_NONE)
 			{
 				BoneToPoseIndex[BoneIndex] = BoneToPoseIndex[ParentIndex];
 			}
 		}
 
-		const TArray<FBoneName>& BoneMap = BaseMesh->GetBoneMap();
+		const TArray<FBoneIdOrIndex>& BoneMap = BaseMesh->GetBoneMap();
 		const int32 NumBonesBoneMap = BoneMap.Num();
 
         // Prepare the skin matrices. They may be in different order, and we only need the ones
@@ -95,10 +102,10 @@ namespace UE::Mutable::Private
 		bool bBonesAffected = false;
 		for (int32 Index = 0; Index < NumBonesBoneMap; ++Index)
 		{
-			const int32 BoneIndex = Skeleton->FindBone(BoneMap[Index]);
-			if (BoneToPoseIndex[BoneIndex] != INDEX_NONE)
+			const int32 PoseIndex = BoneToPoseIndex[BoneMap[Index].Index];
+			if (PoseIndex != INDEX_NONE)
 			{
-				SkinTransforms.Add(PoseMesh->BonePoses[BoneToPoseIndex[BoneIndex]].BoneTransform);
+				SkinTransforms.Add(PoseMesh->BonePoses[PoseIndex].BoneTransform);
 				bBonesAffected = true;
 			}
 			else
@@ -138,7 +145,7 @@ namespace UE::Mutable::Private
 
 		Result->CopyFrom(*BaseMesh);
 
-		SetPoseAsReference(Result, PoseMesh);
+		SetPoseAsReference(Result, PoseMesh, BoneToPoseIndex);
 
 		MeshBufferIterator<EMeshBufferFormat::Float32, float, 3> TargetPositionIterBegin(Result->VertexBuffers, EMeshBufferSemantic::Position, 0);
 		
diff --git a/Engine/Plugins/Mutable/Source/MutableRuntime/Private/MuR/OpMeshBind.h b/Engine/Plugins/Mutable/Source/MutableRuntime/Private/MuR/OpMeshBind.h
--- a/Engine/Plugins/Mutable/Source/MutableRuntime/Private/MuR/OpMeshBind.h
+++ b/Engine/Plugins/Mutable/Source/MutableRuntime/Private/MuR/OpMeshBind.h
@@ -506,8 +506,8 @@ namespace UE::Mutable::Private
 	}
 
 	inline TTuple<TArray<FReshapePointBindingData>, TArray<int32>, TArray<int32>> BindPhysicsBodies( 
-			TArray<const FPhysicsBody*> PhysicsBodies, FShapeMeshTree& ShapeMeshTree, const FMesh* pMesh, 
-			const TArray<FBoneName>& PhysicsToDeform )
+			TArray<const FPhysicsBody*> PhysicsBodies, FShapeMeshTree& ShapeMeshTree, const FMesh& Mesh, const FSkeleton& Skeleton, 
+			const TArray<FName>& PhysicsToDeform )
 	{
 		TTuple<TArray<FReshapePointBindingData>, TArray<int32>, TArray<int32>> ReturnValue;
 
@@ -527,7 +527,7 @@ namespace UE::Mutable::Private
 				const int32 NumBodies = PhysicsBodies[PhysicsBodyIndex]->GetBodyCount();
 				for (int32 BodyIndex = 0; BodyIndex < NumBodies; ++BodyIndex)
 				{
-					if (PhysicsToDeform.Contains(Body->GetBodyBoneId(BodyIndex)))
+					if (PhysicsToDeform.Contains(Body->GetBodyBoneName(BodyIndex)))
 					{
 						BodiesToDeformIndices.Add(BodyIndex);
 					}
@@ -593,12 +593,13 @@ namespace UE::Mutable::Private
 
 			for (const int32 B : BodyIndices)
 			{
-				int32 BoneIndex = pMesh->FindBonePose(Body->GetBodyBoneId(B));
+				const int32 BoneIndex = Skeleton.BoneNames.Find(Body->GetBodyBoneName(B));
+				int32 PoseIndex = Mesh.FindBonePoseByBoneIndex(BoneIndex);
 
 				FTransform3f T = FTransform3f::Identity;
-				if (BoneIndex > 0)
+				if (PoseIndex > 0)
 				{
-					pMesh->GetBonePoseTransform(BoneIndex, T);
+					Mesh.GetBonePoseTransform(PoseIndex, T);
 				}
 
 				const int32 SphereCount = Body->GetSphereCount(B);
@@ -1014,7 +1015,7 @@ namespace UE::Mutable::Private
 
 
 	inline TTuple<TArray<FReshapePointBindingData>, TArray<int32>> BindPose(
-			const FMesh* Mesh, FShapeMeshTree& ShapeMeshTree, const TArray<FBoneName>& BonesToDeform )
+			const FMesh& Mesh, const FSkeleton& Skeleton, FShapeMeshTree& ShapeMeshTree, const TArray<FName>& BonesToDeform )
 	{
 		UE::Geometry::FAxisAlignedBox3d ShapeAABBox = ShapeMeshTree.GetBoundingBox();
 
@@ -1027,27 +1028,27 @@ namespace UE::Mutable::Private
 		TArray<FReshapePointBindingData>& SkeletonBindDataArray = ReturnValue.Get<0>();
 		TArray<int32>& BoneIndices = ReturnValue.Get<1>();
 
-		const int32 BoneCount = Mesh->GetBonePoseCount();
+		const int32 BoneCount = Mesh.GetBonePoseCount();
 		SkeletonBindDataArray.Reserve(BoneCount);
 		BoneIndices.Reserve(BoneCount);
 
 		for (int32 BoneIndex = 0; BoneIndex < BoneCount; ++BoneIndex)
 		{
-			const EBoneUsageFlags BoneUsageFlags = Mesh->BonePoses[BoneIndex].BoneUsageFlags;
+			const EBoneUsageFlags BoneUsageFlags = Mesh.BonePoses[BoneIndex].BoneUsageFlags;
 
 			if (EnumHasAnyFlags(BoneUsageFlags, EBoneUsageFlags::Root))
 			{
 				continue;
 			}
 
-			if (!BonesToDeform.Contains(Mesh->BonePoses[BoneIndex].BoneId))
+			if (!BonesToDeform.Contains(Skeleton.BoneNames[Mesh.BonePoses[BoneIndex].BoneId.Index]))
 			{
 				continue;
 			}
 
 			FReshapeVertexBindingData BindData;
 			constexpr float MaskWeight = 1.0f;
-			BindReshapePoint(ShapeMeshTree, Mesh->BonePoses[BoneIndex].BoneTransform.GetLocation(), FVector3f::ZAxisVector, MaskWeight, BindData, BindValidityTolerance);
+			BindReshapePoint(ShapeMeshTree, Mesh.BonePoses[BoneIndex].BoneTransform.GetLocation(), FVector3f::ZAxisVector, MaskWeight, BindData, BindValidityTolerance);
 
 			// Only add binding if there is a chance of the bone moving.
 			if (BindData.Weight > UE_SMALL_NUMBER && BindData.IsValidBinding())
@@ -1089,7 +1090,7 @@ namespace UE::Mutable::Private
 		check(Masks.MaskWeight == 0 ||
 			FMath::CountLeadingZeros(Masks.MaskWeight) + FMath::CountTrailingZeros(Masks.MaskWeight) == 32 - 8);
 		
-		// No overlaped channels.
+		// No overlapped channels.
 		check((Masks.ClusterId & Masks.MaskWeight) == 0);
 
 		return Masks;
@@ -1101,7 +1102,7 @@ namespace UE::Mutable::Private
     inline void MeshBindShapeReshape(
 			FMesh* Result,
 			const FMesh* BaseMesh, const FMesh* ShapeMesh, 
-			const TArray<FBoneName>& BonesToDeform, const TArray<FBoneName>& PhysicsToDeform,
+			const TArray<FName>& BonesToDeform, const TArray<FName>& PhysicsToDeform,
 			EMeshBindShapeFlags BindFlags, FMeshBindColorChannelUsages ColorChannelUsages,
 			bool& bOutSuccess)
     {
@@ -1123,7 +1124,7 @@ namespace UE::Mutable::Private
 
 		const FMeshBindColorChannelUsageMasks ColorUsagesMasks = MakeColorChannelUsageMasks(ColorChannelUsages);
 
-		// Early out if nothing will be modified and the vertices discarted. return null in this
+		// Early out if nothing will be modified and the vertices discarded. return null in this
 		// case indicating nothing has modified so the Base Mesh can be reused.
 		const bool bSkeletonModification = BaseMesh->GetSkeleton() && bReshapeSkeleton;
 		const bool bPhysicsModification = (BaseMesh->GetPhysicsBody() || BaseMesh->AdditionalPhysicsBodies.Num()) && bReshapePhysics;
@@ -1243,12 +1244,12 @@ namespace UE::Mutable::Private
 
 		// Bind the skeleton bones
 		// \TODO: Build bind data only for actually modified bones?
-		if (bReshapeSkeleton && BonesToDeform.Num())
+		if (bReshapeSkeleton && BonesToDeform.Num() && Result->GetSkeleton())
 		{
 			MUTABLE_CPUPROFILER_SCOPE(BindSkeleton);
 
 			TTuple<TArray<FReshapePointBindingData>, TArray<int32>> SkeletonBindingData = 
-					BindPose(Result, ShapeMeshTree, BonesToDeform);
+					BindPose(*Result, *Result->GetSkeleton(), ShapeMeshTree, BonesToDeform);
 			const TArray<FReshapePointBindingData>& SkeletonBindDataArray = SkeletonBindingData.Get<0>();
 			const TArray<int32>& BoneIndices = SkeletonBindingData.Get<1>();
 
@@ -1279,12 +1280,12 @@ namespace UE::Mutable::Private
 		}
 
 		const FPhysicsBody* ResultPhysicsBody = Result->PhysicsBody.Get();
-		if (bReshapePhysics && (ResultPhysicsBody || Result->AdditionalPhysicsBodies.Num()) && PhysicsToDeform.Num())
+		if (bReshapePhysics && Result->GetSkeleton() && (ResultPhysicsBody || Result->AdditionalPhysicsBodies.Num()) && PhysicsToDeform.Num())
 		{
 			MUTABLE_CPUPROFILER_SCOPE(BindPhysicsBody);
 
-			// Gather bodies respecting order, firts main physics body then additional bodies.
-			// Null entries are need to be able to mantin that order.  
+			// Gather bodies respecting order, first main physics body then additional bodies.
+			// Null entries are need to be able to maintain that order.  
 			TArray<const FPhysicsBody*> PhysicsBodiesToBind;
 			{
 				PhysicsBodiesToBind.Reserve(Result->AdditionalPhysicsBodies.Num() + 1);
@@ -1297,7 +1298,7 @@ namespace UE::Mutable::Private
 			}
 
 			TTuple<TArray<FReshapePointBindingData>, TArray<int32>, TArray<int32>> PhysicsBindingData = 
-					BindPhysicsBodies(PhysicsBodiesToBind, ShapeMeshTree, Result, PhysicsToDeform);
+					BindPhysicsBodies(PhysicsBodiesToBind, ShapeMeshTree, *Result, *Result->GetSkeleton(), PhysicsToDeform);
 
 			const TArray<FReshapePointBindingData>& PhysicsBindDataArray = PhysicsBindingData.Get<0>();
 			const TArray<int32>& DeformedBodyIndices = PhysicsBindingData.Get<1>();
diff --git a/Engine/Plugins/Mutable/Source/MutableRuntime/Private/MuR/OpMeshClipMorphPlane.h b/Engine/Plugins/Mutable/Source/MutableRuntime/Private/MuR/OpMeshClipMorphPlane.h
--- a/Engine/Plugins/Mutable/Source/MutableRuntime/Private/MuR/OpMeshClipMorphPlane.h
+++ b/Engine/Plugins/Mutable/Source/MutableRuntime/Private/MuR/OpMeshClipMorphPlane.h
@@ -73,7 +73,7 @@ namespace UE::Mutable::Private
 	//---------------------------------------------------------------------------------------------
 	inline void MeshClipMorphPlane(FMesh* Result, const FMesh* pBase, const FVector3f& Origin, const FVector3f& Normal, float Dist, float Factor, float Radius,
 		float Radius2, float Angle, const FShape& SelectionShape, bool bRemoveIfAllVerticesCulled, bool& bOutSuccess, 
-		const FBoneName* BoneId = nullptr, float VertexSelectionBoneMaxRadius = -1.f)
+		FName BoneName = NAME_None, float VertexSelectionBoneMaxRadius = -1.f)
 	{
 		bOutSuccess = true;
 		//float Radius = 8.f;
@@ -104,37 +104,31 @@ namespace UE::Mutable::Private
         TManagedPtr<const FSkeleton> BaseSkeleton = pBase->GetSkeleton();
 
 		TBitArray<> AffectedBoneMapIndices;
-		const int32 BaseBoneIndex = BaseSkeleton && BoneId ? BaseSkeleton->FindBone(*BoneId) : INDEX_NONE;
+		const int32 BaseBoneIndex = BaseSkeleton && BoneName != NAME_None ? BaseSkeleton->BoneNames.Find(BoneName) : INDEX_NONE;
         if (BaseBoneIndex != INDEX_NONE)
 		{
-			const TArray<FBoneName>& BoneMap = pBase->BoneMap;
-			AffectedBoneMapIndices.SetNum(BoneMap.Num(), false);
+			const int32 NumBonesBoneMap = pBase->BoneMap.Num();
+			const TArray<FBoneIdOrIndex>& BoneMap = pBase->BoneMap;
+			AffectedBoneMapIndices.SetNum(NumBonesBoneMap, false);
 
-			const int32 BoneCount = BaseSkeleton->GetBoneCount();
+			const int32 NumBones = BaseSkeleton->GetNumBones();
 			TBitArray<> AffectedSkeletonBones;
-			AffectedSkeletonBones.SetNum(BoneCount, false);
+			AffectedSkeletonBones.SetNum(NumBones, false);
 
-            for (int32 BoneIndex = 0; BoneIndex < BoneCount; ++BoneIndex)
+			for (int32 BoneIndex = 0; BoneIndex < NumBones; ++BoneIndex)
 			{
-				const int32 ParentBoneIndex = BaseSkeleton->GetBoneParent(BoneIndex);
+				const int32 ParentBoneIndex = BaseSkeleton->BoneParents[BoneIndex];
 				check(ParentBoneIndex < BoneIndex);
-				
+
 				const bool bIsBoneAffected = (AffectedSkeletonBones.IsValidIndex(ParentBoneIndex) && AffectedSkeletonBones[ParentBoneIndex])
 					|| BoneIndex == BaseBoneIndex;
-				
-				if (!bIsBoneAffected)
-				{
-					continue;
-				}
-				
-				AffectedSkeletonBones[BoneIndex] = true;
 
-				const FBoneName& AffectedBone = BaseSkeleton->GetBoneName(BoneIndex);
-				const int32 AffectedBoneMapIndex = BoneMap.Find(AffectedBone);
-				if (AffectedBoneMapIndex != INDEX_NONE)
-				{
-					AffectedBoneMapIndices[AffectedBoneMapIndex] = true;
-				}
+				AffectedSkeletonBones[BoneIndex] = bIsBoneAffected;
+			}
+
+			for (int32 Index = 0; Index < NumBonesBoneMap; ++Index)
+			{
+				AffectedBoneMapIndices[Index] = AffectedSkeletonBones[BoneMap[Index].Index];
 			}
 
 			// Look for affected vertex indices
diff --git a/Engine/Plugins/Mutable/Source/MutableRuntime/Private/MuR/OpMeshReshape.h b/Engine/Plugins/Mutable/Source/MutableRuntime/Private/MuR/OpMeshReshape.h
--- a/Engine/Plugins/Mutable/Source/MutableRuntime/Private/MuR/OpMeshReshape.h
+++ b/Engine/Plugins/Mutable/Source/MutableRuntime/Private/MuR/OpMeshReshape.h
@@ -566,6 +566,7 @@ namespace UE::Mutable::Private
 	inline void ApplyToPhysicsBodies(
 			FPhysicsBody& PBody, int32& InOutNumProcessedBindPoints, 
 			const FMesh& BaseMesh, 
+			const FSkeleton& BaseSkeleton, 
 			TArrayView<const FReshapePointBindingData> BindingData, 
 			TArrayView<const int32> UsedIndices, 
 			const FShapeMeshDescriptorApply& Shape)
@@ -586,11 +587,13 @@ namespace UE::Mutable::Private
 		// Retrieve them in the same order the boxes where put in, so they can be linked to the physics body volumes.
 		for ( const int32 B : UsedIndices )
 		{	
-			int32 BoneIndex = BaseMesh.FindBonePose(PBody.GetBodyBoneId(B));
+			const int32 BoneIndex = BaseSkeleton.BoneNames.Find(PBody.GetBodyBoneName(B));
+			const int32 PoseIndex = BaseMesh.FindBonePoseByBoneIndex(BoneIndex);
+
 			FTransform3f BoneTransform = FTransform3f::Identity;
-			if (BoneIndex >= 0)
+			if (PoseIndex >= 0)
 			{
-				BaseMesh.GetBonePoseTransform(BoneIndex, BoneTransform);
+				BaseMesh.GetBonePoseTransform(PoseIndex, BoneTransform);
 			}
 			
 			FTransform3f InvBoneTransform = BoneTransform.Inverse();
@@ -704,7 +707,8 @@ namespace UE::Mutable::Private
 	}
 
 	inline void ApplyToAllPhysicsBodies(
-		FMesh& OutNewMesh, const FMesh& BaseMesh, 
+		FMesh& OutNewMesh,
+		const FMesh& BaseMesh, const FSkeleton& BaseSkeleton,
 		TArrayView<const FReshapePointBindingData> BindingData,
 		TArrayView<const int32> UsedIndices,
 		TArrayView<const int32> BodyOffsets,
@@ -721,7 +725,7 @@ namespace UE::Mutable::Private
 		{
 			TArrayView<const int32> BodyUsedIndices(UsedIndices.GetData() + IndicesBegin, IndicesEnd - IndicesBegin);
 
-			ApplyToPhysicsBodies(OutBody, NumProcessedBindPoints, BaseMesh, BindingData, BodyUsedIndices, Shape);
+			ApplyToPhysicsBodies(OutBody, NumProcessedBindPoints, BaseMesh, BaseSkeleton, BindingData, BodyUsedIndices, Shape);
 		};
 
 
@@ -977,7 +981,7 @@ namespace UE::Mutable::Private
 		const FPhysicsBody* OldPhysicsBody = Result->PhysicsBody.Get();
 		const int32 AdditionalPhysicsBodiesNum = Result->AdditionalPhysicsBodies.Num();
 
-		if (bReshapePhysicsVolumes && (OldPhysicsBody || AdditionalPhysicsBodiesNum))
+		if (bReshapePhysicsVolumes && BaseMesh->GetSkeleton() && (OldPhysicsBody || AdditionalPhysicsBodiesNum))
 		{	
 			MUTABLE_CPUPROFILER_SCOPE(ReshapePhysicsBodies);
 			
@@ -1022,7 +1026,7 @@ namespace UE::Mutable::Private
 						(const int32*)PhyiscsBindOffsetsBuffer.GetBufferData(0), 
 					    PhyiscsBindOffsetsBuffer.GetElementCount());
 
-				ApplyToAllPhysicsBodies(*Result, *BaseMesh, BindingData, UsedIndices, Offsets, ShapeDescriptor);
+				ApplyToAllPhysicsBodies(*Result, *BaseMesh, *BaseMesh->GetSkeleton(), BindingData, UsedIndices, Offsets, ShapeDescriptor);
 			}
 		}
 	}
diff --git a/Engine/Plugins/Mutable/Source/MutableRuntime/Private/MuR/OpMeshTransformWithBone.cpp b/Engine/Plugins/Mutable/Source/MutableRuntime/Private/MuR/OpMeshTransformWithBone.cpp
--- a/Engine/Plugins/Mutable/Source/MutableRuntime/Private/MuR/OpMeshTransformWithBone.cpp
+++ b/Engine/Plugins/Mutable/Source/MutableRuntime/Private/MuR/OpMeshTransformWithBone.cpp
@@ -10,11 +10,13 @@
 #include "MuR/ParallelExecutionUtils.h"
 #include "PackedNormal.h"
 
+#include "UObject/NameTypes.h"
+
 namespace UE::Mutable::Private
 {
 
 
-void MeshTransformWithBoneInline(FMesh* Mesh, const FMatrix44f& Transform, const FBoneName& InBoneName, const float InThresholdFactor)
+void MeshTransformWithBoneInline(FMesh* Mesh, const FMatrix44f& Transform, FName InBoneName, const float InThresholdFactor)
 {
 	MUTABLE_CPUPROFILER_SCOPE(MeshTransformWithBone);
 
@@ -26,7 +28,7 @@ void MeshTransformWithBoneInline(FMesh* Mesh, const FMatrix44f& Transform, const
 
 	TManagedPtr<const FSkeleton> BaseSkeleton = Mesh->GetSkeleton();
 
-	const int32 BaseBoneIndex = BaseSkeleton ? BaseSkeleton->FindBone(InBoneName.Id) : INDEX_NONE;
+	const int32 BaseBoneIndex = BaseSkeleton ? BaseSkeleton->BoneNames.Find(InBoneName) : INDEX_NONE;
 	if (BaseBoneIndex == INDEX_NONE)
 	{
 		return;
@@ -37,16 +39,16 @@ void MeshTransformWithBoneInline(FMesh* Mesh, const FMatrix44f& Transform, const
 	// Find affected bones
 	{
 		const int32 NumBonesBoneMap = Mesh->BoneMap.Num();
-		const TArray<FBoneName>& BoneMap = Mesh->BoneMap;
+		const TArray<FBoneIdOrIndex>& BoneMap = Mesh->BoneMap;
 		AffectedBoneMapIndices.SetNum(NumBonesBoneMap, false);
 
-		const int32 BoneCount = BaseSkeleton->GetBoneCount();
+		const int32 NumBones = BaseSkeleton->GetNumBones();
 		TBitArray<> AffectedSkeletonBones;
-		AffectedSkeletonBones.SetNum(BoneCount, false);
+		AffectedSkeletonBones.SetNum(NumBones, false);
 
-		for (int32 BoneIndex = 0; BoneIndex < BoneCount; ++BoneIndex)
+		for (int32 BoneIndex = 0; BoneIndex < NumBones; ++BoneIndex)
 		{
-			const int32 ParentBoneIndex = BaseSkeleton->GetBoneParent(BoneIndex);
+			const int32 ParentBoneIndex = BaseSkeleton->BoneParents[BoneIndex];
 			check(ParentBoneIndex < BoneIndex);
 
 			const bool bIsBoneAffected = (AffectedSkeletonBones.IsValidIndex(ParentBoneIndex) && AffectedSkeletonBones[ParentBoneIndex])
@@ -55,15 +57,9 @@ void MeshTransformWithBoneInline(FMesh* Mesh, const FMatrix44f& Transform, const
 			AffectedSkeletonBones[BoneIndex] = bIsBoneAffected;
 		}
 
-		for (int32 BoneIndex = 0; BoneIndex < NumBonesBoneMap; ++BoneIndex)
+		for (int32 Index = 0; Index < NumBonesBoneMap; ++Index)
 		{
-			const FBoneName& BoneName = BoneMap[BoneIndex];
-			
-			const int32 SkeletonIndex = BaseSkeleton->FindBone(BoneName);
-			if (SkeletonIndex != INDEX_NONE && AffectedSkeletonBones[SkeletonIndex])
-			{
-				AffectedBoneMapIndices[BoneIndex] = true;
-			}
+			AffectedBoneMapIndices[Index] = AffectedSkeletonBones[BoneMap[Index].Index];
 		}
 	}
 
diff --git a/Engine/Plugins/Mutable/Source/MutableRuntime/Private/MuR/OpMeshTransformWithBone.h b/Engine/Plugins/Mutable/Source/MutableRuntime/Private/MuR/OpMeshTransformWithBone.h
--- a/Engine/Plugins/Mutable/Source/MutableRuntime/Private/MuR/OpMeshTransformWithBone.h
+++ b/Engine/Plugins/Mutable/Source/MutableRuntime/Private/MuR/OpMeshTransformWithBone.h
@@ -4,12 +4,13 @@
 
 #include "Math/MathFwd.h"
 
+class FName;
+
 namespace UE::Mutable::Private
 {
 class FMesh;
-struct FBoneName;
 
 /**  */
-extern void MeshTransformWithBoneInline(FMesh* Mesh, const FMatrix44f& Transform, const FBoneName& BoneName, const float Threshold);
+extern void MeshTransformWithBoneInline(FMesh* Mesh, const FMatrix44f& Transform, FName BoneName, const float Threshold);
 
 }
diff --git a/Engine/Plugins/Mutable/Source/MutableRuntime/Private/MuR/Operations.cpp b/Engine/Plugins/Mutable/Source/MutableRuntime/Private/MuR/Operations.cpp
--- a/Engine/Plugins/Mutable/Source/MutableRuntime/Private/MuR/Operations.cpp
+++ b/Engine/Plugins/Mutable/Source/MutableRuntime/Private/MuR/Operations.cpp
@@ -132,7 +132,6 @@ namespace UE::Mutable::Private
 		{ EDataType::Mesh,			},	// ME_TRANSFORM
 		{ EDataType::Mesh,			},	// ME_CLIPMORPHPLANE
 		{ EDataType::Mesh,			},	// ME_CLIPWITHMESH
-		{ EDataType::Mesh,			},	// ME_SETSKELETON
 		{ EDataType::Mesh,			},	// ME_PROJECT
 		{ EDataType::Mesh,			},	// ME_APPLYPOSE
 		{ EDataType::Mesh,			},	// ME_BINDSHAPE
@@ -825,14 +824,6 @@ namespace UE::Mutable::Private
 
 			break;
 		}
- 
-        case EOpType::ME_SETSKELETON :
-        {
-			OP::MeshSetSkeletonArgs Args = program.GetOpArgs<OP::MeshSetSkeletonArgs>(at);
-            f(Args.Source);
-            f(Args.Skeleton);
-            break;
-        }
 
         case EOpType::ME_PROJECT :
         {
diff --git a/Engine/Plugins/Mutable/Source/MutableRuntime/Private/MuR/PhysicsBody.cpp b/Engine/Plugins/Mutable/Source/MutableRuntime/Private/MuR/PhysicsBody.cpp
--- a/Engine/Plugins/Mutable/Source/MutableRuntime/Private/MuR/PhysicsBody.cpp
+++ b/Engine/Plugins/Mutable/Source/MutableRuntime/Private/MuR/PhysicsBody.cpp
@@ -167,6 +167,7 @@ namespace UE::Mutable::Private
 		Result->CustomId = CustomId;
 
 		Result->Bodies = Bodies;
+		Result->BodiesBoneNames = BodiesBoneNames;
 		Result->BoneIds = BoneIds;
 		Result->BodiesCustomIds = BodiesCustomIds;
 
@@ -176,20 +177,10 @@ namespace UE::Mutable::Private
 		return Result;
 	}
 
-	void FPhysicsBody::SetCustomId(int32 InCustomId)
-	{
-		CustomId = InCustomId;
-	}
-
-	int32 FPhysicsBody::GetCustomId() const
-	{
-		return CustomId;
-	}
-
 	void FPhysicsBody::SetBodyCount(int32 Count)
 	{
 		Bodies.SetNum(Count);
-		BoneIds.SetNum(Count);
+		BodiesBoneNames.SetNum(Count);
 		BodiesCustomIds.Init(-1, Count);
 	}
 
@@ -198,30 +189,16 @@ namespace UE::Mutable::Private
 		return Bodies.Num();
 	}
 
-	const FBoneName& FPhysicsBody::GetBodyBoneId(int32 B) const
-	{
-		check(BoneIds.IsValidIndex(B));
-		return BoneIds[B];
-	}
-
-	void FPhysicsBody::SetBodyBoneId(int32 B, const FBoneName& BoneId)
+	void FPhysicsBody::SetBody(int32 BodyIndex, FName InBoneName, int32 InBodyCustomId)
 	{
-		check(BoneIds.IsValidIndex(B));
-		BoneIds[B] = BoneId;
+		BodiesBoneNames[BodyIndex] = InBoneName;
+		BodiesCustomIds[BodyIndex] = InBodyCustomId;
 	}
 
-	int32 FPhysicsBody::GetBodyCustomId(int32 B) const
+	FName FPhysicsBody::GetBodyBoneName(int32 BodyIndex) const
 	{
-		check(B >= 0 && B < BodiesCustomIds.Num());
-
-		return BodiesCustomIds[B];
-	}
-
-	void FPhysicsBody::SetBodyCustomId(int32 B, int32 BodyCustomId)
-	{
-		check(B >= 0 && B < BodiesCustomIds.Num());
-
-		BodiesCustomIds[B] = BodyCustomId;
+		check(BodiesBoneNames.IsValidIndex(BodyIndex));
+		return BodiesBoneNames[BodyIndex];
 	}
 
 	int32 FPhysicsBody::GetSphereCount(int32 B) const
@@ -590,6 +567,20 @@ namespace UE::Mutable::Private
 		return Bodies[B].TaperedCapsules[I].Name;
 	}
 
+	void FPhysicsBody::ConvertBoneIdsToNames(const TMap<uint32, FName>& ConstantNames)
+	{
+		if (BodiesBoneNames.IsEmpty())
+		{
+			BodiesBoneNames.Reserve(BoneIds.Num());
+
+			for (uint32 BoneId : BoneIds)
+			{
+				BodiesBoneNames.Add(ConstantNames.FindChecked(BoneId));
+			}
+
+			BoneIds.Empty();
+		}
+	}
 
 	void FPhysicsBody::Serialise(FOutputArchive& Arch) const
 	{
diff --git a/Engine/Plugins/Mutable/Source/MutableRuntime/Private/MuR/Skeleton.cpp b/Engine/Plugins/Mutable/Source/MutableRuntime/Private/MuR/Skeleton.cpp
--- a/Engine/Plugins/Mutable/Source/MutableRuntime/Private/MuR/Skeleton.cpp
+++ b/Engine/Plugins/Mutable/Source/MutableRuntime/Private/MuR/Skeleton.cpp
@@ -27,98 +27,40 @@ namespace UE::Mutable::Private
 	{
 		TManagedPtr<FSkeleton> Result = MakeManaged<FSkeleton>();
 
-		Result->BoneIds = BoneIds;
+		Result->BoneNames = BoneNames;
+		//Result->BoneIds = BoneIds; // Runtime operations work with FNames
 		Result->BoneParents = BoneParents;
 
-		// For debug
-		Result->DebugBoneNames = DebugBoneNames;
-
 		return Result;
 	}
 
 
-	int32 FSkeleton::GetBoneCount() const
+	int32 FSkeleton::AddBone(const FName InBoneName, const int16 InParentIndex)
 	{
-		return BoneIds.Num();
+		BoneNames.Add(InBoneName);
+		return BoneParents.Add(InParentIndex);
 	}
 
-
-	void FSkeleton::SetBoneCount(int32 NumBones)
+	int32 FSkeleton::FindOrAddBone(const FName InBoneName, const int16 InParentIndex)
 	{
-		DebugBoneNames.SetNum(NumBones);
-		BoneIds.SetNum(NumBones);
-		BoneParents.Init(INDEX_NONE, NumBones);
-	}
-
-
-	const FName FSkeleton::GetDebugName(int32 Index) const
-	{
-		if (DebugBoneNames.IsValidIndex(Index))
+		int32 BoneIndex = BoneNames.Find(InBoneName);
+		if (BoneIndex != INDEX_NONE)
 		{
-			return DebugBoneNames[Index];
+			return BoneIndex;
 		}
 
-		return FName("Unknown Bone");
-	}
-
-
-	void FSkeleton::SetDebugName(const int32 Index, const FName BoneName)
-	{
-		if (DebugBoneNames.IsValidIndex(Index))
-		{
-			DebugBoneNames[Index] = BoneName;
-		}
+		return AddBone(InBoneName, InParentIndex);
 	}
 
-
-	int32 FSkeleton::GetBoneParent(int32 Index) const
-	{
-		if (BoneParents.IsValidIndex(Index))
-		{
-			return BoneParents[Index];
-		}
-
-		return INDEX_NONE;
-	}
-
-
-	void FSkeleton::SetBoneParent(int32 Index, int32 ParentIndex)
-	{
-		check(ParentIndex >= -1 && ParentIndex < GetBoneCount() && ParentIndex < 0xffff);
-		check(BoneParents.IsValidIndex(Index));
-		
-		if (BoneParents.IsValidIndex(Index))
-		{
-			BoneParents[Index] = (int16)ParentIndex;
-		}
-	}
-
-
-	const FBoneName& FSkeleton::GetBoneName(int32 Index) const
-	{
-		check(BoneIds.IsValidIndex(Index));
-		return BoneIds[Index];
-	}
-
-	
-	void FSkeleton::SetBoneName(int32 Index, const FBoneName& BoneName)
-	{
-		check(BoneIds.IsValidIndex(Index));
-		if (BoneIds.IsValidIndex(Index))
-		{
-			BoneIds[Index] = BoneName;
-		}
-	}
-
-
-	int32 FSkeleton::FindBone(const FBoneName& BoneName) const
+	int32 FSkeleton::GetNumBones() const
 	{
-		return BoneIds.Find(BoneName);
+		return !BoneNames.IsEmpty() ? BoneNames.Num() : BoneIds.Num();
 	}
 
 
 	void FSkeleton::Serialise(FOutputArchive& Arch) const
 	{
+		//Arch << BoneNames; Serialized as ConstantNames
 		Arch << BoneIds;
 		Arch << BoneParents;
 	}
@@ -126,6 +68,7 @@ namespace UE::Mutable::Private
 
 	void FSkeleton::Unserialise(FInputArchive& Arch)
 	{
+		//Arch >> BoneNames; Serialized as ConstantNames
 		Arch >> BoneIds;
 		Arch >> BoneParents;
 	}
diff --git a/Engine/Plugins/Mutable/Source/MutableTools/Internal/MuT/NodeMeshClipMorphPlane.h b/Engine/Plugins/Mutable/Source/MutableTools/Internal/MuT/NodeMeshClipMorphPlane.h
--- a/Engine/Plugins/Mutable/Source/MutableTools/Internal/MuT/NodeMeshClipMorphPlane.h
+++ b/Engine/Plugins/Mutable/Source/MutableTools/Internal/MuT/NodeMeshClipMorphPlane.h
@@ -48,7 +48,7 @@ namespace UE::Mutable::Private
 
 		//! Define the root bone of the subhierarchy of the mesh that will be affected.
 		//! Only one of Box or Bone Hierarchy can be used (the last one set)
-		UE_API void SetVertexSelectionBone(const FBoneName& BoneId, float maxEffectRadius);
+		UE_API void SetVertexSelectionBone(FName BoneId, float maxEffectRadius);
 
 
 	protected:
diff --git a/Engine/Plugins/Mutable/Source/MutableTools/Internal/MuT/NodeMeshMorph.h b/Engine/Plugins/Mutable/Source/MutableTools/Internal/MuT/NodeMeshMorph.h
--- a/Engine/Plugins/Mutable/Source/MutableTools/Internal/MuT/NodeMeshMorph.h
+++ b/Engine/Plugins/Mutable/Source/MutableTools/Internal/MuT/NodeMeshMorph.h
@@ -27,8 +27,8 @@ namespace UE::Mutable::Private
 		bool bReshapeSkeleton = false;
 		bool bReshapePhysicsVolumes = false;
 
-		TArray<FBoneName> BonesToDeform;
-		TArray<FBoneName> PhysicsToDeform;
+		TArray<FName> BonesToDeform;
+		TArray<FName> PhysicsToDeform;
 
 	public:
 
diff --git a/Engine/Plugins/Mutable/Source/MutableTools/Internal/MuT/NodeMeshReshape.h b/Engine/Plugins/Mutable/Source/MutableTools/Internal/MuT/NodeMeshReshape.h
--- a/Engine/Plugins/Mutable/Source/MutableTools/Internal/MuT/NodeMeshReshape.h
+++ b/Engine/Plugins/Mutable/Source/MutableTools/Internal/MuT/NodeMeshReshape.h
@@ -32,8 +32,8 @@ namespace UE::Mutable::Private
 		EVertexColorUsage ColorBChannelUsage = EVertexColorUsage::None;
 		EVertexColorUsage ColorAChannelUsage = EVertexColorUsage::None;
 
-		TArray<FBoneName> BonesToDeform;
-		TArray<FBoneName> PhysicsToDeform;
+		TArray<FName> BonesToDeform;
+		TArray<FName> PhysicsToDeform;
 
 	public:
 
diff --git a/Engine/Plugins/Mutable/Source/MutableTools/Internal/MuT/NodeModifierMeshClipMorphPlane.h b/Engine/Plugins/Mutable/Source/MutableTools/Internal/MuT/NodeModifierMeshClipMorphPlane.h
--- a/Engine/Plugins/Mutable/Source/MutableTools/Internal/MuT/NodeModifierMeshClipMorphPlane.h
+++ b/Engine/Plugins/Mutable/Source/MutableTools/Internal/MuT/NodeModifierMeshClipMorphPlane.h
@@ -14,8 +14,6 @@
 namespace UE::Mutable::Private
 {
 
-	struct FBoneName;
-
 	struct FClipMorphPlaneParameters
 	{
 		// Morph field parameters
@@ -35,7 +33,7 @@ namespace UE::Mutable::Private
 		EClipVertexSelectionType VertexSelectionType = EClipVertexSelectionType::None;
 		FVector3f SelectionBoxOrigin = FVector3f(0, 0, 0);
 		FVector3f SelectionBoxRadius = FVector3f(0, 0, 0);
-		FBoneName VertexSelectionBone;
+		FName VertexSelectionBone = NAME_None;
 
 		// Max distance a vertex can have to the bone in order to be affected. A negative value
 		// means no limit.
@@ -72,7 +70,7 @@ namespace UE::Mutable::Private
 
 		//! Define the root bone of the subhierarchy of the mesh that will be affected.
 		//! Only one of Box or Bone Hierarchy can be used (the last one set)
-		UE_API void SetVertexSelectionBone(const FBoneName& BoneId, float maxEffectRadius);
+		UE_API void SetVertexSelectionBone(FName BoneId, float maxEffectRadius);
 
 	protected:
 
diff --git a/Engine/Plugins/Mutable/Source/MutableTools/Internal/MuT/NodeModifierMeshTransformWithBone.h b/Engine/Plugins/Mutable/Source/MutableTools/Internal/MuT/NodeModifierMeshTransformWithBone.h
--- a/Engine/Plugins/Mutable/Source/MutableTools/Internal/MuT/NodeModifierMeshTransformWithBone.h
+++ b/Engine/Plugins/Mutable/Source/MutableTools/Internal/MuT/NodeModifierMeshTransformWithBone.h
@@ -20,7 +20,7 @@ class NodeModifierMeshTransformWithBone : public NodeModifier
 public:
 
 	/* */
-	FBoneName BoneName;
+	FName BoneName = NAME_None;
 	float ThresholdFactor;
 
 	/** */
diff --git a/Engine/Plugins/Mutable/Source/MutableTools/Private/MuT/ASTOpConstantResource.cpp b/Engine/Plugins/Mutable/Source/MutableTools/Private/MuT/ASTOpConstantResource.cpp
--- a/Engine/Plugins/Mutable/Source/MutableTools/Private/MuT/ASTOpConstantResource.cpp
+++ b/Engine/Plugins/Mutable/Source/MutableTools/Private/MuT/ASTOpConstantResource.cpp
@@ -100,8 +100,8 @@ namespace UE::Mutable::Private
 				return Program.ConstantMeshContentIndices.Add(FConstantResourceIndex{(uint32)MeshIndex, 0});
 			};
 			
-			// Split generated mesh data in 4 parts. Geometry and Pose and Physiscs and Metadata.
-			// Inidices for a given rom are sorted by content flag value.
+			// Split generated mesh data in 4 parts. Geometry and Pose and Physics and Metadata.
+			// Indices for a given rom are sorted by content flag value.
 			static_assert(EMeshContentFlags::GeometryData < EMeshContentFlags::PoseData);
 			static_assert(EMeshContentFlags::PoseData     < EMeshContentFlags::PhysicsData);
 			static_assert(EMeshContentFlags::PhysicsData  < EMeshContentFlags::MetaData);
@@ -139,16 +139,38 @@ namespace UE::Mutable::Private
 
 			// Pose Mesh
 			{
-				const EMeshCopyFlags PoseDataCopyFlags = 
-						EMeshCopyFlags::WithPoses  |
-						EMeshCopyFlags::WithBoneMap;
+				const EMeshCopyFlags PoseDataCopyFlags =
+					EMeshCopyFlags::WithPoses |
+					EMeshCopyFlags::WithBoneMap;
 
 				TManagedPtr<FMesh> MeshPoseData = MeshData->Clone(PoseDataCopyFlags);
 
+				if (MeshData->Skeleton)
+				{
+					MUTABLE_CPUPROFILER_SCOPE(ConvertBoneNamesToIds_Pose);
+
+					const TArray<FName>& BoneNames = MeshData->Skeleton->BoneNames;
+
+					for (FBoneIdOrIndex& Bone : MeshPoseData->BoneMap)
+					{
+						Bone.Id = Program.AddConstant(BoneNames[Bone.Index]);
+					}
+
+					for (FMesh::FBonePose& BonePose : MeshPoseData->BonePoses)
+					{
+						BonePose.BoneId.Id = Program.AddConstant(BoneNames[BonePose.BoneId.Index]);
+					}
+				}
+				else
+				{
+					check(MeshPoseData->BoneMap.IsEmpty() && MeshPoseData->BonePoses.IsEmpty());
+				}
+
+
 				for (const TPair<EMeshBufferType, FMeshBufferSet>& AdditionalBuffer : MeshData->AdditionalBuffers)
 				{
-					const bool bIsPoseBufferType = 	
-							AdditionalBuffer.Key == EMeshBufferType::SkeletonDeformBinding;
+					const bool bIsPoseBufferType =
+						AdditionalBuffer.Key == EMeshBufferType::SkeletonDeformBinding;
 
 					if (bIsPoseBufferType)
 					{
@@ -157,7 +179,7 @@ namespace UE::Mutable::Private
 				}
 
 				MeshPoseData->MeshIDPrefix = 0;
-				AddMeshToProgram(MeshPoseData);	
+				AddMeshToProgram(MeshPoseData);
 				EnumAddFlags(MeshContentFlags, EMeshContentFlags::PoseData);
 			}
 
@@ -168,6 +190,16 @@ namespace UE::Mutable::Private
 						EMeshCopyFlags::WithPhysicsBody;
 
 				TManagedPtr<FMesh> MeshPhysicsData = MeshData->Clone(PhysicsDataCopyFlags);
+
+				if (TManagedPtr<FPhysicsBody> PhysicsBody = ConstCastManagedPtr<FPhysicsBody>(MeshPhysicsData->PhysicsBody))
+				{
+					MUTABLE_CPUPROFILER_SCOPE(ConvertBoneNamesToIds_Physics);
+					PhysicsBody->BoneIds.Reserve(PhysicsBody->GetBodyCount());
+					for (FName BoneName : PhysicsBody->BodiesBoneNames)
+					{
+						PhysicsBody->BoneIds.Add(Program.AddConstant(BoneName));
+					}
+				}
 			
 				// Copy components related additional buffers.
 				for (const TPair<EMeshBufferType, FMeshBufferSet>& AdditionalBuffer : MeshData->AdditionalBuffers)
@@ -411,14 +443,13 @@ namespace UE::Mutable::Private
 					OP::MeshConstantArgs Args;
 					FMemory::Memzero(Args);
 
-					TManagedPtr<FMesh> MeshValue = StaticCastManagedPtr<const FMesh>(Value)->Clone();
+					TManagedPtr<const FMesh> MeshValue = StaticCastManagedPtr<const FMesh>(Value);
 					check(MeshValue);
 
 					Args.Skeleton = -1;
 					if (TManagedPtr<const FSkeleton> MeshSkeleton = MeshValue->GetSkeleton())
 					{
 						Args.Skeleton = Program.AddConstant(MeshSkeleton);
-						MeshValue->SetSkeleton(nullptr);
 					}
 
 					check(!MeshValue->PassthroughObject);
diff --git a/Engine/Plugins/Mutable/Source/MutableTools/Private/MuT/ASTOpMeshBindShape.cpp b/Engine/Plugins/Mutable/Source/MutableTools/Private/MuT/ASTOpMeshBindShape.cpp
--- a/Engine/Plugins/Mutable/Source/MutableTools/Private/MuT/ASTOpMeshBindShape.cpp
+++ b/Engine/Plugins/Mutable/Source/MutableTools/Private/MuT/ASTOpMeshBindShape.cpp
@@ -83,14 +83,14 @@ namespace UE::Mutable::Private
 		Result = HashCombineFast(Result, GetTypeHash(BChannelUsage));
 		Result = HashCombineFast(Result, GetTypeHash(AChannelUsage));
 
-		for (const FBoneName& S : BonesToDeform)
+		for (FName Bone : BonesToDeform)
 		{
-			Result = HashCombineFast(Result, GetTypeHash(S.Id));
+			Result = HashCombineFast(Result, GetTypeHash(Bone));
 		}
 
-		for (const FBoneName& S : PhysicsToDeform)
+		for (FName PhysicsBone : PhysicsToDeform)
 		{
-			Result = HashCombineFast(Result, GetTypeHash(S.Id));
+			Result = HashCombineFast(Result, GetTypeHash(PhysicsBone));
 		}
 
 		return Result;
@@ -185,15 +185,15 @@ namespace UE::Mutable::Private
 			AppendCode(program.ByteCode, Args);
 
 			AppendCode(program.ByteCode, (int32)BonesToDeform.Num());
-			for (const FBoneName& S : BonesToDeform)
+			for (FName Bone : BonesToDeform)
 			{
-				AppendCode(program.ByteCode, S);
+				AppendCode(program.ByteCode, program.AddConstant(Bone));
 			}
 
 			AppendCode(program.ByteCode, (int32)PhysicsToDeform.Num());
-			for (const FBoneName& S : PhysicsToDeform)
+			for (FName PhysicsBone : PhysicsToDeform)
 			{
-				AppendCode(program.ByteCode, S);
+				AppendCode(program.ByteCode, program.AddConstant(PhysicsBone));
 			}
 		}
 	}
diff --git a/Engine/Plugins/Mutable/Source/MutableTools/Private/MuT/ASTOpMeshBindShape.h b/Engine/Plugins/Mutable/Source/MutableTools/Private/MuT/ASTOpMeshBindShape.h
--- a/Engine/Plugins/Mutable/Source/MutableTools/Private/MuT/ASTOpMeshBindShape.h
+++ b/Engine/Plugins/Mutable/Source/MutableTools/Private/MuT/ASTOpMeshBindShape.h
@@ -24,8 +24,8 @@ namespace UE::Mutable::Private
 		ASTChild Mesh;
 		ASTChild Shape;
 
-		TArray<FBoneName> BonesToDeform;
-		TArray<FBoneName> PhysicsToDeform;
+		TArray<FName> BonesToDeform;
+		TArray<FName> PhysicsToDeform;
 
 		uint32 BindingMethod = 0;
 		
diff --git a/Engine/Plugins/Mutable/Source/MutableTools/Private/MuT/ASTOpMeshClipMorphPlane.cpp b/Engine/Plugins/Mutable/Source/MutableTools/Private/MuT/ASTOpMeshClipMorphPlane.cpp
--- a/Engine/Plugins/Mutable/Source/MutableTools/Private/MuT/ASTOpMeshClipMorphPlane.cpp
+++ b/Engine/Plugins/Mutable/Source/MutableTools/Private/MuT/ASTOpMeshClipMorphPlane.cpp
@@ -100,8 +100,7 @@ namespace UE::Mutable::Private
 			Args.VertexSelectionType = VertexSelectionType;
 			if (VertexSelectionType == EClipVertexSelectionType::BoneHierarchy)
 			{
-				check(VertexSelectionBone.Id < MAX_uint32);
-				Args.VertexSelectionShapeOrBone = VertexSelectionBone.Id;
+				Args.VertexSelectionShapeOrBone = Program.AddConstant(VertexSelectionBone);
 			}
 			else if (VertexSelectionType == EClipVertexSelectionType::Shape)
 			{
diff --git a/Engine/Plugins/Mutable/Source/MutableTools/Private/MuT/ASTOpMeshClipMorphPlane.h b/Engine/Plugins/Mutable/Source/MutableTools/Private/MuT/ASTOpMeshClipMorphPlane.h
--- a/Engine/Plugins/Mutable/Source/MutableTools/Private/MuT/ASTOpMeshClipMorphPlane.h
+++ b/Engine/Plugins/Mutable/Source/MutableTools/Private/MuT/ASTOpMeshClipMorphPlane.h
@@ -21,7 +21,7 @@ namespace UE::Mutable::Private
 
 		FShape MorphShape;
 		FShape SelectionShape;
-		FBoneName VertexSelectionBone;
+		FName VertexSelectionBone = NAME_None;
 
 		EClipVertexSelectionType VertexSelectionType = EClipVertexSelectionType::None;
 		EFaceCullStrategy FaceCullStrategy = EFaceCullStrategy::AllVerticesCulled;
diff --git a/Engine/Plugins/Mutable/Source/MutableTools/Private/MuT/ASTOpMeshExtractLayoutBlocks.cpp b/Engine/Plugins/Mutable/Source/MutableTools/Private/MuT/ASTOpMeshExtractLayoutBlocks.cpp
--- a/Engine/Plugins/Mutable/Source/MutableTools/Private/MuT/ASTOpMeshExtractLayoutBlocks.cpp
+++ b/Engine/Plugins/Mutable/Source/MutableTools/Private/MuT/ASTOpMeshExtractLayoutBlocks.cpp
@@ -11,7 +11,6 @@
 #include "MuT/ASTOpMeshApplyPose.h"
 #include "MuT/ASTOpMeshApplyLayout.h"
 #include "MuT/ASTOpMeshMerge.h"
-#include "MuT/ASTOpMeshSetSkeleton.h"
 #include "MuT/ASTOpMeshPrepareLayout.h"
 #include "MuT/ASTOpConstantResource.h"
 #include "MuR/Model.h"
@@ -208,14 +207,6 @@ namespace UE::Mutable::Private
 			break;
 		}
 
-		case EOpType::ME_SETSKELETON:
-		{
-			Ptr<ASTOpMeshSetSkeleton> NewOp = UE::Mutable::Private::Clone<ASTOpMeshSetSkeleton>(at);
-			NewOp->Source = Visit( NewOp->Source.child(), currentSinkOp);
-			newAt = NewOp;
-			break;
-		}
-
 		case EOpType::ME_ADDMETADATA:
 		{
 			Ptr<ASTOpMeshAddMetadata> NewOp = UE::Mutable::Private::Clone<ASTOpMeshAddMetadata>(at);
diff --git a/Engine/Plugins/Mutable/Source/MutableTools/Private/MuT/ASTOpMeshFormat.cpp b/Engine/Plugins/Mutable/Source/MutableTools/Private/MuT/ASTOpMeshFormat.cpp
--- a/Engine/Plugins/Mutable/Source/MutableTools/Private/MuT/ASTOpMeshFormat.cpp
+++ b/Engine/Plugins/Mutable/Source/MutableTools/Private/MuT/ASTOpMeshFormat.cpp
@@ -22,7 +22,6 @@
 #include "MuT/ASTOpMeshApplyPose.h"
 #include "MuT/ASTOpMeshApplyLayout.h"
 #include "MuT/ASTOpMeshMerge.h"
-#include "MuT/ASTOpMeshSetSkeleton.h"
 #include "MuT/ASTOpSwitch.h"
 
 #include "GPUSkinPublicDefs.h"
@@ -302,14 +301,6 @@ UE::Mutable::Private::Ptr<ASTOp> Sink_MeshFormatAST::Visit(const UE::Mutable::Pr
 		break;
 	}
 
-	case EOpType::ME_SETSKELETON:
-	{
-		Ptr<ASTOpMeshSetSkeleton> NewOp = UE::Mutable::Private::Clone<ASTOpMeshSetSkeleton>(at);
-		NewOp->Source = Visit(NewOp->Source.child(), currentFormatOp);
-		newAt = NewOp;
-		break;
-	}
-
 	case EOpType::ME_ADDMETADATA:
 	{
 		Ptr<ASTOpMeshAddMetadata> newOp = UE::Mutable::Private::Clone<ASTOpMeshAddMetadata>(at);
diff --git a/Engine/Plugins/Mutable/Source/MutableTools/Private/MuT/ASTOpMeshSetSkeleton.cpp b/Engine/Plugins/Mutable/Source/MutableTools/Private/MuT/ASTOpMeshSetSkeleton.cpp
deleted file mode 100644
diff --git a/Engine/Plugins/Mutable/Source/MutableTools/Private/MuT/ASTOpMeshSetSkeleton.h b/Engine/Plugins/Mutable/Source/MutableTools/Private/MuT/ASTOpMeshSetSkeleton.h
deleted file mode 100644
diff --git a/Engine/Plugins/Mutable/Source/MutableTools/Private/MuT/ASTOpMeshTransformWithBone.cpp b/Engine/Plugins/Mutable/Source/MutableTools/Private/MuT/ASTOpMeshTransformWithBone.cpp
--- a/Engine/Plugins/Mutable/Source/MutableTools/Private/MuT/ASTOpMeshTransformWithBone.cpp
+++ b/Engine/Plugins/Mutable/Source/MutableTools/Private/MuT/ASTOpMeshTransformWithBone.cpp
@@ -74,7 +74,7 @@ void UE::Mutable::Private::ASTOpMeshTransformWithBone::Link(FProgram& Program, F
 			Args.Matrix = Matrix->linkedAddress;
 		}
 
-		Args.BoneId = BoneName.Id;
+		Args.BoneId = Program.AddConstant(BoneName);
 		Args.ThresholdFactor = ThresholdFactor;
 
 		linkedAddress = static_cast<OP::ADDRESS>(Program.OpAddress.Num());
diff --git a/Engine/Plugins/Mutable/Source/MutableTools/Private/MuT/ASTOpMeshTransformWithBone.h b/Engine/Plugins/Mutable/Source/MutableTools/Private/MuT/ASTOpMeshTransformWithBone.h
--- a/Engine/Plugins/Mutable/Source/MutableTools/Private/MuT/ASTOpMeshTransformWithBone.h
+++ b/Engine/Plugins/Mutable/Source/MutableTools/Private/MuT/ASTOpMeshTransformWithBone.h
@@ -22,7 +22,7 @@ class ASTOpMeshTransformWithBone final : public ASTOp
 public:
 	ASTChild SourceMesh;
 	ASTChild Matrix;
-	FBoneName BoneName;
+	FName BoneName = NAME_None;
 	float ThresholdFactor = 0.05f;
 public:
 
diff --git a/Engine/Plugins/Mutable/Source/MutableTools/Private/MuT/CodeOptimiser.cpp b/Engine/Plugins/Mutable/Source/MutableTools/Private/MuT/CodeOptimiser.cpp
--- a/Engine/Plugins/Mutable/Source/MutableTools/Private/MuT/CodeOptimiser.cpp
+++ b/Engine/Plugins/Mutable/Source/MutableTools/Private/MuT/CodeOptimiser.cpp
@@ -877,8 +877,6 @@ namespace UE::Mutable::Private
 		// that require it.
 		CollectAllMeshesForSkeletonVisitorAST requireSkeletonCollector( roots );
 
-		TArray<FAddMeshSkeleton> replacementsFound;
-
 		ASTOp::Traverse_TopDown_Unique_Imprecise( roots, [&](UE::Mutable::Private::Ptr<ASTOp>& at )
 		{
 			// Only recurse instance construction ops.
@@ -911,58 +909,37 @@ namespace UE::Mutable::Private
 						return true;
 					});
 
-					// Create a mesh just with the unified skeleton
+					// Create a mesh with the unified skeleton
 					TManagedPtr<FSkeleton> FinalSkeleton = MakeManaged<FSkeleton>();
-					for (const auto& meshAt: subtreeMeshes)
+					for (const auto& MeshAt: subtreeMeshes)
 					{
-						TManagedPtr<const FMesh> pMesh = StaticCastManagedPtr<const FMesh>(meshAt->GetValue());
-						TManagedPtr<const FSkeleton> SourceSkeleton = pMesh ? pMesh->GetSkeleton() : nullptr;
+						TManagedPtr<FMesh> Mesh = ConstCastManagedPtr<FMesh>(StaticCastManagedPtr<const FMesh>(MeshAt->GetValue()));
+						TManagedPtr<const FSkeleton> SourceSkeleton = Mesh ? Mesh->GetSkeleton() : nullptr;
 						if (SourceSkeleton)
 						{
-							ExtendSkeleton(FinalSkeleton.Get(),SourceSkeleton.Get());
-						}
-					}
+							TArray<int32> RemappedSourceBoneIndices;
+							MergeSkeletons(*FinalSkeleton.Get(), *SourceSkeleton.Get(), RemappedSourceBoneIndices);
 
-					replacementsFound.Emplace( at, subtreeMeshes, FinalSkeleton );
-				}
-			}
-
-			return processChildren;
-		});
-
-
-		// Iterate all meshes again
-		ASTOp::Traverse_TopDown_Unique_Imprecise( roots, [&](UE::Mutable::Private::Ptr<ASTOp>& at )
-		{
-			if (at->GetOpType()==EOpType::ME_CONSTANT)
-			{
-				ASTOpConstantResource* typedOp = static_cast<ASTOpConstantResource*>(at.get());
+							// Update the meshs' skeleton and remap bone indices
+							Mesh->SetSkeleton(FinalSkeleton);
 
-				for(FAddMeshSkeleton& Rep: replacementsFound)
-				{
-					if (Rep.ContributingMeshes.Contains(at))
-					{
-						TManagedPtr<const FMesh> pMesh = StaticCastManagedPtr<const FMesh>(typedOp->GetValue());
-						pMesh->CheckIntegrity();
-
-						TManagedPtr<FMesh> NewMesh = MakeManaged<FMesh>();
-						bool bOutSuccess = false;
-						MeshRemapSkeleton(NewMesh.Get(), pMesh.Get(), Rep.FinalSkeleton, bOutSuccess);
-
-						if (bOutSuccess)
-						{
-							NewMesh->CheckIntegrity();
-							UE::Mutable::Private::Ptr<ASTOpConstantResource> newOp = new ASTOpConstantResource();
-							newOp->Type = EOpType::ME_CONSTANT;
-							newOp->SetValue(NewMesh);
-							newOp->SourceDataDescriptor = at->GetSourceDataDescriptor();
+							const int32 NumBonesBoneMap = Mesh->BoneMap.Num();
+							for (int32 Index = 0; Index < NumBonesBoneMap; ++Index)
+							{
+								Mesh->BoneMap[Index].Index = RemappedSourceBoneIndices[Mesh->BoneMap[Index].Index];
+							}
 
-							ASTOp::Replace(at, newOp);
+							const int32 NumBonePoses = Mesh->BonePoses.Num();
+							for (int32 Index = 0; Index < NumBonePoses; ++Index)
+							{
+								Mesh->BonePoses[Index].BoneId.Index = RemappedSourceBoneIndices[Mesh->BonePoses[Index].BoneId.Index];
+							}
 						}
 					}
 				}
 			}
-			return true;
+
+			return processChildren;
 		});
 	}
 
diff --git a/Engine/Plugins/Mutable/Source/MutableTools/Private/MuT/Compiler.cpp b/Engine/Plugins/Mutable/Source/MutableTools/Private/MuT/Compiler.cpp
--- a/Engine/Plugins/Mutable/Source/MutableTools/Private/MuT/Compiler.cpp
+++ b/Engine/Plugins/Mutable/Source/MutableTools/Private/MuT/Compiler.cpp
@@ -167,7 +167,6 @@ namespace UE::Mutable::Private
 		{ true,		{ 0, 0, 0, 0, 0, 0, 0, 0, 0, 0, 0, 0, 0, 0, 0, 0, 0, 0, 0, 0, 0, 0, 0, 0, 0, 0, 0, 0, 0 }		},	// ME_TRANSFORM
 		{ true,		{ 0, 0, 0, 0, 0, 0, 0, 0, 0, 0, 0, 0, 0, 0, 0, 0, 0, 0, 0, 0, 0, 0, 0, 0, 0, 0, 0, 0, 0 }		},	// ME_CLIPMORPHPLANE
 		{ true,		{ 0, 0, 0, 0, 0, 0, 0, 0, 0, 0, 0, 0, 0, 0, 0, 0, 0, 0, 0, 0, 0, 0, 0, 0, 0, 0, 0, 0, 0 }		},	// ME_CLIPWITHMESH
-		{ true,		{ 0, 0, 0, 0, 0, 0, 0, 0, 0, 0, 0, 0, 0, 0, 0, 0, 0, 0, 0, 0, 0, 0, 0, 0, 0, 0, 0, 0, 0 }		},	// ME_SETSKELETON
 		{ true,		{ 0, 0, 0, 0, 0, 0, 0, 0, 0, 0, 0, 0, 0, 0, 0, 0, 0, 0, 0, 0, 0, 0, 0, 0, 0, 0, 0, 0, 0 }		},	// ME_PROJECT
 		{ true,		{ 0, 0, 0, 0, 0, 0, 0, 0, 0, 0, 0, 0, 0, 0, 0, 0, 0, 0, 0, 0, 0, 0, 0, 0, 0, 0, 0, 0, 0 }		},	// ME_APPLYPOSE
 		{ true,		{ 0, 0, 0, 0, 0, 0, 0, 0, 0, 0, 0, 0, 0, 0, 0, 0, 0, 0, 0, 0, 0, 0, 0, 0, 0, 0, 0, 0, 0 }		},	// ME_BINDSHAPE
diff --git a/Engine/Plugins/Mutable/Source/MutableTools/Private/MuT/ErrorLog.cpp b/Engine/Plugins/Mutable/Source/MutableTools/Private/MuT/ErrorLog.cpp
--- a/Engine/Plugins/Mutable/Source/MutableTools/Private/MuT/ErrorLog.cpp
+++ b/Engine/Plugins/Mutable/Source/MutableTools/Private/MuT/ErrorLog.cpp
@@ -304,7 +304,6 @@ namespace UE::Mutable::Private
 		TEXT("ME_TRANSFORM     "),
 		TEXT("ME_CLIPMORPHPLANE"),
         TEXT("ME_CLIPWITHMESH  "),
-        TEXT("ME_SETSKELETON   "),
         TEXT("ME_PROJECT       "),
         TEXT("ME_APPLYPOSE     "),
 		TEXT("ME_BINDSHAPE	   "),
diff --git a/Engine/Plugins/Mutable/Source/MutableTools/Private/MuT/NodeMeshClipMorphPlane.cpp b/Engine/Plugins/Mutable/Source/MutableTools/Private/MuT/NodeMeshClipMorphPlane.cpp
--- a/Engine/Plugins/Mutable/Source/MutableTools/Private/MuT/NodeMeshClipMorphPlane.cpp
+++ b/Engine/Plugins/Mutable/Source/MutableTools/Private/MuT/NodeMeshClipMorphPlane.cpp
@@ -55,11 +55,11 @@ namespace UE::Mutable::Private
 	}
 
 	//---------------------------------------------------------------------------------------------
-	void NodeMeshClipMorphPlane::SetVertexSelectionBone(const FBoneName& BoneId, float maxEffectRadius)
+	void NodeMeshClipMorphPlane::SetVertexSelectionBone(FName BoneName, float MaxEffectRadius)
 	{
 		Parameters.VertexSelectionType = EClipVertexSelectionType::BoneHierarchy;
-		Parameters.VertexSelectionBone = BoneId;
-		Parameters.MaxEffectRadius = maxEffectRadius;
+		Parameters.VertexSelectionBone = BoneName;
+		Parameters.MaxEffectRadius = MaxEffectRadius;
 	}
 
 }
diff --git a/Engine/Plugins/Mutable/Source/MutableTools/Private/MuT/NodeModifierMeshClipMorphPlane.cpp b/Engine/Plugins/Mutable/Source/MutableTools/Private/MuT/NodeModifierMeshClipMorphPlane.cpp
--- a/Engine/Plugins/Mutable/Source/MutableTools/Private/MuT/NodeModifierMeshClipMorphPlane.cpp
+++ b/Engine/Plugins/Mutable/Source/MutableTools/Private/MuT/NodeModifierMeshClipMorphPlane.cpp
@@ -37,10 +37,10 @@ namespace UE::Mutable::Private
 	}
 
 
-	void NodeModifierMeshClipMorphPlane::SetVertexSelectionBone(const FBoneName& BoneId, float maxEffectRadius)
+	void NodeModifierMeshClipMorphPlane::SetVertexSelectionBone(FName BoneName, float maxEffectRadius)
 	{
 		Parameters.VertexSelectionType = EClipVertexSelectionType::BoneHierarchy;
-		Parameters.VertexSelectionBone = BoneId;
+		Parameters.VertexSelectionBone = BoneName;
 		Parameters.MaxEffectRadius = maxEffectRadius;
 	}
 
