[Mutable] Fixed Image Descriptor of Break Material Nodes.

#rnx

#Synced-CL 49922449

diff --git a/Engine/Plugins/Mutable/Source/MutableRuntime/Private/MuR/CodeRunner.cpp b/Engine/Plugins/Mutable/Source/MutableRuntime/Private/MuR/CodeRunner.cpp
--- a/Engine/Plugins/Mutable/Source/MutableRuntime/Private/MuR/CodeRunner.cpp
+++ b/Engine/Plugins/Mutable/Source/MutableRuntime/Private/MuR/CodeRunner.cpp
@@ -8627,9 +8627,12 @@ namespace UE::Mutable::Private
 				const FString& ParameterName = Program.ConstantStrings[Args.ParameterName];
 				const FParameterKey ParameterKey = { FName(ParameterName), Args.LayerIndex };
 				const FMaterial::FImageParameterData* ImageParameterData = Material->ImageParameters.Find(ParameterKey);
+				bool bParameterFound = false;
 
 				if (ImageParameterData)
 				{
+					bParameterFound = true;
+
 					const TVariant<OP::ADDRESS, TManagedPtr<const FImage>> ImageParameter = ImageParameterData->ImageParameter;
 					if (ImageParameter.IsType<OP::ADDRESS>())
 					{
@@ -8645,11 +8648,27 @@ namespace UE::Mutable::Private
 						StoreValidDesc(Item);
 					}
 				}
-				else if (Material->PassthroughObject.IsResolved())
+				else if(const UMaterialInterface* MaterialInterface = Material->PassthroughObject.Get())
 				{
-					//TODO(Max)
+					UTexture* ParameterValue = nullptr;
+					FHashedMaterialParameterInfo ParameterInfo;
+					ParameterInfo.Name = FScriptName(FName(ParameterName));
+					ParameterInfo.Index = ParameterKey.LayerIndex;
+					ParameterInfo.Association = ParameterInfo.Index == INDEX_NONE ? GlobalParameter : LayerParameter;
+
+					// Find image using the parameter info
+					bParameterFound = MaterialInterface->GetTextureParameterValue(ParameterInfo, ParameterValue);
+
+					if (bParameterFound)
+					{
+						FExtendedImageDesc& Result = ImageDescResults.FindOrAdd(Item.At);
+						Result = GetExternalImageDesc(ParameterValue);
+
+						StoreValidDesc(Item);
+					}
 				}
-				else
+
+				if(!bParameterFound)
 				{
 					ImageDescResults.FindOrAdd(Item.At);
 					StoreValidDesc(Item);
diff --git a/Engine/Plugins/Mutable/Source/MutableTools/Internal/MuT/AST.h b/Engine/Plugins/Mutable/Source/MutableTools/Internal/MuT/AST.h
--- a/Engine/Plugins/Mutable/Source/MutableTools/Internal/MuT/AST.h
+++ b/Engine/Plugins/Mutable/Source/MutableTools/Internal/MuT/AST.h
@@ -848,6 +848,7 @@ namespace UE::Mutable::Private
         public:
             TMap<const ASTOp*, FImageDesc> m_results;
 			bool bIsFormatFromImageParameter = false;
+			FParameterKey ParameterKey;
         };
 
         //!
diff --git a/Engine/Plugins/Mutable/Source/MutableTools/Private/MuT/ASTOpConstantResource.cpp b/Engine/Plugins/Mutable/Source/MutableTools/Private/MuT/ASTOpConstantResource.cpp
--- a/Engine/Plugins/Mutable/Source/MutableTools/Private/MuT/ASTOpConstantResource.cpp
+++ b/Engine/Plugins/Mutable/Source/MutableTools/Private/MuT/ASTOpConstantResource.cpp
@@ -559,7 +559,7 @@ namespace UE::Mutable::Private
 
 
 	//-------------------------------------------------------------------------------------------------
-	FImageDesc ASTOpConstantResource::GetImageDesc(bool, class FGetImageDescContext*) const
+	FImageDesc ASTOpConstantResource::GetImageDesc(bool, class FGetImageDescContext* Context) const
 	{
 		FImageDesc Result;
 
@@ -576,9 +576,13 @@ namespace UE::Mutable::Private
 		}
 		else if (Type == EOpType::MI_CONSTANT)
 		{
-			if (TMap<FParameterKey, ASTChild>::TConstIterator Iterator = ImageOperations.CreateConstIterator())
+			if (Context)
 			{
-				Result = Iterator.Value()->GetImageDesc();
+				if (const ASTChild* Image = ImageOperations.Find(Context->ParameterKey))
+				{
+					check(Image->Child);
+					Result = Image->Child->GetImageDesc();
+				}
 			}
 		}
 		else
diff --git a/Engine/Plugins/Mutable/Source/MutableTools/Private/MuT/ASTOpMaterialBreak.cpp b/Engine/Plugins/Mutable/Source/MutableTools/Private/MuT/ASTOpMaterialBreak.cpp
--- a/Engine/Plugins/Mutable/Source/MutableTools/Private/MuT/ASTOpMaterialBreak.cpp
+++ b/Engine/Plugins/Mutable/Source/MutableTools/Private/MuT/ASTOpMaterialBreak.cpp
@@ -111,7 +111,16 @@ namespace UE::Mutable::Private
 
 		if (Material)
 		{
-			Result = Material->GetImageDesc();
+			// Local context in case it is necessary
+			FGetImageDescContext LocalContext;
+			if (!Context)
+			{
+				Context = &LocalContext;
+			}
+
+			Context->ParameterKey = ParameterKey;
+
+			Result = Material->GetImageDesc(returnBestOption, Context);
 		}
 
 		return Result;
diff --git a/Engine/Plugins/Mutable/Source/MutableTools/Private/MuT/ASTOpMaterialModify.cpp b/Engine/Plugins/Mutable/Source/MutableTools/Private/MuT/ASTOpMaterialModify.cpp
--- a/Engine/Plugins/Mutable/Source/MutableTools/Private/MuT/ASTOpMaterialModify.cpp
+++ b/Engine/Plugins/Mutable/Source/MutableTools/Private/MuT/ASTOpMaterialModify.cpp
@@ -197,9 +197,16 @@ namespace UE::Mutable::Private
 
 		FImageDesc Result;
 
-		if (Material)
+		if (Context)
 		{
-			Result = Material->GetImageDesc();
+			if (const ASTChild* Image = ParametersToModify.Find(Context->ParameterKey))
+			{
+				if (Image->Child)
+				{
+					check(GetOpDataType(Image->Child->GetOpType()) == EDataType::Image);
+					Result = Image->Child->GetImageDesc();
+				}
+			}
 		}
 
 		return Result;
diff --git a/Engine/Plugins/Mutable/Source/MutableTools/Private/MuT/ASTOpParameter.cpp b/Engine/Plugins/Mutable/Source/MutableTools/Private/MuT/ASTOpParameter.cpp
--- a/Engine/Plugins/Mutable/Source/MutableTools/Private/MuT/ASTOpParameter.cpp
+++ b/Engine/Plugins/Mutable/Source/MutableTools/Private/MuT/ASTOpParameter.cpp
@@ -197,15 +197,19 @@ namespace UE::Mutable::Private
 	}
 
 
-	FImageDesc ASTOpParameter::GetImageDesc(bool, FGetImageDescContext*) const
+	FImageDesc ASTOpParameter::GetImageDesc(bool, FGetImageDescContext* Context) const
 	{
 		FImageDesc Result;
 
 		if (Type == EOpType::MI_PARAMETER)
 		{
-			if (TMap<FParameterKey, ASTChild>::TConstIterator Iterator = ImageOperations.CreateConstIterator())
+			if (Context)
 			{
-				Result = Iterator.Value()->GetImageDesc();
+				if (const ASTChild* Image = ImageOperations.Find(Context->ParameterKey))
+				{
+					check(Image->Child);
+					Result = Image->Child->GetImageDesc();
+				}
 			}
 		}
 
