[Mutable] Reimplemented deleted code.

#rnx

#Synced-CL 49922449

diff --git a/Engine/Plugins/Mutable/Source/MutableRuntime/Private/MuR/CodeRunner.cpp b/Engine/Plugins/Mutable/Source/MutableRuntime/Private/MuR/CodeRunner.cpp
--- a/Engine/Plugins/Mutable/Source/MutableRuntime/Private/MuR/CodeRunner.cpp
+++ b/Engine/Plugins/Mutable/Source/MutableRuntime/Private/MuR/CodeRunner.cpp
@@ -8691,22 +8691,25 @@ namespace UE::Mutable::Private
 				CacheMaterialOp.Type = FScheduledOp::EType::Full;
 
 				TManagedPtr<const FMaterial> Material = LoadMaterial(CacheMaterialOp);
-			
+
 				check(Args.ParameterName < (uint32)Program.ConstantStrings.Num());
 				const FString& ParameterName = Program.ConstantStrings[Args.ParameterName];
 				const FParameterKey ParameterKey = { FName(ParameterName), Args.LayerIndex };
 				const FMaterial::FImageParameterData* ImageParameterData = Material->ImageParameters.Find(ParameterKey);
+				bool bParameterFound = false;
 
 				if (ImageParameterData)
 				{
+					bParameterFound = true;
+
 					const TVariant<OP::ADDRESS, TManagedPtr<const FImage>> ImageParameter = ImageParameterData->ImageParameter;
 					if (ImageParameter.IsType<OP::ADDRESS>())
 					{
 						FScheduledOp ItemNextStage(Item.At, Item, 2);
 						ItemNextStage.CustomState = ImageParameter.Get<OP::ADDRESS>();
 
-						AddOp(ItemNextStage,
-								FScheduledOp(ImageParameter.Get<OP::ADDRESS>(), Item));
+						AddOp(ItemNextStage, 
+							FScheduledOp(ImageParameter.Get<OP::ADDRESS>(), Item));
 					}
 					else
 					{
@@ -8714,11 +8717,27 @@ namespace UE::Mutable::Private
 						StoreValidDesc(Item);
 					}
 				}
-				else if (Material->PassthroughObject.IsResolved())
+				else if (const UMaterialInterface* MaterialInterface = Material->PassthroughObject.Get())
 				{
-					//TODO(Max)
+					UTexture* ParameterValue = nullptr;
+					FHashedMaterialParameterInfo ParameterInfo;
+					ParameterInfo.Name = FScriptName(FName(ParameterName));
+					ParameterInfo.Index = ParameterKey.LayerIndex;
+					ParameterInfo.Association = ParameterInfo.Index == INDEX_NONE ? GlobalParameter : LayerParameter;
+
+					// Find image using the parameter info
+					bParameterFound = MaterialInterface->GetTextureParameterValue(ParameterInfo, ParameterValue);
+
+					if (bParameterFound)
+					{
+						FExtendedImageDesc& Result = ImageDescResults.FindOrAdd(Item.At);
+						Result = GetExternalImageDesc(ParameterValue);
+
+						StoreValidDesc(Item);
+					}
 				}
-				else
+
+				if (!bParameterFound)
 				{
 					ImageDescResults.FindOrAdd(Item.At);
 					StoreValidDesc(Item);
diff --git a/Engine/Plugins/Mutable/Source/MutableTools/Private/MuT/ASTOpConstantResource.cpp b/Engine/Plugins/Mutable/Source/MutableTools/Private/MuT/ASTOpConstantResource.cpp
--- a/Engine/Plugins/Mutable/Source/MutableTools/Private/MuT/ASTOpConstantResource.cpp
+++ b/Engine/Plugins/Mutable/Source/MutableTools/Private/MuT/ASTOpConstantResource.cpp
@@ -528,7 +528,7 @@ namespace UE::Mutable::Private
 
 
 	//-------------------------------------------------------------------------------------------------
-	FImageDesc ASTOpConstantResource::GetImageDesc(bool, class FGetImageDescContext*) const
+	FImageDesc ASTOpConstantResource::GetImageDesc(bool, class FGetImageDescContext* Context) const
 	{
 		FImageDesc Result;
 
@@ -545,9 +545,13 @@ namespace UE::Mutable::Private
 		}
 		else if (Type == EOpType::MI_CONSTANT)
 		{
-			if (TMap<FParameterKey, ASTChild>::TConstIterator Iterator = ImageOperations.CreateConstIterator())
+			if (Context)
 			{
-				Result = Iterator.Value()->GetImageDesc();
+				if (const ASTChild* Image = ImageOperations.Find(Context->ParameterKey))
+				{
+					check(Image->Child);
+					Result = Image->Child->GetImageDesc();
+				}
 			}
 		}
 		else
