[mutable] Fixed typos

#rnx

#Synced-CL 50094116

diff --git a/Engine/Plugins/Mutable/Source/MutableRuntime/Internal/MuR/ManagedPointer.h b/Engine/Plugins/Mutable/Source/MutableRuntime/Internal/MuR/ManagedPointer.h
--- a/Engine/Plugins/Mutable/Source/MutableRuntime/Internal/MuR/ManagedPointer.h
+++ b/Engine/Plugins/Mutable/Source/MutableRuntime/Internal/MuR/ManagedPointer.h
@@ -1,5 +1,7 @@
 // Copyright Epic Games, Inc. All Rights Reserved.
 
+#pragma once
+
 #include "Async/Mutex.h"
 #include "Misc/EnumClassFlags.h"
 #include "Templates/MemoryOps.h"
@@ -7,7 +9,6 @@
 #include "Misc/ScopeLock.h"
 #include "Templates/UnrealTemplate.h"
 
-#pragma once
 
 namespace UE::Mutable::Private
 {
@@ -1233,7 +1234,7 @@ public:
 private:
 	/**
 	 * Constructs a strong pointer from a weak pointer, allowing you to access the object (if it
-	 * hasn't expired yet.)  Remember, if the object did not have strong refernces and have been 
+	 * hasn't expired yet.)  Remember, if the object did not have strong references and have been 
 	 * deleted the strong pointer will not be valid.  You should always check to make sure this strong
 	 * pointer is valid before trying to dereference the strong pointer.
 	 *
@@ -1258,7 +1259,7 @@ private:
 
 	/**
 	 * Constructs a strong pointer from a weak pointer, allowing you to access the object (if it
-	 * hasn't expired yet.)  Remember, if the object did not have strong refernces and have been 
+	 * hasn't expired yet.)  Remember, if the object did not have strong references and have been 
 	 * deleted the strong pointer will not be valid.  You should always check to make sure this strong
 	 * pointer is valid before trying to dereference the strong pointer.
 	 *
