[mutable] Avoided bug causing some tab spawners appearing multiple times in the UI.
- There seems to be a non-mutable bug causing this but, since it makes not a lot of sense to have one type of tab outside the "tooling" part of the UI I added it there. That seems to fix the issue for us.

#rnx

#Synced-CL 50094116

diff --git a/Engine/Plugins/Mutable/Source/CustomizableObjectEditor/Private/MuCOE/CustomizableObjectInstanceEditor.cpp b/Engine/Plugins/Mutable/Source/CustomizableObjectEditor/Private/MuCOE/CustomizableObjectInstanceEditor.cpp
--- a/Engine/Plugins/Mutable/Source/CustomizableObjectEditor/Private/MuCOE/CustomizableObjectInstanceEditor.cpp
+++ b/Engine/Plugins/Mutable/Source/CustomizableObjectEditor/Private/MuCOE/CustomizableObjectInstanceEditor.cpp
@@ -139,7 +139,8 @@ void FCustomizableObjectInstanceEditor::RegisterTabSpawners(const TSharedRef<cla
 		.SetGroup(MenuStructure.GetToolsCategory());
 
 	InTabManager->RegisterTabSpawner(TextureAnalyzerTabId, FOnSpawnTab::CreateSP(this, &FCustomizableObjectInstanceEditor::SpawnTab_TextureAnalyzer))
-		.SetDisplayName(LOCTEXT("InstanceTextureAnalyzer", "Texture Analyzer"));
+		.SetDisplayName(LOCTEXT("InstanceTextureAnalyzer", "Texture Analyzer"))
+		.SetGroup(MenuStructure.GetToolsCategory());
 }
 
 void FCustomizableObjectInstanceEditor::UnregisterTabSpawners(const TSharedRef<class FTabManager>& InTabManager)
