[mutable] Fixed some pin hidding/showing interactions wrongfully not being marked as transactions.

#rnx

#Synced-CL 50094116

diff --git a/Engine/Plugins/Mutable/Source/CustomizableObjectEditor/Private/MuCOE/Nodes/CustomizableObjectNode.cpp b/Engine/Plugins/Mutable/Source/CustomizableObjectEditor/Private/MuCOE/Nodes/CustomizableObjectNode.cpp
--- a/Engine/Plugins/Mutable/Source/CustomizableObjectEditor/Private/MuCOE/Nodes/CustomizableObjectNode.cpp
+++ b/Engine/Plugins/Mutable/Source/CustomizableObjectEditor/Private/MuCOE/Nodes/CustomizableObjectNode.cpp
@@ -660,16 +660,24 @@ int32 UCustomizableObjectNode::GetLOD() const
 }
 
 
-void UCustomizableObjectNode::SetPinHidden(UEdGraphPin& Pin, bool bHidden)
+void UCustomizableObjectNode::SetPinHidden(UEdGraphPin& InPin, bool bHidden)
 {
-	Pin.SafeSetHidden(bHidden && CanPinBeHidden(Pin));
+	const FText TransactionName = bHidden ? (LOCTEXT("HidePinTransaction", "Hide Pin")) : LOCTEXT("ShowPinTransaction", "Show Pin");
+	FScopedTransaction Transaction(TransactionName);
+	Modify();
+	
+	InPin.SafeSetHidden(bHidden && CanPinBeHidden(InPin));
 	GetGraph()->NotifyGraphChanged();
 }
 
 
-void UCustomizableObjectNode::SetPinHidden(const TArray<UEdGraphPin*>& PinsToHide, bool bHidden)
+void UCustomizableObjectNode::SetPinHidden(const TArray<UEdGraphPin*>& InPins, bool bHidden)
 {
-	for (UEdGraphPin* Pin : PinsToHide)
+	const FText TransactionName = bHidden ? (LOCTEXT("HideAllPinsTransaction", "Hide All Pins")) : LOCTEXT("ShowAllPinsTransaction", "Show All Pins");
+	FScopedTransaction Transaction(TransactionName);
+	Modify();
+	
+	for (UEdGraphPin* Pin : InPins)
 	{
 		Pin->SafeSetHidden(bHidden && CanPinBeHidden(*Pin));
 	}
diff --git a/Engine/Plugins/Mutable/Source/CustomizableObjectEditor/Private/MuCOE/PinViewer/SPinViewer.cpp b/Engine/Plugins/Mutable/Source/CustomizableObjectEditor/Private/MuCOE/PinViewer/SPinViewer.cpp
--- a/Engine/Plugins/Mutable/Source/CustomizableObjectEditor/Private/MuCOE/PinViewer/SPinViewer.cpp
+++ b/Engine/Plugins/Mutable/Source/CustomizableObjectEditor/Private/MuCOE/PinViewer/SPinViewer.cpp
@@ -315,20 +315,14 @@ void SPinViewer::OnFilterTextChanged(const FText& SearchText)
 
 FReply SPinViewer::OnShowAllPressed() const
 {
-	FScopedTransaction Transaction(LOCTEXT("ShowAllPinsTransaction", "Show All Pins"));
-	Node->Modify();
 	Node->SetPinHidden(Node->GetAllNonOrphanPins(), false);
-
 	return FReply::Handled();
 }
 
 
 FReply SPinViewer::OnHideAllPressed() const
 {
-	FScopedTransaction Transaction(LOCTEXT("HideAllPinsTransaction", "Hide All Pins"));
-	Node->Modify();
 	Node->SetPinHidden(Node->GetAllNonOrphanPins(), true);
-
 	return FReply::Handled();
 }
 
diff --git a/Engine/Plugins/Mutable/Source/CustomizableObjectEditor/Private/MuCOE/PinViewer/SPinViewerListRow.cpp b/Engine/Plugins/Mutable/Source/CustomizableObjectEditor/Private/MuCOE/PinViewer/SPinViewerListRow.cpp
--- a/Engine/Plugins/Mutable/Source/CustomizableObjectEditor/Private/MuCOE/PinViewer/SPinViewerListRow.cpp
+++ b/Engine/Plugins/Mutable/Source/CustomizableObjectEditor/Private/MuCOE/PinViewer/SPinViewerListRow.cpp
@@ -177,8 +177,6 @@ void SPinViewerListRow::SetAdditionalWidgetVisibility(const EVisibility InVisibi
 
 void SPinViewerListRow::OnPinVisibilityCheckStateChanged(ECheckBoxState NewRadioState)
 {
-	FScopedTransaction Transaction(LOCTEXT("ChangedPinVisibilityTransaction", "Changed Pin Visibility"));
-	PinViewer->Node->Modify();
 	PinViewer->Node->SetPinHidden(*PinReference.Get(), NewRadioState == ECheckBoxState::Unchecked);
 }
 
diff --git a/Engine/Plugins/Mutable/Source/CustomizableObjectEditor/Public/MuCOE/Nodes/CustomizableObjectNode.h b/Engine/Plugins/Mutable/Source/CustomizableObjectEditor/Public/MuCOE/Nodes/CustomizableObjectNode.h
--- a/Engine/Plugins/Mutable/Source/CustomizableObjectEditor/Public/MuCOE/Nodes/CustomizableObjectNode.h
+++ b/Engine/Plugins/Mutable/Source/CustomizableObjectEditor/Public/MuCOE/Nodes/CustomizableObjectNode.h
@@ -265,7 +265,7 @@ public:
 	UE_API int32 GetLOD() const;
 	
 	/** Set a pin to be hidden or not. */
-	UE_API void SetPinHidden(UEdGraphPin& Pin, bool bHidden);
+	UE_API void SetPinHidden(UEdGraphPin& InPin, bool bHidden);
 
 	/** Set an array of pins to be hidden or not. 
 	 *
