[Mutable] Fix code viewer crash when previewing meshes without bonemaps.

#rnx

#Synced-CL 50154202

diff --git a/Engine/Plugins/Mutable/Source/CustomizableObject/Private/MuCO/UnrealConversionUtils.cpp b/Engine/Plugins/Mutable/Source/CustomizableObject/Private/MuCO/UnrealConversionUtils.cpp
--- a/Engine/Plugins/Mutable/Source/CustomizableObject/Private/MuCO/UnrealConversionUtils.cpp
+++ b/Engine/Plugins/Mutable/Source/CustomizableObject/Private/MuCO/UnrealConversionUtils.cpp
@@ -276,10 +276,25 @@ namespace UnrealConversionUtils
 			// InBoneMaps may contain BoneMaps from other sections. Copy the bones belonging to this mesh.
 			FirstBone += InFirstBoneMapIndex;
 			
-			Section.BoneMap.Reserve(BoneCount);
-			for (int32 BoneMapIndex = 0; BoneMapIndex < BoneCount; ++BoneMapIndex, ++FirstBone)
+			if (InBoneMap.IsValidIndex(FirstBone + BoneCount - 1))
 			{
-				Section.BoneMap.Add(InRefSkeleton.FindBoneIndex(InBoneMap[FirstBone]));
+				Section.BoneMap.Reserve(BoneCount);
+				for (int32 BoneMapIndex = 0; BoneMapIndex < BoneCount; ++BoneMapIndex, ++FirstBone)
+				{
+					const int32 BoneIndex = InRefSkeleton.FindBoneIndex(InBoneMap[FirstBone]);
+					if (BoneIndex != INDEX_NONE)
+					{
+						Section.BoneMap.Add(BoneIndex);
+					}
+					else
+					{
+						Section.BoneMap.Add(0);
+					}
+				}
+			}
+			else
+			{
+				Section.BoneMap.Init(0, BoneCount);
 			}
 		}
 	}
