[Mutable] Add support for Clothing in Skeletal Mesh parameters.

Moved clothing to Core. Mernging operations are now done inside Mutable Core while executing the VM operations.

#Synced-CL 50155263

#Synced-CL 50155263

diff --git a/Engine/Plugins/Mutable/Source/CustomizableObject/Internal/MuCO/CustomizableObjectClothingTypes.h b/Engine/Plugins/Mutable/Source/CustomizableObject/Internal/MuCO/CustomizableObjectClothingTypes.h
deleted file mode 100644
diff --git a/Engine/Plugins/Mutable/Source/CustomizableObject/Internal/MuCO/CustomizableObjectCompilerTypes.h b/Engine/Plugins/Mutable/Source/CustomizableObject/Internal/MuCO/CustomizableObjectCompilerTypes.h
--- a/Engine/Plugins/Mutable/Source/CustomizableObject/Internal/MuCO/CustomizableObjectCompilerTypes.h
+++ b/Engine/Plugins/Mutable/Source/CustomizableObject/Internal/MuCO/CustomizableObjectCompilerTypes.h
@@ -51,9 +51,6 @@ struct FCompilationOptions
 	// Used to enable the use of real time morph targets.
 	bool bRealTimeMorphTargetsEnabled = false;
 
-	// Used to enable the use of clothing.
-	bool bClothingEnabled = false;
-
 	// Used to enable 16 bit bone weights
 	bool b16BitBoneWeightsEnabled = false;
 
diff --git a/Engine/Plugins/Mutable/Source/CustomizableObject/Internal/MuCO/CustomizableObjectInstancePrivate.h b/Engine/Plugins/Mutable/Source/CustomizableObject/Internal/MuCO/CustomizableObjectInstancePrivate.h
--- a/Engine/Plugins/Mutable/Source/CustomizableObject/Internal/MuCO/CustomizableObjectInstancePrivate.h
+++ b/Engine/Plugins/Mutable/Source/CustomizableObject/Internal/MuCO/CustomizableObjectInstancePrivate.h
@@ -141,9 +141,6 @@ struct FCustomizableInstanceComponentData
 	UPROPERTY(Transient)
 	FReferencedPhysicsAssets PhysicsAssets;
 
-	/** Clothing PhysicsAssets required by the current generated instance. PhysicsAssets to be loaded and merged.*/
-	TArray<TPair<int32, int32>> ClothingPhysicsAssetsToStream;
-
 	FComponentDescriptor LastComponentDescriptor;
 
 	UPROPERTY(Transient)
@@ -348,9 +345,6 @@ public:
 	UPROPERTY(Transient)
 	TArray<FCustomizableInstanceComponentData> ComponentsData;
 	
-	UPROPERTY(Transient)
-	TArray< TObjectPtr<UPhysicsAsset> > ClothingPhysicsAssets;
-
 	// To keep loaded AnimBPs referenced and prevent GC
 	UPROPERTY(Transient, Category = Animation, editfixedsize, VisibleAnywhere)
 	TArray<TSubclassOf<UAnimInstance>> GatheredAnimBPs;
@@ -423,7 +417,7 @@ UE_API USkeleton* MergeSkeletons(UCustomizableObject& CustomizableObject, FCusto
 
 UE_API void InitSkeletalMeshData(const TSharedRef<FUpdateContextPrivate>& Context, USkeletalMesh* SkeletalMesh, const FMutableRefSkeletalMeshData& RefSkeletalMeshData, const UCustomizableObject& CustomizableObject, UE::Mutable::Private::FComponentId ComponentId);
 
-UE_API void BuildOrCopyClothingData(const TSharedRef<FUpdateContextPrivate>& Context, USkeletalMesh* SkeletalMesh, const UModelResources& CustomizableObjectInstance, int32 ComponentIndex, const TArray<TObjectPtr<UPhysicsAsset>>& ClothingPhysicsAssets);
+UE_API void BuildOrCopyClothingData(const TSharedRef<FUpdateContextPrivate>& Context, USkeletalMesh* SkeletalMesh, const UModelResources& CustomizableObjectInstance, int32 InstanceComponentIndex);
 
 UE_API bool BuildSkeletonData(const TSharedRef<FUpdateContextPrivate>& Context, USkeletalMesh& SkeletalMesh, int32 ComponentIndex);
 
diff --git a/Engine/Plugins/Mutable/Source/CustomizableObject/Internal/MuCO/CustomizableObjectPrivate.h b/Engine/Plugins/Mutable/Source/CustomizableObject/Internal/MuCO/CustomizableObjectPrivate.h
--- a/Engine/Plugins/Mutable/Source/CustomizableObject/Internal/MuCO/CustomizableObjectPrivate.h
+++ b/Engine/Plugins/Mutable/Source/CustomizableObject/Internal/MuCO/CustomizableObjectPrivate.h
@@ -4,7 +4,6 @@
 
 #include "MuCO/CustomizableObject.h"
 #include "MuCO/CustomizableObjectCompilerTypes.h"
-#include "MuCO/CustomizableObjectClothingTypes.h"
 #include "MuCO/CustomizableObjectStreamedResourceData.h"
 #include "MuCO/StateMachine.h"
 #include "MuCO/CustomizableObjectUIData.h"
@@ -51,6 +50,7 @@ class UAnimInstance;
 class UAssetUserData;
 class USkeletalMeshLODSettings;
 class UModelResources;
+class ITargetPlatform;
 struct FModelStreamableBulkData;
 struct FObjectAndNameAsStringProxyArchive;
 struct FCustomizableObjectInstanceDescriptor;
@@ -66,109 +66,6 @@ FGuid CUSTOMIZABLEOBJECT_API GenerateDataDistributionIdentifier(const UCustomiza
 
 #endif
 
-// A USTRUCT version of FMeshToMeshVertData in SkeletalMeshTypes.h
-// We are taking advantage of the padding data to store from which asset this data comes from
-// maintaining the same memory footprint than the original.
-USTRUCT()
-struct FCustomizableObjectMeshToMeshVertData
-{
-	GENERATED_USTRUCT_BODY()
-
-	FCustomizableObjectMeshToMeshVertData() = default;
-
-	explicit FCustomizableObjectMeshToMeshVertData(const FMeshToMeshVertData& Original)
-		: Weight(Original.Weight)
-	{
-		for (int32 i = 0; i < 4; ++i)
-		{
-			PositionBaryCoordsAndDist[i] = Original.PositionBaryCoordsAndDist[i];
-			NormalBaryCoordsAndDist[i] = Original.NormalBaryCoordsAndDist[i];
-			TangentBaryCoordsAndDist[i] = Original.TangentBaryCoordsAndDist[i];
-			SourceMeshVertIndices[i] = Original.SourceMeshVertIndices[i];
-		}
-	}
-		
-
-	explicit operator FMeshToMeshVertData() const
-	{
-		FMeshToMeshVertData ReturnValue;
-
-		for (int32 i = 0; i < 4; ++i)
-		{
-			ReturnValue.PositionBaryCoordsAndDist[i] = PositionBaryCoordsAndDist[i];
-			ReturnValue.NormalBaryCoordsAndDist[i] = NormalBaryCoordsAndDist[i];
-			ReturnValue.TangentBaryCoordsAndDist[i] = TangentBaryCoordsAndDist[i];
-			ReturnValue.SourceMeshVertIndices[i] = SourceMeshVertIndices[i];
-		}
-		ReturnValue.Weight = Weight;
-		ReturnValue.Padding = 0;
-
-		return ReturnValue;
-	}
-
-	
-	// Barycentric coords and distance along normal for the position of the final vert
-	UPROPERTY()
-	float PositionBaryCoordsAndDist[4] = {0.f};
-
-	// Barycentric coords and distance along normal for the location of the unit normal endpoint
-	// Actual normal = ResolvedNormalPosition - ResolvedPosition
-	UPROPERTY()
-	float NormalBaryCoordsAndDist[4] = {0.f};
-
-	// Barycentric coords and distance along normal for the location of the unit Tangent endpoint
-	// Actual normal = ResolvedNormalPosition - ResolvedPosition
-	UPROPERTY()
-	float TangentBaryCoordsAndDist[4] = {0.f};
-
-	// Contains the 3 indices for verts in the source mesh forming a triangle, the last element
-	// is a flag to decide how the skinning works, 0xffff uses no simulation, and just normal
-	// skinning, anything else uses the source mesh and the above skin data to get the final position
-	UPROPERTY()
-	uint16	 SourceMeshVertIndices[4] = {0u, 0u, 0u, 0u};
-
-	UPROPERTY()
-	float Weight = 0.0f;
-
-	// Non serialized, unused padding. This is present in the FMeshToMeshVertData struct as Padding for alignment.
-	uint32 UnusedPadding = 0;
-
-	/**
-	 * Serializer
-	 *
-	 * @param Ar - archive to serialize with
-	 * @param V - vertex to serialize
-	 * @return archive that was used
-	 */
-	friend FArchive& operator<<(FArchive& Ar, FCustomizableObjectMeshToMeshVertData& V)
-	{
-		Ar << V.PositionBaryCoordsAndDist[0];
-		Ar << V.PositionBaryCoordsAndDist[1];
-		Ar << V.PositionBaryCoordsAndDist[2];
-		Ar << V.PositionBaryCoordsAndDist[3];
-		Ar << V.NormalBaryCoordsAndDist[0];
-		Ar << V.NormalBaryCoordsAndDist[1];
-		Ar << V.NormalBaryCoordsAndDist[2];
-		Ar << V.NormalBaryCoordsAndDist[3];
-		Ar << V.TangentBaryCoordsAndDist[0];
-		Ar << V.TangentBaryCoordsAndDist[1];
-		Ar << V.TangentBaryCoordsAndDist[2];
-		Ar << V.TangentBaryCoordsAndDist[3];
-
-		Ar << V.SourceMeshVertIndices[0];
-		Ar << V.SourceMeshVertIndices[1];
-		Ar << V.SourceMeshVertIndices[2];
-		Ar << V.SourceMeshVertIndices[3];
-
-		Ar << V.Weight;
-
-		return Ar;
-	}
-};
-static_assert(sizeof(FCustomizableObjectMeshToMeshVertData) == sizeof(float)*4*3 + sizeof(uint16)*4 + sizeof(float) + sizeof(uint32));
-template<> struct TCanBulkSerialize<FCustomizableObjectMeshToMeshVertData> { enum { Value = true }; };
-
-
 // Warning! MutableCompiledDataHeader must be the first data serialized in a stream
 struct MutableCompiledDataStreamHeader
 {
@@ -194,7 +91,6 @@ struct FCustomizableObjectStreameableResourceId
 		None                  = 0,
 		AssetUserData         = 1,
 		RealTimeMorphTarget   = 2,
-		Clothing              = 3,
 		Socket	              = 4,
 	};
 
@@ -696,17 +592,13 @@ struct FMutableMeshMetadata
 	
 	UPROPERTY()
 	uint32 MorphMetadataId = 0;
-
-	UPROPERTY()
-	uint32 ClothingMetadataId = 0;
-
+	
 	UPROPERTY()
 	uint32 SurfaceMetadataId = 0;
 
 	friend FArchive& operator<<(FArchive& Ar, FMutableMeshMetadata& Elem)
 	{
 		Ar << Elem.MorphMetadataId;
-		Ar << Elem.ClothingMetadataId;
 		Ar << Elem.SurfaceMetadataId;
 
 		return Ar;
@@ -735,41 +627,6 @@ struct FMutableSurfaceMetadata
 };
 
 
-USTRUCT()
-struct FClothingStreamable
-{
-	GENERATED_USTRUCT_BODY()
-	
-	UPROPERTY()
-	int32 ClothingAssetIndex = INDEX_NONE;
-	
-	UPROPERTY()
-	int32 ClothingAssetLOD = INDEX_NONE;
-	
-	UPROPERTY()
-	int32 PhysicsAssetIndex = INDEX_NONE;
-
-	UPROPERTY()
-	uint32 Size = 0;
-
-	UPROPERTY()
-	FMutableStreamableBlock Block;
-
-	UPROPERTY()
-	uint32 SourceId = 0;
-
-	friend FArchive& operator<<(FArchive& Ar, FClothingStreamable& Elem)
-	{
-		Ar << Elem.ClothingAssetIndex;
-		Ar << Elem.ClothingAssetLOD;
-		Ar << Elem.PhysicsAssetIndex;
-		Ar << Elem.Size;
-		Ar << Elem.Block;
-		Ar << Elem.SourceId;
-		return Ar;
-	}
-};
-
 USTRUCT()
 struct FMorphTargetVertexData
 {
@@ -925,8 +782,7 @@ namespace UE::Mutable::Private
 		None = 0,
 		Model,
 		RealTimeMorph,
-		Clothing,
-
+		
 		DataTypeCount
 	};
 
@@ -1098,10 +954,7 @@ namespace UE::Mutable::Private
 
 		/** */
 		FModelStreamableData MorphStreamableData;
-
-		/** */
-		FModelStreamableData ClothingStreamableData;
-
+		
 		/** List of files to serialize. Each file has a list of binary blocks to be serialized. */
 		TArray<FFile> BulkDataFiles;
 	};
@@ -1156,9 +1009,7 @@ struct FModelStreamableBulkData
 {
 	/** Map of Hash to Streaming blocks, used to stream a block of data representing a resource from the BulkData */
 	TMap<uint32, FMutableStreamableBlock> ModelStreamables;
-
-	TMap<uint32, FClothingStreamable> ClothingStreamables;
-
+	
 	TMap<uint32, FRealTimeMorphStreamable> RealTimeMorphStreamables;
 
 	TArray<FByteBulkData> StreamableBulkData;
@@ -1180,7 +1031,6 @@ struct FModelStreamableBulkData
 	friend FArchive& operator<<(FArchive& Ar, FModelStreamableBulkData& Data)
 	{
 		Ar << Data.ModelStreamables;
-		Ar << Data.ClothingStreamables;
 		Ar << Data.RealTimeMorphStreamables;
 		Ar << Data.DDCValues;
 		// Don't serialize FByteBulkData manually, the data will be skipped.
@@ -1240,6 +1090,10 @@ public:
 	
 	UPROPERTY()
 	TMap<uint32, TSoftObjectPtr<UObject>> PassthroughObjects;
+
+	/** Objects that we can not directly reference since they are private. Pointing the original package is not of our interest since we would be including unnecessary assets. */
+	UPROPERTY()
+	TArray<TObjectPtr<UObject>> StolenPrivateObjects;
 	
 #if WITH_EDITORONLY_DATA
 	/** Runtime referenced textures used by the UE::Mutable::Private::FModel. */
@@ -1293,13 +1147,6 @@ public:
 	TMap<FIntegerParameterOptionKey, FIntegerParameterOptionDataTable> IntParameterOptionDataTable;
 #endif
 	
-	UPROPERTY()
-	TArray<FCustomizableObjectClothConfigData> ClothSharedConfigsData;	
-
-	UPROPERTY()
-	TArray<FCustomizableObjectClothingAssetData> ClothingAssetsData;
-
-
 	/** Currently not used, this option should be selectable from editor maybe as a compilation flag */
 	UPROPERTY()
 	bool bAllowClothingPhysicsEditsPropagation = true;
@@ -1708,6 +1555,8 @@ enum class ECustomizableObjectVersions : int32
 	ComponentId2,
 
 	BackoutGenerateBulkFileWithDataLimit,
+
+	ClothCore,
 	
 	ConvertBoneIdsToBoneNames,
 
@@ -1886,8 +1735,6 @@ public:
 	CUSTOMIZABLEOBJECT_API bool IsTableMaterialsParentCheckDisabled() const;
 
 	CUSTOMIZABLEOBJECT_API bool IsRealTimeMorphTargetsEnabled() const;
-
-	CUSTOMIZABLEOBJECT_API bool IsClothingEnabled() const;
 	
 	CUSTOMIZABLEOBJECT_API bool Is16BitBoneWeightsEnabled() const;
 	
diff --git a/Engine/Plugins/Mutable/Source/CustomizableObject/Internal/MuCO/CustomizableObjectSystemPrivate.h b/Engine/Plugins/Mutable/Source/CustomizableObject/Internal/MuCO/CustomizableObjectSystemPrivate.h
--- a/Engine/Plugins/Mutable/Source/CustomizableObject/Internal/MuCO/CustomizableObjectSystemPrivate.h
+++ b/Engine/Plugins/Mutable/Source/CustomizableObject/Internal/MuCO/CustomizableObjectSystemPrivate.h
@@ -69,17 +69,6 @@ struct FMorphTargetMeshData
 };
 
 
-/** Serialized Cloth data. */
-struct FClothingMeshData
-{
-	int32 ClothingAssetIndex = INDEX_NONE;
-	int32 ClothingAssetLOD = INDEX_NONE;
-
-	/** Per vertex data. */
-	TArray<FCustomizableObjectMeshToMeshVertData> Data;
-};
-
-
 /** Mapping of FMorphTargetVertexData local names to global names. */
 struct FMappedMorphTargetMeshData
 {
@@ -106,14 +95,6 @@ void LoadMorphTargetsData(FMutableStreamRequest& MutableDataStreamer, const UE::
 
 void LoadMorphTargetsMetadata(const FMutableStreamRequest& MutableDataStreamer, const UE::Mutable::Private::TManagedRef<const UE::Mutable::Private::FMesh>& Mesh, TMap<uint32, FMorphTargetMeshData>& StreamingResult);
 
-
-/** Request the Mutable Data Streamer to load Cloth blocks. Does not stream them.
- *  @param MutableDataStreamer Streamer to request the blocks.
- *  @param Mesh Mutable mesh to obtain which blocks to load.
- *  @param StreamingResult Key is the block. Value is the container where the data will be deserialized once streamed. */
-void LoadClothing(FMutableStreamRequest& MutableDataStreamer, const UE::Mutable::Private::TManagedRef<const UE::Mutable::Private::FMesh>& Mesh, TMap<uint32, FClothingMeshData>& StreamingResult);
-
-
 /** End a Customizable Object Instance Update. All code paths of an update have to end here. */
 void FinishUpdateGlobal(const TSharedRef<FUpdateContextPrivate>& Context);
 
@@ -412,9 +393,6 @@ struct FInstanceUpdateData
 	/** Key is the Component Name. Value is the final Morph Target data to copy into the Skeletal Mesh. */
 	TMap<UE::Mutable::Private::FComponentId, FSkeletalMeshMorphTargets> RealTimeMorphTargets;
 
-	/** Key is the Block Id. Value is the LoadAdditionalAssetsAndData read destination. */
-	TMap<uint32, FClothingMeshData> ClothingMeshData;
-
 	struct FSkeletonData
 	{
 		TArray<TObjectPtr<USkeleton>> Skeletons;
diff --git a/Engine/Plugins/Mutable/Source/CustomizableObject/Internal/MuCO/UnrealConversionUtils.h b/Engine/Plugins/Mutable/Source/CustomizableObject/Internal/MuCO/UnrealConversionUtils.h
--- a/Engine/Plugins/Mutable/Source/CustomizableObject/Internal/MuCO/UnrealConversionUtils.h
+++ b/Engine/Plugins/Mutable/Source/CustomizableObject/Internal/MuCO/UnrealConversionUtils.h
@@ -15,14 +15,15 @@
 #include "Tasks/Task.h"
 
 
+class UPhysicsAsset;
 struct FClothBufferIndexMapping;
 struct FMeshToMeshVertData;
 struct FCustomizableObjectMeshToMeshVertData;
-struct FClothingMeshData;
 struct FMorphTargetMeshData;
 
 namespace UE::Mutable::Private
 {
+	class FPhysicsBody;
 	class FMesh;
 	class FMeshBufferSet;
 }
@@ -158,12 +159,8 @@ namespace UnrealConversionUtils
 
 	CUSTOMIZABLEOBJECT_API UE::Tasks::FTask ConvertSkeletalMeshFromRuntimeData(USkeletalMesh* SkeletalMesh, int32 LODIndex, int32 SectionIndex, UE::Mutable::Private::TManagedPtr<UE::Mutable::Private::FMesh> Result);
 
-	void GetSectionClothData(const UE::Mutable::Private::FMesh& MutableMesh, int32 LODIndex, const TMap<uint32, FClothingMeshData>& ClothingMeshData, TArray<FSectionClothData>& SectionsClothData, int32& NumClothingDataNotFound);
+	TSharedPtr<UE::Mutable::Private::FPhysicsBody> MakePhysicsBodyFromAsset(UPhysicsAsset* Asset, const TArray<uint8>& BodySetupRelevancyMap);
 	
-	void CopyMeshToMeshClothData(TArray<FSectionClothData>& SectionsClothData);
-
-	/** Based on FSkeletalMeshLODModel::GetClothMappingData() */
-	void CreateClothMapping(const FSectionClothData& SectionClothData, TArray<FMeshToMeshVertData>& MappingData, TArray<FClothBufferIndexMapping>& ClothIndexMapping);
 	
-	void ClothVertexBuffers(FSkeletalMeshLODRenderData& LODResource, const UE::Mutable::Private::FMesh& MutableMesh, const TMap<uint32, FClothingMeshData>& ClothingMeshData, int32 LODIndex);
+	void ClothVertexBuffers(FSkeletalMeshLODRenderData& LODResource, const UE::Mutable::Private::FMesh& MutableMesh);
 }
diff --git a/Engine/Plugins/Mutable/Source/CustomizableObject/Private/MuCO/CustomizableObject.cpp b/Engine/Plugins/Mutable/Source/CustomizableObject/Private/MuCO/CustomizableObject.cpp
--- a/Engine/Plugins/Mutable/Source/CustomizableObject/Private/MuCO/CustomizableObject.cpp
+++ b/Engine/Plugins/Mutable/Source/CustomizableObject/Private/MuCO/CustomizableObject.cpp
@@ -588,6 +588,36 @@ void UModelResources::InitCookData(UObject& InCustomizableObject)
 			Extension->MovePrivateReferencesToContainer(ResourceData.Data, &InCustomizableObject);
 		}
 	}
+	
+	for (TMap<uint32, TSoftObjectPtr<UObject>>::TIterator It = PassthroughObjects.CreateIterator(); It; ++It)
+	{
+		TSoftObjectPtr<UObject>& PassthroughObject = It->Value;
+		
+		UObject* Object = UE::Mutable::Private::LoadObject(PassthroughObject);
+		check(Object);
+		
+		const bool bDuplicate = Object->GetOuter() != &InCustomizableObject && !Object->HasAnyFlags(RF_Public);
+		if (bDuplicate)
+		{
+			UObject* DuplicatedObject = DuplicateObject(Object, &InCustomizableObject);
+
+			PassthroughObject = DuplicatedObject;
+			StolenPrivateObjects.Add(DuplicatedObject);
+
+			if (UClothingAssetCommon* ClothingAsset = Cast<UClothingAssetCommon>(DuplicatedObject))
+			{
+				for (TMap<FName, TObjectPtr<UClothConfigBase>>::TIterator ConfigIt = ClothingAsset->ClothConfigs.CreateIterator(); ConfigIt; ++ConfigIt)
+				{
+					TObjectPtr<UClothConfigBase>& ClothConfig = ConfigIt->Value;
+					
+					UClothConfigBase* DuplicatedClothConfig = DuplicateObject<UClothConfigBase>(ClothConfig, &InCustomizableObject);
+					
+					ClothConfig = DuplicatedClothConfig;
+					StolenPrivateObjects.Add(DuplicatedClothConfig);
+				}
+			}
+		}
+	}
 }
 
 
@@ -829,7 +859,6 @@ FString GetDataTypeExtension(UE::Mutable::Private::EStreamableDataType DataType)
 	case UE::Mutable::Private::EStreamableDataType::None:	return TEXT(".mut");
 	case UE::Mutable::Private::EStreamableDataType::Model: return TEXT("_M.mut");
 	case UE::Mutable::Private::EStreamableDataType::RealTimeMorph: return TEXT("_RTM.mut");
-	case UE::Mutable::Private::EStreamableDataType::Clothing:	return TEXT("_C.mut");
 	default:
 	{
 		unimplemented();
@@ -1152,14 +1181,6 @@ bool UCustomizableObjectPrivate::IsRealTimeMorphTargetsEnabled() const
 }
 
 
-bool UCustomizableObjectPrivate::IsClothingEnabled() const
-{
-	PRAGMA_DISABLE_DEPRECATION_WARNINGS // TODO Remove >5.6
-	return GetPublic()->bEnableClothing;
-	PRAGMA_ENABLE_DEPRECATION_WARNINGS
-}
-
-
 bool UCustomizableObjectPrivate::Is16BitBoneWeightsEnabled() const
 {
 	PRAGMA_DISABLE_DEPRECATION_WARNINGS // TODO Remove >5.6
@@ -2102,7 +2123,6 @@ void FModelStreamableBulkData::Serialize(FArchive& Ar, UObject* Owner, bool bCoo
 	MUTABLE_CPUPROFILER_SCOPE(FModelStreamableBulkData::Serialize);
 
 	Ar << ModelStreamables;
-	Ar << ClothingStreamables;
 	Ar << RealTimeMorphStreamables;
 
 	if (bCooked)
@@ -2404,7 +2424,6 @@ FCompilationOptions UCustomizableObjectPrivate::GetCompileOptions() const
 	Options.EmbeddedDataBytesLimit = EmbeddedDataBytesLimit;
 	Options.CustomizableObjectNumBoneInfluences = ICustomizableObjectModule::Get().GetNumBoneInfluences();
 	Options.bRealTimeMorphTargetsEnabled = IsRealTimeMorphTargetsEnabled();
-	Options.bClothingEnabled = IsClothingEnabled();
 	Options.b16BitBoneWeightsEnabled = Is16BitBoneWeightsEnabled();
 	Options.bSkinWeightProfilesEnabled = IsAltSkinWeightProfilesEnabled();
 	Options.bPhysicsAssetMergeEnabled = IsPhysicsAssetMergeEnabled();
@@ -2455,16 +2474,6 @@ namespace UE::Mutable::Private
 				PlatformData->MorphStreamableData.Get(Block.Id, TArrayView64<uint8>(DestData.GetData() + Block.Offset, Block.Size), bDropData);
 			}
 		}
-		else if (DataType == EStreamableDataType::Clothing)
-		{
-
-			for (int32 BlockIndex = 0; BlockIndex < NumBlocks; ++BlockIndex)
-			{
-				const FBlock& Block = Blocks[BlockIndex];
-				check(Block.Offset + Block.Size <= DestSize);
-				PlatformData->ClothingStreamableData.Get(Block.Id, TArrayView64<uint8>(DestData.GetData() + Block.Offset, Block.Size), bDropData);
-			}
-		}
 		else
 		{
 			checkf(false, TEXT("Unknown file DataType found."));
@@ -2593,21 +2602,6 @@ namespace UE::Mutable::Private
 
 			DefaultBucket.DataSize += RealTimeMorphCategory.DataSize;
 		}
-
-		// Clothing. Iterate clothing streamables to accumulate their sizes.
-		{
-			// Add Clothing to the Default bucket
-			FFileCategory& ClothingCategory = DefaultBucket.Categories.AddDefaulted_GetRef();
-			ClothingCategory.Id.DataType = EStreamableDataType::Clothing;
-			
-			const TMap<uint32, FClothingStreamable>& ClothingStreamables = ModelStreamableBulkData.ClothingStreamables;
-			for (const TPair<uint32, FClothingStreamable>& ClothStreamable : ClothingStreamables)
-			{
-				ClothingCategory.DataSize += ClothStreamable.Value.Size;
-			}
-
-			DefaultBucket.DataSize += ClothingCategory.DataSize;
-		}
 		
 		// Limited number of files in each bucket. Find the ideal file distribution between categories based on the accumulated size of their resources.
 		TArray<FFileCategory> Categories;
@@ -2971,28 +2965,7 @@ namespace UE::Mutable::Private
 				AddNode(RootNode, NumRoms, RealTimeMorphCategory, CurrentBlock);
 			}
 		}
-
-		{
-			const TMap<uint32, FClothingStreamable>& ClothingStreamables = ModelStreamableBulkData.ClothingStreamables;
-
-			const FFileCategoryID ClothingCategory = { EStreamableDataType::Clothing, (uint16)UE::Mutable::Private::EDataType::None, (uint16)EMutableFileFlags::None };
-
-			for (const TPair<uint32, FClothingStreamable>& ClothStreamable : ClothingStreamables)
-			{
-				const uint32 BlockSize = ClothStreamable.Value.Size;
-
-				FBlock CurrentBlock = 
-				{ 
-					.Id = ClothStreamable.Key, 
-					.SourceId = ClothStreamable.Value.SourceId, 
-					.Size = BlockSize, 
-					.Offset = 0 
-				};
-
-				AddNode(RootNode, NumRoms, ClothingCategory, CurrentBlock);
-			}
-		}
-
+		
 		// Create Files list
 		for (TPair<FFileCategoryID, FClassifyNode>& Node : RootNode)
 		{
@@ -3035,19 +3008,6 @@ namespace UE::Mutable::Private
 				}
 				break;
 			}
-			case EStreamableDataType::Clothing:
-			{
-				for (FBlock& Block : File.Blocks)
-				{
-					Block.Offset = SourceOffset;
-					SourceOffset += Block.Size;
-
-					FMutableStreamableBlock& StreamableBlock = ModelStreamableBulkData.ClothingStreamables[Block.Id].Block;
-					StreamableBlock.FileId = FileIndex;
-					StreamableBlock.Offset = Block.Offset;
-				}
-				break;
-			}
 			default:
 				UE_LOG(LogMutable, Error, TEXT("Unknown DataType found while fixing streaming block files ids."));
 				unimplemented();
@@ -3121,7 +3081,6 @@ void SerializeCompilationOptionsForDDC(FArchive& Ar, FCompilationOptions& Option
 	Ar << Options.OptimizationLevel;
 	Ar << Options.CustomizableObjectNumBoneInfluences;
 	Ar << Options.bRealTimeMorphTargetsEnabled;
-	Ar << Options.bClothingEnabled;
 	Ar << Options.b16BitBoneWeightsEnabled;
 	Ar << Options.bSkinWeightProfilesEnabled;
 	Ar << Options.bPhysicsAssetMergeEnabled;
@@ -3379,7 +3338,6 @@ void UCustomizableObjectPrivate::LogMemory()
 		{
 			const FModelStreamableBulkData& Bulk = *ModelStreamableData->StreamingData;
 			StreamableData += Bulk.ModelStreamables.GetAllocatedSize();
-			StreamableData += Bulk.ClothingStreamables.GetAllocatedSize();
 			StreamableData += Bulk.RealTimeMorphStreamables.GetAllocatedSize();
 			StreamableData += Bulk.StreamableBulkData.GetAllocatedSize();
 		}
@@ -3404,8 +3362,6 @@ void UCustomizableObjectPrivate::LogMemory()
 				Arch << ModelResources->SurfaceMetadata;
 				Arch << ModelResources->ParameterUIDataMap;
 				Arch << ModelResources->StateUIDataMap;
-				Arch << ModelResources->ClothSharedConfigsData;
-				Arch << ModelResources->ClothingAssetsData;
 				Arch << ModelResources->ComponentNamesPerObjectComponent;
 				Arch << ModelResources->MinLODs;
 				Arch << ModelResources->MinQualityLevelLODs;
diff --git a/Engine/Plugins/Mutable/Source/CustomizableObject/Private/MuCO/CustomizableObjectInstance.cpp b/Engine/Plugins/Mutable/Source/CustomizableObject/Private/MuCO/CustomizableObjectInstance.cpp
--- a/Engine/Plugins/Mutable/Source/CustomizableObject/Private/MuCO/CustomizableObjectInstance.cpp
+++ b/Engine/Plugins/Mutable/Source/CustomizableObject/Private/MuCO/CustomizableObjectInstance.cpp
@@ -2210,7 +2210,7 @@ bool UCustomizableInstancePrivate::UpdateSkeletalMesh_PostBeginUpdate0(UCustomiz
 		}
 
 		BuildOrCopyMorphTargetsData(Context, SkeletalMesh, ComponentIndex);
-		BuildOrCopyClothingData(Context, SkeletalMesh, ModelResources, ComponentIndex, ClothingPhysicsAssets);
+		BuildOrCopyClothingData(Context, SkeletalMesh, ModelResources, ComponentIndex);
 
 		FSkeletalMeshRenderData* RenderData = SkeletalMesh->GetResourceForRendering();
 		ensure(RenderData && RenderData->LODRenderData.Num() > 0);
@@ -3991,54 +3991,16 @@ void BuildOrCopyMorphTargetsData(const TSharedRef<FUpdateContextPrivate>& Contex
 	SkeletalMesh->InitMorphTargets(bInKeepEmptyMorphTargets); // True to avoid removing streamed Morph Targets.
 }
 
-namespace 
-{
-	// Only used to be able to create new clothing assets and assign a new guid to them without the factory.
-	class UCustomizableObjectClothingAsset : public UClothingAssetCommon
-	{
-	public:
-		void AssignNewGuid()
-		{
-			AssetGuid = FGuid::NewGuid();
-		}
-	};
-
-}
 
 void BuildOrCopyClothingData(
 		const TSharedRef<FUpdateContextPrivate>& Context, 
 		USkeletalMesh* SkeletalMesh, 
 		const UModelResources& ModelResources,
-		int32 ComponentIndex,
-		const TArray<TObjectPtr<UPhysicsAsset>>& ClothingPhysicsAssets)
+		int32 ComponentIndex)
 {
 	MUTABLE_CPUPROFILER_SCOPE(UCustomizableInstancePrivate::BuildOrCopyClothingData);
 
-	struct FSectionClothMetadata
-	{
-		int32 SectionIndex;
-		int32 LODIndex;
-		int32 ClothAssetIndex;
-		int32 ClothAssetLodIndex;
-		uint32 NumVertices; // Upper bound
-	};
-
-	struct FPerClothAssetData
-	{
-		int32 MinLOD = 0;
-		TArray<TArray<int32, TInlineAllocator<8>>, TInlineAllocator<8>> AttachedSections; // Indices in SectionsWithCloth for render sections attached to this ClothAsset.
-		FName Name;
-		UPhysicsAsset* PhysicsAsset = nullptr;
-		UClothingAssetCommon* ClothingAsset = nullptr;
-	};
-	
-	const TArray<FCustomizableObjectClothingAssetData>& ClothingAssetsData = ModelResources.ClothingAssetsData;
-	const TArray<FCustomizableObjectClothConfigData>& ClothSharedConfigsData = ModelResources.ClothSharedConfigsData;
-
-	if (!(ClothingAssetsData.Num() && Context->InstanceUpdateData.ClothingMeshData.Num()))
-	{
-		return;
-	}
+	TArray<UClothingAssetBase*> ClothingAssets;
 
 	const TSharedPtr<FInstanceUpdateData::FComponent>& Component = Context->InstanceUpdateData.Components[ComponentIndex];
 	if (!Component)
@@ -4047,399 +4009,112 @@ void BuildOrCopyClothingData(
 	}
 
 	const FName& ComponentName = Context->ComponentNames[Component->Id];
-
-	TArray<FSectionClothMetadata> SectionClothMetadata; // Sections must be sorted ascending
-	SectionClothMetadata.Reserve(32);
 	
-	TBitArray<> LODsWithClothing;
-	LODsWithClothing.Init(false, Component->LODCount);
-
-	// Keep in mind that clothing does not do the Hacky Mutable Streaming Copy. This is because LOD data can not be shared between LODs.
-	// This means that LOD loops are a bit different form usual. With the hacky Mutable streaming, we must generate the requested and the hacky copied ones.
-	
-	// Metadata
+	// Streamed
+	for (int32 LODIndex = Context->FirstLODAvailable[ComponentName]; LODIndex < Context->FirstResidentLOD[ComponentName]; ++LODIndex)
 	{
-		int32 NumClothingDataNotFound = 0;
-			
-		for (int32 LODIndex = Context->FirstLODAvailable[ComponentName]; LODIndex < Component->LODCount; ++LODIndex)
-		{
-			const TSharedRef<FInstanceUpdateData::FLOD>& LOD = Context->InstanceUpdateData.LODs[Component->FirstLOD+LODIndex];
-
-			if (UE::Mutable::Private::TManagedPtr<const UE::Mutable::Private::FMesh> MutableMesh = LOD->Mesh)
-			{
-				for (int32 SectionIndex = 0; SectionIndex < MutableMesh->GetSurfaceCount(); ++SectionIndex)
-				{
-					if (!MutableMesh->Surfaces[SectionIndex].SubMeshes.IsValidIndex(0))
-					{
-						continue;
-					}
-					
-					const UE::Mutable::Private::FSurfaceSubMesh& SubMesh = MutableMesh->Surfaces[SectionIndex].SubMeshes[0];
-					const uint32 ClothResourceId = ModelResources.MeshMetadata[SubMesh.ExternalId].ClothingMetadataId;
+		FSkeletalMeshLODRenderData& RenderLOD = SkeletalMesh->GetResourceForRendering()->LODRenderData[LODIndex];
 
-					if (ClothResourceId == 0)
-					{
-						continue;
-					}
-
-					if (MutableMesh->Surfaces[SectionIndex].SubMeshes.Num() > 1)
-					{
-						UE_LOG(LogMutable, Error, TEXT("Section %d has more than one submesh! Skipping section."), SectionIndex);
-						continue;
-					}
-					
-					FClothingMeshData* SectionClothingData = Context->InstanceUpdateData.ClothingMeshData.Find(ClothResourceId);
-
-					if (!SectionClothingData)
-					{
-						++NumClothingDataNotFound;
-						continue;
-					}
-
-					check(SectionClothingData->ClothingAssetIndex != INDEX_NONE);
-					check(SectionClothingData->ClothingAssetLOD != INDEX_NONE);
-
-					const int32 ClothAssetIndex = SectionClothingData->ClothingAssetIndex;
-					const int32 ClothAssetLodIndex = SectionClothingData->ClothingAssetLOD;
-
-					check(SectionClothingData->ClothingAssetIndex == ClothAssetIndex);
-
-					// Defensive check, this indicates the clothing data might be stale and needs to be recompiled.
-					// Should never happen.
-					if (!ensure(ClothAssetIndex >= 0 && ClothAssetIndex < ClothingAssetsData.Num() && 
-								ClothingAssetsData[ClothAssetIndex].LodData.Num()))
-					{
-						continue;
-					}
-
-					const uint32 NumVertices = SubMesh.VertexEnd - SubMesh.VertexBegin;
-					SectionClothMetadata.Add(FSectionClothMetadata { SectionIndex, LODIndex, ClothAssetIndex, ClothAssetLodIndex, NumVertices });
-
-					LODsWithClothing[LODIndex] = true;
-				}
-			}
-		}
+		const TSharedRef<FInstanceUpdateData::FLOD>& LOD = Context->InstanceUpdateData.LODs[Component->FirstLOD + LODIndex];
 		
-		if (NumClothingDataNotFound > 0)
-		{
-			UE_LOG(LogMutable, Error, TEXT("Some clothing data could not be loaded properly, clothing assets may not behave as expected."));
-		}
-	}
-
-	// No clothing, early out.
-	if (!SectionClothMetadata.Num())
-	{
-		return;
-	}
-	
-	TMap<int32, FPerClothAssetData> PerClothAssetData;
-	PerClothAssetData.Reserve(32);
-
-	// Per Cloth Asset data
-	{
-		// Gather attached sections clothing asset LOD.
-		for (int32 MetadataIndex = 0; MetadataIndex < SectionClothMetadata.Num(); ++MetadataIndex)
+		if (!LOD->Mesh->HasClothing())
 		{
-			const FSectionClothMetadata& SectionClothing = SectionClothMetadata[MetadataIndex];
-			FPerClothAssetData& AssetData = PerClothAssetData.FindOrAdd(SectionClothing.ClothAssetIndex); 
-
-			AssetData.MinLOD = FMath::Min(AssetData.MinLOD, SectionClothing.ClothAssetLodIndex);
-
-			int32 MaxLOD = FMath::Max<int32>(AssetData.AttachedSections.Num() - 1, SectionClothing.ClothAssetLodIndex);
-
-			AssetData.AttachedSections.SetNum(MaxLOD + 1);
-			AssetData.AttachedSections[SectionClothing.ClothAssetLodIndex].Add(MetadataIndex);
+			continue;
 		}
+		
+		TArray<FClothBufferIndexMapping> Mappings;
 
-		for (TPair<int32, FPerClothAssetData>& Data : PerClothAssetData)
+		constexpr int32 Stride = sizeof(FMeshToMeshVertData);
+			
+		int32 NumVertices = 0; // Upper bound
+			
+		const int32 NumSections = LOD->Mesh->Surfaces.Num();
+		for (int32 SectionIndex = 0; SectionIndex < NumSections; ++SectionIndex)
 		{
-			int32 ClothAssetIndex = Data.Key;
-			FPerClothAssetData& ClothAssetData = Data.Value; 
-	
-			ClothAssetData.Name = ClothingAssetsData[ClothAssetIndex].Name;
-			ClothAssetData.PhysicsAsset = ClothingPhysicsAssets[ClothAssetIndex];
+			FSkelMeshRenderSection& RenderSection = RenderLOD.RenderSections[SectionIndex];
+					
+			FClothBufferIndexMapping Mapping;
+			Mapping.BaseVertexIndex = RenderSection.BaseVertexIndex;
+			Mapping.MappingOffset = NumVertices;
+			Mapping.LODBiasStride = RenderSection.NumVertices;
+					
+			Mappings.Add(Mapping);
+						
+			NumVertices += RenderSection.NumVertices;
 		}
+			
+		RenderLOD.ClothVertexBuffer.SetMetadata(Mappings, Stride, NumVertices);
 	}
 	
-	TArray<UnrealConversionUtils::FSectionClothData> SectionsClothData; // Sorted by LOD, Section
-	SectionsClothData.Reserve(32);
-	
-	// Data
+	// Residents
+	for (int32 LODIndex = Context->FirstResidentLOD[ComponentName]; LODIndex < Context->NumLODsAvailable[ComponentName]; ++LODIndex)
 	{
-		MUTABLE_CPUPROFILER_SCOPE(DiscoverSectionsWithCloth);
-
-		int32 NumClothingDataNotFound = 0;
-
-		for (int32 LODIndex = Context->FirstResidentLOD[ComponentName]; LODIndex < Component->LODCount; ++LODIndex)
-		{
-			const TSharedRef<FInstanceUpdateData::FLOD>& LOD = Context->InstanceUpdateData.LODs[Component->FirstLOD+LODIndex];
-
-			if (UE::Mutable::Private::TManagedPtr<const UE::Mutable::Private::FMesh> MutableMesh = LOD->Mesh)
-			{
-				GetSectionClothData(*MutableMesh, LODIndex, Context->InstanceUpdateData.ClothingMeshData, SectionsClothData, NumClothingDataNotFound);
-			}
-		}
+		const TSharedRef<FInstanceUpdateData::FLOD>& LOD = Context->InstanceUpdateData.LODs[Component->FirstLOD + LODIndex];
 
-		if (NumClothingDataNotFound > 0)
+		if (!LOD->Mesh->HasClothing())
 		{
-			UE_LOG(LogMutable, Error, TEXT("Some clothing data could not be loaded properly, clothing assets may not behave as expected."));
+			continue;
 		}
 		
-		CopyMeshToMeshClothData(SectionsClothData);
+		FSkeletalMeshLODRenderData& RenderLOD = SkeletalMesh->GetResourceForRendering()->LODRenderData[LODIndex];
+		UnrealConversionUtils::ClothVertexBuffers(RenderLOD, *LOD->Mesh);
 	}
-
-	// Create Clothing Assets
-	{ 
-		MUTABLE_CPUPROFILER_SCOPE(CreateClothingAssets)
-
-		auto CreateNewClothConfigFromData = [](UObject* Outer, const FCustomizableObjectClothConfigData& ConfigData) -> UClothConfigCommon* 
-		{
-			UClass* ClothConfigClass = FindObject<UClass>(nullptr, *ConfigData.ClassPath);
-			if (ClothConfigClass)
-			{
-				UClothConfigCommon* ClothConfig = NewObject<UClothConfigCommon>(Outer, ClothConfigClass);
-				if (ClothConfig)
-				{
-					FMemoryReaderView MemoryReader(ConfigData.ConfigBytes);
-					ClothConfig->Serialize(MemoryReader);
-
-					return ClothConfig;
-				}
-			}
-
-			return nullptr;
-		};
-
-		TArray<TTuple<FName, UClothConfigCommon*>> SharedConfigs;
-		SharedConfigs.Reserve(ClothSharedConfigsData.Num());
- 
-		for (const FCustomizableObjectClothConfigData& ConfigData : ClothSharedConfigsData)
-		{
-			UClothConfigCommon* ClothConfig = CreateNewClothConfigFromData(SkeletalMesh, ConfigData);
-			if (ClothConfig)
-			{
-				SharedConfigs.Emplace(ConfigData.ConfigName, ClothConfig);
-			}
-		}
 	
-		bool bAllNamesUnique = true;
-		TArray<FName, TInlineAllocator<8>> UniqueAssetNames;
-
-		for (TPair<int32, FPerClothAssetData>& AssetData : PerClothAssetData)
-		{
-			const int32 PrevNumUniqueElems = UniqueAssetNames.Num();
-			const int32 ElemIndex = UniqueAssetNames.AddUnique(AssetData.Value.Name);
-
-			if (ElemIndex < PrevNumUniqueElems)
-			{
-				bAllNamesUnique = false;
-				break;
-			}
-		}
-
-		for (TPair<int32, FPerClothAssetData>& AssetData : PerClothAssetData)
-		{
-			int32 AssetIndex = AssetData.Key;
-			FPerClothAssetData& ClothAssetData = AssetData.Value;  
-
-			FName ClothingAssetObjectName = bAllNamesUnique 
-					? ClothAssetData.Name
-					: FName(FString::Printf(TEXT("%s_%d"), *ClothAssetData.Name.ToString(), AssetIndex));
-
-			UCustomizableObjectClothingAsset* NewClothingAsset = NewObject<UCustomizableObjectClothingAsset>(SkeletalMesh, ClothingAssetObjectName);
-			NewClothingAsset->AssignNewGuid();
-		
-			const int32 NumClothLODs = ClothingAssetsData[AssetIndex].LodData.Num() - ClothAssetData.MinLOD;
-
-			NewClothingAsset->LodData.SetNum(NumClothLODs);
-			for (int32 LODIndex = 0; LODIndex < NumClothLODs; ++LODIndex)
-			{
-				NewClothingAsset->LodData[LODIndex] = ClothingAssetsData[AssetIndex].LodData[LODIndex + ClothAssetData.MinLOD];
-			}
-
-			// Reconstruct clothing asset lod map.
-			NewClothingAsset->LodMap.Init(INDEX_NONE, Component->LODCount);
-			for (int32 LODIndex = 0; LODIndex < NumClothLODs; ++LODIndex)
-			{
-				for (int32 SectionWithClothIndex : ClothAssetData.AttachedSections[LODIndex])
-				{
-					NewClothingAsset->LodMap[SectionClothMetadata[SectionWithClothIndex].LODIndex] = LODIndex; 
-				}
-			}
-
-			NewClothingAsset->UsedBoneIndices = ClothingAssetsData[AssetIndex].UsedBoneIndices;
-			NewClothingAsset->UsedBoneNames = ClothingAssetsData[AssetIndex].UsedBoneNames;
-			NewClothingAsset->ReferenceBoneIndex = ClothingAssetsData[AssetIndex].ReferenceBoneIndex;
-			NewClothingAsset->RefreshBoneMapping(SkeletalMesh);
-			NewClothingAsset->CalculateReferenceBoneIndex();	
-			NewClothingAsset->PhysicsAsset = ClothAssetData.PhysicsAsset;
-
-			for (const FCustomizableObjectClothConfigData& ConfigData : ClothingAssetsData[AssetIndex].ConfigsData)
-			{
-				UClothConfigCommon* ClothConfig = CreateNewClothConfigFromData(NewClothingAsset, ConfigData);
-				if (ClothConfig)
-				{
-					NewClothingAsset->ClothConfigs.Add(ConfigData.ConfigName, ClothConfig);
-				}
-			}
-
-			for (const TTuple<FName, UClothConfigCommon*>& SharedConfig : SharedConfigs)
-			{
-				NewClothingAsset->ClothConfigs.Add(SharedConfig);
-			}
-
-			ClothAssetData.ClothingAsset = NewClothingAsset;
-			SkeletalMesh->GetMeshClothingAssets().AddUnique(NewClothingAsset);
-		}
-	}
-
-	const bool bAllowClothingPhysicsEdits = !bDisableClothingPhysicsEditsPropagation &&
-		ModelResources.bAllowClothingPhysicsEditsPropagation &&
-		!Context->bStreamMeshLODs;
-	
-	if (bAllowClothingPhysicsEdits)
-	{
-		if (IMutableClothingModule* MutableClothingModule = FModuleManager::GetModulePtr<IMutableClothingModule>(MUTABLE_CLOTHING_MODULE_NAME))
-		{
-			for (TPair<int32, FPerClothAssetData>& Data : PerClothAssetData)
-			{
-				FPerClothAssetData& ClothAssetData = Data.Value; 
-				UClothingAssetCommon* ClothingAsset = ClothAssetData.ClothingAsset;
-			
-				if (!ClothingAsset)
-				{
-					continue;
-				}
-
-				bool bNeedsLodTransitionUpdate = false;
-				for (int32 LODIndex = 0; LODIndex < ClothingAsset->LodData.Num(); ++LODIndex)
-				{
-					TArray<TArrayView<FMeshToMeshVertData>, TInlineAllocator<8>> MeshToMeshDataViews;
-
-					for (int32 AttachedSectionIndex : ClothAssetData.AttachedSections[LODIndex])
-					{
-						MeshToMeshDataViews.Add(MakeArrayView(SectionsClothData[AttachedSectionIndex].MappingData));
-					}
-
-					bool bModified = MutableClothingModule->UpdateClothSimulationLOD(
-							LODIndex, *ClothingAsset, MakeConstArrayView(MeshToMeshDataViews));
-				
-					bNeedsLodTransitionUpdate = bNeedsLodTransitionUpdate || bModified;
-				}
-
-				if (bNeedsLodTransitionUpdate)
-				{
-					// This needs to happen after all LODs have been processed.
-					for (int32 LODIndex = 0; LODIndex < ClothingAsset->LodData.Num(); ++LODIndex)
-					{
-						MutableClothingModule->FixLODTransitionMappings(LODIndex, *ClothingAsset);
-					}
-				}
-			}	
-		}
-		else
-		{
-			UE_LOG(LogMutable, Warning, TEXT("MutableClothing plugin could not be found. Make sure the plugin is enabled if you want to use advanced clothing features."));
-		}
-	}
-
-	TArray<TArray<FMeshToMeshVertData>> ResidentLODMappingData;
-	ResidentLODMappingData.SetNum(Component->LODCount);
-
-	TArray<TArray<FClothBufferIndexMapping>> ResidentLODClothIndexMapping;
-	ResidentLODClothIndexMapping.SetNum(Component->LODCount);
+	// Other data
+	const int32 FirstLOD = Context->bStreamMeshLODs ?
+		Context->FirstLODAvailable[ComponentName] :
+		Context->GetFirstRequestedLOD()[ComponentName];
 
-	// Zero all LODs (even those which do not use cloth).
-	for (int32 LODIndex = Context->FirstLODAvailable[ComponentName]; LODIndex < Component->LODCount; ++LODIndex)
+	for (int32 LODIndex = FirstLOD; LODIndex < Component->LODCount; ++LODIndex)
 	{
 		const TSharedRef<FInstanceUpdateData::FLOD>& LOD = Context->InstanceUpdateData.LODs[Component->FirstLOD + LODIndex];
-		
-		UE::Mutable::Private::TManagedPtr<const UE::Mutable::Private::FMesh> MutableMesh = LOD->Mesh;
-		if (!MutableMesh)
+	
+		if (!LOD->Mesh->HasClothing())
 		{
 			continue;
 		}
-		
-		ResidentLODClothIndexMapping[LODIndex].SetNumZeroed(MutableMesh->GetSurfaceCount());
-	}
 
-	// Create the mapping of cloth LODs.
-	for (const UnrealConversionUtils::FSectionClothData& Data : SectionsClothData)
-	{
-		CreateClothMapping(Data, ResidentLODMappingData[Data.LODIndex], ResidentLODClothIndexMapping[Data.LODIndex]);
-	}
-
-	FSkeletalMeshRenderData* RenderResource = SkeletalMesh->GetResourceForRendering();
-	{
-		MUTABLE_CPUPROFILER_SCOPE(InitClothRenderData)
-
-		// Streamed
-		for (int32 LODIndex = Context->FirstLODAvailable[ComponentName]; LODIndex < Context->FirstResidentLOD[ComponentName]; ++LODIndex)
+		FSkeletalMeshLODRenderData& RenderLOD = SkeletalMesh->GetResourceForRendering()->LODRenderData[LODIndex];
+		
+		const int32 NumSections = LOD->Mesh->Surfaces.Num();
+		for (int32 SectionIndex = 0; SectionIndex < NumSections; ++SectionIndex)
 		{
-			FSkeletalMeshLODRenderData& LODModel = RenderResource->LODRenderData[LODIndex];
-
-			if (LODsWithClothing[LODIndex])
-			{
-				TArray<FClothBufferIndexMapping> ClothIndexMapping;
-				ClothIndexMapping.SetNumZeroed(LODModel.RenderSections.Num());
-
-				int32 Stride = sizeof(FMeshToMeshVertData);
+			const UE::Mutable::Private::FCloth& Cloth = LOD->Mesh->ClothSections[SectionIndex];
 			
-				int32 NumVertices = 0; // Upper bound
-				for (const FSectionClothMetadata& Metadata : SectionClothMetadata)
-				{
-					if (Metadata.LODIndex == LODIndex)
-					{
-						// Based on FSkeletalMeshLODModel::GetClothMappingData().
-
-						FSkelMeshRenderSection& RenderSection = LODModel.RenderSections[Metadata.SectionIndex];
-						
-						check(Metadata.NumVertices == RenderSection.NumVertices); // Both values are upper bounds since we can not know the exact number of vertices without executing the code. 
-						
-						ClothIndexMapping[Metadata.SectionIndex].BaseVertexIndex = RenderSection.BaseVertexIndex;
-						ClothIndexMapping[Metadata.SectionIndex].MappingOffset = NumVertices;
-						ClothIndexMapping[Metadata.SectionIndex].LODBiasStride = Metadata.NumVertices;
-						
-						NumVertices += Metadata.NumVertices;
-					}
-				}
-				
-				LODModel.ClothVertexBuffer.SetMetadata(ClothIndexMapping, Stride, NumVertices);
-			}
-		}
-
-
-		// Resident
-		for (int32 LODIndex = Context->FirstResidentLOD[ComponentName]; LODIndex < Context->NumLODsAvailable[ComponentName]; ++LODIndex)
-		{
-			FSkeletalMeshLODRenderData& LODModel = RenderResource->LODRenderData[LODIndex];
-	
-			if (LODsWithClothing[LODIndex])
+			if (!Cloth.IsValid())
 			{
-				LODModel.ClothVertexBuffer.Init(ResidentLODMappingData[LODIndex], ResidentLODClothIndexMapping[LODIndex]);
+				continue;
 			}
+			
+			FSkelMeshRenderSection& RenderSection = RenderLOD.RenderSections[SectionIndex];
+			
+			UClothingAssetCommon* Asset = Cloth.ClothingAsset.Get();
+			const uint32 NumClothVertices = Cloth.Data.Num() ? Cloth.Data.Num() : RenderSection.NumVertices;
+			
+			RenderSection.CorrespondClothAssetIndex = ClothingAssets.AddUnique(Asset);
+
+			RenderSection.ClothMappingDataLODs.AddDefaulted(1);
+			RenderSection.ClothMappingDataLODs[0].SetNum(NumClothVertices);
+			
+			FClothingSectionData& ClothingData = RenderSection.ClothingData;
+			ClothingData.AssetGuid = Asset->GetAssetGuid();
+			ClothingData.AssetLodIndex = Asset->LodMap[LODIndex];
 		}
 	}
 	
-	for (const FSectionClothMetadata& Metadata : SectionClothMetadata)
+	for (int32 Index = 0; Index < ClothingAssets.Num(); ++Index)
 	{
-		FSkeletalMeshLODRenderData& LODModel = RenderResource->LODRenderData[Metadata.LODIndex];
-		FSkelMeshRenderSection& SectionData = LODModel.RenderSections[Metadata.SectionIndex];
-
-		// Ideally we would copy the data of all LODs, but we do not have this information in the initial generation. In any case,
-		// ClothMappingDataLODs is only used for CPU Skinning, and some engine checks (they only check the array size).
-		// The size must be a multiple of SectionData.NumVertices. Currently Mutable only supports one influence per vertex (NumVertices * 1).
-		SectionData.ClothMappingDataLODs.AddDefaulted(1);
-		SectionData.ClothMappingDataLODs[0].SetNum(SectionData.NumVertices); // = MoveTemp(SectionWithCloth.MappingData);
-
-		FPerClothAssetData& AssetData = PerClothAssetData.FindChecked(Metadata.ClothAssetIndex);
+		UClothingAssetCommon* Asset = CastChecked<UClothingAssetCommon>(ClothingAssets[Index]);
 
-		SectionData.CorrespondClothAssetIndex = SkeletalMesh->GetClothingAssetIndex(AssetData.ClothingAsset);
-		SectionData.ClothingData.AssetGuid = AssetData.ClothingAsset->GetAssetGuid();
-		SectionData.ClothingData.AssetLodIndex = AssetData.ClothingAsset->LodMap[Metadata.LODIndex];
+		Asset = DuplicateObject<UClothingAssetCommon>(Asset, SkeletalMesh);
+		Asset->RefreshBoneMapping(SkeletalMesh);
+		Asset->CalculateReferenceBoneIndex();
+		
+		ClothingAssets[Index] = Asset;
 	}
 	
-	SkeletalMesh->SetHasActiveClothingAssets(!SectionClothMetadata.IsEmpty());
+	SkeletalMesh->SetMeshClothingAssets(ClothingAssets);
+	SkeletalMesh->SetHasActiveClothingAssets(!ClothingAssets.IsEmpty());
 }
 
 
@@ -4718,10 +4393,6 @@ UE::Tasks::FTask LoadAdditionalAssetsAndData(const TSharedRef<FUpdateContextPriv
 
 	TArray<TSharedRef<FInstanceUpdateData::FComponent>>& Components = Context->InstanceUpdateData.Components;
 	
-	const int32 NumClothingAssets = ModelResources.ClothingAssetsData.Num();
-	InstancePrivate->ClothingPhysicsAssets.Reset(NumClothingAssets);
-	InstancePrivate->ClothingPhysicsAssets.SetNum(NumClothingAssets);
-
 	InstancePrivate->GatheredAnimBPs.Empty();
 	InstancePrivate->AnimBPGameplayTags.Reset();
 	InstancePrivate->AnimBpPhysicsAssets.Reset();
@@ -4836,88 +4507,6 @@ UE::Tasks::FTask LoadAdditionalAssetsAndData(const TSharedRef<FUpdateContextPriv
 					LoadMorphTargetsMetadata(StreamRequest, MutableMesh.ToStrongRef(), Context->InstanceUpdateData.RealTimeMorphTargetMeshData);
 				}
 			}
-
-			// Cloth
-			{
-				// Data
-				// From FirstResidentLOD instead of FirstRequestedLOD since clothing we generate all LODs, even the hacky streaming copied ones.
-				for (int32 LODIndex = Context->FirstResidentLOD[ComponentName]; LODIndex < Component->LODCount; ++LODIndex)
-				{
-					const TSharedRef<FInstanceUpdateData::FLOD>& LOD = Context->InstanceUpdateData.LODs[Component->FirstLOD + LODIndex];
-
-					UE::Mutable::Private::TManagedPtr<const UE::Mutable::Private::FMesh> MutableMesh = LOD->Mesh;
-					if (!MutableMesh)
-					{
-						continue;
-					}
-
-					if (bUpdateMeshes)
-					{
-						LoadClothing(StreamRequest, MutableMesh.ToStrongRef(), Context->InstanceUpdateData.ClothingMeshData);
-					}
-				}
-
-				// Metadata
-				for (int32 LODIndex = Context->FirstLODAvailable[ComponentName]; LODIndex < Component->LODCount; ++LODIndex)
-				{
-					const TSharedRef<FInstanceUpdateData::FLOD>& LOD = Context->InstanceUpdateData.LODs[Component->FirstLOD + LODIndex];
-
-					UE::Mutable::Private::TManagedPtr<const UE::Mutable::Private::FMesh> MutableMesh = LOD->Mesh;
-					if (!MutableMesh)
-					{
-						continue;
-					}
-				
-					const TArray<uint64>& StreamedResources = MutableMesh->GetStreamedResources();
-
-					for (uint64 ResourceId : StreamedResources)
-					{
-						FCustomizableObjectStreameableResourceId TypedResourceId = BitCast<FCustomizableObjectStreameableResourceId>(ResourceId);
-
-						if (TypedResourceId.Type == (uint8)FCustomizableObjectStreameableResourceId::EType::Clothing)
-						{
-							check(TypedResourceId.Id != 0 && TypedResourceId.Id <= TNumericLimits<uint32>::Max());
-
-							const TMap<uint32, FClothingStreamable>& ClothingStreamables = ModelStreamableBulkData->ClothingStreamables;
-							if (const FClothingStreamable* ClothingStreamable = ClothingStreamables.Find(TypedResourceId.Id))
-							{
-								FClothingMeshData& ReadDestData = Context->InstanceUpdateData.ClothingMeshData.FindOrAdd(TypedResourceId.Id);
-								ReadDestData.ClothingAssetIndex = ClothingStreamable->ClothingAssetIndex;
-								ReadDestData.ClothingAssetLOD = ClothingStreamable->ClothingAssetLOD;
- 							
-								// TODO: Add async loading of ClothingAsset Data. This could be loaded as an streamead resource similar to and the asset user data.
-								int32 ClothingAssetIndex = ClothingStreamable->ClothingAssetIndex; 
-								int32 PhysicsAssetIndex = ClothingStreamable->PhysicsAssetIndex;					
-								const TSoftObjectPtr<UPhysicsAsset>& PhysicsAsset = ModelResources.PhysicsAssets.IsValidIndex(PhysicsAssetIndex) 
-										? ModelResources.PhysicsAssets[PhysicsAssetIndex] 
-										: nullptr;
-
-								// The entry should always be in the map
-								if (!PhysicsAsset.IsNull())
-								{
-									if (PhysicsAsset.Get())
-									{
-										if (InstancePrivate->ClothingPhysicsAssets.IsValidIndex(ClothingAssetIndex))
-										{
-											InstancePrivate->ClothingPhysicsAssets[ClothingAssetIndex] = PhysicsAsset.Get();
-										}
-											ComponentData.ClothingPhysicsAssetsToStream.Emplace(ClothingAssetIndex, PhysicsAssetIndex);
-									}
-									else
-									{
-										ComponentData.ClothingPhysicsAssetsToStream.Emplace(ClothingAssetIndex, PhysicsAssetIndex);
-										AssetsToStream.AddUnique(PhysicsAsset.ToSoftObjectPath());
-									}
-								}
-							}
-							else
-							{
-								UE_LOG(LogMutable, Error, TEXT("Invalid streamed clothing data block [%llu] found."), TypedResourceId.Id);
-							}
-						}
-					}
-				}
-			}
 		}
 		
 		for (int32 LODIndex = Context->FirstLODAvailable[ComponentName]; LODIndex < Component->LODCount; ++LODIndex)
@@ -5279,19 +4868,6 @@ void AdditionalAssetsAsyncLoaded(const TSharedRef<FUpdateContextPrivate> Context
 #endif
 		}
 		PhysicsAssets.PhysicsAssetToLoad.Empty();
-		
-		// Loaded Clothing PhysicsAssets 
-		for ( TPair<int32, int32>& AssetToStream : ComponentData.ClothingPhysicsAssetsToStream )
-		{
-			const int32 AssetIndex = AssetToStream.Key;
-
-			if (InstancePrivate->ClothingPhysicsAssets.IsValidIndex(AssetIndex) && ModelResources.PhysicsAssets.IsValidIndex(AssetToStream.Value))
-			{
-				const TSoftObjectPtr<UPhysicsAsset>& PhysicsAssetPtr = ModelResources.PhysicsAssets[AssetToStream.Value];
-				InstancePrivate->ClothingPhysicsAssets[AssetIndex] = PhysicsAssetPtr.Get();
-			}
-		}
-		ComponentData.ClothingPhysicsAssetsToStream.Empty();
 
 		// Loaded anim BPs
 		for (TPair<FName, TSoftClassPtr<UAnimInstance>>& SlotAnimBP : ComponentData.AnimSlotToBP)
diff --git a/Engine/Plugins/Mutable/Source/CustomizableObject/Private/MuCO/CustomizableObjectMeshUpdate.cpp b/Engine/Plugins/Mutable/Source/CustomizableObject/Private/MuCO/CustomizableObjectMeshUpdate.cpp
--- a/Engine/Plugins/Mutable/Source/CustomizableObject/Private/MuCO/CustomizableObjectMeshUpdate.cpp
+++ b/Engine/Plugins/Mutable/Source/CustomizableObject/Private/MuCO/CustomizableObjectMeshUpdate.cpp
@@ -323,8 +323,6 @@ void FCustomizableObjectMeshStreamIn::RequestMeshUpdate(const FContext& Context)
 
 				LoadMorphTargetsData(StreamRequest, Mesh.ToStrongRef(), RefThis->OperationData->MorphTargetMeshData);
 				LoadMorphTargetsMetadata(StreamRequest, Mesh.ToStrongRef(), RefThis->OperationData->MorphTargetMeshData);
-				
-				LoadClothing(StreamRequest, Mesh.ToStrongRef(), RefThis->OperationData->ClothingMeshData);
 			}
 
 			// Stream data.
@@ -421,7 +419,7 @@ void FCustomizableObjectMeshStreamIn::ConvertMesh(const FContext& Context, bool&
 		UnrealConversionUtils::CopyMutableIndexBuffers(LODResource, MutableMesh.Get(), Mesh->SurfaceIDs[LODIndexAbsolute], bOutMarkRenderStateDirty);
 		UnrealConversionUtils::CopyMutableSkinWeightProfilesBuffers(LODResource, *const_cast<UCustomizableObjectSkeletalMesh*>(Mesh), LODIndexAbsolute, MutableMesh.Get(), Mesh->SkinWeightProfileIDs);
 		UnrealConversionUtils::MorphTargetVertexInfoBuffers(LODResource, *Mesh, *MutableMesh.Get(), OperationData->MorphTargetMeshData, LODIndexAbsolute, bGenerateCPUMorphTargetsIfNeeded);
-		UnrealConversionUtils::ClothVertexBuffers(LODResource, *MutableMesh.Get(), OperationData->ClothingMeshData, LODIndexAbsolute);
+		UnrealConversionUtils::ClothVertexBuffers(LODResource, *MutableMesh.Get());
 		
 		UnrealConversionUtils::UpdateSkeletalMeshLODRenderDataBuffersSize(LODResource);
 	}
diff --git a/Engine/Plugins/Mutable/Source/CustomizableObject/Private/MuCO/CustomizableObjectMeshUpdate.h b/Engine/Plugins/Mutable/Source/CustomizableObject/Private/MuCO/CustomizableObjectMeshUpdate.h
--- a/Engine/Plugins/Mutable/Source/CustomizableObject/Private/MuCO/CustomizableObjectMeshUpdate.h
+++ b/Engine/Plugins/Mutable/Source/CustomizableObject/Private/MuCO/CustomizableObjectMeshUpdate.h
@@ -44,9 +44,6 @@ struct FMutableUpdateMeshContext
 	/** Serialized data per block. */
 	TMap<uint32, FMorphTargetMeshData> MorphTargetMeshData;
 
-	/** Serialized data per block. */
-	TMap<uint32, FClothingMeshData> ClothingMeshData;
-
 	/** Task setup during the nominal operation of the mesh streaming system. It will be null in non-streaming operation. */
 	TSharedPtr<UE::Tasks::FTaskEvent> MeshStreamingPrerequisite;
 
diff --git a/Engine/Plugins/Mutable/Source/CustomizableObject/Private/MuCO/CustomizableObjectSystem.cpp b/Engine/Plugins/Mutable/Source/CustomizableObject/Private/MuCO/CustomizableObjectSystem.cpp
--- a/Engine/Plugins/Mutable/Source/CustomizableObject/Private/MuCO/CustomizableObjectSystem.cpp
+++ b/Engine/Plugins/Mutable/Source/CustomizableObject/Private/MuCO/CustomizableObjectSystem.cpp
@@ -4232,66 +4232,6 @@ void LoadMorphTargetsMetadata(const FMutableStreamRequest& MutableDataStreamer,
 }
 
 
-void LoadClothing(FMutableStreamRequest& MutableDataStreamer, const UE::Mutable::Private::TManagedRef<const UE::Mutable::Private::FMesh>& Mesh, TMap<uint32, FClothingMeshData>& StreamingResult)
-{
-	MUTABLE_CPUPROFILER_SCOPE(LoadClothing);
-
-	TArray<uint32> ClothingStreamableBlocksToStream;
-
-	const TSharedPtr<FModelStreamableBulkData>& ModelStreamableBulkData = MutableDataStreamer.GetModelStreamableBulkData();
-
-	for (uint64 ResourceId : Mesh->GetStreamedResources())
-	{
-		FCustomizableObjectStreameableResourceId TypedResourceId = BitCast<FCustomizableObjectStreameableResourceId>(ResourceId);
-
-		if (TypedResourceId.Type == static_cast<uint8>(FCustomizableObjectStreameableResourceId::EType::Clothing))
-		{
-			check(TypedResourceId.Id != 0 && TypedResourceId.Id <= TNumericLimits<uint32>::Max());
-			
-			if (ModelStreamableBulkData->ClothingStreamables.Contains(TypedResourceId.Id))
-			{
-				ClothingStreamableBlocksToStream.AddUnique(TypedResourceId.Id);
-			}
-			else
-			{
-				UE_LOG(LogMutable, Error, TEXT("Invalid streamed clothing data block [%llu] found."), TypedResourceId.Id);
-			}
-		}
-	}
-	
-	// Clothing blocks to stream 
-	for (const int32 BlockId : ClothingStreamableBlocksToStream)
-	{
-		MUTABLE_CPUPROFILER_SCOPE(ClothingStreamingRequest_Alloc);
-
-		const FClothingStreamable& ClothingStreamable = ModelStreamableBulkData->ClothingStreamables[BlockId];
-		const FMutableStreamableBlock& Block = ClothingStreamable.Block;
-
-		FClothingMeshData& ReadDestData = StreamingResult.FindOrAdd(BlockId);
-
-		// Only request blocks once.
-		if (ReadDestData.Data.Num())
-		{
-			continue;
-		}
-
-		ReadDestData.ClothingAssetIndex = ClothingStreamable.ClothingAssetIndex;
-		ReadDestData.ClothingAssetLOD = ClothingStreamable.ClothingAssetLOD;
-
-		check(ClothingStreamable.Size % sizeof(FCustomizableObjectMeshToMeshVertData) == 0);
-		const uint32 NumElems = ClothingStreamable.Size / sizeof(FCustomizableObjectMeshToMeshVertData);
-
-		ReadDestData.Data.SetNumUninitialized(NumElems);
-
-		MutableDataStreamer.AddBlock(
-			Block,
-			UE::Mutable::Private::EStreamableDataType::Clothing,
-			0,
-			MakeArrayView(reinterpret_cast<uint8*>(ReadDestData.Data.GetData()), ReadDestData.Data.Num() * sizeof(FCustomizableObjectMeshToMeshVertData)));
-	}
-}
-
-
 void ReconstructMorphTargets(const UE::Mutable::Private::FMesh& Mesh, const TArray<FName>& GlobalNames, const TMap<uint32, FMappedMorphTargetMeshData>& MappedMorphTargets, TArray<FMorphTargetLODModel>& OutMorphTargets)
 {
 	const UE::Mutable::Private::FMeshBufferSet& MeshSet = Mesh.GetVertexBuffers();
diff --git a/Engine/Plugins/Mutable/Source/CustomizableObject/Private/MuCO/UnrealConversionUtils.cpp b/Engine/Plugins/Mutable/Source/CustomizableObject/Private/MuCO/UnrealConversionUtils.cpp
--- a/Engine/Plugins/Mutable/Source/CustomizableObject/Private/MuCO/UnrealConversionUtils.cpp
+++ b/Engine/Plugins/Mutable/Source/CustomizableObject/Private/MuCO/UnrealConversionUtils.cpp
@@ -1814,181 +1814,131 @@ namespace UnrealConversionUtils
 				{
 					Result->MorphDataBuffer.Append({0, 0, 0});
 				}
-
+					
 				Result->EnsureSurfaceData();
-			},
-			ConversionTask);
-
-		return FinalConversionTask;
-	}
-
-
-	void GetSectionClothData(const UE::Mutable::Private::FMesh& MutableMesh, int32 LODIndex, const TMap<uint32, FClothingMeshData>& ClothingMeshData, TArray<FSectionClothData>& SectionsClothData, int32& NumClothingDataNotFound)
-	{
-		MUTABLE_CPUPROFILER_SCOPE(GetSectionClothData);
-		
-		const UE::Mutable::Private::FMeshBufferSet& MeshSet = MutableMesh.GetVertexBuffers();
-
-		int32 ClothingIndexBuffer, ClothingIndexChannel;
-		MeshSet.FindChannel(UE::Mutable::Private::EMeshBufferSemantic::Other, 2, &ClothingIndexBuffer, &ClothingIndexChannel);
 
-		int32 ClothingResourceBuffer, ClothingResourceChannel;
-		MeshSet.FindChannel(UE::Mutable::Private::EMeshBufferSemantic::Other, 3, &ClothingResourceBuffer, &ClothingResourceChannel);
-
-		if (ClothingIndexBuffer >= 0 && ClothingResourceBuffer >= 0)
-		{
-			const int32* const ClothingDataBuffer = reinterpret_cast<const int32*>(MeshSet.GetBufferData(ClothingIndexBuffer));
-			const uint32* const ClothingDataResource = reinterpret_cast<const uint32*>(MeshSet.GetBufferData(ClothingResourceBuffer));
-
-			const int32 SurfaceCount = MutableMesh.GetSurfaceCount();
-			for (int32 SectionIndex = 0; SectionIndex < SurfaceCount; ++SectionIndex)
-			{
-				int32 FirstVertex, VerticesCount, FirstIndex, IndicesCount, UnusedBoneIndex, UnusedBoneCount;
-				MutableMesh.GetSurface(SectionIndex, FirstVertex, VerticesCount, FirstIndex, IndicesCount, UnusedBoneIndex, UnusedBoneCount);
-
-				if (VerticesCount == 0 || IndicesCount == 0)
+				// Cloth
 				{
-					continue;
-				}
-				
-				// A Section has cloth data on all its vertices or it does not have it in any.
-				// It can be determined if this section has clothing data just looking at the 
-				// first vertex of the section.
-				TArrayView<const int32> ClothingDataView(ClothingDataBuffer + FirstVertex, VerticesCount);
-				TArrayView<const uint32> ClothingResourceView(ClothingDataResource + FirstVertex, VerticesCount);
+					FSkeletalMeshVertexClothBuffer& ClothVertexBuffer = Context->OriginalLOD->ClothVertexBuffer;
 
-				const int32 IndexCount = MutableMesh.GetIndexBuffers().GetElementCount();
+					const TArray<FClothBufferIndexMapping>& SectionClothIndexMappings = ClothVertexBuffer.GetClothIndexMapping();
 
-				TArrayView<const uint16> IndicesView16Bits;
-				TArrayView<const uint32> IndicesView32Bits;
-				
-				if (IndexCount)
-				{
-					if (MutableMesh.GetIndexBuffers().GetElementSize(0) == 2)
-					{
-						const uint16* IndexPtr = (const uint16*)MutableMesh.GetIndexBuffers().GetBufferData(0);
-						IndicesView16Bits = TArrayView<const uint16>(IndexPtr + FirstIndex, IndicesCount);
-					}
-					else
+					const bool bHasClothData = SectionClothIndexMappings.IsValidIndex(SectionIndex);
+					
+					UClothingAssetCommon* ClothingAsset = Cast<UClothingAssetCommon>(SkeletalMesh->GetClothingAsset(Section.ClothingData.AssetGuid));
+					if (bHasClothData && ClothingAsset)
 					{
-						const uint32* IndexPtr = (const uint32*)MutableMesh.GetIndexBuffers().GetBufferData(0);
-						IndicesView32Bits = TArrayView<const uint32>(IndexPtr + FirstIndex, IndicesCount);
-					}
-				}
-			
-				if (!ClothingDataView.Num())
-				{
-					continue;
-				}
+						Result->ClothSections[0].ClothingAsset = UE::Mutable::Private::TPassthroughObjectPtr(ClothingAsset);
 
-				const uint32 ClothResourceId = ClothingResourceView[0];
-				if (ClothResourceId == 0)
-				{
-					continue;
-				}
+						const FClothBufferIndexMapping& SectionClothIndexMapping = ClothVertexBuffer.GetClothIndexMapping()[SectionIndex];
+						
+						const uint32 FirstClothVertex = SectionClothIndexMapping.BaseVertexIndex;
+						
+						uint32 FirstClothVertexNextSection = ClothVertexBuffer.GetNumVertices();
+						for (int32 MappingIndex = SectionIndex + 1; MappingIndex < SectionClothIndexMappings.Num(); ++MappingIndex)
+						{
+							const int32 BaseVertexIndex = SectionClothIndexMappings[MappingIndex].BaseVertexIndex;
+							if (BaseVertexIndex)
+							{
+								FirstClothVertexNextSection = BaseVertexIndex;
+							}
+						}
+						const uint32 NumClothVertex = FirstClothVertexNextSection - FirstClothVertex;
+						const uint32 NumInfluences = SectionClothIndexMapping.LODBiasStride / Section.NumVertices;
 
-				const FClothingMeshData* SectionClothingData = ClothingMeshData.Find(ClothResourceId);
+						uint32 NumClothLOD = NumClothVertex / (NumInfluences * Section.NumVertices);
+						if (NumClothLOD > 1)
+						{
+							NumClothLOD = 1;
+							UE_LOG(LogMutable, Error, TEXT("\"Mappings to Same LOD\" only supported in Ray Tracing \"Cloth LODBias Mode\" option for [%s] Skeletal Mesh"), *SkeletalMesh->GetName());
+						}
+						
+						uint32 NumVertices = SectionClothIndexMapping.LODBiasStride * NumClothLOD;
 
-				if (!SectionClothingData)
+						Result->ClothSections[0].Data.SetNumUninitialized(NumVertices);
+						FMemory::Memcpy(Result->ClothSections[0].Data.GetData(), &ClothVertexBuffer.MappingData(SectionClothIndexMapping.BaseVertexIndex), NumVertices * sizeof(FMeshToMeshVertData));
+					}	
+				}
+
+				// Physics Asset
 				{
-					++NumClothingDataNotFound;
-					continue;
+					// TODO
 				}
+			},
+			ConversionTask);
 
-				check(SectionClothingData->Data.Num());
-
-				SectionsClothData.Add(FSectionClothData {
-					SectionIndex,
-					LODIndex,
-					FirstVertex,
-					IndicesView16Bits,
-					IndicesView32Bits, 
-					ClothingDataView, 
-					MakeArrayView(SectionClothingData->Data.GetData(), SectionClothingData->Data.Num()),
-					TArray<FMeshToMeshVertData>() });
-			}
-		}
+		return FinalConversionTask;
 	}
+	
 
-
-	void CopyMeshToMeshClothData(TArray<FSectionClothData>& SectionsClothData)
+	void ClothVertexBuffers(FSkeletalMeshLODRenderData& LODResource, const UE::Mutable::Private::FMesh& MutableMesh)
 	{
-		MUTABLE_CPUPROFILER_SCOPE(ClothCopyMeshToMeshData)
+		const int32 NumSections = MutableMesh.Surfaces.Num();
+		check(NumSections == MutableMesh.ClothSections.Num());
 		
-		// Copy MeshToMeshData.
-		for (FSectionClothData& SectionWithCloth : SectionsClothData)
+		if (NumSections == 1) // Single section. No need to reconstruct data.
 		{
-			const int32 NumVertices = SectionWithCloth.ClothingDataIndicesView.Num();
-
-			TArray<FMeshToMeshVertData>& ClothMappingData = SectionWithCloth.MappingData;
-			ClothMappingData.SetNum(NumVertices);
-
-			// Copy mesh to mesh data indexed by the index stored per vertex at compile time. 
-			for (int32 VertexIdx = 0; VertexIdx < NumVertices; ++VertexIdx)
+			const uint32 NumClothVertices = MutableMesh.ClothSections[0].Data.Num();
+			if (MutableMesh.ClothSections[0].IsValid() && NumClothVertices)
 			{
-				// Possible Optimization: Gather consecutive indices in ClothingDataView and Memcpy the whole range.
-				// FMeshToMeshVertData and FCustomizableObjectMeshToMeshVertData have the same memory footprint and
-				// bytes in a FCustomizableObjectMeshToMeshVertData form a valid FMeshToMeshVertData (not the other way around).
-
-				static_assert(sizeof(FCustomizableObjectMeshToMeshVertData) == sizeof(FMeshToMeshVertData));
-				static_assert(TIsTrivial<FCustomizableObjectMeshToMeshVertData>::Value);
-				static_assert(TIsTrivial<FMeshToMeshVertData>::Value);
-
-				const int32 VertexDataIndex = SectionWithCloth.ClothingDataIndicesView[VertexIdx];
-				check(VertexDataIndex >= 0);
-
-				const FCustomizableObjectMeshToMeshVertData& SrcData = SectionWithCloth.ClothingDataView[VertexDataIndex];
-
-				FMeshToMeshVertData& DstData = ClothMappingData[VertexIdx];
-				FMemory::Memcpy(&DstData, &SrcData, sizeof(FMeshToMeshVertData));
+				FClothBufferIndexMapping Mapping;
+				Mapping.BaseVertexIndex = 0;
+				Mapping.MappingOffset = 0;
+				Mapping.LODBiasStride = NumClothVertices;
+		
+				TArray<FClothBufferIndexMapping> Mappings;
+				Mappings.Add(Mapping);
+			
+				TArrayView<const FMeshToMeshVertData> DataView = TArrayView(MutableMesh.ClothSections[0].Data.GetData(), MutableMesh.ClothSections[0].Data.Num());
+				LODResource.ClothVertexBuffer.Init(DataView, Mappings);
 			}
 		}
-	}
-
-
-	void CreateClothMapping(const FSectionClothData& SectionClothData, TArray<FMeshToMeshVertData>& MappingData, TArray<FClothBufferIndexMapping>& ClothIndexMapping)
-	{
-		ClothIndexMapping[SectionClothData.SectionIndex] = FClothBufferIndexMapping {
-			static_cast<uint32>(SectionClothData.BaseVertex),
-			static_cast<uint32>(MappingData.Num()),
-			static_cast<uint32>(SectionClothData.MappingData.Num()) };
+		else if (NumSections > 1)
+		{
+			int32 DataSize = 0;
+			for (int32 SectionIndex = 0; SectionIndex < NumSections; ++SectionIndex)
+			{
+				DataSize += MutableMesh.ClothSections[SectionIndex].Data.Num();
+			}
 		
-		MappingData.Append(SectionClothData.MappingData);
-	}
-
-
-	void ClothVertexBuffers(FSkeletalMeshLODRenderData& LODResource, const UE::Mutable::Private::FMesh& MutableMesh, const TMap<uint32, FClothingMeshData>& ClothingMeshData, int32 LODIndex)
-	{
-		TArray<FSectionClothData> SectionsClothData;
-		SectionsClothData.Reserve(32);
-
-		int32 NumClothingDataNotFound = 0;
+			if (DataSize)
+			{
+				TArray<FMeshToMeshVertData> Data;
+				Data.Reserve(DataSize);
 
-		GetSectionClothData(MutableMesh, LODIndex, ClothingMeshData, SectionsClothData, NumClothingDataNotFound);
-	
-		if (NumClothingDataNotFound > 0)
-		{
-			UE_LOG(LogMutable, Error, TEXT("Some clothing data could not be loaded properly, clothing assets may not behave as expected."));
-		}
+				for (const UE::Mutable::Private::FCloth& Cloth : MutableMesh.ClothSections)
+				{
+					Data.Append(Cloth.Data);
+				}
+			
+				TArray<FClothBufferIndexMapping> Mappings;
+				Mappings.SetNum(NumSections);
 
-		if (!SectionsClothData.Num())
-		{
-			return; // Nothing to do.
-		}
+				uint32 VertexOffset = 0;
+				for (int32 SectionIndex = 0; SectionIndex < NumSections; ++SectionIndex)
+				{
+					const UE::Mutable::Private::FCloth& Cloth = MutableMesh.ClothSections[SectionIndex];
+					const uint32 BaseVertexIndex = LODResource.RenderSections[SectionIndex].BaseVertexIndex;
+				
+					FClothBufferIndexMapping& Mapping = Mappings[SectionIndex];
+					if (Cloth.IsValid())
+					{
+						Mapping.BaseVertexIndex = BaseVertexIndex; 
+						Mapping.MappingOffset = VertexOffset;
+						Mapping.LODBiasStride = VertexOffset;
 
-		CopyMeshToMeshClothData(SectionsClothData);
-		
-		TArray<FMeshToMeshVertData> MappingData;
-		MappingData.Reserve(Algo::Accumulate(SectionsClothData, 0, [](int32 Sum, const FSectionClothData& Element){	return Sum + Element.MappingData.Num(); })); // Optimization.
-		
-		TArray<FClothBufferIndexMapping> ClothIndexMapping;
-		ClothIndexMapping.SetNumZeroed(LODResource.RenderSections.Num());
+						VertexOffset += Cloth.Data.Num();
+					}
+					else
+					{
+						Mapping.BaseVertexIndex = 0; 
+						Mapping.MappingOffset = 0;
+						Mapping.LODBiasStride = 0;
+					}
+				}
 
-		for (const FSectionClothData& Data : SectionsClothData)
-		{
-			CreateClothMapping(Data, MappingData, ClothIndexMapping);
+				LODResource.ClothVertexBuffer.Init(Data, Mappings);
+			}
 		}
-		
-		LODResource.ClothVertexBuffer.Init(MappingData, ClothIndexMapping);
 	}
 }
diff --git a/Engine/Plugins/Mutable/Source/CustomizableObject/Public/MuCO/CustomizableObject.h b/Engine/Plugins/Mutable/Source/CustomizableObject/Public/MuCO/CustomizableObject.h
--- a/Engine/Plugins/Mutable/Source/CustomizableObject/Public/MuCO/CustomizableObject.h
+++ b/Engine/Plugins/Mutable/Source/CustomizableObject/Public/MuCO/CustomizableObject.h
@@ -280,9 +280,6 @@ private:
 	UPROPERTY(EditAnywhere, Category = CompileOptions)
 	bool bEnableRealTimeMorphTargets = false;
 
-	UPROPERTY(EditAnywhere, Category = CompileOptions)
-	bool bEnableClothing = false;
-
 	UPROPERTY(EditAnywhere, Category = CompileOptions)
 	bool bEnable16BitBoneWeights = false;
 
diff --git a/Engine/Plugins/Mutable/Source/CustomizableObjectEditor/Private/MuCOE/CustomizableObjectCompiler.cpp b/Engine/Plugins/Mutable/Source/CustomizableObjectEditor/Private/MuCOE/CustomizableObjectCompiler.cpp
--- a/Engine/Plugins/Mutable/Source/CustomizableObjectEditor/Private/MuCOE/CustomizableObjectCompiler.cpp
+++ b/Engine/Plugins/Mutable/Source/CustomizableObjectEditor/Private/MuCOE/CustomizableObjectCompiler.cpp
@@ -1479,14 +1479,6 @@ void FCustomizableObjectCompiler::FinishLoadingModelDataFromDDC()
 							}
 							break;
 						}
-						case UE::Mutable::Private::EStreamableDataType::Clothing:
-						{
-							for (const UE::Mutable::Private::FBlock& Block : File->Blocks)
-							{
-								PlatformData->ClothingStreamableData.Set(Block.Id, TempData.GetData() + Block.Offset, Block.Size);
-							}
-							break;
-						}
 
 						default:
 							unimplemented();
@@ -1569,7 +1561,6 @@ UE::Mutable::Private::Ptr<UE::Mutable::Private::Node> FCustomizableObjectCompile
 
 	FCompilationOptions CompilerOptions = InCompilerOptions;
 	CompilerOptions.bRealTimeMorphTargetsEnabled = Object->GetPrivate()->IsRealTimeMorphTargetsEnabled();
-	CompilerOptions.bClothingEnabled = Object->GetPrivate()->IsClothingEnabled();
 	CompilerOptions.b16BitBoneWeightsEnabled = Object->GetPrivate()->Is16BitBoneWeightsEnabled();
 	CompilerOptions.bSkinWeightProfilesEnabled = Object->GetPrivate()->IsAltSkinWeightProfilesEnabled();
 	CompilerOptions.bPhysicsAssetMergeEnabled = Object->GetPrivate()->IsPhysicsAssetMergeEnabled();
@@ -1699,77 +1690,7 @@ void FCustomizableObjectCompiler::FinishCompilationTask()
 
 			MeshData.Value.Data.Empty();
 		}
-
-		ModelStreamableBulkData->ClothingStreamables.Empty(32);
-
-		uint64 ClothingDataOffsetInBytes = 0;
-		for (TPair<uint32, FClothingMeshDataSource>& MeshData : CompilationContext->ClothingPerMeshData)
-		{
-			const uint32 DataSizeInBytes = (uint32)MeshData.Value.Data.Num() * sizeof(FCustomizableObjectMeshToMeshVertData);
-			FClothingStreamable& ResourceMeshData = ModelStreamableBulkData->ClothingStreamables.FindOrAdd(MeshData.Key);
-
-			check(ResourceMeshData.ClothingAssetIndex == INDEX_NONE);
-			check(ResourceMeshData.ClothingAssetLOD == INDEX_NONE);
-			check(ResourceMeshData.Size == 0);
-
-			ResourceMeshData.ClothingAssetIndex = MeshData.Value.ClothingAssetIndex;
-			ResourceMeshData.ClothingAssetLOD = MeshData.Value.ClothingAssetLOD;
-			ResourceMeshData.PhysicsAssetIndex = MeshData.Value.PhysicsAssetIndex;
-			ResourceMeshData.Size = DataSizeInBytes;
-			EMutableFileFlags Flags = EMutableFileFlags::None;
-			ResourceMeshData.Block = FMutableStreamableBlock{ .FileId = uint32(0), .Flags = uint16(Flags), .Offset = ClothingDataOffsetInBytes };
-			ResourceMeshData.SourceId = MeshData.Value.SourceId;
-
-			ClothingDataOffsetInBytes += DataSizeInBytes;
-
-			PlatformData->ClothingStreamableData.Set(
-				MeshData.Key,
-				reinterpret_cast<uint8*>(MeshData.Value.Data.GetData()),
-				DataSizeInBytes);
-
-			MeshData.Value.Data.Empty();
-		}
-
-		ModelResources->ClothingAssetsData = MoveTemp(CompilationContext->ClothingAssetsData);
-
-		// A clothing backend, e.g. Chaos cloth, can use 2 config files, one owned by the asset, and another that is shared 
-		// among all assets in a SkeletalMesh. When merging different assets in a skeletalmesh we need to make sure only one of 
-		// the shared is used. In that case we will keep the first visited of a type and will be stored separated from the asset.
-		// TODO: Shared configs, which typically controls the quality of the simulation (iterations, etc), probably should be specified 
-		// somewhere else to give more control with which config ends up used. 
-		auto IsSharedConfigData = [](const FCustomizableObjectClothConfigData& ConfigData) -> bool
-			{
-				const UClass* ConfigClass = FindObject<UClass>(nullptr, *ConfigData.ClassPath);
-				return ConfigClass ? static_cast<bool>(Cast<UClothSharedConfigCommon>(ConfigClass->GetDefaultObject())) : false;
-			};
-
-		// Find shared configs to be used (One of each type) 
-		for (FCustomizableObjectClothingAssetData& ClothingAssetData : ModelResources->ClothingAssetsData)
-		{
-			for (FCustomizableObjectClothConfigData& ClothConfigData : ClothingAssetData.ConfigsData)
-			{
-				if (IsSharedConfigData(ClothConfigData))
-				{
-					FCustomizableObjectClothConfigData* FoundConfig = ModelResources->ClothSharedConfigsData.FindByPredicate(
-						[Name = ClothConfigData.ConfigName](const FCustomizableObjectClothConfigData& Other)
-						{
-							return Name == Other.ConfigName;
-						});
-
-					if (!FoundConfig)
-					{
-						ModelResources->ClothSharedConfigsData.AddDefaulted_GetRef() = ClothConfigData;
-					}
-				}
-			}
-		}
-
-		// Remove shared configs
-		for (FCustomizableObjectClothingAssetData& ClothingAssetData : ModelResources->ClothingAssetsData)
-		{
-			ClothingAssetData.ConfigsData.RemoveAllSwap(IsSharedConfigData);
-		}
-
+		
 		ModelResources->SurfaceMetadata = MoveTemp(CompilationContext->SurfaceMetadata);
 		ModelResources->MeshMetadata = MoveTemp(CompilationContext->MeshMetadata);
 	}
diff --git a/Engine/Plugins/Mutable/Source/CustomizableObjectEditor/Private/MuCOE/GenerateMutableSource/GenerateMutableSource.h b/Engine/Plugins/Mutable/Source/CustomizableObjectEditor/Private/MuCOE/GenerateMutableSource/GenerateMutableSource.h
--- a/Engine/Plugins/Mutable/Source/CustomizableObjectEditor/Private/MuCOE/GenerateMutableSource/GenerateMutableSource.h
+++ b/Engine/Plugins/Mutable/Source/CustomizableObjectEditor/Private/MuCOE/GenerateMutableSource/GenerateMutableSource.h
@@ -46,7 +46,6 @@ class UCustomizableObjectNodeGroupProjectorParameter;
 class UCONodeSkeletalMeshSectionBase;
 class UCustomizableObjectNodeSwitchBase;
 class UCustomizableObjectNodeModifierClipWithMesh;
-struct FCustomizableObjectClothingAssetData;
 
 class FCustomizableObjectCompiler;
 class UTextureLODSettings;
@@ -382,20 +381,6 @@ struct FRealTimeMorphMeshData
 };
 
 
-/** See FClothingMeshData. */
-struct FClothingMeshDataSource
-{
-	int32 ClothingAssetIndex = INDEX_NONE;
-	int32 ClothingAssetLOD = INDEX_NONE;
-	int32 PhysicsAssetIndex = INDEX_NONE;
-	TArray<FCustomizableObjectMeshToMeshVertData> Data;
-
-	/* Used to group data when generating bulk data files. 
-	 * This property should not be taken into consideration when comparing structs.
-	 */
-	uint32 SourceId = 0;
-};
-
 struct FGroupProjectorTempData
 {
 	class UCustomizableObjectNodeGroupProjectorParameter* CustomizableObjectNodeGroupProjectorParameter;
@@ -547,11 +532,7 @@ struct FMutableCompilationContext
 
 	/** Data used for MorphTarget reconstruction. */
 	TMap<uint32, FRealTimeMorphMeshData> RealTimeMorphTargetPerMeshData;
-
-	/** Data used for Clothing reconstruction. */
-	TArray<FCustomizableObjectClothingAssetData> ClothingAssetsData;
-	TMap<uint32, FClothingMeshDataSource> ClothingPerMeshData;
-
+	
 	/** Stores the physics assets gathered from the SkeletalMesh nodes during compilation, to be used in mesh generation in-game */
 	TArray<TSoftObjectPtr<UPhysicsAsset>> PhysicsAssets;
 
diff --git a/Engine/Plugins/Mutable/Source/CustomizableObjectEditor/Private/MuCOE/GenerateMutableSource/GenerateMutableSourceMesh.cpp b/Engine/Plugins/Mutable/Source/CustomizableObjectEditor/Private/MuCOE/GenerateMutableSource/GenerateMutableSourceMesh.cpp
--- a/Engine/Plugins/Mutable/Source/CustomizableObjectEditor/Private/MuCOE/GenerateMutableSource/GenerateMutableSourceMesh.cpp
+++ b/Engine/Plugins/Mutable/Source/CustomizableObjectEditor/Private/MuCOE/GenerateMutableSource/GenerateMutableSourceMesh.cpp
@@ -135,7 +135,7 @@ void SetBuffersFormat(FMutableCompilationContext& Context,
 	const bool bHasSkinWeights = !bIgnoreSkinning && MutableBonesPerVertex > 0;
 	const bool bHasVertexColors = SourceMesh.GetHasVertexColors();
 	const bool bHasRealTimeMorphs = !bIgnoreMorphs && Context.Options.bRealTimeMorphTargetsEnabled && SourceMesh.GetMorphTargets().Num();
-	const bool bHasClothing = !bIgnoreSkinning && Context.Options.bClothingEnabled;
+	const bool bHasClothing = !bIgnoreSkinning && SourceMesh.HasActiveClothingAssetsForLOD(LODIndex);
 	const bool bHasSkinWeightProfiles = !bIgnoreSkinning && Context.Options.bSkinWeightProfilesEnabled && LODModel.SkinWeightProfiles.Num();
 
 	int32 MutableBufferCount = 2; // Position & Tangents
@@ -1823,317 +1823,6 @@ UE::Tasks::FTask ConvertSkeletalMeshToMutable(TSharedRef<FMeshConversionContext>
 		}
 	}
 
-	// Clothing
-	uint32 ClothingUniqueHash = 0;
-	if (!bIgnoreSkinning && Context.Options.bClothingEnabled)
-	{
-		// Create new asset or find an already created one if the section has clothing assets.
-		// clothing assets are shared among all LODs in a section
-		int32 ClothingAssetIndex = INDEX_NONE;
-		int32 PhysicsAssetIndex = INDEX_NONE;
-
-		Invoke([&](int32& OutClothingAssetIndex, int32& OutPhysicsAssetIndex) -> void
-			{
-				const UClothingAssetBase* ClothingAssetBase = SkeletalMesh->GetSectionClothingAsset(LODIndex, SectionIndex);
-
-				if (!ClothingAssetBase)
-				{
-					return;
-				}
-
-				const UClothingAssetCommon* Asset = Cast<UClothingAssetCommon>(ClothingAssetBase);
-				if (!Asset)
-				{
-					return;
-				}
-
-				OutPhysicsAssetIndex = Context.PhysicsAssets.AddUnique(Asset->PhysicsAsset);
-
-				int32 FoundIndex = Context.ClothingAssetsData.IndexOfByPredicate(
-					[SourceAsset = Asset](const FCustomizableObjectClothingAssetData& Asset)
-					{
-						bool bIsSameAsset =
-							Asset.OriginalAssetGuid == SourceAsset->GetAssetGuid() &&
-							Asset.Name == SourceAsset->GetFName() &&
-							Asset.ReferenceBoneIndex == SourceAsset->ReferenceBoneIndex &&
-							Asset.UsedBoneIndices == SourceAsset->UsedBoneIndices &&
-							Asset.UsedBoneNames == SourceAsset->UsedBoneNames &&
-							Asset.LodMap == SourceAsset->LodMap &&
-							Asset.LodData.Num() == SourceAsset->LodData.Num();
-
-						const int32 LodDataNum = Asset.LodData.Num();
-						for (int32 LodDataIndex = LodDataNum - 1; LodDataIndex >= 0 && bIsSameAsset; --LodDataIndex)
-						{
-							const FClothPhysicalMeshData& DataA = Asset.LodData[LodDataIndex].PhysicalMeshData;
-							const FClothPhysicalMeshData& DataB = SourceAsset->LodData[LodDataIndex].PhysicalMeshData;
-
-							bIsSameAsset = bIsSameAsset && DataA.NumFixedVerts == DataB.NumFixedVerts;
-							bIsSameAsset = bIsSameAsset && DataA.MaxBoneWeights == DataB.MaxBoneWeights;
-							bIsSameAsset = bIsSameAsset && DataA.Vertices == DataB.Vertices;
-							bIsSameAsset = bIsSameAsset && DataA.Normals == DataB.Normals;
-							bIsSameAsset = bIsSameAsset && DataA.Indices == DataB.Indices;
-							bIsSameAsset = bIsSameAsset && DataA.InverseMasses == DataB.InverseMasses;
-							bIsSameAsset = bIsSameAsset && DataA.EuclideanTethers.Tethers == DataB.EuclideanTethers.Tethers;
-							bIsSameAsset = bIsSameAsset && DataA.GeodesicTethers.Tethers == DataB.GeodesicTethers.Tethers;
-
-							bIsSameAsset = bIsSameAsset && DataA.WeightMaps.Num() == DataB.WeightMaps.Num();
-							bIsSameAsset = bIsSameAsset && DataA.SelfCollisionVertexSet.Num() == DataB.SelfCollisionVertexSet.Num();
-							bIsSameAsset = bIsSameAsset && DataA.BoneData.Num() == DataB.BoneData.Num();
-
-							static_assert(std::is_trivially_copyable_v<FClothVertBoneData>);
-							// Assume the FClothVertBoneData does not have any padding. In case there was padding, same assets should
-							// have the same unset memory so false negatives can only happen with different assets that have the 
-							// same data. This reasoning relays on the fact that the data buffers have been copied byte for byte using 
-							// Memcpy or similar.
-							bIsSameAsset = bIsSameAsset && FMemory::Memcmp(
-								DataA.BoneData.GetData(), DataB.BoneData.GetData(), DataA.BoneData.Num() * sizeof(FClothVertBoneData)) == 0;
-
-							if (bIsSameAsset)
-							{
-								for (const TPair<uint32, FPointWeightMap>& WeightMap : DataA.WeightMaps)
-								{
-									const FPointWeightMap* FoundWeightMap = DataB.WeightMaps.Find(WeightMap.Key);
-
-									if (!FoundWeightMap)
-									{
-										bIsSameAsset = false;
-										break;
-									}
-
-									if (FoundWeightMap->Values != WeightMap.Value.Values)
-									{
-										bIsSameAsset = false;
-										break;
-									}
-								}
-							}
-
-							if (bIsSameAsset)
-							{
-								TArray<int32> DataASelfCollisionVertexArray = DataA.SelfCollisionVertexSet.Array();
-								Algo::Sort(DataASelfCollisionVertexArray);
-
-								TArray<int32> DataBSelfCollisionVertexArray = DataB.SelfCollisionVertexSet.Array();
-								Algo::Sort(DataBSelfCollisionVertexArray);
-
-								bIsSameAsset = DataASelfCollisionVertexArray == DataBSelfCollisionVertexArray;
-							}
-						}
-
-						return bIsSameAsset;
-					});
-
-				if (FoundIndex != INDEX_NONE)
-				{
-					OutClothingAssetIndex = FoundIndex;
-					return;
-				}
-
-				OutClothingAssetIndex = Context.ClothingAssetsData.AddDefaulted();
-				FCustomizableObjectClothingAssetData& AssetData = Context.ClothingAssetsData[OutClothingAssetIndex];
-
-				AssetData.LodData = Asset->LodData;
-				AssetData.LodMap = Asset->LodMap;
-				AssetData.ReferenceBoneIndex = Asset->ReferenceBoneIndex;
-				AssetData.UsedBoneIndices = Asset->UsedBoneIndices;
-				AssetData.UsedBoneNames = Asset->UsedBoneNames;
-				AssetData.OriginalAssetGuid = Asset->GetAssetGuid();
-				AssetData.Name = Asset->GetFName();
-
-				// Store raw clothing config serialized raw data, and info to recreate it afterwards.
-				for (const TPair<FName, TObjectPtr<UClothConfigBase>>& ClothConfig : Asset->ClothConfigs)
-				{
-					FCustomizableObjectClothConfigData& ConfigData = AssetData.ConfigsData.AddDefaulted_GetRef();
-					ConfigData.ClassPath = ClothConfig.Value->GetClass()->GetPathName();
-					ConfigData.ConfigName = ClothConfig.Key;
-
-					FMemoryWriter MemoryWriter(ConfigData.ConfigBytes);
-					ClothConfig.Value->Serialize(MemoryWriter);
-				}
-			}, ClothingAssetIndex, PhysicsAssetIndex);
-
-		if (ClothingAssetIndex != INDEX_NONE && MeshSection.ClothMappingDataLODs.IsEmpty())
-		{
-			UE_LOG(LogMutable, Display, TEXT("SkeletalMesh [%s] is missing ClothMappingDataLODs. "), *GetNameSafe(SkeletalMesh));
-		}
-
-		if (ClothingAssetIndex != INDEX_NONE && !MeshSection.ClothMappingDataLODs.IsEmpty())
-		{
-			using namespace UE::Mutable::Private;
-			check(MutableMesh->GetVertexBuffers().GetBufferCount() >= NextBufferIndex + 2);
-
-			// Clothing buffers already initialized by the SetBuffer
-			check(MutableMesh->GetVertexBuffers().Buffers[NextBufferIndex].Channels.Num() == 1);
-			check(MutableMesh->GetVertexBuffers().Buffers[NextBufferIndex].Channels[0].Semantic == EMeshBufferSemantic::Other);
-			check(MutableMesh->GetVertexBuffers().Buffers[NextBufferIndex].Channels[0].SemanticIndex == 2);
-			check(MutableMesh->GetVertexBuffers().Buffers[NextBufferIndex].Channels[0].Format == EMeshBufferFormat::Int32);
-
-			check(MutableMesh->GetVertexBuffers().Buffers[NextBufferIndex + 1].Channels.Num() == 1);
-			check(MutableMesh->GetVertexBuffers().Buffers[NextBufferIndex + 1].Channels[0].Semantic == EMeshBufferSemantic::Other);
-			check(MutableMesh->GetVertexBuffers().Buffers[NextBufferIndex + 1].Channels[0].SemanticIndex == 3);
-			check(MutableMesh->GetVertexBuffers().Buffers[NextBufferIndex + 1].Channels[0].Format == EMeshBufferFormat::UInt32);
-
-			TArrayView<int32> ClothSectionIndexView(reinterpret_cast<int32*>(
-				MutableMesh->GetVertexBuffers().GetBufferData(NextBufferIndex)), VertexCount);
-			for (int32& Elem : ClothSectionIndexView)
-			{
-				Elem = -1;
-			}
-
-			constexpr uint32 InvalidResourceId = 0;
-			TArrayView<uint32> ClothSectionResourceIdView(reinterpret_cast<uint32*>(
-				MutableMesh->GetVertexBuffers().GetBufferData(NextBufferIndex + 1)), VertexCount);
-			for (uint32& Elem : ClothSectionResourceIdView)
-			{
-				Elem = InvalidResourceId;
-			}
-
-			const TArray<FMeshToMeshVertData>& ClothMappingData = MeshSection.ClothMappingDataLODs[0];
-
-			// Similar test as the one used on FSkeletalMeshObjectGPUSkin::FVertexFactoryData::InitAPEXClothVertexFactories
-			// Here should work as expexted, but in the reference code I'm not sure it always works. It is worth investigate
-			// in that direction if at some point multiple influences don't work as expected.
-			const bool bUseMutlipleInfluences = ClothMappingData.Num() > MeshSection.NumVertices;
-
-			// Constant defined in ClothMeshUtils.cpp with the following comment:
-			// // This must match NUM_INFLUENCES_PER_VERTEX in GpuSkinCacheComputeShader.usf and GpuSkinVertexFactory.ush
-			// // TODO: Make this easier to change in without messing things up
-			// TODO: find a better place to keep this constant.
-			constexpr int32 NumInfluencesPerVertex = 5;
-
-			int32 MeshToMeshDataIndex = 0;
-
-			constexpr int32 MaxSupportedInfluences = 1;
-			for (int32& Elem : ClothSectionIndexView)
-			{
-				Elem = MeshToMeshDataIndex;
-				MeshToMeshDataIndex += MaxSupportedInfluences;
-			}
-
-			const int32 ClothDataIndexBase = 0;
-
-			const int32 ClothDataStride = bUseMutlipleInfluences ? NumInfluencesPerVertex : 1;
-			const int32 NumClothMappingDataVerts = ClothMappingData.Num() / ClothDataStride;
-
-			FClothingMeshDataSource ClothingMeshData;
-
-			const FCustomizableObjectClothingAssetData& ClothingAssetData = Context.ClothingAssetsData[ClothingAssetIndex];
-
-			ClothingMeshData.ClothingAssetIndex = ClothingAssetIndex;
-			ClothingMeshData.ClothingAssetLOD = ClothingAssetData.LodMap[LODIndex];
-			ClothingMeshData.PhysicsAssetIndex = PhysicsAssetIndex;
-			ClothingMeshData.Data.Reserve(NumClothMappingDataVerts * ClothDataStride);
-
-			for (int32 Idx = 0; Idx < NumClothMappingDataVerts * ClothDataStride; Idx += ClothDataStride)
-			{
-				// If bUseMutlipleInfluences we will only take the element with higher weight ignoring the other ones.
-				TArrayView<const FMeshToMeshVertData> Influences(ClothMappingData.GetData() + Idx, ClothDataStride);
-				const FMeshToMeshVertData* MaxInfluence = MaxElement(Influences.begin(), Influences.end(),
-					[](const FMeshToMeshVertData& A, const FMeshToMeshVertData& B) { return A.Weight < B.Weight; });
-
-				ClothingMeshData.Data.Emplace(*MaxInfluence);
-			}
-
-			for (FCustomizableObjectMeshToMeshVertData& ClothingDataElem : ClothingMeshData.Data)
-			{
-				// Currently if the cloth mapping uses multiple influences, these are ignored and only 
-				// the one with the highest weight is used. We set the weight to 1.0, but
-				// this value will be ignored anyway.
-				ClothingDataElem.Weight = 1.0f;
-			}
-
-			const uint32 DataHash = CityHash32(
-				reinterpret_cast<const char*>(ClothingMeshData.Data.GetData()),
-				ClothingMeshData.Data.Num() * sizeof(FCustomizableObjectMeshToMeshVertData));
-
-			ClothingUniqueHash = DataHash == InvalidResourceId ? DataHash + 1 : DataHash;
-
-			const FClothingMeshDataSource* FoundHash = Context.ClothingPerMeshData.Find(ClothingUniqueHash);
-
-			bool bIsDataAlreadyCollected = false;
-
-			auto CompareClothingData = [](const FClothingMeshDataSource& A, const FClothingMeshDataSource& B) -> bool
-				{
-					return A.ClothingAssetIndex == B.ClothingAssetIndex &&
-						A.ClothingAssetLOD == B.ClothingAssetLOD &&
-						A.PhysicsAssetIndex == B.PhysicsAssetIndex &&
-						A.Data.Num() == B.Data.Num() &&
-						FMemory::Memcmp(A.Data.GetData(), B.Data.GetData(), A.Data.Num() * sizeof(FCustomizableObjectMeshToMeshVertData)) == 0;
-				};
-
-			if (FoundHash)
-			{
-				bIsDataAlreadyCollected = CompareClothingData(*FoundHash, ClothingMeshData);
-			}
-
-			// NOTE: This way of unique hash generation guarantees all valid values can be used but given its 
-			// sequential nature a cascade of changes can occur if new meshes are added. Not many hash collisions 
-			// are expected so it should not be problematic.
-			if (FoundHash && !bIsDataAlreadyCollected)
-			{
-				uint32 NumTries = 0;
-				for (; NumTries < TNumericLimits<uint32>::Max(); ++NumTries)
-				{
-					FoundHash = Context.ClothingPerMeshData.Find(ClothingUniqueHash);
-
-					if (!FoundHash)
-					{
-						break;
-					}
-
-					bIsDataAlreadyCollected = CompareClothingData(*FoundHash, ClothingMeshData);
-
-					if (bIsDataAlreadyCollected)
-					{
-						break;
-					}
-
-					ClothingUniqueHash = ClothingUniqueHash + 1 == InvalidResourceId ? InvalidResourceId + 1 : ClothingUniqueHash + 1;
-				}
-
-				if (NumTries == TNumericLimits<uint32>::Max())
-				{
-					UE_LOG(LogMutable, Warning, TEXT("Maximum number of meshes with clothing reached, some cloth meshes may not work as expected."));
-					ClothingUniqueHash = InvalidResourceId;
-				}
-			}
-
-			// fill the resource buffer with the generated unique id.
-			for (uint32& Elem : ClothSectionResourceIdView)
-			{
-				Elem = ClothingUniqueHash;
-			}
-
-			FCustomizableObjectStreameableResourceId StreamedClothResource;
-			StreamedClothResource.Id = ClothingUniqueHash;
-			StreamedClothResource.Type = static_cast<uint8>(FCustomizableObjectStreameableResourceId::EType::Clothing);
-
-			MutableMesh->AddStreamedResource(BitCast<uint64>(StreamedClothResource));
-
-			if (!bIsDataAlreadyCollected)
-			{
-				FClothingMeshDataSource& NewClothingMeshData = Context.ClothingPerMeshData.FindOrAdd(ClothingUniqueHash);
-
-				check(NewClothingMeshData.ClothingAssetIndex == INDEX_NONE);
-				check(NewClothingMeshData.Data.IsEmpty());
-
-				NewClothingMeshData.ClothingAssetIndex = ClothingMeshData.ClothingAssetIndex;
-				NewClothingMeshData.ClothingAssetLOD = ClothingMeshData.ClothingAssetLOD;
-				NewClothingMeshData.PhysicsAssetIndex = ClothingMeshData.PhysicsAssetIndex;
-
-				NewClothingMeshData.Data = MoveTemp(ClothingMeshData.Data);
-				NewClothingMeshData.SourceId = MeshId;
-			}
-
-			NextBufferIndex += 2;
-		}
-		else
-		{
-			MutableMesh->VertexBuffers.RemoveBuffer(NextBufferIndex + 1);
-			MutableMesh->VertexBuffers.RemoveBuffer(NextBufferIndex);
-		}
-	}
-
 	if (!bIgnoreSkinning)
 	{
 		if (MeshSection.NumVertices > 0)
@@ -2543,10 +2232,10 @@ UE::Tasks::FTask ConvertSkeletalMeshToMutable(TSharedRef<FMeshConversionContext>
 			MutableMesh->AddAdditionalPhysicsBody(MutableBody);
 		}
 	}
-
-
+	
 	// Ensure Surface Data
 	UE::Mutable::Private::FMeshSurface& MeshSurface = MutableMesh->Surfaces.Emplace_GetRef();
+	MutableMesh->ClothSections.AddDefaulted();
 	MeshSurface.BoneMapCount = MutableMesh->BoneMap.Num();
 
 	MeshSurface.SubMeshes.Emplace(
@@ -2563,7 +2252,6 @@ UE::Tasks::FTask ConvertSkeletalMeshToMutable(TSharedRef<FMeshConversionContext>
 		FMemory::Memzero(MeshMetadata);
 
 		MeshMetadata.MorphMetadataId = MorphUniqueHash;
-		MeshMetadata.ClothingMetadataId = ClothingUniqueHash;
 		MeshMetadata.SurfaceMetadataId = AddUniqueSurfaceMetadata(Source.Metadata, Context);
 
 		auto HashMeshMetadataFunc = [](const FMutableMeshMetadata& Data) -> uint32
@@ -2590,7 +2278,28 @@ UE::Tasks::FTask ConvertSkeletalMeshToMutable(TSharedRef<FMeshConversionContext>
 
 
 		MeshSurface.SubMeshes.Last().ExternalId = MeshMetadataUniqueHash;
+	}
+	
+	// Clothing
+	if (!bIgnoreSkinning)
+	{
+		const bool bHasClothData = MeshSection.ClothMappingDataLODs.IsValidIndex(0) && MeshSection.ClothMappingDataLODs[0].Num();
+					
+		UClothingAssetCommon* ClothingAsset = Cast<UClothingAssetCommon>(SkeletalMesh->GetClothingAsset(Section.ClothingData.AssetGuid));
+		if (bHasClothData && ClothingAsset)
+		{
+			MutableMesh->ClothSections[0].ClothingAsset = Context.PassthroughObjectFactory.Add(*ClothingAsset);
 
+			if (MeshSection.ClothMappingDataLODs.Num() > 1)
+			{
+				UE_LOG(LogMutable, Error, TEXT("\"Mappings to Same LOD\" only supported in Ray Tracing \"Cloth LODBias Mode\" option for [%s] Skeletal Mesh"), *SkeletalMesh->GetName());	
+			}
+			
+			const uint32 NumClothVertices = MeshSection.ClothMappingDataLODs[0].Num();
+			
+			MutableMesh->ClothSections[0].Data.SetNumUninitialized(NumClothVertices);
+			FMemory::Memcpy(MutableMesh->ClothSections[0].Data.GetData(), &MeshSection.ClothMappingDataLODs[0][0], NumClothVertices * sizeof(FMeshToMeshVertData));
+		}
 	}
 
 	MeshConversionContext->Result = MutableMesh;
diff --git a/Engine/Plugins/Mutable/Source/CustomizableObjectEditor/Private/MuCOE/GenerateMutableSource/GenerateMutableSourceTable.cpp b/Engine/Plugins/Mutable/Source/CustomizableObjectEditor/Private/MuCOE/GenerateMutableSource/GenerateMutableSourceTable.cpp
--- a/Engine/Plugins/Mutable/Source/CustomizableObjectEditor/Private/MuCOE/GenerateMutableSource/GenerateMutableSourceTable.cpp
+++ b/Engine/Plugins/Mutable/Source/CustomizableObjectEditor/Private/MuCOE/GenerateMutableSource/GenerateMutableSourceTable.cpp
@@ -24,7 +24,6 @@
 #include "Rendering/SkeletalMeshLODModel.h"
 #include "Rendering/SkeletalMeshModel.h"
 #include "MuCOE/CustomizableObjectVersionBridge.h"
-#include "ClothingAsset.h"
 
 #define LOCTEXT_NAMESPACE "CustomizableObjectEditor"
 
diff --git a/Engine/Plugins/Mutable/Source/MutableRuntime/Internal/MuR/Cloth.h b/Engine/Plugins/Mutable/Source/MutableRuntime/Internal/MuR/Cloth.h
new file mode 100644
--- /dev/null
+++ b/Engine/Plugins/Mutable/Source/MutableRuntime/Internal/MuR/Cloth.h
@@ -0,0 +1,43 @@
+// Copyright Epic Games, Inc. All Rights Reserved.
+
+#pragma once
+
+#include "MuR/RefCounted.h"
+
+#include "PassthroughObject.h"
+#include "ClothingAsset.h"
+#include "MemoryTrackingAllocationPolicy.h"
+#include "MeshBufferSet.h"
+#include "Rendering/SkeletalMeshVertexClothBuffer.h"
+
+#define UE_API MUTABLERUNTIME_API
+
+
+namespace UE::Mutable::Private
+{
+	class FCloth : public FResource
+	{
+	public:
+		FCloth() = default;
+
+		FCloth(const FCloth& Other) = default;
+		
+		FCloth& operator=(const FCloth& Other) = default;
+
+		UE_API bool operator==(const FCloth& Other) const;
+		
+		UE_API virtual int32 GetDataSize() const override;
+		
+		UE_API void Serialise(FOutputArchive& Ar) const;
+		UE_API void Unserialise(FInputArchive& Ar);
+		
+		UE_API bool IsValid() const;
+		
+		TArray<FMeshToMeshVertData, FDefaultMemoryTrackingAllocator<MemoryCounters::FMeshMemoryCounter>> Data;
+		
+		TPassthroughObjectPtr<UClothingAssetCommon> ClothingAsset;
+	};
+}
+
+
+#undef UE_API
diff --git a/Engine/Plugins/Mutable/Source/MutableRuntime/Internal/MuR/Mesh.h b/Engine/Plugins/Mutable/Source/MutableRuntime/Internal/MuR/Mesh.h
--- a/Engine/Plugins/Mutable/Source/MutableRuntime/Internal/MuR/Mesh.h
+++ b/Engine/Plugins/Mutable/Source/MutableRuntime/Internal/MuR/Mesh.h
@@ -17,14 +17,14 @@
 #include "MuR/Serialisation.h"
 #include "MuR/Skeleton.h"
 #include "MuR/PassthroughObject.h"
-#include "Templates/Tuple.h"
-
 #include "MuR/ManagedPointer.h"
-
+#include "MuR/Cloth.h"
+#include "Templates/Tuple.h"
 #include "UObject/StrongObjectPtr.h"
 
 #include <type_traits>
 
+
 #define UE_API MUTABLERUNTIME_API
 
 class FString;
@@ -192,6 +192,7 @@ namespace UE::Mutable::Private
 		WithAdditionalPhysics = 1 << 15,
 		WithStreamedResources = 1 << 16,
 		WithMorphData = 1 << 17,
+		WithClothData = 1 << 18,
 
 		AllFlags = 0xFFFFFFFF
 	};
@@ -424,9 +425,9 @@ namespace UE::Mutable::Private
         UE_API uint32 GetId() const;
 
     	UE_API bool HasMorphs() const;
-    
-	public:
 
+    	UE_API bool HasClothing() const;
+    	
 		template<typename Type>
 		using TMemoryTrackedArray = TArray<Type, FDefaultMemoryTrackingAllocator<MemoryCounters::FMeshMemoryCounter>>;
 
@@ -467,6 +468,9 @@ namespace UE::Mutable::Private
 
     	FMeshMorph Morph;
     	
+    	/** Only valid if the mesh is a LOD (has multiple sections). */
+    	TArray<FCloth> ClothSections;
+    	
 		/** Externally provided SkeletonIDs of the skeletons required by this mesh. */
 		TArray<uint32> SkeletonIDs;
 
@@ -596,6 +600,8 @@ namespace UE::Mutable::Private
 				bEqual &= *AdditionalPhysicsBodies[i] == *Other.AdditionalPhysicsBodies[i];
 			}
 
+			if (bEqual) bEqual = ClothSections == Other.ClothSections;
+			
 			return bEqual;
 		}
 
diff --git a/Engine/Plugins/Mutable/Source/MutableRuntime/Internal/MuR/Model.h b/Engine/Plugins/Mutable/Source/MutableRuntime/Internal/MuR/Model.h
--- a/Engine/Plugins/Mutable/Source/MutableRuntime/Internal/MuR/Model.h
+++ b/Engine/Plugins/Mutable/Source/MutableRuntime/Internal/MuR/Model.h
@@ -580,7 +580,8 @@ namespace UE::Mutable::Private
 		template<class CreateMeshFunc>
         void GetConstant(
 				uint32 MeshConstantIndex, 
-				int32 SkeletonConstantIndex, 
+				int32 SkeletonConstantIndex,
+				const TPassthroughObjectPtr<UClothingAssetCommon>& ClothAsset,
 				TManagedPtr<const FMesh>& OutMesh, 
 				EMeshContentFlags FilterContentFlags, 
 				CreateMeshFunc&& CreateMesh) const
@@ -779,6 +780,11 @@ namespace UE::Mutable::Private
 				Result->ConvertBoneIdsToIndices();
 			}
 
+			if (Result->ClothSections.Num())
+			{
+				Result->ClothSections[0].ClothingAsset = ClothAsset;	
+			}
+			
 			OutMesh = Result;
 		}
 
diff --git a/Engine/Plugins/Mutable/Source/MutableRuntime/Internal/MuR/OpMeshMerge.h b/Engine/Plugins/Mutable/Source/MutableRuntime/Internal/MuR/OpMeshMerge.h
--- a/Engine/Plugins/Mutable/Source/MutableRuntime/Internal/MuR/OpMeshMerge.h
+++ b/Engine/Plugins/Mutable/Source/MutableRuntime/Internal/MuR/OpMeshMerge.h
@@ -7,21 +7,23 @@
 #include "MuR/MutableMath.h"
 #include "MuR/OpMeshFormat.h"
 #include "MuR/MutableTrace.h"
+#include "MuR/MutableRuntimeModule.h"
 
 #include "Engine/SkeletalMesh.h"
 #include "Rendering/MorphTargetVertexCodec.h"
+#include "Rendering/SkeletalMeshVertexBuffer.h"
 
 #include "UObject/StrongObjectPtr.h"
 
 namespace UE::Mutable::Private
 {
-
 	struct FMeshMergeScratchMeshes
 	{
 		TManagedPtr<FMesh> FirstReformat;
 		TManagedPtr<FMesh> SecondReformat;
 	};
 
+	
 	inline void FixMorphIndices(FMesh& Result, int32 VertexIndexOffset)
 	{
         MUTABLE_CPUPROFILER_SCOPE(CompressedMorphs_Merge_FixMorphIndices)
@@ -57,6 +59,7 @@ namespace UE::Mutable::Private
 		}
 	}
 
+	
 	inline void MergeMorphs(FMesh& Result, const FMesh& MeshA, const FMesh& MeshB)
 	{
         MUTABLE_CPUPROFILER_SCOPE(CompressedMorphs_Merge_Merge)
@@ -283,6 +286,61 @@ namespace UE::Mutable::Private
 		}
 	}
 
+	
+	/** Copy cloth from MeshA. MeshB only required to for data integrity. */
+	inline void CopyClothFirst(FMesh& Result, const FMesh& MeshA, const FMesh& MeshB)
+	{
+		Result.ClothSections[0].ClothingAsset = MeshA.ClothSections[0].ClothingAsset;
+			
+		const uint32 NumClothVerticesA = MeshA.ClothSections[0].Data.Num();
+		if (!NumClothVerticesA)
+		{
+			return;
+		}
+
+		const int32 NumInfluences = MeshA.ClothSections[0].Data.Num() / MeshA.VertexBuffers.GetElementCount();
+		const uint32 NumClothVerticesB = NumInfluences * MeshB.VertexBuffers.GetElementCount();
+		
+		Result.ClothSections[0].Data.Reserve(NumClothVerticesA + NumClothVerticesB);
+
+		Result.ClothSections[0].Data.Append(MeshA.ClothSections[0].Data);
+		
+		// Add MeshB vertices without simulation. See FMeshToMeshVertData::SourceMeshVertIndices.
+		for (uint32 VertexIndex = 0; VertexIndex < NumClothVerticesB; ++VertexIndex)
+		{
+			FMeshToMeshVertData& Data = Result.ClothSections[0].Data.AddZeroed_GetRef();
+			Data.SourceMeshVertIndices[3] = static_cast<uint16>(0xffff);
+		}
+	}
+	
+	
+	/** Copy cloth from MeshB. MeshA only required to for data integrity. */
+	inline void CopyClothSecond(FMesh& Result, const FMesh& MeshA, const FMesh& MeshB)
+	{
+		Result.ClothSections[0].ClothingAsset = MeshB.ClothSections[0].ClothingAsset;
+
+		const uint32 NumClothVerticesB = MeshB.ClothSections[0].Data.Num();
+		if (!NumClothVerticesB)
+		{
+			return;
+		}
+		
+		const int32 NumInfluences = MeshB.ClothSections[0].Data.Num() / MeshB.VertexBuffers.GetElementCount();
+		const uint32 NumClothVerticesA =  NumInfluences * MeshA.VertexBuffers.GetElementCount();
+
+		Result.ClothSections[0].Data.Reserve(NumClothVerticesA + NumClothVerticesB);
+		
+		// Add MeshA vertices without simulation. See FMeshToMeshVertData::SourceMeshVertIndices.
+		for (uint32 VertexIndex = 0; VertexIndex < NumClothVerticesA; ++VertexIndex)
+		{
+			FMeshToMeshVertData& Data = Result.ClothSections[0].Data.AddZeroed_GetRef();
+			Data.SourceMeshVertIndices[3] = static_cast<uint16>(0xffff);					
+		}
+		
+		Result.ClothSections[0].Data.Append(MeshB.ClothSections[0].Data);
+	}
+	
+	
 	//---------------------------------------------------------------------------------------------
 	//! Merge two meshes into one new mesh
 	//---------------------------------------------------------------------------------------------
@@ -726,6 +784,7 @@ namespace UE::Mutable::Private
 				}
 
 				FMeshSurface& NewSurface = Result->Surfaces.AddDefaulted_GetRef();
+				Result->ClothSections.AddDefaulted();
 				NewSurface.BoneMapCount = Result->BoneMap.Num();
 
 				int32 NumFirstSubMeshes = 0;
@@ -763,7 +822,7 @@ namespace UE::Mutable::Private
 
 				// Add pFirst surfaces
 				Result->Surfaces = pFirst->Surfaces;
-
+				
 				const int32 FirstVertexEnd = pFirst->GetVertexCount();
 				const int32 FirstIndexEnd = pFirst->GetIndexCount();
 
@@ -779,6 +838,8 @@ namespace UE::Mutable::Private
 				}
 
 				NewSurface.BoneMapIndex += NumFirstBonesInBoneMap;
+				
+				Result->ClothSections.SetNum(Result->Surfaces.Num());
 			}
 
 			for (const FMeshSurface& Surface : Result->Surfaces)
@@ -1369,8 +1430,7 @@ namespace UE::Mutable::Private
 			Result->SkeletalMeshes.AddUnique(pSecond->SkeletalMeshes[Index]);
 		}
 
-		// Mesh morphs.
-		
+		// Morphs
 		if (pFirst->HasMorphs() && !pSecond->HasMorphs())
 		{
 			Result->Morph = pFirst->Morph;
@@ -1387,5 +1447,35 @@ namespace UE::Mutable::Private
 		{
 			MergeMorphs(*Result, *pFirst, *pSecond);	
 		}
-	}	
+
+		// Cloth
+		if (bMergeSurfaces)
+		{
+			if (pFirst->ClothSections[0].IsValid() && !pSecond->ClothSections[0].IsValid())
+			{
+				CopyClothFirst(*Result, *pFirst, *pSecond);	
+			}
+			else if (!pFirst->ClothSections[0].IsValid() && pSecond->ClothSections[0].IsValid())
+			{
+				CopyClothSecond(*Result, *pFirst, *pSecond);
+			}
+			else if (pFirst->ClothSections[0].IsValid() && pSecond->ClothSections[0].IsValid())
+			{
+				UE_LOG(LogMutableCore, Error, TEXT("Merging two mesh sections with cloth is not supported."));
+				CopyClothFirst(*Result, *pFirst, *pSecond);
+			}
+		}
+		else
+		{
+			for (int32 SectionIndexFirst = 0; SectionIndexFirst < pFirst->ClothSections.Num(); ++SectionIndexFirst)
+			{
+				Result->ClothSections[SectionIndexFirst] = pFirst->ClothSections[SectionIndexFirst];
+			}
+
+			for (int32 SectionIndexSecond = 0; SectionIndexSecond < pSecond->ClothSections.Num(); ++SectionIndexSecond)
+			{
+				Result->ClothSections[pFirst->Surfaces.Num() + SectionIndexSecond] = pSecond->ClothSections[SectionIndexSecond];
+			}
+		}
+	}
 }
diff --git a/Engine/Plugins/Mutable/Source/MutableRuntime/Internal/MuR/Operations.h b/Engine/Plugins/Mutable/Source/MutableRuntime/Internal/MuR/Operations.h
--- a/Engine/Plugins/Mutable/Source/MutableRuntime/Internal/MuR/Operations.h
+++ b/Engine/Plugins/Mutable/Source/MutableRuntime/Internal/MuR/Operations.h
@@ -490,6 +490,8 @@ namespace UE::Mutable::Private
             // If not negative, index of the skeleton to set to the mesh from the skeleton
             // constant array.
             int32 Skeleton;
+        	
+        	PASSTHROUGH_ID ClothID;
         };
 
     	struct MaterialConstantArgs
diff --git a/Engine/Plugins/Mutable/Source/MutableRuntime/Internal/MuR/PassthroughObject.h b/Engine/Plugins/Mutable/Source/MutableRuntime/Internal/MuR/PassthroughObject.h
--- a/Engine/Plugins/Mutable/Source/MutableRuntime/Internal/MuR/PassthroughObject.h
+++ b/Engine/Plugins/Mutable/Source/MutableRuntime/Internal/MuR/PassthroughObject.h
@@ -27,13 +27,13 @@ namespace UE::Mutable::Private
 		{
 			Resource->Set<PASSTHROUGH_ID>(PASSTHROUGH_ID_INVALID);
 		}
-
-		explicit TPassthroughObjectPtr(PASSTHROUGH_ID Id) : Resource(MakeShared<TVariant<TStrongObjectPtr<UObject>, PASSTHROUGH_ID>>())
+		
+		TPassthroughObjectPtr(PASSTHROUGH_ID Id) : Resource(MakeShared<TVariant<TStrongObjectPtr<UObject>, PASSTHROUGH_ID>>())
 		{
 			Resource->Set<PASSTHROUGH_ID>(Id);
 		}
 
-		explicit TPassthroughObjectPtr(T* Object) : Resource(MakeShared<TVariant<TStrongObjectPtr<UObject>, PASSTHROUGH_ID>>())
+		TPassthroughObjectPtr(T* Object) : Resource(MakeShared<TVariant<TStrongObjectPtr<UObject>, PASSTHROUGH_ID>>())
 		{
 			Resource->Set<TStrongObjectPtr<UObject>>(TStrongObjectPtr<UObject>(Object));
 		}
@@ -47,11 +47,23 @@ namespace UE::Mutable::Private
 		{
 		}
 
+		TPassthroughObjectPtr(TPassthroughObjectPtr&& Other) : Resource(Other.Resource)
+		{
+			Other.Resource = MakeShared<TVariant<TStrongObjectPtr<UObject>, PASSTHROUGH_ID>>();
+		}
+		
 		TPassthroughObjectPtr& operator=(const TPassthroughObjectPtr& Other)
 		{
 			Resource = Other.Resource;
 			return *this;
 		}
+		
+		TPassthroughObjectPtr& operator=(TPassthroughObjectPtr&& Other)
+		{
+			Resource = Other.Resource;
+			Other.Resource = MakeShared<TVariant<TStrongObjectPtr<UObject>, PASSTHROUGH_ID>>();
+			return *this;
+		}
 
 		bool operator==(const TPassthroughObjectPtr& Other) const
 		{
@@ -71,12 +83,18 @@ namespace UE::Mutable::Private
 		
 		T* Get() const
 		{
-			if (UObject* Object = Resource->Get<TStrongObjectPtr<UObject>>().Get())
+			if (Resource->IsType<PASSTHROUGH_ID>())
 			{
-				return CastChecked<T>(Object);
+				return nullptr;
+			}
+			
+			UObject* Object = Resource->Get<TStrongObjectPtr<UObject>>().Get();
+			if (!Object)
+			{
+				return nullptr;
 			}
 
-			return nullptr;
+			return CastChecked<T>(Object);
 		}
 
 		PASSTHROUGH_ID GetId() const
@@ -84,9 +102,22 @@ namespace UE::Mutable::Private
 			return Resource->Get<PASSTHROUGH_ID>(); 
 		}
 
+		PASSTHROUGH_ID Reset()
+		{
+			PASSTHROUGH_ID Id = PASSTHROUGH_ID_INVALID; 
+			
+			if (Resource->IsType<PASSTHROUGH_ID>())
+			{
+				Id = Resource->Get<PASSTHROUGH_ID>();
+				}
+			
+			Resource = MakeShared<TVariant<TStrongObjectPtr<UObject>, PASSTHROUGH_ID>>();
+			return Id; 
+		}
+		
 		bool IsValid() const
 		{
-			return !Resource->IsType<PASSTHROUGH_ID>() || Resource->Get<PASSTHROUGH_ID>() != PASSTHROUGH_ID_INVALID;
+			return Resource->IsType<TStrongObjectPtr<UObject>>() || Resource->Get<PASSTHROUGH_ID>() != PASSTHROUGH_ID_INVALID;
 		}
 		
 		explicit operator bool() const
@@ -103,22 +134,25 @@ namespace UE::Mutable::Private
 	{
 	public:
 		template<typename T>
-		TPassthroughObjectPtr<T> Find(PASSTHROUGH_ID Id)
+		TPassthroughObjectPtr<T> FindChecked(PASSTHROUGH_ID Id)
 		{
-			if (TPassthroughObjectPtr<UObject>* Result = PassthroughObjects.Find(Id))
+			if (Id == PASSTHROUGH_ID_INVALID)
 			{
-				return TPassthroughObjectPtr<T>(*Result);
+				return {};
 			}
 			else
 			{
-				return {};
+				return PassthroughObjects.FindChecked(Id);
 			}
 		}
 		
 		template<typename T>
 		void Add(PASSTHROUGH_ID Id)
 		{
-			PassthroughObjects.Add(Id, TPassthroughObjectPtr<T>(Id));
+			if (Id != PASSTHROUGH_ID_INVALID)
+			{
+				PassthroughObjects.Add(Id, TPassthroughObjectPtr<T>(Id));
+			}
 		}
 
 		MUTABLERUNTIME_API void Load(const TMap<PASSTHROUGH_ID, TSoftObjectPtr<UObject>>& ResolutionMap, TArray<FSoftObjectPath>& AssetsToStream) const;
diff --git a/Engine/Plugins/Mutable/Source/MutableRuntime/Internal/MuR/Serialisation.h b/Engine/Plugins/Mutable/Source/MutableRuntime/Internal/MuR/Serialisation.h
--- a/Engine/Plugins/Mutable/Source/MutableRuntime/Internal/MuR/Serialisation.h
+++ b/Engine/Plugins/Mutable/Source/MutableRuntime/Internal/MuR/Serialisation.h
@@ -21,6 +21,7 @@
 #include "StructUtils/InstancedStruct.h"
 #include "UObject/StrongObjectPtr.h"
 
+struct FMeshToMeshVertData;
 class UTexture;
 class USkeletalMesh;
 
@@ -343,6 +344,9 @@ namespace UE::Mutable::Private
 	MUTABLERUNTIME_API void operator<<(FOutputArchive&, const FName&);
 	MUTABLERUNTIME_API void operator>>(FInputArchive&, FName&);
 
+	//---------------------------------------------------------------------------------------------
+	MUTABLERUNTIME_API void operator<<(FOutputArchive&, const FMeshToMeshVertData&);
+	MUTABLERUNTIME_API void operator>>(FInputArchive&, FMeshToMeshVertData&);
 
 	//---------------------------------------------------------------------------------------------
 	template<typename T, typename Alloc> 
diff --git a/Engine/Plugins/Mutable/Source/MutableRuntime/MutableRuntime.Build.cs b/Engine/Plugins/Mutable/Source/MutableRuntime/MutableRuntime.Build.cs
--- a/Engine/Plugins/Mutable/Source/MutableRuntime/MutableRuntime.Build.cs
+++ b/Engine/Plugins/Mutable/Source/MutableRuntime/MutableRuntime.Build.cs
@@ -21,7 +21,8 @@ namespace UnrealBuildTool.Rules
 				"Core",
 				"CoreUObject",
 				"Engine",
-				"ImageCore"
+				"ImageCore",
+				"ClothingSystemRuntimeCommon",
 				}
 			);
 
diff --git a/Engine/Plugins/Mutable/Source/MutableRuntime/Private/MuR/Cloth.cpp b/Engine/Plugins/Mutable/Source/MutableRuntime/Private/MuR/Cloth.cpp
new file mode 100644
--- /dev/null
+++ b/Engine/Plugins/Mutable/Source/MutableRuntime/Private/MuR/Cloth.cpp
@@ -0,0 +1,42 @@
+// Copyright Epic Games, Inc. All Rights Reserved.
+
+#include "MuR/Cloth.h"
+
+#include "Serialization/MemoryReader.h"
+#include "Serialization/MemoryWriter.h"
+
+
+namespace UE::Mutable::Private
+{
+	bool FCloth::operator==(const FCloth& Other) const
+	{
+		return ClothingAsset == Other.ClothingAsset &&
+			Data == Other.Data;
+	}
+
+
+	int32 FCloth::GetDataSize() const
+	{
+		return Data.GetAllocatedSize();
+	}
+
+
+	void FCloth::Serialise(FOutputArchive& Ar) const
+	{
+		Ar << Data;
+		// Ar << PassthroughObject; // Serialized in the OP Args.
+	}
+
+
+	void FCloth::Unserialise(FInputArchive& Ar)
+	{
+		Ar >> Data;
+	}
+
+
+	bool FCloth::IsValid() const
+	{
+		return ClothingAsset.IsValid();
+	}
+}
+
diff --git a/Engine/Plugins/Mutable/Source/MutableRuntime/Private/MuR/CodeRunner.cpp b/Engine/Plugins/Mutable/Source/MutableRuntime/Private/MuR/CodeRunner.cpp
--- a/Engine/Plugins/Mutable/Source/MutableRuntime/Private/MuR/CodeRunner.cpp
+++ b/Engine/Plugins/Mutable/Source/MutableRuntime/Private/MuR/CodeRunner.cpp
@@ -1095,8 +1095,10 @@ namespace UE::Mutable::Private
 			OP::MeshConstantArgs Args = Program.GetOpArgs<OP::MeshConstantArgs>(Item.At);
 			EMeshContentFlags MeshContentFlags = (EMeshContentFlags)Item.ExecutionOptions;
 
+			const TPassthroughObjectPtr<UClothingAssetCommon> ClothAsset = FindPassthrough<UClothingAssetCommon>(Args.ClothID);
+        		
 			TManagedPtr<const FMesh> Source;
-			Program.GetConstant(Args.Value, Args.Skeleton, Source, MeshContentFlags, 
+			Program.GetConstant(Args.Value, Args.Skeleton, ClothAsset, Source, MeshContentFlags, 
 					[this](int32 BudgetReserve){ return CreateMesh(BudgetReserve); });
 			
 			if (!Source)
diff --git a/Engine/Plugins/Mutable/Source/MutableRuntime/Private/MuR/CodeRunner.h b/Engine/Plugins/Mutable/Source/MutableRuntime/Private/MuR/CodeRunner.h
--- a/Engine/Plugins/Mutable/Source/MutableRuntime/Private/MuR/CodeRunner.h
+++ b/Engine/Plugins/Mutable/Source/MutableRuntime/Private/MuR/CodeRunner.h
@@ -603,10 +603,7 @@ namespace  UE::Mutable::Private
 		{
 			if (LiveInstance->PassthroughObjectLoader)
 			{
-				TPassthroughObjectPtr<T> PassthroughObject = LiveInstance->PassthroughObjectLoader->Find<T>(Id);
-				check(PassthroughObject); // Object not discovered in the CodeRunnerBegin.
-				
-				return PassthroughObject;
+				return LiveInstance->PassthroughObjectLoader->FindChecked<T>(Id);
 			}
 			else
 			{
diff --git a/Engine/Plugins/Mutable/Source/MutableRuntime/Private/MuR/CodeRunnerBegin.cpp b/Engine/Plugins/Mutable/Source/MutableRuntime/Private/MuR/CodeRunnerBegin.cpp
--- a/Engine/Plugins/Mutable/Source/MutableRuntime/Private/MuR/CodeRunnerBegin.cpp
+++ b/Engine/Plugins/Mutable/Source/MutableRuntime/Private/MuR/CodeRunnerBegin.cpp
@@ -1891,6 +1891,7 @@ namespace UE::Mutable::Private
 					check(!Item.bEvaluate);
 
 					OP::MeshConstantArgs Args = Program.GetOpArgs<OP::MeshConstantArgs>(Item.At);
+					PassthroughObjectLoader->Add<UMaterialInterface>(Args.ClothID);
 
 					StoreNone(Item);
 					break;
diff --git a/Engine/Plugins/Mutable/Source/MutableRuntime/Private/MuR/Mesh.cpp b/Engine/Plugins/Mutable/Source/MutableRuntime/Private/MuR/Mesh.cpp
--- a/Engine/Plugins/Mutable/Source/MutableRuntime/Private/MuR/Mesh.cpp
+++ b/Engine/Plugins/Mutable/Source/MutableRuntime/Private/MuR/Mesh.cpp
@@ -9,6 +9,7 @@
 #include "HAL/UnrealMemory.h"
 #include "MuR/MeshPrivate.h"
 #include "MuR/MutableTrace.h"
+#include "MuR/Serialisation.h"
 #include "MuR/OpMeshClipWithMesh.h"
 #include "Spatial/PointHashGrid3.h"
 #include "Engine/SkeletalMesh.h"
@@ -182,6 +183,11 @@ void FMesh::CopyFrom(const FMesh& From, EMeshCopyFlags InFlags)
 		Morph = From.Morph;
 		MorphDataBuffer = From.MorphDataBuffer;	
 	}
+	
+	if (EnumHasAnyFlags(InFlags, EMeshCopyFlags::WithClothData))
+	{
+		ClothSections = From.ClothSections;
+	}
 }
 
 
@@ -197,6 +203,20 @@ bool FMesh::HasMorphs() const
 }
 
 
+bool FMesh::HasClothing() const
+{
+	for (const FCloth& Cloth : ClothSections)
+	{
+		if (Cloth.IsValid())
+		{
+			return true;
+		}
+	}
+	
+	return false;
+}
+
+
 int FMesh::GetVertexCount() const
 {
     return GetVertexBuffers().GetElementCount();
@@ -675,6 +695,11 @@ int32 FMesh::GetDataSize() const
 		AdditionalBuffersSize += AdditionalBuffer.Value.GetDataSize();
 	}
 
+	for (const FCloth& ClothSection : ClothSections)
+	{
+		AdditionalBuffersSize += ClothSection.GetDataSize();
+	}
+	
 	return sizeof(FMesh)
 		+ IndexBuffers.GetDataSize()
 		+ VertexBuffers.GetDataSize()
@@ -879,6 +904,7 @@ void FMesh::EnsureSurfaceData()
 	if (!Surfaces.Num() && VertexBuffers.GetElementCount())
 	{
 		FMeshSurface& NewSurface = Surfaces.Emplace_GetRef();
+		ClothSections.AddDefaulted();
 
 		FSurfaceSubMesh& SubMesh = NewSurface.SubMeshes.Emplace_GetRef();
 		SubMesh.VertexBegin = 0;
@@ -939,6 +965,7 @@ void FMesh::FBonePose::Unserialise(FInputArchive& Arch)
 void FMesh::Serialise(FOutputArchive& Arch) const
 {
 	check(!IsReference());
+	check(ClothSections.Num() <= 1);
 	
 	Arch << IndexBuffers;
 	Arch << VertexBuffers;
@@ -962,6 +989,8 @@ void FMesh::Serialise(FOutputArchive& Arch) const
 	Arch << AdditionalPhysicsBodies;
 
 	Arch << MeshIDPrefix;
+	
+	Arch << ClothSections;
 }
 
 
@@ -992,6 +1021,8 @@ void FMesh::Unserialise(FInputArchive& Arch)
 	Arch >> AdditionalPhysicsBodies;
 
 	Arch >> MeshIDPrefix;
+	
+	Arch >> ClothSections;
 }
 
 
@@ -1039,6 +1070,8 @@ bool FMesh::IsSimilar(const FMesh& Other) const
 	{
 		bEqual = VertexBuffers.IsSimilarRobust(Other.VertexBuffers, false);
 	}
+	
+	bEqual = bEqual && (ClothSections == Other.ClothSections);
 
 	return bEqual;
 }
diff --git a/Engine/Plugins/Mutable/Source/MutableRuntime/Private/MuR/OpImageProject.cpp b/Engine/Plugins/Mutable/Source/MutableRuntime/Private/MuR/OpImageProject.cpp
--- a/Engine/Plugins/Mutable/Source/MutableRuntime/Private/MuR/OpImageProject.cpp
+++ b/Engine/Plugins/Mutable/Source/MutableRuntime/Private/MuR/OpImageProject.cpp
@@ -3054,6 +3054,7 @@ void MeshProject(FMesh* Result, const FMesh* pMesh, const FProjector& projector,
 	if (bOutSuccess)
 	{
 		Result->Surfaces.Empty();
+		Result->ClothSections.Empty();
 		Result->EnsureSurfaceData();
 		Result->ResetStaticFormatFlags();
 	}
diff --git a/Engine/Plugins/Mutable/Source/MutableRuntime/Private/MuR/OpMeshDifference.h b/Engine/Plugins/Mutable/Source/MutableRuntime/Private/MuR/OpMeshDifference.h
--- a/Engine/Plugins/Mutable/Source/MutableRuntime/Private/MuR/OpMeshDifference.h
+++ b/Engine/Plugins/Mutable/Source/MutableRuntime/Private/MuR/OpMeshDifference.h
@@ -204,6 +204,7 @@ namespace UE::Mutable::Private
 
         // Morphs are always single surfaced
         Result->Surfaces.Empty();
+        Result->ClothSections.Empty();
         Result->EnsureSurfaceData();
 	}
 
diff --git a/Engine/Plugins/Mutable/Source/MutableRuntime/Private/MuR/OpMeshFormat.cpp b/Engine/Plugins/Mutable/Source/MutableRuntime/Private/MuR/OpMeshFormat.cpp
--- a/Engine/Plugins/Mutable/Source/MutableRuntime/Private/MuR/OpMeshFormat.cpp
+++ b/Engine/Plugins/Mutable/Source/MutableRuntime/Private/MuR/OpMeshFormat.cpp
@@ -719,6 +719,7 @@ namespace UE::Mutable::Private
 		Result->AdditionalPhysicsBodies = Source->AdditionalPhysicsBodies;
 
 		Result->Surfaces = Source->Surfaces;
+		Result->ClothSections = Source->ClothSections;
 
 		Result->ResetStaticFormatFlags();
 		Result->EnsureSurfaceData();
diff --git a/Engine/Plugins/Mutable/Source/MutableRuntime/Private/MuR/OpMeshRemove.cpp b/Engine/Plugins/Mutable/Source/MutableRuntime/Private/MuR/OpMeshRemove.cpp
--- a/Engine/Plugins/Mutable/Source/MutableRuntime/Private/MuR/OpMeshRemove.cpp
+++ b/Engine/Plugins/Mutable/Source/MutableRuntime/Private/MuR/OpMeshRemove.cpp
@@ -230,6 +230,7 @@ namespace UE::Mutable::Private
 			if (!Result->Surfaces[I].SubMeshes.Num())
 			{
 				Result->Surfaces.RemoveAt(I, EAllowShrinking::No);
+				Result->ClothSections.RemoveAt(I, EAllowShrinking::No);
 			}
 		}
 
@@ -307,6 +308,35 @@ namespace UE::Mutable::Private
     }
 
 	
+	void CopyMoveBuffer(uint8* BufferData, uint32 ElemSize, int32 NumVertices, const TBitArray<>& UsedVertices)
+	{
+		int32 LastFreeVertexIndex = 0;
+
+		for (int32 VertexIndex = UsedVertices.FindFrom(true, 0); VertexIndex >= 0;)
+		{
+			const int32 UsedSpanBegin = VertexIndex;
+			VertexIndex = UsedVertices.FindFrom(false, VertexIndex);
+			VertexIndex = VertexIndex >= 0 ? VertexIndex : NumVertices;
+			const int32 UsedSpanEnd = VertexIndex;
+
+			VertexIndex = UsedVertices.FindFrom(true, VertexIndex);
+
+			// Copy vertex buffer span.	
+			const int32 UsedSpanSize = UsedSpanEnd - UsedSpanBegin;
+            
+			if (LastFreeVertexIndex != UsedSpanBegin)
+			{
+				FMemory::Memmove(
+						BufferData + LastFreeVertexIndex*ElemSize,
+						BufferData + UsedSpanBegin*ElemSize,
+						UsedSpanSize*ElemSize);
+			}
+
+			LastFreeVertexIndex += UsedSpanSize;
+		}
+	}
+	
+	
     void MeshRemoveVerticesWithCullSet(FMesh* Result, const TBitArray<>& VerticesToCull, bool bRemoveIfAllVerticesCulled)
     { 
         MUTABLE_CPUPROFILER_SCOPE(MeshRemoveVerticesWithCullSet);
@@ -453,6 +483,7 @@ namespace UE::Mutable::Private
         // This data will only be accessed by indices that have been mapped, No need to initialize. 
 		UsedVerticesMap.SetNumUninitialized(NumVertices);
 #endif
+    	
 		FMeshBufferSet& VertexBufferSet = Result->GetVertexBuffers();
 
 		const int32 NumBuffers = VertexBufferSet.GetBufferCount();
@@ -493,38 +524,26 @@ namespace UE::Mutable::Private
         {
             for (int32 BufferIndex = 0; BufferIndex < NumBuffers; ++BufferIndex)
             {
-                uint8* BufferData = VertexBufferSet.GetBufferData(BufferIndex); 
-                const uint32 ElemSize = VertexBufferSet.GetElementSize(BufferIndex);
-
-                int32 LastFreeVertexIndex = 0;
-
-                for (int32 VertexIndex = UsedVertices.FindFrom(true, 0); VertexIndex >= 0;)
-                {
-                    const int32 UsedSpanBegin = VertexIndex;
-                    VertexIndex = UsedVertices.FindFrom(false, VertexIndex);
-                    VertexIndex = VertexIndex >= 0 ? VertexIndex : NumVertices;
-                    const int32 UsedSpanEnd = VertexIndex;
-
-                    VertexIndex = UsedVertices.FindFrom(true, VertexIndex);
-
-                    // Copy vertex buffer span.	
-                    const int32 UsedSpanSize = UsedSpanEnd - UsedSpanBegin;
-                    
-                    if (LastFreeVertexIndex != UsedSpanBegin)
-                    {
-                        FMemory::Memmove(
-                                BufferData + LastFreeVertexIndex*ElemSize,
-                                BufferData + UsedSpanBegin*ElemSize,
-                                UsedSpanSize*ElemSize);
-                    }
-
-                    LastFreeVertexIndex += UsedSpanSize;
-                }
+            	uint8* BufferData = VertexBufferSet.GetBufferData(BufferIndex); 
+            	const uint32 ElemSize = VertexBufferSet.GetElementSize(BufferIndex);
 
-                check(LastFreeVertexIndex == NumVerticesRemaining);
+            	CopyMoveBuffer(BufferData, ElemSize, NumVertices, UsedVertices);
             }
         }
 		Result->GetVertexBuffers().SetElementCount(NumVerticesRemaining);
+    	
+    	check(Result->ClothSections.Num() == 1); // Check that we are only clipping meshes with one section.
+    	
+    	const uint32 NumClothVertices = Result->ClothSections[0].Data.Num();
+    	if (Result->ClothSections[0].IsValid() && NumClothVertices)
+    	{
+    		uint8* BufferData = reinterpret_cast<uint8*>(Result->ClothSections[0].Data.GetData());
+    		const int32 NumInfluences = NumClothVertices / NumVerticesRemaining;
+    		const uint32 ElemSize = sizeof(FMeshToMeshVertData) * NumInfluences;
+
+    		CopyMoveBuffer(BufferData, ElemSize, NumVertices, UsedVertices);
+    		Result->ClothSections[0].Data.SetNum(NumVerticesRemaining * NumInfluences);
+    	}
 
 		int32 LastFreeFaceIndex = 0;
 		if (bHasFaces)
diff --git a/Engine/Plugins/Mutable/Source/MutableRuntime/Private/MuR/PassthroughObject.cpp b/Engine/Plugins/Mutable/Source/MutableRuntime/Private/MuR/PassthroughObject.cpp
--- a/Engine/Plugins/Mutable/Source/MutableRuntime/Private/MuR/PassthroughObject.cpp
+++ b/Engine/Plugins/Mutable/Source/MutableRuntime/Private/MuR/PassthroughObject.cpp
@@ -1,4 +1,4 @@
-﻿// Copyright Epic Games, Inc. All Rights Reserved.
+// Copyright Epic Games, Inc. All Rights Reserved.
 
 #include "MuR/PassthroughObject.h"
 
@@ -22,7 +22,7 @@ namespace UE::Mutable::Private
 
 	void FPassthroughObjectLoader::Load(const TMap<PASSTHROUGH_ID, TSoftObjectPtr<UObject>>& ResolutionMap, TArray<FSoftObjectPath>& AssetsToStream) const
 	{
-		for (const TPair<PASSTHROUGH_ID, TPassthroughObjectPtr<UObject>>& Pair:  PassthroughObjects)
+		for (const TPair<PASSTHROUGH_ID, TPassthroughObjectPtr<UObject>>& Pair : PassthroughObjects)
 		{
 			const TPassthroughObjectPtr<UObject>& ExternalObject = Pair.Value;
 				
@@ -33,10 +33,7 @@ namespace UE::Mutable::Private
 			}
 
 			const PASSTHROUGH_ID ID = ExternalObject.Resource->Get<PASSTHROUGH_ID>();
-			if (ID == PASSTHROUGH_ID_INVALID)
-			{
-				continue;
-			}
+			check(ID != PASSTHROUGH_ID_INVALID);
 
 			const TSoftObjectPtr<UObject>* Result = ResolutionMap.Find(ID);
 			if (!Result)
@@ -63,11 +60,7 @@ namespace UE::Mutable::Private
 			}
 
 			const PASSTHROUGH_ID ID = ExternalObject.Resource->Get<PASSTHROUGH_ID>();
-			if (ID == PASSTHROUGH_ID_INVALID)
-			{
-				ExternalObject.Resource->Set<TStrongObjectPtr<UObject>>(nullptr);
-				continue;
-			}
+			check(ID != PASSTHROUGH_ID_INVALID);
 			
 			const TSoftObjectPtr<UObject>* Result = ResolutionMap.Find(ID);
 			if (!Result)
diff --git a/Engine/Plugins/Mutable/Source/MutableRuntime/Private/MuR/Serialisation.cpp b/Engine/Plugins/Mutable/Source/MutableRuntime/Private/MuR/Serialisation.cpp
--- a/Engine/Plugins/Mutable/Source/MutableRuntime/Private/MuR/Serialisation.cpp
+++ b/Engine/Plugins/Mutable/Source/MutableRuntime/Private/MuR/Serialisation.cpp
@@ -3,6 +3,7 @@
 
 #include "MuR/Serialisation.h"
 
+#include "SkeletalMeshTypes.h"
 #include "Containers/Array.h"
 #include "HAL/PlatformCrt.h"
 #include "HAL/PlatformMath.h"
@@ -79,9 +80,35 @@ namespace UE::Mutable::Private
 	    arch >> Temp;
 	    v = FName(Temp);
     }
-	
-	
-	void operator<<(FOutputArchive& Arch, const bool& T)
+
+
+    void operator<<(FOutputArchive& Ar, const FMeshToMeshVertData& Value)
+    {
+    	Ar << Value.PositionBaryCoordsAndDist; 
+    	Ar << Value.NormalBaryCoordsAndDist;
+    	Ar << Value.TangentBaryCoordsAndDist;
+    	Ar << Value.SourceMeshVertIndices[0];
+    	Ar << Value.SourceMeshVertIndices[1];
+    	Ar << Value.SourceMeshVertIndices[2];
+    	Ar << Value.SourceMeshVertIndices[3];
+    	Ar << Value.Weight;
+    }
+
+
+    void operator>>(FInputArchive& Ar, FMeshToMeshVertData& Value)
+    {
+    	Ar >> Value.PositionBaryCoordsAndDist; 
+    	Ar >> Value.NormalBaryCoordsAndDist;
+    	Ar >> Value.TangentBaryCoordsAndDist;
+    	Ar >> Value.SourceMeshVertIndices[0];
+    	Ar >> Value.SourceMeshVertIndices[1];
+    	Ar >> Value.SourceMeshVertIndices[2];
+    	Ar >> Value.SourceMeshVertIndices[3];
+    	Ar >> Value.Weight;
+    }
+
+
+    void operator<<(FOutputArchive& Arch, const bool& T)
     {
     	uint8 S = T ? 1 : 0;
     	Arch.Stream->Write(&S, sizeof(uint8));
diff --git a/Engine/Plugins/Mutable/Source/MutableTools/Private/MuT/ASTOpConstantResource.cpp b/Engine/Plugins/Mutable/Source/MutableTools/Private/MuT/ASTOpConstantResource.cpp
--- a/Engine/Plugins/Mutable/Source/MutableTools/Private/MuT/ASTOpConstantResource.cpp
+++ b/Engine/Plugins/Mutable/Source/MutableTools/Private/MuT/ASTOpConstantResource.cpp
@@ -114,7 +114,8 @@ namespace UE::Mutable::Private
 						EMeshCopyFlags::WithSurfaces      | 
 						EMeshCopyFlags::WithVertexBuffers |
 						EMeshCopyFlags::WithIndexBuffers  |
-						EMeshCopyFlags::WithLayouts;
+						EMeshCopyFlags::WithLayouts		  |
+						EMeshCopyFlags::WithClothData;
 
 				TManagedPtr<FMesh> MeshGeometryData = MeshData->Clone(GeometryDataCopyFlags);
 
@@ -453,6 +454,8 @@ namespace UE::Mutable::Private
 					check(!MeshValue->PassthroughObject);
 					
 					Args.Value = AddConstantMesh(Program, MeshValue, *Options);
+					Args.ClothID = MeshValue->ClothSections.Num() ? MeshValue->ClothSections[0].ClothingAsset.Reset() : PASSTHROUGH_ID_INVALID;
+
 					int32 DataDescIndex = Options->AdditionalData.SourceMeshPerConstant.Add(SourceDataDescriptor);
 					check(DataDescIndex == Args.Value);
 
@@ -534,9 +537,7 @@ namespace UE::Mutable::Private
 						MaterialValue->ImageParameters.Add(ImageOperation.Key, ImageParameterData);
 					}
 
-					Args.ID = MaterialValue->PassthroughObject.GetId();
-					MaterialValue->PassthroughObject = {}; // Not serialized but compared. Set to invalid to avoid duplicating it.
-					
+					Args.ID = MaterialValue->PassthroughObject.Reset(); // Not serialized but compared. Set to invalid to avoid duplicating it.
 					Args.Value = Program.AddConstant(MaterialValue);
 					
 					linkedAddress = (OP::ADDRESS)Program.OpAddress.Num();
diff --git a/Engine/Source/Runtime/ClothingSystemRuntimeCommon/Private/ClothingAsset.cpp b/Engine/Source/Runtime/ClothingSystemRuntimeCommon/Private/ClothingAsset.cpp
--- a/Engine/Source/Runtime/ClothingSystemRuntimeCommon/Private/ClothingAsset.cpp
+++ b/Engine/Source/Runtime/ClothingSystemRuntimeCommon/Private/ClothingAsset.cpp
@@ -1361,9 +1361,11 @@ void UClothingAssetCommon::PostLoad()
 		ClothingCachedDataFlags |= EClothingCachedDataFlagsCommon::Tethers;
 
 		// ReferenceBoneIndex is only required when rebinding the cloth.
-		USkeletalMesh* const OwnerMesh = CastChecked<USkeletalMesh>(GetOuter());
-		RefreshBoneMapping(OwnerMesh);
-		CalculateReferenceBoneIndex();
+		if (USkeletalMesh* const OwnerMesh = Cast<USkeletalMesh>(GetOuter()))
+		{
+			RefreshBoneMapping(OwnerMesh);
+			CalculateReferenceBoneIndex();
+		}
 	}
 
 	// After fixing the content, we are ready to call functions that rely on it
diff --git a/Engine/Source/Runtime/Engine/Private/Rendering/SkeletalMeshVertexClothBuffer.cpp b/Engine/Source/Runtime/Engine/Private/Rendering/SkeletalMeshVertexClothBuffer.cpp
--- a/Engine/Source/Runtime/Engine/Private/Rendering/SkeletalMeshVertexClothBuffer.cpp
+++ b/Engine/Source/Runtime/Engine/Private/Rendering/SkeletalMeshVertexClothBuffer.cpp
@@ -174,7 +174,7 @@ void FSkeletalMeshVertexClothBuffer::SerializeMetaData(FArchive& Ar)
 * Initializes the buffer with the given vertices.
 * @param InVertices - The vertices to initialize the buffer with.
 */
-void FSkeletalMeshVertexClothBuffer::Init(const TArray<FMeshToMeshVertData>& InMappingData, const TArray<FClothBufferIndexMapping>& InClothIndexMapping)
+void FSkeletalMeshVertexClothBuffer::Init(const TArrayView<const FMeshToMeshVertData>& InMappingData, const TArrayView<const FClothBufferIndexMapping>& InClothIndexMapping)
 {
 	// Allocate new data
 	AllocateData();
diff --git a/Engine/Source/Runtime/Engine/Public/Rendering/SkeletalMeshVertexClothBuffer.h b/Engine/Source/Runtime/Engine/Public/Rendering/SkeletalMeshVertexClothBuffer.h
--- a/Engine/Source/Runtime/Engine/Public/Rendering/SkeletalMeshVertexClothBuffer.h
+++ b/Engine/Source/Runtime/Engine/Public/Rendering/SkeletalMeshVertexClothBuffer.h
@@ -46,7 +46,7 @@ public:
 	* @param InVertices - The vertices to initialize the buffer with.
 	* @param InClothIndexMapping - The cloth deformer mapping data offset for a specific section/LODBias.
 	*/
-	ENGINE_API void Init(const TArray<FMeshToMeshVertData>& InMappingData, const TArray<FClothBufferIndexMapping>& InClothIndexMapping);
+	ENGINE_API void Init(const TArrayView<const FMeshToMeshVertData>& InMappingData, const TArrayView<const FClothBufferIndexMapping>& InClothIndexMapping);
 
 
 	/**
diff --git a/Engine/Source/Runtime/Engine/Public/SkeletalMeshTypes.h b/Engine/Source/Runtime/Engine/Public/SkeletalMeshTypes.h
--- a/Engine/Source/Runtime/Engine/Public/SkeletalMeshTypes.h
+++ b/Engine/Source/Runtime/Engine/Public/SkeletalMeshTypes.h
@@ -116,6 +116,18 @@ struct FMeshToMeshVertData
 	uint32	 Padding;
 
 	friend ENGINE_API FArchive& operator<<(FArchive& Ar, FMeshToMeshVertData& V);
+	
+	bool operator==(const FMeshToMeshVertData& Other) const
+	{
+		return PositionBaryCoordsAndDist == Other.PositionBaryCoordsAndDist &&
+			NormalBaryCoordsAndDist == Other.NormalBaryCoordsAndDist &&
+			TangentBaryCoordsAndDist == Other.TangentBaryCoordsAndDist &&
+			SourceMeshVertIndices[0] == Other.SourceMeshVertIndices[0] &&
+			SourceMeshVertIndices[1] == Other.SourceMeshVertIndices[1] &&
+			SourceMeshVertIndices[2] == Other.SourceMeshVertIndices[2] &&
+			SourceMeshVertIndices[3] == Other.SourceMeshVertIndices[3] &&
+			Weight == Other.Weight;
+	}
 };
 
 /**
