[Mutable] Fix mesh merge not remapping BoneMap indices.
- Fix potential unaligned uint32 reads that could result in undefined behavior

#rnx

#Synced-CL 50155787

diff --git a/Engine/Plugins/Mutable/Source/MutableRuntime/Internal/MuR/OpMeshMerge.h b/Engine/Plugins/Mutable/Source/MutableRuntime/Internal/MuR/OpMeshMerge.h
--- a/Engine/Plugins/Mutable/Source/MutableRuntime/Internal/MuR/OpMeshMerge.h
+++ b/Engine/Plugins/Mutable/Source/MutableRuntime/Internal/MuR/OpMeshMerge.h
@@ -818,7 +818,7 @@ namespace UE::Mutable::Private
 			else
 			{
 				// Add the BoneMap of the second mesh
-				Result->BoneMap.Append(pSecond->BoneMap);
+				Result->BoneMap.Append(SecondBoneMap);
 
 				// Add pFirst surfaces
 				Result->Surfaces = pFirst->Surfaces;
diff --git a/Engine/Plugins/Mutable/Source/MutableRuntime/Private/MuR/CodeRunner.cpp b/Engine/Plugins/Mutable/Source/MutableRuntime/Private/MuR/CodeRunner.cpp
--- a/Engine/Plugins/Mutable/Source/MutableRuntime/Private/MuR/CodeRunner.cpp
+++ b/Engine/Plugins/Mutable/Source/MutableRuntime/Private/MuR/CodeRunner.cpp
@@ -2430,31 +2430,43 @@ namespace UE::Mutable::Private
 					Data += sizeof(Args);
 
 					// Rebuilding array of bone names ----
+					TArray<FName> BonesToDeform;
+					TArray<FName> PhysicsToDeform;
+
 					int32 NumBones;
 					FMemory::Memcpy(&NumBones, Data, sizeof(int32)); 
 					Data += sizeof(int32);
 					
-					TArrayView<const uint32> BoneIds(reinterpret_cast<const uint32*>(Data), NumBones);
-					TArray<FName> BonesToDeform;
-					BonesToDeform.Reserve(NumBones);
-					for (uint32 BoneId : BoneIds)
+					TArray<uint32> BoneIds;
+					if (NumBones > 0)
 					{
-						BonesToDeform.Add(Program.ConstantNames[BoneId]);
+						BoneIds.SetNumUninitialized(NumBones);
+						FMemory::Memcpy(BoneIds.GetData(), Data, NumBones * sizeof(uint32));
+						Data += NumBones * sizeof(uint32);
+
+						BonesToDeform.Reserve(NumBones);
+						for (uint32 BoneId : BoneIds)
+						{
+							BonesToDeform.Add(Program.ConstantNames[BoneId]);
+						}
 					}
-					Data += NumBones * sizeof(uint32);
 
 					int32 NumPhysicsBodies;
-					FMemory::Memcpy(&NumPhysicsBodies, Data, sizeof(int32)); 
+					FMemory::Memcpy(&NumPhysicsBodies, Data, sizeof(int32));
 					Data += sizeof(int32);
 
-					TArrayView<const uint32> PhysicsBoneIds(reinterpret_cast<const uint32*>(Data), NumPhysicsBodies);
-					TArray<FName> PhysicsToDeform;
-					PhysicsToDeform.Reserve(NumPhysicsBodies);
-					for (uint32 PhysicsBoneId : PhysicsBoneIds)
+					if (NumPhysicsBodies > 0)
 					{
-						PhysicsToDeform.Add(Program.ConstantNames[PhysicsBoneId]);
+						BoneIds.SetNumUninitialized(NumPhysicsBodies, EAllowShrinking::No);
+						FMemory::Memcpy(BoneIds.GetData(), Data, NumPhysicsBodies * sizeof(uint32));
+						Data += NumPhysicsBodies * sizeof(uint32);
+
+						PhysicsToDeform.Reserve(NumPhysicsBodies);
+						for (uint32 PhysicsBoneId : BoneIds)
+						{
+							PhysicsToDeform.Add(Program.ConstantNames[PhysicsBoneId]);
+						}
 					}
-					Data += NumPhysicsBodies * sizeof(uint32);
 
 					EMeshBindShapeFlags BindFlags = static_cast<EMeshBindShapeFlags>(Args.flags);
 					EMeshContentFlags MeshContentFilter = (EMeshContentFlags)Item.ExecutionOptions;
