[Mutable] Fix crash when deserializing meshes with cloth.

#Synced-CL 50155263

diff --git a/Engine/Plugins/Mutable/Source/MutableRuntime/Private/MuR/Mesh.cpp b/Engine/Plugins/Mutable/Source/MutableRuntime/Private/MuR/Mesh.cpp
--- a/Engine/Plugins/Mutable/Source/MutableRuntime/Private/MuR/Mesh.cpp
+++ b/Engine/Plugins/Mutable/Source/MutableRuntime/Private/MuR/Mesh.cpp
@@ -105,6 +105,7 @@ void FMesh::CopyFrom(const FMesh& From, EMeshCopyFlags InFlags)
 	if (EnumHasAnyFlags(InFlags, EMeshCopyFlags::WithSurfaces))
 	{
 		Surfaces = From.Surfaces;
+		ClothSections.SetNum(Surfaces.Num());
 	}
 
 	if (EnumHasAnyFlags(InFlags, EMeshCopyFlags::WithSkeleton))
