[Mutable] Fix missing include MaterialInterface header.

#Synced-CL 50155263

diff --git a/Engine/Plugins/Mutable/Source/CustomizableObjectEditor/Private/MuCOE/Nodes/CONodeModifierExtendSkeletalMeshSection.cpp b/Engine/Plugins/Mutable/Source/CustomizableObjectEditor/Private/MuCOE/Nodes/CONodeModifierExtendSkeletalMeshSection.cpp
--- a/Engine/Plugins/Mutable/Source/CustomizableObjectEditor/Private/MuCOE/Nodes/CONodeModifierExtendSkeletalMeshSection.cpp
+++ b/Engine/Plugins/Mutable/Source/CustomizableObjectEditor/Private/MuCOE/Nodes/CONodeModifierExtendSkeletalMeshSection.cpp
@@ -3,6 +3,7 @@
 #include "MuCOE/Nodes/CONodeModifierExtendSkeletalMeshSection.h"
 
 #include "Materials/MaterialParameters.h"
+#include "Materials/MaterialInterface.h"
 #include "MuCO/CustomizableObjectCustomVersion.h"
 #include "MuCOE/CustomizableObjectEditor_Deprecated.h"
 #include "MuCOE/CustomizableObjectGraph.h"
