[Mutable] Fixed Skeletal Mesh Break crash.
+ Added a new error showing that we do not support breaking Skeletal Mesh Materials into Mutable Textures.
+ Added missing layer index to parameters comming from breaking a Material Parameter Node.

#rnx

#Synced-CL 50102962

diff --git a/Engine/Plugins/Mutable/Source/CustomizableObject/Internal/MuCO/CustomizableObjectPrivate.h b/Engine/Plugins/Mutable/Source/CustomizableObject/Internal/MuCO/CustomizableObjectPrivate.h
--- a/Engine/Plugins/Mutable/Source/CustomizableObject/Internal/MuCO/CustomizableObjectPrivate.h
+++ b/Engine/Plugins/Mutable/Source/CustomizableObject/Internal/MuCO/CustomizableObjectPrivate.h
@@ -1711,6 +1711,8 @@ enum class ECustomizableObjectVersions : int32
 	
 	ConvertBoneIdsToBoneNames,
 
+	AddLayerIndexToMatParamTextures,
+
 	// -----<new versions can be added above this line>--------
 	LastCustomizableObjectVersion
 };
diff --git a/Engine/Plugins/Mutable/Source/CustomizableObjectEditor/Private/MuCOE/GenerateMutableSource/GenerateMutableSource.cpp b/Engine/Plugins/Mutable/Source/CustomizableObjectEditor/Private/MuCOE/GenerateMutableSource/GenerateMutableSource.cpp
--- a/Engine/Plugins/Mutable/Source/CustomizableObjectEditor/Private/MuCOE/GenerateMutableSource/GenerateMutableSource.cpp
+++ b/Engine/Plugins/Mutable/Source/CustomizableObjectEditor/Private/MuCOE/GenerateMutableSource/GenerateMutableSource.cpp
@@ -412,9 +412,8 @@ uint32 GetTypeHash(const FGeneratedKey& Key)
 
 uint32 GetTypeHash(const FMaterialBreakParameter& MaterialBreakParameter)
 {
-	uint32 Hash = GetTypeHash(MaterialBreakParameter.ParameterName.ToString());
+	uint32 Hash = GetTypeHash(MaterialBreakParameter.ParameterKey);
 	Hash = HashCombine(Hash, GetTypeHash(MaterialBreakParameter.ParameterType));
-	Hash = HashCombine(Hash, GetTypeHash(MaterialBreakParameter.LayerIndex));
 	return Hash;
 }
 
diff --git a/Engine/Plugins/Mutable/Source/CustomizableObjectEditor/Private/MuCOE/GenerateMutableSource/GenerateMutableSource.h b/Engine/Plugins/Mutable/Source/CustomizableObjectEditor/Private/MuCOE/GenerateMutableSource/GenerateMutableSource.h
--- a/Engine/Plugins/Mutable/Source/CustomizableObjectEditor/Private/MuCOE/GenerateMutableSource/GenerateMutableSource.h
+++ b/Engine/Plugins/Mutable/Source/CustomizableObjectEditor/Private/MuCOE/GenerateMutableSource/GenerateMutableSource.h
@@ -146,14 +146,11 @@ struct FMorphNodeData
 /** Representation of the root of a Material Break node compilation branch. */
 struct FMaterialBreakParameter 
 {
-	/** The name of the material parameter. */
-	FName ParameterName;
+	/** Name and Layer Index of the material parameter. */
+	UE::Mutable::Private::FParameterKey ParameterKey;
 	
 	/** The type of data this parameter relates to. */
 	FName ParameterType;
-	
-	/** The index of the parameter layer. */
-	int32 LayerIndex;
 
 	bool operator==(const FMaterialBreakParameter& OtherItem) const = default;
 };
diff --git a/Engine/Plugins/Mutable/Source/CustomizableObjectEditor/Private/MuCOE/GenerateMutableSource/GenerateMutableSourceColor.cpp b/Engine/Plugins/Mutable/Source/CustomizableObjectEditor/Private/MuCOE/GenerateMutableSource/GenerateMutableSourceColor.cpp
--- a/Engine/Plugins/Mutable/Source/CustomizableObjectEditor/Private/MuCOE/GenerateMutableSource/GenerateMutableSourceColor.cpp
+++ b/Engine/Plugins/Mutable/Source/CustomizableObjectEditor/Private/MuCOE/GenerateMutableSource/GenerateMutableSourceColor.cpp
@@ -409,7 +409,7 @@ UE::Mutable::Private::Ptr<UE::Mutable::Private::NodeColour> GenerateMutableSourc
 				MaterialBreak->ParameterKey = {ParameterName, LayerIndex};
 
 				// Generate the material that is going to be broken into parameters
-				GenerationContext.MaterialBreakParameterStack.Add(FMaterialBreakParameter{ ParameterName, Pin->PinType.PinCategory, LayerIndex });
+				GenerationContext.MaterialBreakParameterStack.Add(FMaterialBreakParameter{ {ParameterName, LayerIndex}, Pin->PinType.PinCategory });
 				UE::Mutable::Private::Ptr<UE::Mutable::Private::NodeMaterial> LinkedNode = GenerateMutableSourceMaterial(LinkedMaterial, GenerationContext, FGenerationMaterialOptions());
 				MaterialBreak->MaterialSource = LinkedNode;
 				GenerationContext.MaterialBreakParameterStack.Pop();
diff --git a/Engine/Plugins/Mutable/Source/CustomizableObjectEditor/Private/MuCOE/GenerateMutableSource/GenerateMutableSourceFloat.cpp b/Engine/Plugins/Mutable/Source/CustomizableObjectEditor/Private/MuCOE/GenerateMutableSource/GenerateMutableSourceFloat.cpp
--- a/Engine/Plugins/Mutable/Source/CustomizableObjectEditor/Private/MuCOE/GenerateMutableSource/GenerateMutableSourceFloat.cpp
+++ b/Engine/Plugins/Mutable/Source/CustomizableObjectEditor/Private/MuCOE/GenerateMutableSource/GenerateMutableSourceFloat.cpp
@@ -429,7 +429,7 @@ UE::Mutable::Private::Ptr<UE::Mutable::Private::NodeScalar> GenerateMutableSourc
 				MaterialBreak->ParameterKey = {ParameterName, LayerIndex};
 
 				// Generate the material that is going to be broken into parameters
-				GenerationContext.MaterialBreakParameterStack.Add(FMaterialBreakParameter{ ParameterName, Pin->PinType.PinCategory, LayerIndex });
+				GenerationContext.MaterialBreakParameterStack.Add(FMaterialBreakParameter{ {ParameterName, LayerIndex}, Pin->PinType.PinCategory});
 				UE::Mutable::Private::Ptr<UE::Mutable::Private::NodeMaterial> LinkedNode = GenerateMutableSourceMaterial(LinkedMaterial, GenerationContext, FGenerationMaterialOptions());
 				MaterialBreak->MaterialSource = LinkedNode;
 				GenerationContext.MaterialBreakParameterStack.Pop();
diff --git a/Engine/Plugins/Mutable/Source/CustomizableObjectEditor/Private/MuCOE/GenerateMutableSource/GenerateMutableSourceImage.cpp b/Engine/Plugins/Mutable/Source/CustomizableObjectEditor/Private/MuCOE/GenerateMutableSource/GenerateMutableSourceImage.cpp
--- a/Engine/Plugins/Mutable/Source/CustomizableObjectEditor/Private/MuCOE/GenerateMutableSource/GenerateMutableSourceImage.cpp
+++ b/Engine/Plugins/Mutable/Source/CustomizableObjectEditor/Private/MuCOE/GenerateMutableSource/GenerateMutableSourceImage.cpp
@@ -1251,7 +1251,7 @@ UE::Mutable::Private::Ptr<UE::Mutable::Private::NodeImage> GenerateMutableSource
 				GenerationOptions.ReferenceTextureSize = ReferenceTextureSize;
 
 				// Generate the material that is going to be broken into parameters
-				GenerationContext.MaterialBreakParameterStack.Add(FMaterialBreakParameter{ ParameterName, Pin->PinType.PinCategory, LayerIndex });
+				GenerationContext.MaterialBreakParameterStack.Add(FMaterialBreakParameter{ {ParameterName, LayerIndex}, Pin->PinType.PinCategory});
 				UE::Mutable::Private::Ptr<UE::Mutable::Private::NodeMaterial> LinkedNode = GenerateMutableSourceMaterial(LinkedMaterial, GenerationContext, GenerationOptions);
 				MaterialBreak->MaterialSource = LinkedNode;
 				GenerationContext.MaterialBreakParameterStack.Pop();
@@ -1262,7 +1262,7 @@ UE::Mutable::Private::Ptr<UE::Mutable::Private::NodeImage> GenerateMutableSource
 				}
 				else
 				{
-					GenerationContext.Log(LOCTEXT("MaterialBreakNodeError", "Something went wrong generation the Material Node linked to the Material Break Node."), Node);
+					GenerationContext.Log(LOCTEXT("MaterialBreakNodeError", "An error occurred generating the Material Source."), Node, EMessageSeverity::Error);
 				}
 			}
 			else
diff --git a/Engine/Plugins/Mutable/Source/CustomizableObjectEditor/Private/MuCOE/GenerateMutableSource/GenerateMutableSourceMaterial.cpp b/Engine/Plugins/Mutable/Source/CustomizableObjectEditor/Private/MuCOE/GenerateMutableSource/GenerateMutableSourceMaterial.cpp
--- a/Engine/Plugins/Mutable/Source/CustomizableObjectEditor/Private/MuCOE/GenerateMutableSource/GenerateMutableSourceMaterial.cpp
+++ b/Engine/Plugins/Mutable/Source/CustomizableObjectEditor/Private/MuCOE/GenerateMutableSource/GenerateMutableSourceMaterial.cpp
@@ -398,8 +398,8 @@ UE::Mutable::Private::Ptr<UE::Mutable::Private::NodeMaterial> GenerateMutableSou
 			if (!CurrentBreakParameter.ParameterType.IsNone())
 			{
 				FHashedMaterialParameterInfo ParameterInfo;
-				ParameterInfo.Name = FScriptName(CurrentBreakParameter.ParameterName);
-				ParameterInfo.Index = CurrentBreakParameter.LayerIndex;
+				ParameterInfo.Name = FScriptName(CurrentBreakParameter.ParameterKey.ParameterName);
+				ParameterInfo.Index = CurrentBreakParameter.ParameterKey.LayerIndex;
 				ParameterInfo.Association = ParameterInfo.Index == INDEX_NONE ? GlobalParameter : LayerParameter;
 
 				if (!TypedNodeMaterialConstant->Material)
@@ -434,7 +434,7 @@ UE::Mutable::Private::Ptr<UE::Mutable::Private::NodeMaterial> GenerateMutableSou
 
 							if (UE::Mutable::Private::Ptr<UE::Mutable::Private::NodeImage> FinalNodeImag = ResizeTextureByNumMips(ConstantImageNode, MipsToSkip))
 							{
-								MaterialNode->ImageValues.Add({CurrentBreakParameter.ParameterName, CurrentBreakParameter.LayerIndex}, FinalNodeImag);
+								MaterialNode->ImageValues.Add(CurrentBreakParameter.ParameterKey, FinalNodeImag);
 							}
 						}
 					}
@@ -449,7 +449,7 @@ UE::Mutable::Private::Ptr<UE::Mutable::Private::NodeMaterial> GenerateMutableSou
 						{
 							UE::Mutable::Private::Ptr<UE::Mutable::Private::NodeImageConstant> ImageNode = new UE::Mutable::Private::NodeImageConstant();
 							ImageNode->Image = GenerateImageConstant(PassthroughTexture, GenerationContext, true);
-							MaterialNode->ImageValues.Add({CurrentBreakParameter.ParameterName, CurrentBreakParameter.LayerIndex}, ImageNode);
+							MaterialNode->ImageValues.Add(CurrentBreakParameter.ParameterKey, ImageNode);
 						}
 					}
 
@@ -458,7 +458,7 @@ UE::Mutable::Private::Ptr<UE::Mutable::Private::NodeMaterial> GenerateMutableSou
 						FLinearColor Color;
 						if (Material->GetVectorParameterValue(ParameterInfo, Color))
 						{
-							MaterialNode->ColorValues.Add({CurrentBreakParameter.ParameterName, CurrentBreakParameter.LayerIndex}, Color);
+							MaterialNode->ColorValues.Add(CurrentBreakParameter.ParameterKey, Color);
 						}
 					}
 
@@ -467,7 +467,7 @@ UE::Mutable::Private::Ptr<UE::Mutable::Private::NodeMaterial> GenerateMutableSou
 						float Scalar;
 						if (Material->GetScalarParameterValue(ParameterInfo, Scalar))
 						{
-							MaterialNode->ScalarValues.Add({CurrentBreakParameter.ParameterName, CurrentBreakParameter.LayerIndex}, Scalar);
+							MaterialNode->ScalarValues.Add(CurrentBreakParameter.ParameterKey, Scalar);
 						}
 					}
 				}
@@ -588,7 +588,8 @@ UE::Mutable::Private::Ptr<UE::Mutable::Private::NodeMaterial> GenerateMutableSou
 
 		// Store the parameter that needs to be processed
 		const FMaterialBreakParameter CurrentBreakParameter = GenerationContext.GetCurrentMaterialBreakParameter();
-		if (!CurrentBreakParameter.ParameterType.IsNone())
+		
+		if (!CurrentBreakParameter.ParameterType.IsNone() && CurrentBreakParameter.ParameterType == UEdGraphSchema_CustomizableObject::PC_Texture)
 		{
 			if (!MaterialParameterNode->ReferenceValue)
 			{
@@ -598,42 +599,39 @@ UE::Mutable::Private::Ptr<UE::Mutable::Private::NodeMaterial> GenerateMutableSou
 			{
 				// Data to find the parameter value inside the material
 				FHashedMaterialParameterInfo ParameterInfo;
-				ParameterInfo.Name = FScriptName(CurrentBreakParameter.ParameterName);
-				ParameterInfo.Index = CurrentBreakParameter.LayerIndex;
+				ParameterInfo.Name = FScriptName(CurrentBreakParameter.ParameterKey.ParameterName);
+				ParameterInfo.Index = CurrentBreakParameter.ParameterKey.LayerIndex;
 				ParameterInfo.Association = ParameterInfo.Index == INDEX_NONE ? GlobalParameter : LayerParameter;
 
-				if (CurrentBreakParameter.ParameterType == UEdGraphSchema_CustomizableObject::PC_Texture)
-				{
-					//TODO(Max - GM): Move this to the generatemutablesourceimage?
-					UE::Mutable::Private::Ptr<UE::Mutable::Private::NodeImageFromMaterialParameter> TextureNode = new UE::Mutable::Private::NodeImageFromMaterialParameter();
-					TextureNode->ImageParameterName = CurrentBreakParameter.ParameterName;
+				//TODO(Max - GM): Move this to the generatemutablesourceimage?
+				UE::Mutable::Private::Ptr<UE::Mutable::Private::NodeImageFromMaterialParameter> TextureNode = new UE::Mutable::Private::NodeImageFromMaterialParameter();
+				TextureNode->ImageParameterKey = CurrentBreakParameter.ParameterKey;
 
-					// Resize image node.
-					UE::Mutable::Private::Ptr<UE::Mutable::Private::NodeImageResize> ResizeNode = new UE::Mutable::Private::NodeImageResize();
-					ResizeNode->Base = TextureNode;
-					ResizeNode->bRelative = false;
+				// Resize image node.
+				UE::Mutable::Private::Ptr<UE::Mutable::Private::NodeImageResize> ResizeNode = new UE::Mutable::Private::NodeImageResize();
+				ResizeNode->Base = TextureNode;
+				ResizeNode->bRelative = false;
 
-					UE::Mutable::Private::FImageSize TextureSize(1);
+				UE::Mutable::Private::FImageSize TextureSize(1);
 
-					//Find reference Texture from the material parameter
-					UTexture* Texture = nullptr;
-					bool bParameterFound = MaterialParameterNode->ReferenceValue->GetTextureParameterValue(ParameterInfo, Texture);
-					UTexture2D* ReferenceTexture = Cast<UTexture2D>(Texture);
+				//Find reference Texture from the material parameter
+				UTexture* Texture = nullptr;
+				bool bParameterFound = MaterialParameterNode->ReferenceValue->GetTextureParameterValue(ParameterInfo, Texture);
+				UTexture2D* ReferenceTexture = Cast<UTexture2D>(Texture);
 
-					if (bParameterFound && ReferenceTexture)
-					{
-						const uint32 LODBias = ComputeLODBiasForTexture(GenerationContext, *ReferenceTexture, nullptr, MaterialOptions.ReferenceTextureSize);
-						TextureSize.X = FMath::Max(ReferenceTexture->Source.GetSizeX() >> LODBias, 1);
-						TextureSize.Y = FMath::Max(ReferenceTexture->Source.GetSizeY() >> LODBias, 1);
-					}
+				if (bParameterFound && ReferenceTexture)
+				{
+					const uint32 LODBias = ComputeLODBiasForTexture(GenerationContext, *ReferenceTexture, nullptr, MaterialOptions.ReferenceTextureSize);
+					TextureSize.X = FMath::Max(ReferenceTexture->Source.GetSizeX() >> LODBias, 1);
+					TextureSize.Y = FMath::Max(ReferenceTexture->Source.GetSizeY() >> LODBias, 1);
+				}
 
-					ResizeNode->SizeX = TextureSize.X;
-					ResizeNode->SizeY = TextureSize.Y;
+				ResizeNode->SizeX = TextureSize.X;
+				ResizeNode->SizeY = TextureSize.Y;
 
-					// Force the same format that the default texture if any.
-					UE::Mutable::Private::Ptr<UE::Mutable::Private::NodeImageFormat> FormatNode = GenerateMutableImageFormat(ResizeNode, ReferenceTexture, GenerationContext, Node);
-					MaterialNode->ImageParametersToCompile.Add({ FName(ParameterInfo.Name), ParameterInfo.Index }, FormatNode);
-				}
+				// Force the same format that the default texture if any.
+				UE::Mutable::Private::Ptr<UE::Mutable::Private::NodeImageFormat> FormatNode = GenerateMutableImageFormat(ResizeNode, ReferenceTexture, GenerationContext, Node);
+				MaterialNode->ImageParametersToCompile.Add(CurrentBreakParameter.ParameterKey, FormatNode);
 			}
 		}
 
@@ -643,8 +641,8 @@ UE::Mutable::Private::Ptr<UE::Mutable::Private::NodeMaterial> GenerateMutableSou
 	else if (const UCONodeSkeletalMeshBreak* TypedSkeletalMeshBreakNode = Cast<UCONodeSkeletalMeshBreak>(Node))
 	{
 		const UCONodeSkeletalMeshBreakPinData& OverrideMaterialPinData = TypedSkeletalMeshBreakNode->GetPinData<UCONodeSkeletalMeshBreakPinData>(*Pin);
-
 		const FName TargetSlotName = OverrideMaterialPinData.TargetMaterialSlotName;
+
 		if (!TargetSlotName.IsNone())
 		{
 			if (const UEdGraphPin* ConnectedMeshPin = FollowInputPin(*TypedSkeletalMeshBreakNode->InputPassthroughMesh.Get()))
@@ -654,6 +652,17 @@ UE::Mutable::Private::Ptr<UE::Mutable::Private::NodeMaterial> GenerateMutableSou
 				MaterialFromSlotNode->SlotName = TargetSlotName;
 
 				Result = MaterialFromSlotNode;
+
+				// Log unsuported textures error if needed
+				const FMaterialBreakParameter CurrentBreakParameter = GenerationContext.GetCurrentMaterialBreakParameter();
+				if (!CurrentBreakParameter.ParameterType.IsNone() && CurrentBreakParameter.ParameterType == UEdGraphSchema_CustomizableObject::PC_Texture)
+				{
+					FText Msg = FText::Format(LOCTEXT("NodeSKMBreak_UnsupportedOperation", "Breaking the Skeletal Mesh Material parameter [{0}] into a Mutable Texture is currently unsupported. Use it as a passthrough texture."),
+						FText::FromName(CurrentBreakParameter.ParameterKey.ParameterName));
+					GenerationContext.Log(Msg, Node, EMessageSeverity::Error);
+
+					Result = nullptr;
+				}
 			}
 		}
 	}
diff --git a/Engine/Plugins/Mutable/Source/MutableRuntime/Internal/MuR/Operations.h b/Engine/Plugins/Mutable/Source/MutableRuntime/Internal/MuR/Operations.h
--- a/Engine/Plugins/Mutable/Source/MutableRuntime/Internal/MuR/Operations.h
+++ b/Engine/Plugins/Mutable/Source/MutableRuntime/Internal/MuR/Operations.h
@@ -1207,6 +1207,7 @@ namespace UE::Mutable::Private
 		{
 			ADDRESS MaterialParameter;
 			CONSTANT_STRING_INDEX ParameterName;
+			int32 LayerIndex;
 		};
 
     	struct ExternalArgs
diff --git a/Engine/Plugins/Mutable/Source/MutableRuntime/Private/MuR/CodeRunner.cpp b/Engine/Plugins/Mutable/Source/MutableRuntime/Private/MuR/CodeRunner.cpp
--- a/Engine/Plugins/Mutable/Source/MutableRuntime/Private/MuR/CodeRunner.cpp
+++ b/Engine/Plugins/Mutable/Source/MutableRuntime/Private/MuR/CodeRunner.cpp
@@ -8739,7 +8739,12 @@ namespace UE::Mutable::Private
 			// Get the parameter texture from the UMaterial
 			if (MaterialInterface)
 			{
-				MaterialInterface->GetTextureParameterValue(FName(ParameterName), Image);
+				FHashedMaterialParameterInfo ParameterInfo;
+				ParameterInfo.Name = FScriptName(FName(ParameterName));
+				ParameterInfo.Index = Args.LayerIndex;
+				ParameterInfo.Association = ParameterInfo.Index == INDEX_NONE ? GlobalParameter : LayerParameter;
+
+				MaterialInterface->GetTextureParameterValue(ParameterInfo, Image);
 			}
 
 			FExtendedImageDesc& Result = ImageDescResults.FindOrAdd(Item.At);
diff --git a/Engine/Plugins/Mutable/Source/MutableTools/Internal/MuT/NodeImageFromMaterialParameter.h b/Engine/Plugins/Mutable/Source/MutableTools/Internal/MuT/NodeImageFromMaterialParameter.h
--- a/Engine/Plugins/Mutable/Source/MutableTools/Internal/MuT/NodeImageFromMaterialParameter.h
+++ b/Engine/Plugins/Mutable/Source/MutableTools/Internal/MuT/NodeImageFromMaterialParameter.h
@@ -5,6 +5,7 @@
 #include "MuR/RefCounted.h"
 #include "MuT/Node.h"
 #include "MuT/NodeImage.h"
+#include "MuR/Material.h"
 
 #define UE_API MUTABLETOOLS_API
 
@@ -14,7 +15,7 @@ namespace UE::Mutable::Private
 	{
 	public:
 
-		FName ImageParameterName;
+		FParameterKey ImageParameterKey;
 
 	public:
 
diff --git a/Engine/Plugins/Mutable/Source/MutableTools/Private/MuT/ASTOpMaterialBreakImageParameter.cpp b/Engine/Plugins/Mutable/Source/MutableTools/Private/MuT/ASTOpMaterialBreakImageParameter.cpp
--- a/Engine/Plugins/Mutable/Source/MutableTools/Private/MuT/ASTOpMaterialBreakImageParameter.cpp
+++ b/Engine/Plugins/Mutable/Source/MutableTools/Private/MuT/ASTOpMaterialBreakImageParameter.cpp
@@ -29,7 +29,7 @@ namespace UE::Mutable::Private
 		{
 			const ASTOpImageFromMaterialParameter* Other = static_cast<const ASTOpImageFromMaterialParameter*>(&OtherUntyped);
 			return MaterialParameter == Other->MaterialParameter &&
-				ParameterName == Other->ParameterName;
+				ParameterKey == Other->ParameterKey;
 		}
 		return false;
 	}
@@ -38,7 +38,7 @@ namespace UE::Mutable::Private
 	UE::Mutable::Private::Ptr<ASTOp> ASTOpImageFromMaterialParameter::Clone(MapChildFuncRef MapChild) const
 	{
 		Ptr<ASTOpImageFromMaterialParameter> Result = new ASTOpImageFromMaterialParameter();
-		Result->ParameterName = ParameterName;
+		Result->ParameterKey = ParameterKey;
 		Result->MaterialParameter = MaterialParameter;
 
 		return Result;
@@ -55,7 +55,7 @@ namespace UE::Mutable::Private
 	{
 		uint32 Result = GetTypeHash(GetOpType());
 		
-		Result = HashCombineFast(Result, GetTypeHash(ParameterName));
+		Result = HashCombineFast(Result, GetTypeHash(ParameterKey));
 		Result = HashCombineFast(Result, GetTypeHash(MaterialParameter.Name));
 		
 		return Result;
@@ -84,7 +84,8 @@ namespace UE::Mutable::Private
 			FMemory::Memzero(Args);
 
 			Args.MaterialParameter = OP::ADDRESS(LinkedParameterIndex);
-			Args.ParameterName = Program.AddConstant(ParameterName.ToString());
+			Args.ParameterName = Program.AddConstant(ParameterKey.ParameterName.ToString());
+			Args.LayerIndex = ParameterKey.LayerIndex;
 
 			AppendCode(Program.ByteCode, Args);
 		}
diff --git a/Engine/Plugins/Mutable/Source/MutableTools/Private/MuT/ASTOpMaterialBreakImageParameter.h b/Engine/Plugins/Mutable/Source/MutableTools/Private/MuT/ASTOpMaterialBreakImageParameter.h
--- a/Engine/Plugins/Mutable/Source/MutableTools/Private/MuT/ASTOpMaterialBreakImageParameter.h
+++ b/Engine/Plugins/Mutable/Source/MutableTools/Private/MuT/ASTOpMaterialBreakImageParameter.h
@@ -18,7 +18,7 @@ namespace UE::Mutable::Private
 	public:
 
 		FParameterDesc MaterialParameter;
-		FName ParameterName;
+		FParameterKey ParameterKey;
 
 	public:
 
diff --git a/Engine/Plugins/Mutable/Source/MutableTools/Private/MuT/ASTOpParameter.h b/Engine/Plugins/Mutable/Source/MutableTools/Private/MuT/ASTOpParameter.h
--- a/Engine/Plugins/Mutable/Source/MutableTools/Private/MuT/ASTOpParameter.h
+++ b/Engine/Plugins/Mutable/Source/MutableTools/Private/MuT/ASTOpParameter.h
@@ -29,7 +29,7 @@ namespace UE::Mutable::Private
 		//** Ranges adding dimensions to this parameter. */
 		TArray<FRangeData> Ranges;
 
-		/** Image Parameters are evaluated at runtime. */
+		/** Image Parameters which are evaluated at runtime. */
 		TMap<FParameterKey, ASTChild> ImageOperations;
 
 		/** Index of the parameter in the final program parameter list. 
diff --git a/Engine/Plugins/Mutable/Source/MutableTools/Private/MuT/CodeGenerator_Image.cpp b/Engine/Plugins/Mutable/Source/MutableTools/Private/MuT/CodeGenerator_Image.cpp
--- a/Engine/Plugins/Mutable/Source/MutableTools/Private/MuT/CodeGenerator_Image.cpp
+++ b/Engine/Plugins/Mutable/Source/MutableTools/Private/MuT/CodeGenerator_Image.cpp
@@ -1994,7 +1994,7 @@ namespace UE::Mutable::Private
 	{
 		Ptr<ASTOpImageFromMaterialParameter> ImageOp = new ASTOpImageFromMaterialParameter();
 		ImageOp->MaterialParameter = InOptions.MaterialParameter;
-		ImageOp->ParameterName = InNode->ImageParameterName;
+		ImageOp->ParameterKey = InNode->ImageParameterKey;
 
 		Result.op = ImageOp;
 	}
