[mutable] Simplified calls to BeginUpdate_MutableThread and removed some info from the Context
- Removed OldState and OldParameters from the Update context, as they are no longer needed.
- Fixed a pontential issue with the removal of FMesh objects from the cache that prevented the removal of some if they were marked as locked.
- Made the "bDisableMeshStreaming" and "bDisableTextureStreaming" toggles non interactuable when the LiveUpdate has been enabled. If the LiveUpdate is enabled, nor Mesh nor Texture streaming can be enabled.

#rnx

#Synced-CL 50195280

diff --git a/Engine/Plugins/Mutable/Source/CustomizableObject/Internal/MuCO/CustomizableObjectSystemPrivate.h b/Engine/Plugins/Mutable/Source/CustomizableObject/Internal/MuCO/CustomizableObjectSystemPrivate.h
--- a/Engine/Plugins/Mutable/Source/CustomizableObject/Internal/MuCO/CustomizableObjectSystemPrivate.h
+++ b/Engine/Plugins/Mutable/Source/CustomizableObject/Internal/MuCO/CustomizableObjectSystemPrivate.h
@@ -570,9 +570,6 @@ public:
 	/** Manage loading Passthrough Objects discovered on the StartUpdate. */
 	TSharedPtr<UE::Mutable::Private::FPassthroughObjectLoader> PassthroughObjectLoader;
 	
-	/** The parameters the Live Instance had before being reused. */
-	TSharedPtr<const UE::Mutable::Private::FParameters> OldParameters;
-	
 	/** Weak reference to the instance we are operating on.
 	 *It is weak because we don't want to lock it in case it becomes irrelevant in the game while operations are pending and it needs to be destroyed. */
 	TWeakObjectPtr<UCustomizableObjectInstance> Instance;
@@ -605,9 +602,6 @@ public:
 	/** This option comes from the operation request. It is used to reduce the number of mipmaps that mutable must generate for images.  */
 	int32 MipsToSkip = 0;
 	
-	/** The state the Live Instance had before being reused. */
-	int32 OldState = -1;
-	
 	EQueuePriorityType PriorityType = EQueuePriorityType::Low;
 
 	EUpdateResult UpdateResult = EUpdateResult::Success;
@@ -661,6 +655,11 @@ public:
 
 	/** true if the update has been optimized (skips all Tasks and calls FinishUpdateGlobal directly on the Enqueue). */
 	uint8 bOptimizedUpdate : 1 = false;
+	
+	/** Determines if the Cachelayer1 (locked and unlocked data) should be cleared or not during the execution of the BeginUpdate_MutableThread. 
+	 * It is only useful during the initial generation as we do not need to consider clearing the layer one during LOD or MIP generation
+	 */
+	bool bClearCacheLayer1 = false;
 };
 
 
diff --git a/Engine/Plugins/Mutable/Source/CustomizableObject/Private/MuCO/CustomizableObjectMeshUpdate.cpp b/Engine/Plugins/Mutable/Source/CustomizableObject/Private/MuCO/CustomizableObjectMeshUpdate.cpp
--- a/Engine/Plugins/Mutable/Source/CustomizableObject/Private/MuCO/CustomizableObjectMeshUpdate.cpp
+++ b/Engine/Plugins/Mutable/Source/CustomizableObject/Private/MuCO/CustomizableObjectMeshUpdate.cpp
@@ -249,10 +249,7 @@ namespace impl
 		
 		// Main instance generation step
 		// note: as this method has already been run we provide the param and stats of the instance to prevent the full clearing of its contents
-		UE::Mutable::Private::TManagedPtr<const UE::Mutable::Private::FInstance> Instance = System->BeginUpdate_MutableThread(LiveInstance,
-			OperationData->LiveInstance->Parameters,
-			OperationData->LiveInstance->State,
-			LODMask);
+		UE::Mutable::Private::TManagedPtr<const UE::Mutable::Private::FInstance> Instance = System->BeginUpdate_MutableThread(LiveInstance, LODMask, false);
 		
 		check(Instance);
 
diff --git a/Engine/Plugins/Mutable/Source/CustomizableObject/Private/MuCO/CustomizableObjectMipDataProvider.cpp b/Engine/Plugins/Mutable/Source/CustomizableObject/Private/MuCO/CustomizableObjectMipDataProvider.cpp
--- a/Engine/Plugins/Mutable/Source/CustomizableObject/Private/MuCO/CustomizableObjectMipDataProvider.cpp
+++ b/Engine/Plugins/Mutable/Source/CustomizableObject/Private/MuCO/CustomizableObjectMipDataProvider.cpp
@@ -145,11 +145,8 @@ namespace Impl
 			{
 				LiveInstance->Cache->Clear();
 			}
-			
-			Instance = System->BeginUpdate_MutableThread(LiveInstance,
-				LiveInstance->Parameters,
-				LiveInstance->State,
-				LODMask);
+
+			Instance = System->BeginUpdate_MutableThread(LiveInstance, LODMask, false);
 			
 			check(Instance);
 		}
diff --git a/Engine/Plugins/Mutable/Source/CustomizableObject/Private/MuCO/CustomizableObjectSystem.cpp b/Engine/Plugins/Mutable/Source/CustomizableObject/Private/MuCO/CustomizableObjectSystem.cpp
--- a/Engine/Plugins/Mutable/Source/CustomizableObject/Private/MuCO/CustomizableObjectSystem.cpp
+++ b/Engine/Plugins/Mutable/Source/CustomizableObject/Private/MuCO/CustomizableObjectSystem.cpp
@@ -2272,6 +2272,8 @@ namespace impl
 		{
 			OperationData->MutableSystem->SetImagePixelConversionOverride( OperationData->PixelFormatOverride );
 		}
+
+		check(OperationData->LiveInstance);
 		
 		// Prevent the usage of the already cached elements by fully clearing the cache.
 		if (!CVarReuseLiveInstanceCacheBetweenUpdates.GetValueOnAnyThread())
@@ -2279,14 +2281,10 @@ namespace impl
 			OperationData->LiveInstance->Cache->Clear();
 		}
 		
-		UE::Mutable::Private::TManagedPtr<const UE::Mutable::Private::FInstance> Instance = OperationData->MutableSystem->BeginUpdate_MutableThread(OperationData->LiveInstance.ToSharedRef(),
-			OperationData->OldParameters,
-			OperationData->OldState,
-			UE::Mutable::Private::FSystem::AllLODs);
-
-		// This data is no longer useful
-		OperationData->OldParameters = nullptr;
-		OperationData->OldState = -1;
+		UE::Mutable::Private::TManagedPtr<const UE::Mutable::Private::FInstance> Instance = OperationData->MutableSystem->BeginUpdate_MutableThread(
+			OperationData->LiveInstance.ToSharedRef(),
+			UE::Mutable::Private::FSystem::AllLODs,
+			OperationData->bClearCacheLayer1);
 		
 		OperationData->MutableInstance = Instance;
 		
@@ -3029,7 +3027,7 @@ namespace impl
 		}
 		else
 		{
-			Operation->bNeverStreamMips = StateData ? StateData->bDisableTextureStreaming : false;
+			Operation->bNeverStreamMips = StateData ? StateData->bLiveUpdateMode || StateData->bDisableTextureStreaming : false;
 		}
 #else
 		Operation->bNeverStreamMips = true;
@@ -3085,7 +3083,6 @@ namespace impl
 		
 		// Task: Mutable Update and GetMesh
 		//-------------------------------------------------------------
-		Operation->bStreamMeshLODs = IsStreamingEnabled(*Object, State);
 
 #if WITH_EDITOR
 		Operation->PixelFormatOverride = SystemPrivate->ImageFormatOverrideFunc;
@@ -3164,14 +3161,13 @@ namespace impl
 
 			Operation->bLiveUpdateMode = SystemPrivate->EnableMutableLiveUpdate && (StateData ? StateData->bLiveUpdateMode : false);
 		
-			// todo : make the stream settings in Operation respond to the settings in Operation->bLiveUpdateMode
-			if (Operation->bLiveUpdateMode && 
-				(Operation->bStreamMeshLODs || !Operation->bNeverStreamMips || Operation->MipsToSkip > 0))
-			{
-				UE_LOG(LogMutable, Warning, TEXT("Instance LiveUpdateMode does not yet support progressive streaming of Mutable textures or meshes. Disabling LiveUpdateMode for this update."));
-				Operation->bLiveUpdateMode = false;
-			}
-
+			Operation->bStreamMeshLODs = Operation->bLiveUpdateMode ? false : IsStreamingEnabled(*Object, State);
+			
+			// Ensure that in the case of running with the live update enabled we have mesh and mip streaming disabled
+			check(!Operation->bLiveUpdateMode || Operation->bNeverStreamMips);
+			check(!Operation->bLiveUpdateMode || Operation->MipsToSkip == 0);
+			check(!Operation->bLiveUpdateMode || !Operation->bStreamMeshLODs);
+			
 			Operation->LiveInstance = MutableSystem->NewInstance(
 				Operation->Model,
 				ObjectPrivate->MeshIdRegistry,
@@ -3188,12 +3184,15 @@ namespace impl
 			{
 				if (Instance->GetPrivate()->LiveInstance.IsValid())
 				{
+					// If we allow the cache to be reused then grab the one from the old instance and prepare it for this one
 					// Reuse the cache of the live instance that was added during a previous update
 					Operation->LiveInstance->Cache = Instance->GetPrivate()->LiveInstance->Cache;
-				
-					// Keep the old params and state to be able to later determine if the reused cache requires clearing before being used
-					Operation->OldParameters = Instance->GetPrivate()->LiveInstance->Parameters;
-					Operation->OldState = Instance->GetPrivate()->LiveInstance->State;
+					
+					const TSharedPtr<const UE::Mutable::Private::FParameters> PreviousParameters = Instance->GetPrivate()->LiveInstance->Parameters;
+					const int32 PreviousState = Instance->GetPrivate()->LiveInstance->State;
+					
+					Operation->bClearCacheLayer1 = PreviousState != Operation->LiveInstance->State || 
+						UE::Mutable::Private::FSystem::CheckForUpdatedParameters(Operation->LiveInstance.ToSharedRef(), PreviousParameters);
 				}
 				
 				Instance->GetPrivate()->LiveInstance = Operation->LiveInstance;
diff --git a/Engine/Plugins/Mutable/Source/CustomizableObjectEditor/Private/MuCOE/Nodes/CustomizableObjectNodeObject.h b/Engine/Plugins/Mutable/Source/CustomizableObjectEditor/Private/MuCOE/Nodes/CustomizableObjectNodeObject.h
--- a/Engine/Plugins/Mutable/Source/CustomizableObjectEditor/Private/MuCOE/Nodes/CustomizableObjectNodeObject.h
+++ b/Engine/Plugins/Mutable/Source/CustomizableObjectEditor/Private/MuCOE/Nodes/CustomizableObjectNodeObject.h
@@ -40,11 +40,11 @@ struct FCustomizableObjectState
 	ETextureCompressionStrategy TextureCompressionStrategy = ETextureCompressionStrategy::None;
 
 	/** If this is enabled, Mesh streaming won't be used for this state, and all LODs will be generated when an instance is first updated. */
-	UPROPERTY(EditAnywhere, Category = CustomizableObject)
+	UPROPERTY(EditAnywhere, Category = CustomizableObject, meta = (EditCondition="!bLiveUpdateMode"))
 	bool bDisableMeshStreaming = false;
 
 	/** If this is enabled, texture streaming won't be used for this state, and full images will be generated when an instance is first updated. */
-	UPROPERTY(EditAnywhere, Category = CustomizableObject)
+	UPROPERTY(EditAnywhere, Category = CustomizableObject, meta = (EditCondition="!bLiveUpdateMode"))
 	bool bDisableTextureStreaming = false;
 
 	/** LiveUpdateMode will reuse instance temp. data between updates and speed up update times, but spend much more memory. Good for customization screens, not for actual gameplay modes. */
diff --git a/Engine/Plugins/Mutable/Source/MutableRuntime/Internal/MuR/System.h b/Engine/Plugins/Mutable/Source/MutableRuntime/Internal/MuR/System.h
--- a/Engine/Plugins/Mutable/Source/MutableRuntime/Internal/MuR/System.h
+++ b/Engine/Plugins/Mutable/Source/MutableRuntime/Internal/MuR/System.h
@@ -493,10 +493,17 @@ namespace UE::Mutable::Private
 			return uint32_t(m_usedRangeIndices.Num()) - 1;
 		}
 
-		
-		void Prepare(const FProgram& InProgram, int32 InState);
 
-		
+		/**
+		 * Sets up the Cache so it can be used by the CodeRunner
+		 * @param InProgram The mutable program
+		 * @param InState The state you want the cache to be "bound" to. It is used to determine which operations should be locked and which not.
+		 * @note When determining which Ops gets locked the state is used as the ones that are part of the current state will be left unlocked while
+		 * the ones from other states will get marked as locked.
+		 */
+		MUTABLERUNTIME_API void Prepare(const UE::Mutable::Private::FProgram& InProgram, int32 InState);
+
+
 		void Init(uint32 Size)
 		{
 			// This clear would prevent live update cache reuse
@@ -791,6 +798,7 @@ namespace UE::Mutable::Private
 			return Data->DataTypeIndex != 0;
 		}
 
+		
 		float GetScalar(FCacheAddress Address)
 		{
 			if (!Address.At)
@@ -808,6 +816,7 @@ namespace UE::Mutable::Private
 			return Data->ScalarResult;
 		}
 
+		
 		int32 GetInt(FCacheAddress Address)
 		{
 			if (!Address.At)
@@ -826,6 +835,7 @@ namespace UE::Mutable::Private
 			return Data->DataTypeIndex;
 		}
 
+		
 		FVector4f GetColor(FCacheAddress Address)
 		{
 			if (!Address.At) 
@@ -843,6 +853,7 @@ namespace UE::Mutable::Private
 			return ColorResults[Data->DataTypeIndex];
 		}
 
+		
 		FMatrix44f GetMatrix(FCacheAddress Address)
 		{
 			if (!Address.At)
@@ -860,6 +871,7 @@ namespace UE::Mutable::Private
 			return MatrixResults[Data->DataTypeIndex];
 		}
 
+		
 		TManagedPtr<const FMaterial> GetMaterial(FCacheAddress at)
 		{
 			if (!at.At)
@@ -884,6 +896,7 @@ namespace UE::Mutable::Private
 			return FoundMaterial;
 		}
 
+		
 		FProjector GetProjector(FCacheAddress at)
 		{
 			if (!at.At) return FProjector();
@@ -960,6 +973,7 @@ namespace UE::Mutable::Private
 			return Result;
 		}
 
+		
 		TManagedPtr<const FMesh> GetMesh(FCacheAddress Address)
 		{
 			if (!Address.At || Address.At >= OpExecutionData.size_code()) 
@@ -980,7 +994,7 @@ namespace UE::Mutable::Private
 			check(Data->OpHitCount > 0);
 
 			--Data->OpHitCount;
-			if (Data->OpHitCount == 0 && !Data->IsCacheLocked)
+			if (Data->OpHitCount == 0)
 			{
 				// If this is the last cached reference, now Result is the only reference, which is given 
 				// to the calling code. 
@@ -1081,6 +1095,7 @@ namespace UE::Mutable::Private
 			return ExtensionDataResults[Data->DataTypeIndex].Pin();
 		}
 		
+		
 		TManagedPtr<const FInstancedStruct> GetInstancedStruct(FCacheAddress Address)
 		{
 			if (!Address.At) 
@@ -1112,6 +1127,7 @@ namespace UE::Mutable::Private
 			Data.IsDescCacheValid = true;
 		}
 
+		
 		void SetBool(FCacheAddress at, bool v)
 		{
 			check(at.At < OpExecutionData.size_code());
@@ -1122,6 +1138,7 @@ namespace UE::Mutable::Private
 			Data->IsValueValid = true;
 		}
 
+		
 		void SetInt(FCacheAddress at, int32 v)
 		{
 			check(at.At < OpExecutionData.size_code());
@@ -1132,6 +1149,7 @@ namespace UE::Mutable::Private
 			Data->IsValueValid = true;
 		}
 
+		
 		void SetScalar(FCacheAddress Address, float v)
 		{
 			check(Address.At < OpExecutionData.size_code());
@@ -1142,6 +1160,7 @@ namespace UE::Mutable::Private
 			Data->IsValueValid = true;
 		}
 
+		
 		void SetColour(FCacheAddress Address, const FVector4f& v)
 		{
 			check(Address.At < OpExecutionData.size_code());
@@ -1162,6 +1181,7 @@ namespace UE::Mutable::Private
 			check(Data->DataTypeIndex != 0);
 		}
 
+		
 		void SetMatrix(FCacheAddress Address, const FMatrix44f& v)
 		{
 			check(Address.At < OpExecutionData.size_code());
@@ -1182,6 +1202,7 @@ namespace UE::Mutable::Private
 			check(Data->DataTypeIndex != 0);
 		}
 
+		
 		void SetMaterial(FCacheAddress Address, const TManagedPtr<const FMaterial>& Value)
 		{
 			check(Value);		// All material pointers provided must be valid
@@ -1203,6 +1224,7 @@ namespace UE::Mutable::Private
 			check(Data->DataTypeIndex != 0);
 		}
 		
+		
 		void SetProjector(FCacheAddress Address, const FProjector& v)
 		{
 			check(Address.At < OpExecutionData.size_code());
@@ -1223,6 +1245,7 @@ namespace UE::Mutable::Private
 			check(Data->DataTypeIndex != 0);
 		}
 
+		
 		void SetInstance(FCacheAddress Address, const TManagedPtr<const FInstance>& Value)
 		{
 			check(Address.At < OpExecutionData.size_code());
@@ -1244,6 +1267,7 @@ namespace UE::Mutable::Private
 			check(Data->DataTypeIndex != 0);
 		}
 
+		
 		void SetExtensionData(FCacheAddress Address, const TManagedPtr<const FExtensionData>& Value)
 		{
 			check(Address.At < OpExecutionData.size_code());
@@ -1264,6 +1288,7 @@ namespace UE::Mutable::Private
 			check(Data->DataTypeIndex != 0);
 		}
 
+		
 		void SetInstancedStruct(FCacheAddress Address, const TManagedPtr<const FInstancedStruct>& Value)
 		{
 			check(Address.At < OpExecutionData.size_code());
@@ -1284,6 +1309,7 @@ namespace UE::Mutable::Private
 			check(Data->DataTypeIndex != 0);
 		}
 		
+		
 		void SetImage(FCacheAddress Address, const TManagedPtr<const FImage>& Value)
 		{
 			check(Address.At < OpExecutionData.size_code());
@@ -1321,6 +1347,7 @@ namespace UE::Mutable::Private
 			UE::Mutable::Private::UpdateLLMStats();
 		}
 
+		
 		void SetMesh(FCacheAddress Address, const TManagedPtr<const FMesh>& Value)
 		{
 			check(Address.At < OpExecutionData.size_code());
@@ -1360,6 +1387,7 @@ namespace UE::Mutable::Private
 			UE::Mutable::Private::UpdateLLMStats();
 		}
 
+		
 		void SetSkeletalMesh(FCacheAddress Address, const TManagedPtr<const FSkeletalMesh>& Value)
 		{
 			check(Address.At < OpExecutionData.size_code());
@@ -1382,6 +1410,7 @@ namespace UE::Mutable::Private
 			UE::Mutable::Private::UpdateLLMStats();
 		}
 
+		
 		void SetLayout(FCacheAddress Address, const TManagedPtr<const FLayout>& Value)
 		{
 			check(Address.At < OpExecutionData.size_code());
@@ -1404,6 +1433,7 @@ namespace UE::Mutable::Private
 			UE::Mutable::Private::UpdateLLMStats();
 		}
 
+		
 		void SetString(FCacheAddress Address, const TManagedPtr<const String>& Value)
 		{
 			check(Address.At < OpExecutionData.size_code());
@@ -1897,9 +1927,8 @@ namespace UE::Mutable::Private
 
 		/** Potentially expensive. Must be called in the Mutable Thread. */
 		UE_API TManagedPtr<const FInstance> BeginUpdate_MutableThread(const TSharedRef<FLiveInstance>& LiveInstance,
-			const TSharedPtr<const FParameters>& InOldParameters,
-			int32 InOldState,
-			uint32 InLodMask);
+			uint32 InLodMask,
+			bool bClearCacheLayer1);
 		
 		//! Only valid between BeginUpdate and EndUpdate
 		//! Calculate the description of an image, without generating it.
@@ -1986,7 +2015,7 @@ namespace UE::Mutable::Private
 		/** Increased each time a rom is loaded. Used to keep track of which roms to evict. */
 		static uint64 RomTick;
 		
-		bool CheckUpdatedParameters(const TSharedRef<FLiveInstance>& LiveInstance, const TSharedPtr<const FParameters>&);
+		MUTABLERUNTIME_API static bool CheckForUpdatedParameters(const TSharedRef<FLiveInstance>& LiveInstance, const TSharedPtr<const FParameters>& InOldParameters);
 		
 		void RunCode(
 			const TSharedRef<FLiveInstance>& InLiveInstance,
diff --git a/Engine/Plugins/Mutable/Source/MutableRuntime/Private/MuR/System.cpp b/Engine/Plugins/Mutable/Source/MutableRuntime/Private/MuR/System.cpp
--- a/Engine/Plugins/Mutable/Source/MutableRuntime/Private/MuR/System.cpp
+++ b/Engine/Plugins/Mutable/Source/MutableRuntime/Private/MuR/System.cpp
@@ -212,9 +212,7 @@ namespace UE::Mutable::Private
 
 	
     TManagedPtr<const FInstance> FSystem::BeginUpdate_MutableThread(const TSharedRef<FLiveInstance>& LiveInstance,
-		const TSharedPtr<const FParameters>& InOldParameters,
-		int32 InOldState,
-		uint32 InLodMask)
+		uint32 InLodMask, bool bClearCacheLayer1)
     {
 		LLM_SCOPE_BYNAME(TEXT("MutableRuntime"));
 		MUTABLE_CPUPROFILER_SCOPE(BeginUpdate_MutableThread);
@@ -228,9 +226,9 @@ namespace UE::Mutable::Private
 
 		const FProgram& Program = LiveInstance->Model->Program;
 
-    	const int32 NewState = LiveInstance->State;
+    	const int32 LiveInstanceState = LiveInstance->State;
     	
-		bool bValidState = NewState >= 0 && NewState < (int32)Program.States.Num();
+		const bool bValidState = LiveInstanceState >= 0 && LiveInstanceState < (int32)Program.States.Num();
 		if (!bValidState)
 		{
 			UE_LOG(LogMutableCore, Error, TEXT("Invalid state in mutable update."));
@@ -239,43 +237,44 @@ namespace UE::Mutable::Private
 
 		// This may free resources that allow us to use less memory.
 		LiveInstance->Instance = nullptr;
-    	
-    	// Remove cached data
-		{
+
+    	// Program cache setup
+    	{
+			// Remove all unlocked data (data from the state used to Prepare the cache)
 			LiveInstance->Cache->ClearLayer0();
-    	
-			// If we changed the state or the parameters that are not in this state, we need to rebuild all.
-			if (InOldState != NewState || CheckUpdatedParameters(LiveInstance, InOldParameters))
+			
+			if (bClearCacheLayer1)
 			{
+				// Remove all locked data (data from other states)
 				LiveInstance->Cache->ClearLayer1();
 			}
-		}
+			
+			// Prepare instance cache
+			LiveInstance->Cache->Prepare(Program, LiveInstance->State);
+    	}
 
-		WorkingMemoryManager.BeginRunnerThread();
     	
-		// Prepare instance cache
-    	LiveInstance->Cache->Prepare(Program, NewState);
-    	
-		OP::ADDRESS RootAt = Program.States[NewState].Root;
-		RunCode(LiveInstance, RootAt, InLodMask);
-
-		TManagedPtr<const FInstance> Result;
-
-		if (!bUnrecoverableError)
+    	TManagedPtr<const FInstance> Result;
+		WorkingMemoryManager.BeginRunnerThread();
 		{
-			Result = LiveInstance->Cache->GetInstance(FCacheAddress(RootAt, 0, 0));
-		}
+			const OP::ADDRESS RootAt = Program.States[LiveInstanceState].Root;
+			
+			RunCode(LiveInstance, RootAt, InLodMask);
+			if (!bUnrecoverableError)
+			{
+				Result = LiveInstance->Cache->GetInstance(FCacheAddress(RootAt, 0, 0));
+			}
 
-		// Debug check to see if we managed the op-hit-counts correctly
-		LiveInstance->Cache->CheckHitCountsCleared();
-		LiveInstance->Instance = Result;
+			// Debug check to see if we managed the op-hit-counts correctly
+			LiveInstance->Cache->CheckHitCountsCleared();
+			LiveInstance->Instance = Result;
 
-    	if (!Result)
-		{
-			// In case of failure return an empty instance, to prevent following code to have to check it every time
-			Result = MakeManaged<FInstance>();
+			if (!Result)
+			{
+				// In case of failure return an empty instance, to prevent following code to have to check it every time
+				Result = MakeManaged<FInstance>();
+			}
 		}
-
 		WorkingMemoryManager.EndRunnerThread(LiveInstance);
     	
 		return Result;
@@ -905,7 +904,7 @@ namespace UE::Mutable::Private
     }
 
 	
-	bool FSystem::CheckUpdatedParameters(const TSharedRef<FLiveInstance>& LiveInstance, const TSharedPtr<const FParameters>& InOldParameters)
+	bool FSystem::CheckForUpdatedParameters(const TSharedRef<FLiveInstance>& LiveInstance, const TSharedPtr<const FParameters>& InOldParameters)
     {
         bool bFullBuild = false;
 
@@ -914,7 +913,7 @@ namespace UE::Mutable::Private
 			return true;
 		}
 
-    	const TSharedPtr<const FParameters> NewParameters = LiveInstance->Parameters.ToSharedRef();
+    	const TSharedRef<const FParameters> NewParameters = LiveInstance->Parameters.ToSharedRef();
 
         // check what parameters have changed
         const FProgram& Program = LiveInstance->Model->Program;
@@ -1302,7 +1301,6 @@ namespace UE::Mutable::Private
 	}
 	
 
-
 	int64 FWorkingMemoryManager::GetCurrentMemoryBytes() const
 	{
 		MUTABLE_CPUPROFILER_SCOPE(GetCurrentMemoryBytes);
