[mutable] Changed check failure to handled error when a cyclical reference is found when looking for the root CO (GenerateMutableSource).
- A user mistake should never be handled with a crash.
- Fixed crash that happened when opening a CO that exposes a circular reference with another CO.
- Fixed possible crash when running the PostEditChangeProperty of the Group Node. It was not taking into consideration that the value returned by the GetFullGraphRootNode may be null (cycle)

#rnx

#Synced-CL 50196343

diff --git a/Engine/Plugins/Mutable/Source/CustomizableObjectEditor/Private/MuCOE/CustomizableObjectCompiler.cpp b/Engine/Plugins/Mutable/Source/CustomizableObjectEditor/Private/MuCOE/CustomizableObjectCompiler.cpp
--- a/Engine/Plugins/Mutable/Source/CustomizableObjectEditor/Private/MuCOE/CustomizableObjectCompiler.cpp
+++ b/Engine/Plugins/Mutable/Source/CustomizableObjectEditor/Private/MuCOE/CustomizableObjectCompiler.cpp
@@ -284,25 +284,33 @@ void FCustomizableObjectCompiler::Compile(const TSharedRef<FCompilationRequest>&
 	
 	if (!UCustomizableObjectSystem::IsActive())
 	{
-		UE_LOG(LogMutable, Warning, TEXT("Failed to compile Customizable Object [%s]. Mutable is disabled. To enable it set the CVar Mutable.Enabled to true."), *CurrentObject->GetName());
+		CompilerLog( LOCTEXT("CompilationFailure_MutableIsDisabled","Failed to compile Customizable Object. Mutable is disabled. To enable it set the CVar Mutable.Enabled to true."), CurrentObject, EMessageSeverity::Warning);
 		CompleteRequest(ECompilationStatePrivate::Completed, ECompilationResultPrivate::Errors);
 		return;
 	}
 
-	UCustomizableObject* RootObject = GraphTraversal::GetRootObject(CurrentObject);
-	check(RootObject);
+	const UCustomizableObject* RootObject = GraphTraversal::GetRootObject(CurrentObject);
+	if (!RootObject)
+	{
+		const UCustomizableObject* OtherObject = GraphTraversal::GetCylindricallyReferencedObject(CurrentObject);
+		check(OtherObject);
+
+		CompilerLog( LOCTEXT("CompilationFailure_COCircularReferenceFound", "Failed to compile Customizable Object. A circular reference between this and another CO was found."), {CurrentObject, OtherObject}, EMessageSeverity::Warning);
+		CompleteRequest(ECompilationStatePrivate::Completed, ECompilationResultPrivate::Errors);
+		return;
+	}
 
 	if (RootObject->VersionBridge && !RootObject->VersionBridge->GetClass()->ImplementsInterface(UCustomizableObjectVersionBridgeInterface::StaticClass()))
 	{
-		UE_LOG(LogMutable, Warning, TEXT("In Customizable Object [%s], the VersionBridge asset [%s] does not implement the required UCustomizableObjectVersionBridgeInterface."),
-			*RootObject->GetName(), *RootObject->VersionBridge.GetName());
+		const FText WarningText = FText::Format(LOCTEXT("CompilationFailure_RootNotImplementingVersionBridgeInterface", "Failed to compile Customizable Object. The VersionBridge asset {0} does not implement the required UCustomizableObjectVersionBridgeInterface"),  FText::FromString(RootObject->VersionBridge.GetName()));
+		CompilerLog( WarningText, {RootObject, RootObject->VersionBridge}, EMessageSeverity::Warning);
 		CompleteRequest(ECompilationStatePrivate::Completed, ECompilationResultPrivate::Errors);
 		return;
 	}
 
 	if (!CurrentOptions.bIsCooking && IsRunningCookCommandlet())
 	{
-		UE_LOG(LogMutable, Display, TEXT("Editor compilation suspended for Customizable Object [%s]. Can not compile COs when the cook commandlet is running. "), *CurrentObject->GetName());
+		CompilerLog( LOCTEXT("CompilationSuspended_CookInProgress", "Editor compilation suspended. Can not compile COs when the cook commandlet is running."), CurrentObject, EMessageSeverity::Info);
 		CompleteRequest(ECompilationStatePrivate::Completed, ECompilationResultPrivate::Errors);
 		return;
 	}
diff --git a/Engine/Plugins/Mutable/Source/CustomizableObjectEditor/Private/MuCOE/CustomizableObjectEditorModule.cpp b/Engine/Plugins/Mutable/Source/CustomizableObjectEditor/Private/MuCOE/CustomizableObjectEditorModule.cpp
--- a/Engine/Plugins/Mutable/Source/CustomizableObjectEditor/Private/MuCOE/CustomizableObjectEditorModule.cpp
+++ b/Engine/Plugins/Mutable/Source/CustomizableObjectEditor/Private/MuCOE/CustomizableObjectEditorModule.cpp
@@ -14,7 +14,6 @@
 #include "HAL/FileManager.h"
 #include "MuCO/CustomizableObjectSystem.h"		// For defines related to memory function replacements.
 #include "MuCO/CustomizableObjectInstanceUsage.h"
-#include "MuCO/CustomizableSkeletalMeshActor.h"
 #include "MuCO/ICustomizableObjectModule.h"		// For instance editor command utility function
 #include "MuCO/UnrealPortabilityHelpers.h"
 #include "MuCOE/CustomizableInstanceDetails.h"
@@ -1122,8 +1121,11 @@ TMap<FName, FGuid> FCustomizableObjectEditorModule::GetParticipatingObjects(cons
 
 	TArray<const UCustomizableObject*> VisitedObjects;
     UCustomizableObjectNodeObject* RootNode = GraphTraversal::GetFullGraphRootNode(Object, VisitedObjects);
-
-	GraphTraversal::VisitNodes(*RootNode, Visit);
+	
+	if (RootNode)
+	{
+		GraphTraversal::VisitNodes(*RootNode, Visit);
+	}
 	
 	// Done
 	return Result;
diff --git a/Engine/Plugins/Mutable/Source/CustomizableObjectEditor/Private/MuCOE/GraphTraversal.cpp b/Engine/Plugins/Mutable/Source/CustomizableObjectEditor/Private/MuCOE/GraphTraversal.cpp
--- a/Engine/Plugins/Mutable/Source/CustomizableObjectEditor/Private/MuCOE/GraphTraversal.cpp
+++ b/Engine/Plugins/Mutable/Source/CustomizableObjectEditor/Private/MuCOE/GraphTraversal.cpp
@@ -12,7 +12,6 @@
 #include "MuCOE/Nodes/CONodeModifierExtendSkeletalMeshSection.h"
 #include "MuCOE/Nodes/CustomizableObjectNodeExternalPin.h"
 #include "MuCOE/Nodes/CONodeSkeletalMeshSection.h"
-#include "MuCOE/Nodes/CustomizableObjectNodeMaterialParameter.h"
 #include "MuCOE/Nodes/CONodeSkeletalMeshSectionVariation.h"
 #include "MuCOE/Nodes/CustomizableObjectNodeMeshMorph.h"
 #include "MuCOE/Nodes/CustomizableObjectNodeMeshMorphStackApplication.h"
@@ -24,7 +23,6 @@
 #include "MuCOE/Nodes/CustomizableObjectNodeObject.h"
 #include "MuCOE/Nodes/CustomizableObjectNodeSkeletalMesh.h"
 #include "MuCOE/Nodes/CustomizableObjectNodeStaticMesh.h"
-#include "MuCOE/Nodes/CustomizableObjectNodeStaticString.h"
 #include "MuCOE/Nodes/CustomizableObjectNodeTable.h"
 
 #include "MuCOE/Nodes/CustomizableObjectNodeAnimationPose.h"
@@ -302,7 +300,7 @@ bool GetParentsUntilRoot(const UCustomizableObject* Object, TArray<UCustomizable
 		}
 		else
 		{
-			// This object has already been visted which means that there is a Cycle between Customizable Objects
+			// This object has already been visited which means that there is a Cycle between Customizable Objects
 			return false;
 		}
 
@@ -371,6 +369,30 @@ UCustomizableObject* GetFullGraphRootObject(const UCustomizableObjectNodeObject*
 }
 
 
+UCustomizableObject* GetCylindricallyReferencedObject(const UCustomizableObjectNodeObject* Node, TArray<UCustomizableObject*>& VisitedObjects)
+{
+	if (Node->ParentObject != nullptr)
+	{
+		VisitedObjects.Add(Node->ParentObject);
+
+		UCustomizableObjectNodeObject* Root = GetRootNode(Node->ParentObject);
+		check(Root);
+		
+		if (VisitedObjects.Contains(Root->ParentObject))
+		{
+			//There is a cycle so report about what node is the cause
+			return Root->ParentObject;
+		}
+		else
+		{
+			return GetCylindricallyReferencedObject(Root, VisitedObjects);
+		}
+	}
+
+	return nullptr;
+}
+
+
 UCustomizableObject* GraphTraversal::GetObject(const UCustomizableObjectNode& Node)
 {
 	if (Node.IsInMacro())
@@ -406,6 +428,22 @@ const UCustomizableObject* GraphTraversal::GetRootObject(const UCustomizableObje
 }
 
 
+const UCustomizableObject* GraphTraversal::GetCylindricallyReferencedObject(const UCustomizableObject* ChildObject)
+{
+	// Grab a node to start the search -> Get the root since it should be always present
+	UCustomizableObjectNodeObject* ObjectRootNode = GetRootNode(ChildObject);
+
+	if (ObjectRootNode && ObjectRootNode->ParentObject)
+	{
+		TArray<UCustomizableObject*> VisitedNodes;
+		return GetCylindricallyReferencedObject(ObjectRootNode, VisitedNodes);
+	}
+
+	// No parent object found so there can be no cycle
+	return nullptr;
+}
+
+
 void RecursiveVisitNodes(UCustomizableObjectNode& CurrentNode, const TFunction<void(UCustomizableObjectNode&)>& VisitFunction,
 	TArray<const UCustomizableObjectNodeMacroInstance*>& MacroContext, TSet<TPair<UCustomizableObjectNode*, TArray<const UCustomizableObjectNodeMacroInstance*>>>& VisitedNodes, const TMultiMap<FGuid, UCustomizableObjectNodeObject*>& ObjectGroupMap)
 {
@@ -777,7 +815,7 @@ void GetNodeGroupObjectNodeMappingImmersive(UCustomizableObject* Object, IAssetR
 {
 	MUTABLE_CPUPROFILER_SCOPE(GetNodeGroupObjectNodeMappingImmersive);
 	
-	if (Visited.Contains(Object))
+	if (!Object || Visited.Contains(Object))
 	{
 		return;
 	}
diff --git a/Engine/Plugins/Mutable/Source/CustomizableObjectEditor/Private/MuCOE/GraphTraversal.h b/Engine/Plugins/Mutable/Source/CustomizableObjectEditor/Private/MuCOE/GraphTraversal.h
--- a/Engine/Plugins/Mutable/Source/CustomizableObjectEditor/Private/MuCOE/GraphTraversal.h
+++ b/Engine/Plugins/Mutable/Source/CustomizableObjectEditor/Private/MuCOE/GraphTraversal.h
@@ -65,7 +65,11 @@ bool GetParentsUntilRoot(const UCustomizableObject* Object, TArray<UCustomizable
 /** Returns true if the Candidate is parent of the current Customizable Object */
 bool HasCandidateAsParent(UCustomizableObjectNodeObject* Node, UCustomizableObject* ParentCandidate);
 
-/** Return the full graph Customizable Object root of the given node. */
+/** Return the full graph Customizable Object root of the given node.
+ * @param Node The CO Node we want to get its root CO
+ * @param VisitedObjects List of objects used to record the search path already taken. It is used to determine if a cycle has been found.
+ * return The root CO. It can be null if a cycle has been found.
+ */
 UCustomizableObject* GetFullGraphRootObject(const UCustomizableObjectNodeObject* Node, TArray<UCustomizableObject*>& VisitedObjects);
 
 namespace GraphTraversal
@@ -82,6 +86,14 @@ namespace GraphTraversal
 		 * object does not have any parent.	*/
 	CUSTOMIZABLEOBJECTEDITOR_API UCustomizableObject* GetRootObject(UCustomizableObject* ChildObject);
 	CUSTOMIZABLEOBJECTEDITOR_API const UCustomizableObject* GetRootObject(const UCustomizableObject* ChildObject);
+
+	/**
+	 * Provided a CO it looks for cylindrical references while going up in search of the root CO. It is quite similar with GetRootObject but, instead
+	 * of returning the root CO will instead return you the CO that was found to be referencing the provided node
+	 * @param ChildObject The object where to start the search for cylindrical references.
+	 * @return The CO that is referencing the provided one.
+	 */
+	CUSTOMIZABLEOBJECTEDITOR_API const UCustomizableObject* GetCylindricallyReferencedObject(const UCustomizableObject* ChildObject);
 	
 	/** Return the full graph Customizable Object Node root of the node given as parameter */
 	UCustomizableObjectNodeObject* GetFullGraphRootNode(const UCustomizableObject* Object, TArray<const UCustomizableObject*>& VisitedObjects);
diff --git a/Engine/Plugins/Mutable/Source/CustomizableObjectEditor/Private/MuCOE/Nodes/CustomizableObjectNodeObjectGroup.cpp b/Engine/Plugins/Mutable/Source/CustomizableObjectEditor/Private/MuCOE/Nodes/CustomizableObjectNodeObjectGroup.cpp
--- a/Engine/Plugins/Mutable/Source/CustomizableObjectEditor/Private/MuCOE/Nodes/CustomizableObjectNodeObjectGroup.cpp
+++ b/Engine/Plugins/Mutable/Source/CustomizableObjectEditor/Private/MuCOE/Nodes/CustomizableObjectNodeObjectGroup.cpp
@@ -59,7 +59,8 @@ void UCustomizableObjectNodeObjectGroup::PostEditChangeProperty(FPropertyChanged
 					TArray<const UCustomizableObject*> VisitedObjects;
 					CurrentRootNode = GraphTraversal::GetFullGraphRootNode(CurrentRootNode->ParentObject, VisitedObjects);
 				}
-				if (!CurrentRootNode->ParentObject)
+				
+				if (CurrentRootNode && !CurrentRootNode->ParentObject)
 				{
 					for (FCustomizableObjectState& State : CurrentRootNode->States)
 					{
