[Mutable] Fix crash when deserializing meshes with cloth.

#Synced-CL 50155263

#was-backout 50176689

#Synced-CL 50155263

diff --git a/Engine/Plugins/Mutable/Source/CustomizableObject/Internal/MuCO/CustomizableObjectPrivate.h b/Engine/Plugins/Mutable/Source/CustomizableObject/Internal/MuCO/CustomizableObjectPrivate.h
--- a/Engine/Plugins/Mutable/Source/CustomizableObject/Internal/MuCO/CustomizableObjectPrivate.h
+++ b/Engine/Plugins/Mutable/Source/CustomizableObject/Internal/MuCO/CustomizableObjectPrivate.h
@@ -1563,7 +1563,9 @@ enum class ECustomizableObjectVersions : int32
 	AddLayerIndexToMatParamTextures,
 
 	ClothCoreSerializationIssue,
-	
+
+	ClothCoreSerializationIssue2,
+
 	// -----<new versions can be added above this line>--------
 	LastCustomizableObjectVersion
 };
diff --git a/Engine/Plugins/Mutable/Source/MutableRuntime/Internal/MuR/Model.h b/Engine/Plugins/Mutable/Source/MutableRuntime/Internal/MuR/Model.h
--- a/Engine/Plugins/Mutable/Source/MutableRuntime/Internal/MuR/Model.h
+++ b/Engine/Plugins/Mutable/Source/MutableRuntime/Internal/MuR/Model.h
@@ -762,6 +762,7 @@ namespace UE::Mutable::Private
 						Result->VertexBuffers = MetaDataMesh->VertexBuffers;
 						Result->IndexBuffers = MetaDataMesh->IndexBuffers;
 						Result->Surfaces = MetaDataMesh->Surfaces;
+						Result->ClothSections.SetNum(MetaDataMesh->Surfaces.Num());
 					}
 
 					Result->Tags = MetaDataMesh->Tags; 
diff --git a/Engine/Plugins/Mutable/Source/MutableRuntime/Private/MuR/Mesh.cpp b/Engine/Plugins/Mutable/Source/MutableRuntime/Private/MuR/Mesh.cpp
--- a/Engine/Plugins/Mutable/Source/MutableRuntime/Private/MuR/Mesh.cpp
+++ b/Engine/Plugins/Mutable/Source/MutableRuntime/Private/MuR/Mesh.cpp
@@ -105,6 +105,7 @@ void FMesh::CopyFrom(const FMesh& From, EMeshCopyFlags InFlags)
 	if (EnumHasAnyFlags(InFlags, EMeshCopyFlags::WithSurfaces))
 	{
 		Surfaces = From.Surfaces;
+		ClothSections.SetNum(From.Surfaces.Num());
 	}
 
 	if (EnumHasAnyFlags(InFlags, EMeshCopyFlags::WithSkeleton))
