[Mutable] Added a new tab to show the compilation messages inside the CO Editor.

#rnx

#Synced-CL 50197798

diff --git a/Engine/Plugins/Mutable/Source/CustomizableObjectEditor/Internal/MuCOE/CompileRequest.h b/Engine/Plugins/Mutable/Source/CustomizableObjectEditor/Internal/MuCOE/CompileRequest.h
--- a/Engine/Plugins/Mutable/Source/CustomizableObjectEditor/Internal/MuCOE/CompileRequest.h
+++ b/Engine/Plugins/Mutable/Source/CustomizableObjectEditor/Internal/MuCOE/CompileRequest.h
@@ -7,7 +7,7 @@
 #define UE_API CUSTOMIZABLEOBJECTEDITOR_API
 
 class UCustomizableObject;
-
+class FCustomizableObjectCustomLogList;
 
 struct FCompilationRequest
 {
@@ -50,6 +50,8 @@ public:
 	bool bSkipIfCompiled = false;
 	bool bSkipIfNotOutOfDate = false;
 	bool bSilentCompilation = true;
+
+	TSharedPtr<FCustomizableObjectCustomLogList> EditorCustomLogListPtr;
 };
 
 #undef UE_API
diff --git a/Engine/Plugins/Mutable/Source/CustomizableObjectEditor/Private/MuCOE/CustomizableObjectCompiler.cpp b/Engine/Plugins/Mutable/Source/CustomizableObjectEditor/Private/MuCOE/CustomizableObjectCompiler.cpp
--- a/Engine/Plugins/Mutable/Source/CustomizableObjectEditor/Private/MuCOE/CustomizableObjectCompiler.cpp
+++ b/Engine/Plugins/Mutable/Source/CustomizableObjectEditor/Private/MuCOE/CustomizableObjectCompiler.cpp
@@ -35,6 +35,8 @@
 #include "DerivedDataCacheKey.h"
 #include "DerivedDataRequestOwner.h"
 #include "MuR/External/Operation.h"
+#include "IMessageLogListing.h"
+
 
 class UTexture2D;
 
@@ -281,6 +283,12 @@ void FCustomizableObjectCompiler::Compile(const TSharedRef<FCompilationRequest>&
 	CurrentRequest = InCompileRequest.ToSharedPtr();
 	CurrentObject = Object;
 	CurrentOptions = InCompileRequest->Options;
+
+	// If it's compiling inside an editor, clear any previous message of the terminal
+	if (CurrentRequest->EditorCustomLogListPtr)
+	{
+		CurrentRequest->EditorCustomLogListPtr->MessageLogListing->ClearMessages();
+	}
 	
 	if (!UCustomizableObjectSystem::IsActive())
 	{
@@ -1807,7 +1815,7 @@ void FCustomizableObjectCompiler::AddCompileNotification(const FText& Compilatio
 	.Notification(!CurrentRequest->bSilentCompilation)
 	.CustomNotification()
 	.FixNotification()
-	.Log();
+	.Log(CurrentRequest->EditorCustomLogListPtr);
 }
 
 
@@ -1849,12 +1857,18 @@ void FCustomizableObjectCompiler::NotifyCompilationErrors() const
 		:
 		FText::Format(LOCTEXT("CompilationFinished_WithoutIgnoreds", "{0} finished compiling with {1} {1}|plural(one=warning,other=warnings) and {2} {2}|plural(one=error,other=errors)."), Prefix, NumWarnings, NumErrors);
 	
+	// Open the Stats tab if it was previously closed an there are errors/warnings
+	if (CurrentRequest->EditorCustomLogListPtr && !NoWarningsOrErrors)
+	{
+		CurrentRequest->EditorCustomLogListPtr->DisplayStatesTab();
+	}
+
 	FCustomizableObjectEditorLogger::CreateLog(Message)
 	.Category(ELoggerCategory::Compilation)
 	.Severity(Severity)
 	.Notification(!CurrentRequest->bSilentCompilation || !NoWarningsOrErrors)
 	.CustomNotification()
-	.Log();
+	.Log(CurrentRequest->EditorCustomLogListPtr);
 }
 
 
@@ -1862,12 +1876,12 @@ void FCustomizableObjectCompiler::CompilerLog(const FText& Message, const TArray
 {
 	if (CompilationLogsContainer.AddMessage(Message, Context, MessageSeverity, SpamBin)) // Cache the message for later reference
 	{
-		FCustomizableObjectEditorLogger::CreateLog(Message)
+			FCustomizableObjectEditorLogger::CreateLog(Message)
 			.Severity(MessageSeverity)
 			.Context(Context)
 			.BaseObject(bAddBaseObjectInfo)
 			.SpamBin(SpamBin)
-			.Log();
+			.Log(CurrentRequest ? CurrentRequest->EditorCustomLogListPtr : nullptr); // Compilations from a debugger do not have a compilation request
 	}
 }
 
diff --git a/Engine/Plugins/Mutable/Source/CustomizableObjectEditor/Private/MuCOE/CustomizableObjectEditor.cpp b/Engine/Plugins/Mutable/Source/CustomizableObjectEditor/Private/MuCOE/CustomizableObjectEditor.cpp
--- a/Engine/Plugins/Mutable/Source/CustomizableObjectEditor/Private/MuCOE/CustomizableObjectEditor.cpp
+++ b/Engine/Plugins/Mutable/Source/CustomizableObjectEditor/Private/MuCOE/CustomizableObjectEditor.cpp
@@ -66,6 +66,7 @@
 #include "MuCO/UnrealMutableImageProvider.h"
 #include "MuCOE/CompileRequest.h"
 #include "MuCOE/CustomizableObjectInstanceFactory.h"
+#include "MessageLogModule.h"
 
 #include UE_INLINE_GENERATED_CPP_BY_NAME(CustomizableObjectEditor)
 
@@ -89,6 +90,7 @@ const FName FCustomizableObjectEditor::PerformanceAnalyzerTabId(TEXT("Customizab
 const FName FCustomizableObjectEditor::TagExplorerTabId(TEXT("CustomizableObjectEditor_TagExplorer"));
 const FName FCustomizableObjectEditor::ObjectDebuggerTabId(TEXT("CustomizableObjectEditor_ObjectDebugger"));
 const FName FCustomizableObjectEditor::PopulationClassTagManagerTabId(TEXT("CustomizableObjectEditor_PopulationClassTabManager"));
+const FName FCustomizableObjectEditor::StatsTabId(TEXT("CustomizableObjectEditor_StatsTabManager"));
 
 
 void UUpdateClassWrapper::DelegatedCallback(UCustomizableObjectInstance* Instance)
@@ -135,6 +137,10 @@ void FCustomizableObjectEditor::RegisterTabSpawners(const TSharedRef<class FTabM
 	InTabManager->RegisterTabSpawner(TagExplorerTabId, FOnSpawnTab::CreateSP(this, &FCustomizableObjectEditor::SpawnTab_TagExplorer))
 		.SetDisplayName(LOCTEXT("TagExplorerTab", "Tag Explorer"))
 		.SetGroup(WorkspaceMenuCategoryRef);
+	
+	InTabManager->RegisterTabSpawner(StatsTabId, FOnSpawnTab::CreateSP(this, &FCustomizableObjectEditor::SpawnTab_Stats))
+		.SetDisplayName(LOCTEXT("StatsTab", "Mutable Log"))
+		.SetGroup(WorkspaceMenuCategoryRef);
 }
 
 
@@ -147,6 +153,7 @@ void FCustomizableObjectEditor::UnregisterTabSpawners(const TSharedRef<class FTa
 	InTabManager->UnregisterTabSpawner( AdvancedPreviewSettingsTabId );
 	InTabManager->UnregisterTabSpawner( TextureAnalyzerTabId );
 	InTabManager->UnregisterTabSpawner( PerformanceAnalyzerTabId );
+	InTabManager->UnregisterTabSpawner( StatsTabId );
 }	
 
 
@@ -236,16 +243,31 @@ void FCustomizableObjectEditor::InitCustomizableObjectEditor(const EToolkitMode:
 	CustomizableObjectEditorAdvancedPreviewSettings->LoadProfileEnvironment();
 	AdvancedPreviewSettingsWidget = CustomizableObjectEditorAdvancedPreviewSettings;
 
-	const TSharedRef<FTabManager::FLayout> StandaloneDefaultLayout = FTabManager::NewLayout( "Standalone_CustomizableObjectEditor_Layout_v1.4" )
+	COEditorLogList = MakeShareable(new FCustomizableObjectCustomLogList(SharedThis(this), StatsWidget));
+
+	const TSharedRef<FTabManager::FLayout> StandaloneDefaultLayout = FTabManager::NewLayout( "Standalone_CustomizableObjectEditor_Layout_v1.5" )
 	->AddArea
 	(
 		FTabManager::NewPrimaryArea()->SetOrientation(Orient_Horizontal)
 		->Split
 		(
-			FTabManager::NewStack()
+			FTabManager::NewSplitter()->SetOrientation(Orient_Vertical)
 			->SetSizeCoefficient(0.6f)
-			->SetHideTabWell(true)
-			->AddTab(GraphTabId, ETabState::OpenedTab)
+			->Split
+			(
+				FTabManager::NewStack()
+				->SetHideTabWell(true)
+				->SetSizeCoefficient(0.8f)
+				->AddTab(GraphTabId, ETabState::OpenedTab)
+				->SetForegroundTab(GraphTabId)
+			)
+			->Split
+			(
+				FTabManager::NewStack()
+				->SetSizeCoefficient(0.2f)
+				->AddTab(StatsTabId, ETabState::OpenedTab)
+				->SetForegroundTab(StatsTabId)
+			)
 		)
 		->Split
 		(
@@ -1310,6 +1332,7 @@ void FCustomizableObjectEditor::CompileObject(bool bOnlySelectedParameters, bool
 	TSharedRef<FCompilationRequest> CompileRequest = MakeShared<FCompilationRequest>(*CustomizableObject);
 	CompileRequest->bSilentCompilation = false;
 	CompileRequest->Options.bGatherReferences = bGatherReferences;
+	CompileRequest->EditorCustomLogListPtr = COEditorLogList;
 
 	if (bOnlySelectedParameters)
 	{
@@ -1843,6 +1866,18 @@ TSharedRef<SDockTab> FCustomizableObjectEditor::SpawnTab_TagExplorer(const FSpaw
 }
 
 
+TSharedRef<SDockTab> FCustomizableObjectEditor::SpawnTab_Stats(const FSpawnTabArgs& Args)
+{
+	check(Args.GetTabId() == StatsTabId);
+
+	return SNew(SDockTab)
+	.Label(LOCTEXT("StatsTabLabel", "Stats"))
+	[
+		StatsWidget.ToSharedRef()
+	];
+}
+
+
 UCustomizableObject* FCustomizableObjectEditor::GetAbsoluteCOParent(const UCustomizableObjectNodeObject* const Root)
 {
 	if (Root->ParentObject != nullptr)
@@ -1904,6 +1939,7 @@ void FCustomizableObjectEditor::OnCustomizableObjectStatusChanged(FCustomizableO
 			CompileRequest->bSkipIfNotOutOfDate = true;
 			CompileRequest->bSilentCompilation = false;
 			CompileRequest->SetDerivedDataCachePolicy(GetDerivedDataCachePolicyForEditor());
+			CompileRequest->EditorCustomLogListPtr = COEditorLogList;
 
 			ICustomizableObjectEditorModulePrivate::GetChecked().EnqueueCompileRequest(CompileRequest);
 		}
@@ -1916,4 +1952,39 @@ void FCustomizableObjectEditor::OnChangeDebugPlatform(TSharedPtr<FString> NewSel
 	DebugTargetPlatformName = NewSelection;
 }
 
+
+//---------------------------------------------------------------
+
+FCustomizableObjectCustomLogList::FCustomizableObjectCustomLogList(TSharedPtr<FAssetEditorToolkit> InEditorPtr, TSharedPtr<SWidget>& StatsWidget)
+{
+	EditorPtr = InEditorPtr;
+
+	FMessageLogModule& MessageLogModule = FModuleManager::LoadModuleChecked<FMessageLogModule>("MessageLog");
+	FMessageLogInitializationOptions LogOptions;
+	// Show Pages so that user is never allowed to clear log messages
+	LogOptions.bShowPages = false;
+	LogOptions.bShowFilters = false;
+	LogOptions.bAllowClear = false;
+	LogOptions.MaxPageCount = 1;
+
+	MessageLogListing = MessageLogModule.CreateLogListing("MutableEditorLog", LogOptions);
+	check(MessageLogListing);
+
+	StatsWidget = MessageLogModule.CreateLogListingWidget(MessageLogListing.ToSharedRef());
+}
+
+
+void FCustomizableObjectCustomLogList::DisplayStatesTab()
+{
+	TSharedPtr<FAssetEditorToolkit> Editor = EditorPtr.Pin();
+
+	if (!Editor || !Editor->GetTabManager())
+	{
+		return;
+	}
+
+	Editor->GetTabManager()->TryInvokeTab(FCustomizableObjectEditor::StatsTabId);
+}
+
+
 #undef LOCTEXT_NAMESPACE
diff --git a/Engine/Plugins/Mutable/Source/CustomizableObjectEditor/Private/MuCOE/CustomizableObjectEditor.h b/Engine/Plugins/Mutable/Source/CustomizableObjectEditor/Private/MuCOE/CustomizableObjectEditor.h
--- a/Engine/Plugins/Mutable/Source/CustomizableObjectEditor/Private/MuCOE/CustomizableObjectEditor.h
+++ b/Engine/Plugins/Mutable/Source/CustomizableObjectEditor/Private/MuCOE/CustomizableObjectEditor.h
@@ -10,6 +10,8 @@
 #include "MuT/Node.h"
 #include "Widgets/Input/SNumericDropDown.h"
 
+#include "IMessageLogListing.h"
+
 #include "CustomizableObjectEditor.generated.h"
 
 struct FMutableSourceTextureData;
@@ -78,6 +80,30 @@ public:
 	FCreatePreviewInstanceFlagDelegate Delegate;
 };
 
+/** 
+* Class that adds a stats tab to the CustomizableObjectEditor
+* Also allows the managment of this tab
+*/
+class FCustomizableObjectCustomLogList
+{
+public:
+
+	FCustomizableObjectCustomLogList(TSharedPtr<FAssetEditorToolkit> InEditorPtr, TSharedPtr<SWidget>& StatsWidget);
+
+	/** Tries to reopen the States tab if it's been closed (and needs to be reopened) to show a compilation error or warning */
+	void DisplayStatesTab();
+
+public:
+
+	TSharedPtr<IMessageLogListing> MessageLogListing;
+
+private:
+
+	/** Pointer to the editor that contains this Stats */
+	TWeakPtr<FAssetEditorToolkit> EditorPtr;
+
+};
+
 
 /**
  * CustomizableObject Editor class
@@ -193,6 +219,7 @@ private:
 	TSharedRef<SDockTab> SpawnTab_TextureAnalyzer(const FSpawnTabArgs& Args);
 	TSharedRef<SDockTab> SpawnTab_PerformanceAnalyzer(const FSpawnTabArgs& Args);
 	TSharedRef<SDockTab> SpawnTab_TagExplorer(const FSpawnTabArgs& Args);
+	TSharedRef<SDockTab> SpawnTab_Stats(const FSpawnTabArgs& Args);
 	
 	/** Binds commands associated with the Static Mesh Editor. */
 	void BindCommands();
@@ -274,8 +301,10 @@ public:
 	static const FName TagExplorerTabId;
 	static const FName ObjectDebuggerTabId;
 	static const FName PopulationClassTagManagerTabId;
+	static const FName StatsTabId;
 
 private:
+
 	/** The currently viewed object. */
 	TObjectPtr<UCustomizableObject> CustomizableObject = nullptr;
 	TObjectPtr<UCustomizableObjectInstance> PreviewInstance = nullptr;
@@ -335,6 +364,10 @@ private:
 	
 	bool bDebugForceLargeLODBias = false;
 
+	TSharedPtr<FCustomizableObjectCustomLogList> COEditorLogList;
+	
+	TSharedPtr<SWidget> StatsWidget;
+
 protected:
 	/** @return the documentation location for this editor */
 	virtual FString GetDocumentationLink() const override;
diff --git a/Engine/Plugins/Mutable/Source/CustomizableObjectEditor/Private/MuCOE/CustomizableObjectEditorLogger.cpp b/Engine/Plugins/Mutable/Source/CustomizableObjectEditor/Private/MuCOE/CustomizableObjectEditorLogger.cpp
--- a/Engine/Plugins/Mutable/Source/CustomizableObjectEditor/Private/MuCOE/CustomizableObjectEditorLogger.cpp
+++ b/Engine/Plugins/Mutable/Source/CustomizableObjectEditor/Private/MuCOE/CustomizableObjectEditorLogger.cpp
@@ -15,6 +15,9 @@
 #include "Toolkits/ToolkitManager.h"
 #include "Widgets/Notifications/SNotificationList.h"
 
+#include "IMessageLogListing.h"
+
+
 #define LOCTEXT_NAMESPACE "CustomizableObject"
 
 
@@ -227,9 +230,9 @@ FLogParameters& FLogParameters::SpamBin(const ELoggerSpamBin SpamBin)
 }
 
 
-void FLogParameters::Log()
+void FLogParameters::Log(TSharedPtr<FCustomizableObjectCustomLogList> MessageLogListing)
 {
-	Logger.Log(*this);
+	Logger.Log(*this, MessageLogListing);
 }
 
 
@@ -246,7 +249,7 @@ FLogParameters FCustomizableObjectEditorLogger::CreateLog(FText&& Text)
 }
 
 
-void FCustomizableObjectEditorLogger::Log(FLogParameters& LogParameters)
+void FCustomizableObjectEditorLogger::Log(FLogParameters& LogParameters, TSharedPtr<FCustomizableObjectCustomLogList> MessageLogListing)
 {
 	FMessageLog MessageLog(LOG_NAME);
 	const TSharedRef<FTokenizedMessage> Message = MessageLog.Message(LogParameters.ParamSeverity);
@@ -279,6 +282,14 @@ void FCustomizableObjectEditorLogger::Log(FLogParameters& LogParameters)
 		}
 	}
 
+	if (MessageLogListing)
+	{
+		MessageLogListing->MessageLogListing->AddMessage(Message);
+
+		// Since we are compiling from an editor, we do not need to send any popup notification 
+		return;
+	}
+
 	if (LogParameters.bParamNotification)
 	{
 		FCategoryData& CategoryData = CategoriesData.FindOrAdd(LogParameters.ParamCategory);
diff --git a/Engine/Plugins/Mutable/Source/CustomizableObjectEditor/Private/MuCOE/CustomizableObjectEditorLogger.h b/Engine/Plugins/Mutable/Source/CustomizableObjectEditor/Private/MuCOE/CustomizableObjectEditorLogger.h
--- a/Engine/Plugins/Mutable/Source/CustomizableObjectEditor/Private/MuCOE/CustomizableObjectEditorLogger.h
+++ b/Engine/Plugins/Mutable/Source/CustomizableObjectEditor/Private/MuCOE/CustomizableObjectEditorLogger.h
@@ -3,6 +3,7 @@
 #pragma once
 
 #include "Logging/TokenizedMessage.h"
+#include "MuCOE/CustomizableObjectEditor.h"
 
 class UCustomizableObjectNode;
 
@@ -57,7 +58,7 @@ public:
 	FLogParameters& SpamBin(ELoggerSpamBin SpamBin);
 	
 	/** Actually display the log. */
-	void Log();
+	void Log(TSharedPtr<FCustomizableObjectCustomLogList> MessageLogListing = nullptr);
 	
 private:
 	/** Constructor. Contains all required parameters. */
@@ -102,7 +103,7 @@ public:
 	static FLogParameters CreateLog(FText&& Text);
 
 	/** Display the log. Equivalent to FLogParameters::Log(). */
-	void Log(FLogParameters& LogParameters);
+	void Log(FLogParameters& LogParameters, TSharedPtr<FCustomizableObjectCustomLogList> MessageLogListing = nullptr);
 
 	/** Dismiss a fixed notification. */
 	static void DismissNotification(ELoggerCategory Category);
