[Mutable] Fixed LOCTEXT key conflict.

#rnx

#Synced-CL 50197798

diff --git a/Engine/Plugins/Mutable/Source/CustomizableObjectEditor/Private/MuCOE/GenerateMutableSource/GenerateMutableSourceColor.cpp b/Engine/Plugins/Mutable/Source/CustomizableObjectEditor/Private/MuCOE/GenerateMutableSource/GenerateMutableSourceColor.cpp
--- a/Engine/Plugins/Mutable/Source/CustomizableObjectEditor/Private/MuCOE/GenerateMutableSource/GenerateMutableSourceColor.cpp
+++ b/Engine/Plugins/Mutable/Source/CustomizableObjectEditor/Private/MuCOE/GenerateMutableSource/GenerateMutableSourceColor.cpp
@@ -420,7 +420,7 @@ UE::Mutable::Private::Ptr<UE::Mutable::Private::NodeColour> GenerateMutableSourc
 				}
 				else
 				{
-					GenerationContext.Log(LOCTEXT("MaterialBreakNodeError", "Something went wrong generation the Material Node linked to the Material Break Node."), Node);
+					GenerationContext.Log(LOCTEXT("MaterialBreakNodeError", "An error occurred generating the Material Source."), Node, EMessageSeverity::Error);
 				}
 			}
 			else
diff --git a/Engine/Plugins/Mutable/Source/CustomizableObjectEditor/Private/MuCOE/GenerateMutableSource/GenerateMutableSourceFloat.cpp b/Engine/Plugins/Mutable/Source/CustomizableObjectEditor/Private/MuCOE/GenerateMutableSource/GenerateMutableSourceFloat.cpp
--- a/Engine/Plugins/Mutable/Source/CustomizableObjectEditor/Private/MuCOE/GenerateMutableSource/GenerateMutableSourceFloat.cpp
+++ b/Engine/Plugins/Mutable/Source/CustomizableObjectEditor/Private/MuCOE/GenerateMutableSource/GenerateMutableSourceFloat.cpp
@@ -440,7 +440,7 @@ UE::Mutable::Private::Ptr<UE::Mutable::Private::NodeScalar> GenerateMutableSourc
 				}
 				else
 				{
-					GenerationContext.Log(LOCTEXT("MaterialBreakNodeError", "Something went wrong generation the Material Node linked to the Material Break Node."), Node);
+					GenerationContext.Log(LOCTEXT("MaterialBreakNodeError", "An error occurred generating the Material Source."), Node, EMessageSeverity::Error);
 				}
 			}
 			else
