[mutable] VariationNode fixes and updates

UE-296979 Fixed wrong order of the VariationNode input pins possibly due to bad backwards compatibility fix up
- Looks like a previous fix up did intend to call for the reconstruction of the node after performing the fixes, but instead it was reconstructing the node before, which fixed the backend of the node but not the frontend of it.

UE-268615 Removed from the input variation pin the name of the tag as that is already part of the pin that is used to set the tag for the target variation.

#rnx

#Synced-CL 50244055

diff --git a/Engine/Plugins/Mutable/Source/CustomizableObject/Internal/MuCO/CustomizableObjectCustomVersion.h b/Engine/Plugins/Mutable/Source/CustomizableObject/Internal/MuCO/CustomizableObjectCustomVersion.h
--- a/Engine/Plugins/Mutable/Source/CustomizableObject/Internal/MuCO/CustomizableObjectCustomVersion.h
+++ b/Engine/Plugins/Mutable/Source/CustomizableObject/Internal/MuCO/CustomizableObjectCustomVersion.h
@@ -215,6 +215,8 @@ struct FCustomizableObjectCustomVersion
 		PassthroughParameter,
 		
 		MeshSectionNodeInputPinsAsReferences,
+		
+		WrongVariationNodeInputOrderAfterSerializationFix,
 
 		// -----<new versions can be added above this line>-------------------------------------------------
 		VersionPlusOne,
diff --git a/Engine/Plugins/Mutable/Source/CustomizableObjectEditor/Private/MuCOE/Nodes/CustomizableObjectNodeVariation.cpp b/Engine/Plugins/Mutable/Source/CustomizableObjectEditor/Private/MuCOE/Nodes/CustomizableObjectNodeVariation.cpp
--- a/Engine/Plugins/Mutable/Source/CustomizableObjectEditor/Private/MuCOE/Nodes/CustomizableObjectNodeVariation.cpp
+++ b/Engine/Plugins/Mutable/Source/CustomizableObjectEditor/Private/MuCOE/Nodes/CustomizableObjectNodeVariation.cpp
@@ -11,6 +11,7 @@
 #include UE_INLINE_GENERATED_CPP_BY_NAME(CustomizableObjectNodeVariation)
 
 class UCustomizableObjectNodeRemapPins;
+class UCustomizableObjectNodeRemapPinsByName;
 
 #define LOCTEXT_NAMESPACE "CustomizableObjectEditor"
 
@@ -47,14 +48,11 @@ void UCustomizableObjectNodeVariation::AllocateDefaultPins(UCustomizableObjectNo
 		{
 			VariationPin->PinType.ContainerType = EPinContainerType::Array;
 		}
-
-		FString TagName = GetTagDisplayName(VariationsData[VariationIndex].Tag);
-		VariationPin->PinFriendlyName = FText::Format(LOCTEXT("Variation_Pin_FriendlyName", "Variation {0} [{1}]"), VariationIndex, FText::FromString(TagName));
 		
 		VariationsPins[VariationIndex] = VariationPin;
 
-		FText VariationTagSufix = LOCTEXT("TagSufix", " - Tag");
-		const FName TagPinName = FName(PinName.ToString() + VariationTagSufix.ToString());
+		const FText VariationTagSuffix = LOCTEXT("TagSuffix", " - Tag");
+		const FName TagPinName = FName(PinName.ToString() + VariationTagSuffix.ToString());
 		UEdGraphPin* VariationTagPin = CustomCreatePin(EGPD_Input, Schema->PC_String, TagPinName, FText::FromName(TagPinName));
 
 		VariationTagPins[VariationIndex] = VariationTagPin;
@@ -129,9 +127,7 @@ FString UCustomizableObjectNodeVariation::GetVariationTag(int32 Index, TArray<co
 	{
 		if (const UEdGraphPin* ConnectedStringPin = FollowInputPin(*TagPin))
 		{
-			const UEdGraphPin* SourceStringPin = GraphTraversal::FindIOPinSourceThroughMacroContext(*ConnectedStringPin, MacroContext);
-
-			if (SourceStringPin)
+			if (const UEdGraphPin* SourceStringPin = GraphTraversal::FindIOPinSourceThroughMacroContext(*ConnectedStringPin, MacroContext))
 			{
 				if (const UCustomizableObjectNodeStaticString* StringNode = Cast<UCustomizableObjectNodeStaticString>(SourceStringPin->GetOwningNode()))
 				{
@@ -195,11 +191,6 @@ void UCustomizableObjectNodeVariation::BackwardsCompatibleFixup(int32 Customizab
 {
 	Super::BackwardsCompatibleFixup(CustomizableObjectCustomVersion);
 	
-	if (CustomizableObjectCustomVersion == FCustomizableObjectCustomVersion::NodeVariationSerializationIssue)
-	{
-		ReconstructNode();
-	}
-
 	if (CustomizableObjectCustomVersion == FCustomizableObjectCustomVersion::EnableMutableMacrosNewVersion)
 	{
 		if (VariationTagPins.IsEmpty())
@@ -209,15 +200,61 @@ void UCustomizableObjectNodeVariation::BackwardsCompatibleFixup(int32 Customizab
 			for (int32 VariationIndex = VariationsData.Num() - 1; VariationIndex >= 0; --VariationIndex)
 			{
 				const FName PinName = FName(FString::Printf(TEXT("Variation %d"), VariationIndex));
-				FText VariationTagSufix = LOCTEXT("TagSufix", " - Tag");
-				const FName TagPinName = FName(PinName.ToString() + VariationTagSufix.ToString());
+				FText VariationTagSuffix = LOCTEXT("TagSuffix", " - Tag");
+				const FName TagPinName = FName(PinName.ToString() + VariationTagSuffix.ToString());
 				UEdGraphPin* VariationTagPin = CustomCreatePin(EGPD_Input, UEdGraphSchema_CustomizableObject::PC_String, TagPinName, FText::FromName(TagPinName));
 				VariationTagPins[VariationIndex] = VariationTagPin;
 			}
 		}
 	}
+	else if (CustomizableObjectCustomVersion == FCustomizableObjectCustomVersion::WrongVariationNodeInputOrderAfterSerializationFix)
+	{
+		UCustomizableObjectNodeRemapPinsByName* Remapper =  NewObject<UCustomizableObjectNodeRemapPinsByName>();
+
+		FixupReconstructPins(Remapper, [](UCustomizableObjectNode* Node, UCustomizableObjectNodeRemapPins* Action)
+		{
+			// Code extracted from AllocateDefaultPins
+				
+			UCustomizableObjectNodeVariation* CastedNode = Cast<UCustomizableObjectNodeVariation>(Node);
+			check(CastedNode)
+			
+			const UEdGraphSchema_CustomizableObject* Schema = GetDefault<UEdGraphSchema_CustomizableObject>();
 
-	// TODO(Max UE-268615): Remove the tag name from the variation pins since now there is a pin with a textbox for it.
+			const FName Category = CastedNode->GetCategory();
+			const bool bIsInputPinArray = CastedNode->IsInputPinArray();
+
+			{
+				CastedNode->CustomCreatePin(EGPD_Output, Category);
+			}
+			
+			CastedNode->VariationsPins.SetNum(CastedNode->VariationsData.Num());
+			CastedNode->VariationTagPins.SetNum(CastedNode->VariationsData.Num());
+
+			for (int32 VariationIndex = CastedNode->VariationsData.Num() - 1; VariationIndex >= 0; --VariationIndex)
+			{
+				const FName PinName = FName(FString::Printf( TEXT("Variation %d"), VariationIndex));
+				UEdGraphPin* VariationPin = CastedNode->CustomCreatePin(EGPD_Input, Category, PinName, FText::FromName(PinName));
+				if (bIsInputPinArray)
+				{
+					VariationPin->PinType.ContainerType = EPinContainerType::Array;
+				}
+				
+				CastedNode->VariationsPins[VariationIndex] = VariationPin;
+
+				const FText VariationTagSuffix = LOCTEXT("TagSuffix", " - Tag");
+				const FName TagPinName = FName(PinName.ToString() + VariationTagSuffix.ToString());
+				UEdGraphPin* VariationTagPin = CastedNode->CustomCreatePin(EGPD_Input, Schema->PC_String, TagPinName, FText::FromName(TagPinName));
+
+				CastedNode->VariationTagPins[VariationIndex] = VariationTagPin;
+			}
+
+			UEdGraphPin* Pin = CastedNode->CustomCreatePin(EGPD_Input, Category, TEXT("Default"), LOCTEXT("Default", "Default"));
+			if (bIsInputPinArray)
+			{
+				Pin->PinType.ContainerType = EPinContainerType::Array;
+			}
+		});
+	}
 }
 
 
