[Mutable] Allow MIPs and LODs be executed in parallel.

CVar toggleable.

#Synced-CL 50155263

diff --git a/Engine/Plugins/Mutable/Source/CustomizableObject/Internal/MuCO/FMutableTaskGraph.h b/Engine/Plugins/Mutable/Source/CustomizableObject/Internal/MuCO/FMutableTaskGraph.h
--- a/Engine/Plugins/Mutable/Source/CustomizableObject/Internal/MuCO/FMutableTaskGraph.h
+++ b/Engine/Plugins/Mutable/Source/CustomizableObject/Internal/MuCO/FMutableTaskGraph.h
@@ -9,6 +9,8 @@
 #include "Containers/Queue.h"
 
 
+extern TAutoConsoleVariable<bool> CVarEnableParallelUpdates;
+
 /** Mutable Tasks System.
  *
  * Allows launching tasks in different threads:
@@ -61,6 +63,12 @@ public:
 	/** Create and launch a task on the Mutable Thread with Normal priority. */
 	UE::Tasks::FTask AddMutableThreadTask(const TCHAR* DebugName, TUniqueFunction<void()>&& TaskBody, const TArray<UE::Tasks::FTask>& Prerequisites = {});
 
+	UE::Tasks::FTask AddGenerationTask(const TCHAR* DebugName, TUniqueFunction<void()>&& TaskBody, const TArray<UE::Tasks::FTask>& Prerequisites = {});
+
+	UE::Tasks::FTask AddImageStreamingTask(const TCHAR* DebugName, TUniqueFunction<void()>&& TaskBody, const TArray<UE::Tasks::FTask>& Prerequisites = {});
+
+	UE::Tasks::FTask AddMeshStreamingTask(const TCHAR* DebugName, TUniqueFunction<void()>&& TaskBody, const TArray<UE::Tasks::FTask>& Prerequisites = {});
+	
 	/** Create and launch a task on the Mutable Thread with Low priority. */
 	uint32 AddMutableThreadTaskLowPriority(const TCHAR* DebugName, TFunction<void()>&& TaskBody);
 
@@ -148,5 +156,13 @@ private:
 	 *  Completed = Mutable Thread lock has been requested. May not be locked due to tasks still running (see LastMutableTaskBeforeLock).
 	 *  Not Completed = Mutable Thread unlocked. */
 	UE::Tasks::FTaskEvent MutableThreadUnlockEvent { UE_SOURCE_LOCATION };
+	
+	mutable FCriticalSection CriticalSectionLastGenerationTask;
+	mutable FCriticalSection CriticalSectionLastImageStreaming;
+	mutable FCriticalSection CriticalSectionLastMeshStreamingTask;
+	
+	UE::Tasks::FTask LastGenerationTask;
+	UE::Tasks::FTask LastImageStreamingTask;
+	UE::Tasks::FTask LastMeshStreamingTask;
 };
 
diff --git a/Engine/Plugins/Mutable/Source/CustomizableObject/Private/MuCO/CustomizableObjectMeshUpdate.cpp b/Engine/Plugins/Mutable/Source/CustomizableObject/Private/MuCO/CustomizableObjectMeshUpdate.cpp
--- a/Engine/Plugins/Mutable/Source/CustomizableObject/Private/MuCO/CustomizableObjectMeshUpdate.cpp
+++ b/Engine/Plugins/Mutable/Source/CustomizableObject/Private/MuCO/CustomizableObjectMeshUpdate.cpp
@@ -295,12 +295,24 @@ void FCustomizableObjectMeshStreamIn::RequestMeshUpdate(const FContext& Context)
 	
 	TaskSynchronization.Increment();
 
-	MutableTaskId = CustomizableObjectSystem->MutableTaskGraph.AddMutableThreadTaskLowPriority(
-		TEXT("Mutable_MeshUpdate"),
-		[SharedOperationData = OperationData, RefThis = TRefCountPtr<FCustomizableObjectMeshStreamIn>(this)]() mutable
-		{
-			impl::Task_Mutable_UpdateMesh(SharedOperationData, RefThis);
-		});
+	if (CVarEnableParallelUpdates.GetValueOnAnyThread())
+	{
+		CustomizableObjectSystem->MutableTaskGraph.AddMeshStreamingTask(
+			TEXT("Mutable_MeshUpdate"),
+			[SharedOperationData = OperationData, RefThis = TRefCountPtr<FCustomizableObjectMeshStreamIn>(this)]() mutable
+			{
+				impl::Task_Mutable_UpdateMesh(SharedOperationData, RefThis);
+			});
+	}
+	else
+	{
+		MutableTaskId = CustomizableObjectSystem->MutableTaskGraph.AddMutableThreadTaskLowPriority(
+			TEXT("Mutable_MeshUpdate"),
+			[SharedOperationData = OperationData, RefThis = TRefCountPtr<FCustomizableObjectMeshStreamIn>(this)]() mutable
+			{
+				impl::Task_Mutable_UpdateMesh(SharedOperationData, RefThis);
+			});
+	}
 
 	// Stream data (morphs, clothing...).
 	check(OperationData->MeshStreamingPrerequisite.IsValid())
diff --git a/Engine/Plugins/Mutable/Source/CustomizableObject/Private/MuCO/CustomizableObjectMipDataProvider.cpp b/Engine/Plugins/Mutable/Source/CustomizableObject/Private/MuCO/CustomizableObjectMipDataProvider.cpp
--- a/Engine/Plugins/Mutable/Source/CustomizableObject/Private/MuCO/CustomizableObjectMipDataProvider.cpp
+++ b/Engine/Plugins/Mutable/Source/CustomizableObject/Private/MuCO/CustomizableObjectMipDataProvider.cpp
@@ -418,12 +418,24 @@ int32 FMutableTextureMipDataProvider::GetMips(const FTextureUpdateContext& Conte
 				UCustomizableObjectSystemPrivate* CustomizableObjectSystem = UCustomizableObjectSystem::GetInstance()->GetPrivate();
 				if (CustomizableObjectSystem)
 				{
-					OperationData->MutableTaskId = CustomizableObjectSystem->MutableTaskGraph.AddMutableThreadTaskLowPriority(
-						TEXT("Mutable_MipUpdate"),
-						[OperationData = OperationData]()
-						{
-							Impl::Task_Mutable_UpdateImage(OperationData);
-						});
+					if (CVarEnableParallelUpdates.GetValueOnAnyThread())
+					{
+						CustomizableObjectSystem->MutableTaskGraph.AddImageStreamingTask(
+							TEXT("Mutable_MipUpdate"),
+							[OperationData = OperationData]()
+							{
+								Impl::Task_Mutable_UpdateImage(OperationData);
+							});
+					}
+					else
+					{
+						OperationData->MutableTaskId = CustomizableObjectSystem->MutableTaskGraph.AddMutableThreadTaskLowPriority(
+							TEXT("Mutable_MipUpdate"),
+							[OperationData = OperationData]()
+							{
+								Impl::Task_Mutable_UpdateImage(OperationData);
+							});
+					}
 				}
 			},
 			StreamTask,
diff --git a/Engine/Plugins/Mutable/Source/CustomizableObject/Private/MuCO/FMutableTaskGraph.cpp b/Engine/Plugins/Mutable/Source/CustomizableObject/Private/MuCO/FMutableTaskGraph.cpp
--- a/Engine/Plugins/Mutable/Source/CustomizableObject/Private/MuCO/FMutableTaskGraph.cpp
+++ b/Engine/Plugins/Mutable/Source/CustomizableObject/Private/MuCO/FMutableTaskGraph.cpp
@@ -31,6 +31,12 @@ static TAutoConsoleVariable<bool> CVarEnableMutableTaskLowPriority(
 	TEXT("Enable or disable Mutable Tasks with Low priority. If disabled, all task will have the same priority. "),
 	ECVF_Scalability);
 
+TAutoConsoleVariable<bool> CVarEnableParallelUpdates(
+	TEXT("mutable.EnableParallelUpdates"),
+	false,
+	TEXT("Allow initial generation, Mips and LODs to be run in parallel."),
+	ECVF_Scalability);
+
 
 /** Status of the Mutable Thread lock. Only used for Insights Regions. */
 enum class ERegionLockStatus
@@ -79,17 +85,70 @@ void FMutableTaskGraph::AddGameThreadTask(const TCHAR* DebugName, TUniqueFunctio
 
 UE::Tasks::FTask FMutableTaskGraph::AddMutableThreadTask(const TCHAR* DebugName, TUniqueFunction<void()>&& TaskBody, const TArray<UE::Tasks::FTask>& InPrerequisites)
 {
-	FScopeLock Lock(&MutableTaskLock);
+	if (CVarEnableParallelUpdates.GetValueOnAnyThread())
+	{
+		return AddGenerationTask(DebugName, Forward<TUniqueFunction<void()>>(TaskBody), InPrerequisites);
+	}
+	else
+	{
+		FScopeLock Lock(&MutableTaskLock);
+
+		TArray<UE::Tasks::FTask, TInlineAllocator<2>> Prerequisites = LastMutableTask.IsValid() ?
+			TArray<UE::Tasks::FTask, TInlineAllocator<2>>{ LastMutableTask, MutableThreadUnlockEvent } :
+			TArray<UE::Tasks::FTask, TInlineAllocator<2>>{};
+
+		Prerequisites.Append(InPrerequisites);
+	
+		LastMutableTask = UE::Tasks::Launch(DebugName, Forward<TUniqueFunction<void()>>(TaskBody), Prerequisites, TASKGRAPH_PRIORITY);
+	
+		return LastMutableTask;
+	}
+}
+
+
+UE::Tasks::FTask FMutableTaskGraph::AddGenerationTask(const TCHAR* DebugName, TUniqueFunction<void()>&& TaskBody, const TArray<UE::Tasks::FTask>& InPrerequisites)
+{
+	FScopeLock Lock(&CriticalSectionLastGenerationTask);
+	
+	TArray<UE::Tasks::FTask, TInlineAllocator<1>> Prerequisites = LastMutableTask.IsValid() 
+		? TArray<UE::Tasks::FTask, TInlineAllocator<1>>{ LastMutableTask } 
+		: TArray<UE::Tasks::FTask, TInlineAllocator<1>>{};
+
+	Prerequisites.Append(InPrerequisites);
+	
+	LastGenerationTask = UE::Tasks::Launch(DebugName, Forward<TUniqueFunction<void()>>(TaskBody), Prerequisites, TASKGRAPH_PRIORITY);
+	return LastGenerationTask;
+}
+
+
+UE::Tasks::FTask FMutableTaskGraph::AddImageStreamingTask(const TCHAR* DebugName, TUniqueFunction<void()>&& TaskBody, const TArray<UE::Tasks::FTask>& InPrerequisites)
+{
+	FScopeLock Lock(&CriticalSectionLastImageStreaming);
+	
+	TArray<UE::Tasks::FTask, TInlineAllocator<1>> Prerequisites = LastImageStreamingTask.IsValid() 
+		? TArray<UE::Tasks::FTask, TInlineAllocator<1>>{ LastImageStreamingTask } 
+		: TArray<UE::Tasks::FTask, TInlineAllocator<1>>{};
+
+	Prerequisites.Append(InPrerequisites);
+  
+	LastImageStreamingTask = UE::Tasks::Launch(DebugName, Forward<TUniqueFunction<void()>>(TaskBody), Prerequisites, TASKGRAPH_PRIORITY);
+	return LastImageStreamingTask;
+}
+
 
-	TArray<UE::Tasks::FTask, TInlineAllocator<2>> Prerequisites = LastMutableTask.IsValid() 
-		? TArray<UE::Tasks::FTask, TInlineAllocator<2>>{LastMutableTask, MutableThreadUnlockEvent } 
-		: TArray<UE::Tasks::FTask, TInlineAllocator<2>>{};
+UE::Tasks::FTask FMutableTaskGraph::AddMeshStreamingTask(const TCHAR* DebugName, TUniqueFunction<void()>&& TaskBody, const TArray<UE::Tasks::FTask>& InPrerequisites)
+{
+	FScopeLock Lock(&CriticalSectionLastMeshStreamingTask);
+	
+	TArray<UE::Tasks::FTask, TInlineAllocator<1>> Prerequisites = LastMeshStreamingTask.IsValid() 
+		? TArray<UE::Tasks::FTask, TInlineAllocator<1>>{ LastMeshStreamingTask } 
+		: TArray<UE::Tasks::FTask, TInlineAllocator<1>>{};
 
 	Prerequisites.Append(InPrerequisites);
 	
-	LastMutableTask = UE::Tasks::Launch(DebugName, Forward<TUniqueFunction<void()>>(TaskBody), Prerequisites, TASKGRAPH_PRIORITY);
 	
-	return LastMutableTask;
+	LastMeshStreamingTask = UE::Tasks::Launch(DebugName, Forward<TUniqueFunction<void()>>(TaskBody), Prerequisites, TASKGRAPH_PRIORITY);
+	return LastMeshStreamingTask;
 }
 
 
diff --git a/Engine/Plugins/Mutable/Source/MutableRuntime/Internal/MuR/System.h b/Engine/Plugins/Mutable/Source/MutableRuntime/Internal/MuR/System.h
--- a/Engine/Plugins/Mutable/Source/MutableRuntime/Internal/MuR/System.h
+++ b/Engine/Plugins/Mutable/Source/MutableRuntime/Internal/MuR/System.h
@@ -581,38 +581,7 @@ namespace UE::Mutable::Private
 				return false;
 			}
 		}
-
-		/** */
-		void CheckHitCountsCleared()
-		{
-#if !(UE_BUILD_SHIPPING || UE_BUILD_TEST)
-			//MUTABLE_CPUPROFILER_SCOPE(CheckHitCountsCleared);
-
-			//int32 IncorrectCount = 0;
-			//CodeContainer<int>::iterator it = m_opHitCount.begin();
-			//for (; it.IsValid(); ++it)
-			//{
-			//	int32 Count = *it;
-			//	if (Count>0 && Count < UE_MUTABLE_CACHE_COUNT_LIMIT)
-			//	{
-			//		// We don't manage the hitcounts of small types that don't use much memory
-			//		if (m_resources[it.get_address()].Key==2)
-			//		{
-			//			// Op hitcount should have reached 0, otherwise, it means we requested an operation but never read 
-			//			// the result.
-			//			++IncorrectCount;
-			//		}
-			//	}
-			//}
-
-			//if (IncorrectCount > 0)
-			//{
-			//	UE_LOG(LogMutableCore, Log, TEXT("The op-hit-count didn't hit 0 for %5d operations. This may mean that too much memory is cached."), IncorrectCount);
-			//	//check(false);
-			//}
-#endif
-		}
-
+		
 
 		void Clear()
 		{
@@ -1776,88 +1745,8 @@ namespace UE::Mutable::Private
 
     		return Result;
     	}
-
-
+    	
 		void LogWorkingMemory(const class CodeRunner* CurrentRunner) const;
-
-
-#if !(UE_BUILD_SHIPPING || UE_BUILD_TEST)
-		static constexpr uint64 InvalidRunnerId = 0;
-		std::atomic<uint64> DebugCurrentRunnerId {1};
-		
-		/** Temp variable with the ID of the thread running code, for debugging. */
-		uint32 DebugRunnerThreadID = FThread::InvalidThreadId;
-		uint64 DebugRunnerID = InvalidRunnerId;
-#endif
-
-		/** This is a development-only check to make sure calls to resource management happen in the correct thread. */
-		FORCEINLINE void BeginRunnerThread()
-		{
-#if !(UE_BUILD_SHIPPING || UE_BUILD_TEST)
-			// If this check fails it means have not set up correctly all the paths to debug threading for resource management.
-			check(DebugRunnerThreadID == FThread::InvalidThreadId);
-			check(DebugRunnerID == InvalidRunnerId)
-
-			DebugRunnerThreadID = FPlatformTLS::GetCurrentThreadId();
-			DebugRunnerID = ++DebugCurrentRunnerId;
-#endif
-		}
-
-		/** This is a development-only check to make sure calls to resource management happen in the correct thread. */
-		FORCEINLINE void ResetRunnerThread()
-		{
-#if !(UE_BUILD_SHIPPING || UE_BUILD_TEST)
-			// If this check fails it means have not set up correctly all the paths to debug threading for resource management.
-			check(DebugRunnerThreadID == FThread::InvalidThreadId);
-			check(DebugRunnerID != InvalidRunnerId)
-
-			DebugRunnerThreadID = FPlatformTLS::GetCurrentThreadId();
-#endif
-		}
-
-		/** This is a development-only check to make sure calls to resource management happen in the correct thread. */
-		FORCEINLINE void InvalidateRunnerThread()
-		{
-#if !(UE_BUILD_SHIPPING || UE_BUILD_TEST)
-			// If this check fails it means have not set up correctly all the paths to debug threading for resource management.
-			check(DebugRunnerThreadID != FThread::InvalidThreadId);
-			check(DebugRunnerID != InvalidRunnerId)
-
-			DebugRunnerThreadID = FThread::InvalidThreadId;
-#endif
-		}
-
-		/** This is a development-only check to make sure calls to resource management happen in the correct thread. */
-		FORCEINLINE void CheckRunnerThread()
-		{
-#if !(UE_BUILD_SHIPPING || UE_BUILD_TEST)
-			// If this check fails it means have not set up correctly all the paths to debug threading for resource management.
-			check(DebugRunnerThreadID != FThread::InvalidThreadId);
-			check(DebugRunnerID != InvalidRunnerId);
-
-			// If this check fails it means we are doing resource management from a thread that is not the CodeRunner::RunCode
-			// thread, and this is not allowed.
-			check(DebugRunnerThreadID == FPlatformTLS::GetCurrentThreadId());
-			check(DebugRunnerID == DebugCurrentRunnerId);
-#endif
-		}
-
-
-		/** This is a development-only check to make sure calls to resource management happen in the correct thread. */
-		FORCEINLINE void EndRunnerThread(const TSharedRef<FLiveInstance>& InLiveInstance)
-		{
-			InLiveInstance->Cache->CheckHitCountsCleared();
-
-#if !(UE_BUILD_SHIPPING || UE_BUILD_TEST)
-			// If this check fails it means have not set up correctly all the paths to debug threading for resource management.
-			check(DebugRunnerThreadID != FThread::InvalidThreadId);
-			check(DebugRunnerID != InvalidRunnerId);
-
-			DebugRunnerThreadID = FThread::InvalidThreadId;
-			DebugRunnerID = InvalidRunnerId;
-#endif
-		}
-
 	};
 	
 	
diff --git a/Engine/Plugins/Mutable/Source/MutableRuntime/Private/MuR/CodeRunner.h b/Engine/Plugins/Mutable/Source/MutableRuntime/Private/MuR/CodeRunner.h
--- a/Engine/Plugins/Mutable/Source/MutableRuntime/Private/MuR/CodeRunner.h
+++ b/Engine/Plugins/Mutable/Source/MutableRuntime/Private/MuR/CodeRunner.h
@@ -550,8 +550,6 @@ namespace  UE::Mutable::Private
 
 		inline TManagedPtr<FImage> CloneOrTakeOver(TManagedPtr<const FImage>&& Resource)
 		{
-			System.WorkingMemoryManager.CheckRunnerThread();
-	
 			TManagedPtr<FImage> Result;
 			if (!Resource.IsUniqueReference())
 			{
@@ -570,8 +568,6 @@ namespace  UE::Mutable::Private
 
 		[[nodiscard]] inline TManagedPtr<FMesh> CreateMesh(int32 BudgetReserveSize = 0)
 		{
-			System.WorkingMemoryManager.CheckRunnerThread();
-
 			EnsureBudgetBelow(BudgetReserveSize);
 
 			TManagedPtr<FMesh> Result = MakeManaged<FMesh>();
@@ -580,8 +576,6 @@ namespace  UE::Mutable::Private
 
 		[[nodiscard]] inline TManagedPtr<FMesh> CloneOrTakeOver(TManagedPtr<const FMesh>&& Resource)
 		{
-			System.WorkingMemoryManager.CheckRunnerThread();
-
 			TManagedPtr<FMesh> Result;
 			if (!Resource.IsUniqueReference())
 			{
@@ -632,8 +626,6 @@ namespace  UE::Mutable::Private
     	
     	inline TManagedPtr<FImage> CreateImageImp(uint32 SizeX, uint32 SizeY, uint32 Lods, EImageFormat Format, EInitializationType Init)
     	{
-    		System.WorkingMemoryManager.CheckRunnerThread();
-			
     		// Make room in the budget
     		uint32 DataSize = FImage::CalculateDataSize(SizeX, SizeY, Lods, Format);
     		EnsureBudgetBelow(DataSize);
diff --git a/Engine/Plugins/Mutable/Source/MutableRuntime/Private/MuR/CodeRunnerTasked.cpp b/Engine/Plugins/Mutable/Source/MutableRuntime/Private/MuR/CodeRunnerTasked.cpp
--- a/Engine/Plugins/Mutable/Source/MutableRuntime/Private/MuR/CodeRunnerTasked.cpp
+++ b/Engine/Plugins/Mutable/Source/MutableRuntime/Private/MuR/CodeRunnerTasked.cpp
@@ -398,13 +398,9 @@ namespace UE::Mutable::Private
 						}
 					}
 
-					System.WorkingMemoryManager.InvalidateRunnerThread();
-
 					UE::Tasks::AddNested(UE::Tasks::Launch(TEXT("CodeRunnerFromIssuedTasksTask"),
 						[Runner = AsShared(), ProfileContext = MoveTemp(ProfileContext)]() mutable
 						{
-							Runner->System.WorkingMemoryManager.ResetRunnerThread();
-
 							constexpr bool bForceInlineExecution = false;
 							Runner->Run(MoveTemp(ProfileContext), bForceInlineExecution);
 						},
@@ -432,13 +428,9 @@ namespace UE::Mutable::Private
 			{
 				if (FTimespan::FromSeconds(FPlatformTime::Seconds()) > TimeOut)
 				{
-					System.WorkingMemoryManager.InvalidateRunnerThread();
-
 					UE::Tasks::AddNested(UE::Tasks::Launch(TEXT("CodeRunnerFromTimeoutTask"),
 						[Runner = AsShared(), ProfileContext = MoveTemp(ProfileContext)]() mutable
 						{
-							Runner->System.WorkingMemoryManager.ResetRunnerThread();
-
 							constexpr bool bForceInlineExecution = false;
 							Runner->Run(MoveTemp(ProfileContext), bForceInlineExecution);
 						},
diff --git a/Engine/Plugins/Mutable/Source/MutableRuntime/Private/MuR/System.cpp b/Engine/Plugins/Mutable/Source/MutableRuntime/Private/MuR/System.cpp
--- a/Engine/Plugins/Mutable/Source/MutableRuntime/Private/MuR/System.cpp
+++ b/Engine/Plugins/Mutable/Source/MutableRuntime/Private/MuR/System.cpp
@@ -255,7 +255,6 @@ namespace UE::Mutable::Private
 
     	
     	TManagedPtr<const FInstance> Result;
-		WorkingMemoryManager.BeginRunnerThread();
 		{
 			const OP::ADDRESS RootAt = Program.States[LiveInstanceState].Root;
 			
@@ -266,7 +265,6 @@ namespace UE::Mutable::Private
 			}
 
 			// Debug check to see if we managed the op-hit-counts correctly
-			LiveInstance->Cache->CheckHitCountsCleared();
 			LiveInstance->Instance = Result;
 
 			if (!Result)
@@ -275,8 +273,6 @@ namespace UE::Mutable::Private
 				Result = MakeManaged<FInstance>();
 			}
 		}
-		WorkingMemoryManager.EndRunnerThread(LiveInstance);
-    	
 		return Result;
 	}
 
@@ -349,8 +345,6 @@ namespace UE::Mutable::Private
     	UE::Tasks::FTask RunnerTask = UE::Tasks::Launch(TEXT("FSystem::GetImage"),
 		[Context, System = this, RootAddress, LiveInstance, MipsToSkip]()
 		{
-			System->WorkingMemoryManager.BeginRunnerThread();
-
 			Context->Runner = CodeRunner::Create(
 				LiveInstance,
 				System->Settings,
@@ -382,8 +376,6 @@ namespace UE::Mutable::Private
 			{
 				Result = MakeManaged<FImage>(16, 16, 1, EImageFormat::RGBA_UByte, EInitializationType::Black);
 			}
-
-			System->WorkingMemoryManager.EndRunnerThread(LiveInstance);
 			
 			return Result;
 		},
@@ -430,7 +422,6 @@ namespace UE::Mutable::Private
 		if (GetOpDataType(OpType) == EDataType::Image)
 		{
 			// GetImageDesc may call normal execution paths where meshes are computed.
-			WorkingMemoryManager.BeginRunnerThread();
 					
 			int8 executionOptions = 0;
 			TSharedRef<CodeRunner> Runner = CodeRunner::Create(
@@ -448,8 +439,6 @@ namespace UE::Mutable::Private
 			Runner->StartRun(bForceInlineExecution);
 
 			Result = Runner->GetImageDescResult(RootAddress);
-
-			WorkingMemoryManager.EndRunnerThread(LiveInstance);
 		}
 		
 		return Result;
@@ -505,8 +494,6 @@ namespace UE::Mutable::Private
 				LiveInstance->Cache->ClearDescCache();
 			}
 
-			System->WorkingMemoryManager.BeginRunnerThread();
-
 			// GetImageDesc may call normal execution paths where meshes are computed.
 			int8 ExecutionOptions = 0;
 			Context->Runner = CodeRunner::Create(
@@ -531,7 +518,6 @@ namespace UE::Mutable::Private
 		{
 			FExtendedImageDesc Result = Context->Runner->GetImageDescResult(RootAddress);
 
-			System->WorkingMemoryManager.EndRunnerThread(LiveInstance);
 			return Result;
 		},
 		UE::Tasks::Prerequisites(RunnerTask),
@@ -608,8 +594,6 @@ namespace UE::Mutable::Private
 		UE::Tasks::FTask RunnerTask = UE::Tasks::Launch(TEXT("FSystem::GetMesh"),
 		[Context, System = this, RootAddress, LiveInstance, MeshContentFilter]()
 		{
-			System->WorkingMemoryManager.BeginRunnerThread();
-			
 			uint8 ExecutionOptions = static_cast<uint8>(MeshContentFilter);
 			Context->Runner = CodeRunner::Create(
 			LiveInstance,
@@ -639,8 +623,6 @@ namespace UE::Mutable::Private
 				Result = Context->Runner->LoadMesh(FCacheAddress(RootAddress, 0, static_cast<uint8>(MeshContentFilter)));
 			}
 			
-			System->WorkingMemoryManager.EndRunnerThread(LiveInstance);
-			
 			return Result;
 		},
 		UE::Tasks::Prerequisites(RunnerTask),
@@ -690,8 +672,6 @@ namespace UE::Mutable::Private
 		UE::Tasks::FTask RunnerTask = UE::Tasks::Launch(TEXT("FSystem::GetSkeletalMesh"),
 		[Context, System = this, RootAddress, LiveInstance, SkeletalMeshContentFilter]()
 		{
-			System->WorkingMemoryManager.BeginRunnerThread();
-
 			Context->Runner = CodeRunner::Create(
 			LiveInstance,
 					System->Settings,
@@ -720,8 +700,6 @@ namespace UE::Mutable::Private
 				Result = Context->Runner->LoadSkeletalMesh(FCacheAddress(RootAddress, 0, static_cast<uint8>(SkeletalMeshContentFilter)));
 			}
 			
-			System->WorkingMemoryManager.EndRunnerThread(LiveInstance);
-			
 			return Result;
 		},
 		UE::Tasks::Prerequisites(RunnerTask),
@@ -771,8 +749,6 @@ namespace UE::Mutable::Private
 		UE::Tasks::FTask RunnerTask = UE::Tasks::Launch(TEXT("FSystem::GetMaterial"),
 		[Context, System = this, RootAddress, LiveInstance]()
 		{
-			System->WorkingMemoryManager.BeginRunnerThread();
-				
 			Context->Runner = CodeRunner::Create(
 			LiveInstance,
 					System->Settings,
@@ -801,8 +777,6 @@ namespace UE::Mutable::Private
 				Result = Context->Runner->LoadMaterial(FCacheAddress(RootAddress, 0, 0));
 			}
 			
-			System->WorkingMemoryManager.EndRunnerThread(LiveInstance);
-			
 			return Result;
 		},
 		UE::Tasks::Prerequisites(RunnerTask),
@@ -818,9 +792,6 @@ namespace UE::Mutable::Private
 
 		LiveInstance->Instance = nullptr;
 
-		// Debug check to see if we managed the op-hit-counts correctly
-		LiveInstance->Cache->CheckHitCountsCleared();
-
 		// We don't want to clear the cache layer 1 because it contains data that can be useful for a 
 		// future update (same states, just runtime parameters changed).
 		//LiveInstance->Cache->ClearLayer1();
@@ -964,8 +935,6 @@ namespace UE::Mutable::Private
 	
 	bool FSystem::BuildBool(const TSharedRef<FLiveInstance>& InLiveInstance, OP::ADDRESS InAddress)
     {
-    	WorkingMemoryManager.BeginRunnerThread();
-    	
     	RunCode(InLiveInstance, InAddress);
 	
     	bool bResult = false;
@@ -974,16 +943,12 @@ namespace UE::Mutable::Private
     		bResult = InLiveInstance->Cache->GetBool(FCacheAddress(InAddress, 0, 0));
     	}
 	
-    	WorkingMemoryManager.EndRunnerThread(InLiveInstance);
-	
     	return bResult;    
     }
 	
 	
 	int32 FSystem::BuildInt(const TSharedRef<FLiveInstance>& InLiveInstance, OP::ADDRESS InAddress)
 	{
-		WorkingMemoryManager.BeginRunnerThread();
-    	
 		RunCode(InLiveInstance, InAddress);
 
 		int32 Result = 0;
@@ -992,16 +957,12 @@ namespace UE::Mutable::Private
 			Result = InLiveInstance->Cache->GetInt(FCacheAddress(InAddress, 0, 0));
 		}
 
-		WorkingMemoryManager.EndRunnerThread(InLiveInstance);
-
 		return Result;
 	}
 
 	
 	float FSystem::BuildScalar(const TSharedRef<FLiveInstance>& InLiveInstance, OP::ADDRESS InAddress)
 	{
-    	WorkingMemoryManager.BeginRunnerThread();
-    	
     	RunCode(InLiveInstance, InAddress);
 
     	float Result = 0.0f;		
@@ -1010,16 +971,12 @@ namespace UE::Mutable::Private
     		Result = InLiveInstance->Cache->GetScalar(FCacheAddress(InAddress, 0, 0));
     	}
 
-    	WorkingMemoryManager.EndRunnerThread(InLiveInstance);
-
     	return Result;
 	}
 
 
 	FVector4f FSystem::BuildColor(const TSharedRef<FLiveInstance>& InLiveInstance, OP::ADDRESS InAddress)
 	{
-		WorkingMemoryManager.BeginRunnerThread();
-
 		FVector4f Result(0,0,0,1);
     	
 		UE::Mutable::Private::EOpType OpType = InLiveInstance->Model->Program.GetOpType(InAddress);
@@ -1032,16 +989,12 @@ namespace UE::Mutable::Private
 			}
 		}
 
-		WorkingMemoryManager.EndRunnerThread(InLiveInstance);
-
 		return Result;
 	}
 
 	
 	FProjector FSystem::BuildProjector(const TSharedRef<FLiveInstance>& InLiveInstance, OP::ADDRESS InAddress)
 	{
-		WorkingMemoryManager.BeginRunnerThread();
-    	
     	RunCode(InLiveInstance, InAddress);
 
 		FProjector Result;
@@ -1050,16 +1003,12 @@ namespace UE::Mutable::Private
 			Result = InLiveInstance->Cache->GetProjector(FCacheAddress(InAddress, 0, 0));
 		}
 
-		WorkingMemoryManager.EndRunnerThread(InLiveInstance);
-
 		return Result;
 	}
 
 
 	TManagedPtr<const FImage> FSystem::BuildImage(const TSharedRef<FLiveInstance>& InLiveInstance, OP::ADDRESS InAddress, int32 MipsToSkip)
 	{
-		WorkingMemoryManager.BeginRunnerThread();
-
 		TManagedPtr<const FImage> Result;
     	
 		UE::Mutable::Private::EOpType OpType = InLiveInstance->Model->Program.GetOpType(InAddress);
@@ -1072,16 +1021,12 @@ namespace UE::Mutable::Private
 			}			
 		}
 
-		WorkingMemoryManager.EndRunnerThread(InLiveInstance);
-
 		return Result;
 	}
 
 
 	TManagedPtr<const FMesh> FSystem::BuildMesh(const TSharedRef<FLiveInstance>& InLiveInstance, OP::ADDRESS InRootAddress, EMeshContentFlags MeshContentFilter)
 	{
-		WorkingMemoryManager.BeginRunnerThread();
-
 		TManagedPtr<const FMesh> Result;
     	
 		UE::Mutable::Private::EOpType OpType = InLiveInstance->Model->Program.GetOpType(InRootAddress);
@@ -1095,16 +1040,12 @@ namespace UE::Mutable::Private
 			}
 		}
 
-		WorkingMemoryManager.EndRunnerThread(InLiveInstance);
-
 		return Result;
 	}
 
 
 	TManagedPtr<const FInstance> FSystem::BuildInstance(const TSharedRef<FLiveInstance>& InLiveInstance, OP::ADDRESS InAddress)
 	{
-		WorkingMemoryManager.BeginRunnerThread();
-
 		TManagedPtr<const FInstance> Result;
     	
 		UE::Mutable::Private::EOpType OpType = InLiveInstance->Model->Program.GetOpType(InAddress);
@@ -1117,16 +1058,12 @@ namespace UE::Mutable::Private
 			}
 		}
 
-		WorkingMemoryManager.EndRunnerThread(InLiveInstance);
-
 		return Result;
 	}
 
 
 	TManagedPtr<const FLayout> FSystem::BuildLayout(const TSharedRef<FLiveInstance>& InLiveInstance, OP::ADDRESS InAddress)
 	{
-		WorkingMemoryManager.BeginRunnerThread();
-
 		TManagedPtr<const FLayout> Result;
     	
 		if (InLiveInstance->Model->Program.States[0].Root)
@@ -1142,16 +1079,12 @@ namespace UE::Mutable::Private
 			}
 		}
 
-		WorkingMemoryManager.EndRunnerThread(InLiveInstance);
-
 		return Result;
 	}
 
 
 	TManagedPtr<const String> FSystem::BuildString(const TSharedRef<FLiveInstance>& InLiveInstance, OP::ADDRESS InAddress)
 	{
-		WorkingMemoryManager.BeginRunnerThread();
-
 		TManagedPtr<const String> Result;
     	
 		if (InLiveInstance->Model->Program.States[0].Root)
@@ -1167,8 +1100,6 @@ namespace UE::Mutable::Private
 			}
 		}
 
-		WorkingMemoryManager.EndRunnerThread(InLiveInstance);
-
 		return Result;
 	}
 	
