[Mutable] Retrived check that indicates if a surface edit node is being used correctly.
+Removed uncessary textures from edit nodes.

#rnx

#Synced-CL 50245735

diff --git a/Engine/Plugins/Mutable/Source/CustomizableObjectEditor/Private/MuCOE/GenerateMutableSource/GenerateMutableSourceModifier.cpp b/Engine/Plugins/Mutable/Source/CustomizableObjectEditor/Private/MuCOE/GenerateMutableSource/GenerateMutableSourceModifier.cpp
--- a/Engine/Plugins/Mutable/Source/CustomizableObjectEditor/Private/MuCOE/GenerateMutableSource/GenerateMutableSourceModifier.cpp
+++ b/Engine/Plugins/Mutable/Source/CustomizableObjectEditor/Private/MuCOE/GenerateMutableSource/GenerateMutableSourceModifier.cpp
@@ -317,7 +317,6 @@ UE::Mutable::Private::Ptr<UE::Mutable::Private::NodeModifier> GenerateMutableSou
 			SurfNode->LODs[LODIndex].MeshAdd = AddMeshNode;
 
 			const int32 NumImages = TypedNodeExt->GetNumParameters(EMaterialParameterType::Texture);
-			SurfNode->LODs[LODIndex].Textures.SetNum(NumImages);
 			for (int32 ImageIndex = 0; ImageIndex < NumImages; ++ImageIndex)
 			{
 				UE::Mutable::Private::NodeImagePtr ImageNode;
@@ -353,9 +352,14 @@ UE::Mutable::Private::Ptr<UE::Mutable::Private::NodeModifier> GenerateMutableSou
 					}
 				}
 
+				if (ImageNode)
+				{
+					UE::Mutable::Private::NodeModifierSurfaceEdit::FTexture Texture;
+					Texture.Extend = ImageNode;
+					Texture.MaterialParameterKey = UE::Mutable::Private::FParameterKey(MaterialParameterName, MaterialParameterLayerIndex);
 
-				SurfNode->LODs[LODIndex].Textures[ImageIndex].Extend = ImageNode;
-				SurfNode->LODs[LODIndex].Textures[ImageIndex].MaterialParameterKey = UE::Mutable::Private::FParameterKey(MaterialParameterName, MaterialParameterLayerIndex);
+					SurfNode->LODs[LODIndex].Textures.Add(Texture);
+				}
 			}
 		}
 
@@ -449,9 +453,8 @@ UE::Mutable::Private::Ptr<UE::Mutable::Private::NodeModifier> GenerateMutableSou
 			GenerationContext.FromLOD = 0;
 			GenerationContext.CurrentLOD = LODIndex;
 
-
 			const int32 NumImages = TypedNodeEdit->GetNumParameters(EMaterialParameterType::Texture);
-			SurfNode->LODs[LODIndex].Textures.SetNum(NumImages);
+
 			for (int32 ImageIndex = 0; ImageIndex < NumImages; ++ImageIndex)
 			{
 				const FNodeMaterialParameterId ImageId = TypedNodeEdit->GetParameterId(EMaterialParameterType::Texture, ImageIndex);
@@ -462,14 +465,12 @@ UE::Mutable::Private::Ptr<UE::Mutable::Private::NodeModifier> GenerateMutableSou
 					//check(ParentMaterialNode->IsImageMutableMode(ImageIndex)); // Ensured at graph time. If it fails, something is wrong.
 
 					const UEdGraphPin* ConnectedImagePin = FollowInputPin(*TypedNodeEdit->GetUsedImagePin(ImageId));
-
-					UE::Mutable::Private::NodeModifierSurfaceEdit::FTexture& ImagePatch = SurfNode->LODs[LODIndex].Textures[ImageIndex];
-					
 					FName MaterialParameterName = TypedNodeEdit->GetParameterName(EMaterialParameterType::Texture, ImageIndex);
 					int32 MaterialParameterLayerIndex = TypedNodeEdit->GetParameterLayerIndex(EMaterialParameterType::Texture, ImageIndex);
 					
-					// \todo: expose these two options?
+					UE::Mutable::Private::NodeModifierSurfaceEdit::FTexture ImagePatch;
 					ImagePatch.MaterialParameterKey = UE::Mutable::Private::FParameterKey(MaterialParameterName, MaterialParameterLayerIndex);
+					// \todo: expose these two options?
 					ImagePatch.PatchBlendType = UE::Mutable::Private::EBlendType::BT_BLEND;
 					ImagePatch.bPatchApplyToAlpha = true;
 
@@ -490,6 +491,7 @@ UE::Mutable::Private::Ptr<UE::Mutable::Private::NodeModifier> GenerateMutableSou
 					FIntPoint GridSize = TypedNodeEdit->Layout->GetGridSize();
 					FVector2f GridSizeF = FVector2f(GridSize);
 					ImagePatch.PatchBlocks.Reserve(TypedNodeEdit->Layout->Blocks.Num());
+
 					for (const FCustomizableObjectLayoutBlock& LayoutBlock : TypedNodeEdit->Layout->Blocks)
 					{
 						FBox2f Rect;
@@ -497,6 +499,8 @@ UE::Mutable::Private::Ptr<UE::Mutable::Private::NodeModifier> GenerateMutableSou
 						Rect.Max = FVector2f(LayoutBlock.Max) / GridSizeF;
 						ImagePatch.PatchBlocks.Add(Rect);
 					}
+
+					SurfNode->LODs[LODIndex].Textures.Add(ImagePatch);
 				}
 			}
 		}
diff --git a/Engine/Plugins/Mutable/Source/MutableTools/Private/MuT/CodeGenerator.cpp b/Engine/Plugins/Mutable/Source/MutableTools/Private/MuT/CodeGenerator.cpp
--- a/Engine/Plugins/Mutable/Source/MutableTools/Private/MuT/CodeGenerator.cpp
+++ b/Engine/Plugins/Mutable/Source/MutableTools/Private/MuT/CodeGenerator.cpp
@@ -874,9 +874,8 @@ namespace UE::Mutable::Private
 		int32 ComponentId = Options.Component ? Options.Component->Id : -1;
 		GetModifiersFor(ComponentId, SurfaceNode->Tags, bModifiersForBeforeOperations, Modifiers);
 
-		// TODO (Max): UE-354034 Fix this function and uncomment it
 		// This pass on the modifiers is only to detect errors that cannot be detected at the point they are applied.
-		//CheckModifiersForSurface(*SurfaceNode, Modifiers, Options.LODIndex);
+		CheckModifiersForSurface(*SurfaceNode, Modifiers, Options.LODIndex);
 		
         // Generate the mesh
 		FMeshGenerationStaticOptions MeshStaticOptions(ComponentId, Options.LODIndex);
@@ -2769,13 +2768,10 @@ namespace UE::Mutable::Private
 				const NodeModifierSurfaceEdit* Edit = static_cast<const NodeModifierSurfaceEdit*>(Mod.Node);
 				const NodeMaterialModify* Material = static_cast<const NodeMaterialModify*>(Node.Material.get());
 
-				bool bAffectsCurrentLOD = Edit->LODs.IsValidIndex(CurrentLOD);
-				if (!bAffectsCurrentLOD)
-				{
-					continue;
-				}
-
-				if (Material->ImageParameters.IsEmpty() || Edit->LODs[CurrentLOD].Textures.IsEmpty())
+				// If this edit is not affected by the current LOD, then ignore.
+				// If both arrays are empty means that there are no textures to check
+				// If there is a MeshRemove, this edit does not affect textures
+				if (!Edit->LODs.IsValidIndex(CurrentLOD) || (Material->ImageParameters.IsEmpty() && Edit->LODs[CurrentLOD].Textures.IsEmpty()) || Edit->LODs[CurrentLOD].MeshRemove)
 				{
 					continue;
 				}
