[Mutable] Create FPhysicsBodies when converting meshes from runtime data.
Move PhysicsAssets and AdditionalPhysicsAssets to Passthrough Objects.

#rnx

#Synced-CL 50203008

diff --git a/Engine/Plugins/Mutable/Source/CustomizableObject/Internal/MuCO/CustomizableObjectInstancePrivate.h b/Engine/Plugins/Mutable/Source/CustomizableObject/Internal/MuCO/CustomizableObjectInstancePrivate.h
--- a/Engine/Plugins/Mutable/Source/CustomizableObject/Internal/MuCO/CustomizableObjectInstancePrivate.h
+++ b/Engine/Plugins/Mutable/Source/CustomizableObject/Internal/MuCO/CustomizableObjectInstancePrivate.h
@@ -72,23 +72,6 @@ enum ECOInstanceFlags
 ENUM_CLASS_FLAGS(ECOInstanceFlags);
 
 
-USTRUCT()
-struct FReferencedPhysicsAssets
-{
-	GENERATED_USTRUCT_BODY();
-	
-	TArray<int32> PhysicsAssetToLoad;
-	
-	UPROPERTY(Transient)
-	TArray< TObjectPtr<UPhysicsAsset> > PhysicsAssetsToMerge;
-
-	TArray<int32> AdditionalPhysicsAssetsToLoad;
-
-	UPROPERTY(Transient)
-	TArray<TObjectPtr<UPhysicsAsset>> AdditionalPhysicsAssets;
-};
-
-
 USTRUCT()
 struct FReferencedSkeletons
 {
@@ -137,10 +120,6 @@ struct FCustomizableInstanceComponentData
 	UPROPERTY(Transient)
 	FReferencedSkeletons Skeletons;
 	
-	/** PhysicsAssets required by the current generated instance. PhysicsAssets to be loaded and merged.*/
-	UPROPERTY(Transient)
-	FReferencedPhysicsAssets PhysicsAssets;
-
 	FComponentDescriptor LastComponentDescriptor;
 
 	UPROPERTY(Transient)
diff --git a/Engine/Plugins/Mutable/Source/CustomizableObject/Internal/MuCO/CustomizableObjectPrivate.h b/Engine/Plugins/Mutable/Source/CustomizableObject/Internal/MuCO/CustomizableObjectPrivate.h
--- a/Engine/Plugins/Mutable/Source/CustomizableObject/Internal/MuCO/CustomizableObjectPrivate.h
+++ b/Engine/Plugins/Mutable/Source/CustomizableObject/Internal/MuCO/CustomizableObjectPrivate.h
@@ -484,9 +484,6 @@ struct FAnimBpOverridePhysicsAssetsInfo
 	UPROPERTY()
 	TSoftClassPtr<UAnimInstance> AnimInstanceClass;
 
-	UPROPERTY()
-	TSoftObjectPtr<UPhysicsAsset> SourceAsset;
-
 	UPROPERTY()
 	int32 PropertyIndex = -1;
 
@@ -1566,6 +1563,8 @@ enum class ECustomizableObjectVersions : int32
 
 	ClothCoreSerializationIssue2,
 
+	MovePhysicsAssetCore,
+
 	// -----<new versions can be added above this line>--------
 	LastCustomizableObjectVersion
 };
diff --git a/Engine/Plugins/Mutable/Source/CustomizableObject/Internal/MuCO/UnrealConversionUtils.h b/Engine/Plugins/Mutable/Source/CustomizableObject/Internal/MuCO/UnrealConversionUtils.h
--- a/Engine/Plugins/Mutable/Source/CustomizableObject/Internal/MuCO/UnrealConversionUtils.h
+++ b/Engine/Plugins/Mutable/Source/CustomizableObject/Internal/MuCO/UnrealConversionUtils.h
@@ -26,6 +26,7 @@ namespace UE::Mutable::Private
 	class FPhysicsBody;
 	class FMesh;
 	class FMeshBufferSet;
+	class FPhysicsBody;
 }
 struct FReferenceSkeleton;
 struct FMutableSurfaceMetadata;
@@ -33,6 +34,7 @@ class UModelResources;
 class FSkeletalMeshLODRenderData;
 class USkeletalMesh;
 class USkeleton;
+class USkeletalBodySetup;
 
 namespace UnrealConversionUtils
 {
@@ -163,4 +165,10 @@ namespace UnrealConversionUtils
 	
 	
 	void ClothVertexBuffers(FSkeletalMeshLODRenderData& LODResource, const UE::Mutable::Private::FMesh& MutableMesh);
+
+	CUSTOMIZABLEOBJECT_API void PropagateBoneUsageFlagsThroughMeshPose(UE::Mutable::Private::FMesh& Mesh);
+
+	CUSTOMIZABLEOBJECT_API UE::Mutable::Private::TManagedPtr<UE::Mutable::Private::FPhysicsBody> CreatePhysicsBodyForMesh(
+		const UPhysicsAsset& PhysicsAsset,
+		const UE::Mutable::Private::FMesh& Mesh);
 }
diff --git a/Engine/Plugins/Mutable/Source/CustomizableObject/Private/MuCO/CustomizableObject.cpp b/Engine/Plugins/Mutable/Source/CustomizableObject/Private/MuCO/CustomizableObject.cpp
--- a/Engine/Plugins/Mutable/Source/CustomizableObject/Private/MuCO/CustomizableObject.cpp
+++ b/Engine/Plugins/Mutable/Source/CustomizableObject/Private/MuCO/CustomizableObject.cpp
@@ -3309,7 +3309,7 @@ void UCustomizableObjectPrivate::LogMemory()
 		int32 PhysicsTotal = Program.ConstantPhysicsBodies.GetAllocatedSize();
 		for (const UE::Mutable::Private::TManagedPtr<const UE::Mutable::Private::FPhysicsBody>& Entry : Program.ConstantPhysicsBodies)
 		{
-			PhysicsTotal += sizeof(UE::Mutable::Private::FPhysicsBody) + Entry->BodiesBoneNames.GetAllocatedSize() + Entry->Bodies.GetAllocatedSize() + Entry->BodiesCustomIds.GetAllocatedSize();
+			PhysicsTotal += sizeof(UE::Mutable::Private::FPhysicsBody) + Entry->BodiesBoneNames.GetAllocatedSize() + Entry->Bodies.GetAllocatedSize();
 			for (const UE::Mutable::Private::FPhysicsBodyAggregate& Body : Entry->Bodies)
 			{
 				PhysicsTotal += Body.Spheres.GetAllocatedSize() + Body.Boxes.GetAllocatedSize() + Body.Convex.GetAllocatedSize() + Body.Sphyls.GetAllocatedSize() + Body.TaperedCapsules.GetAllocatedSize();
@@ -3413,7 +3413,6 @@ void UCustomizableObjectPrivate::LogMemory()
 bool FAnimBpOverridePhysicsAssetsInfo::operator==(const FAnimBpOverridePhysicsAssetsInfo& Rhs) const
 {
 	return AnimInstanceClass == Rhs.AnimInstanceClass &&
-		SourceAsset == Rhs.SourceAsset &&
 		PropertyIndex == Rhs.PropertyIndex;
 }
 
@@ -3484,25 +3483,20 @@ uint32 GetTypeHash(const FTextureCache::FId& Key)
 FArchive& operator<<(FArchive& Ar, FAnimBpOverridePhysicsAssetsInfo& Info)
 {
 	FString AnimInstanceClassPathString;
-	FString PhysicsAssetPathString;
 
 	if (Ar.IsLoading())
 	{
 		Ar << AnimInstanceClassPathString;
-		Ar << PhysicsAssetPathString;
 		Ar << Info.PropertyIndex;
 
 		Info.AnimInstanceClass = TSoftClassPtr<UAnimInstance>(AnimInstanceClassPathString);
-		Info.SourceAsset = TSoftObjectPtr<UPhysicsAsset>(FSoftObjectPath(PhysicsAssetPathString));
 	}
 
 	if (Ar.IsSaving())
 	{
 		AnimInstanceClassPathString = Info.AnimInstanceClass.ToString();
-		PhysicsAssetPathString = Info.SourceAsset.ToString();
 
 		Ar << AnimInstanceClassPathString;
-		Ar << PhysicsAssetPathString;
 		Ar << Info.PropertyIndex;
 	}
 
diff --git a/Engine/Plugins/Mutable/Source/CustomizableObject/Private/MuCO/CustomizableObjectInstance.cpp b/Engine/Plugins/Mutable/Source/CustomizableObject/Private/MuCO/CustomizableObjectInstance.cpp
--- a/Engine/Plugins/Mutable/Source/CustomizableObject/Private/MuCO/CustomizableObjectInstance.cpp
+++ b/Engine/Plugins/Mutable/Source/CustomizableObject/Private/MuCO/CustomizableObjectInstance.cpp
@@ -53,6 +53,8 @@
 #include "MuCO/Plugins/IMutableClothingModule.h"
 #include "Rendering/SkeletalMeshModel.h"
 
+#include "MuR/PassthroughObject.h"
+
 #include UE_INLINE_GENERATED_CPP_BY_NAME(CustomizableObjectInstance)
 
 #if WITH_EDITOR
@@ -1480,18 +1482,25 @@ UPhysicsAsset* GetOrBuildMainPhysicsAsset(
 		return nullptr;
 	}
 
-	FCustomizableInstanceComponentData& ComponentData = InstancePrivate->GetComponentData(Component->Id);
+	const FName ComponentName = Context->ComponentNames[Component->Id];
 	
 	TArray<TObjectPtr<UPhysicsAsset>> ValidAssets;
 
-	const int32 NumPhysicsAssets = ComponentData.PhysicsAssets.PhysicsAssetsToMerge.Num();
-	for (int32 PhysicsAssetIndex = 0; PhysicsAssetIndex < NumPhysicsAssets; ++PhysicsAssetIndex)
+	const int32 FirstLOD = Context->bStreamMeshLODs ?
+		Context->FirstLODAvailable[ComponentName] :
+		Context->GetFirstRequestedLOD()[ComponentName];
+	
+	for (int32 LODIndex = FirstLOD; LODIndex < Component->LODCount; ++LODIndex)
 	{
-		const TObjectPtr<UPhysicsAsset>& PhysicsAsset = ComponentData.PhysicsAssets.PhysicsAssetsToMerge[PhysicsAssetIndex];
+		const TSharedRef<FInstanceUpdateData::FLOD>& LOD = Context->InstanceUpdateData.LODs[Component->FirstLOD + LODIndex];
+		check(LOD->Mesh);
 
-		if (PhysicsAsset)
+		for (const UE::Mutable::Private::TPassthroughObjectPtr<UPhysicsAsset>& PassthroughPhysicsAsset : LOD->Mesh->PhysicsAssets)
 		{
-			ValidAssets.AddUnique(PhysicsAsset);
+			if (UPhysicsAsset* PhysicsAsset = PassthroughPhysicsAsset.Get())
+			{
+				ValidAssets.AddUnique(PhysicsAsset);
+			}
 		}
 	}
 
@@ -1739,8 +1748,6 @@ UPhysicsAsset* GetOrBuildMainPhysicsAsset(
 	Result->UpdateBodySetupIndexMap();
 	Result->UpdateBoundsBodiesArray();
 
-	ComponentData.PhysicsAssets.PhysicsAssetsToMerge.Empty();
-
 	return Result;
 }
 
@@ -2156,12 +2163,19 @@ bool UCustomizableInstancePrivate::UpdateSkeletalMesh_PostBeginUpdate0(UCustomiz
 #endif
 		}
 
-		const int32 NumAdditionalPhysicsNum = ComponentMesh->AdditionalPhysicsBodies.Num();
-		for (int32 I = 0; I < NumAdditionalPhysicsNum; ++I)
+		const int32 NumAdditionalPhysicsBodies = ComponentMesh->AdditionalPhysicsBodies.Num();
+		check(ComponentMesh->AdditionalPhysicsAssets.Num() == NumAdditionalPhysicsBodies)
+		for (int32 I = 0; I < NumAdditionalPhysicsBodies; ++I)
 		{
-			UE::Mutable::Private::TManagedPtr<const UE::Mutable::Private::FPhysicsBody> AdditionalPhysicsBody = ComponentMesh->AdditionalPhysicsBodies[I];
-			
-			check(AdditionalPhysicsBody);
+			const UE::Mutable::Private::TManagedPtr<const UE::Mutable::Private::FPhysicsBody>& AdditionalPhysicsBody = ComponentMesh->AdditionalPhysicsBodies[I];
+			TObjectPtr<UPhysicsAsset> PhysicsAssetTemplate = ComponentMesh->AdditionalPhysicsAssets[I].Get();
+			check(AdditionalPhysicsBody && PhysicsAssetTemplate);
+
+			if (!AdditionalPhysicsBody || !PhysicsAssetTemplate)
+			{
+				continue;
+			}
+
 			if (!AdditionalPhysicsBody->bBodiesModified)
 			{
 				continue;
@@ -2181,16 +2195,12 @@ bool UCustomizableInstancePrivate::UpdateSkeletalMesh_PostBeginUpdate0(UCustomiz
 
 			FAnimBpGeneratedPhysicsAssets& PhysicsAssetsUsedByAnimBp = AnimBpPhysicsAssets.FindOrAdd(AnimInstanceClass);
 
-			TObjectPtr<UPhysicsAsset> PhysicsAssetTemplate = TObjectPtr<UPhysicsAsset>(Info.SourceAsset.Get());
-
-			check(PhysicsAssetTemplate);
-
 			FAnimInstanceOverridePhysicsAsset& Entry =
 				PhysicsAssetsUsedByAnimBp.AnimInstancePropertyIndexAndPhysicsAssets.Emplace_GetRef();
 
 			Entry.PropertyIndex = Info.PropertyIndex;
 			Entry.PhysicsAsset = MakePhysicsAssetFromTemplateAndMutableBody(
-				Context,PhysicsAssetTemplate, AdditionalPhysicsBody.Get(), ComponentIndex);
+				Context, PhysicsAssetTemplate, AdditionalPhysicsBody.Get(), ComponentIndex);
 		}
 
 		// Add sockets from the SkeletalMesh of reference and from the MutableMesh
@@ -4589,25 +4599,7 @@ UE::Tasks::FTask LoadAdditionalAssetsAndData(const TSharedRef<FUpdateContextPriv
 			for (int32 TagIndex = 0; TagIndex < MutableMesh->GetTagCount(); ++TagIndex)
 			{
 				FString Tag = MutableMesh->GetTag(TagIndex);
-				if (Tag.RemoveFromStart("__PA:"))
-				{
-					const int32 AssetIndex = FCString::Atoi(*Tag);
-					const TSoftObjectPtr<UPhysicsAsset>& PhysicsAsset = ModelResources.PhysicsAssets.IsValidIndex(AssetIndex) ? ModelResources.PhysicsAssets[AssetIndex] : nullptr;
 
-					if (!PhysicsAsset.IsNull())
-					{
-						if (UPhysicsAsset* LoadedPhysicsAssetPtr = PhysicsAsset.Get())
-						{
-							ComponentData.PhysicsAssets.PhysicsAssetsToMerge.Add(LoadedPhysicsAssetPtr);
-						}
-						else
-						{
-							ComponentData.PhysicsAssets.PhysicsAssetToLoad.Add(AssetIndex);
-							AssetsToStream.AddUnique(PhysicsAsset.ToSoftObjectPath());
-						}
-					}
-				}
-				
 				if (Tag.RemoveFromStart("__AnimBP:"))
 				{
 					FString SlotIndexString, AnimBpIndexString;
@@ -4639,15 +4631,6 @@ UE::Tasks::FTask LoadAdditionalAssetsAndData(const TSharedRef<FUpdateContextPriv
 				}
 #endif
 			}
-
-			const int32 AdditionalPhysicsNum = MutableMesh->AdditionalPhysicsBodies.Num();
-			for (int32 I = 0; I < AdditionalPhysicsNum; ++I)
-			{
-				const int32 ExternalId = MutableMesh->AdditionalPhysicsBodies[I]->CustomId;
-				
-				ComponentData.PhysicsAssets.AdditionalPhysicsAssetsToLoad.Add(ExternalId);
-				AssetsToStream.Add(ModelResources.AnimBpOverridePhysiscAssetsInfo[ExternalId].SourceAsset.ToSoftObjectPath());
-			}
 		}
 
 
@@ -4856,34 +4839,6 @@ void AdditionalAssetsAsyncLoaded(const TSharedRef<FUpdateContextPrivate> Context
 			}
 		}
 
-		// Loaded PhysicsAssets
-		FReferencedPhysicsAssets& PhysicsAssets = ComponentData.PhysicsAssets;
-		for(const int32 PhysicsAssetIndex : PhysicsAssets.PhysicsAssetToLoad)
-		{
-			check(ModelResources.PhysicsAssets.IsValidIndex(PhysicsAssetIndex));
-			const TSoftObjectPtr<UPhysicsAsset>& PhysicsAsset = ModelResources.PhysicsAssets[PhysicsAssetIndex];
-			PhysicsAssets.PhysicsAssetsToMerge.Add(PhysicsAsset.Get());
-
-#if WITH_EDITOR
-			if (!PhysicsAsset.Get())
-			{
-				if (!PhysicsAsset.IsNull())
-				{
-					FString ErrorMsg = FString::Printf(TEXT("Mutable couldn't load the PhysicsAsset [%s] and won't be merged. If it has been deleted or renamed, please recompile all the mutable objects that use it."), *PhysicsAsset.GetAssetName());
-					UE_LOG(LogMutable, Error, TEXT("%s"), *ErrorMsg);
-
-					FMessageLog MessageLog("Mutable");
-					MessageLog.Notify(FText::FromString(ErrorMsg), EMessageSeverity::Error, true);
-				}
-				else
-				{
-					ensure(false); // Couldn't load the PhysicsAsset, and we don't know which PhysicsAsset
-				}
-			}
-#endif
-		}
-		PhysicsAssets.PhysicsAssetToLoad.Empty();
-
 		// Loaded anim BPs
 		for (TPair<FName, TSoftClassPtr<UAnimInstance>>& SlotAnimBP : ComponentData.AnimSlotToBP)
 		{
@@ -4905,17 +4860,6 @@ void AdditionalAssetsAsyncLoaded(const TSharedRef<FUpdateContextPrivate> Context
 			}
 #endif
 		}
-
-		const int32 AdditionalPhysicsNum = ComponentData.PhysicsAssets.AdditionalPhysicsAssetsToLoad.Num();
-		ComponentData.PhysicsAssets.AdditionalPhysicsAssets.Reserve(AdditionalPhysicsNum);
-		for (int32 I = 0; I < AdditionalPhysicsNum; ++I)
-		{
-			// Make the loaded assets references strong.
-			const int32 AnimBpPhysicsOverrideIndex = ComponentData.PhysicsAssets.AdditionalPhysicsAssetsToLoad[I];
-			ComponentData.PhysicsAssets.AdditionalPhysicsAssets.Add( 
-				ModelResources.AnimBpOverridePhysiscAssetsInfo[AnimBpPhysicsOverrideIndex].SourceAsset.Get());
-		}
-		ComponentData.PhysicsAssets.AdditionalPhysicsAssetsToLoad.Empty();
 	}
 	
 #if WITH_EDITOR
diff --git a/Engine/Plugins/Mutable/Source/CustomizableObject/Private/MuCO/UnrealConversionUtils.cpp b/Engine/Plugins/Mutable/Source/CustomizableObject/Private/MuCO/UnrealConversionUtils.cpp
--- a/Engine/Plugins/Mutable/Source/CustomizableObject/Private/MuCO/UnrealConversionUtils.cpp
+++ b/Engine/Plugins/Mutable/Source/CustomizableObject/Private/MuCO/UnrealConversionUtils.cpp
@@ -20,6 +20,8 @@
 #include "SkinnedAssetCompiler.h"
 #include "Rendering/SkeletalMeshRenderData.h"
 #include "Rendering/MorphTargetVertexCodec.h"
+#include "PhysicsEngine/PhysicsAsset.h"
+#include "PhysicsEngine/SkeletalBodySetup.h"
 
 
 class USkeleton;
@@ -1255,7 +1257,29 @@ namespace UnrealConversionUtils
 
 					Result->SetBoneMap(MutableBoneMap);
 				}
-					
+
+				// Physics
+				TObjectPtr<UPhysicsAsset> PhysicsAsset = SkeletalMesh->GetPhysicsAsset();
+				if (!bIgnoreSkeleton && PhysicsAsset && Result->GetSkeleton()/* && Context.Options.bPhysicsAssetMergeEnabled*/)
+				{
+					UE::Mutable::Private::TManagedPtr<UE::Mutable::Private::FPhysicsBody> PhysicsBody = UnrealConversionUtils::CreatePhysicsBodyForMesh(*PhysicsAsset, *Result);
+
+					if (PhysicsBody)
+					{
+						Result->SetPhysicsBody(PhysicsBody);
+						Result->PhysicsAssets.Emplace(UE::Mutable::Private::TPassthroughObjectPtr(PhysicsAsset.Get()));
+
+						// Set Physics usage flag
+						const int32 NumBodySetups = PhysicsBody->GetBodyCount();
+						for (int32 BodyIndex = 0; BodyIndex < NumBodySetups; ++BodyIndex)
+						{
+							const int32 PoseIndex = Result->Skeleton->BoneNames.Find(PhysicsBody->BodiesBoneNames[BodyIndex]); // Bone Poses are 1 : 1 with the skeleton bones
+							check(PoseIndex != INDEX_NONE);
+							EnumAddFlags(Result->BonePoses[PoseIndex].BoneUsageFlags, UE::Mutable::Private::EBoneUsageFlags::Physics);
+						}
+					}
+				}
+				
 
 				// Mesh data
 				int32 FirstVertexIndex = Context->OriginalLOD->RenderSections[SectionIndex].BaseVertexIndex;
@@ -1941,4 +1965,188 @@ namespace UnrealConversionUtils
 			}
 		}
 	}
+
+
+	void PropagateBoneUsageFlagsThroughMeshPose(UE::Mutable::Private::FMesh& Mesh)
+	{
+		if (!Mesh.Skeleton)
+		{
+			return;
+		}
+
+		const int32 NumBonePoses = Mesh.BonePoses.Num();
+		check(NumBonePoses == Mesh.Skeleton->GetNumBones());
+
+		for (int32 PoseIndex = NumBonePoses - 1; PoseIndex >= 0; --PoseIndex)
+		{
+			const UE::Mutable::Private::EBoneUsageFlags BoneUsageFlags = Mesh.BonePoses[PoseIndex].BoneUsageFlags;
+
+			constexpr UE::Mutable::Private::EBoneUsageFlags FlagsToPropagate =
+				UE::Mutable::Private::EBoneUsageFlags::Skinning | UE::Mutable::Private::EBoneUsageFlags::Physics | UE::Mutable::Private::EBoneUsageFlags::Deform;
+			if (EnumHasAnyFlags(BoneUsageFlags, FlagsToPropagate))
+			{
+				const UE::Mutable::Private::EBoneUsageFlags ParentPropagationFlags =
+					(EnumHasAnyFlags(BoneUsageFlags, UE::Mutable::Private::EBoneUsageFlags::Skinning)
+						? UE::Mutable::Private::EBoneUsageFlags::SkinningParent : UE::Mutable::Private::EBoneUsageFlags::None) |
+					(EnumHasAnyFlags(BoneUsageFlags, UE::Mutable::Private::EBoneUsageFlags::Physics)
+						? UE::Mutable::Private::EBoneUsageFlags::PhysicsParent : UE::Mutable::Private::EBoneUsageFlags::None) |
+					(EnumHasAnyFlags(BoneUsageFlags, UE::Mutable::Private::EBoneUsageFlags::Deform)
+						? UE::Mutable::Private::EBoneUsageFlags::DeformParent : UE::Mutable::Private::EBoneUsageFlags::None);
+
+				int32 ParentPoseIndex = Mesh.Skeleton->BoneParents[PoseIndex]; // Bone Poses are 1 : 1 with the skeleton bones
+				while (ParentPoseIndex != INDEX_NONE)
+				{
+					if (EnumHasAllFlags(Mesh.BonePoses[ParentPoseIndex].BoneUsageFlags, ParentPropagationFlags))
+					{
+						break;
+					}
+
+					EnumAddFlags(Mesh.BonePoses[ParentPoseIndex].BoneUsageFlags, ParentPropagationFlags);
+
+					ParentPoseIndex = Mesh.Skeleton->BoneParents[ParentPoseIndex];
+				}
+			}
+		}
+	}
+
+
+	UE::Mutable::Private::TManagedPtr<UE::Mutable::Private::FPhysicsBody> CreatePhysicsBodyForMesh(const UPhysicsAsset& PhysicsAsset, const UE::Mutable::Private::FMesh& Mesh)
+	{
+		const TArray<TObjectPtr<USkeletalBodySetup>>& BodySetups = PhysicsAsset.SkeletalBodySetups;
+		const int32 NumBodySetups = BodySetups.Num();
+
+		TArray<int32> ActiveBodySetups;
+		ActiveBodySetups.Reserve(NumBodySetups);
+
+		check(Mesh.Skeleton);
+
+		for (int32 BodyIndex = 0; BodyIndex < NumBodySetups; ++BodyIndex)
+		{
+			const TObjectPtr<USkeletalBodySetup>& BodySetup = BodySetups[BodyIndex];
+			if (ensure(BodySetup))
+			{
+				const int32 BoneIndex = Mesh.Skeleton->BoneNames.Find(BodySetup->BoneName);
+				if (BoneIndex != INDEX_NONE)
+				{
+					ActiveBodySetups.Add(BodyIndex);
+				}
+			}
+		}
+
+		if (ActiveBodySetups.IsEmpty())
+		{
+			return nullptr;
+		}
+
+		UE::Mutable::Private::TManagedPtr<UE::Mutable::Private::FPhysicsBody> PhysicsBody =
+			UE::Mutable::Private::MakeManaged<UE::Mutable::Private::FPhysicsBody>();
+
+		const int32 NumActiveBodySetups = ActiveBodySetups.Num();
+		PhysicsBody->SetBodyCount(NumActiveBodySetups);
+
+		auto GetKBodyElemFlags = [](const FKShapeElem& KElem) -> uint32
+			{
+				uint8 ElemCollisionEnabled = static_cast<uint8>(KElem.GetCollisionEnabled());
+
+				uint32 Flags = static_cast<uint32>(ElemCollisionEnabled);
+				Flags = Flags | (static_cast<uint32>(KElem.GetContributeToMass()) << 8);
+
+				return Flags;
+			};
+
+		for (int32 BodyIndex = 0; BodyIndex < NumActiveBodySetups; ++BodyIndex)
+		{
+			const TObjectPtr<USkeletalBodySetup>& BodySetup = BodySetups[ActiveBodySetups[BodyIndex]];
+
+			PhysicsBody->SetBody(BodyIndex, BodySetup->BoneName);
+
+			const int32 NumSpheres = BodySetup->AggGeom.SphereElems.Num();
+			PhysicsBody->SetSphereCount(BodyIndex, NumSpheres);
+
+			for (int32 I = 0; I < NumSpheres; ++I)
+			{
+				const FKSphereElem& SphereElem = BodySetup->AggGeom.SphereElems[I];
+				PhysicsBody->SetSphere(BodyIndex, I, FVector3f(SphereElem.Center), SphereElem.Radius);
+
+				const FString ElemName = SphereElem.GetName().ToString();
+				PhysicsBody->SetSphereName(BodyIndex, I, StringCast<ANSICHAR>(*ElemName).Get());
+				PhysicsBody->SetSphereFlags(BodyIndex, I, GetKBodyElemFlags(SphereElem));
+			}
+
+			const int32 NumBoxes = BodySetup->AggGeom.BoxElems.Num();
+			PhysicsBody->SetBoxCount(BodyIndex, NumBoxes);
+
+			for (int32 I = 0; I < NumBoxes; ++I)
+			{
+				const FKBoxElem& BoxElem = BodySetup->AggGeom.BoxElems[I];
+				PhysicsBody->SetBox(BodyIndex, I,
+					FVector3f(BoxElem.Center),
+					FQuat4f(BoxElem.Rotation.Quaternion()),
+					FVector3f(BoxElem.X, BoxElem.Y, BoxElem.Z));
+
+				const FString KElemName = BoxElem.GetName().ToString();
+				PhysicsBody->SetBoxName(BodyIndex, I, StringCast<ANSICHAR>(*KElemName).Get());
+				PhysicsBody->SetBoxFlags(BodyIndex, I, GetKBodyElemFlags(BoxElem));
+			}
+
+			const int32 NumConvex = BodySetup->AggGeom.ConvexElems.Num();
+			PhysicsBody->SetConvexCount(BodyIndex, NumConvex);
+			for (int32 I = 0; I < NumConvex; ++I)
+			{
+				const FKConvexElem& ConvexElem = BodySetup->AggGeom.ConvexElems[I];
+
+				// Convert to FVector3f
+				TArray<FVector3f> VertexData;
+				VertexData.SetNumUninitialized(ConvexElem.VertexData.Num());
+				for (int32 Elem = VertexData.Num() - 1; Elem >= 0; --Elem)
+				{
+					VertexData[Elem] = FVector3f(ConvexElem.VertexData[Elem]);
+				}
+
+				PhysicsBody->SetConvexMesh(BodyIndex, I,
+					TArrayView<const FVector3f>(VertexData.GetData(), ConvexElem.VertexData.Num()),
+					TArrayView<const int32>(ConvexElem.IndexData.GetData(), ConvexElem.IndexData.Num()));
+
+				PhysicsBody->SetConvexTransform(BodyIndex, I, FTransform3f(ConvexElem.GetTransform()));
+
+				const FString KElemName = ConvexElem.GetName().ToString();
+				PhysicsBody->SetConvexName(BodyIndex, I, StringCast<ANSICHAR>(*KElemName).Get());
+				PhysicsBody->SetConvexFlags(BodyIndex, I, GetKBodyElemFlags(ConvexElem));
+			}
+
+			const int32 NumSphyls = BodySetup->AggGeom.SphylElems.Num();
+			PhysicsBody->SetSphylCount(BodyIndex, NumSphyls);
+
+			for (int32 I = 0; I < NumSphyls; ++I)
+			{
+				const FKSphylElem& SphylElem = BodySetup->AggGeom.SphylElems[I];
+				PhysicsBody->SetSphyl(BodyIndex, I,
+					FVector3f(SphylElem.Center),
+					FQuat4f(SphylElem.Rotation.Quaternion()),
+					SphylElem.Radius, SphylElem.Length);
+
+				const FString KElemName = SphylElem.GetName().ToString();
+				PhysicsBody->SetSphylName(BodyIndex, I, StringCast<ANSICHAR>(*KElemName).Get());
+				PhysicsBody->SetSphylFlags(BodyIndex, I, GetKBodyElemFlags(SphylElem));
+			}
+
+			const int32 NumTaperedCapsules = BodySetup->AggGeom.TaperedCapsuleElems.Num();
+			PhysicsBody->SetTaperedCapsuleCount(BodyIndex, NumTaperedCapsules);
+
+			for (int32 I = 0; I < NumTaperedCapsules; ++I)
+			{
+				const FKTaperedCapsuleElem& TaperedCapsuleElem = BodySetup->AggGeom.TaperedCapsuleElems[I];
+				PhysicsBody->SetTaperedCapsule(BodyIndex, I,
+					FVector3f(TaperedCapsuleElem.Center),
+					FQuat4f(TaperedCapsuleElem.Rotation.Quaternion()),
+					TaperedCapsuleElem.Radius0, TaperedCapsuleElem.Radius1, TaperedCapsuleElem.Length);
+
+				const FString KElemName = TaperedCapsuleElem.GetName().ToString();
+				PhysicsBody->SetTaperedCapsuleName(BodyIndex, I, StringCast<ANSICHAR>(*KElemName).Get());
+				PhysicsBody->SetTaperedCapsuleFlags(BodyIndex, I, GetKBodyElemFlags(TaperedCapsuleElem));
+			}
+		}
+
+		return PhysicsBody;
+	}
 }
diff --git a/Engine/Plugins/Mutable/Source/CustomizableObjectEditor/Private/MuCOE/GenerateMutableSource/GenerateMutableSourceMesh.cpp b/Engine/Plugins/Mutable/Source/CustomizableObjectEditor/Private/MuCOE/GenerateMutableSource/GenerateMutableSourceMesh.cpp
--- a/Engine/Plugins/Mutable/Source/CustomizableObjectEditor/Private/MuCOE/GenerateMutableSource/GenerateMutableSourceMesh.cpp
+++ b/Engine/Plugins/Mutable/Source/CustomizableObjectEditor/Private/MuCOE/GenerateMutableSource/GenerateMutableSourceMesh.cpp
@@ -729,153 +729,6 @@ TArray<TTuple<UPhysicsAsset*, int32>> GetPhysicsAssetsFromAnimInstance(FMutableG
 	return Result;
 }
 
-// Could be removed but may produce patch size
-TArray<uint8> MakePhysicsAssetBodySetupRelevancyMap(const FMutableCompilationContext& Context, UPhysicsAsset* Asset, const UE::Mutable::Private::TManagedPtr<UE::Mutable::Private::FMesh>& Mesh)
-{
-	const int32 BodySetupsNum = Asset->SkeletalBodySetups.Num();
-
-	TArray<uint8> RelevancyMap;
-	RelevancyMap.Init(0, BodySetupsNum);
-
-	UE::Mutable::Private::TManagedPtr<const UE::Mutable::Private::FSkeleton> Skeleton = Mesh->GetSkeleton();
-	if (!Skeleton)
-	{
-		return RelevancyMap;
-	}
-
-	for (int32 BodyIndex = 0; BodyIndex < BodySetupsNum; ++BodyIndex)
-	{
-		RelevancyMap[BodyIndex] = Skeleton->BoneNames.Find(Asset->SkeletalBodySetups[BodyIndex]->BoneName) != INDEX_NONE;
-	}
-
-	return RelevancyMap;
-}
-
-UE::Mutable::Private::TManagedPtr<UE::Mutable::Private::FPhysicsBody> MakePhysicsBodyFromAsset(UPhysicsAsset* Asset, const UE::Mutable::Private::FSkeleton& MutableSkeleton, const TArray<uint8>& BodySetupRelevancyMap)
-{
-	check(Asset);
-	check(Asset->SkeletalBodySetups.Num() == BodySetupRelevancyMap.Num());
-
-	// Find BodySetups with relevant bones.
-	TArray<TObjectPtr<USkeletalBodySetup>>& SkeletalBodySetups = Asset->SkeletalBodySetups;
-	
-	const int32 NumRelevantSetups = Algo::CountIf(BodySetupRelevancyMap, [](uint8 V) { return V; });
-
-	UE::Mutable::Private::TManagedPtr<UE::Mutable::Private::FPhysicsBody> PhysicsBody = UE::Mutable::Private::MakeManaged<UE::Mutable::Private::FPhysicsBody>();
-	
-	PhysicsBody->SetBodyCount(NumRelevantSetups);
-
-	auto GetKBodyElemFlags = [](const FKShapeElem& KElem) -> uint32
-	{
-		uint8 ElemCollisionEnabled = static_cast<uint8>( KElem.GetCollisionEnabled() );
-		
-		uint32 Flags = static_cast<uint32>( ElemCollisionEnabled );
-		Flags = Flags | (static_cast<uint32>(KElem.GetContributeToMass()) << 8);
-
-		return Flags; 
-	};
-
-	for (int32 B = 0, SourceBodyIndex = 0; B < NumRelevantSetups; ++B)
-	{
-		if (!BodySetupRelevancyMap[SourceBodyIndex])
-		{
-			continue;
-		}
-
-		TObjectPtr<USkeletalBodySetup>& BodySetup = SkeletalBodySetups[SourceBodyIndex++];
-
-		PhysicsBody->SetBody(B, BodySetup->BoneName, INDEX_NONE);
-		
-		const int32 NumSpheres = BodySetup->AggGeom.SphereElems.Num();
-		PhysicsBody->SetSphereCount(B, NumSpheres);
-
-		for (int32 I = 0; I < NumSpheres; ++I)
-		{
-			const FKSphereElem& SphereElem = BodySetup->AggGeom.SphereElems[I];
-			PhysicsBody->SetSphere(B, I, FVector3f(SphereElem.Center), SphereElem.Radius);
-
-			const FString ElemName = SphereElem.GetName().ToString();
-			PhysicsBody->SetSphereName(B, I, TCHAR_TO_ANSI(*ElemName));	
-			PhysicsBody->SetSphereFlags(B, I, GetKBodyElemFlags(SphereElem));
-		}
-
-		const int32 NumBoxes = BodySetup->AggGeom.BoxElems.Num();
-		PhysicsBody->SetBoxCount(B, NumBoxes);
-
-		for (int32 I = 0; I < NumBoxes; ++I)
-		{
-			const FKBoxElem& BoxElem = BodySetup->AggGeom.BoxElems[I];
-			PhysicsBody->SetBox(B, I, 
-					FVector3f(BoxElem.Center), 
-					FQuat4f(BoxElem.Rotation.Quaternion()), 
-					FVector3f(BoxElem.X, BoxElem.Y, BoxElem.Z));
-
-			const FString KElemName = BoxElem.GetName().ToString();
-			PhysicsBody->SetBoxName(B, I, TCHAR_TO_ANSI(*KElemName));
-			PhysicsBody->SetBoxFlags(B, I, GetKBodyElemFlags(BoxElem));
-		}
-
-		const int32 NumConvex = BodySetup->AggGeom.ConvexElems.Num();
-		PhysicsBody->SetConvexCount(B, NumConvex);
-		for (int32 I = 0; I < NumConvex; ++I)
-		{
-			const FKConvexElem& ConvexElem = BodySetup->AggGeom.ConvexElems[I];
-
-			// Convert to FVector3f
-			TArray<FVector3f> VertexData;
-			VertexData.SetNumUninitialized(ConvexElem.VertexData.Num());
-			for (int32 Elem = VertexData.Num() - 1; Elem >= 0; --Elem)
-			{
-				VertexData[Elem] = FVector3f(ConvexElem.VertexData[Elem]);
-			}
-			
-			PhysicsBody->SetConvexMesh(B, I,
-					TArrayView<const FVector3f>(VertexData.GetData(), ConvexElem.VertexData.Num()),
-					TArrayView<const int32>(ConvexElem.IndexData.GetData(), ConvexElem.IndexData.Num()));
-
-			PhysicsBody->SetConvexTransform(B, I, FTransform3f(ConvexElem.GetTransform()));
-			
-			const FString KElemName = ConvexElem.GetName().ToString();
-			PhysicsBody->SetConvexName(B, I, TCHAR_TO_ANSI(*KElemName));
-			PhysicsBody->SetConvexFlags(B, I, GetKBodyElemFlags(ConvexElem));
-		}
-
-		const int32 NumSphyls = BodySetup->AggGeom.SphylElems.Num();
-		PhysicsBody->SetSphylCount( B, NumSphyls );
-
-		for (int32 I = 0; I < NumSphyls; ++I)
-		{
-			const FKSphylElem& SphylElem = BodySetup->AggGeom.SphylElems[I];
-			PhysicsBody->SetSphyl( B, I, 
-					FVector3f(SphylElem.Center), 
-					FQuat4f(SphylElem.Rotation.Quaternion()), 
-					SphylElem.Radius, SphylElem.Length );
-
-			const FString KElemName = SphylElem.GetName().ToString();
-			PhysicsBody->SetSphylName(B, I, TCHAR_TO_ANSI(*KElemName));
-			PhysicsBody->SetSphylFlags(B, I, GetKBodyElemFlags(SphylElem));
-		}
-
-		const int32 NumTaperedCapsules = BodySetup->AggGeom.TaperedCapsuleElems.Num();
-		PhysicsBody->SetTaperedCapsuleCount( B, NumTaperedCapsules );
-
-		for (int32 I = 0; I < NumTaperedCapsules; ++I)
-		{
-			const FKTaperedCapsuleElem& TaperedCapsuleElem = BodySetup->AggGeom.TaperedCapsuleElems[I];
-			PhysicsBody->SetTaperedCapsule(B, I, 
-					FVector3f(TaperedCapsuleElem.Center), 
-					FQuat4f(TaperedCapsuleElem.Rotation.Quaternion()), 
-					TaperedCapsuleElem.Radius0, TaperedCapsuleElem.Radius1, TaperedCapsuleElem.Length);
-			
-			const FString KElemName = TaperedCapsuleElem.GetName().ToString();
-			PhysicsBody->SetTaperedCapsuleName(B, I, TCHAR_TO_ANSI(*KElemName));
-			PhysicsBody->SetTaperedCapsuleFlags(B, I, GetKBodyElemFlags(TaperedCapsuleElem));
-		}
-	}
-
-	return PhysicsBody;
-}
-
 
 namespace UE::Mutable::Private
 {
@@ -1260,7 +1113,7 @@ UE::Tasks::FTask ConvertSkeletalMeshToMutable(TSharedRef<FMeshConversionContext>
 	bool bBoneMapModified = false;
 	TArray<FBoneIndexType> BoneMap;
 
-	// Check if the Skeleton is valid and build the UE::Mutable::Private::Skeleton
+	// Check if the Skeleton is valid and build the UE::Mutable::Private::FSkeleton
 	if (!bIgnoreSkinning)
 	{
 		const USkeleton* Skeleton = SkeletalMesh->GetSkeleton();
@@ -1526,11 +1379,11 @@ UE::Tasks::FTask ConvertSkeletalMeshToMutable(TSharedRef<FMeshConversionContext>
 		// cache the result here to avoid calling it repeatedly.
 		const TArray<UMorphTarget*>& SkeletalMeshMorphTargets = SkeletalMesh->GetMorphTargets();
 
-		// Find realtime MorphTargets to be used.
+		// Find real-time MorphTargets to be used.
 		TArray<const UMorphTarget*> UsedMorphTargets;
 		UsedMorphTargets.Reserve(SkeletalMeshMorphTargets.Num());
 
-		// Add SkeletalMesh node used defined realtime morph targets to a temporal array where
+		// Add SkeletalMesh node used defined real-time morph targets to a temporal array where
 		// the actual to be used real-time morphs names will be placed.        
 		TArray<FName> UsedMorphTargetsNames = Invoke([&]()
 			{
@@ -1973,236 +1826,29 @@ UE::Tasks::FTask ConvertSkeletalMeshToMutable(TSharedRef<FMeshConversionContext>
 
 	if (!bIgnorePhysics && SkeletalMesh->GetPhysicsAsset() && MutableMesh->GetSkeleton() && Context.Options.bPhysicsAssetMergeEnabled)
 	{
-		// Find BodySetups with relevant bones.
-		TArray<TObjectPtr<USkeletalBodySetup>>& SkeletalBodySetups = SkeletalMesh->GetPhysicsAsset()->SkeletalBodySetups;
-
-		TArray<TObjectPtr<USkeletalBodySetup>> RelevantBodySetups;
-		RelevantBodySetups.Reserve(SkeletalBodySetups.Num());
-
-		TArray<uint8> DiscardedBodySetups;
-		DiscardedBodySetups.Init(1, SkeletalBodySetups.Num());
+		UPhysicsAsset& PhysicsAsset = *SkeletalMesh->GetPhysicsAsset();
 
-		const UE::Mutable::Private::FSkeleton& MutableSkeleton = *MutableMesh->GetSkeleton();
+		UE::Mutable::Private::TManagedPtr<UE::Mutable::Private::FPhysicsBody> PhysicsBody = UnrealConversionUtils::CreatePhysicsBodyForMesh(PhysicsAsset, *MutableMesh);
 
-		for (int32 BodySetupIndex = 0; BodySetupIndex < SkeletalBodySetups.Num(); ++BodySetupIndex)
+		if (PhysicsBody)
 		{
-			TObjectPtr<USkeletalBodySetup>& BodySetup = SkeletalBodySetups[BodySetupIndex];
-			if (!BodySetup)
-			{
-				continue;
-			}
-
-			const int32 BonePoseIndex = MutableMesh->FindBonePoseByBoneIndex(MutableSkeleton.BoneNames.Find(BodySetup->BoneName));
+			MutableMesh->SetPhysicsBody(PhysicsBody);
+			MutableMesh->PhysicsAssets.Emplace(Context.PassthroughObjectFactory.Add(PhysicsAsset));
 
-			if (BonePoseIndex == INDEX_NONE)
+			// Set Physics usage flag
+			const int32 NumBodySetups = PhysicsBody->GetBodyCount();
+			for (int32 BodyIndex = 0; BodyIndex < NumBodySetups; ++BodyIndex)
 			{
-				continue;
-			}
 
-			RelevantBodySetups.Add(BodySetup);
-			DiscardedBodySetups[BodySetupIndex] = 0;
-			EnumAddFlags(MutableMesh->BonePoses[BonePoseIndex].BoneUsageFlags, UE::Mutable::Private::EBoneUsageFlags::Physics);
-		}
-
-		const int32 NumDiscardedSetups = Algo::CountIf(DiscardedBodySetups, [](const uint8& V) { return V; });
-
-		constexpr bool bOptOutOfIncompleteBodyWarnings = true;
-		if (NumDiscardedSetups > 0 && !bOptOutOfIncompleteBodyWarnings)
-		{
-			FString PhysicsSetupsRemovedMsg =
-				FString::Printf(TEXT("PhysicsBodySetups in %s attached to bones"), *(SkeletalMesh->GetPhysicsAsset()->GetName()));
-
-			constexpr int32 MaxNumDiscardedShown = 3;
-
-			int32 NumDiscardedShown = 0;
-			for (int32 I = 0; I < SkeletalBodySetups.Num() && NumDiscardedShown < MaxNumDiscardedShown; ++I)
-			{
-				if (DiscardedBodySetups[I] && SkeletalBodySetups[I])
-				{
-					PhysicsSetupsRemovedMsg += (NumDiscardedShown <= 0 ? " " : ", ") + SkeletalBodySetups[I]->BoneName.ToString();
-					++NumDiscardedShown;
-				}
+				const int32 PoseIndex = MutableMesh->Skeleton->BoneNames.Find(PhysicsBody->BodiesBoneNames[BodyIndex]); // Bone Poses are 1 : 1 with the skeleton bones
+				check(PoseIndex != INDEX_NONE);
+				EnumAddFlags(MutableMesh->BonePoses[PoseIndex].BoneUsageFlags, UE::Mutable::Private::EBoneUsageFlags::Physics);
 			}
-
-			if (NumDiscardedShown < NumDiscardedSetups)
-			{
-				PhysicsSetupsRemovedMsg += FString::Printf(TEXT("... and %d more "), NumDiscardedSetups - MaxNumDiscardedShown);
-			}
-
-			PhysicsSetupsRemovedMsg += FString::Printf(TEXT("have been discarded because they are not present in the SkeletalMesh [%s] Skeleton."),
-				*SkeletalMesh->GetName());
-
-			Context.Log(FText::FromString(PhysicsSetupsRemovedMsg), Source.MessageContext, EMessageSeverity::Warning);
 		}
-
-		UE::Mutable::Private::TManagedPtr<UE::Mutable::Private::FPhysicsBody> PhysicsBody = 
-				UE::Mutable::Private::MakeManaged<UE::Mutable::Private::FPhysicsBody>();
-
-		const int32 NumBodySetups = RelevantBodySetups.Num();
-		PhysicsBody->SetBodyCount(NumBodySetups);
-
-		auto GetKBodyElemFlags = [](const FKShapeElem& KElem) -> uint32
-			{
-				uint8 ElemCollisionEnabled = static_cast<uint8>(KElem.GetCollisionEnabled());
-
-				uint32 Flags = static_cast<uint32>(ElemCollisionEnabled);
-				Flags = Flags | (static_cast<uint32>(KElem.GetContributeToMass()) << 8);
-
-				return Flags;
-			};
-
-		for (int32 B = 0; B < NumBodySetups; ++B)
-		{
-			TObjectPtr<USkeletalBodySetup>& BodySetup = RelevantBodySetups[B];
-
-			PhysicsBody->SetBody(B, BodySetup->BoneName);
-
-			const int32 NumSpheres = BodySetup->AggGeom.SphereElems.Num();
-			PhysicsBody->SetSphereCount(B, NumSpheres);
-
-			for (int32 I = 0; I < NumSpheres; ++I)
-			{
-				const FKSphereElem& SphereElem = BodySetup->AggGeom.SphereElems[I];
-				PhysicsBody->SetSphere(B, I, FVector3f(SphereElem.Center), SphereElem.Radius);
-
-				const FString ElemName = SphereElem.GetName().ToString();
-				PhysicsBody->SetSphereName(B, I, StringCast<ANSICHAR>(*ElemName).Get());
-				PhysicsBody->SetSphereFlags(B, I, GetKBodyElemFlags(SphereElem));
-			}
-
-			const int32 NumBoxes = BodySetup->AggGeom.BoxElems.Num();
-			PhysicsBody->SetBoxCount(B, NumBoxes);
-
-			for (int32 I = 0; I < NumBoxes; ++I)
-			{
-				const FKBoxElem& BoxElem = BodySetup->AggGeom.BoxElems[I];
-				PhysicsBody->SetBox(B, I,
-					FVector3f(BoxElem.Center),
-					FQuat4f(BoxElem.Rotation.Quaternion()),
-					FVector3f(BoxElem.X, BoxElem.Y, BoxElem.Z));
-
-				const FString KElemName = BoxElem.GetName().ToString();
-				PhysicsBody->SetBoxName(B, I, StringCast<ANSICHAR>(*KElemName).Get());
-				PhysicsBody->SetBoxFlags(B, I, GetKBodyElemFlags(BoxElem));
-			}
-
-			const int32 NumConvex = BodySetup->AggGeom.ConvexElems.Num();
-			PhysicsBody->SetConvexCount(B, NumConvex);
-			for (int32 I = 0; I < NumConvex; ++I)
-			{
-				const FKConvexElem& ConvexElem = BodySetup->AggGeom.ConvexElems[I];
-
-				// Convert to FVector3f
-				TArray<FVector3f> VertexData;
-				VertexData.SetNumUninitialized(ConvexElem.VertexData.Num());
-				for (int32 Elem = VertexData.Num() - 1; Elem >= 0; --Elem)
-				{
-					VertexData[Elem] = FVector3f(ConvexElem.VertexData[Elem]);
-				}
-
-				PhysicsBody->SetConvexMesh(B, I,
-					TArrayView<const FVector3f>(VertexData.GetData(), ConvexElem.VertexData.Num()),
-					TArrayView<const int32>(ConvexElem.IndexData.GetData(), ConvexElem.IndexData.Num()));
-
-				PhysicsBody->SetConvexTransform(B, I, FTransform3f(ConvexElem.GetTransform()));
-
-				const FString KElemName = ConvexElem.GetName().ToString();
-				PhysicsBody->SetConvexName(B, I, StringCast<ANSICHAR>(*KElemName).Get());
-				PhysicsBody->SetConvexFlags(B, I, GetKBodyElemFlags(ConvexElem));
-			}
-
-			const int32 NumSphyls = BodySetup->AggGeom.SphylElems.Num();
-			PhysicsBody->SetSphylCount(B, NumSphyls);
-
-			for (int32 I = 0; I < NumSphyls; ++I)
-			{
-				const FKSphylElem& SphylElem = BodySetup->AggGeom.SphylElems[I];
-				PhysicsBody->SetSphyl(B, I,
-					FVector3f(SphylElem.Center),
-					FQuat4f(SphylElem.Rotation.Quaternion()),
-					SphylElem.Radius, SphylElem.Length);
-
-				const FString KElemName = SphylElem.GetName().ToString();
-				PhysicsBody->SetSphylName(B, I, StringCast<ANSICHAR>(*KElemName).Get());
-				PhysicsBody->SetSphylFlags(B, I, GetKBodyElemFlags(SphylElem));
-			}
-
-			const int32 NumTaperedCapsules = BodySetup->AggGeom.TaperedCapsuleElems.Num();
-			PhysicsBody->SetTaperedCapsuleCount(B, NumTaperedCapsules);
-
-			for (int32 I = 0; I < NumTaperedCapsules; ++I)
-			{
-				const FKTaperedCapsuleElem& TaperedCapsuleElem = BodySetup->AggGeom.TaperedCapsuleElems[I];
-				PhysicsBody->SetTaperedCapsule(B, I,
-					FVector3f(TaperedCapsuleElem.Center),
-					FQuat4f(TaperedCapsuleElem.Rotation.Quaternion()),
-					TaperedCapsuleElem.Radius0, TaperedCapsuleElem.Radius1, TaperedCapsuleElem.Length);
-
-				const FString KElemName = TaperedCapsuleElem.GetName().ToString();
-				PhysicsBody->SetTaperedCapsuleName(B, I, StringCast<ANSICHAR>(*KElemName).Get());
-				PhysicsBody->SetTaperedCapsuleFlags(B, I, GetKBodyElemFlags(TaperedCapsuleElem));
-			}
-		}
-
-		MutableMesh->SetPhysicsBody(PhysicsBody);
-
-		// Add tags
-		if (SkeletalMesh->GetPhysicsAsset() && PhysicsBody && PhysicsBody->GetBodyCount())
-		{
-			TSoftObjectPtr<UPhysicsAsset> PhysicsAsset = SkeletalMesh->GetPhysicsAsset();
-
-			const int32 AssetIndex = Context.PhysicsAssets.AddUnique(PhysicsAsset);
-			FString PhysicsAssetTag = FString("__PA:") + FString::FromInt(AssetIndex);
-
-			AddTagToMutableMeshUnique(*MutableMesh, PhysicsAssetTag);
-		}
-
 	}
 
-	if (MutableMesh->GetSkeleton()) 
-	{
-		const UE::Mutable::Private::FSkeleton& MutableSkeleton = *MutableMesh->GetSkeleton();
-
-		// Set Bone Parenting usages. This has to be done after all primary usages are set.
-		for (int32 I = MutableMesh->GetBonePoseCount() - 1; I >= 0; --I)
-		{
-			UE::Mutable::Private::FMesh::FBonePose& BonePose = MutableMesh->BonePoses[I];
-
-			constexpr UE::Mutable::Private::EBoneUsageFlags FlagsToPropagate =
-				UE::Mutable::Private::EBoneUsageFlags::Skinning | UE::Mutable::Private::EBoneUsageFlags::Physics | UE::Mutable::Private::EBoneUsageFlags::Deform;
-			if (EnumHasAnyFlags(BonePose.BoneUsageFlags, FlagsToPropagate))
-			{
-				const int32 ParentIndex = MutableSkeleton.BoneParents[BonePose.BoneId.Index];
-
-				if (ParentIndex == INDEX_NONE)
-				{
-					continue;
-				}
-
-				int32 PoseIndex = MutableMesh->FindBonePoseByBoneIndex(ParentIndex);
-				if (PoseIndex == INDEX_NONE)
-				{
-					continue;
-				}
-
-				const UE::Mutable::Private::EBoneUsageFlags ParentPropagationFlags =
-					(EnumHasAnyFlags(BonePose.BoneUsageFlags, UE::Mutable::Private::EBoneUsageFlags::Skinning)
-						? UE::Mutable::Private::EBoneUsageFlags::SkinningParent : UE::Mutable::Private::EBoneUsageFlags::None) |
-					(EnumHasAnyFlags(BonePose.BoneUsageFlags, UE::Mutable::Private::EBoneUsageFlags::Physics)
-						? UE::Mutable::Private::EBoneUsageFlags::PhysicsParent : UE::Mutable::Private::EBoneUsageFlags::None) |
-					(EnumHasAnyFlags(BonePose.BoneUsageFlags, UE::Mutable::Private::EBoneUsageFlags::Deform)
-						? UE::Mutable::Private::EBoneUsageFlags::DeformParent : UE::Mutable::Private::EBoneUsageFlags::None);
-
-				do
-				{
-					EnumAddFlags(MutableMesh->BonePoses[PoseIndex].BoneUsageFlags, ParentPropagationFlags);
-
-					PoseIndex = MutableMesh->FindBonePoseByBoneIndex(MutableSkeleton.BoneParents[MutableMesh->BonePoses[PoseIndex].BoneId.Index]);
-
-				} while (PoseIndex != INDEX_NONE);
-			}
-		}
-	}
+	// Set Bone Parenting usages. This has to be done after all primary usages are set.
+	UnrealConversionUtils::PropagateBoneUsageFlagsThroughMeshPose(*MutableMesh);
 
 	const bool bAnimPhysicsManipulationEnabled = Context.Options.bAnimBpPhysicsManipulationEnabled;
 
@@ -2220,16 +1866,16 @@ UE::Tasks::FTask ConvertSkeletalMeshToMutable(TSharedRef<FMeshConversionContext>
 			{
 				Info.AnimInstanceClass = Source.AnimInstance;
 				Info.PropertyIndex = PropertyIndex;
-				Info.SourceAsset = TSoftObjectPtr<UPhysicsAsset>(PropertyAsset);
 			}
 
 			const int32 PhysicsAssetId = Context.AnimBpOverridePhysicsAssetsInfo.AddUnique(Info);
 
-			UE::Mutable::Private::TManagedPtr<UE::Mutable::Private::FPhysicsBody> MutableBody = MakePhysicsBodyFromAsset(PropertyAsset, *MutableMesh->GetSkeleton(),
-				MakePhysicsAssetBodySetupRelevancyMap(Context, PropertyAsset, MutableMesh));
+			UE::Mutable::Private::TManagedPtr<UE::Mutable::Private::FPhysicsBody> MutableBody = UnrealConversionUtils::CreatePhysicsBodyForMesh(*PropertyAsset, *MutableMesh);
+
 			MutableBody->CustomId = PhysicsAssetId;
 
 			MutableMesh->AddAdditionalPhysicsBody(MutableBody);
+			MutableMesh->AdditionalPhysicsAssets.Emplace(Context.PassthroughObjectFactory.Add(*PropertyAsset));
 		}
 	}
 	
@@ -2525,30 +2171,14 @@ UE::Mutable::Private::TManagedPtr<UE::Mutable::Private::FMesh> ConvertStaticMesh
 
 		MutableMesh->GetIndexBuffers().SetBuffer(0, ElementSize, ChannelCount, Semantics, SemanticIndices, Formats, Components, Offsets);
 
-		//if (ElementSize==4)
-		//{
-		//	uint32* pDest = reinterpret_cast<uint32*>( MutableMesh->GetIndexBuffers().GetBufferData(0) );
-		//	const uint32* pSource = reinterpret_cast<const uint32*>( TypedNode->StaticMesh->RenderData->LODResources[LOD].MultiSizeIndexContainer.GetIndexBuffer()->GetPointerTo(IndexStart) );
-
-		//	for ( int i=0; i<IndexCount; ++i )
-		//	{
-		//		*pDest = *pSource - VertexStart;
-		//		++pDest;
-		//		++pSource;
-		//	}
-		//}
-		//else
-		{
-			FIndexArrayView Source = StaticMesh->GetRenderData()->LODResources[LODIndex].IndexBuffer.GetArrayView();
-			//const uint16* pSource = reinterpret_cast<const uint16*>( StaticMesh->RenderData->LODResources[LOD].IndexBuffer.Indices.GetResourceData() );
-			//pSource += IndexStart;
-			uint16* pDest = reinterpret_cast<uint16*>(MutableMesh->GetIndexBuffers().GetBufferData(0));
+		FIndexArrayView Source = StaticMesh->GetRenderData()->LODResources[LODIndex].IndexBuffer.GetArrayView();
 
-			for (int32 i = 0; i < IndexCount; ++i)
-			{
-				*pDest = Source[IndexStart + i] - VertexStart;
-				++pDest;
-			}
+		uint16* pDest = reinterpret_cast<uint16*>(MutableMesh->GetIndexBuffers().GetBufferData(0));
+
+		for (int32 i = 0; i < IndexCount; ++i)
+		{
+			*pDest = Source[IndexStart + i] - VertexStart;
+			++pDest;
 		}
 	}
 
diff --git a/Engine/Plugins/Mutable/Source/CustomizableObjectEditor/Private/MuCOE/SMutableCodeViewer.cpp b/Engine/Plugins/Mutable/Source/CustomizableObjectEditor/Private/MuCOE/SMutableCodeViewer.cpp
--- a/Engine/Plugins/Mutable/Source/CustomizableObjectEditor/Private/MuCOE/SMutableCodeViewer.cpp
+++ b/Engine/Plugins/Mutable/Source/CustomizableObjectEditor/Private/MuCOE/SMutableCodeViewer.cpp
@@ -1565,33 +1565,40 @@ void SMutableCodeViewer::GenerateElementRecursive(const int32& InStateIndex, UE:
 
 		using OpEnumFlags = UE::Mutable::Private::OP::MeshAddMetadataArgs::EnumFlags;
 
-		// Set count to 1 if is not a list.
-		int32 TagCount      = int32(!EnumHasAnyFlags(Args.Flags, OpEnumFlags::IsTagList));
-		int32 ResourceCount = int32(!EnumHasAnyFlags(Args.Flags, OpEnumFlags::IsResourceList));
-		int32 SkeletonCount = int32(!EnumHasAnyFlags(Args.Flags, OpEnumFlags::IsSkeletonList));
-
-		if (!TagCount)
+		int32 TagCount = EnumHasAnyFlags(Args.Flags, OpEnumFlags::HasTags) ? 1 : 0;
+		if (EnumHasAnyFlags(Args.Flags, OpEnumFlags::IsTagList) &&
+			InProgram.ConstantUInt32Lists.IsValidIndex(Args.Tags.ListAddress))
 		{
-			TagCount = InProgram.ConstantUInt32Lists.IsValidIndex(Args.Tags.ListAddress) 
-					? InProgram.ConstantUInt32Lists[Args.Tags.ListAddress].Num() 
-					: 0;
+			TagCount = InProgram.ConstantUInt32Lists[Args.Tags.ListAddress].Num();
 		}
 
-		if (!ResourceCount)
+		int32 ResourceCount = EnumHasAnyFlags(Args.Flags, OpEnumFlags::HasResources) ? 1 : 0;;
+		if (EnumHasAnyFlags(Args.Flags, OpEnumFlags::IsResourceList) &&
+			ensure(InProgram.ConstantUInt64Lists.IsValidIndex(Args.ResourceIds.ListAddress)))
 		{
 			ResourceCount = InProgram.ConstantUInt64Lists.IsValidIndex(Args.ResourceIds.ListAddress) 
 					? InProgram.ConstantUInt64Lists[Args.ResourceIds.ListAddress].Num() 
 					: 0;
 		}
 
-		if (!SkeletonCount)
+		int32 SkeletonCount = EnumHasAnyFlags(Args.Flags, OpEnumFlags::HasSkeleton) ? 1 : 0;
+		if (EnumHasAnyFlags(Args.Flags, OpEnumFlags::IsSkeletonList) &&
+			ensure(InProgram.ConstantUInt32Lists.IsValidIndex(Args.SkeletonIds.ListAddress)))
 		{
-			SkeletonCount = InProgram.ConstantUInt32Lists.IsValidIndex(Args.SkeletonIds.ListAddress) 
-					? InProgram.ConstantUInt32Lists[Args.SkeletonIds.ListAddress].Num() 
-					: 0;
+			SkeletonCount = InProgram.ConstantUInt32Lists[Args.SkeletonIds.ListAddress].Num();
+		}
+
+		const int32 PhysicsAssetCount = Args.PhysicsAssetId != UE::Mutable::Private::PASSTHROUGH_ID_INVALID ? 1 : 0;
+
+		int32 AdditionalPhysicsAssetCount = EnumHasAnyFlags(Args.Flags, OpEnumFlags::HasAdditionalPhysicsAsset) ? 1 : 0;
+		if (EnumHasAnyFlags(Args.Flags, OpEnumFlags::IsAdditionalPhysicsAssetList) &&
+			ensure(InProgram.ConstantUInt32Lists.IsValidIndex(Args.AdditionalPhysicsAssetsIds.ListAddress)))
+		{
+			AdditionalPhysicsAssetCount = InProgram.ConstantUInt32Lists[Args.AdditionalPhysicsAssetsIds.ListAddress].Num();
 		}
 
-		FString Caption = FString::Printf(TEXT("add %d tags %d resources and %d skeletons to "), TagCount, ResourceCount, SkeletonCount);
+		FString Caption = FString::Printf(TEXT("add %d Tags, %d Resources, %d Skeleton, %d PhysicsAsset, and %d Additional PhysicsAssets to "),
+			TagCount, ResourceCount, SkeletonCount, PhysicsAssetCount, AdditionalPhysicsAssetCount);
 		AddOpFunc(Args.Source, Caption);
 
 		break;
diff --git a/Engine/Plugins/Mutable/Source/MutableRuntime/Internal/MuR/Mesh.h b/Engine/Plugins/Mutable/Source/MutableRuntime/Internal/MuR/Mesh.h
--- a/Engine/Plugins/Mutable/Source/MutableRuntime/Internal/MuR/Mesh.h
+++ b/Engine/Plugins/Mutable/Source/MutableRuntime/Internal/MuR/Mesh.h
@@ -29,6 +29,7 @@
 
 class FString;
 class USkeletalMesh;
+class UPhysicsAsset;
 
 namespace UE::Mutable::Private
 {
@@ -484,6 +485,12 @@ namespace UE::Mutable::Private
 		/** Additional physics bodies referenced by the mesh that don't merge. */
 		TArray<TManagedPtr<const FPhysicsBody>> AdditionalPhysicsBodies;
 
+		/** Physics assets referenced by the mesh to be merged. */
+		TArray<TPassthroughObjectPtr<UPhysicsAsset>> PhysicsAssets;
+
+		/** Additional physics assets referenced by the mesh that don't merge. */
+		TArray<TPassthroughObjectPtr<UPhysicsAsset>> AdditionalPhysicsAssets;
+
 		/** 
 		 * Texture Layout blocks attached to this mesh. They are const because they could be shared with
 		 * other meshes, so they need to be cloned and replaced if a modification is needed.
diff --git a/Engine/Plugins/Mutable/Source/MutableRuntime/Internal/MuR/OpMeshMerge.h b/Engine/Plugins/Mutable/Source/MutableRuntime/Internal/MuR/OpMeshMerge.h
--- a/Engine/Plugins/Mutable/Source/MutableRuntime/Internal/MuR/OpMeshMerge.h
+++ b/Engine/Plugins/Mutable/Source/MutableRuntime/Internal/MuR/OpMeshMerge.h
@@ -911,15 +911,21 @@ namespace UE::Mutable::Private
 		{
 			MUTABLE_CPUPROFILER_SCOPE(PhysicsBodies);
 
-			// Appends InPhysicsBody to OutPhysicsBody removing Bodies that are equal, have same bone and customId and its properies are identical.	
+			// Append PhysicsAssets
+			Result->PhysicsAssets = pFirst->PhysicsAssets;
+
+			for (const TPassthroughObjectPtr<UPhysicsAsset>& PhysicsAsset : pSecond->PhysicsAssets)
+			{
+				Result->PhysicsAssets.AddUnique(PhysicsAsset);
+			}
+
+			// Appends InPhysicsBody to OutPhysicsBody removing Bodies that are equal and their properies are identical.	
 			auto AppendPhysicsBodiesUnique = [](FPhysicsBody& OutPhysicsBody, const FPhysicsBody& InPhysicsBody) -> bool
 			{
 				TArray<FName>& OutBones = OutPhysicsBody.BodiesBoneNames;
-				TArray<int32>& OutCustomIds = OutPhysicsBody.BodiesCustomIds;
 				TArray<FPhysicsBodyAggregate>& OutBodies = OutPhysicsBody.Bodies;
 
 				const TArray<FName>& InBones = InPhysicsBody.BodiesBoneNames;
-				const TArray<int32>& InCustomIds = InPhysicsBody.BodiesCustomIds;
 				const TArray<FPhysicsBodyAggregate>& InBodies = InPhysicsBody.Bodies;
 
 				const int32 InBodyCount = InPhysicsBody.GetBodyCount();
@@ -929,20 +935,11 @@ namespace UE::Mutable::Private
 
 				for (int32 InBodyIndex = 0; InBodyIndex < InBodyCount; ++InBodyIndex)
 				{
-					int32 FoundIndex = INDEX_NONE;
-					for (int32 OutBodyIndex = 0; OutBodyIndex < OutBodyCount; ++OutBodyIndex)
-					{
-						if (InCustomIds[InBodyIndex] == OutCustomIds[OutBodyIndex] && InBones[InBodyIndex] == OutBones[OutBodyIndex])
-						{
-							FoundIndex = OutBodyIndex;
-							break;
-						}
-					}
+					const int32 FoundIndex = OutBones.Find(InBones[InBodyIndex]);
 
 					if (FoundIndex == INDEX_NONE)
 					{
 						OutBones.Add(InBones[InBodyIndex]);
-						OutCustomIds.Add(InCustomIds[InBodyIndex]);
 						OutBodies.Add(InBodies[InBodyIndex]);
 
 						bModified |= true;
@@ -1028,9 +1025,14 @@ namespace UE::Mutable::Private
 
 			Result->AdditionalPhysicsBodies.Reserve(MaxAdditionalPhysicsResultNum);
 			Result->AdditionalPhysicsBodies.Append(pFirst->AdditionalPhysicsBodies);
+
+			Result->AdditionalPhysicsAssets.Reserve(MaxAdditionalPhysicsResultNum);
+			Result->AdditionalPhysicsAssets.Append(pFirst->AdditionalPhysicsAssets);
 			
 			// Not very many additional bodies expected, do a quadratic search to have unique bodies based on external id.
 			const int32 NumSecondAdditionalBodies = pSecond->AdditionalPhysicsBodies.Num();
+			check(pSecond->AdditionalPhysicsAssets.Num() == NumSecondAdditionalBodies);
+
 			for (int32 Index = 0; Index < NumSecondAdditionalBodies; ++Index)
 			{
 				const int32 CustomIdToMerge = pSecond->AdditionalPhysicsBodies[Index]->CustomId;
@@ -1043,6 +1045,7 @@ namespace UE::Mutable::Private
 				if (!Found)
 				{
 					Result->AdditionalPhysicsBodies.Add(pSecond->AdditionalPhysicsBodies[Index]);
+					Result->AdditionalPhysicsAssets.Add(pSecond->AdditionalPhysicsAssets[Index]);
 				}
 			}
 		}
diff --git a/Engine/Plugins/Mutable/Source/MutableRuntime/Internal/MuR/Operations.h b/Engine/Plugins/Mutable/Source/MutableRuntime/Internal/MuR/Operations.h
--- a/Engine/Plugins/Mutable/Source/MutableRuntime/Internal/MuR/Operations.h
+++ b/Engine/Plugins/Mutable/Source/MutableRuntime/Internal/MuR/Operations.h
@@ -1092,22 +1092,27 @@ namespace UE::Mutable::Private
         {
             enum class EnumFlags : uint16
             {
-                None           = 0,
-                IsTagList      = 1 << 0,
-                IsResourceList = 1 << 1,
-                IsSkeletonList = 1 << 2
+                None			= 0,
+                HasTags			= 1 << 0,
+				IsTagList		= 1 << 1,
+                HasResources	= 1 << 2,
+                IsResourceList	= 1 << 3,
+                HasSkeleton		= 1 << 4, // remove
+                IsSkeletonList	= 1 << 5, // remove
+				HasAdditionalPhysicsAsset = 1 << 6,
+                IsAdditionalPhysicsAssetList = 1 << 7
             };
 
             union TagAddressOrUInt32ListAddress
             {
                 ADDRESS TagAddress;
-                ADDRESS ListAddress;
+				LIST_UINT32_ADDRESS ListAddress;
             };
 
             union SkeletonIdOrUInt32ListAddress
             {
                 uint32 SkeletonId;
-                ADDRESS ListAddress;
+				LIST_UINT32_ADDRESS ListAddress;
             };
 
             union ResourceIdOrUInt64ListAddress
@@ -1116,14 +1121,22 @@ namespace UE::Mutable::Private
                 ADDRESS ListAddress;
             };
 
+			union PassthroughIDOrUInt32ListAddress
+			{
+				PASSTHROUGH_ID PassthroughID;
+				LIST_UINT32_ADDRESS ListAddress;
+			};
+
             EnumFlags Flags; 
 			int16 BonePosePriority;
             ADDRESS Source;
             TagAddressOrUInt32ListAddress Tags;
             SkeletonIdOrUInt32ListAddress SkeletonIds;
             ResourceIdOrUInt64ListAddress ResourceIds;
+			PASSTHROUGH_ID PhysicsAssetId; // Main Physics Asset of the mesh
+			PassthroughIDOrUInt32ListAddress AdditionalPhysicsAssetsIds; // Physics Assets used by AnimBPs
         };
-        static_assert(sizeof(MeshAddMetadataArgs) == sizeof(uint64)*3, "MeshAddMetadataArgs has an unexpected size.");
+        static_assert(sizeof(MeshAddMetadataArgs) == sizeof(uint64)*4, "MeshAddMetadataArgs has an unexpected size.");
 
 
         struct InstanceAddArgs
diff --git a/Engine/Plugins/Mutable/Source/MutableRuntime/Internal/MuR/PassthroughObject.h b/Engine/Plugins/Mutable/Source/MutableRuntime/Internal/MuR/PassthroughObject.h
--- a/Engine/Plugins/Mutable/Source/MutableRuntime/Internal/MuR/PassthroughObject.h
+++ b/Engine/Plugins/Mutable/Source/MutableRuntime/Internal/MuR/PassthroughObject.h
@@ -11,6 +11,7 @@ namespace UE::Mutable::Private
 	class FPassthroughObjectLoader;
 	
 	using PASSTHROUGH_ID = uint32;
+	static_assert(sizeof(PASSTHROUGH_ID) == 4); // Passthrough object Ids are stored in the program as uint32
 
 	constexpr PASSTHROUGH_ID PASSTHROUGH_ID_INVALID = TNumericLimits<PASSTHROUGH_ID>::Max();
 	
diff --git a/Engine/Plugins/Mutable/Source/MutableRuntime/Internal/MuR/PhysicsBody.h b/Engine/Plugins/Mutable/Source/MutableRuntime/Internal/MuR/PhysicsBody.h
--- a/Engine/Plugins/Mutable/Source/MutableRuntime/Internal/MuR/PhysicsBody.h
+++ b/Engine/Plugins/Mutable/Source/MutableRuntime/Internal/MuR/PhysicsBody.h
@@ -165,7 +165,7 @@ namespace UE::Mutable::Private
 		UE_API void SetBodyCount(int32 B);
 		UE_API int32 GetBodyCount() const;
 	
-		UE_API void SetBody(int32 BodyIndex, FName BoneName, int32 CustomId = INDEX_NONE);
+		UE_API void SetBody(int32 BodyIndex, FName BoneName);
 
 		UE_API FName GetBodyBoneName(int32 Index) const;
 
@@ -271,7 +271,6 @@ namespace UE::Mutable::Private
 		TArray<FName> BodiesBoneNames; // Runtime
 		TArray<uint32> BoneIds; // Serialized
 		TArray<FPhysicsBodyAggregate> Bodies;
-		TArray<int32> BodiesCustomIds;
 
 		bool bBodiesModified = false;
 
@@ -284,7 +283,6 @@ namespace UE::Mutable::Private
         	return CustomId        == Other.CustomId        &&
 				   BoneIds		   == Other.BoneIds			&&
         		   Bodies          == Other.Bodies          &&
-        	       BodiesCustomIds == Other.BodiesCustomIds &&
 				   bBodiesModified == Other.bBodiesModified;
         }
 	};
diff --git a/Engine/Plugins/Mutable/Source/MutableRuntime/Private/MuR/CodeRunner.cpp b/Engine/Plugins/Mutable/Source/MutableRuntime/Private/MuR/CodeRunner.cpp
--- a/Engine/Plugins/Mutable/Source/MutableRuntime/Private/MuR/CodeRunner.cpp
+++ b/Engine/Plugins/Mutable/Source/MutableRuntime/Private/MuR/CodeRunner.cpp
@@ -2896,77 +2896,61 @@ namespace UE::Mutable::Private
 					Result = CloneOrTakeOver(MoveTemp(SourceMesh));
 
 					using OpEnumFlags = OP::MeshAddMetadataArgs::EnumFlags;
-					bool bIsTagList      = EnumHasAnyFlags(OpArgs.Flags, OpEnumFlags::IsTagList);
-					bool bIsResourceList = EnumHasAnyFlags(OpArgs.Flags, OpEnumFlags::IsResourceList);
-					bool bIsSkeletonList = EnumHasAnyFlags(OpArgs.Flags, OpEnumFlags::IsSkeletonList);
 
-					if (bIsTagList)
+					// Append Tags
+					if (EnumHasAnyFlags(OpArgs.Flags, OpEnumFlags::IsTagList) &&
+						Program.ConstantUInt32Lists.IsValidIndex(OpArgs.Tags.ListAddress))
 					{
-						if (Program.ConstantUInt32Lists.IsValidIndex(OpArgs.Tags.ListAddress))
-						{
-							const TArray<uint32>& ConstantTagList = Program.ConstantUInt32Lists[OpArgs.Tags.ListAddress];
-							Result->Tags.Reserve(Result->Tags.Num() + ConstantTagList.Num());
+						const TArray<uint32>& ConstantTagList = Program.ConstantUInt32Lists[OpArgs.Tags.ListAddress];
+						Result->Tags.Reserve(Result->Tags.Num() + ConstantTagList.Num());
 
-							for (uint32 TagIndex : ConstantTagList)
+						for (uint32 TagIndex : ConstantTagList)
+						{
+							if (ensure(Program.ConstantStrings.IsValidIndex(TagIndex)))
 							{
-								if (Program.ConstantStrings.IsValidIndex(TagIndex))
-								{
-									Result->Tags.Add(Program.ConstantStrings[TagIndex]);
-								}
-								else
-								{
-									check(false);
-								}
+								Result->Tags.Add(Program.ConstantStrings[TagIndex]);
 							}
 						}
-						else
-						{
-							check(false);
-						}
+					}
+					else if(EnumHasAnyFlags(OpArgs.Flags, OpEnumFlags::HasTags) &&
+						Program.ConstantStrings.IsValidIndex(OpArgs.Tags.TagAddress))
+					{
+						Result->Tags.Add(Program.ConstantStrings[OpArgs.Tags.TagAddress]);
 					}
 					else
 					{
-						if (Program.ConstantStrings.IsValidIndex(OpArgs.Tags.TagAddress))
-						{
-							Result->Tags.Add(Program.ConstantStrings[OpArgs.Tags.TagAddress]);
-						}
-						else
-						{
-							check(false);
-						}
+						check(!EnumHasAnyFlags(OpArgs.Flags, OpEnumFlags::HasTags | OpEnumFlags::IsTagList));
 					}
-
-					if (bIsResourceList)
+					
+					// Append Streamed Resources
+					if (EnumHasAnyFlags(OpArgs.Flags, OpEnumFlags::IsResourceList) &&
+						Program.ConstantUInt64Lists.IsValidIndex(OpArgs.ResourceIds.ListAddress))
 					{
-						if (Program.ConstantUInt64Lists.IsValidIndex(OpArgs.ResourceIds.ListAddress))
-						{
-							Result->StreamedResources.Append(Program.ConstantUInt64Lists[OpArgs.ResourceIds.ListAddress]);
-						}
-						else
-						{
-							check(false);
-						}
+						Result->StreamedResources.Append(Program.ConstantUInt64Lists[OpArgs.ResourceIds.ListAddress]);
 					}
-					else
+					else if (EnumHasAnyFlags(OpArgs.Flags, OpEnumFlags::HasResources))
 					{
 						Result->StreamedResources.Add(OpArgs.ResourceIds.ResourceId);
 					}
-
-					if (bIsSkeletonList)
+					else
 					{
-						if (Program.ConstantUInt32Lists.IsValidIndex(OpArgs.SkeletonIds.ListAddress))
-						{
-							Result->SkeletonIDs.Append(Program.ConstantUInt32Lists[OpArgs.SkeletonIds.ListAddress]);
-						}
-						else
-						{
-							check(false);
-						}
+						check(!EnumHasAnyFlags(OpArgs.Flags, OpEnumFlags::IsResourceList));
 					}
-					else
+					
+					// Append Skeletons 
+					if (EnumHasAnyFlags(OpArgs.Flags, OpEnumFlags::IsSkeletonList) &&
+						Program.ConstantUInt32Lists.IsValidIndex(OpArgs.SkeletonIds.ListAddress))
+					{
+						Result->SkeletonIDs.Append(Program.ConstantUInt32Lists[OpArgs.SkeletonIds.ListAddress]);
+					}
+					else if(EnumHasAnyFlags(OpArgs.Flags, OpEnumFlags::HasSkeleton))
 					{
 						Result->SkeletonIDs.Add(OpArgs.SkeletonIds.SkeletonId);
 					}
+					else
+					{
+						check(!EnumHasAnyFlags(OpArgs.Flags, OpEnumFlags::IsSkeletonList));
+					}
 
 					if (Result->BonePosePriority != OpArgs.BonePosePriority)
 					{
@@ -2977,6 +2961,33 @@ namespace UE::Mutable::Private
 							BonePose.BonePosePriority = OpArgs.BonePosePriority;
 						}
 					}
+
+					if (OpArgs.PhysicsAssetId != PASSTHROUGH_ID_INVALID)
+					{
+						Result->PhysicsAssets.Emplace(FindPassthrough<UPhysicsAsset>(OpArgs.PhysicsAssetId));
+					}
+
+					if (EnumHasAnyFlags(OpArgs.Flags, OpEnumFlags::IsAdditionalPhysicsAssetList) &&
+						Program.ConstantUInt32Lists.IsValidIndex(OpArgs.AdditionalPhysicsAssetsIds.ListAddress))
+					{
+						const TArray<uint32>& ConstantPassthroughIdList = Program.ConstantUInt32Lists[OpArgs.AdditionalPhysicsAssetsIds.ListAddress];
+						Result->AdditionalPhysicsAssets.Reserve(Result->AdditionalPhysicsAssets.Num() + ConstantPassthroughIdList.Num());
+
+						for (uint32 PassthroughId : ConstantPassthroughIdList)
+						{
+							check(PassthroughId != PASSTHROUGH_ID_INVALID);
+							Result->AdditionalPhysicsAssets.Emplace(FindPassthrough<UPhysicsAsset>(PassthroughId));
+						}
+					}
+					else if (EnumHasAnyFlags(OpArgs.Flags, OpEnumFlags::HasAdditionalPhysicsAsset))
+					{
+						check(OpArgs.AdditionalPhysicsAssetsIds.PassthroughID != PASSTHROUGH_ID_INVALID);
+						Result->AdditionalPhysicsAssets.Emplace(FindPassthrough<UPhysicsAsset>(OpArgs.AdditionalPhysicsAssetsIds.PassthroughID));
+					}
+					else
+					{
+						check(!EnumHasAnyFlags(OpArgs.Flags, OpEnumFlags::IsAdditionalPhysicsAssetList));
+					}
 				}
 
 				StoreMesh(Item, Result);
diff --git a/Engine/Plugins/Mutable/Source/MutableRuntime/Private/MuR/CodeRunnerBegin.cpp b/Engine/Plugins/Mutable/Source/MutableRuntime/Private/MuR/CodeRunnerBegin.cpp
--- a/Engine/Plugins/Mutable/Source/MutableRuntime/Private/MuR/CodeRunnerBegin.cpp
+++ b/Engine/Plugins/Mutable/Source/MutableRuntime/Private/MuR/CodeRunnerBegin.cpp
@@ -2232,7 +2232,36 @@ namespace UE::Mutable::Private
 
 					const OP::MeshAddMetadataArgs OpArgs = Program.GetOpArgs<OP::MeshAddMetadataArgs>(Item.At);
 					PushOp(FScheduledOpInline(OpArgs.Source, Item));
-					
+
+					using OpEnumFlags = OP::MeshAddMetadataArgs::EnumFlags;
+
+					if (OpArgs.PhysicsAssetId != PASSTHROUGH_ID_INVALID)
+					{
+						PassthroughObjectLoader->Add<UPhysicsAsset>(OpArgs.PhysicsAssetId);
+					}
+
+					if (EnumHasAnyFlags(OpArgs.Flags, OpEnumFlags::IsAdditionalPhysicsAssetList) &&
+						Program.ConstantUInt32Lists.IsValidIndex(OpArgs.AdditionalPhysicsAssetsIds.ListAddress))
+					{
+						const TArray<uint32>& ConstantPassthroughIdList = Program.ConstantUInt32Lists[OpArgs.AdditionalPhysicsAssetsIds.ListAddress];
+						
+						for (uint32 PassthroughId : ConstantPassthroughIdList)
+						{
+							if (PassthroughId != PASSTHROUGH_ID_INVALID)
+							{
+								PassthroughObjectLoader->Add<UPhysicsAsset>(PassthroughId);
+							}
+						}
+					}
+					else if (EnumHasAnyFlags(OpArgs.Flags, OpEnumFlags::HasAdditionalPhysicsAsset))
+					{
+						PassthroughObjectLoader->Add<UPhysicsAsset>(OpArgs.AdditionalPhysicsAssetsIds.PassthroughID);
+					}
+					else
+					{
+						check(!EnumHasAnyFlags(OpArgs.Flags, OpEnumFlags::IsAdditionalPhysicsAssetList));
+					}
+
 					StoreNone(Item);
 					break;
 				}
diff --git a/Engine/Plugins/Mutable/Source/MutableRuntime/Private/MuR/OpMeshFormat.cpp b/Engine/Plugins/Mutable/Source/MutableRuntime/Private/MuR/OpMeshFormat.cpp
--- a/Engine/Plugins/Mutable/Source/MutableRuntime/Private/MuR/OpMeshFormat.cpp
+++ b/Engine/Plugins/Mutable/Source/MutableRuntime/Private/MuR/OpMeshFormat.cpp
@@ -695,7 +695,9 @@ namespace UE::Mutable::Private
 
 		// Copy the rest of the data
 		Result->SetSkeleton(Source->GetSkeleton());
+
 		Result->SetPhysicsBody(Source->GetPhysicsBody());
+		Result->PhysicsAssets = Source->PhysicsAssets;
 
 		Result->Layouts.Empty();
 		for (const TManagedPtr<const FLayout>& Layout : Source->Layouts)
@@ -717,6 +719,7 @@ namespace UE::Mutable::Private
 
 		// A shallow copy is done here, it should not be a problem.
 		Result->AdditionalPhysicsBodies = Source->AdditionalPhysicsBodies;
+		Result->AdditionalPhysicsAssets = Source->AdditionalPhysicsAssets;
 
 		Result->Surfaces = Source->Surfaces;
 		Result->ClothSections = Source->ClothSections;
diff --git a/Engine/Plugins/Mutable/Source/MutableRuntime/Private/MuR/PhysicsBody.cpp b/Engine/Plugins/Mutable/Source/MutableRuntime/Private/MuR/PhysicsBody.cpp
--- a/Engine/Plugins/Mutable/Source/MutableRuntime/Private/MuR/PhysicsBody.cpp
+++ b/Engine/Plugins/Mutable/Source/MutableRuntime/Private/MuR/PhysicsBody.cpp
@@ -169,7 +169,6 @@ namespace UE::Mutable::Private
 		Result->Bodies = Bodies;
 		Result->BodiesBoneNames = BodiesBoneNames;
 		Result->BoneIds = BoneIds;
-		Result->BodiesCustomIds = BodiesCustomIds;
 
 		Result->bBodiesModified = bBodiesModified;
 
@@ -181,7 +180,6 @@ namespace UE::Mutable::Private
 	{
 		Bodies.SetNum(Count);
 		BodiesBoneNames.SetNum(Count);
-		BodiesCustomIds.Init(-1, Count);
 	}
 
 	int32 FPhysicsBody::GetBodyCount() const
@@ -189,10 +187,9 @@ namespace UE::Mutable::Private
 		return Bodies.Num();
 	}
 
-	void FPhysicsBody::SetBody(int32 BodyIndex, FName InBoneName, int32 InBodyCustomId)
+	void FPhysicsBody::SetBody(int32 BodyIndex, FName InBoneName)
 	{
 		BodiesBoneNames[BodyIndex] = InBoneName;
-		BodiesCustomIds[BodyIndex] = InBodyCustomId;
 	}
 
 	FName FPhysicsBody::GetBodyBoneName(int32 BodyIndex) const
@@ -587,7 +584,6 @@ namespace UE::Mutable::Private
 		Arch << CustomId;
 		Arch << Bodies;
 		Arch << BoneIds;
-		Arch << BodiesCustomIds;
 		Arch << bBodiesModified;
 	}
 
@@ -596,7 +592,6 @@ namespace UE::Mutable::Private
 		Arch >> CustomId;        	
 		Arch >> Bodies;
 		Arch >> BoneIds;
-		Arch >> BodiesCustomIds;
 		Arch >> bBodiesModified;
 	}
 }
diff --git a/Engine/Plugins/Mutable/Source/MutableTools/Private/MuT/ASTOpMeshAddMetadata.cpp b/Engine/Plugins/Mutable/Source/MutableTools/Private/MuT/ASTOpMeshAddMetadata.cpp
--- a/Engine/Plugins/Mutable/Source/MutableTools/Private/MuT/ASTOpMeshAddMetadata.cpp
+++ b/Engine/Plugins/Mutable/Source/MutableTools/Private/MuT/ASTOpMeshAddMetadata.cpp
@@ -30,12 +30,14 @@ namespace UE::Mutable::Private
 		if (OtherUntyped.GetOpType() == GetOpType())
 		{
 			const ASTOpMeshAddMetadata* Other = static_cast<const ASTOpMeshAddMetadata*>(&OtherUntyped);
-			return 
-				Source      == Other->Source      && 
-				Tags        == Other->Tags        && 
+			return
+				Source == Other->Source &&
+				Tags == Other->Tags &&
 				ResourceIds == Other->ResourceIds &&
 				SkeletonIds == Other->SkeletonIds &&
-				BonePosePriority == Other->BonePosePriority;
+				BonePosePriority == Other->BonePosePriority &&
+				PhysicsAssetId == Other->PhysicsAssetId &&
+				AdditionalPhysicsAssetIds == Other->AdditionalPhysicsAssetIds;
 		}
 
 		return false;
@@ -50,7 +52,8 @@ namespace UE::Mutable::Private
 		NewOp->ResourceIds = ResourceIds;
 		NewOp->SkeletonIds = SkeletonIds;
 		NewOp->BonePosePriority = BonePosePriority;
-
+		NewOp->PhysicsAssetId = PhysicsAssetId;
+		NewOp->AdditionalPhysicsAssetIds = AdditionalPhysicsAssetIds;
 		return NewOp;
 	}
 
@@ -86,44 +89,68 @@ namespace UE::Mutable::Private
 			using OpEnumFlags = OP::MeshAddMetadataArgs::EnumFlags;
 
 			Args.Flags = OpEnumFlags::None;
-			EnumAddFlags(Args.Flags, Tags.Num()        != 1 ? OpEnumFlags::IsTagList      : OpEnumFlags::None);
-			EnumAddFlags(Args.Flags, ResourceIds.Num() != 1 ? OpEnumFlags::IsResourceList : OpEnumFlags::None);
-			EnumAddFlags(Args.Flags, SkeletonIds.Num() != 1 ? OpEnumFlags::IsSkeletonList : OpEnumFlags::None);
-
-			bool bIsTagList      = EnumHasAnyFlags(Args.Flags, OpEnumFlags::IsTagList);
-			bool bIsResourceList = EnumHasAnyFlags(Args.Flags, OpEnumFlags::IsResourceList);
-			bool bIsSkeletonList = EnumHasAnyFlags(Args.Flags, OpEnumFlags::IsSkeletonList);
-			if (bIsTagList)
-			{
-				Args.Tags.ListAddress = Program.AddConstant(Tags);
-			}
-			else
+			
+			if (Tags.Num())
 			{
-				check(Tags.Num() > 0);
-				Args.Tags.TagAddress = Program.AddConstant(Tags[0]);
+				EnumAddFlags(Args.Flags, OpEnumFlags::HasTags);
+
+				if (Tags.Num() > 1)
+				{
+					EnumAddFlags(Args.Flags, OpEnumFlags::IsTagList);
+					Args.Tags.ListAddress = Program.AddConstant(Tags);
+				}
+				else
+				{
+					Args.Tags.TagAddress = Program.AddConstant(Tags[0]);
+				}
 			}
 
-			if (bIsResourceList)
-			{
-				Args.ResourceIds.ListAddress = Program.AddConstant(ResourceIds);
-			}
-			else
+			if (ResourceIds.Num())
 			{
-				check(ResourceIds.Num() > 0);
-				Args.ResourceIds.ResourceId = ResourceIds[0];
+				EnumAddFlags(Args.Flags, OpEnumFlags::HasResources);
+
+				if (ResourceIds.Num() > 1)
+				{
+					EnumAddFlags(Args.Flags, OpEnumFlags::IsResourceList);
+					Args.ResourceIds.ListAddress = Program.AddConstant(ResourceIds);
+				}
+				else
+				{
+					Args.ResourceIds.ResourceId = ResourceIds[0];
+				}
 			}
 
-			if (bIsSkeletonList)
+			if (SkeletonIds.Num())
 			{
-				Args.SkeletonIds.ListAddress = Program.AddConstant(SkeletonIds);
-			}
-			else
-			{
-				check(SkeletonIds.Num() > 0);
-				Args.SkeletonIds.SkeletonId = SkeletonIds[0];
+				EnumAddFlags(Args.Flags, OpEnumFlags::HasSkeleton);
+				if (SkeletonIds.Num() > 1)
+				{
+					EnumAddFlags(Args.Flags, OpEnumFlags::IsSkeletonList);
+					Args.SkeletonIds.ListAddress = Program.AddConstant(SkeletonIds);
+				}
+				else
+				{
+					Args.SkeletonIds.SkeletonId = SkeletonIds[0];
+				}
 			}
 
 			Args.BonePosePriority = BonePosePriority;
+			
+			Args.PhysicsAssetId = PhysicsAssetId;
+
+			if (AdditionalPhysicsAssetIds.Num())
+			{
+				EnumAddFlags(Args.Flags, OpEnumFlags::HasAdditionalPhysicsAsset);
+				if (AdditionalPhysicsAssetIds.Num() > 1)
+				{
+					EnumAddFlags(Args.Flags, OpEnumFlags::IsAdditionalPhysicsAssetList);
+					Args.AdditionalPhysicsAssetsIds.ListAddress = Program.AddConstant(AdditionalPhysicsAssetIds);
+				}
+				else
+				{
+					Args.AdditionalPhysicsAssetsIds.PassthroughID = AdditionalPhysicsAssetIds[0];
+				}
+			}
 
 			Program.OpAddress.Add((uint32)Program.ByteCode.Num());
 			AppendCode(Program.ByteCode, EOpType::ME_ADDMETADATA);
diff --git a/Engine/Plugins/Mutable/Source/MutableTools/Private/MuT/ASTOpMeshAddMetadata.h b/Engine/Plugins/Mutable/Source/MutableTools/Private/MuT/ASTOpMeshAddMetadata.h
--- a/Engine/Plugins/Mutable/Source/MutableTools/Private/MuT/ASTOpMeshAddMetadata.h
+++ b/Engine/Plugins/Mutable/Source/MutableTools/Private/MuT/ASTOpMeshAddMetadata.h
@@ -12,7 +12,7 @@ namespace UE::Mutable::Private
 {
 	struct FProgram;
 
-	/** Add tags to a mesh. */
+	/** Add metadata to a mesh. */
 	class ASTOpMeshAddMetadata final : public ASTOp
 	{
 	public:
@@ -32,6 +32,13 @@ namespace UE::Mutable::Private
 		/** Bone Pose priority when merging two meshes. */
 		int16 BonePosePriority = 0;
 
+		/** PhysicsAsset to add */
+		PASSTHROUGH_ID PhysicsAssetId = PASSTHROUGH_ID_INVALID;
+
+		/** Additional PhysicsAssets */
+		TArray<PASSTHROUGH_ID> AdditionalPhysicsAssetIds;
+
+
 	public:
 
 		ASTOpMeshAddMetadata();
diff --git a/Engine/Plugins/Mutable/Source/MutableTools/Private/MuT/ASTOpSwitch.cpp b/Engine/Plugins/Mutable/Source/MutableTools/Private/MuT/ASTOpSwitch.cpp
--- a/Engine/Plugins/Mutable/Source/MutableTools/Private/MuT/ASTOpSwitch.cpp
+++ b/Engine/Plugins/Mutable/Source/MutableTools/Private/MuT/ASTOpSwitch.cpp
@@ -711,11 +711,14 @@ namespace UE::Mutable::Private
 			{
 			case EOpType::ME_ADDMETADATA:
 			{
-				// Move the add tags out of the switch if all tags are the same
+				// Move the metadata out of the switch if all metadatas are the same
 				bool bAllMetadataIsTheSame = true;
 				TArray<FString> Tags;
 				TArray<uint64> ResourceIds;
 				TArray<uint32> SkeletonIds;
+				int16 BonePosePriority = 0;
+				PASSTHROUGH_ID PhysicsAssetId = PASSTHROUGH_ID_INVALID;
+				TArray<uint32> AdditionalPhysicsAssetIds;
 
 				if (Default)
 				{
@@ -724,6 +727,9 @@ namespace UE::Mutable::Private
 					Tags = Typed->Tags;
 					ResourceIds = Typed->ResourceIds;
 					SkeletonIds = Typed->SkeletonIds;
+					BonePosePriority = Typed->BonePosePriority;
+					PhysicsAssetId = Typed->PhysicsAssetId;
+					AdditionalPhysicsAssetIds = Typed->AdditionalPhysicsAssetIds;
 				}
 
 				for (const FCase& Case : Cases)
@@ -735,16 +741,22 @@ namespace UE::Mutable::Private
 
 					check(Case.Branch->GetOpType() == EOpType::ME_ADDMETADATA);
 					const ASTOpMeshAddMetadata* Typed = static_cast<const ASTOpMeshAddMetadata*>(Case.Branch.child().get());
-					if (Tags.IsEmpty() && ResourceIds.IsEmpty() && SkeletonIds.IsEmpty())
+					if (Tags.IsEmpty() && ResourceIds.IsEmpty() && SkeletonIds.IsEmpty() && PhysicsAssetId == PASSTHROUGH_ID_INVALID && AdditionalPhysicsAssetIds.IsEmpty())
 					{
 						Tags = Typed->Tags;
 						ResourceIds = Typed->ResourceIds;
 						SkeletonIds = Typed->SkeletonIds;
+						BonePosePriority = Typed->BonePosePriority;
+						PhysicsAssetId = Typed->PhysicsAssetId;
+						AdditionalPhysicsAssetIds = Typed->AdditionalPhysicsAssetIds;
 					}
 					else if (
 							Typed->Tags != Tags || 
 							Typed->ResourceIds != ResourceIds || 
-							Typed->SkeletonIds != SkeletonIds)
+							Typed->SkeletonIds != SkeletonIds ||
+							Typed->BonePosePriority != BonePosePriority ||
+							Typed->PhysicsAssetId != PhysicsAssetId ||
+							Typed->AdditionalPhysicsAssetIds != AdditionalPhysicsAssetIds)
 					{
 						bAllMetadataIsTheSame = false;
 						break;
@@ -757,6 +769,9 @@ namespace UE::Mutable::Private
 					New->Tags = Tags;
 					New->ResourceIds = ResourceIds;
 					New->SkeletonIds = SkeletonIds;
+					New->BonePosePriority= BonePosePriority;
+					New->PhysicsAssetId = PhysicsAssetId;
+					New->AdditionalPhysicsAssetIds = AdditionalPhysicsAssetIds;
 
 					{
 						Ptr<ASTOpSwitch> NewSwitch = UE::Mutable::Private::Clone<ASTOpSwitch>(this);
diff --git a/Engine/Plugins/Mutable/Source/MutableTools/Private/MuT/CodeGenerator_Mesh.cpp b/Engine/Plugins/Mutable/Source/MutableTools/Private/MuT/CodeGenerator_Mesh.cpp
--- a/Engine/Plugins/Mutable/Source/MutableTools/Private/MuT/CodeGenerator_Mesh.cpp
+++ b/Engine/Plugins/Mutable/Source/MutableTools/Private/MuT/CodeGenerator_Mesh.cpp
@@ -811,12 +811,23 @@ namespace UE::Mutable::Private
 				}
 
 				// Separate metadata from the mesh
-				TArray<FString> Tags = ReferenceMesh->Tags;
-				TArray<uint64> ResourceIds = ReferenceMesh->StreamedResources;
-				TArray<uint32> SkeletonIds = ReferenceMesh->SkeletonIDs;
-				int32 BonePosePriority = 0;
-
-				if (Mesh->Tags.Num() + Mesh->StreamedResources.Num() + Mesh->SkeletonIDs.Num() > 0 || Mesh->BonePosePriority != 0)
+				TArray<FString> Tags = ReferenceMesh->Tags; // __AnimBP Tags are added to the ReferenceMesh
+
+				TArray<uint64> ResourceIds;
+				TArray<uint32> SkeletonIds;
+				int16 BonePosePriority = 0;
+				PASSTHROUGH_ID PhysicsAssetId = PASSTHROUGH_ID_INVALID;
+				TArray<PASSTHROUGH_ID> AdditionalPhysicsAssetId;
+
+				const bool bMeshHasMetadata = ReferenceMesh->Tags.Num() ||
+					Mesh->SkeletonIDs.Num() > 0 ||
+					Mesh->Tags.Num() ||
+					Mesh->StreamedResources.Num() ||
+					Mesh->BonePosePriority != 0 ||
+					Mesh->PhysicsAssets.Num() ||
+					Mesh->AdditionalPhysicsAssets.Num();
+
+				if (bMeshHasMetadata)
 				{
 					for (FString& Tag : Mesh->Tags)
 					{
@@ -835,12 +846,25 @@ namespace UE::Mutable::Private
 
 					BonePosePriority = Mesh->BonePosePriority;
 
+					if (Mesh->PhysicsAssets.Num())
+					{
+						check(Mesh->PhysicsAssets.Num() == 1);
+						PhysicsAssetId =  Mesh->PhysicsAssets[0].GetId();
+					}
+
+					for (const TPassthroughObjectPtr<UPhysicsAsset>& AdditionalPhysicsAsset : Mesh->AdditionalPhysicsAssets)
+					{
+						AdditionalPhysicsAssetId.Add(AdditionalPhysicsAsset.GetId());
+					}
+
 					TManagedPtr<FMesh> MetadatalessMesh = CloneOrTakeOver(MoveTemp(Mesh));
 
 					MetadatalessMesh->Tags.SetNum(0, EAllowShrinking::No);
 					MetadatalessMesh->StreamedResources.SetNum(0, EAllowShrinking::No);
 					MetadatalessMesh->SkeletonIDs.SetNum(0, EAllowShrinking::No);
 					MetadatalessMesh->BonePosePriority = 0;
+					MetadatalessMesh->PhysicsAssets.SetNum(0, EAllowShrinking::No);
+					MetadatalessMesh->AdditionalPhysicsAssets.SetNum(0, EAllowShrinking::No);
 
 					Mesh = MetadatalessMesh;
 				}
@@ -912,7 +936,7 @@ namespace UE::Mutable::Private
 				Result.BaseMeshOp = LastMeshOp;
 
 				// Add the metadata operation
-				if (Tags.Num() + ResourceIds.Num() + SkeletonIds.Num() > 0 || BonePosePriority != 0)
+				if (bMeshHasMetadata)
 				{
 					Ptr<ASTOpMeshAddMetadata> AddMetadataOp = new ASTOpMeshAddMetadata;
 					AddMetadataOp->Source = LastMeshOp;
@@ -922,6 +946,9 @@ namespace UE::Mutable::Private
 					AddMetadataOp->SkeletonIds = MoveTemp(SkeletonIds);
 					AddMetadataOp->BonePosePriority = BonePosePriority;
 
+					AddMetadataOp->PhysicsAssetId = PhysicsAssetId;
+					AddMetadataOp->AdditionalPhysicsAssetIds = MoveTemp(AdditionalPhysicsAssetId);
+
 					LastMeshOp = AddMetadataOp;
 				}
 
