[Mutable] Changed Mutable core variable types.
+Now layer indexes are an int8 instead of an int32.
+Now parameter names are a FName instead of a FString
+Added a function to store Arrays of FNames into the Program.

#rnx

#Synced-CL 50245735

diff --git a/Engine/Plugins/Mutable/Source/CustomizableObject/Internal/MuCO/CustomizableObjectPrivate.h b/Engine/Plugins/Mutable/Source/CustomizableObject/Internal/MuCO/CustomizableObjectPrivate.h
--- a/Engine/Plugins/Mutable/Source/CustomizableObject/Internal/MuCO/CustomizableObjectPrivate.h
+++ b/Engine/Plugins/Mutable/Source/CustomizableObject/Internal/MuCO/CustomizableObjectPrivate.h
@@ -1565,6 +1565,8 @@ enum class ECustomizableObjectVersions : int32
 
 	MovePhysicsAssetCore,
 
+	ChangedLayerIndexDataType,
+
 	// -----<new versions can be added above this line>--------
 	LastCustomizableObjectVersion
 };
diff --git a/Engine/Plugins/Mutable/Source/CustomizableObject/Private/MuCO/CustomizableObjectSystem.cpp b/Engine/Plugins/Mutable/Source/CustomizableObject/Private/MuCO/CustomizableObjectSystem.cpp
--- a/Engine/Plugins/Mutable/Source/CustomizableObject/Private/MuCO/CustomizableObjectSystem.cpp
+++ b/Engine/Plugins/Mutable/Source/CustomizableObject/Private/MuCO/CustomizableObjectSystem.cpp
@@ -2209,7 +2209,7 @@ namespace impl
 				{
 					FInstanceUpdateData::FScalar& Scalar = Context->InstanceUpdateData.Scalars.AddDefaulted_GetRef();
 					Scalar.Name = ScalarParameter.Key.ParameterName;
-					Scalar.MaterialLayer = ScalarParameter.Key.LayerIndex;
+					Scalar.MaterialLayer = (int32)ScalarParameter.Key.LayerIndex;
 					Scalar.Scalar = ScalarParameter.Value;
 					
 				}
@@ -2221,7 +2221,7 @@ namespace impl
 				{
 					FInstanceUpdateData::FVector& Vector = Context->InstanceUpdateData.Vectors.AddDefaulted_GetRef();
 					Vector.Name = VectorParameter.Key.ParameterName;
-					Vector.MaterialLayer = VectorParameter.Key.LayerIndex;
+					Vector.MaterialLayer = (int32)VectorParameter.Key.LayerIndex;
 					Vector.Vector = VectorParameter.Value;
 				}
 
@@ -2234,7 +2234,7 @@ namespace impl
 					
 					TSharedRef<FInstanceUpdateData::FImage>& Image = Context->InstanceUpdateData.Images.Add_GetRef(MakeShared<FInstanceUpdateData::FImage>());
 					Image->ParameterName = ImageParameter.Key.ParameterName;
-					Image->MaterialLayer = ImageParameter.Key.LayerIndex;
+					Image->MaterialLayer = (int32)ImageParameter.Key.LayerIndex;
 					Image->ImageID = UE::Mutable::Private::FSystem::GetImageId(Context->LiveInstance->ImageIdRegistry, *Context->Model.Get(), *Context->Parameters.Get(), Address);
 					Image->FullImageSizeX = 0;
 					Image->FullImageSizeY = 0;
diff --git a/Engine/Plugins/Mutable/Source/CustomizableObjectEditor/Private/MuCOE/GenerateMutableSource/GenerateMutableSourceColor.cpp b/Engine/Plugins/Mutable/Source/CustomizableObjectEditor/Private/MuCOE/GenerateMutableSource/GenerateMutableSourceColor.cpp
--- a/Engine/Plugins/Mutable/Source/CustomizableObjectEditor/Private/MuCOE/GenerateMutableSource/GenerateMutableSourceColor.cpp
+++ b/Engine/Plugins/Mutable/Source/CustomizableObjectEditor/Private/MuCOE/GenerateMutableSource/GenerateMutableSourceColor.cpp
@@ -403,13 +403,14 @@ UE::Mutable::Private::Ptr<UE::Mutable::Private::NodeColour> GenerateMutableSourc
 			FName ParameterName = TypedNodeMaterialBreak->GetPinParameterName(*Pin);
 			if (!ParameterName.IsNone())
 			{
-				int32 LayerIndex = TypedNodeMaterialBreak->GetPinParameterLayerIndex(*Pin);
+				const int32 LayerIndex = TypedNodeMaterialBreak->GetPinParameterLayerIndex(*Pin);
+				check(LayerIndex <= MAX_int8);
 
-				// Set the name of the parameter that represents the pin
-				MaterialBreak->ParameterKey = {ParameterName, LayerIndex};
+				// Set the key of the parameter that represents the pin
+				MaterialBreak->ParameterKey = {ParameterName, (int8)LayerIndex};
 
 				// Generate the material that is going to be broken into parameters
-				GenerationContext.MaterialBreakParameterStack.Add(FMaterialBreakParameter{ {ParameterName, LayerIndex}, Pin->PinType.PinCategory });
+				GenerationContext.MaterialBreakParameterStack.Add(FMaterialBreakParameter{ MaterialBreak->ParameterKey, Pin->PinType.PinCategory });
 				UE::Mutable::Private::Ptr<UE::Mutable::Private::NodeMaterial> LinkedNode = GenerateMutableSourceMaterial(LinkedMaterial, GenerationContext, FGenerationMaterialOptions());
 				MaterialBreak->MaterialSource = LinkedNode;
 				GenerationContext.MaterialBreakParameterStack.Pop();
diff --git a/Engine/Plugins/Mutable/Source/CustomizableObjectEditor/Private/MuCOE/GenerateMutableSource/GenerateMutableSourceFloat.cpp b/Engine/Plugins/Mutable/Source/CustomizableObjectEditor/Private/MuCOE/GenerateMutableSource/GenerateMutableSourceFloat.cpp
--- a/Engine/Plugins/Mutable/Source/CustomizableObjectEditor/Private/MuCOE/GenerateMutableSource/GenerateMutableSourceFloat.cpp
+++ b/Engine/Plugins/Mutable/Source/CustomizableObjectEditor/Private/MuCOE/GenerateMutableSource/GenerateMutableSourceFloat.cpp
@@ -423,13 +423,14 @@ UE::Mutable::Private::Ptr<UE::Mutable::Private::NodeScalar> GenerateMutableSourc
 			FName ParameterName = TypedNodeMaterialBreak->GetPinParameterName(*Pin);
 			if (!ParameterName.IsNone())
 			{
-				int32 LayerIndex = TypedNodeMaterialBreak->GetPinParameterLayerIndex(*Pin);
+				const int32 LayerIndex = TypedNodeMaterialBreak->GetPinParameterLayerIndex(*Pin);
+				check(LayerIndex <= MAX_int8);
 
-				// Set the name of the parameter that represents the pin
-				MaterialBreak->ParameterKey = {ParameterName, LayerIndex};
+				// Set the key of the parameter that represents the pin
+				MaterialBreak->ParameterKey = {ParameterName, (int8)LayerIndex};
 
 				// Generate the material that is going to be broken into parameters
-				GenerationContext.MaterialBreakParameterStack.Add(FMaterialBreakParameter{ {ParameterName, LayerIndex}, Pin->PinType.PinCategory});
+				GenerationContext.MaterialBreakParameterStack.Add(FMaterialBreakParameter{ MaterialBreak->ParameterKey, Pin->PinType.PinCategory});
 				UE::Mutable::Private::Ptr<UE::Mutable::Private::NodeMaterial> LinkedNode = GenerateMutableSourceMaterial(LinkedMaterial, GenerationContext, FGenerationMaterialOptions());
 				MaterialBreak->MaterialSource = LinkedNode;
 				GenerationContext.MaterialBreakParameterStack.Pop();
diff --git a/Engine/Plugins/Mutable/Source/CustomizableObjectEditor/Private/MuCOE/GenerateMutableSource/GenerateMutableSourceImage.cpp b/Engine/Plugins/Mutable/Source/CustomizableObjectEditor/Private/MuCOE/GenerateMutableSource/GenerateMutableSourceImage.cpp
--- a/Engine/Plugins/Mutable/Source/CustomizableObjectEditor/Private/MuCOE/GenerateMutableSource/GenerateMutableSourceImage.cpp
+++ b/Engine/Plugins/Mutable/Source/CustomizableObjectEditor/Private/MuCOE/GenerateMutableSource/GenerateMutableSourceImage.cpp
@@ -1241,17 +1241,18 @@ UE::Mutable::Private::Ptr<UE::Mutable::Private::NodeImage> GenerateMutableSource
 			FName ParameterName = TypedNodeMaterialBreak->GetPinParameterName(*Pin);
 			if (!ParameterName.IsNone())
 			{
-				int32 LayerIndex = TypedNodeMaterialBreak->GetPinParameterLayerIndex(*Pin);
+				const int32 LayerIndex = TypedNodeMaterialBreak->GetPinParameterLayerIndex(*Pin);
+				check(LayerIndex <= MAX_int8);
 
-				// Set the name of the parameter that represents the pin
-				MaterialBreak->ParameterKey = {ParameterName, LayerIndex};
+				// Set the key of the parameter that represents the pin
+				MaterialBreak->ParameterKey = {ParameterName, (int8)LayerIndex};
 
 				// Store the reference texture into the material generation options to generate all the textures with the correct size
 				FGenerationMaterialOptions GenerationOptions;
 				GenerationOptions.ReferenceTextureSize = ReferenceTextureSize;
 
 				// Generate the material that is going to be broken into parameters
-				GenerationContext.MaterialBreakParameterStack.Add(FMaterialBreakParameter{ {ParameterName, LayerIndex}, Pin->PinType.PinCategory});
+				GenerationContext.MaterialBreakParameterStack.Add(FMaterialBreakParameter{ MaterialBreak->ParameterKey, Pin->PinType.PinCategory});
 				UE::Mutable::Private::Ptr<UE::Mutable::Private::NodeMaterial> LinkedNode = GenerateMutableSourceMaterial(LinkedMaterial, GenerationContext, GenerationOptions);
 				MaterialBreak->MaterialSource = LinkedNode;
 				GenerationContext.MaterialBreakParameterStack.Pop();
diff --git a/Engine/Plugins/Mutable/Source/CustomizableObjectEditor/Private/MuCOE/GenerateMutableSource/GenerateMutableSourceMaterial.cpp b/Engine/Plugins/Mutable/Source/CustomizableObjectEditor/Private/MuCOE/GenerateMutableSource/GenerateMutableSourceMaterial.cpp
--- a/Engine/Plugins/Mutable/Source/CustomizableObjectEditor/Private/MuCOE/GenerateMutableSource/GenerateMutableSourceMaterial.cpp
+++ b/Engine/Plugins/Mutable/Source/CustomizableObjectEditor/Private/MuCOE/GenerateMutableSource/GenerateMutableSourceMaterial.cpp
@@ -168,13 +168,16 @@ UE::Mutable::Private::Ptr<UE::Mutable::Private::NodeMaterial> GenerateMutableSou
 			const UEdGraphPin* ImagePin = TypedNodeMaterialBase->GetParameterPin(EMaterialParameterType::Texture, ImageIndex);
 			const FName ImageName = TypedNodeMaterialBase->GetParameterName(EMaterialParameterType::Texture, ImageIndex);
 			const int32 LayerIndex = TypedNodeMaterialBase->GetParameterLayerIndex(EMaterialParameterType::Texture, ImageIndex);
+			check(LayerIndex <= MAX_int8);
+
+			UE::Mutable::Private::FParameterKey ImageKey = { ImageName, (int8)LayerIndex };
+
 			const bool bIsImagePinLinked = ImagePin && FollowInputPin(*ImagePin);
 
 			if (bIsImagePinLinked && !TypedNodeMaterialBase->IsImageMutableMode(ImageIndex))
 			{
 				if (const UEdGraphPin* ConnectedPin = FollowInputPin(*ImagePin))
 				{
-					UE::Mutable::Private::FParameterKey ImageKey = { ImageName, LayerIndex };
 					FGeneratedImageProperties Props = GenerateImageProperties(TypedNodeMaterialBase, nullptr, ImageKey, true, GenerationContext);
 
 					// This is a connected pass-through texture that simply has to be passed to the core
@@ -239,8 +242,7 @@ UE::Mutable::Private::Ptr<UE::Mutable::Private::NodeMaterial> GenerateMutableSou
 						}
 					}
 
-					// Generate image key to identify it
-					UE::Mutable::Private::FParameterKey ImageKey = { ImageName, LayerIndex };
+					// Generate image properties
 					FGeneratedImageProperties Props = GenerateImageProperties(TypedNodeMaterialBase, ReferenceTexture, ImageKey, false, GenerationContext);
 
 					// Generate the texture nodes
@@ -342,15 +344,14 @@ UE::Mutable::Private::Ptr<UE::Mutable::Private::NodeMaterial> GenerateMutableSou
 					FName VectorName = TypedNodeMaterialBase->GetParameterName(EMaterialParameterType::Vector, VectorIndex);
 					FNodeMaterialParameterId VectorId = TypedNodeMaterialBase->GetParameterId(EMaterialParameterType::Vector, VectorIndex);
 					const int32 LayerIndex = TypedNodeMaterialBase->GetParameterLayerIndex(EMaterialParameterType::Vector, VectorIndex);
+					check(LayerIndex <= MAX_int8);
 
 					// Store info to generate colors from the material of a table node
 					GenerationContext.CurrentMaterialTableParameter = VectorName.ToString();
 					GenerationContext.CurrentMaterialParameterId = VectorId.ParameterId.ToString();
 
 					UE::Mutable::Private::Ptr<UE::Mutable::Private::NodeColour> ColorNode = GenerateMutableSourceColor(ConnectedPin, GenerationContext);
-					UE::Mutable::Private::FParameterKey ColorKey = { VectorName, LayerIndex };
-
-					BaseMaterial->ColorParameters.Add(ColorKey, ColorNode);
+					BaseMaterial->ColorParameters.Add({ VectorName, (int8)LayerIndex }, ColorNode);
 				}
 			}
 		}
@@ -369,15 +370,14 @@ UE::Mutable::Private::Ptr<UE::Mutable::Private::NodeMaterial> GenerateMutableSou
 					FName ScalarName = TypedNodeMaterialBase->GetParameterName(EMaterialParameterType::Scalar, ScalarIndex);
 					FNodeMaterialParameterId ScalarId = TypedNodeMaterialBase->GetParameterId(EMaterialParameterType::Scalar, ScalarIndex);
 					const int32 LayerIndex = TypedNodeMaterialBase->GetParameterLayerIndex(EMaterialParameterType::Scalar, ScalarIndex);
+					check(LayerIndex <= MAX_int8);
 					
 					// Store info to generate scalars from the material of a table
 					GenerationContext.CurrentMaterialTableParameter = ScalarName.ToString();
 					GenerationContext.CurrentMaterialParameterId = ScalarId.ParameterId.ToString();
 
 					UE::Mutable::Private::Ptr<UE::Mutable::Private::NodeScalar> ScalarNode = GenerateMutableSourceFloat(ConnectedPin, GenerationContext);
-					UE::Mutable::Private::FParameterKey ScalarKey = { ScalarName, LayerIndex };
-
-					BaseMaterial->ScalarParameters.Add(ScalarKey, ScalarNode);
+					BaseMaterial->ScalarParameters.Add({ ScalarName, (int8)LayerIndex }, ScalarNode);
 				}
 			}
 		}
@@ -399,7 +399,7 @@ UE::Mutable::Private::Ptr<UE::Mutable::Private::NodeMaterial> GenerateMutableSou
 			{
 				FHashedMaterialParameterInfo ParameterInfo;
 				ParameterInfo.Name = FScriptName(CurrentBreakParameter.ParameterKey.ParameterName);
-				ParameterInfo.Index = CurrentBreakParameter.ParameterKey.LayerIndex;
+				ParameterInfo.Index = (int32)CurrentBreakParameter.ParameterKey.LayerIndex;
 				ParameterInfo.Association = ParameterInfo.Index == INDEX_NONE ? GlobalParameter : LayerParameter;
 
 				if (!TypedNodeMaterialConstant->Material)
@@ -600,7 +600,7 @@ UE::Mutable::Private::Ptr<UE::Mutable::Private::NodeMaterial> GenerateMutableSou
 				// Data to find the parameter value inside the material
 				FHashedMaterialParameterInfo ParameterInfo;
 				ParameterInfo.Name = FScriptName(CurrentBreakParameter.ParameterKey.ParameterName);
-				ParameterInfo.Index = CurrentBreakParameter.ParameterKey.LayerIndex;
+				ParameterInfo.Index = (int32)CurrentBreakParameter.ParameterKey.LayerIndex;
 				ParameterInfo.Association = ParameterInfo.Index == INDEX_NONE ? GlobalParameter : LayerParameter;
 
 				//TODO(Max - GM): Move this to the generatemutablesourceimage?
@@ -720,7 +720,9 @@ UE::Mutable::Private::Ptr<UE::Mutable::Private::NodeMaterial> GenerateMutableSou
 			}
 			
 			const int32 LayerIndex = TypedNodeMaterialModify->GetPinParameterLayerIndex(*InputPin);
-			UE::Mutable::Private::FParameterKey ParameterKey = { ParameterName, LayerIndex };
+			check(LayerIndex <= MAX_int8);
+
+			UE::Mutable::Private::FParameterKey ParameterKey = { ParameterName, (int8)LayerIndex };
 
 			const FName PinType = InputPin->PinType.PinCategory;
 			if (PinType == UEdGraphSchema_CustomizableObject::PC_Texture)
@@ -760,8 +762,7 @@ UE::Mutable::Private::Ptr<UE::Mutable::Private::NodeMaterial> GenerateMutableSou
 
 				if (const UEdGraphPin* ConnectedPin = FollowInputPin(*InputPin))
 				{
-					UE::Mutable::Private::FParameterKey ImageKey = { ParameterName, LayerIndex };
-					FGeneratedImageProperties Props = GenerateImageProperties(TypedNodeMaterialBase, nullptr, ImageKey, true, GenerationContext);
+					FGeneratedImageProperties Props = GenerateImageProperties(Node, nullptr, ParameterKey, true, GenerationContext);
 
 					// This is a connected pass-through texture that simply has to be passed to the core
 					UE::Mutable::Private::Ptr<UE::Mutable::Private::NodeImage> PassThroughImagePtr = GenerateMutableSourceImage(ConnectedPin, GenerationContext, 0);
diff --git a/Engine/Plugins/Mutable/Source/CustomizableObjectEditor/Private/MuCOE/GenerateMutableSource/GenerateMutableSourceModifier.cpp b/Engine/Plugins/Mutable/Source/CustomizableObjectEditor/Private/MuCOE/GenerateMutableSource/GenerateMutableSourceModifier.cpp
--- a/Engine/Plugins/Mutable/Source/CustomizableObjectEditor/Private/MuCOE/GenerateMutableSource/GenerateMutableSourceModifier.cpp
+++ b/Engine/Plugins/Mutable/Source/CustomizableObjectEditor/Private/MuCOE/GenerateMutableSource/GenerateMutableSourceModifier.cpp
@@ -320,8 +320,9 @@ UE::Mutable::Private::Ptr<UE::Mutable::Private::NodeModifier> GenerateMutableSou
 			for (int32 ImageIndex = 0; ImageIndex < NumImages; ++ImageIndex)
 			{
 				UE::Mutable::Private::NodeImagePtr ImageNode;
-				FName MaterialParameterName = TypedNodeExt->GetParameterName(EMaterialParameterType::Texture, ImageIndex);
-				int32 MaterialParameterLayerIndex = TypedNodeExt->GetParameterLayerIndex(EMaterialParameterType::Texture, ImageIndex);
+				const FName& MaterialParameterName = TypedNodeExt->GetParameterName(EMaterialParameterType::Texture, ImageIndex);
+				const int32 MaterialParameterLayerIndex = TypedNodeExt->GetParameterLayerIndex(EMaterialParameterType::Texture, ImageIndex);
+				check(MaterialParameterLayerIndex <= MAX_int8);
 
 				if (!ImageNode) // If
 				{
@@ -356,7 +357,7 @@ UE::Mutable::Private::Ptr<UE::Mutable::Private::NodeModifier> GenerateMutableSou
 				{
 					UE::Mutable::Private::NodeModifierSurfaceEdit::FTexture Texture;
 					Texture.Extend = ImageNode;
-					Texture.MaterialParameterKey = UE::Mutable::Private::FParameterKey(MaterialParameterName, MaterialParameterLayerIndex);
+					Texture.MaterialParameterKey = UE::Mutable::Private::FParameterKey(MaterialParameterName, (int8)MaterialParameterLayerIndex);
 
 					SurfNode->LODs[LODIndex].Textures.Add(Texture);
 				}
@@ -465,11 +466,12 @@ UE::Mutable::Private::Ptr<UE::Mutable::Private::NodeModifier> GenerateMutableSou
 					//check(ParentMaterialNode->IsImageMutableMode(ImageIndex)); // Ensured at graph time. If it fails, something is wrong.
 
 					const UEdGraphPin* ConnectedImagePin = FollowInputPin(*TypedNodeEdit->GetUsedImagePin(ImageId));
-					FName MaterialParameterName = TypedNodeEdit->GetParameterName(EMaterialParameterType::Texture, ImageIndex);
-					int32 MaterialParameterLayerIndex = TypedNodeEdit->GetParameterLayerIndex(EMaterialParameterType::Texture, ImageIndex);
+					const FName& MaterialParameterName = TypedNodeEdit->GetParameterName(EMaterialParameterType::Texture, ImageIndex);
+					const int32 MaterialParameterLayerIndex = TypedNodeEdit->GetParameterLayerIndex(EMaterialParameterType::Texture, ImageIndex);
+					check(MaterialParameterLayerIndex <= MAX_int8);
 					
 					UE::Mutable::Private::NodeModifierSurfaceEdit::FTexture ImagePatch;
-					ImagePatch.MaterialParameterKey = UE::Mutable::Private::FParameterKey(MaterialParameterName, MaterialParameterLayerIndex);
+					ImagePatch.MaterialParameterKey = UE::Mutable::Private::FParameterKey(MaterialParameterName,  (int8)MaterialParameterLayerIndex);
 					// \todo: expose these two options?
 					ImagePatch.PatchBlendType = UE::Mutable::Private::EBlendType::BT_BLEND;
 					ImagePatch.bPatchApplyToAlpha = true;
diff --git a/Engine/Plugins/Mutable/Source/CustomizableObjectEditor/Private/MuCOE/SMutableMaterialViewer.cpp b/Engine/Plugins/Mutable/Source/CustomizableObjectEditor/Private/MuCOE/SMutableMaterialViewer.cpp
--- a/Engine/Plugins/Mutable/Source/CustomizableObjectEditor/Private/MuCOE/SMutableMaterialViewer.cpp
+++ b/Engine/Plugins/Mutable/Source/CustomizableObjectEditor/Private/MuCOE/SMutableMaterialViewer.cpp
@@ -434,7 +434,7 @@ void SMutableMaterialViewer::SetMaterial(const UE::Mutable::Private::TManagedPtr
 				TSharedPtr<FMutableMaterialImageParameterElement> ImageParameter = MakeShared<FMutableMaterialImageParameterElement>();
 				ImageParameter->Index = ImageParameterIndex++;
 				ImageParameter->ParameterName = ImageParameterTuple.Key.ParameterName;
-				ImageParameter->LayerIndex = ImageParameterTuple.Key.LayerIndex;
+				ImageParameter->LayerIndex = (int32)ImageParameterTuple.Key.LayerIndex;
 				ImageParameter->Value = ImageParameterTuple.Value;
 				ImageParameters.Add(ImageParameter);
 			}
@@ -447,7 +447,7 @@ void SMutableMaterialViewer::SetMaterial(const UE::Mutable::Private::TManagedPtr
 				TSharedPtr<FMutableMaterialColorParameterElement> ColorParameter = MakeShared<FMutableMaterialColorParameterElement>();
 				ColorParameter->Index = ColorParameterIndex++;
 				ColorParameter->ParameterName = ColorParameterTuple.Key.ParameterName;
-				ColorParameter->LayerIndex = ColorParameterTuple.Key.LayerIndex;
+				ColorParameter->LayerIndex = (int32)ColorParameterTuple.Key.LayerIndex;
 				ColorParameter->Value = ColorParameterTuple.Value;
 				ColorParameters.Add(ColorParameter);
 			}
@@ -460,7 +460,7 @@ void SMutableMaterialViewer::SetMaterial(const UE::Mutable::Private::TManagedPtr
 				TSharedPtr<FMutableMaterialScalarParameterElement> ScalarParameter = MakeShared<FMutableMaterialScalarParameterElement>();
 				ScalarParameter->Index = ScalarParameterIndex++;
 				ScalarParameter->ParameterName = ScalarParameterTuple.Key.ParameterName;
-				ScalarParameter->LayerIndex = ScalarParameterTuple.Key.LayerIndex;
+				ScalarParameter->LayerIndex = (int32)ScalarParameterTuple.Key.LayerIndex;
 				ScalarParameter->Value = ScalarParameterTuple.Value;
 				ScalarParameters.Add(ScalarParameter);
 			}
diff --git a/Engine/Plugins/Mutable/Source/MutableRuntime/Internal/MuR/Material.h b/Engine/Plugins/Mutable/Source/MutableRuntime/Internal/MuR/Material.h
--- a/Engine/Plugins/Mutable/Source/MutableRuntime/Internal/MuR/Material.h
+++ b/Engine/Plugins/Mutable/Source/MutableRuntime/Internal/MuR/Material.h
@@ -23,7 +23,7 @@ namespace UE::Mutable::Private
 	struct FParameterKey
 	{
 		FName ParameterName;
-		int32 LayerIndex = -1;
+		int8 LayerIndex = -1;
 
 		inline void Serialise(FOutputArchive& arch) const;
 		inline void Unserialise(FInputArchive& arch);
diff --git a/Engine/Plugins/Mutable/Source/MutableRuntime/Internal/MuR/Model.h b/Engine/Plugins/Mutable/Source/MutableRuntime/Internal/MuR/Model.h
--- a/Engine/Plugins/Mutable/Source/MutableRuntime/Internal/MuR/Model.h
+++ b/Engine/Plugins/Mutable/Source/MutableRuntime/Internal/MuR/Model.h
@@ -419,6 +419,8 @@ namespace UE::Mutable::Private
 
 		UE_API uint32 AddConstant(const FName& InName);
 
+		UE_API OP::ADDRESS AddConstant(const TArray<FName>& NameList);
+
         UE_API OP::CONSTANT_MATRIX_INDEX AddConstant(const FMatrix44f& Constant);
     	
         UE_API OP::CONSTANT_SHAPE_INDEX AddConstant(const FShape& Constant);
diff --git a/Engine/Plugins/Mutable/Source/MutableRuntime/Internal/MuR/Operations.h b/Engine/Plugins/Mutable/Source/MutableRuntime/Internal/MuR/Operations.h
--- a/Engine/Plugins/Mutable/Source/MutableRuntime/Internal/MuR/Operations.h
+++ b/Engine/Plugins/Mutable/Source/MutableRuntime/Internal/MuR/Operations.h
@@ -1215,14 +1215,14 @@ namespace UE::Mutable::Private
 		{
 			ADDRESS Material;
 			CONSTANT_STRING_INDEX ParameterName;
-			int32 LayerIndex;
+			int8 LayerIndex;
 		};
 
 		struct MaterialBreakImageParameterArgs
 		{
 			ADDRESS MaterialParameter;
 			CONSTANT_STRING_INDEX ParameterName;
-			int32 LayerIndex;
+			int8 LayerIndex;
 		};
 
     	struct ExternalArgs
diff --git a/Engine/Plugins/Mutable/Source/MutableRuntime/Private/MuR/CodeRunner.cpp b/Engine/Plugins/Mutable/Source/MutableRuntime/Private/MuR/CodeRunner.cpp
--- a/Engine/Plugins/Mutable/Source/MutableRuntime/Private/MuR/CodeRunner.cpp
+++ b/Engine/Plugins/Mutable/Source/MutableRuntime/Private/MuR/CodeRunner.cpp
@@ -5689,9 +5689,8 @@ namespace UE::Mutable::Private
 			{
 				TManagedPtr<const FMaterial> Material = LoadMaterial(FCacheAddress(Args.Material, Item));
 
-				check(Args.ParameterName < (uint32)Program.ConstantStrings.Num());
-				const FString& ParameterName = Program.ConstantStrings[Args.ParameterName];
-				const FParameterKey ParameterKey = { FName(ParameterName), Args.LayerIndex };
+				const FName& ParameterName = Program.ConstantNames[Args.ParameterName];
+				const FParameterKey ParameterKey = { ParameterName, Args.LayerIndex };
 				const FMaterial::FImageParameterData* ImageParameterData = Material->ImageParameters.Find(ParameterKey);
 				bool bParameterFound = false;
 
@@ -5717,8 +5716,8 @@ namespace UE::Mutable::Private
 				{
 					UTexture* ParameterValue = nullptr;
 					FHashedMaterialParameterInfo ParameterInfo;
-					ParameterInfo.Name = FScriptName(FName(ParameterName));
-					ParameterInfo.Index = ParameterKey.LayerIndex;
+					ParameterInfo.Name = FScriptName(ParameterName);
+					ParameterInfo.Index = (int32)ParameterKey.LayerIndex;
 					ParameterInfo.Association = ParameterInfo.Index == INDEX_NONE ? GlobalParameter : LayerParameter;
 
 					// Find image using the parameter info
@@ -5735,7 +5734,7 @@ namespace UE::Mutable::Private
 				
 				if(!bParameterFound)
 				{
-					UE_LOG(LogMutableCore, Error, TEXT("Expected Texture parameter [%s] not found in Break Material node."), *ParameterName);
+					UE_LOG(LogMutableCore, Error, TEXT("Expected Texture parameter [%s] not found in Break Material node."), *ParameterName.ToString());
 					AbortRun();
 				}
 
@@ -5744,10 +5743,8 @@ namespace UE::Mutable::Private
 			case 2:
 			{
 				TManagedPtr<const FMaterial> Material = LoadMaterial(FCacheAddress(Args.Material, Item));
-
-				check(Args.ParameterName < (uint32)Program.ConstantStrings.Num());
-				const FString& ParameterName = Program.ConstantStrings[Args.ParameterName];
-				const FParameterKey ParameterKey = { FName(ParameterName), Args.LayerIndex };
+				const FName& ParameterName = Program.ConstantNames[Args.ParameterName];
+				const FParameterKey ParameterKey = { ParameterName, Args.LayerIndex };
 				const FMaterial::FImageParameterData* ImageParameterData = Material->ImageParameters.Find(ParameterKey);
 
 				if (!ImageParameterData)
@@ -6221,10 +6218,8 @@ namespace UE::Mutable::Private
 			case 1:
 			{
 				TManagedPtr<const FMaterial> Material = LoadMaterial(FCacheAddress(Args.Material, Item));
-				
-				check(Args.ParameterName < (uint32)Program.ConstantStrings.Num());
-				const FString& ParameterName = Program.ConstantStrings[Args.ParameterName];
-				const FParameterKey ParameterKey = { FName(ParameterName), Args.LayerIndex };
+				const FName& ParameterName = Program.ConstantNames[Args.ParameterName];
+				const FParameterKey ParameterKey = { ParameterName, Args.LayerIndex };
 				const float* Value = Material->ScalarParameters.Find(ParameterKey);
 				bool bParameterFound = false;
 
@@ -6239,8 +6234,8 @@ namespace UE::Mutable::Private
 					float ParameterValue = 0.0f;
 
 					FHashedMaterialParameterInfo ParameterInfo;
-					ParameterInfo.Name = FScriptName(FName(ParameterName));
-					ParameterInfo.Index = ParameterKey.LayerIndex;
+					ParameterInfo.Name = FScriptName(ParameterName);
+					ParameterInfo.Index = (int32)ParameterKey.LayerIndex;
 					ParameterInfo.Association = ParameterInfo.Index == INDEX_NONE ? GlobalParameter : LayerParameter;
 
 					// Find float using the parameter info
@@ -6254,7 +6249,7 @@ namespace UE::Mutable::Private
 
 				if(!bParameterFound)
 				{
-					UE_LOG(LogMutableCore, Error, TEXT("Expected Scalar parameter [%s] not found in Break Material node."), *ParameterName);
+					UE_LOG(LogMutableCore, Error, TEXT("Expected Scalar parameter [%s] not found in Break Material node."), *ParameterName.ToString());
 					AbortRun();
 				}
 
@@ -6561,10 +6556,8 @@ namespace UE::Mutable::Private
 			case 1:
 			{
 				TManagedPtr<const FMaterial> Material = LoadMaterial(FCacheAddress(Args.Material, Item));
-
-				check(Args.ParameterName < (uint32)Program.ConstantStrings.Num());
-				const FString& ParameterName = Program.ConstantStrings[Args.ParameterName];
-				const FParameterKey ParameterKey = { FName(ParameterName), Args.LayerIndex };
+				const FName& ParameterName = Program.ConstantNames[Args.ParameterName];
+				const FParameterKey ParameterKey = { ParameterName, Args.LayerIndex };
 				const FVector4f* Value = Material->ColorParameters.Find(ParameterKey);
 				bool bParameterFound = false;
 
@@ -6579,8 +6572,8 @@ namespace UE::Mutable::Private
 					FLinearColor ParameterValue;
 
 					FHashedMaterialParameterInfo ParameterInfo;
-					ParameterInfo.Name = FScriptName(FName(ParameterName));
-					ParameterInfo.Index = ParameterKey.LayerIndex;
+					ParameterInfo.Name = FScriptName(ParameterName);
+					ParameterInfo.Index = (int32)ParameterKey.LayerIndex;
 					ParameterInfo.Association = ParameterInfo.Index == INDEX_NONE ? GlobalParameter : LayerParameter;
 
 					// Find float using the parameter info
@@ -6594,7 +6587,7 @@ namespace UE::Mutable::Private
 
 				if(!bParameterFound)
 				{
-					UE_LOG(LogMutableCore, Error, TEXT("Expected Color parameter [%s] not found in Break Material node."), *ParameterName);
+					UE_LOG(LogMutableCore, Error, TEXT("Expected Color parameter [%s] not found in Break Material node."), *ParameterName.ToString());
 					AbortRun();
 				}
 
@@ -6717,17 +6710,17 @@ namespace UE::Mutable::Private
 				for (int32 ParameterIndex = 0; ParameterIndex < ImageParameterNames.Num(); ++ParameterIndex)
 				{
 					// Get Image Parameter Name
-					const FString& ParameterName = Program.ConstantStrings[ImageParameterNames[ParameterIndex]];
-					
+					const FName& ParameterName = Program.ConstantNames[ImageParameterNames[ParameterIndex]];
+
 					// Get Image Parameter Layer
-					int32 LayerIndex;
-					FMemory::Memcpy(&LayerIndex, Data, sizeof(int32));
-					Data += sizeof(int32);
+					int8 LayerIndex;
+					FMemory::Memcpy(&LayerIndex, Data, sizeof(int8));
+					Data += sizeof(int8);
 
 					// Construct image parameter info
 					FHashedMaterialParameterInfo ImageParameterInfo;
-					ImageParameterInfo.Name = FScriptName(FName(ParameterName));
-					ImageParameterInfo.Index = LayerIndex;
+					ImageParameterInfo.Name = FScriptName(ParameterName);
+					ImageParameterInfo.Index = (int32)LayerIndex;
 					ImageParameterInfo.Association = ImageParameterInfo.Index == INDEX_NONE ? GlobalParameter : LayerParameter;
 
 					// Find Texture using the parameter info
@@ -6743,7 +6736,7 @@ namespace UE::Mutable::Private
 						FMaterial::FImageParameterData NewImageParameter;
 						NewImageParameter.ImageParameter = NewParameterImage;
 
-						Material->ImageParameters.Add({FName(ImageParameterInfo.Name), ImageParameterInfo.Index}, NewImageParameter);
+						Material->ImageParameters.Add({ParameterName, LayerIndex}, NewImageParameter);
 					}
 				}
 			}
@@ -6780,9 +6773,9 @@ namespace UE::Mutable::Private
 					FMemory::Memcpy(&NewName, Data, sizeof(OP::ADDRESS));
 					Data += sizeof(OP::ADDRESS);
 
-					uint32 LayerIndex;
-					FMemory::Memcpy(&LayerIndex, Data, sizeof(int32));
-					Data += sizeof(int32);
+					int8 LayerIndex;
+					FMemory::Memcpy(&LayerIndex, Data, sizeof(int8));
+					Data += sizeof(int8);
 
 					OP::ADDRESS NewParameterOperation;
 					FMemory::Memcpy(&NewParameterOperation, Data, sizeof(OP::ADDRESS));
@@ -6806,7 +6799,7 @@ namespace UE::Mutable::Private
 			case 1:
 			{
 				TArray<OP::ADDRESS> ParameterNames;
-				TArray<int32> ParameterLayerIndex;
+				TArray<int8> ParameterLayerIndex;
 				TArray<OP::ADDRESS> ParameterOperations;
 				TArray<int32> ImageParameterIndexArray;
 
@@ -6817,10 +6810,10 @@ namespace UE::Mutable::Private
 					ParameterNames.Add(NewName);
 					Data += sizeof(OP::ADDRESS);
 
-					uint32 LayerIndex;
-					FMemory::Memcpy(&LayerIndex, Data, sizeof(int32));
+					int8 LayerIndex;
+					FMemory::Memcpy(&LayerIndex, Data, sizeof(int8));
 					ParameterLayerIndex.Add(LayerIndex);
-					Data += sizeof(int32);
+					Data += sizeof(int8);
 
 					OP::ADDRESS NewParameterOperation;
 					FMemory::Memcpy(&NewParameterOperation, Data, sizeof(OP::ADDRESS));
@@ -6849,10 +6842,10 @@ namespace UE::Mutable::Private
 					for (int32 ParameterIndex = 0; ParameterIndex < ParameterCount; ++ParameterIndex)
 					{
 						// Get the parameter name from the constant string
-						const FName ParameterName = FName(Program.ConstantStrings[ParameterNames[ParameterIndex]]);
+						const FName& ParameterName = Program.ConstantNames[ParameterNames[ParameterIndex]];
 
 						// Construct the parameter key
-						FParameterKey ParameterKey = { ParameterName, ParameterLayerIndex[ParameterIndex] };						
+						FParameterKey ParameterKey = { ParameterName, ParameterLayerIndex[ParameterIndex] };
 						EDataType OperationType = GetOpDataType(Program.GetOpType(ParameterOperations[ParameterIndex]));
 
 						switch (OperationType)
@@ -8647,10 +8640,8 @@ namespace UE::Mutable::Private
 				CacheMaterialOp.Type = FScheduledOp::EType::Full;
 
 				TManagedPtr<const FMaterial> Material = LoadMaterial(CacheMaterialOp);
-
-				check(Args.ParameterName < (uint32)Program.ConstantStrings.Num());
-				const FString& ParameterName = Program.ConstantStrings[Args.ParameterName];
-				const FParameterKey ParameterKey = { FName(ParameterName), Args.LayerIndex };
+				const FName& ParameterName = Program.ConstantNames[Args.ParameterName];
+				const FParameterKey ParameterKey = { ParameterName, Args.LayerIndex };
 				const FMaterial::FImageParameterData* ImageParameterData = Material->ImageParameters.Find(ParameterKey);
 				bool bParameterFound = false;
 
@@ -8677,8 +8668,8 @@ namespace UE::Mutable::Private
 				{
 					UTexture* ParameterValue = nullptr;
 					FHashedMaterialParameterInfo ParameterInfo;
-					ParameterInfo.Name = FScriptName(FName(ParameterName));
-					ParameterInfo.Index = ParameterKey.LayerIndex;
+					ParameterInfo.Name = FScriptName(ParameterName);
+					ParameterInfo.Index = (int32)ParameterKey.LayerIndex;
 					ParameterInfo.Association = ParameterInfo.Index == INDEX_NONE ? GlobalParameter : LayerParameter;
 
 					// Find image using the parameter info
@@ -8707,10 +8698,8 @@ namespace UE::Mutable::Private
 				CacheMaterialOp.Type = FScheduledOp::EType::Full;
 
 				TManagedPtr<const FMaterial> Material = LoadMaterial(CacheMaterialOp);
-
-				check(Args.ParameterName < (uint32)Program.ConstantStrings.Num());
-				const FString& ParameterName = Program.ConstantStrings[Args.ParameterName];
-				const FParameterKey ParameterKey = { FName(ParameterName), Args.LayerIndex };
+				const FName& ParameterName = Program.ConstantNames[Args.ParameterName];
+				const FParameterKey ParameterKey = { ParameterName, Args.LayerIndex };
 				const FMaterial::FImageParameterData* ImageParameterData = Material->ImageParameters.Find(ParameterKey);
 					
 				if (ImageParameterData)
@@ -8744,8 +8733,7 @@ namespace UE::Mutable::Private
 			UMaterialInterface* MaterialInterface = Parameters.GetMaterialValue(Args.MaterialParameter, Index.Get());
 
 			// Get the texture parameter name
-			check(Args.ParameterName < (uint32)Program.ConstantStrings.Num());
-			const FString& ParameterName = Program.ConstantStrings[Args.ParameterName];
+			const FName& ParameterName = Program.ConstantNames[Args.ParameterName];
 
 			UTexture* Image = nullptr;
 
@@ -8753,8 +8741,8 @@ namespace UE::Mutable::Private
 			if (MaterialInterface)
 			{
 				FHashedMaterialParameterInfo ParameterInfo;
-				ParameterInfo.Name = FScriptName(FName(ParameterName));
-				ParameterInfo.Index = Args.LayerIndex;
+				ParameterInfo.Name = FScriptName(ParameterName);
+				ParameterInfo.Index = (int32)Args.LayerIndex;
 				ParameterInfo.Association = ParameterInfo.Index == INDEX_NONE ? GlobalParameter : LayerParameter;
 
 				MaterialInterface->GetTextureParameterValue(ParameterInfo, Image);
diff --git a/Engine/Plugins/Mutable/Source/MutableRuntime/Private/MuR/CodeRunnerBegin.cpp b/Engine/Plugins/Mutable/Source/MutableRuntime/Private/MuR/CodeRunnerBegin.cpp
--- a/Engine/Plugins/Mutable/Source/MutableRuntime/Private/MuR/CodeRunnerBegin.cpp
+++ b/Engine/Plugins/Mutable/Source/MutableRuntime/Private/MuR/CodeRunnerBegin.cpp
@@ -2811,10 +2811,8 @@ namespace UE::Mutable::Private
 						case 1:
 						{
 							TManagedPtr<const FMaterial> Material = LoadMaterial(Args.Material);
-
-							check(Args.ParameterName < (uint32)Program.ConstantStrings.Num());
-							const FString& ParameterName = Program.ConstantStrings[Args.ParameterName];
-							const FParameterKey ParameterKey = { FName(ParameterName), Args.LayerIndex };
+							const FName& ParameterName = Program.ConstantNames[Args.ParameterName];
+							const FParameterKey ParameterKey = { ParameterName, Args.LayerIndex };
 							const EOpType InputOpType = Program.GetOpType(Item.At);
 							const EDataType InputDataType = GetOpDataType(InputOpType);
 
@@ -2980,9 +2978,9 @@ namespace UE::Mutable::Private
 							FMemory::Memcpy(&NewName, Data, sizeof(OP::ADDRESS));
 							Data += sizeof(OP::ADDRESS);
 
-							uint32 LayerIndex;
-							FMemory::Memcpy(&LayerIndex, Data, sizeof(int32));
-							Data += sizeof(int32);
+							int8 LayerIndex;
+							FMemory::Memcpy(&LayerIndex, Data, sizeof(int8));
+							Data += sizeof(int8);
 
 							OP::ADDRESS NewParameterOperation;
 							FMemory::Memcpy(&NewParameterOperation, Data, sizeof(OP::ADDRESS));
@@ -3024,9 +3022,9 @@ namespace UE::Mutable::Private
 							FMemory::Memcpy(&NewName, Data, sizeof(OP::ADDRESS));
 							Data += sizeof(OP::ADDRESS);
 
-							uint32 LayerIndex;
-							FMemory::Memcpy(&LayerIndex, Data, sizeof(int32));
-							Data += sizeof(int32);
+							int8 LayerIndex;
+							FMemory::Memcpy(&LayerIndex, Data, sizeof(int8));
+							Data += sizeof(int8);
 
 							OP::ADDRESS ParameterOperation;
 							FMemory::Memcpy(&ParameterOperation, Data, sizeof(OP::ADDRESS));
diff --git a/Engine/Plugins/Mutable/Source/MutableRuntime/Private/MuR/CodeRunnerTasked.cpp b/Engine/Plugins/Mutable/Source/MutableRuntime/Private/MuR/CodeRunnerTasked.cpp
--- a/Engine/Plugins/Mutable/Source/MutableRuntime/Private/MuR/CodeRunnerTasked.cpp
+++ b/Engine/Plugins/Mutable/Source/MutableRuntime/Private/MuR/CodeRunnerTasked.cpp
@@ -2566,15 +2566,19 @@ namespace UE::Mutable::Private
 			// Get material parameter from the array of parameters
 			UMaterialInterface* Material = Parameters.GetMaterialValue(Args.MaterialParameter, Index.Get());
 
-			// Get the texture parameter name
-			check(Args.ParameterName < (uint32)Program.ConstantStrings.Num());
-			const FString& ParameterName = Program.ConstantStrings[Args.ParameterName];
-
 			// Get the parameter texture from the UMaterial
 			if (Material)
 			{
+				// Get the texture parameter name
+				const FName& ParameterName = Program.ConstantNames[Args.ParameterName];
+
+				FHashedMaterialParameterInfo ParameterInfo;
+				ParameterInfo.Name = FScriptName(ParameterName);
+				ParameterInfo.Index = (int32)Args.LayerIndex;
+				ParameterInfo.Association = ParameterInfo.Index == INDEX_NONE ? GlobalParameter : LayerParameter;
+
 				UTexture* Texture;
-				bool bParameterFound = Material->GetTextureParameterValue(FName(ParameterName), Texture);
+				bool bParameterFound = Material->GetTextureParameterValue(ParameterInfo, Texture);
 
 				if (bParameterFound && Texture)
 				{
diff --git a/Engine/Plugins/Mutable/Source/MutableRuntime/Private/MuR/Model.cpp b/Engine/Plugins/Mutable/Source/MutableRuntime/Private/MuR/Model.cpp
--- a/Engine/Plugins/Mutable/Source/MutableRuntime/Private/MuR/Model.cpp
+++ b/Engine/Plugins/Mutable/Source/MutableRuntime/Private/MuR/Model.cpp
@@ -417,6 +417,7 @@ namespace UE::Mutable::Private
 	    return (OP::CONSTANT_STRING_INDEX)ConstantStrings.AddUnique(Str);
     }
 
+
 	uint32 FProgram::AddConstant(const FName& InName)
 	{
 		const FString NameString = InName.ToString().ToLower();
@@ -433,6 +434,23 @@ namespace UE::Mutable::Private
 		return NameId;
 	}
 
+
+	OP::ADDRESS FProgram::AddConstant(const TArray<FName>& NameList)
+	{
+		const int32 NumNames = NameList.Num();
+
+		TArray<uint32> NameAddrs;
+		NameAddrs.SetNumUninitialized(NumNames);
+
+		for (int32 NameIndex = 0; NameIndex < NumNames; ++NameIndex)
+		{
+			NameAddrs[NameIndex] = AddConstant(NameList[NameIndex]);
+		}
+
+		return (OP::ADDRESS)AddConstant(NameAddrs);
+	}
+
+
     OP::ADDRESS FProgram::AddConstant(const TArray<uint32>& UInt32List)
     {
 	    return (OP::ADDRESS)ConstantUInt32Lists.AddUnique(UInt32List);
diff --git a/Engine/Plugins/Mutable/Source/MutableRuntime/Private/MuR/Operations.cpp b/Engine/Plugins/Mutable/Source/MutableRuntime/Private/MuR/Operations.cpp
--- a/Engine/Plugins/Mutable/Source/MutableRuntime/Private/MuR/Operations.cpp
+++ b/Engine/Plugins/Mutable/Source/MutableRuntime/Private/MuR/Operations.cpp
@@ -328,9 +328,9 @@ namespace UE::Mutable::Private
 				FMemory::Memcpy(&NewName, Data, sizeof(OP::ADDRESS));
 				Data += sizeof(OP::ADDRESS);
 
-				uint32 LayerIndex;
-				FMemory::Memcpy(&LayerIndex, Data, sizeof(int32));
-				Data += sizeof(int32);
+				int8 LayerIndex;
+				FMemory::Memcpy(&LayerIndex, Data, sizeof(int8));
+				Data += sizeof(int8);
 
 				OP::ADDRESS NewParameterOperation;
 				FMemory::Memcpy(&NewParameterOperation, Data, sizeof(OP::ADDRESS));
diff --git a/Engine/Plugins/Mutable/Source/MutableTools/Private/MuT/ASTOpMaterialBreak.cpp b/Engine/Plugins/Mutable/Source/MutableTools/Private/MuT/ASTOpMaterialBreak.cpp
--- a/Engine/Plugins/Mutable/Source/MutableTools/Private/MuT/ASTOpMaterialBreak.cpp
+++ b/Engine/Plugins/Mutable/Source/MutableTools/Private/MuT/ASTOpMaterialBreak.cpp
@@ -95,7 +95,7 @@ namespace UE::Mutable::Private
 				Args.Material = Material->linkedAddress;
 			}
 
-			Args.ParameterName = Program.AddConstant(ParameterKey.ParameterName.ToString());
+			Args.ParameterName = Program.AddConstant(ParameterKey.ParameterName);
 			Args.LayerIndex = ParameterKey.LayerIndex;
 
 			AppendCode(Program.ByteCode, Args);
diff --git a/Engine/Plugins/Mutable/Source/MutableTools/Private/MuT/ASTOpMaterialBreakImageParameter.cpp b/Engine/Plugins/Mutable/Source/MutableTools/Private/MuT/ASTOpMaterialBreakImageParameter.cpp
--- a/Engine/Plugins/Mutable/Source/MutableTools/Private/MuT/ASTOpMaterialBreakImageParameter.cpp
+++ b/Engine/Plugins/Mutable/Source/MutableTools/Private/MuT/ASTOpMaterialBreakImageParameter.cpp
@@ -84,7 +84,7 @@ namespace UE::Mutable::Private
 			FMemory::Memzero(Args);
 
 			Args.MaterialParameter = OP::ADDRESS(LinkedParameterIndex);
-			Args.ParameterName = Program.AddConstant(ParameterKey.ParameterName.ToString());
+			Args.ParameterName = Program.AddConstant(ParameterKey.ParameterName);
 			Args.LayerIndex = ParameterKey.LayerIndex;
 
 			AppendCode(Program.ByteCode, Args);
diff --git a/Engine/Plugins/Mutable/Source/MutableTools/Private/MuT/ASTOpMaterialModify.cpp b/Engine/Plugins/Mutable/Source/MutableTools/Private/MuT/ASTOpMaterialModify.cpp
--- a/Engine/Plugins/Mutable/Source/MutableTools/Private/MuT/ASTOpMaterialModify.cpp
+++ b/Engine/Plugins/Mutable/Source/MutableTools/Private/MuT/ASTOpMaterialModify.cpp
@@ -129,12 +129,12 @@ namespace UE::Mutable::Private
 			for (const TPair<FParameterKey, ASTChild>& Parameter : ParametersToModify)
 			{
 				//Parameter Name
-				AppendCode(Program.ByteCode,Program.AddConstant(Parameter.Key.ParameterName.ToString()));
+				AppendCode(Program.ByteCode,Program.AddConstant(Parameter.Key.ParameterName));
 				
 				// Layer Index
 				AppendCode(Program.ByteCode, Parameter.Key.LayerIndex);
 
-				// Check if the exist a child operation
+				// Check if exists a child operation
 				if (Parameter.Value)
 				{
 					AppendCode(Program.ByteCode, Parameter.Value->linkedAddress);
diff --git a/Engine/Plugins/Mutable/Source/MutableTools/Private/MuT/ASTOpParameter.cpp b/Engine/Plugins/Mutable/Source/MutableTools/Private/MuT/ASTOpParameter.cpp
--- a/Engine/Plugins/Mutable/Source/MutableTools/Private/MuT/ASTOpParameter.cpp
+++ b/Engine/Plugins/Mutable/Source/MutableTools/Private/MuT/ASTOpParameter.cpp
@@ -113,11 +113,11 @@ namespace UE::Mutable::Private
 				Args.Variable = (OP::ADDRESS)LinkedParameterIndex;
 				
 				// Images
-				TArray<FString> ImageParameterNames;
+				TArray<FName> ImageParameterNames;
 				TArray<uint32> ImageParameterOperations;
 				for (const TPair<FParameterKey, ASTChild>& Element : ImageOperations)
 				{
-					ImageParameterNames.Add(Element.Key.ParameterName.ToString());
+					ImageParameterNames.Add(Element.Key.ParameterName);
 					ImageParameterOperations.Add(Element.Value->linkedAddress);
 		 		}
 
