[Mutable] Minor cloth related code cleanup.

#Synced-CL 50250210

diff --git a/Engine/Plugins/Mutable/Source/CustomizableObject/Private/MuCO/CustomizableObjectInstance.cpp b/Engine/Plugins/Mutable/Source/CustomizableObject/Private/MuCO/CustomizableObjectInstance.cpp
--- a/Engine/Plugins/Mutable/Source/CustomizableObject/Private/MuCO/CustomizableObjectInstance.cpp
+++ b/Engine/Plugins/Mutable/Source/CustomizableObject/Private/MuCO/CustomizableObjectInstance.cpp
@@ -4027,11 +4027,6 @@ void BuildOrCopyClothingData(
 
 		const TSharedRef<FInstanceUpdateData::FLOD>& LOD = Context->InstanceUpdateData.LODs[Component->FirstLOD + LODIndex];
 		
-		if (!LOD->Mesh)
-		{
-			continue;
-		}
-		
 		if (!LOD->Mesh->HasClothing())
 		{
 			continue;
@@ -4062,15 +4057,10 @@ void BuildOrCopyClothingData(
 	}
 	
 	// Residents
-	for (int32 LODIndex = Context->FirstResidentLOD[ComponentName]; LODIndex < Context->NumLODsAvailable[ComponentName]; ++LODIndex)
+	for (int32 LODIndex = Context->GetFirstRequestedLOD()[ComponentName]; LODIndex < Context->NumLODsAvailable[ComponentName]; ++LODIndex)
 	{
 		const TSharedRef<FInstanceUpdateData::FLOD>& LOD = Context->InstanceUpdateData.LODs[Component->FirstLOD + LODIndex];
 
-		if (!LOD->Mesh)
-		{
-			continue;
-		}
-
 		if (!LOD->Mesh->HasClothing())
 		{
 			continue;
@@ -4088,12 +4078,7 @@ void BuildOrCopyClothingData(
 	for (int32 LODIndex = FirstLOD; LODIndex < Component->LODCount; ++LODIndex)
 	{
 		const TSharedRef<FInstanceUpdateData::FLOD>& LOD = Context->InstanceUpdateData.LODs[Component->FirstLOD + LODIndex];
-	
-		if (!LOD->Mesh)
-		{
-			continue;
-		}
-
+		
 		if (!LOD->Mesh->HasClothing())
 		{
 			continue;
diff --git a/Engine/Plugins/Mutable/Source/MutableTools/Private/MuT/ASTOpConstantResource.cpp b/Engine/Plugins/Mutable/Source/MutableTools/Private/MuT/ASTOpConstantResource.cpp
--- a/Engine/Plugins/Mutable/Source/MutableTools/Private/MuT/ASTOpConstantResource.cpp
+++ b/Engine/Plugins/Mutable/Source/MutableTools/Private/MuT/ASTOpConstantResource.cpp
@@ -134,6 +134,12 @@ namespace UE::Mutable::Private
 				}
 
 				MeshGeometryData->MeshIDPrefix = 0;
+				
+				if (MeshGeometryData->ClothSections.Num())
+				{
+					MeshGeometryData->ClothSections[0].ClothingAsset.Reset();
+				}
+				
 				AddMeshToProgram(MeshGeometryData);	
 				EnumAddFlags(MeshContentFlags, EMeshContentFlags::GeometryData);
 			}
@@ -441,7 +447,7 @@ namespace UE::Mutable::Private
 					OP::MeshConstantArgs Args;
 					FMemory::Memzero(Args);
 
-					TManagedPtr<FMesh> MeshValue = StaticCastManagedPtr<const FMesh>(Value)->Clone();
+					TManagedPtr<const FMesh> MeshValue = StaticCastManagedPtr<const FMesh>(Value);
 					check(MeshValue);
 
 					Args.Skeleton = -1;
@@ -454,7 +460,7 @@ namespace UE::Mutable::Private
 					check(!MeshValue->PassthroughObject);
 					
 					Args.Value = AddConstantMesh(Program, MeshValue, *Options);
-					Args.ClothID = MeshValue->ClothSections.Num() ? MeshValue->ClothSections[0].ClothingAsset.Reset() : PASSTHROUGH_ID_INVALID;
+					Args.ClothID = MeshValue->ClothSections.Num() ? MeshValue->ClothSections[0].ClothingAsset.GetId() : PASSTHROUGH_ID_INVALID;
 
 					int32 DataDescIndex = Options->AdditionalData.SourceMeshPerConstant.Add(SourceDataDescriptor);
 					check(DataDescIndex == Args.Value);
