[Mutable] Fix Callback Thread check in FCustomizableObjectMeshStreamIn.

#Synced-CL 50250210

diff --git a/Engine/Plugins/Mutable/Source/CustomizableObject/Private/MuCO/CustomizableObjectMeshUpdate.cpp b/Engine/Plugins/Mutable/Source/CustomizableObject/Private/MuCO/CustomizableObjectMeshUpdate.cpp
--- a/Engine/Plugins/Mutable/Source/CustomizableObject/Private/MuCO/CustomizableObjectMeshUpdate.cpp
+++ b/Engine/Plugins/Mutable/Source/CustomizableObject/Private/MuCO/CustomizableObjectMeshUpdate.cpp
@@ -132,7 +132,7 @@ void FCustomizableObjectMeshStreamIn::DoConvertResources(const FContext& Context
 
 	if (bMarkRenderStateDirty)
 	{
-		PushTask(Context, TT_GameThread, SRA_UPDATE_CALLBACK(MarkRenderStateDirty), TT_None, SRA_UPDATE_CALLBACK(DoCancel));
+		PushTask(Context, TT_GameThread, SRA_UPDATE_CALLBACK(MarkRenderStateDirty), (EThreadType)Context.CurrentThread, SRA_UPDATE_CALLBACK(DoCancel));
 	}
 	else
 	{
