[Mutable] Move Skeletons to Passthrough objects

#rnx

#Synced-CL 50297267

diff --git a/Engine/Plugins/Mutable/Source/CustomizableObject/Internal/MuCO/CustomizableObjectInstancePrivate.h b/Engine/Plugins/Mutable/Source/CustomizableObject/Internal/MuCO/CustomizableObjectInstancePrivate.h
--- a/Engine/Plugins/Mutable/Source/CustomizableObject/Internal/MuCO/CustomizableObjectInstancePrivate.h
+++ b/Engine/Plugins/Mutable/Source/CustomizableObject/Internal/MuCO/CustomizableObjectInstancePrivate.h
@@ -72,23 +72,6 @@ enum ECOInstanceFlags
 ENUM_CLASS_FLAGS(ECOInstanceFlags);
 
 
-USTRUCT()
-struct FReferencedSkeletons
-{
-	GENERATED_USTRUCT_BODY();
-
-	// Merged skeleton if found in the cache
-	UPROPERTY()
-	TObjectPtr<USkeleton> Skeleton;
-
-	UPROPERTY()
-	TArray<uint16> SkeletonIds;
-
-	UPROPERTY()
-	TArray< TObjectPtr<USkeleton> > SkeletonsToMerge;
-};
-
-
 USTRUCT()
 struct FCustomizableInstanceComponentData
 {
@@ -116,9 +99,6 @@ struct FCustomizableInstanceComponentData
 	TArray<FString> MeshPartPaths;
 #endif
 
-	/** Skeletons required by the current generated instance. Skeletons to be loaded and merged.*/
-	UPROPERTY(Transient)
-	FReferencedSkeletons Skeletons;
 	
 	FComponentDescriptor LastComponentDescriptor;
 
@@ -392,7 +372,7 @@ public:
 
 UE_API UPhysicsAsset* GetOrBuildMainPhysicsAsset(const TSharedRef<FUpdateContextPrivate>& Context, TObjectPtr<UPhysicsAsset> TamplateAsset, const UE::Mutable::Private::FPhysicsBody* PhysicsBody, bool bDisableCollisionBetweenAssets, int32 ComponentIndex);
 
-UE_API USkeleton* MergeSkeletons(UCustomizableObject& CustomizableObject, FCustomizableInstanceComponentData& ComponentData, bool& bOutCreatedNewSkeleton);
+UE_API USkeleton* MergeSkeletons(const TArray<TObjectPtr<USkeleton>>& SkeletonsToMerge, UCustomizableObject& CustomizableObject, bool& bOutCreatedNewSkeleton);
 
 UE_API void InitSkeletalMeshData(const TSharedRef<FUpdateContextPrivate>& Context, USkeletalMesh* SkeletalMesh, const FMutableRefSkeletalMeshData& RefSkeletalMeshData, const UCustomizableObject& CustomizableObject, UE::Mutable::Private::FComponentId ComponentId);
 
diff --git a/Engine/Plugins/Mutable/Source/CustomizableObject/Internal/MuCO/CustomizableObjectPrivate.h b/Engine/Plugins/Mutable/Source/CustomizableObject/Internal/MuCO/CustomizableObjectPrivate.h
--- a/Engine/Plugins/Mutable/Source/CustomizableObject/Internal/MuCO/CustomizableObjectPrivate.h
+++ b/Engine/Plugins/Mutable/Source/CustomizableObject/Internal/MuCO/CustomizableObjectPrivate.h
@@ -206,14 +206,14 @@ private:
 class FSkeletonCache
 {
 public:
-	USkeleton* Get(const TArray<uint16>& Key);
+	USkeleton* Get(const TArray<FName>& Key);
 
-	void Add(const TArray<uint16>& Key, USkeleton* Value);
+	void Add(const TArray<FName>& Key, USkeleton* Value);
 
 	void Clear();
 
 private:
-	TMap<TArray<uint16>, TWeakObjectPtr<USkeleton>> MergedSkeletons;
+	TMap<TArray<FName>, TWeakObjectPtr<USkeleton>> MergedSkeletons;
 };
 
 
@@ -1080,10 +1080,6 @@ public:
 	 * LOD render data settings, Mesh sockets, Bounding volumes, etc. */
 	UPROPERTY()
 	TArray<FMutableRefSkeletalMeshData> ReferenceSkeletalMeshesData;	
-
-	/** Skeletons used by the compiled UE::Mutable::Private::FModel. */
-	UPROPERTY()
-	TArray<TSoftObjectPtr<USkeleton>> Skeletons;
 	
 	UPROPERTY()
 	TMap<uint32, TSoftObjectPtr<UObject>> PassthroughObjects;
@@ -1567,6 +1563,8 @@ enum class ECustomizableObjectVersions : int32
 
 	ChangedLayerIndexDataType,
 
+	MutableSkeletonToPassthroughObject,
+	
 	// -----<new versions can be added above this line>--------
 	LastCustomizableObjectVersion
 };
diff --git a/Engine/Plugins/Mutable/Source/CustomizableObject/Internal/MuCO/CustomizableObjectSystemPrivate.h b/Engine/Plugins/Mutable/Source/CustomizableObject/Internal/MuCO/CustomizableObjectSystemPrivate.h
--- a/Engine/Plugins/Mutable/Source/CustomizableObject/Internal/MuCO/CustomizableObjectSystemPrivate.h
+++ b/Engine/Plugins/Mutable/Source/CustomizableObject/Internal/MuCO/CustomizableObjectSystemPrivate.h
@@ -349,7 +349,15 @@ struct FInstanceUpdateData
 		
 		uint32 MaterialIndex = INDEX_NONE;
 	};
-	
+
+	struct FBone
+	{
+		FName BoneName;
+		FMatrix44f MatrixWithScale;
+
+		bool operator==(const FName& OtherBoneName) const { return BoneName == OtherBoneName; };
+	};
+
 	struct FComponent
 	{
 		UE::Mutable::Private::FComponentId Id = INDEX_NONE;
@@ -365,6 +373,8 @@ struct FInstanceUpdateData
 		
 		UE::Mutable::Private::FSkeletalMeshId SkeletalMeshId;
 		UE::Mutable::Private::TManagedPtr<const UE::Mutable::Private::FSkeletalMesh> SkeletalMesh;
+
+		TArray<FBone> BonePose;
 	};
 
 	TArray<TSharedRef<FComponent>> Components;
@@ -379,32 +389,12 @@ struct FInstanceUpdateData
 	TArray<FName> ActiveBones;
 	TArray<FName> BoneMaps;
 	
-	struct FBone
-	{
-		FName BoneName;
-		FMatrix44f MatrixWithScale;
-
-		bool operator==(const FName& OtherBoneName) const { return BoneName == OtherBoneName; };
-	};
-	
 	/** Key is the Block Id. Value is the LoadAdditionalAssetsAndData read destination. */
 	TMap<uint32, FMorphTargetMeshData> RealTimeMorphTargetMeshData;
 	
 	/** Key is the Component Name. Value is the final Morph Target data to copy into the Skeletal Mesh. */
 	TMap<UE::Mutable::Private::FComponentId, FSkeletalMeshMorphTargets> RealTimeMorphTargets;
 
-	struct FSkeletonData
-	{
-		TArray<TObjectPtr<USkeleton>> Skeletons;
-
-		TArray<uint16> SkeletonIds;
-
-		TArray<FBone> BonePose;
-	};
-
-	// Access by instance component index
-	TArray<FSkeletonData> SkeletonsPerInstanceComponent;
-
 	struct FNamedExtensionData
 	{
 		UE::Mutable::Private::TManagedPtr<const UE::Mutable::Private::FExtensionData> Data;
@@ -424,7 +414,6 @@ struct FInstanceUpdateData
 		Scalars.Empty();
 		Vectors.Empty();
 		RealTimeMorphTargets.Empty();
-		SkeletonsPerInstanceComponent.Empty();
 		ExtendedInputPins.Empty();
 	}
 };
diff --git a/Engine/Plugins/Mutable/Source/CustomizableObject/Private/MuCO/CustomizableObject.cpp b/Engine/Plugins/Mutable/Source/CustomizableObject/Private/MuCO/CustomizableObject.cpp
--- a/Engine/Plugins/Mutable/Source/CustomizableObject/Private/MuCO/CustomizableObject.cpp
+++ b/Engine/Plugins/Mutable/Source/CustomizableObject/Private/MuCO/CustomizableObject.cpp
@@ -2009,14 +2009,14 @@ void FMeshCache::Clear()
 }
 
 
-USkeleton* FSkeletonCache::Get(const TArray<uint16>& Key)
+USkeleton* FSkeletonCache::Get(const TArray<FName>& Key)
 {
 	const TWeakObjectPtr<USkeleton>* Result = MergedSkeletons.Find(Key);
 	return Result ? Result->Get() : nullptr;
 }
 
 
-void FSkeletonCache::Add(const TArray<uint16>& Key, USkeleton* Value)
+void FSkeletonCache::Add(const TArray<FName>& Key, USkeleton* Value)
 {
 	if (!Value)
 	{
@@ -2025,7 +2025,7 @@ void FSkeletonCache::Add(const TArray<uint16>& Key, USkeleton* Value)
 
 	MergedSkeletons.Add(Key, Value);
 
-	// Remove invalid SkeletalMeshes from the cache.
+	// Remove invalid Skeletons from the cache.
 	for (auto SkeletonIterator = MergedSkeletons.CreateIterator(); SkeletonIterator; ++SkeletonIterator)
 	{
 		if (SkeletonIterator.Value().IsStale())
@@ -3350,7 +3350,6 @@ void UCustomizableObjectPrivate::LogMemory()
 			if (ModelResources)
 			{
 				Arch << ModelResources->ReferenceSkeletalMeshesData;
-				Arch << ModelResources->Skeletons;
 				Arch << ModelResources->PassthroughObjects;
 				Arch << ModelResources->PhysicsAssets;
 				Arch << ModelResources->AnimBPs;
diff --git a/Engine/Plugins/Mutable/Source/CustomizableObject/Private/MuCO/CustomizableObjectInstance.cpp b/Engine/Plugins/Mutable/Source/CustomizableObject/Private/MuCO/CustomizableObjectInstance.cpp
--- a/Engine/Plugins/Mutable/Source/CustomizableObject/Private/MuCO/CustomizableObjectInstance.cpp
+++ b/Engine/Plugins/Mutable/Source/CustomizableObject/Private/MuCO/CustomizableObjectInstance.cpp
@@ -1075,42 +1075,50 @@ bool AreSkeletonsCompatible(const TArray<TObjectPtr<USkeleton>>& InSkeletons)
 #endif
 
 
-USkeleton* MergeSkeletons(UCustomizableObject& CustomizableObject, FCustomizableInstanceComponentData& ComponentData, bool& bOutCreatedNewSkeleton)
+USkeleton* MergeSkeletons(const TArray<TObjectPtr<USkeleton>>& SkeletonsToMerge, UCustomizableObject& CustomizableObject, bool& bOutCreatedNewSkeleton)
 {
 	MUTABLE_CPUPROFILER_SCOPE(BuildSkeletonData_MergeSkeletons);
 	bOutCreatedNewSkeleton = false;
 	
-	FReferencedSkeletons& ReferencedSkeletons = ComponentData.Skeletons;
-	
-	// Merged skeleton found in the cache
-	if (ReferencedSkeletons.Skeleton)
+	if (SkeletonsToMerge.Contains(nullptr))
 	{
-		USkeleton* MergedSkeleton = ReferencedSkeletons.Skeleton;
-		ReferencedSkeletons.Skeleton = nullptr;
-		return MergedSkeleton;
+		return nullptr;
 	}
 
-	check(!ReferencedSkeletons.SkeletonsToMerge.Contains(nullptr));
-
 	// No need to merge skeletons
-	if (ReferencedSkeletons.SkeletonsToMerge.Num() == 1)
+	if (SkeletonsToMerge.Num() == 1)
+	{
+		check(IsValid(SkeletonsToMerge[0]));
+		return SkeletonsToMerge[0];
+	}
+
+	TArray<FName> SkeletonIds;
+	SkeletonIds.Reserve(SkeletonsToMerge.Num());
+
+	for (const TObjectPtr<USkeleton>& Skeleton : SkeletonsToMerge)
+	{
+		SkeletonIds.Add(Skeleton->GetFName());
+	}
+
+	if (USkeleton* Skeleton = CustomizableObject.GetPrivate()->SkeletonCache.Get(SkeletonIds))
 	{
-		const TObjectPtr<USkeleton> RefSkeleton = ReferencedSkeletons.SkeletonsToMerge[0];
-		ReferencedSkeletons.SkeletonIds.Empty();
-		ReferencedSkeletons.SkeletonsToMerge.Empty();
-		return RefSkeleton;
+		// Merged skeleton found in the cache
+		if (IsValid(Skeleton))
+		{
+			return Skeleton;
+		}
 	}
 
 #if !UE_BUILD_SHIPPING
 	// Test Skeleton compatibility before attempting the merge to avoid a crash.
-	if (!AreSkeletonsCompatible(ReferencedSkeletons.SkeletonsToMerge))
+	if (!AreSkeletonsCompatible(SkeletonsToMerge))
 	{
 		return nullptr;
 	}
 #endif
 
 	FSkeletonMergeParams Params;
-	Params.SkeletonsToMerge = ReferencedSkeletons.SkeletonsToMerge;
+	Params.SkeletonsToMerge = SkeletonsToMerge;
 
 	USkeleton* FinalSkeleton = USkeletalMergingLibrary::MergeSkeletons(Params);
 	if (!FinalSkeleton)
@@ -1157,8 +1165,7 @@ USkeleton* MergeSkeletons(UCustomizableObject& CustomizableObject, FCustomizable
 #endif
 		
 		// Add Skeleton to the cache
-		CustomizableObject.GetPrivate()->SkeletonCache.Add(ReferencedSkeletons.SkeletonIds, FinalSkeleton);
-		ReferencedSkeletons.SkeletonIds.Empty();
+		CustomizableObject.GetPrivate()->SkeletonCache.Add(SkeletonIds, FinalSkeleton);
 
 		bOutCreatedNewSkeleton = true;
 	}
@@ -3661,12 +3668,43 @@ bool BuildSkeletonData(const TSharedRef<FUpdateContextPrivate>& Context, USkelet
 
 	UCustomizableObject* CustomizableObject = Instance->GetCustomizableObject();
 
-	UE::Mutable::Private::FComponentId ComponentId = Context->InstanceUpdateData.Components[ComponentIndex]->Id;
+	const TSharedPtr<FInstanceUpdateData::FComponent>& Component = Context->InstanceUpdateData.Components[ComponentIndex];
+	if (!Component)
+	{
+		return false;
+	}
+
+
+	const FName ComponentName = Context->ComponentNames[Component->Id];
+
+	// Other data
+	const int32 FirstLOD = Context->bStreamMeshLODs ?
+		Context->FirstLODAvailable[ComponentName] :
+		Context->GetFirstRequestedLOD()[ComponentName];
+
+	TArray<TObjectPtr<USkeleton>> SkeletonsToMerge;
+
+	// Add Reference Skeleton to the list of SkeletonsToMerge. TODO: Make this optional
+	const FMutableRefSkeletalMeshData& RefSkeletalMeshData = CustomizableObject->GetPrivate()->GetModelResourcesChecked().ReferenceSkeletalMeshesData[Component->Id];
+	SkeletonsToMerge.Add(RefSkeletalMeshData.Skeleton);
+
+	for (int32 LODIndex = FirstLOD; LODIndex < Component->LODCount; ++LODIndex)
+	{
+		const TSharedRef<FInstanceUpdateData::FLOD>& LOD = Context->InstanceUpdateData.LODs[Component->FirstLOD + LODIndex];
+		check(LOD->Mesh)
+				
+		// Add Skeletons to merge
+		for (const UE::Mutable::Private::TPassthroughObjectPtr<USkeleton>& SkeletonObject : LOD->Mesh->SkeletonObjects)
+		{
+			if (ensure(SkeletonObject.IsResolved()))
+			{
+				SkeletonsToMerge.AddUnique(SkeletonObject.Get());
+			}
+		}
+	}
 
-	FCustomizableInstanceComponentData& ComponentData = InstancePrivate->GetComponentData(ComponentId);
-	
 	bool bCreatedNewSkeleton = false;
-	const TObjectPtr<USkeleton> Skeleton = MergeSkeletons(*CustomizableObject, ComponentData, bCreatedNewSkeleton);
+	const TObjectPtr<USkeleton> Skeleton = MergeSkeletons(SkeletonsToMerge, *CustomizableObject, bCreatedNewSkeleton);
 	if (!Skeleton)
 	{
 		return false;
@@ -3684,7 +3722,7 @@ bool BuildSkeletonData(const TSharedRef<FUpdateContextPrivate>& Context, USkelet
 	{
 		MUTABLE_CPUPROFILER_SCOPE(BuildSkeletonData_ApplyPose);
 		
-		const TArray<FInstanceUpdateData::FBone>& BonePose = Context->InstanceUpdateData.SkeletonsPerInstanceComponent[ComponentIndex].BonePose;
+		const TArray<FInstanceUpdateData::FBone>& BonePose = Component->BonePose;
 		
 		TArray<FMatrix44f>& RefBasesInvMatrix = SkeletalMesh.GetRefBasesInvMatrix();
 		RefBasesInvMatrix.Empty(RawRefBoneCount);
@@ -3736,13 +3774,11 @@ bool BuildSkeletonData(const TSharedRef<FUpdateContextPrivate>& Context, USkelet
 		// If the skeleton is new, it means it has just been merged and the retargeting modes need merging too as the
 		// MergeSkeletons function doesn't do it. Only do it for newly generated ones, not for cached or non-transient ones.
 		GeneratedSkeleton->RecreateBoneTree(&SkeletalMesh);
-		
-		TArray<TObjectPtr<USkeleton>>& SkeletonsToMerge = ComponentData.Skeletons.SkeletonsToMerge;
-		check(SkeletonsToMerge.Num() > 1);
 
 		TMap<FName, EBoneTranslationRetargetingMode::Type> BoneNamesToRetargetingMode;
 
 		const int32 NumberOfSkeletons = SkeletonsToMerge.Num();
+		check(NumberOfSkeletons > 1);
 
 		for (int32 SkeletonIndex = 0; SkeletonIndex < NumberOfSkeletons; ++SkeletonIndex)
 		{
@@ -4401,64 +4437,17 @@ UE::Tasks::FTask LoadAdditionalAssetsAndData(const TSharedRef<FUpdateContextPriv
 
 	TArray<FSoftObjectPath> AssetsToStream;
 
-	TArray<TSharedRef<FInstanceUpdateData::FComponent>>& Components = Context->InstanceUpdateData.Components;
 	
 	InstancePrivate->GatheredAnimBPs.Empty();
 	InstancePrivate->AnimBPGameplayTags.Reset();
 	InstancePrivate->AnimBpPhysicsAssets.Reset();
 
-	// Load Skeletons required by the SubMeshes of the newly generated Mesh, will be merged later
-	for (int32 ComponentIndex = 0; ComponentIndex < Context->InstanceUpdateData.Components.Num(); ++ComponentIndex)
-	{
-		const TSharedPtr<FInstanceUpdateData::FComponent>& Component = Context->InstanceUpdateData.Components[ComponentIndex];
-		if (Component->Id == INDEX_NONE)
-		{
-			continue;
-		}
-
-		const FInstanceUpdateData::FSkeletonData& SkeletonData = Context->InstanceUpdateData.SkeletonsPerInstanceComponent[ComponentIndex];
-		
-		FCustomizableInstanceComponentData& ComponentData = InstancePrivate->GetComponentData(Component->Id);
-
-		// Reuse merged Skeleton if cached
-		ComponentData.Skeletons.Skeleton = CustomizableObject->GetPrivate()->SkeletonCache.Get(SkeletonData.SkeletonIds);
-		if (ComponentData.Skeletons.Skeleton)
-		{
-			ComponentData.Skeletons.SkeletonIds.Empty();
-			ComponentData.Skeletons.SkeletonsToMerge.Empty();
-			continue;
-		}
-
-		// Add Skeletons to merge
-		ComponentData.Skeletons.SkeletonsToMerge = SkeletonData.Skeletons;
-
-		for (const uint32 SkeletonId : SkeletonData.SkeletonIds)
-		{
-			TSoftObjectPtr<USkeleton> AssetPtr = ModelResources.Skeletons.IsValidIndex(SkeletonId) ? ModelResources.Skeletons[SkeletonId] : nullptr;
-			if (AssetPtr.IsNull())
-			{
-				continue;
-			}
-
-			// Add referenced skeletons to the assets to stream
-			ComponentData.Skeletons.SkeletonIds.Add(SkeletonId);
-
-			if (USkeleton* Skeleton = AssetPtr.Get())
-			{
-				ComponentData.Skeletons.SkeletonsToMerge.Add(Skeleton);
-			}
-			else
-			{
-				AssetsToStream.Add(AssetPtr.ToSoftObjectPath());
-			}
-		}
-	}
-	
 	bool bHasInvalidMesh = false;
 	const bool bUpdateMeshes = DoComponentsNeedUpdate(Context, bHasInvalidMesh);
 
 	// Load assets coming from SubMeshes of the newly generated Mesh
-	for (int32 ComponentIndex = 0; ComponentIndex < Context->InstanceUpdateData.Components.Num(); ++ComponentIndex)
+	TArray<TSharedRef<FInstanceUpdateData::FComponent>>& Components = Context->InstanceUpdateData.Components;
+	for (int32 ComponentIndex = 0; ComponentIndex < Components.Num(); ++ComponentIndex)
 	{
 		const TSharedRef<const FInstanceUpdateData::FComponent>& Component = Components[ComponentIndex];
 		FCustomizableInstanceComponentData& ComponentData = InstancePrivate->GetComponentData(Component->Id);
@@ -4806,24 +4795,6 @@ void AdditionalAssetsAsyncLoaded(const TSharedRef<FUpdateContextPrivate> Context
 #endif
 		}
 
-		// Loaded Skeletons
-		FReferencedSkeletons& Skeletons = ComponentData.Skeletons;
-		for (int32 SkeletonIndex : Skeletons.SkeletonIds)
-		{
-			check(ModelResources.Skeletons.IsValidIndex(SkeletonIndex));
-			const TSoftObjectPtr<USkeleton>& AssetPtr = ModelResources.Skeletons[SkeletonIndex];
-			
-			if (USkeleton* Skeleton = AssetPtr.Get())
-			{
-				Skeletons.SkeletonsToMerge.AddUnique(Skeleton);
-			}
-			else
-			{
-				UE_LOG(LogMutable, Error, TEXT("Failed to load Skeleton asset: %s"), *AssetPtr.ToSoftObjectPath().ToString());
-				Context->bIsUpdateAborted = true;
-			}
-		}
-
 		// Loaded anim BPs
 		for (TPair<FName, TSoftClassPtr<UAnimInstance>>& SlotAnimBP : ComponentData.AnimSlotToBP)
 		{
diff --git a/Engine/Plugins/Mutable/Source/CustomizableObject/Private/MuCO/CustomizableObjectSystem.cpp b/Engine/Plugins/Mutable/Source/CustomizableObject/Private/MuCO/CustomizableObjectSystem.cpp
--- a/Engine/Plugins/Mutable/Source/CustomizableObject/Private/MuCO/CustomizableObjectSystem.cpp
+++ b/Engine/Plugins/Mutable/Source/CustomizableObject/Private/MuCO/CustomizableObjectSystem.cpp
@@ -29,6 +29,7 @@
 #include "MuCO/CustomizableObjectSystemPrivate.h"
 #include "MuR/Parameters.h"
 #include "MuR/System.h"
+#include "MuR/PassthroughObject.h"
 
 #if WITH_EDITOR
 #include "Editor.h"
@@ -1880,13 +1881,8 @@ namespace impl
 	{
 		MUTABLE_CPUPROFILER_SCOPE(Subtask_Mutable_PrepareSkeletonData);
 
-		const int32 NumComponents = OperationData->InstanceUpdateData.Components.Num();
-		OperationData->InstanceUpdateData.SkeletonsPerInstanceComponent.SetNum(NumComponents);
-
-		for (int32 ComponentIndex = 0; ComponentIndex< NumComponents; ++ComponentIndex)
+		for (const TSharedRef<FInstanceUpdateData::FComponent>& Component : OperationData->InstanceUpdateData.Components)
 		{
-			const TSharedRef<FInstanceUpdateData::FComponent>& Component = OperationData->InstanceUpdateData.Components[ComponentIndex];
-
 			for (int32 LODIndex = 0; LODIndex < Component->LODCount; ++LODIndex)
 			{
 				TSharedRef<FInstanceUpdateData::FLOD>& LOD = OperationData->InstanceUpdateData.LODs[Component->FirstLOD + LODIndex];
@@ -1902,25 +1898,6 @@ namespace impl
 					continue;
 				}
 
-				FInstanceUpdateData::FSkeletonData& SkeletonData = OperationData->InstanceUpdateData.SkeletonsPerInstanceComponent[ComponentIndex];
-
-				// Add SkeletonIds 
-				const int32 SkeletonIDsCount = Mesh->GetSkeletonIDsCount();
-				for (int32 SkeletonIndex = 0; SkeletonIndex < SkeletonIDsCount; ++SkeletonIndex)
-				{
-					SkeletonData.SkeletonIds.AddUnique(Mesh->GetSkeletonID(SkeletonIndex));
-				}
-
-				// Add Skeletons from Mesh Parameters
-				for (const TStrongObjectPtr<USkeletalMesh>& StrongMesh : Mesh->SkeletalMeshes)
-				{
-					USkeletalMesh* SkeletalMesh = StrongMesh.Get();
-					if (ensure(SkeletalMesh))
-					{
-						SkeletonData.Skeletons.AddUnique(SkeletalMesh->GetSkeleton());
-					}
-				}
-
 				const TArray<FName>& BoneNames = Mesh->GetSkeleton()->BoneNames;
 
 				// Append BoneMap to the array of BoneMaps
@@ -1942,11 +1919,11 @@ namespace impl
 
 					OperationData->InstanceUpdateData.ActiveBones.Add(BoneName);
 
-					if (!SkeletonData.BonePose.FindByKey(BoneName))
+					if (!Component->BonePose.FindByKey(BoneName))
 					{
 						FTransform3f Transform;
 						Mesh->GetBonePoseTransform(BoneIndex, Transform);
-						SkeletonData.BonePose.Add({ BoneName,Transform.Inverse().ToMatrixWithScale() });
+						Component->BonePose.Add({ BoneName,Transform.Inverse().ToMatrixWithScale() });
 					}
 				}
 			}
diff --git a/Engine/Plugins/Mutable/Source/CustomizableObject/Private/MuCO/UnrealConversionUtils.cpp b/Engine/Plugins/Mutable/Source/CustomizableObject/Private/MuCO/UnrealConversionUtils.cpp
--- a/Engine/Plugins/Mutable/Source/CustomizableObject/Private/MuCO/UnrealConversionUtils.cpp
+++ b/Engine/Plugins/Mutable/Source/CustomizableObject/Private/MuCO/UnrealConversionUtils.cpp
@@ -1196,6 +1196,8 @@ namespace UnrealConversionUtils
 						return;
 					}
 
+					Result->SkeletonObjects.Add(UE::Mutable::Private::TPassthroughObjectPtr(Skeleton));
+
 					// Create the skeleton, poses, and BoneMap for this mesh
 					UE::Mutable::Private::TManagedPtr<UE::Mutable::Private::FSkeleton> MutableSkeleton = 
 							UE::Mutable::Private::MakeManaged<UE::Mutable::Private::FSkeleton>();
diff --git a/Engine/Plugins/Mutable/Source/CustomizableObjectEditor/Private/MuCOE/CustomizableObjectCompiler.cpp b/Engine/Plugins/Mutable/Source/CustomizableObjectEditor/Private/MuCOE/CustomizableObjectCompiler.cpp
--- a/Engine/Plugins/Mutable/Source/CustomizableObjectEditor/Private/MuCOE/CustomizableObjectCompiler.cpp
+++ b/Engine/Plugins/Mutable/Source/CustomizableObjectEditor/Private/MuCOE/CustomizableObjectCompiler.cpp
@@ -1668,8 +1668,6 @@ void FCustomizableObjectCompiler::FinishCompilationTask()
 			ModelResources->ExternalOperationVersions.Add(ExternalOperation.Value.GetScriptStruct(), Version);
 		}
 		
-		ModelResources->Skeletons = MoveTemp(CompilationContext->ReferencedSkeletons);
-
 		ModelResources->PhysicsAssets = MoveTemp(CompilationContext->PhysicsAssets);
 		ModelResources->AnimBpOverridePhysiscAssetsInfo = MoveTemp(CompilationContext->AnimBpOverridePhysicsAssetsInfo);
 		
diff --git a/Engine/Plugins/Mutable/Source/CustomizableObjectEditor/Private/MuCOE/GenerateMutableSource/GenerateMutableSourceMesh.cpp b/Engine/Plugins/Mutable/Source/CustomizableObjectEditor/Private/MuCOE/GenerateMutableSource/GenerateMutableSourceMesh.cpp
--- a/Engine/Plugins/Mutable/Source/CustomizableObjectEditor/Private/MuCOE/GenerateMutableSource/GenerateMutableSourceMesh.cpp
+++ b/Engine/Plugins/Mutable/Source/CustomizableObjectEditor/Private/MuCOE/GenerateMutableSource/GenerateMutableSourceMesh.cpp
@@ -1147,13 +1147,8 @@ UE::Tasks::FTask ConvertSkeletalMeshToMutable(TSharedRef<FMeshConversionContext>
 				return UE::Tasks::MakeCompletedTask<void>();
 			}
 
-			// Add the RefSkeleton ID to the mesh.
-			const int32 RefSkeletonID = Context.ReferencedSkeletons.AddUnique(ComponentRefSkeleton);
-			MutableMesh->AddSkeletonID(RefSkeletonID);
-
-			// Add the skeleton to the list of referenced skeletons and add its index to the mesh
-			const int32 SkeletonID = Context.ReferencedSkeletons.AddUnique(const_cast<USkeleton*>(Skeleton));
-			MutableMesh->AddSkeletonID(SkeletonID);
+			// Add the skeleton object
+			MutableMesh->SkeletonObjects.Emplace(Context.PassthroughObjectFactory.Add(*const_cast<USkeleton*>(Skeleton)));
 		}
 
 		// RefSkeleton check
diff --git a/Engine/Plugins/Mutable/Source/CustomizableObjectEditor/Private/MuCOE/SMutableCodeViewer.cpp b/Engine/Plugins/Mutable/Source/CustomizableObjectEditor/Private/MuCOE/SMutableCodeViewer.cpp
--- a/Engine/Plugins/Mutable/Source/CustomizableObjectEditor/Private/MuCOE/SMutableCodeViewer.cpp
+++ b/Engine/Plugins/Mutable/Source/CustomizableObjectEditor/Private/MuCOE/SMutableCodeViewer.cpp
@@ -40,6 +40,7 @@
 #include "Engine/SkeletalMesh.h"
 #include "Internationalization/Regex.h"
 #include "MuCO/CustomizableObjectInstancePrivate.h"
+#include "MuR/PassthroughObject.h"
 
 class FExtender;
 class FReferenceCollector;
@@ -1581,13 +1582,7 @@ void SMutableCodeViewer::GenerateElementRecursive(const int32& InStateIndex, UE:
 					: 0;
 		}
 
-		int32 SkeletonCount = EnumHasAnyFlags(Args.Flags, OpEnumFlags::HasSkeleton) ? 1 : 0;
-		if (EnumHasAnyFlags(Args.Flags, OpEnumFlags::IsSkeletonList) &&
-			ensure(InProgram.ConstantUInt32Lists.IsValidIndex(Args.SkeletonIds.ListAddress)))
-		{
-			SkeletonCount = InProgram.ConstantUInt32Lists[Args.SkeletonIds.ListAddress].Num();
-		}
-
+		const int32 SkeletonCount = Args.SkeletonId != UE::Mutable::Private::PASSTHROUGH_ID_INVALID ? 1 : 0;
 		const int32 PhysicsAssetCount = Args.PhysicsAssetId != UE::Mutable::Private::PASSTHROUGH_ID_INVALID ? 1 : 0;
 
 		int32 AdditionalPhysicsAssetCount = EnumHasAnyFlags(Args.Flags, OpEnumFlags::HasAdditionalPhysicsAsset) ? 1 : 0;
diff --git a/Engine/Plugins/Mutable/Source/MutableRuntime/Internal/MuR/Mesh.h b/Engine/Plugins/Mutable/Source/MutableRuntime/Internal/MuR/Mesh.h
--- a/Engine/Plugins/Mutable/Source/MutableRuntime/Internal/MuR/Mesh.h
+++ b/Engine/Plugins/Mutable/Source/MutableRuntime/Internal/MuR/Mesh.h
@@ -30,6 +30,7 @@
 class FString;
 class USkeletalMesh;
 class UPhysicsAsset;
+class USkeleton;
 
 namespace UE::Mutable::Private
 {
@@ -408,17 +409,6 @@ namespace UE::Mutable::Private
 		//! Return an array containing the bonemaps of all surfaces in the mesh.
 		UE_API const TArray<FBoneIdOrIndex>& GetBoneMap() const;
 
-		//!
-		UE_API int32 GetSkeletonIDsCount() const;
-
-		//!
-		UE_API int32 GetSkeletonID(int32 SkeletonIndex) const;
-
-		//!
-		UE_API void AddSkeletonID(int32 SkeletonID);
-
-        //! \}
-
         /** 
 		 * Get an internal identifier used to reference this mesh in operations like deferred
          * mesh building, or instance updating.
@@ -472,8 +462,8 @@ namespace UE::Mutable::Private
     	/** Only valid if the mesh is a LOD (has multiple sections). */
     	TArray<FCloth> ClothSections;
     	
-		/** Externally provided SkeletonIDs of the skeletons required by this mesh. */
-		TArray<uint32> SkeletonIDs;
+		/** Externally provided skeletons required by this mesh. */
+		TArray<TPassthroughObjectPtr<USkeleton>> SkeletonObjects;
 
 		/** 
 		 * This skeleton and physicsbody are not owned and may be used by other meshes, so it cannot be modified
@@ -579,7 +569,7 @@ namespace UE::Mutable::Private
 			if (bEqual) bEqual = (StreamedResources == Other.StreamedResources);
 			if (bEqual) bEqual = (Surfaces == Other.Surfaces);
 			if (bEqual) bEqual = (Tags == Other.Tags);
-			if (bEqual) bEqual = (SkeletonIDs == Other.SkeletonIDs);
+			if (bEqual) bEqual = (SkeletonObjects == Other.SkeletonObjects);
 
 			for (int32 i = 0; bEqual && i < Layouts.Num(); ++i)
 			{
diff --git a/Engine/Plugins/Mutable/Source/MutableRuntime/Internal/MuR/Model.h b/Engine/Plugins/Mutable/Source/MutableRuntime/Internal/MuR/Model.h
--- a/Engine/Plugins/Mutable/Source/MutableRuntime/Internal/MuR/Model.h
+++ b/Engine/Plugins/Mutable/Source/MutableRuntime/Internal/MuR/Model.h
@@ -768,7 +768,7 @@ namespace UE::Mutable::Private
 					}
 
 					Result->Tags = MetaDataMesh->Tags; 
-					Result->SkeletonIDs = MetaDataMesh->SkeletonIDs; 
+					Result->SkeletonObjects = MetaDataMesh->SkeletonObjects;
 					Result->StreamedResources = MetaDataMesh->StreamedResources; 
 				}
 			}
diff --git a/Engine/Plugins/Mutable/Source/MutableRuntime/Internal/MuR/OpMeshMerge.h b/Engine/Plugins/Mutable/Source/MutableRuntime/Internal/MuR/OpMeshMerge.h
--- a/Engine/Plugins/Mutable/Source/MutableRuntime/Internal/MuR/OpMeshMerge.h
+++ b/Engine/Plugins/Mutable/Source/MutableRuntime/Internal/MuR/OpMeshMerge.h
@@ -696,12 +696,12 @@ namespace UE::Mutable::Private
 		// Skeleton
 		//---------------------------
 
-		// Add SkeletonIDs
-		Result->SkeletonIDs = pFirst->SkeletonIDs;
+		// Add SkeletonObjects
+		Result->SkeletonObjects = pFirst->SkeletonObjects;
 
-		for (const int32 SkeletonID : pSecond->SkeletonIDs)
+		for (const TPassthroughObjectPtr<USkeleton>& SkeletonObject : pSecond->SkeletonObjects)
 		{
-			Result->SkeletonIDs.AddUnique(SkeletonID);
+			Result->SkeletonObjects.AddUnique(SkeletonObject);
 		}
 
 		TManagedPtr<const FSkeleton> FirstSkeleton = pFirst->GetSkeleton();
diff --git a/Engine/Plugins/Mutable/Source/MutableRuntime/Internal/MuR/Operations.h b/Engine/Plugins/Mutable/Source/MutableRuntime/Internal/MuR/Operations.h
--- a/Engine/Plugins/Mutable/Source/MutableRuntime/Internal/MuR/Operations.h
+++ b/Engine/Plugins/Mutable/Source/MutableRuntime/Internal/MuR/Operations.h
@@ -1097,8 +1097,6 @@ namespace UE::Mutable::Private
 				IsTagList		= 1 << 1,
                 HasResources	= 1 << 2,
                 IsResourceList	= 1 << 3,
-                HasSkeleton		= 1 << 4, // remove
-                IsSkeletonList	= 1 << 5, // remove
 				HasAdditionalPhysicsAsset = 1 << 6,
                 IsAdditionalPhysicsAssetList = 1 << 7
             };
@@ -1109,12 +1107,6 @@ namespace UE::Mutable::Private
 				LIST_UINT32_ADDRESS ListAddress;
             };
 
-            union SkeletonIdOrUInt32ListAddress
-            {
-                uint32 SkeletonId;
-				LIST_UINT32_ADDRESS ListAddress;
-            };
-
             union ResourceIdOrUInt64ListAddress
             {
                 uint64 ResourceId;
@@ -1130,8 +1122,8 @@ namespace UE::Mutable::Private
             EnumFlags Flags; 
 			int16 BonePosePriority;
             ADDRESS Source;
+            PASSTHROUGH_ID SkeletonId;
             TagAddressOrUInt32ListAddress Tags;
-            SkeletonIdOrUInt32ListAddress SkeletonIds;
             ResourceIdOrUInt64ListAddress ResourceIds;
 			PASSTHROUGH_ID PhysicsAssetId; // Main Physics Asset of the mesh
 			PassthroughIDOrUInt32ListAddress AdditionalPhysicsAssetsIds; // Physics Assets used by AnimBPs
diff --git a/Engine/Plugins/Mutable/Source/MutableRuntime/Private/MuR/CodeRunner.cpp b/Engine/Plugins/Mutable/Source/MutableRuntime/Private/MuR/CodeRunner.cpp
--- a/Engine/Plugins/Mutable/Source/MutableRuntime/Private/MuR/CodeRunner.cpp
+++ b/Engine/Plugins/Mutable/Source/MutableRuntime/Private/MuR/CodeRunner.cpp
@@ -2895,6 +2895,11 @@ namespace UE::Mutable::Private
 				{
 					Result = CloneOrTakeOver(MoveTemp(SourceMesh));
 
+					if (OpArgs.SkeletonId != PASSTHROUGH_ID_INVALID)
+					{
+						Result->SkeletonObjects.Add(FindPassthrough<USkeleton>(OpArgs.SkeletonId));
+					}
+
 					using OpEnumFlags = OP::MeshAddMetadataArgs::EnumFlags;
 
 					// Append Tags
@@ -2937,21 +2942,7 @@ namespace UE::Mutable::Private
 						check(!EnumHasAnyFlags(OpArgs.Flags, OpEnumFlags::IsResourceList));
 					}
 					
-					// Append Skeletons 
-					if (EnumHasAnyFlags(OpArgs.Flags, OpEnumFlags::IsSkeletonList) &&
-						Program.ConstantUInt32Lists.IsValidIndex(OpArgs.SkeletonIds.ListAddress))
-					{
-						Result->SkeletonIDs.Append(Program.ConstantUInt32Lists[OpArgs.SkeletonIds.ListAddress]);
-					}
-					else if(EnumHasAnyFlags(OpArgs.Flags, OpEnumFlags::HasSkeleton))
-					{
-						Result->SkeletonIDs.Add(OpArgs.SkeletonIds.SkeletonId);
-					}
-					else
-					{
-						check(!EnumHasAnyFlags(OpArgs.Flags, OpEnumFlags::IsSkeletonList));
-					}
-
+					
 					if (Result->BonePosePriority != OpArgs.BonePosePriority)
 					{
 						Result->BonePosePriority = OpArgs.BonePosePriority;
diff --git a/Engine/Plugins/Mutable/Source/MutableRuntime/Private/MuR/CodeRunnerBegin.cpp b/Engine/Plugins/Mutable/Source/MutableRuntime/Private/MuR/CodeRunnerBegin.cpp
--- a/Engine/Plugins/Mutable/Source/MutableRuntime/Private/MuR/CodeRunnerBegin.cpp
+++ b/Engine/Plugins/Mutable/Source/MutableRuntime/Private/MuR/CodeRunnerBegin.cpp
@@ -2233,6 +2233,10 @@ namespace UE::Mutable::Private
 					const OP::MeshAddMetadataArgs OpArgs = Program.GetOpArgs<OP::MeshAddMetadataArgs>(Item.At);
 					PushOp(FScheduledOpInline(OpArgs.Source, Item));
 
+if (OpArgs.SkeletonId != PASSTHROUGH_ID_INVALID)
+					{
+						PassthroughObjectLoader->Add<USkeleton>(OpArgs.SkeletonId);
+					}
 					using OpEnumFlags = OP::MeshAddMetadataArgs::EnumFlags;
 
 					if (OpArgs.PhysicsAssetId != PASSTHROUGH_ID_INVALID)
@@ -2262,6 +2266,8 @@ namespace UE::Mutable::Private
 						check(!EnumHasAnyFlags(OpArgs.Flags, OpEnumFlags::IsAdditionalPhysicsAssetList));
 					}
 
+					
+
 					StoreNone(Item);
 					break;
 				}
diff --git a/Engine/Plugins/Mutable/Source/MutableRuntime/Private/MuR/Mesh.cpp b/Engine/Plugins/Mutable/Source/MutableRuntime/Private/MuR/Mesh.cpp
--- a/Engine/Plugins/Mutable/Source/MutableRuntime/Private/MuR/Mesh.cpp
+++ b/Engine/Plugins/Mutable/Source/MutableRuntime/Private/MuR/Mesh.cpp
@@ -171,7 +171,7 @@ void FMesh::CopyFrom(const FMesh& From, EMeshCopyFlags InFlags)
 	// Copy SkeletonIDs
 	if (EnumHasAnyFlags(InFlags, EMeshCopyFlags::WithSkeletonIDs))
 	{
-		SkeletonIDs = From.SkeletonIDs;
+		SkeletonObjects = From.SkeletonObjects;
 	}
 
 	if (EnumHasAnyFlags(InFlags, EMeshCopyFlags::WithAdditionalPhysics))
@@ -665,25 +665,6 @@ const TArray<FBoneIdOrIndex>& FMesh::GetBoneMap() const
 }
 
 
-int32 FMesh::GetSkeletonIDsCount() const
-{
-    return SkeletonIDs.Num();
-}
-
-
-int32 FMesh::GetSkeletonID(int32 SkeletonIndex) const
-{
-	return SkeletonIDs.IsValidIndex(SkeletonIndex) ? SkeletonIDs[SkeletonIndex] : INDEX_NONE;
-}
-
-
-void FMesh::AddSkeletonID(int32 SkeletonID)
-{
-	check(SkeletonID != INDEX_NONE);
-	SkeletonIDs.AddUnique(SkeletonID);
-}
-
-
 int32 FMesh::GetDataSize() const
 {
 	// TODO: review if other mesh fields like additional physics assets
@@ -973,17 +954,11 @@ void FMesh::Serialise(FOutputArchive& Arch) const
 	Arch << AdditionalBuffers;
 	Arch << Layouts;
 
-	Arch << SkeletonIDs;
-
-	//Arch << Skeleton;
 	Arch << PhysicsBody;
 
 	Arch << uint32(Flags);
 	Arch << Surfaces;
 
-	Arch << Tags;
-	Arch << StreamedResources;
-
 	Arch << BonePoses;
 	Arch << BoneMap;
 
@@ -1002,9 +977,6 @@ void FMesh::Unserialise(FInputArchive& Arch)
 	Arch >> AdditionalBuffers;
 	Arch >> Layouts;
 
-	Arch >> SkeletonIDs;
-
-	//Arch >> Skeleton;
 	Arch >> PhysicsBody;
 
 	uint32 Temp;
@@ -1013,9 +985,6 @@ void FMesh::Unserialise(FInputArchive& Arch)
 
 	Arch >> Surfaces;
 
-	Arch >> Tags;
-	Arch >> StreamedResources;
-
 	Arch >> BonePoses;
 	Arch >> BoneMap;
 
@@ -1039,6 +1008,8 @@ bool FMesh::IsSimilar(const FMesh& Other) const
 	bool bEqual = IndexBuffers == Other.IndexBuffers;
 	bEqual = bEqual && (PassthroughObject == Other.PassthroughObject);
 
+	bEqual = bEqual && (SkeletonObjects == Other.SkeletonObjects);
+
 	if (bEqual && Skeleton != Other.Skeleton)
 	{
 		if (Skeleton && Other.Skeleton)
diff --git a/Engine/Plugins/Mutable/Source/MutableRuntime/Private/MuR/OpMeshFormat.cpp b/Engine/Plugins/Mutable/Source/MutableRuntime/Private/MuR/OpMeshFormat.cpp
--- a/Engine/Plugins/Mutable/Source/MutableRuntime/Private/MuR/OpMeshFormat.cpp
+++ b/Engine/Plugins/Mutable/Source/MutableRuntime/Private/MuR/OpMeshFormat.cpp
@@ -715,7 +715,7 @@ namespace UE::Mutable::Private
 		Result->BonePosePriority = Source->BonePosePriority;
 		Result->BoneMap = Source->BoneMap;
 
-		Result->SkeletonIDs = Source->SkeletonIDs;
+		Result->SkeletonObjects = Source->SkeletonObjects;
 
 		// A shallow copy is done here, it should not be a problem.
 		Result->AdditionalPhysicsBodies = Source->AdditionalPhysicsBodies;
diff --git a/Engine/Plugins/Mutable/Source/MutableTools/Private/MuT/ASTOpMeshAddMetadata.cpp b/Engine/Plugins/Mutable/Source/MutableTools/Private/MuT/ASTOpMeshAddMetadata.cpp
--- a/Engine/Plugins/Mutable/Source/MutableTools/Private/MuT/ASTOpMeshAddMetadata.cpp
+++ b/Engine/Plugins/Mutable/Source/MutableTools/Private/MuT/ASTOpMeshAddMetadata.cpp
@@ -34,7 +34,7 @@ namespace UE::Mutable::Private
 				Source == Other->Source &&
 				Tags == Other->Tags &&
 				ResourceIds == Other->ResourceIds &&
-				SkeletonIds == Other->SkeletonIds &&
+				SkeletonId == Other->SkeletonId &&
 				BonePosePriority == Other->BonePosePriority &&
 				PhysicsAssetId == Other->PhysicsAssetId &&
 				AdditionalPhysicsAssetIds == Other->AdditionalPhysicsAssetIds;
@@ -48,9 +48,9 @@ namespace UE::Mutable::Private
 	{
 		Ptr<ASTOpMeshAddMetadata> NewOp = new ASTOpMeshAddMetadata();
 		NewOp->Source = MapChild(Source.child());
+		NewOp->SkeletonId = SkeletonId;
 		NewOp->Tags = Tags;
 		NewOp->ResourceIds = ResourceIds;
-		NewOp->SkeletonIds = SkeletonIds;
 		NewOp->BonePosePriority = BonePosePriority;
 		NewOp->PhysicsAssetId = PhysicsAssetId;
 		NewOp->AdditionalPhysicsAssetIds = AdditionalPhysicsAssetIds;
@@ -120,20 +120,7 @@ namespace UE::Mutable::Private
 				}
 			}
 
-			if (SkeletonIds.Num())
-			{
-				EnumAddFlags(Args.Flags, OpEnumFlags::HasSkeleton);
-				if (SkeletonIds.Num() > 1)
-				{
-					EnumAddFlags(Args.Flags, OpEnumFlags::IsSkeletonList);
-					Args.SkeletonIds.ListAddress = Program.AddConstant(SkeletonIds);
-				}
-				else
-				{
-					Args.SkeletonIds.SkeletonId = SkeletonIds[0];
-				}
-			}
-
+			Args.SkeletonId = SkeletonId;
 			Args.BonePosePriority = BonePosePriority;
 			
 			Args.PhysicsAssetId = PhysicsAssetId;
diff --git a/Engine/Plugins/Mutable/Source/MutableTools/Private/MuT/ASTOpMeshAddMetadata.h b/Engine/Plugins/Mutable/Source/MutableTools/Private/MuT/ASTOpMeshAddMetadata.h
--- a/Engine/Plugins/Mutable/Source/MutableTools/Private/MuT/ASTOpMeshAddMetadata.h
+++ b/Engine/Plugins/Mutable/Source/MutableTools/Private/MuT/ASTOpMeshAddMetadata.h
@@ -8,6 +8,7 @@
 #include "MuT/AST.h"
 
 
+
 namespace UE::Mutable::Private
 {
 	struct FProgram;
@@ -26,8 +27,8 @@ namespace UE::Mutable::Private
 		/** ResourceIds to add. */
 		TArray<uint64> ResourceIds;
 
-		/** SkeletonIds to add. */
-		TArray<uint32> SkeletonIds;
+		/** SkeletonId to add. */
+		PASSTHROUGH_ID SkeletonId = PASSTHROUGH_ID_INVALID;
 
 		/** Bone Pose priority when merging two meshes. */
 		int16 BonePosePriority = 0;
diff --git a/Engine/Plugins/Mutable/Source/MutableTools/Private/MuT/ASTOpSwitch.cpp b/Engine/Plugins/Mutable/Source/MutableTools/Private/MuT/ASTOpSwitch.cpp
--- a/Engine/Plugins/Mutable/Source/MutableTools/Private/MuT/ASTOpSwitch.cpp
+++ b/Engine/Plugins/Mutable/Source/MutableTools/Private/MuT/ASTOpSwitch.cpp
@@ -715,7 +715,7 @@ namespace UE::Mutable::Private
 				bool bAllMetadataIsTheSame = true;
 				TArray<FString> Tags;
 				TArray<uint64> ResourceIds;
-				TArray<uint32> SkeletonIds;
+				PASSTHROUGH_ID SkeletonId = PASSTHROUGH_ID_INVALID;
 				int16 BonePosePriority = 0;
 				PASSTHROUGH_ID PhysicsAssetId = PASSTHROUGH_ID_INVALID;
 				TArray<uint32> AdditionalPhysicsAssetIds;
@@ -726,7 +726,7 @@ namespace UE::Mutable::Private
 					const ASTOpMeshAddMetadata* Typed = static_cast<const ASTOpMeshAddMetadata*>(Default.child().get());
 					Tags = Typed->Tags;
 					ResourceIds = Typed->ResourceIds;
-					SkeletonIds = Typed->SkeletonIds;
+					SkeletonId = Typed->SkeletonId;
 					BonePosePriority = Typed->BonePosePriority;
 					PhysicsAssetId = Typed->PhysicsAssetId;
 					AdditionalPhysicsAssetIds = Typed->AdditionalPhysicsAssetIds;
@@ -741,11 +741,11 @@ namespace UE::Mutable::Private
 
 					check(Case.Branch->GetOpType() == EOpType::ME_ADDMETADATA);
 					const ASTOpMeshAddMetadata* Typed = static_cast<const ASTOpMeshAddMetadata*>(Case.Branch.child().get());
-					if (Tags.IsEmpty() && ResourceIds.IsEmpty() && SkeletonIds.IsEmpty() && PhysicsAssetId == PASSTHROUGH_ID_INVALID && AdditionalPhysicsAssetIds.IsEmpty())
+					if (Tags.IsEmpty() && ResourceIds.IsEmpty() &&  SkeletonId == PASSTHROUGH_ID_INVALID && PhysicsAssetId == PASSTHROUGH_ID_INVALID && AdditionalPhysicsAssetIds.IsEmpty())
 					{
 						Tags = Typed->Tags;
 						ResourceIds = Typed->ResourceIds;
-						SkeletonIds = Typed->SkeletonIds;
+						SkeletonId = Typed->SkeletonId;
 						BonePosePriority = Typed->BonePosePriority;
 						PhysicsAssetId = Typed->PhysicsAssetId;
 						AdditionalPhysicsAssetIds = Typed->AdditionalPhysicsAssetIds;
@@ -753,7 +753,7 @@ namespace UE::Mutable::Private
 					else if (
 							Typed->Tags != Tags || 
 							Typed->ResourceIds != ResourceIds || 
-							Typed->SkeletonIds != SkeletonIds ||
+							Typed->SkeletonId != SkeletonId ||
 							Typed->BonePosePriority != BonePosePriority ||
 							Typed->PhysicsAssetId != PhysicsAssetId ||
 							Typed->AdditionalPhysicsAssetIds != AdditionalPhysicsAssetIds)
@@ -768,7 +768,7 @@ namespace UE::Mutable::Private
 					Ptr<ASTOpMeshAddMetadata> New = new ASTOpMeshAddMetadata();
 					New->Tags = Tags;
 					New->ResourceIds = ResourceIds;
-					New->SkeletonIds = SkeletonIds;
+					New->SkeletonId = SkeletonId;
 					New->BonePosePriority= BonePosePriority;
 					New->PhysicsAssetId = PhysicsAssetId;
 					New->AdditionalPhysicsAssetIds = AdditionalPhysicsAssetIds;
diff --git a/Engine/Plugins/Mutable/Source/MutableTools/Private/MuT/CodeGenerator_Mesh.cpp b/Engine/Plugins/Mutable/Source/MutableTools/Private/MuT/CodeGenerator_Mesh.cpp
--- a/Engine/Plugins/Mutable/Source/MutableTools/Private/MuT/CodeGenerator_Mesh.cpp
+++ b/Engine/Plugins/Mutable/Source/MutableTools/Private/MuT/CodeGenerator_Mesh.cpp
@@ -775,7 +775,7 @@ namespace UE::Mutable::Private
 		TArray<UE::Tasks::FTask, TInlineAllocator<4>> Requisites;
 		Requisites.Emplace(InOptions);
 
-		TSharedPtr< TManagedPtr<FMesh> > ResolveMesh = MakeShared<TManagedPtr<FMesh>>();
+		TSharedPtr<TManagedPtr<FMesh>> ResolveMesh = MakeShared<TManagedPtr<FMesh>>();
 
 		// Compile-time reference? Add task to resolve as requisite.
 		if (Mesh->IsReference())
@@ -814,13 +814,13 @@ namespace UE::Mutable::Private
 				TArray<FString> Tags = ReferenceMesh->Tags; // __AnimBP Tags are added to the ReferenceMesh
 
 				TArray<uint64> ResourceIds;
-				TArray<uint32> SkeletonIds;
+				PASSTHROUGH_ID SkeletonID = PASSTHROUGH_ID_INVALID;
 				int16 BonePosePriority = 0;
 				PASSTHROUGH_ID PhysicsAssetId = PASSTHROUGH_ID_INVALID;
 				TArray<PASSTHROUGH_ID> AdditionalPhysicsAssetId;
 
 				const bool bMeshHasMetadata = ReferenceMesh->Tags.Num() ||
-					Mesh->SkeletonIDs.Num() > 0 ||
+					Mesh->SkeletonObjects.Num() > 0 ||
 					Mesh->Tags.Num() ||
 					Mesh->StreamedResources.Num() ||
 					Mesh->BonePosePriority != 0 ||
@@ -839,10 +839,7 @@ namespace UE::Mutable::Private
 						ResourceIds.AddUnique(ResourceId);
 					}
 
-					for (uint32 SkeletonId : Mesh->SkeletonIDs)
-					{
-						SkeletonIds.AddUnique(SkeletonId);
-					}
+					SkeletonID = Mesh->SkeletonObjects.IsEmpty() ? PASSTHROUGH_ID_INVALID : Mesh->SkeletonObjects[0].GetId();
 
 					BonePosePriority = Mesh->BonePosePriority;
 
@@ -861,7 +858,7 @@ namespace UE::Mutable::Private
 
 					MetadatalessMesh->Tags.SetNum(0, EAllowShrinking::No);
 					MetadatalessMesh->StreamedResources.SetNum(0, EAllowShrinking::No);
-					MetadatalessMesh->SkeletonIDs.SetNum(0, EAllowShrinking::No);
+					MetadatalessMesh->SkeletonObjects.SetNum(0, EAllowShrinking::No);
 					MetadatalessMesh->BonePosePriority = 0;
 					MetadatalessMesh->PhysicsAssets.SetNum(0, EAllowShrinking::No);
 					MetadatalessMesh->AdditionalPhysicsAssets.SetNum(0, EAllowShrinking::No);
@@ -943,7 +940,7 @@ namespace UE::Mutable::Private
 
 					AddMetadataOp->Tags        = MoveTemp(Tags);
 					AddMetadataOp->ResourceIds = MoveTemp(ResourceIds);
-					AddMetadataOp->SkeletonIds = MoveTemp(SkeletonIds);
+					AddMetadataOp->SkeletonId = SkeletonID;
 					AddMetadataOp->BonePosePriority = BonePosePriority;
 
 					AddMetadataOp->PhysicsAssetId = PhysicsAssetId;
