[Backout] - CL50294950
#rnx
Original CL Desc
-----------------------------------------------------------------
[Mutable] Retrived check that indicates if a surface edit node is being used correctly.
+Removed uncessary textures from edit nodes.

#rnx

#Synced-CL 50245735

#Synced-CL 50297267

diff --git a/Engine/Plugins/Mutable/Source/CustomizableObjectEditor/Private/MuCOE/GenerateMutableSource/GenerateMutableSourceModifier.cpp b/Engine/Plugins/Mutable/Source/CustomizableObjectEditor/Private/MuCOE/GenerateMutableSource/GenerateMutableSourceModifier.cpp
--- a/Engine/Plugins/Mutable/Source/CustomizableObjectEditor/Private/MuCOE/GenerateMutableSource/GenerateMutableSourceModifier.cpp
+++ b/Engine/Plugins/Mutable/Source/CustomizableObjectEditor/Private/MuCOE/GenerateMutableSource/GenerateMutableSourceModifier.cpp
@@ -317,6 +317,7 @@ UE::Mutable::Private::Ptr<UE::Mutable::Private::NodeModifier> GenerateMutableSou
 			SurfNode->LODs[LODIndex].MeshAdd = AddMeshNode;
 
 			const int32 NumImages = TypedNodeExt->GetNumParameters(EMaterialParameterType::Texture);
+			SurfNode->LODs[LODIndex].Textures.SetNum(NumImages);
 			for (int32 ImageIndex = 0; ImageIndex < NumImages; ++ImageIndex)
 			{
 				UE::Mutable::Private::NodeImagePtr ImageNode;
@@ -353,14 +354,9 @@ UE::Mutable::Private::Ptr<UE::Mutable::Private::NodeModifier> GenerateMutableSou
 					}
 				}
 
-				if (ImageNode)
-				{
-					UE::Mutable::Private::NodeModifierSurfaceEdit::FTexture Texture;
-					Texture.Extend = ImageNode;
-					Texture.MaterialParameterKey = UE::Mutable::Private::FParameterKey(MaterialParameterName, (int8)MaterialParameterLayerIndex);
 
-					SurfNode->LODs[LODIndex].Textures.Add(Texture);
-				}
+				SurfNode->LODs[LODIndex].Textures[ImageIndex].Extend = ImageNode;
+				SurfNode->LODs[LODIndex].Textures[ImageIndex].MaterialParameterKey = UE::Mutable::Private::FParameterKey(MaterialParameterName, MaterialParameterLayerIndex);
 			}
 		}
 
@@ -454,8 +450,9 @@ UE::Mutable::Private::Ptr<UE::Mutable::Private::NodeModifier> GenerateMutableSou
 			GenerationContext.FromLOD = 0;
 			GenerationContext.CurrentLOD = LODIndex;
 
-			const int32 NumImages = TypedNodeEdit->GetNumParameters(EMaterialParameterType::Texture);
 
+			const int32 NumImages = TypedNodeEdit->GetNumParameters(EMaterialParameterType::Texture);
+			SurfNode->LODs[LODIndex].Textures.SetNum(NumImages);
 			for (int32 ImageIndex = 0; ImageIndex < NumImages; ++ImageIndex)
 			{
 				const FNodeMaterialParameterId ImageId = TypedNodeEdit->GetParameterId(EMaterialParameterType::Texture, ImageIndex);
@@ -466,11 +463,11 @@ UE::Mutable::Private::Ptr<UE::Mutable::Private::NodeModifier> GenerateMutableSou
 					//check(ParentMaterialNode->IsImageMutableMode(ImageIndex)); // Ensured at graph time. If it fails, something is wrong.
 
 					const UEdGraphPin* ConnectedImagePin = FollowInputPin(*TypedNodeEdit->GetUsedImagePin(ImageId));
-					const FName& MaterialParameterName = TypedNodeEdit->GetParameterName(EMaterialParameterType::Texture, ImageIndex);
+					const FName MaterialParameterName = TypedNodeEdit->GetParameterName(EMaterialParameterType::Texture, ImageIndex);
 					const int32 MaterialParameterLayerIndex = TypedNodeEdit->GetParameterLayerIndex(EMaterialParameterType::Texture, ImageIndex);
 					check(MaterialParameterLayerIndex <= MAX_int8);
 					
-					UE::Mutable::Private::NodeModifierSurfaceEdit::FTexture ImagePatch;
+					UE::Mutable::Private::NodeModifierSurfaceEdit::FTexture& ImagePatch = SurfNode->LODs[LODIndex].Textures[ImageIndex];
 					ImagePatch.MaterialParameterKey = UE::Mutable::Private::FParameterKey(MaterialParameterName,  (int8)MaterialParameterLayerIndex);
 					// \todo: expose these two options?
 					ImagePatch.PatchBlendType = UE::Mutable::Private::EBlendType::BT_BLEND;
@@ -493,7 +490,6 @@ UE::Mutable::Private::Ptr<UE::Mutable::Private::NodeModifier> GenerateMutableSou
 					FIntPoint GridSize = TypedNodeEdit->Layout->GetGridSize();
 					FVector2f GridSizeF = FVector2f(GridSize);
 					ImagePatch.PatchBlocks.Reserve(TypedNodeEdit->Layout->Blocks.Num());
-
 					for (const FCustomizableObjectLayoutBlock& LayoutBlock : TypedNodeEdit->Layout->Blocks)
 					{
 						FBox2f Rect;
@@ -501,8 +497,6 @@ UE::Mutable::Private::Ptr<UE::Mutable::Private::NodeModifier> GenerateMutableSou
 						Rect.Max = FVector2f(LayoutBlock.Max) / GridSizeF;
 						ImagePatch.PatchBlocks.Add(Rect);
 					}
-
-					SurfNode->LODs[LODIndex].Textures.Add(ImagePatch);
 				}
 			}
 		}
diff --git a/Engine/Plugins/Mutable/Source/MutableTools/Private/MuT/CodeGenerator.cpp b/Engine/Plugins/Mutable/Source/MutableTools/Private/MuT/CodeGenerator.cpp
--- a/Engine/Plugins/Mutable/Source/MutableTools/Private/MuT/CodeGenerator.cpp
+++ b/Engine/Plugins/Mutable/Source/MutableTools/Private/MuT/CodeGenerator.cpp
@@ -874,8 +874,9 @@ namespace UE::Mutable::Private
 		int32 ComponentId = Options.Component ? Options.Component->Id : -1;
 		GetModifiersFor(ComponentId, SurfaceNode->Tags, bModifiersForBeforeOperations, Modifiers);
 
+		// TODO (Max): UE-354034 Fix this function and uncomment it
 		// This pass on the modifiers is only to detect errors that cannot be detected at the point they are applied.
-		CheckModifiersForSurface(*SurfaceNode, Modifiers, Options.LODIndex);
+		//CheckModifiersForSurface(*SurfaceNode, Modifiers, Options.LODIndex);
 		
         // Generate the mesh
 		FMeshGenerationStaticOptions MeshStaticOptions(ComponentId, Options.LODIndex);
@@ -2768,10 +2769,13 @@ namespace UE::Mutable::Private
 				const NodeModifierSurfaceEdit* Edit = static_cast<const NodeModifierSurfaceEdit*>(Mod.Node);
 				const NodeMaterialModify* Material = static_cast<const NodeMaterialModify*>(Node.Material.get());
 
-				// If this edit is not affected by the current LOD, then ignore.
-				// If both arrays are empty means that there are no textures to check
-				// If there is a MeshRemove, this edit does not affect textures
-				if (!Edit->LODs.IsValidIndex(CurrentLOD) || (Material->ImageParameters.IsEmpty() && Edit->LODs[CurrentLOD].Textures.IsEmpty()) || Edit->LODs[CurrentLOD].MeshRemove)
+				bool bAffectsCurrentLOD = Edit->LODs.IsValidIndex(CurrentLOD);
+				if (!bAffectsCurrentLOD)
+				{
+					continue;
+				}
+
+				if (Material->ImageParameters.IsEmpty() || Edit->LODs[CurrentLOD].Textures.IsEmpty())
 				{
 					continue;
 				}
