[Mutable] New Rom Manager.

Moved loading roms from CodeRunnerTasked to FRomManager. FRomManager is now thread safe.

#Synced-CL 50294903

diff --git a/Engine/Plugins/Mutable/Source/CustomizableObject/Internal/MuCO/UnrealMutableModelDiskStreamer.h b/Engine/Plugins/Mutable/Source/CustomizableObject/Internal/MuCO/UnrealMutableModelDiskStreamer.h
--- a/Engine/Plugins/Mutable/Source/CustomizableObject/Internal/MuCO/UnrealMutableModelDiskStreamer.h
+++ b/Engine/Plugins/Mutable/Source/CustomizableObject/Internal/MuCO/UnrealMutableModelDiskStreamer.h
@@ -103,8 +103,8 @@ protected:
 	};
 
 	TArray<FObjectData> Objects;
-
-	FCriticalSection FileHandlesCritical;
+	
+	mutable FCriticalSection CriticalSection;
 
 	/** This is used to generate unique ids for read requests. */
 	FOperationID LastOperationID = 0;
diff --git a/Engine/Plugins/Mutable/Source/CustomizableObject/Private/MuCO/CustomizableObject.cpp b/Engine/Plugins/Mutable/Source/CustomizableObject/Private/MuCO/CustomizableObject.cpp
--- a/Engine/Plugins/Mutable/Source/CustomizableObject/Private/MuCO/CustomizableObject.cpp
+++ b/Engine/Plugins/Mutable/Source/CustomizableObject/Private/MuCO/CustomizableObject.cpp
@@ -889,7 +889,7 @@ void UCustomizableObjectPrivate::SaveEmbeddedData(FArchive& Ar)
 	{	
 		if (Ar.IsCooking())
 		{
-			Model->Program.RomsCompileData.Empty();
+			Model->GetProgram().RomsCompileData.Empty();
 		}
 		UnrealMutableOutputStream Stream(Ar);
 		UE::Mutable::Private::FOutputArchive Arch(&Stream);
@@ -2217,7 +2217,7 @@ void UCustomizableObjectPrivate::SetModel(const TSharedPtr<UE::Mutable::Private:
 #endif
 		check(LocalModelResources);
 		
-		UE::Mutable::Private::FProgram& Program = MutableModel->Program;
+		UE::Mutable::Private::FProgram& Program = MutableModel->GetProgram();
 		Program.ExternalOperationProvider = MakeShared<UE::Mutable::Private::FExternalOperationProvider>(*LocalModelResources);
 
 		for (UE::Mutable::Private::FParameterDesc& Parameter : Program.Parameters)
@@ -2574,12 +2574,12 @@ namespace UE::Mutable::Private
 			// Add High-res textures to the Optional bucket
 			OptionalBucket.Categories.Add({ FFileCategoryID(EStreamableDataType::Model, uint16(UE::Mutable::Private::EDataType::Image), uint16(EMutableFileFlags::HighRes)), 0, 0, 0});
 
-			const int32 NumRoms = Model->Program.GetRomCount();
+			const int32 NumRoms = Model->GetProgram().GetRomCount();
 			for (int32 RomIndex = 0; RomIndex < NumRoms; ++RomIndex)
 			{
-				const uint32 BlockSize = Model->Program.GetRomSize(RomIndex);
-				const uint16 BlockResourceType = Model->Program.GetRomDataType(RomIndex) == ERomDataType::Mesh ? uint16(UE::Mutable::Private::EDataType::Mesh) : uint16(UE::Mutable::Private::EDataType::Image);
-				const EMutableFileFlags BlockFlags = Model->Program.IsRomHighRes(RomIndex) ? EMutableFileFlags::HighRes : EMutableFileFlags::None;
+				const uint32 BlockSize = Model->GetProgram().GetRomSize(RomIndex);
+				const uint16 BlockResourceType = Model->GetProgram().GetRomDataType(RomIndex) == ERomDataType::Mesh ? uint16(UE::Mutable::Private::EDataType::Mesh) : uint16(UE::Mutable::Private::EDataType::Image);
+				const EMutableFileFlags BlockFlags = Model->GetProgram().IsRomHighRes(RomIndex) ? EMutableFileFlags::HighRes : EMutableFileFlags::None;
 
 				FFileCategoryID CategoryID = { EStreamableDataType::Model, BlockResourceType, uint16(BlockFlags) };
 				TPair<FFileBucket&, FFileCategory&> It = FindOrAddCategory(FileBuckets, DefaultBucket, CategoryID); // Add block to an existing or new category
@@ -2913,23 +2913,23 @@ namespace UE::Mutable::Private
 		}
 
 		// Root nodes by flags.
-		const uint32 NumRoms = Model->Program.GetRomCount();
+		const uint32 NumRoms = Model->GetProgram().GetRomCount();
 		TMap<FFileCategoryID, FClassifyNode> RootNode;
 
 		// Create blocks data.
 		{
 			for (uint32 RomIndex = 0; RomIndex < NumRoms; ++RomIndex)
 			{
-				uint32 SourceBlockId = Model->Program.GetRomSourceId(RomIndex);
-				const uint32 BlockSize = Model->Program.GetRomSize(RomIndex);
+				uint32 SourceBlockId = Model->GetProgram().GetRomSourceId(RomIndex);
+				const uint32 BlockSize = Model->GetProgram().GetRomSize(RomIndex);
 				
 				uint16 BlockResourceType = 0;
 				EMutableFileFlags BlockFlags = EMutableFileFlags::None;
 				// TODO: Temp. Remove after unifying generated output files code between editor an package. UE-222777 
 				if (bUseRomTypeAndFlagsToFilter)
 				{
-					BlockResourceType = Model->Program.GetRomDataType(RomIndex) == ERomDataType::Mesh ? uint16(UE::Mutable::Private::EDataType::Mesh) : uint16(UE::Mutable::Private::EDataType::Image);
-					BlockFlags = Model->Program.IsRomHighRes(RomIndex) ? EMutableFileFlags::HighRes : EMutableFileFlags::None;
+					BlockResourceType = Model->GetProgram().GetRomDataType(RomIndex) == ERomDataType::Mesh ? uint16(UE::Mutable::Private::EDataType::Mesh) : uint16(UE::Mutable::Private::EDataType::Image);
+					BlockFlags = Model->GetProgram().IsRomHighRes(RomIndex) ? EMutableFileFlags::HighRes : EMutableFileFlags::None;
 				}
 
 				FFileCategoryID CurrentCategory = { EStreamableDataType::Model, BlockResourceType, uint16(BlockFlags) };
@@ -3258,14 +3258,13 @@ void UCustomizableObjectPrivate::LogMemory()
 	const UE::Mutable::Private::FModel* MuModel = GetModel().Get();
 	if (MuModel)
 	{
-		const UE::Mutable::Private::FProgram& Program = MuModel->Program;
+		const UE::Mutable::Private::FProgram& Program = MuModel->GetProgram();
 		int32 ByteCode = Program.ByteCode.GetAllocatedSize();
 		int32 OpAddress = Program.OpAddress.GetAllocatedSize();
 		int32 Roms = Program.Roms.GetAllocatedSize();
 		int32 ConstantImages = Program.ConstantImages.GetAllocatedSize();
 		int32 ConstantImageLODIndices = Program.ConstantImageLODIndices.GetAllocatedSize();
 		int32 ConstantImageLODPermanent = Program.ConstantImageLODsPermanent.GetAllocatedSize();
-		int32 ConstantImageLODStreamed = Program.ConstantImageLODsStreamed.GetAllocatedSize();
 
 		int32 PermanentImages = 0;
 		for (const UE::Mutable::Private::TManagedPtr<const UE::Mutable::Private::FImage>& Entry : Program.ConstantImageLODsPermanent)
@@ -3275,12 +3274,11 @@ void UCustomizableObjectPrivate::LogMemory()
 				PermanentImages += Entry->GetDataSize();
 			}
 		}
-		int32 ImageTotal = ConstantImages + ConstantImageLODIndices + ConstantImageLODPermanent + PermanentImages + ConstantImageLODStreamed;
+		int32 ImageTotal = ConstantImages + ConstantImageLODIndices + ConstantImageLODPermanent + PermanentImages;
 
 		int32 ConstantMeshes = Program.ConstantMeshes.GetAllocatedSize();
 		int32 ConstantMeshContentIndices = Program.ConstantMeshContentIndices.GetAllocatedSize();
 		int32 ConstantMeshesPermanent = Program.ConstantMeshesPermanent.GetAllocatedSize();
-		int32 ConstantMeshesStreamed = Program.ConstantMeshesStreamed.GetAllocatedSize();
 
 		int32 PermanentMeshes = 0;
 		for (const UE::Mutable::Private::TManagedPtr<const UE::Mutable::Private::FMesh>& Entry : Program.ConstantMeshesPermanent)
@@ -3290,7 +3288,7 @@ void UCustomizableObjectPrivate::LogMemory()
 				PermanentMeshes += Entry->GetDataSize();
 			}
 		}
-		int32 MeshesTotal = ConstantMeshes + ConstantMeshContentIndices + ConstantMeshesPermanent + ConstantMeshesStreamed + PermanentMeshes;
+		int32 MeshesTotal = ConstantMeshes + ConstantMeshContentIndices + ConstantMeshesPermanent + PermanentMeshes;
 
 		int32 StringsTotal = Program.ConstantStrings.GetAllocatedSize();
 		for (const FString& Entry : Program.ConstantStrings)
@@ -3388,13 +3386,11 @@ void UCustomizableObjectPrivate::LogMemory()
 		UE_LOG(LogMutable, Log, TEXT("      LODIndices     : %8d"), ConstantImageLODIndices);
 		UE_LOG(LogMutable, Log, TEXT("      PermanentBuf   : %8d"), ConstantImageLODPermanent);
 		UE_LOG(LogMutable, Log, TEXT("      Permanent      : %8d"), PermanentImages);
-		UE_LOG(LogMutable, Log, TEXT("      StreamedBuf    : %8d"), ConstantImageLODStreamed);
 		UE_LOG(LogMutable, Log, TEXT("    Meshes           : %8d"), MeshesTotal);
 		UE_LOG(LogMutable, Log, TEXT("      Buffer         : %8d"), ConstantMeshes);
 		UE_LOG(LogMutable, Log, TEXT("      ContentIndices : %8d"), ConstantMeshContentIndices);
 		UE_LOG(LogMutable, Log, TEXT("      PermanentBuf   : %8d"), ConstantMeshesPermanent);
 		UE_LOG(LogMutable, Log, TEXT("      Permanent      : %8d"), PermanentMeshes);
-		UE_LOG(LogMutable, Log, TEXT("      StreamedBuf    : %8d"), ConstantMeshesStreamed);
 		UE_LOG(LogMutable, Log, TEXT("    Skeletons        : %8d"), SkeletonsTotal);
 		UE_LOG(LogMutable, Log, TEXT("    PhysicBodies     : %8d"), PhysicsTotal);
 		UE_LOG(LogMutable, Log, TEXT("    Parameters       : %8d"), ParametersTotal);
diff --git a/Engine/Plugins/Mutable/Source/CustomizableObject/Private/MuCO/CustomizableObjectMipDataProvider.cpp b/Engine/Plugins/Mutable/Source/CustomizableObject/Private/MuCO/CustomizableObjectMipDataProvider.cpp
--- a/Engine/Plugins/Mutable/Source/CustomizableObject/Private/MuCO/CustomizableObjectMipDataProvider.cpp
+++ b/Engine/Plugins/Mutable/Source/CustomizableObject/Private/MuCO/CustomizableObjectMipDataProvider.cpp
@@ -362,16 +362,16 @@ int32 FMutableTextureMipDataProvider::GetMips(const FTextureUpdateContext& Conte
 
 				for (int32 ImageToLoadIndex : ImageRef.ConstantImagesNeededToGenerate)
 				{
-					const int32 RomId = Model->Program.GetConstantImageRomId(ImageToLoadIndex, LODIndex);
+					const int32 RomId = Model->GetProgram().GetConstantImageRomId(ImageToLoadIndex, LODIndex);
 
-					if (RomId >= 0 && Model->Program.IsRomHighRes(RomId))
+					if (RomId >= 0 && Model->GetProgram().IsRomHighRes(RomId))
 					{
 						FMutableStreamableBlock& Block = ModelStreamableBulkData->ModelStreamables[RomId];
 						if (!Block.IsPrefetched)
 						{
 							Block.IsPrefetched = true;
 
-							MaxBlockSize = FMath::Max(MaxBlockSize, Model->Program.GetRomSize(RomId));
+							MaxBlockSize = FMath::Max(MaxBlockSize, Model->GetProgram().GetRomSize(RomId));
 							RomsToPrefetch.Add(RomId);
 						}
 					}
@@ -391,7 +391,7 @@ int32 FMutableTextureMipDataProvider::GetMips(const FTextureUpdateContext& Conte
 					ModelStreamableBulkData->ModelStreamables[RomId],
 					UE::Mutable::Private::EStreamableDataType::Model,
 					(uint16)UE::Mutable::Private::EDataType::Image,
-					TArrayView<uint8>(AllocatedMemory, Model->Program.GetRomSize(RomId)));
+					TArrayView<uint8>(AllocatedMemory, Model->GetProgram().GetRomSize(RomId)));
 			}
 		}
 
diff --git a/Engine/Plugins/Mutable/Source/CustomizableObject/Private/MuCO/UnrealMutableModelDiskStreamer.cpp b/Engine/Plugins/Mutable/Source/CustomizableObject/Private/MuCO/UnrealMutableModelDiskStreamer.cpp
--- a/Engine/Plugins/Mutable/Source/CustomizableObject/Private/MuCO/UnrealMutableModelDiskStreamer.cpp
+++ b/Engine/Plugins/Mutable/Source/CustomizableObject/Private/MuCO/UnrealMutableModelDiskStreamer.cpp
@@ -65,6 +65,8 @@ bool FUnrealMutableModelBulkReader::PrepareStreamingForObject(UCustomizableObjec
 		return false;
 	}
 
+	FScopeLock Lock(&CriticalSection);
+
 	// See if we can free previuously allocated resources.
 	for (int32 ObjectIndex = 0; ObjectIndex < Objects.Num(); )
 	{
@@ -123,6 +125,8 @@ void FUnrealMutableModelBulkReader::CancelStreamingForObject(const UCustomizable
 		check(false);
 	}
 
+	FScopeLock Lock(&CriticalSection);
+
 	// See if we can free previuously allocated resources
 	for (int32 ObjectIndex = 0; ObjectIndex < Objects.Num(); ++ObjectIndex)
 	{
@@ -148,6 +152,8 @@ bool FUnrealMutableModelBulkReader::AreTherePendingStreamingOperationsForObject(
 		return false;
 	}
 
+	FScopeLock Lock(&CriticalSection);
+
 	for (int32 ObjectIndex = 0; ObjectIndex < Objects.Num(); ++ObjectIndex)
 	{
 		if (Objects[ObjectIndex].Model.Pin() == CustomizableObject->GetPrivate()->GetModel())
@@ -167,6 +173,8 @@ bool FUnrealMutableModelBulkReader::AreTherePendingStreamingOperationsForObject(
 
 void FUnrealMutableModelBulkReader::EndStreaming()
 {
+	FScopeLock Lock(&CriticalSection);
+
 	for (FObjectData& o : Objects)
 	{
 		for (TPair<FOperationID, FReadRequest>& it : o.CurrentReadRequests)
@@ -215,6 +223,8 @@ bool FUnrealMutableModelBulkReader::DoesBlockExist(const UE::Mutable::Private::F
 {
 	MUTABLE_CPUPROFILER_SCOPE(FUnrealMutableModelBulkReader::DoesBlockExist)
 
+	FScopeLock Lock(&CriticalSection);
+	
 	FObjectData* ObjectData = Objects.FindByPredicate(
 			[Model](const FObjectData& ObjectData) { return ObjectData.Model.Pin().Get() == Model; });
 
@@ -269,6 +279,8 @@ UE::Mutable::Private::FModelReader::FOperationID FUnrealMutableModelBulkReader::
 {
 	MUTABLE_CPUPROFILER_SCOPE(FUnrealMutableModelBulkStreamer::OpenReadFile);
 
+	FScopeLock Lock(&CriticalSection);
+	
 	UE_LOG(LogMutable, VeryVerbose, TEXT("Streaming: reading data %08llu."), Key);
 
 	// Find the object we are streaming for
@@ -402,27 +414,22 @@ UE::Mutable::Private::FModelReader::FOperationID FUnrealMutableModelBulkReader::
 				});
 		}
 
-		TSharedPtr<IAsyncReadFileHandle> FileHandle;
+		TSharedPtr<IAsyncReadFileHandle>& Found = ObjectData->ReadFileHandles.FindOrAdd(Block->FileId);
+		if (!Found)
 		{
-			FScopeLock Lock(&FileHandlesCritical);
+			FString FilePath = Private::MakeBlockFilePath(ObjectData->BulkFilePrefix, *Block);
+			Found = MakeShareable(FPlatformFileManager::Get().GetPlatformFile().OpenAsyncRead(*FilePath));
 
-			TSharedPtr<IAsyncReadFileHandle>& Found = ObjectData->ReadFileHandles.FindOrAdd(Block->FileId);
 			if (!Found)
 			{
-				FString FilePath = Private::MakeBlockFilePath(ObjectData->BulkFilePrefix, *Block);
-				Found = MakeShareable(FPlatformFileManager::Get().GetPlatformFile().OpenAsyncRead(*FilePath));
-
-				if (!Found)
-				{
-					UE_LOG(LogMutable, Error, TEXT("Failed to create AsyncReadFileHandle. File Path [%s]."), *FilePath);
-					check(false);
-					return -1;
-				}
+				UE_LOG(LogMutable, Error, TEXT("Failed to create AsyncReadFileHandle. File Path [%s]."), *FilePath);
+				check(false);
+				return -1;
 			}
-
-			FileHandle = Found;
 		}
 
+		TSharedPtr<IAsyncReadFileHandle> FileHandle = Found;
+
 		ReadRequest.FileReadRequest = MakeShareable(FileHandle->ReadRequest(
 			sizeof(MutableCompiledDataStreamHeader) + Block->Offset,
 			Size,
@@ -445,6 +452,8 @@ bool FUnrealMutableModelBulkReader::IsReadCompleted(UE::Mutable::Private::FModel
 {
 	MUTABLE_CPUPROFILER_SCOPE(FUnrealMutableModelBulkStreamer::IsReadCompleted);
 
+	FScopeLock Lock(&CriticalSection);
+	
 	for (FObjectData& o : Objects)
 	{
 		if (FReadRequest* ReadRequest = o.CurrentReadRequests.Find(OperationId))
@@ -476,6 +485,8 @@ bool FUnrealMutableModelBulkReader::EndRead(UE::Mutable::Private::FModelReader::
 {
 	MUTABLE_CPUPROFILER_SCOPE(FUnrealMutableModelBulkStreamer::EndRead);
 
+	FScopeLock Lock(&CriticalSection);
+
 	bool bSuccess = true;
 	bool bFound = false;
 	for (FObjectData& o : Objects)
diff --git a/Engine/Plugins/Mutable/Source/CustomizableObjectEditor/Private/MuCOE/CustomizableObjectCompiler.cpp b/Engine/Plugins/Mutable/Source/CustomizableObjectEditor/Private/MuCOE/CustomizableObjectCompiler.cpp
--- a/Engine/Plugins/Mutable/Source/CustomizableObjectEditor/Private/MuCOE/CustomizableObjectCompiler.cpp
+++ b/Engine/Plugins/Mutable/Source/CustomizableObjectEditor/Private/MuCOE/CustomizableObjectCompiler.cpp
@@ -1122,7 +1122,7 @@ void FCustomizableObjectCompiler::CompleteRequest(ECompilationStatePrivate State
 
 		if (PlatformData->Model)
 		{
-			PlatformData->Model->Program.UnloadRoms();
+			PlatformData->Model->RomManager.UnloadRoms();
 		}
 
 		if (!CurrentOptions.bIsCooking)
@@ -1635,7 +1635,7 @@ void FCustomizableObjectCompiler::FinishCompilationTask()
 	// At this point it is assumed that all data goes into a single file.
 	if (Model)
 	{
-		const int32 NumStreamingFiles = Model->Program.GetRomCount();
+		const int32 NumStreamingFiles = Model->GetProgram().GetRomCount();
 
 		TMap<uint32, FMutableStreamableBlock>& ModelStreamables = PlatformData->ModelStreamableBulkData->ModelStreamables;
 		ModelStreamables.Empty(NumStreamingFiles);
@@ -1643,8 +1643,8 @@ void FCustomizableObjectCompiler::FinishCompilationTask()
 		uint64 Offset = 0;
 		for (int32 FileIndex = 0; FileIndex < NumStreamingFiles; ++FileIndex)
 		{
-			const uint32 ResourceSize = Model->Program.GetRomSize(FileIndex);
-			EMutableFileFlags Flags = Model->Program.IsRomHighRes(FileIndex) ? EMutableFileFlags::HighRes : EMutableFileFlags::None;
+			const uint32 ResourceSize = Model->GetProgram().GetRomSize(FileIndex);
+			EMutableFileFlags Flags = Model->GetProgram().IsRomHighRes(FileIndex) ? EMutableFileFlags::HighRes : EMutableFileFlags::None;
 			ModelStreamables.Add(FileIndex, FMutableStreamableBlock{ .FileId = 0, .Flags = uint16(Flags), .Offset = Offset });
 			Offset += ResourceSize;
 		}
diff --git a/Engine/Plugins/Mutable/Source/CustomizableObjectEditor/Private/MuCOE/SMutableCodeViewer.cpp b/Engine/Plugins/Mutable/Source/CustomizableObjectEditor/Private/MuCOE/SMutableCodeViewer.cpp
--- a/Engine/Plugins/Mutable/Source/CustomizableObjectEditor/Private/MuCOE/SMutableCodeViewer.cpp
+++ b/Engine/Plugins/Mutable/Source/CustomizableObjectEditor/Private/MuCOE/SMutableCodeViewer.cpp
@@ -634,7 +634,7 @@ void SMutableCodeViewer::Construct(const FArguments& InArgs, UCustomizableObject
 					[
 						// Generate a new Constants panel to show the data stored on the current mutable program
 						SAssignNew(ConstantsWidget, SMutableConstantsWidget,
-							&(MutableModel->Program),
+							MutableModel,
 							SharedThis(this))
 					]
 				]
@@ -833,7 +833,7 @@ void SMutableCodeViewer::CacheOperationsMatchingStringPattern()
 		// to start searching
 		SearchPayload.SetupRootBatch(RootNodeAddresses);
 
-		const UE::Mutable::Private::FProgram& Program = MutableModel->Program;
+		const UE::Mutable::Private::FProgram& Program = MutableModel->GetProgram();
 		GetOperationsMatchingStringPattern(SearchString,bIsSearchStringRegularExpression,SearchPayload, Program);
 	
 		// Dump the located resources array onto the navigation array
@@ -1058,7 +1058,7 @@ void SMutableCodeViewer::OnSelectedOperationTypeFromTree()
 	const TSharedPtr<FMutableCodeTreeElement> ReferenceOperationElement = TreeView->GetSelectedItems()[0];
 	
 	const UE::Mutable::Private::EOpType OperationType =
-		MutableModel->Program.GetOpType(ReferenceOperationElement->MutableOperation);
+		MutableModel->GetProgram().GetOpType(ReferenceOperationElement->MutableOperation);
 
 	// Find the operation type directly in our array of operation elements (from the drop down)
 	const TSharedPtr<const FMutableOperationElement>* RepresentativeElement = FoundModelOperationTypeElements.
@@ -1102,7 +1102,7 @@ void SMutableCodeViewer::CacheAddressesOfOperationsOfType()
 	SearchPayload.SetupRootBatch(RootNodeAddresses);
 	
 	// Main update procedure run for the targeted state and the targeted parameter values
-	const UE::Mutable::Private::FProgram& Program = MutableModel->Program;
+	const UE::Mutable::Private::FProgram& Program = MutableModel->GetProgram();
 	GetOperationsOfType(OperationTypeToSearch,SearchPayload, Program);
 	
 	if (!SearchPayload.FoundElements.IsEmpty())
@@ -1171,7 +1171,7 @@ void SMutableCodeViewer::CacheOperationTypesPresentOnModel()
 	// there are
 	{
 		// Get the types and the amount of instances of each unique operation on the mutable model
-		const UE::Mutable::Private::FProgram& Program = MutableModel->Program;
+		const UE::Mutable::Private::FProgram& Program = MutableModel->GetProgram();
 		const uint32 ProgramAddressesCount = Program.OpAddress.Num();
 		
 		// Ensure first operation type is NONE since we are skipping it due to it having to have that type
@@ -1346,7 +1346,7 @@ void SMutableCodeViewer::GenerateAllTreeElements()
 	// 	be ignored due to how we handle them when expanding and contracting elements (OnExpansionChanged)
 
 	// Generate all root nodes
-	const UE::Mutable::Private::FProgram& Program = MutableModel->Program;
+	const UE::Mutable::Private::FProgram& Program = MutableModel->GetProgram();
 	const uint32 StateCount = Program.States.Num();
 	for ( uint32 StateIndex = 0; StateIndex < StateCount; ++StateIndex )
 	{
@@ -1948,7 +1948,7 @@ void SMutableCodeViewer::GetChildrenForInfo(TSharedPtr<FMutableCodeTreeElement>
 	}
 
 	check(MutableModel);
-	const UE::Mutable::Private::FProgram& Program = MutableModel->Program;
+	const UE::Mutable::Private::FProgram& Program = MutableModel->GetProgram();
 
 	UE::Mutable::Private::OP::ADDRESS ParentAddress = InInfo->MutableOperation;
 
@@ -2015,7 +2015,7 @@ void SMutableCodeViewer::OnExpansionChanged(TSharedPtr<FMutableCodeTreeElement>
 void SMutableCodeViewer::GetVisibleChildren(TSharedPtr<FMutableCodeTreeElement> InInfo, TSet<TSharedPtr<FMutableCodeTreeElement>>& OutChildren)
 {
 	check(MutableModel);
-	const UE::Mutable::Private::FProgram& MutableProgram = MutableModel->Program;
+	const UE::Mutable::Private::FProgram& MutableProgram = MutableModel->GetProgram();
 	
 	TArray<TSharedPtr<FMutableCodeTreeElement>> ToSearchForChildren;
 	ToSearchForChildren.Add(InInfo);
@@ -2093,7 +2093,7 @@ void SMutableCodeViewer::OnSelectionChanged(TSharedPtr<FMutableCodeTreeElement>
 	bIsPreviewPendingUpdate = true;
 
 	SelectedOperationAddress = SelectedNodes[0]->MutableOperation;
-	const UE::Mutable::Private::EOpType OperationType = MutableModel->Program.GetOpType(SelectedOperationAddress);
+	const UE::Mutable::Private::EOpType OperationType = MutableModel->GetProgram().GetOpType(SelectedOperationAddress);
 	const UE::Mutable::Private::EDataType OperationDataType = UE::Mutable::Private::GetOpDataType(OperationType);
 	
 	switch (OperationDataType)
@@ -2366,7 +2366,7 @@ void SMutableCodeViewer::TreeExpandElements(TArray<TSharedPtr<FMutableCodeTreeEl
 		// Filter the elements being expanded if the user has defined a desired EDataType
 		if (FilteringDataType != UE::Mutable::Private::EDataType::None)
 		{
-			const UE::Mutable::Private::EOpType OperationType = Item->MutableModel->Program.GetOpType(OperationAddress);
+			const UE::Mutable::Private::EOpType OperationType = Item->MutableModel->GetProgram().GetOpType(OperationAddress);
 			const UE::Mutable::Private::EDataType OperationDataType = UE::Mutable::Private::GetOpDataType(OperationType);
 
 			// If it is not of the desired type then ignore it and continue to the next pending element
@@ -2507,10 +2507,10 @@ void SMutableCodeViewer::CacheRootNodeAddresses()
 
 	TArray<UE::Mutable::Private::OP::ADDRESS> FoundRootNodeAddresses;
 	
-	const int32 StateCount = MutableModel->Program.States.Num();
+	const int32 StateCount = MutableModel->GetProgram().States.Num();
 	for ( int32 StateIndex=0; StateIndex < StateCount; ++StateIndex )
 	{
-		const UE::Mutable::Private::FProgram::FState& State = MutableModel->Program.States[StateIndex];
+		const UE::Mutable::Private::FProgram::FState& State = MutableModel->GetProgram().States[StateIndex];
 		FoundRootNodeAddresses.Add(State.Root);
 	}
 
@@ -2537,7 +2537,7 @@ void SMutableCodeViewer::CacheRootNodeAddresses()
 	SearchPayload.SetupRootBatch(RootNodeAddresses);
 	
 	// Main update procedure run for the targeted state and the targeted parameter values
-	const UE::Mutable::Private::FProgram& Program = MutableModel->Program; 
+	const UE::Mutable::Private::FProgram& Program = MutableModel->GetProgram(); 
 	GetOperationsReferencingConstantResource(ConstantDataType,IndexOnConstantsArray,SearchPayload, Program);
 	
 	// At this point we did get all the addresses of operations that do involve the usage of our resource
@@ -2708,26 +2708,7 @@ bool SMutableCodeViewer::IsConstantResourceUsedByOperation(const int32 IndexOnCo
 			{
 				if (OperationType == UE::Mutable::Private::EOpType::ME_CONSTANT)
 				{
-					UE::Mutable::Private::OP::MeshConstantArgs Args = InProgram.GetOpArgs<UE::Mutable::Private::OP::MeshConstantArgs>(OperationAddress);
-					UE::Mutable::Private::FConstantResourceIndex ResourceIndex = *reinterpret_cast<const UE::Mutable::Private::FConstantResourceIndex*>(&Args.Value);
-					
-					if (!ResourceIndex.Streamable)
-					{
-						bResourceLocated = IndexOnConstantsArray == ResourceIndex.Index;
-					}
-					else
-					{
-						int32 DebuggerIndexOffset = InProgram.ConstantMeshesPermanent.Num();
-						for (const TPair<uint32, UE::Mutable::Private::TManagedPtr<const UE::Mutable::Private::FMesh>>& Entry : InProgram.ConstantMeshesStreamed)
-						{
-							if (Entry.Key == ResourceIndex.Index)
-							{
-								bResourceLocated = IndexOnConstantsArray == DebuggerIndexOffset;
-								break;
-							}
-							++DebuggerIndexOffset;
-						}
-					}
+					bResourceLocated = IndexOnConstantsArray == InProgram.GetOpArgs<UE::Mutable::Private::OP::MeshConstantArgs>(OperationAddress).Value;
 				}
 				break;
 			}
@@ -2864,7 +2845,7 @@ void SMutableCodeViewer::Tick(const FGeometry& AllottedGeometry, const double In
 
 	bIsPreviewPendingUpdate = false;
 
-	const UE::Mutable::Private::EOpType OperationType = MutableModel->Program.GetOpType(SelectedOperationAddress);
+	const UE::Mutable::Private::EOpType OperationType = MutableModel->GetProgram().GetOpType(SelectedOperationAddress);
 	const UE::Mutable::Private::EDataType OperationDataType = UE::Mutable::Private::GetOpDataType(OperationType);
 
 	UE::Mutable::Private::FSettings Settings;
@@ -3207,7 +3188,7 @@ void FMutableCodeTreeElement::SetElementCurrentState(const int32& InStateIndex)
 
 	// Check for an out of bounds value
 	check(MutableModel);
-	UE::Mutable::Private::FProgram& MutableProgram = MutableModel->Program;
+	UE::Mutable::Private::FProgram& MutableProgram = MutableModel->GetProgram();
 	check(InStateIndex >= 0 && InStateIndex < MutableProgram.States.Num());
 
 	CurrentMutableStateIndex = InStateIndex;
@@ -3236,7 +3217,7 @@ void FMutableCodeTreeElement::SetElementCurrentState(const int32& InStateIndex)
 
 void FMutableCodeTreeElement::GenerateLabelText()
 {
-	const UE::Mutable::Private::FProgram& Program = MutableModel->Program;
+	const UE::Mutable::Private::FProgram& Program = MutableModel->GetProgram();
 	const UE::Mutable::Private::EOpType OperationType = Program.GetOpType(MutableOperation);
 	FString OpName = UE::Mutable::Private::s_opNames[static_cast<int32>(OperationType)];
 	OpName.TrimEndInline();
diff --git a/Engine/Plugins/Mutable/Source/CustomizableObjectEditor/Private/MuCOE/SMutableConstantsWidget.cpp b/Engine/Plugins/Mutable/Source/CustomizableObjectEditor/Private/MuCOE/SMutableConstantsWidget.cpp
--- a/Engine/Plugins/Mutable/Source/CustomizableObjectEditor/Private/MuCOE/SMutableConstantsWidget.cpp
+++ b/Engine/Plugins/Mutable/Source/CustomizableObjectEditor/Private/MuCOE/SMutableConstantsWidget.cpp
@@ -6,6 +6,8 @@
 #include "Editor.h"
 #include "Algo/StableSort.h"
 #include "Framework/Views/TableViewMetadata.h"
+#include "MuCO/CustomizableObject.h"
+#include "MuCO/CustomizableObjectSystemPrivate.h"
 #include "MuCOE/SMutableCodeViewer.h"
 #include "MuT/TypeInfo.h"
 #include "Widgets/Input/SButton.h"
@@ -578,15 +580,15 @@ public:
 
 
 
-void SMutableConstantsWidget::Construct(const FArguments& InArgs,const UE::Mutable::Private::FProgram* InMutableProgramPtr,  TSharedPtr<SMutableCodeViewer> InMutableCodeViewerPtr)
+void SMutableConstantsWidget::Construct(const FArguments& InArgs, TSharedPtr<UE::Mutable::Private::FModel> InModel,  TSharedPtr<SMutableCodeViewer> InMutableCodeViewerPtr)
 {
 	// A pointer to MutableCodeViewerPtr is required in order to be able to invoke the preview of our constants
 	check (InMutableCodeViewerPtr);
 	MutableCodeViewerPtr = InMutableCodeViewerPtr;
 	
 	// A pointer to the mutable program object is required to get the constants data
-	check(InMutableProgramPtr);
-	SetProgram(InMutableProgramPtr);
+	check(InModel);
+	SetModel(InModel);
 	
 	// Formatting constants
 	constexpr uint32 InBetweenListsVerticalPadding = 4;
@@ -1008,19 +1010,19 @@ void SMutableConstantsWidget::Construct(const FArguments& InArgs,const UE::Mutab
 
 
 
-void SMutableConstantsWidget::SetProgram(const UE::Mutable::Private::FProgram* InProgram)
+void SMutableConstantsWidget::SetModel(const TSharedPtr<UE::Mutable::Private::FModel>& InModel)
 {
 	// Make sure we do not process the setting of the same program object as the one already set
-	if (MutableProgramPtr == InProgram)
+	if (MutableModel == InModel)
 	{
 		return;
 	}
 	
 	// Set only once the program that is being used. No further updates should be required
-	MutableProgramPtr = InProgram;
+	MutableModel = InModel;
 	
 	// Generate the backend for the lists used in this object
-	if (MutableProgramPtr != nullptr)
+	if (MutableModel != nullptr)
 	{
 		LoadConstantElements();
 	}
@@ -1495,16 +1497,16 @@ void SMutableConstantsWidget::LoadConstantElements()
 
 void SMutableConstantsWidget::LoadConstantStrings()
 {
-	check (MutableProgramPtr);
+	check(MutableModel);
 
-	const int32 ConstantsCount = MutableProgramPtr->ConstantStrings.Num();
+	const int32 ConstantsCount = MutableModel->GetProgram().ConstantStrings.Num();
 	ConstantStringElements.Empty(ConstantsCount);
 	uint64 ConstantStringsAccumulatedSize = 0;
 	
 	for (int32 StringAddressIndex = 0; StringAddressIndex < ConstantsCount; StringAddressIndex++)
 	{
 		TSharedPtr<FMutableConstantStringElement> ConstantStringElement = MakeShared<FMutableConstantStringElement>();
-		ConstantStringElement->MutableString = MutableProgramPtr->ConstantStrings[StringAddressIndex];
+		ConstantStringElement->MutableString = MutableModel->GetProgram().ConstantStrings[StringAddressIndex];
 		ConstantStringElement->IndexOnSourceVector = StringAddressIndex;
 		
 		// Cache resource size
@@ -1520,9 +1522,9 @@ void SMutableConstantsWidget::LoadConstantStrings()
 
 void SMutableConstantsWidget::LoadConstantImages()
 {
-	check (MutableProgramPtr);
+	check(MutableModel);
 
-	const int32 ConstantsCount = MutableProgramPtr->ConstantImages.Num();
+	const int32 ConstantsCount = MutableModel->GetProgram().ConstantImages.Num();
 	if (ConstantImageElements)
 	{
 		ConstantImageElements->Empty(ConstantsCount);
@@ -1534,17 +1536,22 @@ void SMutableConstantsWidget::LoadConstantImages()
 	}
 	
 	uint64 ConstantImagesAccumulatedSize = 0;
+	
+	const TSharedRef<UE::Mutable::Private::FModelReader> ModelReader = UCustomizableObjectSystem::GetInstance()->GetPrivate()->MutableSystem->StreamInterface.ToSharedRef();
 
 	for (int32 ImageIndex = 0; ImageIndex < ConstantsCount; ImageIndex++)
 	{
 		TSharedPtr<FMutableConstantImageElement> ConstantImageElement = MakeShared<FMutableConstantImageElement>();
-
-		MutableProgramPtr->GetConstant(ImageIndex, ConstantImageElement->ImagePtr, 0,
-			[this](int32 x, int32 y, int32 m, UE::Mutable::Private::EImageFormat f, UE::Mutable::Private::EInitializationType i) 
-			{ 
-				return UE::Mutable::Private::MakeManaged<UE::Mutable::Private::FImage>(x, y, m, f, i); 
-			}
-		);
+		
+		MutableModel->GetConstant(ImageIndex, ConstantImageElement->ImagePtr, 0,
+		[this, ModelReader](int32 RomIndex)
+		{
+			return MutableModel->RomManager.GetRom<UE::Mutable::Private::FImage>(RomIndex, ModelReader).GetResult();
+		},
+		[this](int32 x, int32 y, int32 m, UE::Mutable::Private::EImageFormat f, UE::Mutable::Private::EInitializationType i) 
+		{ 
+			return UE::Mutable::Private::MakeManaged<UE::Mutable::Private::FImage>(x, y, m, f, i); 
+		});
 
 		ConstantImageElement->IndexOnSourceVector = ImageIndex;
 		
@@ -1559,17 +1566,9 @@ void SMutableConstantsWidget::LoadConstantImages()
 
 void SMutableConstantsWidget::LoadConstantMeshes()
 {
-	check (MutableProgramPtr);
+	check(MutableModel);
 	
-	TArray<UE::Mutable::Private::TManagedPtr<const UE::Mutable::Private::FMesh>> AllMeshes;
-	AllMeshes.Reserve(MutableProgramPtr->ConstantMeshesPermanent.Num() + MutableProgramPtr->ConstantMeshesStreamed.Num() );
-	AllMeshes.Append(MutableProgramPtr->ConstantMeshesPermanent);
-	for (const TPair<uint32, UE::Mutable::Private::TManagedPtr<const UE::Mutable::Private::FMesh>>& Entry : MutableProgramPtr->ConstantMeshesStreamed)
-	{
-		AllMeshes.Add(Entry.Value);
-	}
-
-	const int32 ConstantsCount = AllMeshes.Num();
+	const int32 ConstantsCount = MutableModel->GetProgram().ConstantMeshes.Num();
 	if (ConstantMeshElements)
 	{
 		ConstantMeshElements->Empty(ConstantsCount);
@@ -1582,10 +1581,22 @@ void SMutableConstantsWidget::LoadConstantMeshes()
 	
 	uint64 ConstantMeshesAccumulatedSize = 0;
 	
+	const TSharedRef<UE::Mutable::Private::FModelReader> ModelReader = UCustomizableObjectSystem::GetInstance()->GetPrivate()->MutableSystem->StreamInterface.ToSharedRef();
+
 	for (int32 MeshIndex = 0; MeshIndex < ConstantsCount; MeshIndex++)
 	{
 		TSharedPtr< FMutableConstantMeshElement> ConstantMeshElement = MakeShared<FMutableConstantMeshElement>();
-		ConstantMeshElement->MeshPtr = AllMeshes[MeshIndex];
+
+		MutableModel->GetConstant(MeshIndex, INDEX_NONE, {}, ConstantMeshElement->MeshPtr, UE::Mutable::Private::EMeshContentFlags::AllFlags,
+		[this, ModelReader](int32 RomIndex)
+		{
+			return MutableModel->RomManager.GetRom<UE::Mutable::Private::FMesh>(RomIndex, ModelReader).GetResult();
+		},
+		[](int32 BudgetReserve)
+		{
+			return UE::Mutable::Private::MakeManaged<UE::Mutable::Private::FMesh>();
+		});
+		
 		ConstantMeshElement->IndexOnSourceVector = MeshIndex;
 
 		// Actual core disk size would be:
@@ -1606,8 +1617,8 @@ void SMutableConstantsWidget::LoadConstantMeshes()
 
 void SMutableConstantsWidget::LoadConstantLayouts()
 {
-	check (MutableProgramPtr);
-	const int32 ConstantsCount = MutableProgramPtr->ConstantLayouts.Num();
+	check(MutableModel);
+	const int32 ConstantsCount = MutableModel->GetProgram().ConstantLayouts.Num();
 	ConstantLayoutElements.Empty(ConstantsCount);
 	
 	UE::Mutable::Private::FOutputMemoryStream Stream;
@@ -1616,7 +1627,7 @@ void SMutableConstantsWidget::LoadConstantLayouts()
 	for (int32 LayoutIndex = 0; LayoutIndex < ConstantsCount; LayoutIndex++)
 	{		
 		TSharedPtr<FMutableConstantLayoutElement> ConstantLayoutElement = MakeShared<FMutableConstantLayoutElement>();
-		ConstantLayoutElement->Layout = MutableProgramPtr->ConstantLayouts[LayoutIndex];
+		ConstantLayoutElement->Layout = MutableModel->GetProgram().ConstantLayouts[LayoutIndex];
 		ConstantLayoutElement->IndexOnSourceVector = LayoutIndex;
 		
 		// Cache resource for later GetBufferSize() call
@@ -1631,8 +1642,8 @@ void SMutableConstantsWidget::LoadConstantLayouts()
 
 void SMutableConstantsWidget::LoadConstantSkeletons()
 {
-	check (MutableProgramPtr);
-	const int32 ConstantsCount = MutableProgramPtr->ConstantSkeletons.Num();
+	check(MutableModel);
+	const int32 ConstantsCount = MutableModel->GetProgram().ConstantSkeletons.Num();
 	ConstantSkeletonElements.Empty(ConstantsCount);
 
 	UE::Mutable::Private::FOutputMemoryStream Stream;
@@ -1641,7 +1652,7 @@ void SMutableConstantsWidget::LoadConstantSkeletons()
 	for (int32 SkeletonIndex = 0; SkeletonIndex < ConstantsCount; SkeletonIndex++)
 	{
 		TSharedPtr<FMutableConstantSkeletonElement> ConstantSkeletonElement = MakeShared<FMutableConstantSkeletonElement>();
-		ConstantSkeletonElement->Skeleton = MutableProgramPtr->ConstantSkeletons[SkeletonIndex];
+		ConstantSkeletonElement->Skeleton = MutableModel->GetProgram().ConstantSkeletons[SkeletonIndex];
 		ConstantSkeletonElement->IndexOnSourceVector = SkeletonIndex;
 		
 		ConstantSkeletonElement->Skeleton->Serialise(Archive);
@@ -1655,8 +1666,8 @@ void SMutableConstantsWidget::LoadConstantSkeletons()
 
 void SMutableConstantsWidget::LoadConstantPhysics()
 {
-	check (MutableProgramPtr);
-	const int32 ConstantsCount = MutableProgramPtr->ConstantPhysicsBodies.Num();
+	check(MutableModel);
+	const int32 ConstantsCount = MutableModel->GetProgram().ConstantPhysicsBodies.Num();
 	ConstantPhysicsElements.Empty(ConstantsCount);
 
 	UE::Mutable::Private::FOutputMemoryStream Stream;
@@ -1665,7 +1676,7 @@ void SMutableConstantsWidget::LoadConstantPhysics()
 	for (int32 PhysicsIndex = 0; PhysicsIndex < ConstantsCount; PhysicsIndex++)
 	{
 		TSharedPtr<FMutableConstantPhysicsElement> ConstantPhysicsElement = MakeShared<FMutableConstantPhysicsElement>();
-		ConstantPhysicsElement->Physics = MutableProgramPtr->ConstantPhysicsBodies[PhysicsIndex];
+		ConstantPhysicsElement->Physics = MutableModel->GetProgram().ConstantPhysicsBodies[PhysicsIndex];
 		ConstantPhysicsElement->IndexOnSourceVector = PhysicsIndex;
 	
 		ConstantPhysicsElement->Physics->Serialise(Archive);
@@ -1680,8 +1691,8 @@ void SMutableConstantsWidget::LoadConstantPhysics()
 
 void SMutableConstantsWidget::LoadConstantMaterials()
 {
-	check (MutableProgramPtr);
-	const int32 ConstantsCount = MutableProgramPtr->ConstantMaterials.Num();
+	check(MutableModel);
+	const int32 ConstantsCount = MutableModel->GetProgram().ConstantMaterials.Num();
 	
 	if (ConstantMaterialElements)
 	{
@@ -1699,7 +1710,7 @@ void SMutableConstantsWidget::LoadConstantMaterials()
 	for (int32 MaterialIndex = 0; MaterialIndex < ConstantsCount; MaterialIndex++)
 	{
 		TSharedPtr<FMutableConstantMaterialElement> ConstantMaterialElement = MakeShared<FMutableConstantMaterialElement>();
-		ConstantMaterialElement->MutableMaterial = MutableProgramPtr->ConstantMaterials[MaterialIndex];
+		ConstantMaterialElement->MutableMaterial = MutableModel->GetProgram().ConstantMaterials[MaterialIndex];
 		ConstantMaterialElement->IndexOnSourceVector = MaterialIndex;
 		ConstantMaterialElement->MutableMaterial->Serialise(Archive);
 
@@ -1713,8 +1724,8 @@ void SMutableConstantsWidget::LoadConstantMaterials()
 
 void SMutableConstantsWidget::LoadConstantProjectors()
 {
-	check (MutableProgramPtr);
-	const int32 ConstantsCount = MutableProgramPtr->ConstantProjectors.Num();
+	check(MutableModel);
+	const int32 ConstantsCount = MutableModel->GetProgram().ConstantProjectors.Num();
 	ConstantProjectorElements.Empty(ConstantsCount);
 
 	UE::Mutable::Private::FOutputMemoryStream Stream;
@@ -1723,7 +1734,7 @@ void SMutableConstantsWidget::LoadConstantProjectors()
 	for (int32 ProjectorIndex = 0; ProjectorIndex < ConstantsCount; ProjectorIndex++)
 	{
 		TSharedPtr<FMutableConstantProjectorElement> ConstantProjectorElement = MakeShared<FMutableConstantProjectorElement>();
-		ConstantProjectorElement->Projector = &(MutableProgramPtr->ConstantProjectors[ProjectorIndex]);
+		ConstantProjectorElement->Projector = &(MutableModel->GetProgram().ConstantProjectors[ProjectorIndex]);
 		ConstantProjectorElement->IndexOnSourceVector = ProjectorIndex;
 		
 		ConstantProjectorElement->Projector->Serialise(Archive);
@@ -1737,8 +1748,8 @@ void SMutableConstantsWidget::LoadConstantProjectors()
 
 void SMutableConstantsWidget::LoadConstantMatrices()
 {
-	check (MutableProgramPtr);
-	const int32 ConstantsCount = MutableProgramPtr->ConstantMatrices.Num();
+	check(MutableModel);
+	const int32 ConstantsCount = MutableModel->GetProgram().ConstantMatrices.Num();
 	ConstantMatrixElements.Empty(ConstantsCount);
 
 	UE::Mutable::Private::FOutputMemoryStream Stream;
@@ -1747,7 +1758,7 @@ void SMutableConstantsWidget::LoadConstantMatrices()
 	for (int32 MatrixIndex = 0; MatrixIndex < ConstantsCount; MatrixIndex++)
 	{
 		TSharedPtr<FMutableConstantMatrixElement> ConstantMatrixElement = MakeShared<FMutableConstantMatrixElement>();
-		ConstantMatrixElement->Matrix = MutableProgramPtr->ConstantMatrices[MatrixIndex];
+		ConstantMatrixElement->Matrix = MutableModel->GetProgram().ConstantMatrices[MatrixIndex];
 		ConstantMatrixElement->IndexOnSourceVector = MatrixIndex;
 		
 		Archive << ConstantMatrixElement->Matrix;
@@ -1761,8 +1772,8 @@ void SMutableConstantsWidget::LoadConstantMatrices()
 
 void SMutableConstantsWidget::LoadConstantShapes()
 {
-	check (MutableProgramPtr);
-	const int32 ConstantsCount = MutableProgramPtr->ConstantShapes.Num();
+	check(MutableModel);
+	const int32 ConstantsCount = MutableModel->GetProgram().ConstantShapes.Num();
 	ConstantShapeElements.Empty(ConstantsCount);
 
 	UE::Mutable::Private::FOutputMemoryStream Stream;
@@ -1771,7 +1782,7 @@ void SMutableConstantsWidget::LoadConstantShapes()
 	for (int32 ShapeIndex = 0; ShapeIndex < ConstantsCount; ShapeIndex++)
 	{
 		TSharedPtr<FMutableConstantShapeElement> ConstantShapeElement = MakeShared<FMutableConstantShapeElement>();
-		ConstantShapeElement->Shape = &(MutableProgramPtr->ConstantShapes[ShapeIndex]);
+		ConstantShapeElement->Shape = &(MutableModel->GetProgram().ConstantShapes[ShapeIndex]);
 		ConstantShapeElement->IndexOnSourceVector = ShapeIndex;
 		
 		ConstantShapeElement->Shape->Serialise(Archive);
@@ -1785,8 +1796,8 @@ void SMutableConstantsWidget::LoadConstantShapes()
 
 void SMutableConstantsWidget::LoadConstantCurves()
 {
-	check (MutableProgramPtr);
-	const int32 ConstantsCount = MutableProgramPtr->ConstantCurves.Num();
+	check(MutableModel);
+	const int32 ConstantsCount = MutableModel->GetProgram().ConstantCurves.Num();
 	ConstantCurveElements.Empty(ConstantsCount);
 	
 	UE::Mutable::Private::FOutputMemoryStream Stream;
@@ -1795,7 +1806,7 @@ void SMutableConstantsWidget::LoadConstantCurves()
 	for (int32 CurveIndex = 0; CurveIndex < ConstantsCount; CurveIndex++)
 	{
 		TSharedPtr<FMutableConstantCurveElement> ConstantCurveElement = MakeShared<FMutableConstantCurveElement>();
-		ConstantCurveElement->Curve = MutableProgramPtr->ConstantCurves[CurveIndex];
+		ConstantCurveElement->Curve = MutableModel->GetProgram().ConstantCurves[CurveIndex];
 		ConstantCurveElement->IndexOnSourceVector = CurveIndex;
 
 		Archive << ConstantCurveElement->Curve;
diff --git a/Engine/Plugins/Mutable/Source/CustomizableObjectEditor/Private/MuCOE/SMutableConstantsWidget.h b/Engine/Plugins/Mutable/Source/CustomizableObjectEditor/Private/MuCOE/SMutableConstantsWidget.h
--- a/Engine/Plugins/Mutable/Source/CustomizableObjectEditor/Private/MuCOE/SMutableConstantsWidget.h
+++ b/Engine/Plugins/Mutable/Source/CustomizableObjectEditor/Private/MuCOE/SMutableConstantsWidget.h
@@ -134,7 +134,7 @@ public:
 	 * @param InMutableProgramPtr - Pointer to the UE::Mutable::Private::FProgram object that holds the constants data.
 	 * @param  InMutableCodeViewerPtr - Pointer to the MutableCodeViewer tasked with the previewing of the constant values
 	 */
-	void Construct(const FArguments& InArgs, const UE::Mutable::Private::FProgram*  InMutableProgramPtr ,  TSharedPtr<SMutableCodeViewer> InMutableCodeViewerPtr);
+	void Construct(const FArguments& InArgs, TSharedPtr<UE::Mutable::Private::FModel> InModel,  TSharedPtr<SMutableCodeViewer> InMutableCodeViewerPtr);
 
 	/** It clears the selected element for all constant view slates (SListView and STileView) except the one that shows the data of the provided type
 	 * @param ExceptionDataType The type of data of the slate container that we want to skip when clearing the selected elements.
@@ -146,8 +146,8 @@ public:
 private:
 
 	/** Mutable object containing the constants data */
-	const UE::Mutable::Private::FProgram* MutableProgramPtr = nullptr;
-
+	TSharedPtr<UE::Mutable::Private::FModel> MutableModel;
+	
 	/** Slate capable of accessing the previewer object */
 	TSharedPtr<SMutableCodeViewer> MutableCodeViewerPtr = nullptr;
 	
@@ -155,7 +155,7 @@ private:
 	 *Sets the back end for the operation of this widget. Each time this is done the ui backend gets updated
 	 * @param InProgram - Mutable program object holding all the constants data
 	 */
-	void SetProgram(const UE::Mutable::Private::FProgram* InProgram);
+	void SetModel(const TSharedPtr<UE::Mutable::Private::FModel>& InModel);
 
 	/*
 	 * Pointers to all slates showing the constants
diff --git a/Engine/Plugins/Mutable/Source/MutableRuntime/Internal/MuR/CodeVisitor.h b/Engine/Plugins/Mutable/Source/MutableRuntime/Internal/MuR/CodeVisitor.h
--- a/Engine/Plugins/Mutable/Source/MutableRuntime/Internal/MuR/CodeVisitor.h
+++ b/Engine/Plugins/Mutable/Source/MutableRuntime/Internal/MuR/CodeVisitor.h
@@ -238,7 +238,7 @@ namespace UE::Mutable::Private
         {
             PARENT::SetDefaultState( STATE() );
 
-            PARENT::Traverse( at, Model->Program );
+            PARENT::Traverse( at, Model->GetProgram() );
         }
 
     protected:
diff --git a/Engine/Plugins/Mutable/Source/MutableRuntime/Internal/MuR/Mesh.h b/Engine/Plugins/Mutable/Source/MutableRuntime/Internal/MuR/Mesh.h
--- a/Engine/Plugins/Mutable/Source/MutableRuntime/Internal/MuR/Mesh.h
+++ b/Engine/Plugins/Mutable/Source/MutableRuntime/Internal/MuR/Mesh.h
@@ -345,7 +345,7 @@ namespace UE::Mutable::Private
         UE_API void SetSkeleton(const TManagedPtr<const FSkeleton>&);
 		UE_API TManagedPtr<const FSkeleton> GetSkeleton() const;
 
-		void ConvertBoneIdsToIndices();
+		UE_API void ConvertBoneIdsToIndices();
 
         //! \}
 
diff --git a/Engine/Plugins/Mutable/Source/MutableRuntime/Internal/MuR/Model.h b/Engine/Plugins/Mutable/Source/MutableRuntime/Internal/MuR/Model.h
--- a/Engine/Plugins/Mutable/Source/MutableRuntime/Internal/MuR/Model.h
+++ b/Engine/Plugins/Mutable/Source/MutableRuntime/Internal/MuR/Model.h
@@ -2,6 +2,7 @@
 
 #pragma once
 
+#include "ClothingAsset.h"
 #include "MuR/Parameters.h"
 #include "MuR/Model.h"
 #include "MuR/System.h"
@@ -15,6 +16,7 @@
 
 #include "Templates/SharedPointer.h"
 #include "ExternalOperationProvider.h"
+#include "RomManager.h"
 
 class UMaterialInterface;
 
@@ -31,7 +33,6 @@
 	class FModelWriter;
 	class FOutputArchive;
     class FModelParametersGenerator;
-
 	
 	/** Used to debug and log. */
 	constexpr bool DebugRom = false;
@@ -247,21 +248,12 @@
 
 		/** Data for every rom required at compile-time. It is empty in cooked data. */
 		TArray<FRomDataCompile> RomsCompileData;
-	
-		/** Loaded roms worth tracking (only images and meshes for now). Stores the rom's data type.*/
-		mutable TSparseArray<uint8> LoadedMemTrackedRoms;
-
+    	
 		/** Constant image mip data is split in 2 sets: ConstantImageLODsPermanent constains data that is always loaded. 
 		* Index with FConstantResourceIndex::Index, when Streamable is 0.
 		*/
 		TArray<TManagedPtr<const FImage>> ConstantImageLODsPermanent;
 
-		/** Constant image mip data is split in 2 sets: ConstantImageLODsStreamed constains data that is streamed in and out. 
-		* Index with FConstantResourceIndex::Index, when Streamable is 1.
-		* This part is empty for an unused Model, and shouldn't be serialised.
-		*/
-		mutable TMap<uint32, TManagedPtr<const FImage>> ConstantImageLODsStreamed;
-
 		/** Constant image mip chain indices: ranges in this array are defined in FImageLODRange and the indices here refer to ConstantImageLODs. */
 		TArray<FConstantResourceIndex> ConstantImageLODIndices;
 
@@ -278,13 +270,7 @@
 		* Index with FConstantResourceIndex::Index, when Streamable is 0.
 		*/
 		TArray<TManagedPtr<const FMesh>> ConstantMeshesPermanent;
-
-		/** Constant mesh data is split in 2 sets: ConstantMeshesStreamed constains data that is streamed in and out.
-		* Index with FConstantResourceIndex::Index, when Streamable is 1.
-		* This part is empty for an unused Model, and shouldn't be serialised.
-		*/
-		mutable TMap<uint32, TManagedPtr<const FMesh>> ConstantMeshesStreamed;
-
+    	
 		/** Constant FExtensionData */
 		TArray<FExtensionDataConstant> ConstantExtensionData;
 
@@ -352,12 +338,6 @@
 		bool bIsValid = true;
 #endif
     	
-    	/** For each model rom, the last time its streamed data was used. */
-    	mutable TMemoryTrackedArray<TPair<uint64, uint64>> RomWeights;
-
-    	/** Count of pending operations for every rom index. */
-    	mutable TMemoryTrackedArray<uint16> PendingOpsPerRom;
-    	
         UE_API void Serialise(FOutputArchive& Arch) const;
 
         void Unserialise(FInputArchive& Arch);
@@ -369,9 +349,6 @@
 
     	UE_API uint32 GetRomSize(int32 Index) const;
     	
-		/** Return true if the given ROM is loaded. */
-		bool IsRomLoaded(int32 RomIndex) const;
-
         UE_API int32 GetRomCount() const;
 
 #if WITH_EDITOR
@@ -388,19 +365,6 @@
 		 */
     	UE_API int32 GetConstantImageRomId(int32 ConstantImageIndex, int32 LODIndex) const;
     	
-    	UE_API void UnloadRoms() const;
-    	
-		/** Unload a rom resource. 
-		* If a pointer to a size is passed in OutDataSize, the value will be set with the size of the unloaded rom.
-		*/
-        void UnloadRom(int32 RomIndex, int32* OutDataSize=nullptr) const;
-
-        int32 GetRomBytes(int32 RomIndex) const;
-
-        void SetMeshRomValue(uint32 RomIndex, const TManagedPtr<FMesh>& Value) const;
-
-        void SetImageRomValue(uint32 RomIndex, const TManagedPtr<FImage>& Value) const;
-    	
         UE_API OP::CONSTANT_EXTENSIONDATA_INDEX AddConstant(const TManagedPtr<const FExtensionData>& Data);
     	
         UE_API OP::ADDRESS AddConstant(const TManagedPtr<const FLayout>& Layout);
@@ -430,20 +394,169 @@
         UE_API OP::ADDRESS AddConstant(const FRichCurve& Constant);
     	
         UE_API OP::ADDRESS AddConstant(const TManagedPtr<const FMaterial>& Material);
+
+		void GetExtensionDataConstant(int32 ConstantIndex, TManagedPtr<const FExtensionData>& Result) const;
+
+        UE_API EOpType GetOpType(OP::ADDRESS At) const;
     	
-        UE_API TManagedPtr<const FImage> GetImageLOD(FConstantResourceIndex Index) const;
+        template<typename ARGS>
+		ARGS GetOpArgs(OP::ADDRESS at) const
+		{
+			ARGS result;
+			uint64 byteCodeAddress = OpAddress[at];
+			byteCodeAddress += sizeof(EOpType);
+			FMemory::Memcpy(&result, &ByteCode[byteCodeAddress], sizeof(ARGS));
+			return result;
+		}
 
-        /** Get a constant image, assuming at least some mips are loaded. The image constant will be composed with lodaded mips if necessary. */
+		template<typename ARGS>
+		void SetOpArgs(OP::ADDRESS At, const ARGS& Args)
+		{
+			uint64 byteCodeAddress = OpAddress[At];
+			byteCodeAddress += sizeof(EOpType);
+			FMemory::Memcpy(&ByteCode[byteCodeAddress], &Args, sizeof(ARGS));
+		}
+
+        UE_API const uint8* GetOpArgsPointer(OP::ADDRESS At) const;
+    	
+        uint8* GetOpArgsPointer(OP::ADDRESS At);
+    };
+
+	
+    /** A Model represents a customisable object with any number of parameters.
+    * When values are given to the parameters, specific Instances can be built, which hold the
+    * built application-usable data.
+	*/
+    class FModel
+	{
+	public:
+    	UE_API FModel();
+    	
+		static UE_API void Serialise(const FModel*, FOutputArchive&);
+		static UE_API TSharedPtr<FModel> StaticUnserialise(FInputArchive&);
+    	
+    	void Serialise(FOutputArchive& Arch) const;
+
+		void Unserialise(FInputArchive& Arch);
+
+		/** Special serialise operation that serialises the data in separate "files". An object
+         * with the ModelStreamer interface is responsible of storing this data and providing
+		 * the "file" concept. 
+		 * If bDropData is set to true, the rom data will be freed as it is serialized.
+		 */
+        static UE_API void Serialise(FModel* Model, FModelWriter& Streamer, bool bDropData = false);
+
+		//! Return true if the model has external data in other files. This kind of models will
+		//! require data streaming when used.
+		UE_API bool HasExternalData() const;
+
+#if WITH_EDITOR
+		//! Return true unless the streamed resources were destroyed, which could happen in the
+		//! editor after recompiling the CO.
+		UE_API bool IsValid() const;
+
+		//! Invalidate the Model. Compiling a compiled CO will invalidate the model kept by previously
+		//! generated resources, like streamed textures.
+		UE_API void Invalidate();
+#endif
+
+		//-----------------------------------------------------------------------------------------
+		// Own interface
+		//-----------------------------------------------------------------------------------------
+
+		//! Create a set of new parameters of the model with the default values.
+		//! If old parameters are provided, they will be reused when possible instead of the
+		//! default values.
+        static UE_API TSharedPtr<FParameters> NewParameters(TSharedPtr<const FModel> Model, const FParameters* OldParameters = nullptr);
+
+		/** Return true if the parameter is multi-dimensional */
+		UE_API bool IsParameterMultidimensional(int32 ParameterIndex) const;
+
+		//! Get the number of states in the model.
+		UE_API int32 GetStateCount() const;
+
+		//! Get a state name by state index from 0 to GetStateCount-1
+		UE_API const FString& GetStateName( int32 StateIndex ) const;
+
+		//! Find a state index by state name
+		UE_API int32 FindState( const FString& Name ) const;
+
+		//! Get the number of parameters available in a particular state.
+		UE_API int32 GetStateParameterCount( int32 StateIndex ) const;
+
+		//! Get the index of one of the parameters in the given state. The index refers to the
+		//! parameters in a FParameters object obtained from this model with NewParameters.
+		UE_API int32 GetStateParameterIndex( int32 StateIndex, int32 ParamIndex ) const;
+	
+		//! Return the default value of a boolean parameter.
+		//! \pre The parameter specified by index is a T_BOOL.
+		//! \param Index Index of the parameter from 0 to GetCount()-1
+    	UE_API bool GetBoolDefaultValue(int32 Index) const;
+
+   		//! Return the default value of a integer parameter.
+		//! \pre The parameter specified by index is a T_INT.
+		//! \param Index Index of the parameter from 0 to GetCount()-1
+    	UE_API int32 GetIntDefaultValue(int32 Index) const;
+
+		//! Return the default value of a float parameter.
+		//! \pre The parameter specified by index is a T_FLOAT.
+		//! \param Index Index of the parameter from 0 to GetCount()-1
+        UE_API float GetFloatDefaultValue(int32 Index) const;
+
+		//! Return the default value of a colour parameter.
+		//! \pre The parameter specified by index is a T_FLOAT.
+        //! \param Index Index of the parameter from 0 to GetCount()-1
+        //! \param R,G,B Pointers to values where every resulting colour channel will be stored
+    	UE_API void GetColourDefaultValue(int32 Index, FVector4f& OutValue) const;
+
+		//! Return the default value of a colour parameter.
+		//! \pre The parameter specified by index is a T_FLOAT.
+        //! \param Index Index of the parameter from 0 to GetCount()-1
+    	UE_API FMatrix44f GetMatrixDefaultValue(int32 Index) const;
+
+    	//! Return the default value of a projector parameter, as a 4x4 matrix. The matrix is supposed to be
+		//! a linear transform in column-major.
+		//! \pre The parameter specified by index is a T_PROJECTOR.
+        //! \param Index Index of the parameter from 0 to GetCount()-1
+        //! \param OutPos Pointer to where the object-space position coordinates of the projector will be stored.
+        //! \param OutDir Pointer to where the object-space direction vector of the projector will be stored.
+        //! \param OutUp Pointer to where the object-space vertically up direction vector
+        //!         of the projector will be stored. This controls the "roll" angle of the
+        //!         projector.
+        //! \param OutScale Pointer to the projector-space scaling of the projector.
+    	UE_API void GetProjectorDefaultValue(int32 Index, EProjectorType* OutProjectionType, FVector3f* OutPos,
+    	 	FVector3f* OutDir, FVector3f* OutUp, FVector3f* OutScale, float* OutProjectionAngle) const;
+    	
+    	 /** Get a constant image, assuming at least some mips are loaded. The image constant will be composed with lodaded mips if necessary.
+    	  * @param RomsStreamed Given a rom index, return the  future of the streamed rom. */
 		template <typename CreateImageFunc>
-        void GetConstant(uint32 ConstantIndex, TManagedPtr<const FImage>& Result, int32 MipsToSkip, const CreateImageFunc& CreateImage) const
+        void GetConstant(uint32 ConstantIndex, TManagedPtr<const FImage>& Result, int32 MipsToSkip, const TFunction<TManagedPtr<const FImage>(int32)>&& RomsStreamed, const CreateImageFunc&& CreateImage) const
         {
-			const int32 NumLODs = ConstantImages[ConstantIndex].LODCount;
+			const int32 NumLODs = Program.ConstantImages[ConstantIndex].LODCount;
 
 			check(NumLODs > 0);
 			const int32 ReallySkippedLODs = FMath::Min(NumLODs - 1, MipsToSkip);
-			const int32 FirstLODIndexIndex = ConstantImages[ConstantIndex].FirstIndex;
-			const int32 FirstTailLOD = ConstantImages[ConstantIndex].LODCount - ConstantImages[ConstantIndex].NumLODsInTail;
+			const int32 FirstLODIndexIndex = Program.ConstantImages[ConstantIndex].FirstIndex;
+			const int32 FirstTailLOD = Program.ConstantImages[ConstantIndex].LODCount - Program.ConstantImages[ConstantIndex].NumLODsInTail;
 
+			auto GetImageLOD = [&](FConstantResourceIndex Index) -> TManagedPtr<const FImage>
+			{
+				TManagedPtr<const FImage> Mip;
+				if (!Index.Streamable)
+				{
+					if (Program.ConstantImageLODsPermanent.IsValidIndex(Index.Index))
+					{
+						Mip = Program.ConstantImageLODsPermanent[Index.Index];
+					}
+				}
+				else
+				{
+					return RomsStreamed(Index.Index);
+				}
+
+				return Mip;
+			};
+			
 			// Compose the result image
 			{
 				MUTABLE_CPUPROFILER_SCOPE(ComposeConstantImage);
@@ -456,7 +569,7 @@
 				for (; LOD <= FirstTailLOD; ++LOD)
 				{
 					MUTABLE_CPUPROFILER_SCOPE(Search_FisrtValid);
-					TManagedPtr<const FImage> FirstValidLODImage = GetImageLOD(ConstantImageLODIndices[FirstLODIndexIndex + LOD]);
+					TManagedPtr<const FImage> FirstValidLODImage = GetImageLOD(Program.ConstantImageLODIndices[FirstLODIndexIndex + LOD]);
 					
 					if (FirstValidLODImage)
 					{
@@ -466,8 +579,8 @@
 			
 				if (LOD < NumLODs)
 				{						
-					int32 SizeX = ConstantImages[ConstantIndex].ImageSizeX; 
-					int32 SizeY = ConstantImages[ConstantIndex].ImageSizeY; 
+					int32 SizeX = Program.ConstantImages[ConstantIndex].ImageSizeX; 
+					int32 SizeY = Program.ConstantImages[ConstantIndex].ImageSizeY; 
 
 					for (int32 I = 0; I < LOD; ++I)
 					{
@@ -475,10 +588,10 @@
 						SizeY = FMath::DivideAndRoundUp(SizeY, 2);
 					}
 					 
-					EImageFormat Format = ConstantImages[ConstantIndex].ImageFormat; 
+					EImageFormat Format = Program.ConstantImages[ConstantIndex].ImageFormat; 
 
 					ResultImage = CreateImage(SizeX, SizeY, NumLODs - LOD, Format, EInitializationType::NotInitialized);
-					ResultImage->Flags = ConstantImages[ConstantIndex].Flags;
+					ResultImage->Flags = Program.ConstantImages[ConstantIndex].Flags;
 				}
 	
 				check(ResultImage);
@@ -491,7 +604,7 @@
 				for (; LOD < FirstTailLOD; ++LOD, ++DestLOD)
 				{
 					MUTABLE_CPUPROFILER_SCOPE(NonTailLODs);
-					TManagedPtr<const FImage> Image = GetImageLOD(ConstantImageLODIndices[FirstLODIndexIndex + LOD]);
+					TManagedPtr<const FImage> Image = GetImageLOD(Program.ConstantImageLODIndices[FirstLODIndexIndex + LOD]);
 
 					if (bFormatNeedsDataResize)
 					{
@@ -510,7 +623,7 @@
 				if (LOD < NumLODs)
 				{
 					MUTABLE_CPUPROFILER_SCOPE(TailLODs);
-					TManagedPtr<const FImage> TailLODsImage = GetImageLOD(ConstantImageLODIndices[FirstLODIndexIndex + FirstTailLOD]);
+					TManagedPtr<const FImage> TailLODsImage = GetImageLOD(Program.ConstantImageLODIndices[FirstLODIndexIndex + FirstTailLOD]);
 
 					if (!ensureAlwaysMsgf(TailLODsImage, TEXT("Missing tail image rom, setting image to black.")))
 					{
@@ -579,6 +692,7 @@
 			}
 		}		
 
+    	/* @param RomsStreamed Given a rom index, return the  future of the streamed rom. */
 		template<class CreateMeshFunc>
         void GetConstant(
 				uint32 MeshConstantIndex, 
@@ -586,11 +700,12 @@
 				const TPassthroughObjectPtr<UClothingAssetCommon>& ClothAsset,
 				TManagedPtr<const FMesh>& OutMesh, 
 				EMeshContentFlags FilterContentFlags, 
+				const TFunction<TManagedPtr<const FMesh>(int32)>&& RomsStreamed,
 				CreateMeshFunc&& CreateMesh) const
         {
 			MUTABLE_CPUPROFILER_SCOPE(GetConstant_Mesh)
 
-			FMeshContentRange MeshContentRange = ConstantMeshes[MeshConstantIndex];
+			FMeshContentRange MeshContentRange = Program.ConstantMeshes[MeshConstantIndex];
 
 			TManagedPtr<const FMesh> EmptyMesh = nullptr;
 			auto GetMeshAtResourceIndex = [&](FConstantResourceIndex ResourceIndex) -> TManagedPtr<const FMesh>
@@ -599,9 +714,9 @@
 
 				if (!ResourceIndex.Streamable)
 				{
-					if (ConstantMeshesPermanent.IsValidIndex(ResourceIndex.Index))
+					if (Program.ConstantMeshesPermanent.IsValidIndex(ResourceIndex.Index))
 					{
-						return ConstantMeshesPermanent[ResourceIndex.Index];
+						return Program.ConstantMeshesPermanent[ResourceIndex.Index];
 					}
 
 					if (!EmptyMesh)
@@ -613,18 +728,13 @@
 				}
 				else
 				{
-					const TManagedPtr<const FMesh>* Found = ConstantMeshesStreamed.Find(ResourceIndex.Index);
-					if (Found)
+					TManagedPtr<const FMesh> Found = RomsStreamed(ResourceIndex.Index);
+					if (!Found)
 					{
-						return *Found;
+						Found = MakeManaged<FMesh>();
 					}
-
-					if (!EmptyMesh)
-					{
-						EmptyMesh = MakeManaged<FMesh>();
-					}
-
-					return EmptyMesh;
+					
+					return Found;
 				}
 			};
 
@@ -637,7 +747,7 @@
 			if (EnumHasAnyFlags(FilterContentFlags, EMeshContentFlags::GeometryData) &&
 				EnumHasAnyFlags(MeshContentRange.GetContentFlags(), EMeshContentFlags::GeometryData))
 			{
-				const FConstantResourceIndex ResourceIndex = ConstantMeshContentIndices[MeshRomCurrentIndex];
+				const FConstantResourceIndex ResourceIndex = Program.ConstantMeshContentIndices[MeshRomCurrentIndex];
 				GeometryMesh = GetMeshAtResourceIndex(ResourceIndex);
 			}
 			MeshRomCurrentIndex += (int32)EnumHasAnyFlags(MeshContentRange.GetContentFlags(), EMeshContentFlags::GeometryData);
@@ -645,7 +755,7 @@
 			if (EnumHasAnyFlags(FilterContentFlags, EMeshContentFlags::PoseData) &&
 				EnumHasAnyFlags(MeshContentRange.GetContentFlags(), EMeshContentFlags::PoseData))
 			{
-				const FConstantResourceIndex ResourceIndex = ConstantMeshContentIndices[MeshRomCurrentIndex];
+				const FConstantResourceIndex ResourceIndex = Program.ConstantMeshContentIndices[MeshRomCurrentIndex];
 				PoseMesh = GetMeshAtResourceIndex(ResourceIndex);
 			}
 			MeshRomCurrentIndex += (int32)EnumHasAnyFlags(MeshContentRange.GetContentFlags(), EMeshContentFlags::PoseData);
@@ -653,7 +763,7 @@
 			if (EnumHasAnyFlags(FilterContentFlags, EMeshContentFlags::PhysicsData) &&
 				EnumHasAnyFlags(MeshContentRange.GetContentFlags(), EMeshContentFlags::PhysicsData))
 			{
-				const FConstantResourceIndex ResourceIndex = ConstantMeshContentIndices[MeshRomCurrentIndex];
+				const FConstantResourceIndex ResourceIndex = Program.ConstantMeshContentIndices[MeshRomCurrentIndex];
 				PhysicsMesh = GetMeshAtResourceIndex(ResourceIndex);
 			}
 			MeshRomCurrentIndex += (int32)EnumHasAnyFlags(MeshContentRange.GetContentFlags(), EMeshContentFlags::PhysicsData);
@@ -661,7 +771,7 @@
 			if (EnumHasAnyFlags(FilterContentFlags, EMeshContentFlags::MetaData) &&
 				EnumHasAnyFlags(MeshContentRange.GetContentFlags(), EMeshContentFlags::MetaData))
 			{
-				const FConstantResourceIndex ResourceIndex = ConstantMeshContentIndices[MeshRomCurrentIndex];
+				const FConstantResourceIndex ResourceIndex = Program.ConstantMeshContentIndices[MeshRomCurrentIndex];
 				MetaDataMesh = GetMeshAtResourceIndex(ResourceIndex);
 			}
 			MeshRomCurrentIndex += (int32)EnumHasAnyFlags(MeshContentRange.GetContentFlags(), EMeshContentFlags::MetaData);
@@ -739,7 +849,7 @@
 
 				if (TManagedPtr<FPhysicsBody> PhysicsBody = ConstCastManagedPtr<FPhysicsBody>(Result->PhysicsBody))
 				{
-					PhysicsBody->ConvertBoneIdsToNames(ConstantNames);
+					PhysicsBody->ConvertBoneIdsToNames(Program.ConstantNames);
 				}
 			}
 
@@ -775,10 +885,9 @@
 
 			Result->MeshIDPrefix = MeshContentRange.MeshIDPrefix;
 
-			check(ConstantSkeletons.Num() > SkeletonConstantIndex);
-			if (ConstantSkeletons.IsValidIndex(SkeletonConstantIndex))
+			if (Program.ConstantSkeletons.IsValidIndex(SkeletonConstantIndex))
 			{
-				Result->Skeleton = ConstantSkeletons[SkeletonConstantIndex];
+				Result->Skeleton = Program.ConstantSkeletons[SkeletonConstantIndex];
 
 				Result->ConvertBoneIdsToIndices();
 			}
@@ -790,143 +899,22 @@
 			
 			OutMesh = Result;
 		}
+    	
+    	FProgram& GetProgram()
+    	{
+    		return Program;
+    	}
 
-		void GetExtensionDataConstant(int32 ConstantIndex, TManagedPtr<const FExtensionData>& Result) const;
-
-        UE_API EOpType GetOpType(OP::ADDRESS At) const;
-    	
-        template<typename ARGS>
-		ARGS GetOpArgs(OP::ADDRESS at) const
-		{
-			ARGS result;
-			uint64 byteCodeAddress = OpAddress[at];
-			byteCodeAddress += sizeof(EOpType);
-			FMemory::Memcpy(&result, &ByteCode[byteCodeAddress], sizeof(ARGS));
-			return result;
-		}
-
-		template<typename ARGS>
-		void SetOpArgs(OP::ADDRESS At, const ARGS& Args)
-		{
-			uint64 byteCodeAddress = OpAddress[At];
-			byteCodeAddress += sizeof(EOpType);
-			FMemory::Memcpy(&ByteCode[byteCodeAddress], &Args, sizeof(ARGS));
-		}
-
-        UE_API const uint8* GetOpArgsPointer(OP::ADDRESS At) const;
-    	
-        uint8* GetOpArgsPointer(OP::ADDRESS At);
-
-        void MarkRomUsed(int32 RomIndex) const;
-    };
-
-	
-    /** A Model represents a customisable object with any number of parameters.
-    * When values are given to the parameters, specific Instances can be built, which hold the
-    * built application-usable data.
-	*/
-    class FModel
-	{
-	public:
-		static UE_API void Serialise(const FModel*, FOutputArchive&);
-		static UE_API TSharedPtr<FModel> StaticUnserialise(FInputArchive&);
-    	
-    	void Serialise(FOutputArchive& Arch) const;
-
-		void Unserialise(FInputArchive& Arch);
-
-		/** Special serialise operation that serialises the data in separate "files". An object
-         * with the ModelStreamer interface is responsible of storing this data and providing
-		 * the "file" concept. 
-		 * If bDropData is set to true, the rom data will be freed as it is serialized.
-		 */
-        static UE_API void Serialise(FModel* Model, FModelWriter& Streamer, bool bDropData = false);
-
-		//! Return true if the model has external data in other files. This kind of models will
-		//! require data streaming when used.
-		UE_API bool HasExternalData() const;
-
-#if WITH_EDITOR
-		//! Return true unless the streamed resources were destroyed, which could happen in the
-		//! editor after recompiling the CO.
-		UE_API bool IsValid() const;
-
-		//! Invalidate the Model. Compiling a compiled CO will invalidate the model kept by previously
-		//! generated resources, like streamed textures.
-		UE_API void Invalidate();
-#endif
-
-		//-----------------------------------------------------------------------------------------
-		// Own interface
-		//-----------------------------------------------------------------------------------------
-
-		//! Create a set of new parameters of the model with the default values.
-		//! If old parameters are provided, they will be reused when possible instead of the
-		//! default values.
-        static UE_API TSharedPtr<FParameters> NewParameters(TSharedPtr<const FModel> Model, const FParameters* OldParameters = nullptr);
-
-		/** Return true if the parameter is multi-dimensional */
-		UE_API bool IsParameterMultidimensional(int32 ParameterIndex) const;
-
-		//! Get the number of states in the model.
-		UE_API int32 GetStateCount() const;
-
-		//! Get a state name by state index from 0 to GetStateCount-1
-		UE_API const FString& GetStateName( int32 StateIndex ) const;
-
-		//! Find a state index by state name
-		UE_API int32 FindState( const FString& Name ) const;
-
-		//! Get the number of parameters available in a particular state.
-		UE_API int32 GetStateParameterCount( int32 StateIndex ) const;
-
-		//! Get the index of one of the parameters in the given state. The index refers to the
-		//! parameters in a FParameters object obtained from this model with NewParameters.
-		UE_API int32 GetStateParameterIndex( int32 StateIndex, int32 ParamIndex ) const;
-
-        //! Free memory used by streaming data that may be loaded again when needed.
-        UE_API void UnloadExternalData();
-    	
-		//! Return the default value of a boolean parameter.
-		//! \pre The parameter specified by index is a T_BOOL.
-		//! \param Index Index of the parameter from 0 to GetCount()-1
-    	UE_API bool GetBoolDefaultValue(int32 Index) const;
-
-   		//! Return the default value of a integer parameter.
-		//! \pre The parameter specified by index is a T_INT.
-		//! \param Index Index of the parameter from 0 to GetCount()-1
-    	UE_API int32 GetIntDefaultValue(int32 Index) const;
-
-		//! Return the default value of a float parameter.
-		//! \pre The parameter specified by index is a T_FLOAT.
-		//! \param Index Index of the parameter from 0 to GetCount()-1
-        UE_API float GetFloatDefaultValue(int32 Index) const;
-
-		//! Return the default value of a colour parameter.
-		//! \pre The parameter specified by index is a T_FLOAT.
-        //! \param Index Index of the parameter from 0 to GetCount()-1
-        //! \param R,G,B Pointers to values where every resulting colour channel will be stored
-    	UE_API void GetColourDefaultValue(int32 Index, FVector4f& OutValue) const;
-
-		//! Return the default value of a colour parameter.
-		//! \pre The parameter specified by index is a T_FLOAT.
-        //! \param Index Index of the parameter from 0 to GetCount()-1
-    	UE_API FMatrix44f GetMatrixDefaultValue(int32 Index) const;
-
-    	//! Return the default value of a projector parameter, as a 4x4 matrix. The matrix is supposed to be
-		//! a linear transform in column-major.
-		//! \pre The parameter specified by index is a T_PROJECTOR.
-        //! \param Index Index of the parameter from 0 to GetCount()-1
-        //! \param OutPos Pointer to where the object-space position coordinates of the projector will be stored.
-        //! \param OutDir Pointer to where the object-space direction vector of the projector will be stored.
-        //! \param OutUp Pointer to where the object-space vertically up direction vector
-        //!         of the projector will be stored. This controls the "roll" angle of the
-        //!         projector.
-        //! \param OutScale Pointer to the projector-space scaling of the projector.
-    	UE_API void GetProjectorDefaultValue(int32 Index, EProjectorType* OutProjectionType, FVector3f* OutPos,
-    	 	FVector3f* OutDir, FVector3f* OutUp, FVector3f* OutScale, float* OutProjectionAngle) const;
-    	
+    	const FProgram& GetProgram() const
+    	{
+    		return Program;
+    	}
+    
+    private:
     	FProgram Program;
+    	
+    public:
+    	FRomManager RomManager;
 	};
 
 }
diff --git a/Engine/Plugins/Mutable/Source/MutableRuntime/Internal/MuR/PhysicsBody.h b/Engine/Plugins/Mutable/Source/MutableRuntime/Internal/MuR/PhysicsBody.h
--- a/Engine/Plugins/Mutable/Source/MutableRuntime/Internal/MuR/PhysicsBody.h
+++ b/Engine/Plugins/Mutable/Source/MutableRuntime/Internal/MuR/PhysicsBody.h
@@ -261,8 +261,7 @@ namespace UE::Mutable::Private
 		UE_API const FString& GetSphylName(int32 B, int32 I) const;
 		UE_API const FString& GetTaperedCapsuleName(int32 B, int32 I) const;
 
-		void ConvertBoneIdsToNames(const TMap<uint32, FName>& ConstantNames);
-		void ConvertNamesToBoneIds(const TMap<uint32, FName>& ConstantNames);
+		UE_API void ConvertBoneIdsToNames(const TMap<uint32, FName>& ConstantNames);
 
 	public:
 		int32 CustomId = -1;
diff --git a/Engine/Plugins/Mutable/Source/MutableRuntime/Internal/MuR/RomManager.cpp b/Engine/Plugins/Mutable/Source/MutableRuntime/Internal/MuR/RomManager.cpp
new file mode 100644
--- /dev/null
+++ b/Engine/Plugins/Mutable/Source/MutableRuntime/Internal/MuR/RomManager.cpp
@@ -0,0 +1,343 @@
+// Copyright Epic Games, Inc. All Rights Reserved.
+
+#include "RomManager.h"
+
+#include "Model.h"
+
+#define UE_API MUTABLERUNTIME_API
+
+
+namespace UE::Mutable::Private
+{
+	template UE_API Tasks::TTask<TManagedPtr<const FImage>> FRomManager::GetRom<FImage>(uint32 RomIndex, TSharedRef<FModelReader> ModelReader);
+	template UE_API Tasks::TTask<TManagedPtr<const FMesh>> FRomManager::GetRom<FMesh>(uint32 RomIndex, TSharedRef<FModelReader> ModelReader);
+
+	template UE_API void FRomManager::SetRom<FImage>(uint32 RomIndex, const TManagedPtr<const FImage>& Task);
+	template UE_API void FRomManager::SetRom<FMesh>(uint32 RomIndex, const TManagedPtr<const FMesh>& Task);
+
+
+	std::atomic<SSIZE_T>& MemoryCounters::FStreamingMemoryCounter::Get()
+	{
+		static std::atomic<SSIZE_T> Counter{0};
+		return Counter;
+	}
+	
+	
+	bool FRomManager::FRomWeight::operator<(const FRomWeight& Other) const
+	{
+		return RomTick < Other.RomTick;
+	}
+
+
+	int64 FRomManager::EnsureBudgetBelow(int64 BytesToFree)
+	{
+		MUTABLE_CPUPROFILER_SCOPE(FRomManager::EnsureBudgetBelow);
+	
+		TScopeLock<FMutex> Lock(CriticalSection);
+
+		int64 FreedBytes = 0;
+		
+		while (true)
+		{
+			MUTABLE_CPUPROFILER_SCOPE(EnsureBudgetBelow_UnloadRom);
+
+			if (FreedBytes >= BytesToFree || !RomWeights.Num())
+			{
+				break;
+			}
+			
+			FRomWeight RomWeight;
+			RomWeights.HeapPop(RomWeight, EAllowShrinking::No);
+
+			switch (static_cast<ERomDataType>(Model.GetProgram().Roms[RomWeight.RomIndex].ResourceType))
+			{
+			case ERomDataType::Image:
+			{
+				Tasks::TTask<TManagedPtr<const FImage>>* Task = ConstantImageLODsStreamed.Find(RomWeight.RomIndex);
+				check(Task);
+					
+				if (Task->IsCompleted())
+				{
+					const TManagedPtr<const FImage>& Data = Task->GetResult();
+					if (Data.IsUniqueReference())
+					{
+						FreedBytes += Data->GetDataSize();
+					}
+				}
+				
+				// Released even when we do not know if it is really destroyed.
+				ConstantImageLODsStreamed.Remove(RomWeight.RomIndex);
+					
+				break;
+			}
+			
+			case ERomDataType::Mesh:
+			{
+				Tasks::TTask<TManagedPtr<const FMesh>>* Task = ConstantMeshesStreamed.Find(RomWeight.RomIndex);
+				check(Task);
+
+				if (Task->IsCompleted())
+				{
+					const TManagedPtr<const FMesh>& Data = Task->GetResult();
+					if (Data.IsUniqueReference())
+					{
+						FreedBytes += Data->GetDataSize();
+					}
+				}
+
+				ConstantMeshesStreamed.Remove(RomWeight.RomIndex);
+					
+				break;
+			}
+			
+			default:
+				check(false);
+				break;
+			}
+		}
+		
+		return FreedBytes;
+	}
+
+	
+	int64 FRomManager::GetRomBytes() const
+	{
+		TScopeLock<FMutex> Lock(CriticalSection);
+		
+		int64 Result = 0;
+		
+		// Count streamable and currently-loaded resources
+		for (TPair<uint32, Tasks::TTask<TManagedPtr<const FImage>>>& Pair : ConstantImageLODsStreamed)
+		{
+			if (TManagedPtr<const FImage> Rom = Pair.Value.GetResult())
+			{
+				Result += Rom->GetDataSize();
+			}
+		}
+		
+		for (TPair<uint32, Tasks::TTask<TManagedPtr<const FMesh>>>& Pair : ConstantMeshesStreamed)
+		{
+			if (TManagedPtr<const FMesh> Rom = Pair.Value.GetResult())
+			{
+				Result += Rom->GetDataSize();
+			}
+		}
+		
+		return Result;
+	}
+
+	
+	uint32 FRomManager::GetRomDiskSize(int32 RomIndex) const
+	{
+		return Model.GetProgram().Roms[RomIndex].Size;
+	}
+
+
+	FRomManager::FRomManager(FModel& InModel) :
+		Model(InModel)
+	{
+	}
+
+	
+	void FRomManager::Serialise(FModelWriter& Streamer, bool bDropData)
+	{
+		FOutputMemoryStream MemStream(16 * 1024 * 1024);
+
+		// Save images and unload from memory
+		for (TPair<uint32, Tasks::TTask<TManagedPtr<const FImage>>>& Entry: ConstantImageLODsStreamed)
+		{
+			int32 RomIndex = Entry.Key;
+			const FRomDataRuntime& RomData = Model.GetProgram().Roms[RomIndex];
+			check(RomData.ResourceType == uint32(ERomDataType::Image));
+
+			// Serialize to memory, to find out final size of this rom
+			MemStream.Reset();
+			FOutputArchive MemoryArch(&MemStream);
+			FImage::Serialise(Entry.Value.GetResult().Get(), MemoryArch);
+			check(RomData.Size == MemStream.GetBufferSize());
+
+			Streamer.OpenWriteFile(RomIndex, true);
+			Streamer.Write(MemStream.GetBuffer(), MemStream.GetBufferSize());
+			Streamer.CloseWriteFile();
+
+			// Do this progressively to avoid duplicating all data in memory.
+			if (bDropData)
+			{
+				Entry = {};
+			}
+		}
+
+		if (bDropData)
+		{
+			ConstantImageLODsStreamed.Empty(0);
+		}
+
+		// Save meshes and unload from memory
+		for (TPair<uint32, Tasks::TTask<TManagedPtr<const FMesh>>>& Entry : ConstantMeshesStreamed)
+		{
+			int32 RomIndex = Entry.Key;
+			const FRomDataRuntime& RomData = Model.GetProgram().Roms[RomIndex];
+			check(RomData.ResourceType == uint32(ERomDataType::Mesh));
+
+			// Serialize to memory, to find out final size of this rom
+			MemStream.Reset();
+			FOutputArchive MemoryArch(&MemStream);
+			FMesh::Serialise(Entry.Value.GetResult().Get(), MemoryArch);
+			check(RomData.Size == MemStream.GetBufferSize());
+
+			Streamer.OpenWriteFile(RomIndex, true);
+			Streamer.Write(MemStream.GetBuffer(), MemStream.GetBufferSize());
+			Streamer.CloseWriteFile();
+
+			// Do this progressively to avoid duplicating all data in memory.
+			if (bDropData)
+			{
+				Entry = {};
+			}
+		}
+
+		if (bDropData)
+		{
+			ConstantMeshesStreamed.Empty(0);
+		}
+	}
+
+	
+	void FRomManager::UnloadRoms()
+	{
+		TScopeLock<FMutex> Lock(CriticalSection);
+
+		ConstantImageLODsStreamed.Empty();
+		ConstantMeshesStreamed.Empty();
+		RomWeights.Empty();
+	}
+
+
+	template <typename Type>
+	void FRomManager::SetRom(uint32 RomIndex, const TManagedPtr<const Type>& Rom)
+	{
+		TScopeLock<FMutex> Lock(CriticalSection);
+
+		if constexpr (std::is_same<Type, FImage>::value)
+		{
+			check(!ConstantImageLODsStreamed.Contains(RomIndex));
+			ConstantImageLODsStreamed.Add(RomIndex, Tasks::MakeCompletedTask<TManagedPtr<const FImage>>(Rom));
+		}
+		else if constexpr (std::is_same<Type, FMesh>::value)
+		{
+			check(!ConstantMeshesStreamed.Contains(RomIndex));
+			ConstantMeshesStreamed.Add(RomIndex, Tasks::MakeCompletedTask<TManagedPtr<const FMesh>>(Rom));
+		}
+		
+		FRomWeight RomWeight;
+		RomWeight.RomIndex = RomIndex;
+		RomWeight.RomTick = FSystem::RomTick++;
+		
+		RomWeights.HeapPush(RomWeight);
+	}
+
+	
+	template <typename Type>
+	Tasks::TTask<TManagedPtr<const Type>> FRomManager::GetRom(uint32 RomIndex, TSharedRef<FModelReader> ModelReader)
+	{
+		TScopeLock<FMutex> Lock(CriticalSection);
+				
+		Tasks::TTask<TManagedPtr<const Type>>* Result = nullptr;
+			
+		if constexpr (std::is_same<Type, FImage>::value)
+		{
+			Result = ConstantImageLODsStreamed.Find(RomIndex);
+		}
+		else if constexpr (std::is_same<Type, FMesh>::value)
+		{
+			Result = ConstantMeshesStreamed.Find(RomIndex);
+		}
+
+		if (Result)
+		{
+			FRomWeight* RomWeight = RomWeights.FindByPredicate([&](const FRomWeight& Element)
+			{
+				return Element.RomIndex == RomIndex;
+			});
+				
+			check(RomWeight);
+			
+			RomWeight->RomTick = FSystem::RomTick++;
+				
+			return *Result;
+		}
+			
+		Tasks::FTaskEvent ReadCompletionEvent(TEXT("FLoadMeshRomsTaskRom"));
+
+		struct FRomLoadOp
+		{
+			using StreamingDataContainerType = TArray<uint8, FDefaultMemoryTrackingAllocator<MemoryCounters::FStreamingMemoryCounter>>;
+
+			FModelReader::FOperationID StreamID = -1;
+			StreamingDataContainerType StreamBuffer;
+		};
+		
+		TSharedRef<FRomLoadOp> RomLoadOp = MakeShared<FRomLoadOp>();
+		const uint32 RomSize = GetRomDiskSize(RomIndex);
+		check(RomSize > 0);
+		
+		RomLoadOp->StreamBuffer.SetNumUninitialized(RomSize);
+
+		TFunction<void(bool)> Callback = [ReadCompletionEvent](bool bSuccess) mutable // Mutable due Trigger not being const
+		{
+			ReadCompletionEvent.Trigger();
+		};
+		
+		EDataType DataType;
+			
+		if constexpr (std::is_same<Type, FImage>::value)
+		{
+			DataType = EDataType::Image;
+		}
+		else if constexpr (std::is_same<Type, FMesh>::value)
+		{
+			DataType = EDataType::Mesh;
+		}
+			
+		RomLoadOp->StreamID = ModelReader->BeginReadBlock(&Model, RomIndex, RomLoadOp->StreamBuffer.GetData(), RomSize, DataType, &Callback);
+		if (RomLoadOp->StreamID < 0)
+		{
+			return Tasks::MakeCompletedTask<TManagedPtr<const Type>>(nullptr);
+		}
+		
+		Tasks::TTask<TManagedPtr<const Type>> Task = Tasks::Launch(UE_SOURCE_LOCATION, [Reader = ModelReader, RomLoadOp]() -> TManagedPtr<const Type>
+		{
+			bool bSuccess = Reader->EndRead(RomLoadOp->StreamID);
+			if (!bSuccess)
+			{
+				return nullptr;
+			}
+
+			FInputMemoryStream Stream(RomLoadOp->StreamBuffer.GetData(), RomLoadOp->StreamBuffer.Num());
+			FInputArchive Arch(&Stream);
+			
+			return Type::StaticUnserialise(Arch);
+		}, 
+		ReadCompletionEvent);
+			
+		FRomWeight RomWeight;
+		RomWeight.RomIndex = RomIndex;
+		RomWeight.RomTick = FSystem::RomTick++;
+				
+		RomWeights.HeapPush(RomWeight);
+			
+		if constexpr  (std::is_same<Type, FImage>::value)
+		{
+			ConstantImageLODsStreamed.Add(RomIndex, Task);
+		}
+		else if constexpr (std::is_same<Type, FMesh>::value)
+		{
+			ConstantMeshesStreamed.Add(RomIndex, Task);
+		}
+			
+		return Task;
+	}
+}
+
+
+#undef UE_API
diff --git a/Engine/Plugins/Mutable/Source/MutableRuntime/Internal/MuR/RomManager.h b/Engine/Plugins/Mutable/Source/MutableRuntime/Internal/MuR/RomManager.h
new file mode 100644
--- /dev/null
+++ b/Engine/Plugins/Mutable/Source/MutableRuntime/Internal/MuR/RomManager.h
@@ -0,0 +1,75 @@
+﻿// Copyright Epic Games, Inc. All Rights Reserved.
+
+#pragma once
+
+#include "ManagedPointer.h"
+#include "MemoryCounters.h"
+#include "Mesh.h"
+#include "Image.h"
+#include "System.h"
+
+#define UE_API MUTABLERUNTIME_API
+
+
+namespace UE::Mutable::Private
+{
+	namespace MemoryCounters
+	{
+		struct FStreamingMemoryCounter
+		{
+			static MUTABLERUNTIME_API std::atomic<SSIZE_T>& Get();
+		};
+	}
+	
+	struct FConstantResourceIndex;
+	
+	class FRomManager
+	{
+	public:
+		FRomManager(FModel& InModel);
+
+		void Serialise(FModelWriter& Streamer, bool bDropData);
+
+		template<typename Type>
+		void SetRom(uint32 RomIndex, const TManagedPtr<const Type>& Rom);
+
+		template<typename Type>
+		Tasks::TTask<TManagedPtr<const Type>> GetRom(uint32 RomIndex, TSharedRef<FModelReader> ModelReader);
+
+		UE_API void UnloadRoms();
+
+		UE_API int64 EnsureBudgetBelow(int64 BytesToFree);
+		
+		UE_API int64 GetRomBytes() const;
+	
+	private:
+		uint32 GetRomDiskSize(int32 RomIndex) const;
+		
+		const FModel& Model;
+		
+		mutable FMutex CriticalSection;
+		
+		struct FRomWeight
+		{
+			uint32 RomIndex = INDEX_NONE;
+			uint32 RomTick = 0;
+			
+			bool operator<(const FRomWeight& Other) const;
+		};
+		
+		TMemoryTrackedArray<FRomWeight> RomWeights;
+		
+		/** Constant mesh data is split in 2 sets: ConstantMeshesStreamed constains data that is streamed in and out.
+		* Index with FConstantResourceIndex::Index, when Streamable is 1.
+		* This part is empty for an unused Model, and shouldn't be serialised. */
+		mutable TMap<uint32, Tasks::TTask<TManagedPtr<const FMesh>>> ConstantMeshesStreamed;
+		
+		/** Constant image mip data is split in 2 sets: ConstantImageLODsStreamed constains data that is streamed in and out. 
+		* Index with FConstantResourceIndex::Index, when Streamable is 1.
+		* This part is empty for an unused Model, and shouldn't be serialised. */
+		mutable TMap<uint32, Tasks::TTask<TManagedPtr<const FImage>>> ConstantImageLODsStreamed;
+	};
+}
+
+
+#undef UE_API
\ No newline at end of file
diff --git a/Engine/Plugins/Mutable/Source/MutableRuntime/Internal/MuR/System.h b/Engine/Plugins/Mutable/Source/MutableRuntime/Internal/MuR/System.h
--- a/Engine/Plugins/Mutable/Source/MutableRuntime/Internal/MuR/System.h
+++ b/Engine/Plugins/Mutable/Source/MutableRuntime/Internal/MuR/System.h
@@ -1629,7 +1629,7 @@ namespace UE::Mutable::Private
 		int32 State = 0;
 		uint32 LODMask = 0;
 		TManagedPtr<const FInstance> Instance;
-		TSharedPtr<const FModel> Model;
+		TSharedPtr<FModel> Model;
 
 		TSharedPtr<const FParameters> Parameters;
 		
@@ -1669,7 +1669,7 @@ namespace UE::Mutable::Private
         /** Make sure the working memory is below the internal budget, even counting with the passed additional memory. 
 		* Return true if it succeeded, false otherwise.
 		*/
-		bool EnsureBudgetBelow(uint64 AdditionalMemory, const TSharedPtr<FProgramCache>& InCurrentCache);
+		bool EnsureBudgetBelow(int64 AdditionalMemory, const TSharedPtr<FLiveInstance>& InCurrentLiveInstance);
 		
 		/** Return true if the memory budget is 90% full. */
 		bool IsMemoryBudgetFull() const;
@@ -1781,7 +1781,7 @@ namespace UE::Mutable::Private
         //! calls to BeginUpdate/EndUpdate.
         //! A call to NewInstance must be paired to a call to ReleasesInstance when the instance is
         //! no longer needed.
-        UE_API TSharedRef<FLiveInstance> NewInstance(const TSharedPtr<const FModel>& Model,
+        UE_API TSharedRef<FLiveInstance> NewInstance(const TSharedPtr<FModel>& Model,
 		const TSharedPtr<UE::Mutable::Private::FMeshIdRegistry>& InMeshIdRegistry,
 		const TSharedPtr<UE::Mutable::Private::FImageIdRegistry>& InImageIdRegistry,
 		const TSharedPtr<UE::Mutable::Private::FMaterialIdRegistry>& InMaterialIdRegistry,
@@ -1794,7 +1794,7 @@ namespace UE::Mutable::Private
 		 * as the life cycle of the returned object is expected to be short and contained during the compilation operation.
 		 * It does not require further initialization as it is fully capable of operating as is.
 		 */
-		UE_API TSharedRef<FLiveInstance> NewBuildInstance(const TSharedPtr<const FModel>& InModel, 
+		UE_API TSharedRef<FLiveInstance> NewBuildInstance(const TSharedPtr<FModel>& InModel, 
 			const TSharedPtr<const FParameters>& InParameters,
 			const TSharedPtr<FExternalResourceProvider>& ExternalResourceProvider);
 				
@@ -1902,7 +1902,7 @@ namespace UE::Mutable::Private
 		FWorkingMemoryManager WorkingMemoryManager;
 
 		/** Increased each time a rom is loaded. Used to keep track of which roms to evict. */
-		static uint64 RomTick;
+		static std::atomic<uint32> RomTick;
 		
 		MUTABLERUNTIME_API static bool CheckForUpdatedParameters(const TSharedRef<FLiveInstance>& LiveInstance, const TSharedPtr<const FParameters>& InOldParameters);
 		
diff --git a/Engine/Plugins/Mutable/Source/MutableRuntime/MutableRuntime.Build.cs b/Engine/Plugins/Mutable/Source/MutableRuntime/MutableRuntime.Build.cs
--- a/Engine/Plugins/Mutable/Source/MutableRuntime/MutableRuntime.Build.cs
+++ b/Engine/Plugins/Mutable/Source/MutableRuntime/MutableRuntime.Build.cs
@@ -29,7 +29,7 @@ namespace UnrealBuildTool.Rules
 			PrivateDependencyModuleNames.AddRange(
 				new string[] {
 				"GeometryCore",
-				"TraceLog"
+				"TraceLog", 
 				}
 			);
 		}
diff --git a/Engine/Plugins/Mutable/Source/MutableRuntime/Private/MuR/CodeRunner.cpp b/Engine/Plugins/Mutable/Source/MutableRuntime/Private/MuR/CodeRunner.cpp
--- a/Engine/Plugins/Mutable/Source/MutableRuntime/Private/MuR/CodeRunner.cpp
+++ b/Engine/Plugins/Mutable/Source/MutableRuntime/Private/MuR/CodeRunner.cpp
@@ -147,14 +147,6 @@ namespace UE::Mutable::Private::Private
 
 #endif
 
-namespace UE::Mutable::Private::MemoryCounters
-{
-	std::atomic<SSIZE_T>& FStreamingMemoryCounter::Get()
-	{
-		static std::atomic<SSIZE_T> Counter{0};
-		return Counter;
-	}
-}
 
 namespace UE::Mutable::Private
 {
@@ -191,7 +183,7 @@ namespace UE::Mutable::Private
 		, ExecutionStrategy(InExecutionStrategy)
 		, System(InSystem)
 		, Model(*InLiveInstance->Model.Get())
-		, Program(InLiveInstance->Model->Program)
+		, Program(InLiveInstance->Model->GetProgram())
 		, Parameters(*InLiveInstance->Parameters)
 		, LODMask(InLodMask)
 		, ExternalResourceProvider(InLiveInstance->ExternalResourceProvider)
@@ -1087,56 +1079,6 @@ namespace UE::Mutable::Private
 		EOpType type = Program.GetOpType(Item.At);
         switch (type)
         {
-
-        case EOpType::ME_CONSTANT:
-        {
-        	MUTABLE_CPUPROFILER_SCOPE(ME_CONSTANT);
-        		
-			OP::MeshConstantArgs Args = Program.GetOpArgs<OP::MeshConstantArgs>(Item.At);
-			EMeshContentFlags MeshContentFlags = (EMeshContentFlags)Item.ExecutionOptions;
-
-			const TPassthroughObjectPtr<UClothingAssetCommon> ClothAsset = FindPassthrough<UClothingAssetCommon>(Args.ClothID);
-        		
-			TManagedPtr<const FMesh> Source;
-			Program.GetConstant(Args.Value, Args.Skeleton, ClothAsset, Source, MeshContentFlags, 
-					[this](int32 BudgetReserve){ return CreateMesh(BudgetReserve); });
-			
-			if (!Source)
-			{
-				return false;
-			}
-
-            StoreMesh(Item, Source);
-
-			//UE_LOG(LogMutableCore, Log, TEXT("Set mesh constant %d."), Item.At);
-			break;
-        }
-
-        case EOpType::IM_CONSTANT:
-        {
-        	MUTABLE_CPUPROFILER_SCOPE(IM_CONSTANT);
-        		
-			OP::ResourceConstantArgs Args = Program.GetOpArgs<OP::ResourceConstantArgs>(Item.At);
-
-			int32 MipsToSkip = Item.ExecutionOptions;
-            TManagedPtr<const FImage> Source;
-			Program.GetConstant(Args.value, Source, MipsToSkip, 
-			[this](int32 SizeX, int32 SizeY, int32 NumLODs, EImageFormat Format, EInitializationType InitPolicy)
-			{
-				return CreateImage(SizeX, SizeY, NumLODs, Format, InitPolicy);
-			});
-
-			// Assume the ROM has been loaded previously in a task generated at IssueOp
-			if (!Source)
-			{
-				return false;
-			}
-
-            StoreImage(Item, Source);
-			//UE_LOG(LogMutableCore, Log, TEXT("Set image constant %d."), Item.At);
-			break;
-        }
-
 		case EOpType::ED_CONSTANT:
 		{
         	MUTABLE_CPUPROFILER_SCOPE(ED_CONSTANT);
@@ -5808,7 +5750,7 @@ namespace UE::Mutable::Private
             return nullptr;
         }
 
-		const FProgram& Program = InModel.Program;
+		const FProgram& Program = InModel.GetProgram();
 		const FParameterDesc& paramDesc = Program.Parameters[ ParameterIndex ];
         for( int32 rangeIndexInParam = 0;
              rangeIndexInParam<paramDesc.Ranges.Num();
@@ -7175,8 +7117,6 @@ namespace UE::Mutable::Private
             RunCode_Conditional(Item);
             break;
 
-        case EOpType::ME_CONSTANT:
-		case EOpType::IM_CONSTANT:
 		case EOpType::ED_CONSTANT:
 		case EOpType::MI_CONSTANT:
             RunCode_ConstantResource(Item);
diff --git a/Engine/Plugins/Mutable/Source/MutableRuntime/Private/MuR/CodeRunner.h b/Engine/Plugins/Mutable/Source/MutableRuntime/Private/MuR/CodeRunner.h
--- a/Engine/Plugins/Mutable/Source/MutableRuntime/Private/MuR/CodeRunner.h
+++ b/Engine/Plugins/Mutable/Source/MutableRuntime/Private/MuR/CodeRunner.h
@@ -16,13 +16,6 @@
 #include "Templates/Tuple.h"
 #include "Concepts/SameAs.h"
 
-namespace UE::Mutable::Private::MemoryCounters
-{
-	struct FStreamingMemoryCounter
-	{
-		static MUTABLERUNTIME_API std::atomic<SSIZE_T>& Get();
-	};
-}
 
 namespace  UE::Mutable::Private
 {
@@ -240,16 +233,7 @@ namespace  UE::Mutable::Private
 				return !Event.IsValid() || Event.IsCompleted(); 
 			}
 		};
-
-		struct FRomLoadOp
-		{
-			using StreamingDataContainerType = TArray<uint8, FDefaultMemoryTrackingAllocator<MemoryCounters::FStreamingMemoryCounter>>;
-
-			int32 RomIndex = -1;
-			FModelReader::FOperationID m_streamID = -1;
-			StreamingDataContainerType m_streamBuffer;
-			UE::Tasks::FTask Event;
-		};
+    	
 
 		class FLoadMeshRomTask : public CodeRunner::FIssuedTask
 		{
@@ -271,12 +255,12 @@ namespace  UE::Mutable::Private
 			virtual bool Complete(CodeRunner*) override;
 
 		private:
-
 			int32 FirstIndex = -1;
 			EMeshContentFlags RomContentFlags = EMeshContentFlags::None;
 			EMeshContentFlags ExecutionContentFlags = EMeshContentFlags::None;
 
-			TArray<int32, TInlineAllocator<4>> RomIndices;
+			/** Key is the rom index. Value is the future of the streamed rom. */
+			TMap<int32, Tasks::TTask<TManagedPtr<const FMesh>>> RomsStreamed;
 		};
 
 		class FLoadImageRomsTask : public CodeRunner::FIssuedTask
@@ -297,7 +281,7 @@ namespace  UE::Mutable::Private
  			int32 ImageIndexBegin = -1;
 			int32 ImageIndexEnd = -1;
 
-			TArray<int32> RomIndices;
+			TMap<int32, Tasks::TTask<TManagedPtr<const FImage>>> RomsStreamed;
 		};
 
 		template <
@@ -337,14 +321,7 @@ namespace  UE::Mutable::Private
 				ScheduledStagePerOp[Op] = Op.Stage + 1;
 				AddChildren(InDeps);
 			}
-
 		}
-
-    	/** Calculate an approximation of memory used by streaming buffers in this class. */
-		int32 GetStreamingMemoryBytes() const
-    	{
-			return RomLoadOps.GetAllocatedSize();			
-    	}
     	
 		/** Calculate an approximation of memory used by manging structures in this class. */
 		int32 GetInternalMemoryBytes() const
@@ -613,7 +590,7 @@ namespace  UE::Mutable::Private
 		FSystem& System;
     	
     	// Dereference optimization. LiveInstance keeps them alive.
-		const FModel& Model;
+		FModel& Model;
 		const FProgram& Program;
 		const FParameters& Parameters;
     	
@@ -637,77 +614,9 @@ namespace  UE::Mutable::Private
     	
     	bool EnsureBudgetBelow( uint64 AdditionalMemory ) const
 	    {
-    		return System.WorkingMemoryManager.EnsureBudgetBelow(AdditionalMemory, LiveInstance->Cache);
+    		return System.WorkingMemoryManager.EnsureBudgetBelow(AdditionalMemory, LiveInstance);
     	}
     	
-    	
-    	struct FRomLoadOps
-    	{
-   		private:
-    		/** Rom read operations already in progress. */
-    		TArray<FRomLoadOp> RomLoadOps;
-
-			CodeRunner* Runner;
-
-    	public:
-
-			FRomLoadOps(CodeRunner& InRunner)
-			{
-				Runner = &InRunner;
-			}
-
-    		FRomLoadOp* Find(const int32 RomIndex)
-    		{
-    			for (FRomLoadOp& RomLoadOp : RomLoadOps)
-    			{
-    				if (RomLoadOp.RomIndex == RomIndex)
-    				{
-    					return &RomLoadOp;
-    				}
-    			}
-
-    			return nullptr;
-    		}
-    	
-    		FRomLoadOp& Create(const int32 RomIndex)
-    		{
-    			for (FRomLoadOp& RomLoadOp : RomLoadOps)
-    			{
-    				if (RomLoadOp.RomIndex == -1)
-    				{
-    					RomLoadOp.RomIndex = RomIndex;
-    					return RomLoadOp;
-    				}
-    			}
-
-    			FRomLoadOp& RomLoadOp = RomLoadOps.AddDefaulted_GetRef();
-    			RomLoadOp.RomIndex = RomIndex;
-
-    			return RomLoadOp;
-    		}
-
-    		void Remove(FRomLoadOp& RomLoadOp)
-    		{
-    			RomLoadOp.RomIndex = -1;
-                RomLoadOp.m_streamBuffer.Empty();
-                RomLoadOp.Event = {};
-            }
-
-    		int32 GetAllocatedSize() const
-    		{
-    			int32 Result = 0;
-			
-    			for (const FRomLoadOp& RomLoadOp : RomLoadOps)
-    			{
-    				Result += RomLoadOp.m_streamBuffer.GetAllocatedSize();
-    			}
-
-    			return Result;
-    		}
-    	};
-    	
-    	FRomLoadOps RomLoadOps = FRomLoadOps(*this);
-    	    	
 		void AddChildren(TArrayView<const FScheduledOp> Deps);
 		//void AddChildren(const FScheduledOp& Deps);
 
diff --git a/Engine/Plugins/Mutable/Source/MutableRuntime/Private/MuR/CodeRunnerBegin.cpp b/Engine/Plugins/Mutable/Source/MutableRuntime/Private/MuR/CodeRunnerBegin.cpp
--- a/Engine/Plugins/Mutable/Source/MutableRuntime/Private/MuR/CodeRunnerBegin.cpp
+++ b/Engine/Plugins/Mutable/Source/MutableRuntime/Private/MuR/CodeRunnerBegin.cpp
@@ -17,9 +17,9 @@ namespace UE::Mutable::Private
 			uint32 InLodMask) :
 		Model(InLiveInstance->Model),
 		Params(InLiveInstance->Parameters),
-		Program(InLiveInstance->Model->Program),
+		Program(InLiveInstance->Model->GetProgram()),
 		LodMask(InLodMask),
-		Executed(InLiveInstance->Model->Program.OpAddress.Num()),
+		Executed(InLiveInstance->Model->GetProgram().OpAddress.Num()),
 		ProgramCache(InLiveInstance->Cache),
 		MeshIdRegistry(InLiveInstance->MeshIdRegistry),
 		ImageIdRegistry(InLiveInstance->ImageIdRegistry),
@@ -1762,7 +1762,7 @@ namespace UE::Mutable::Private
 					{
 						OP::ResourceConstantArgs Args = Program.GetOpArgs<OP::ResourceConstantArgs>(Item.At);
 				
-						check(Args.value < static_cast<uint32>(Model->Program.ConstantStrings.Num()));
+						check(Args.value < static_cast<uint32>(Model->GetProgram().ConstantStrings.Num()));
 
 						const FString& Result = Program.ConstantStrings[Args.value];
 
@@ -2395,7 +2395,7 @@ if (OpArgs.SkeletonId != PASSTHROUGH_ID_INVALID)
 							int32 NumIterations = 0;
 							if (Args.rangeSize)
 							{
-								const EDataType RangeSizeType = GetOpDataType(Model->Program.GetOpType(Args.rangeSize) );
+								const EDataType RangeSizeType = GetOpDataType(Model->GetProgram().GetOpType(Args.rangeSize) );
 								if (RangeSizeType == EDataType::Int)
 								{
 									NumIterations = LoadInt(Args.rangeSize);
diff --git a/Engine/Plugins/Mutable/Source/MutableRuntime/Private/MuR/CodeRunnerTasked.cpp b/Engine/Plugins/Mutable/Source/MutableRuntime/Private/MuR/CodeRunnerTasked.cpp
--- a/Engine/Plugins/Mutable/Source/MutableRuntime/Private/MuR/CodeRunnerTasked.cpp
+++ b/Engine/Plugins/Mutable/Source/MutableRuntime/Private/MuR/CodeRunnerTasked.cpp
@@ -41,6 +41,7 @@
 #include "ProfilingDebugging/CountersTrace.h"
 #include "Async/Fundamental/Scheduler.h"
 #include "MuR/SkeletalMesh.h"
+#include "ClothingAsset.h"
 
 #if MUTABLE_DEBUG_CODERUNNER_TASK_SCHEDULE_CALLSTACK
 #include "GenericPlatform/GenericPlatformStackWalk.h"
@@ -320,7 +321,7 @@ namespace UE::Mutable::Private
 				if (ProfileContext)
 				{
 					++ProfileContext->NumRunOps;
-					++ProfileContext->RunOpsPerType[int32(LiveInstance->Model->Program.GetOpType(Item.At))];
+					++ProfileContext->RunOpsPerType[int32(LiveInstance->Model->GetProgram().GetOpType(Item.At))];
 				}
 			}
 
@@ -2006,8 +2007,7 @@ namespace UE::Mutable::Private
 
 		for (const int32 MeshContentRomIndex : RomsToLoad)
 		{
-			FConstantResourceIndex CurrentIndex = 
-					LocalProgram.ConstantMeshContentIndices[MeshContentRomIndex + FirstIndex];
+			FConstantResourceIndex CurrentIndex = LocalProgram.ConstantMeshContentIndices[MeshContentRomIndex + FirstIndex];
 
 			if (!CurrentIndex.Streamable)
 			{
@@ -2015,75 +2015,26 @@ namespace UE::Mutable::Private
 				continue;
 			}
 
-			int32 RomIndex = CurrentIndex.Index;
+			const int32 RomIndex = CurrentIndex.Index;
 
-			check(RomIndex < LocalProgram.Roms.Num());
-
-			++LocalProgram.PendingOpsPerRom[RomIndex];
-
-			if (DebugRom && (DebugRomAll || RomIndex == DebugRomIndex))
-			{
-				UE_LOG(LogMutableCore, Log, TEXT("Preparing rom %d, now peding ops is %d."), RomIndex, LocalProgram.PendingOpsPerRom[RomIndex]);
-			}
-
-			if (LocalProgram.IsRomLoaded(RomIndex))
-			{
-				continue;
-			}
-
-			RomIndices.Add(RomIndex);
-
-			if (const FRomLoadOp* Result = Runner->RomLoadOps.Find(RomIndex))
-			{
-				ReadCompleteEvents.Add(Result->Event); // Wait for the read operation started by other task
-				continue;
-			}
-
-			FRomLoadOp& RomLoadOp = Runner->RomLoadOps.Create(RomIndex);
-			
-			check(Runner->System.StreamInterface);
+			Tasks::TTask<TManagedPtr<const FMesh>> Task = Runner->Model.RomManager.GetRom<FMesh>(RomIndex, Runner->System.StreamInterface.ToSharedRef());
+			ReadCompleteEvents.Add(Task);
+			RomsStreamed.Add(RomIndex, Task);
 
 			const uint32 RomSize = LocalProgram.Roms[RomIndex].Size;
 			check(RomSize > 0);
 
-			// Free roms if necessary
-			{
-				LocalProgram.MarkRomUsed(RomIndex);
-				Runner->EnsureBudgetBelow(RomSize);
-			}
-
-			const int32 SizeBefore = RomLoadOp.m_streamBuffer.GetAllocatedSize();
-			RomLoadOp.m_streamBuffer.SetNumUninitialized(RomSize);
-			const int32 SizeAfter = RomLoadOp.m_streamBuffer.GetAllocatedSize();
-
-			UE::Tasks::FTaskEvent ReadCompletionEvent(TEXT("FLoadMeshRomsTaskRom"));
-			ReadCompleteEvents.Add(ReadCompletionEvent);
-			RomLoadOp.Event = ReadCompletionEvent;
-
-			TFunction<void(bool)> Callback = [ReadCompletionEvent](bool bSuccess) mutable // Mutable due Trigger not being const
-			{
-				ReadCompletionEvent.Trigger();
-			};
-
-		
-			const uint32 RomId = RomIndex;
-			RomLoadOp.m_streamID = Runner->System.StreamInterface->BeginReadBlock(&Runner->Model, RomId, RomLoadOp.m_streamBuffer.GetData(), RomSize, EDataType::Mesh, &Callback);
-			if (RomLoadOp.m_streamID < 0)
-			{
-				bOutFailed = true;
-				return false;
-			}
+			Runner->EnsureBudgetBelow(RomSize);
 		}
 
 		// Wait for all read operations to end
-		UE::Tasks::FTaskEvent GatherReadsCompletionEvent(TEXT("FLoadMeshRomsTask"));
+		Tasks::FTaskEvent GatherReadsCompletionEvent(TEXT("FLoadMeshRomsTask"));
 		GatherReadsCompletionEvent.AddPrerequisites(ReadCompleteEvents);
 		GatherReadsCompletionEvent.Trigger();
 
 		Event = GatherReadsCompletionEvent;
 			
 		return false; // No worker thread work
-
 	}	
 
 	//---------------------------------------------------------------------------------------------
@@ -2092,62 +2043,31 @@ namespace UE::Mutable::Private
 		// This runs in the Runner thread
 		MUTABLE_CPUPROFILER_SCOPE(FLoadMeshRomTask_Complete);
 		
-		const FProgram& LocalProgram = Runner->Program;
+		OP::MeshConstantArgs Args = Runner->Program.GetOpArgs<OP::MeshConstantArgs>(Op.At);
 		
-		bool bSomeMissingData = false;
-		for (const int32 RomIndex : RomIndices)
-		{
-			// Since task could be reordered, we need to make sure we end the rom read before continuing
-			if (FRomLoadOp* RomLoadOp = Runner->RomLoadOps.Find(RomIndex))
-			{				
-				bool bSuccess = Runner->System.StreamInterface->EndRead(RomLoadOp->m_streamID);
-
-				MUTABLE_CPUPROFILER_SCOPE(Unserialise);
+		const TPassthroughObjectPtr<UClothingAssetCommon> ClothAsset = Runner->FindPassthrough<UClothingAssetCommon>(Args.ClothID);
 
-				if (bSuccess)
-				{
-					FInputMemoryStream Stream(RomLoadOp->m_streamBuffer.GetData(), RomLoadOp->m_streamBuffer.Num());
-					FInputArchive Arch(&Stream);
-
-					TManagedPtr<FMesh> Value = FMesh::StaticUnserialise(Arch);
-
-					LocalProgram.SetMeshRomValue(RomIndex, Value);
-				}
-				else
-				{
-					bSomeMissingData = true;
-				}
-				
-				Runner->RomLoadOps.Remove(*RomLoadOp);
-			}
-		}
-
-		if (bSomeMissingData)
+		const EMeshContentFlags MeshContentFlags = static_cast<EMeshContentFlags>(Op.ExecutionOptions);
+		TManagedPtr<const FMesh> Source;
+		Runner->Model.GetConstant(Args.Value, Args.Skeleton, ClothAsset, Source, MeshContentFlags,
+		[this](int32 RomIndex)
 		{
-			// Some data may be missing. We can try to go on if some requested mesh parts are there. 
-			UE_LOG(LogMutableCore, Verbose, TEXT("FLoadMeshRomsTask::Complete failed: missing data?"));
-		}
-
-		// Process the constant op normally, now that the rom is loaded.	
-		bool bSuccess = Runner->RunCode_ConstantResource(Op);
-
-		if (bSuccess)
+			Tasks::TTask<TManagedPtr<const FMesh>>* Found = RomsStreamed.Find(RomIndex);
+			return Found ? Found->GetResult() : nullptr;
+		},
+		[Runner](int32 BudgetReserve)
 		{
-			for (int32 RomIndex : RomIndices)
-			{
-				check(RomIndex < LocalProgram.Roms.Num());
-
-				Runner->Model.Program.MarkRomUsed(RomIndex);
-				--LocalProgram.PendingOpsPerRom[RomIndex];
-
-				if (DebugRom && (DebugRomAll || RomIndex == DebugRomIndex))
-				{
-					UE_LOG(LogMutableCore, Log, TEXT("FLoadMeshRomsTask::Complete rom %d, now peding ops is %d."), RomIndex, LocalProgram.PendingOpsPerRom[RomIndex]);
-				}
-			}
+			return Runner->CreateMesh(BudgetReserve);
+		});
+			
+		if (!Source)
+		{
+			return false;
 		}
 
-		return bSuccess;
+		Runner->StoreMesh(Op, Source);
+
+		return true;
 	}
 	
 	//---------------------------------------------------------------------------------------------
@@ -2181,63 +2101,19 @@ namespace UE::Mutable::Private
 			}
 
 			int32 RomIndex = ImageConstantResourceIndex.Index;
-			check(RomIndex < LocalProgram.Roms.Num());
-
-			++LocalProgram.PendingOpsPerRom[RomIndex];
-
-			if (DebugRom && (DebugRomAll || RomIndex == DebugRomIndex))
-				UE_LOG(LogMutableCore, Log, TEXT("Preparing rom %d, now peding ops is %d."), RomIndex, LocalProgram.PendingOpsPerRom[RomIndex]);
-
-			if (LocalProgram.IsRomLoaded(RomIndex))
-			{
-				continue;
-			}
-
-			RomIndices.Add(RomIndex);
-
-			if (const FRomLoadOp* Result = Runner->RomLoadOps.Find(RomIndex))
-			{
-				ReadCompleteEvents.Add(Result->Event); // Wait for the read operation started by other task
-				continue;
-			}
-
-			FRomLoadOp& RomLoadOp = Runner->RomLoadOps.Create(RomIndex);
 			
-			check(Runner->System.StreamInterface);
-
+			Tasks::TTask<TManagedPtr<const FImage>> Task = Runner->Model.RomManager.GetRom<FImage>(RomIndex, Runner->System.StreamInterface.ToSharedRef());
+			RomsStreamed.Add(RomIndex, Task);
+			ReadCompleteEvents.Add(Task);
+			
 			const uint32 RomSize = LocalProgram.Roms[RomIndex].Size;
 			check(RomSize > 0);
 
-			// Free roms if necessary
-			{
-				Runner->Model.Program.MarkRomUsed(RomIndex);
-				Runner->EnsureBudgetBelow(RomSize);
-			}
-
-			const int32 SizeBefore = RomLoadOp.m_streamBuffer.GetAllocatedSize();
-			RomLoadOp.m_streamBuffer.SetNumUninitialized(RomSize);
-			const int32 SizeAfter = RomLoadOp.m_streamBuffer.GetAllocatedSize();
-
-			UE::Tasks::FTaskEvent ReadCompletionEvent(TEXT("FLoadImageRomsTaskRom"));
-			ReadCompleteEvents.Add(ReadCompletionEvent);
-			RomLoadOp.Event = ReadCompletionEvent;
-
-			TFunction<void(bool)> Callback = [ReadCompletionEvent](bool bSuccess) mutable // Mutable due Trigger not being const
-			{
-				ReadCompletionEvent.Trigger();
-			};
-			
-			const uint32 RomId = RomIndex;
-			RomLoadOp.m_streamID = Runner->System.StreamInterface->BeginReadBlock(&Runner->Model, RomId, RomLoadOp.m_streamBuffer.GetData(), RomSize, EDataType::Image, &Callback);
-			if (RomLoadOp.m_streamID < 0)
-			{
-				bOutFailed = true;
-				return false;
-			}
+			Runner->EnsureBudgetBelow(RomSize);
 		}
 
 		// Wait for all read operations to end
-		UE::Tasks::FTaskEvent GatherReadsCompletionEvent(TEXT("FLoadImageRomsTask"));
+		Tasks::FTaskEvent GatherReadsCompletionEvent(TEXT("FLoadImageRomsTask"));
 		GatherReadsCompletionEvent.AddPrerequisites(ReadCompleteEvents);
 		GatherReadsCompletionEvent.Trigger();
 
@@ -2250,80 +2126,30 @@ namespace UE::Mutable::Private
 	//---------------------------------------------------------------------------------------------
 	bool CodeRunner::FLoadImageRomsTask::Complete(CodeRunner* Runner)
 	{
-		// This runs in the Runner thread
-		MUTABLE_CPUPROFILER_SCOPE(FLoadImageRomsTask_Complete);
+		OP::ResourceConstantArgs Args = Runner->Model.GetProgram().GetOpArgs<OP::ResourceConstantArgs>(Op.At);
 
-		if (!Runner)
+		const int32 MipsToSkip = Op.ExecutionOptions;
+		TManagedPtr<const FImage> Source;
+		Runner->Model.GetConstant(Args.value, Source, MipsToSkip,
+		[this](int32 RomIndex)
 		{
-			return false;
-		}
-
-		const FProgram& LocalProgram = Runner->Program;
-		
-		bool bSomeMissingData = false;
-
-		for (const int32 RomIndex : RomIndices)
+			Tasks::TTask<TManagedPtr<const FImage>>* Found = RomsStreamed.Find(RomIndex);
+			return Found ? Found->GetResult() : nullptr;
+		},
+		[Runner](int32 SizeX, int32 SizeY, int32 NumLODs, EImageFormat Format, EInitializationType InitPolicy)
 		{
-			// Since task could be reordered, we need to make sure we end the rom read before continuing
-			if (FRomLoadOp* RomLoadOp = Runner->RomLoadOps.Find(RomIndex))
-			{				
-				bool bSuccess = Runner->System.StreamInterface->EndRead(RomLoadOp->m_streamID);
-
-				MUTABLE_CPUPROFILER_SCOPE(Unserialise);
+			return Runner->CreateImage(SizeX, SizeY, NumLODs, Format, InitPolicy);
+		});
 
-				if (bSuccess)
-				{
-					FInputMemoryStream Stream(RomLoadOp->m_streamBuffer.GetData(), RomLoadOp->m_streamBuffer.Num());
-					FInputArchive Arch(&Stream);
-
-					TManagedPtr<FImage> Value = FImage::StaticUnserialise(Arch);
-
-					LocalProgram.SetImageRomValue(RomIndex, Value);
-				}
-				else
-				{
-					bSomeMissingData = true;
-				}
-				
-				Runner->RomLoadOps.Remove(*RomLoadOp);
-			}
-		}
-
-		if (bSomeMissingData)
-		{
-			// Some data may be missing. We can try to go on if some mips are there. 
-			UE_LOG(LogMutableCore, Verbose, TEXT("FLoadImageRomsTask::Complete failed: missing data?"));
-		}
-
-		// Process the constant op normally, now that the rom is loaded.	
-		bool bSuccess = Runner->RunCode_ConstantResource(Op);
-
-		if (bSuccess)
+		// Assume the ROM has been loaded previously in a task generated at IssueOp
+		if (!Source)
 		{
-			for (int32 ImageIndex = ImageIndexBegin; ImageIndex < ImageIndexEnd; ++ImageIndex)
-			{
-				FConstantResourceIndex ImageConstantResourceIndex = LocalProgram.ConstantImageLODIndices[ImageIndex];
-
-				if (!ImageConstantResourceIndex.Streamable)
-				{
-					// This data is always resident.
-					continue;
-				}
-
-				const int32 RomIndex = ImageConstantResourceIndex.Index;
-				check(RomIndex < LocalProgram.Roms.Num());
-
-				Runner->Model.Program.MarkRomUsed(RomIndex);
-				--LocalProgram.PendingOpsPerRom[RomIndex];
-
-				if (DebugRom && (DebugRomAll || RomIndex == DebugRomIndex))
-				{
-					UE_LOG(LogMutableCore, Log, TEXT("FLoadImageRomsTask::Complete rom %d, now peding ops is %d."), RomIndex, LocalProgram.PendingOpsPerRom[RomIndex]);
-				}
-			}
+			return false;
 		}
 
-		return bSuccess;
+		Runner->StoreImage(Op, Source);
+		
+		return true;
 	}
 
 
diff --git a/Engine/Plugins/Mutable/Source/MutableRuntime/Private/MuR/Model.cpp b/Engine/Plugins/Mutable/Source/MutableRuntime/Private/MuR/Model.cpp
--- a/Engine/Plugins/Mutable/Source/MutableRuntime/Private/MuR/Model.cpp
+++ b/Engine/Plugins/Mutable/Source/MutableRuntime/Private/MuR/Model.cpp
@@ -164,9 +164,6 @@ namespace UE::Mutable::Private
 		Arch >> RelevantParameterList;
 		Arch >> ConstantMaterials;
 		Arch >> ConstantNames;
-		
-		PendingOpsPerRom.SetNum(Roms.Num());
-		RomWeights.SetNumZeroed(Roms.Num());
 
 		PostLoad();
 	}
@@ -223,109 +220,7 @@ namespace UE::Mutable::Private
         }
 #endif
     }
-
-
-    void FProgram::UnloadRoms() const
-    {
-    	for (int32 RomIndex = 0; RomIndex < Roms.Num(); ++RomIndex)
-    	{
-    		UnloadRom(RomIndex);
-    	}
-    }
-
-
-    void FProgram::UnloadRom(int32 RomIndex, int32* OutDataSize) const
-    {
-	    if (DebugRom && (DebugRomAll||RomIndex == DebugRomIndex))
-		    UE_LOG(LogMutableCore, Log, TEXT("Unloading rom %d."), RomIndex);
-
-	    if (LoadedMemTrackedRoms.IsValidIndex(RomIndex))
-	    {
-		    LoadedMemTrackedRoms.RemoveAt(RomIndex);
-	    }
-
-	    switch (Roms[RomIndex].ResourceType)
-	    {
-	    case uint32(ERomDataType::Image):
-		    {
-			    TManagedPtr<const FImage> Data;
-			    ConstantImageLODsStreamed.RemoveAndCopyValue(RomIndex, Data);
-			    if (OutDataSize && Data)
-			    {
-				    *OutDataSize = Data->GetDataSize();
-			    }
-			    break;
-		    }
-	    case uint32(ERomDataType::Mesh):
-		    {
-			    TManagedPtr<const FMesh> Data;
-			    ConstantMeshesStreamed.RemoveAndCopyValue(RomIndex, Data);
-			    if (OutDataSize && Data)
-			    {
-				    *OutDataSize = Data->GetDataSize();
-			    }
-			    break;
-		    }
-	    //TODO(Max): Make materials loadable and unloadable from a rom
-	    default:
-		    check(false);
-		    break;
-	    }
-    }
-
-
-    int32 FProgram::GetRomBytes(int32 RomIndex) const
-    {
-	    int32 Result = 0;
-	    switch (Roms[RomIndex].ResourceType)
-	    {
-	    case uint32(ERomDataType::Image):
-		    {
-			    const TManagedPtr<const FImage>* Data = ConstantImageLODsStreamed.Find(RomIndex);
-			    if (Data)
-			    {
-				    Result = (*Data)->GetDataSize();
-			    }
-			    break;
-		    }
-	    case uint32(ERomDataType::Mesh):
-		    {
-			    const TManagedPtr<const FMesh>* Data = ConstantMeshesStreamed.Find(RomIndex);
-			    if (Data)
-			    {
-				    Result = (*Data)->GetDataSize();
-			    }
-			    break;
-		    }
-	    //TODO(Max): Make materials loadable and unloadable from a rom
-	    default:
-		    check(false);
-		    break;
-	    }
-
-	    return Result;
-    }
-
-
-    void FProgram::SetMeshRomValue(uint32 RomIndex, const TManagedPtr<FMesh>& Value) const
-    {
-	    check(Roms[RomIndex].ResourceType == uint32(ERomDataType::Mesh));
-	    check(!ConstantMeshesStreamed.Contains(RomIndex));
-
-	    LoadedMemTrackedRoms.EmplaceAt(RomIndex, (uint8)Roms[RomIndex].ResourceType);
-	    ConstantMeshesStreamed.Add(RomIndex, Value);
-    }
-
-
-    void FProgram::SetImageRomValue(uint32 RomIndex, const TManagedPtr<FImage>& Value) const
-    {
-	    check(Roms[RomIndex].ResourceType == uint32(ERomDataType::Image));
-	    check(!ConstantImageLODsStreamed.Contains(RomIndex));
-
-	    LoadedMemTrackedRoms.EmplaceAt(RomIndex, (uint8)Roms[RomIndex].ResourceType);
-	    ConstantImageLODsStreamed.Add(RomIndex, Value);
-    }
-
+	
 
     OP::CONSTANT_EXTENSIONDATA_INDEX FProgram::AddConstant(const TManagedPtr<const FExtensionData>& Data)
     {
@@ -551,30 +446,7 @@ namespace UE::Mutable::Private
 	    ConstantMaterials.Add(Material);
 	    return Index;
     }
-
-
-    TManagedPtr<const FImage> FProgram::GetImageLOD(FConstantResourceIndex Index) const
-    {
-	    TManagedPtr<const FImage> Mip;
-	    if (!Index.Streamable)
-	    {
-		    if (ConstantImageLODsPermanent.IsValidIndex(Index.Index))
-		    {
-			    Mip = ConstantImageLODsPermanent[Index.Index];
-		    }
-	    }
-	    else
-	    {
-		    const TManagedPtr<const FImage>* Found = ConstantImageLODsStreamed.Find(Index.Index);
-		    if (Found)
-		    {
-			    Mip = *Found;
-		    }
-	    }
-
-	    return Mip;
-    }
-
+	
 
     void FProgram::GetExtensionDataConstant(int32 ConstantIndex, TManagedPtr<const FExtensionData>& Result) const
     {
@@ -618,11 +490,9 @@ namespace UE::Mutable::Private
 	    return pData;
     }
 
-
-    void FProgram::MarkRomUsed(int32 RomIndex) const
+	
+    FModel::FModel(): RomManager(*this)
     {
-	    RomWeights[RomIndex].Key++;
-	    RomWeights[RomIndex].Value = ++FSystem::RomTick;
     }
 
 
@@ -658,68 +528,8 @@ namespace UE::Mutable::Private
     void FModel::Serialise(FModel* Model, FModelWriter& Streamer, bool bDropData)
     {
 		LLM_SCOPE_BYNAME(TEXT("MutableRuntime"));
-
-		UE::Mutable::Private::FProgram& Program = Model->Program;
-
-		FOutputMemoryStream MemStream(16 * 1024 * 1024);
-
-		// Save images and unload from memory
-		for (TPair<uint32, TManagedPtr<const FImage>>& Entry: Program.ConstantImageLODsStreamed )
-		{
-			int32 RomIndex = Entry.Key;
-			const FRomDataRuntime& RomData = Program.Roms[RomIndex];
-			check(RomData.ResourceType == uint32(ERomDataType::Image));
-
-			// Serialize to memory, to find out final size of this rom
-			MemStream.Reset();
-			FOutputArchive MemoryArch(&MemStream);
-			FImage::Serialise(Entry.Value.Get(), MemoryArch);
-			check(RomData.Size == MemStream.GetBufferSize());
-
-			Streamer.OpenWriteFile(RomIndex, true);
-			Streamer.Write(MemStream.GetBuffer(), MemStream.GetBufferSize());
-			Streamer.CloseWriteFile();
-
-			// Do this progressively to avoid duplicating all data in memory.
-			if (bDropData)
-			{
-				Entry.Value.Reset();
-			}
-		}
-
-		if (bDropData)
-		{
-			Program.ConstantImageLODsStreamed.Empty(0);
-		}
-
-		// Save meshes and unload from memory
-		for (TPair<uint32, TManagedPtr<const FMesh>>& ResData : Program.ConstantMeshesStreamed)
-		{
-			int32 RomIndex = ResData.Key;
-			const FRomDataRuntime& RomData = Program.Roms[RomIndex];
-			check(RomData.ResourceType == uint32(ERomDataType::Mesh));
-
-			// Serialize to memory, to find out final size of this rom
-			MemStream.Reset();
-			FOutputArchive MemoryArch(&MemStream);
-			FMesh::Serialise(ResData.Value.Get(), MemoryArch);
-			check(RomData.Size == MemStream.GetBufferSize());
-
-			Streamer.OpenWriteFile(RomIndex, true);
-			Streamer.Write(MemStream.GetBuffer(), MemStream.GetBufferSize());
-			Streamer.CloseWriteFile();
-
-			// Do this progressively to avoid duplicating all data in memory.
-			if (bDropData)
-			{
-				ResData.Value.Reset();
-			}
-		}
-
-		if (bDropData)
-		{
-			Program.ConstantMeshesStreamed.Empty(0);
-		}
+    	
+    	Model->RomManager.Serialise(Streamer, bDropData);
 
 		// Store the main data of the model
 		{
@@ -733,12 +543,6 @@ namespace UE::Mutable::Private
 		}
 	}
 
-
-	bool FModel::HasExternalData() const
-    {
-		return Program.Roms.Num() > 0;
-    }
-
 	
 #if WITH_EDITOR
 	bool FModel::IsValid() const
@@ -754,15 +558,6 @@ namespace UE::Mutable::Private
 #endif
 
 
-	void FModel::UnloadExternalData()
-    {
-		LLM_SCOPE_BYNAME(TEXT("MutableRuntime"));
-
-		Program.ConstantImageLODsStreamed.Empty(0);
-		Program.ConstantMeshesStreamed.Empty(0);
-	}
-
-
 	TSharedPtr<FModel> FModel::StaticUnserialise(FInputArchive& Arch)
     {
 		LLM_SCOPE_BYNAME(TEXT("MutableRuntime"));
@@ -909,25 +704,7 @@ namespace UE::Mutable::Private
     	return Roms[Index].Size;
     }
 
-
-    bool FProgram::IsRomLoaded(int32 RomIndex) const
-    {
-	    switch (Roms[RomIndex].ResourceType)
-	    {
-	    case uint32(ERomDataType::Image):
-		    return ConstantImageLODsStreamed.Contains(RomIndex);
-	    case uint32(ERomDataType::Mesh):
-		    return ConstantMeshesStreamed.Contains(RomIndex);
-	    //TODO(Max): Make materials loadable and unloadable from a rom
-	    default:
-		    check(false);
-		    break;
-	    }
-
-	    return false;
-    }
-
-
+	
     ERomDataType FProgram::GetRomDataType(int32 Index) const
 	{
 		return static_cast<ERomDataType>(Roms[Index].ResourceType);
diff --git a/Engine/Plugins/Mutable/Source/MutableRuntime/Private/MuR/Parameters.cpp b/Engine/Plugins/Mutable/Source/MutableRuntime/Private/MuR/Parameters.cpp
--- a/Engine/Plugins/Mutable/Source/MutableRuntime/Private/MuR/Parameters.cpp
+++ b/Engine/Plugins/Mutable/Source/MutableRuntime/Private/MuR/Parameters.cpp
@@ -73,7 +73,7 @@ namespace UE::Mutable::Private
 
 	const FString& FParameters::GetName( int32 Index ) const
 	{
-		const FProgram& Program = Model->Program;
+		const FProgram& Program = Model->GetProgram();
 		check( Index>=0 && Index<(int32)Program.Parameters.Num() );
 
 		return Program.Parameters[Index].Name;
@@ -82,7 +82,7 @@ namespace UE::Mutable::Private
 
 	const FGuid& FParameters::GetUid( int32 Index ) const
 	{
-		const FProgram& Program = Model->Program;
+		const FProgram& Program = Model->GetProgram();
 		check( Index>=0 && Index<Program.Parameters.Num() );
 
 		return Program.Parameters[Index].UID;
@@ -91,7 +91,7 @@ namespace UE::Mutable::Private
 
 	int32 FParameters::Find( const FString& Name ) const
 	{
-    	const FProgram& Program = Model->Program;
+    	const FProgram& Program = Model->GetProgram();
 
     	int32 result = -1;
 
@@ -111,7 +111,7 @@ namespace UE::Mutable::Private
 
 	EParameterType FParameters::GetType( int32 Index ) const
 	{
-		const FProgram& Program = Model->Program;
+		const FProgram& Program = Model->GetProgram();
 		check( Index>=0 && Index<Program.Parameters.Num() );
 
 		return Program.Parameters[Index].Type;
@@ -122,7 +122,7 @@ namespace UE::Mutable::Private
     {
 		LLM_SCOPE_BYNAME(TEXT("MutableRuntime"));
 
-        const FProgram& Program = Model->Program;
+        const FProgram& Program = Model->GetProgram();
         check( ParamIndex>=0 && ParamIndex<Program.Parameters.Num() );
 		TSharedPtr<FRangeIndex> Range;
 
@@ -276,7 +276,7 @@ namespace UE::Mutable::Private
 
 	int32 FParameters::GetIntPossibleValueCount( int32 ParamIndex ) const
 	{
-		const FProgram& Program = Model->Program;
+		const FProgram& Program = Model->GetProgram();
 		check( ParamIndex>=0 && ParamIndex<(int32)Program.Parameters.Num() );
 		return (int32)Program.Parameters[ParamIndex].PossibleValues.Num();
 	}
@@ -284,7 +284,7 @@ namespace UE::Mutable::Private
 
 	int32 FParameters::GetIntPossibleValue( int32 ParamIndex, int32 valueIndex ) const
 	{
-		const FProgram& Program = Model->Program;
+		const FProgram& Program = Model->GetProgram();
 		check( ParamIndex>=0
 				&& ParamIndex<(int32)Program.Parameters.Num() );
 		check( valueIndex>=0
@@ -296,7 +296,7 @@ namespace UE::Mutable::Private
 
 	int32 FParameters::GetIntValueIndex(int32 ParamIndex, const FString& ValueName) const
 	{
-		const FProgram& Program = Model->Program;
+		const FProgram& Program = Model->GetProgram();
 		check(ParamIndex >= 0
 			&& ParamIndex < (int32)Program.Parameters.Num());
 
@@ -315,7 +315,7 @@ namespace UE::Mutable::Private
 
 	int32 FParameters::GetIntValueIndex(int32 ParamIndex, int32 Value) const
 	{
-		const FProgram& Program = Model->Program;
+		const FProgram& Program = Model->GetProgram();
 		check(ParamIndex >= 0
 			&& ParamIndex < (int32)Program.Parameters.Num());
 
@@ -332,7 +332,7 @@ namespace UE::Mutable::Private
 
 	const FString& FParameters::GetIntPossibleValueName( int32 ParamIndex, int32 valueIndex ) const
 	{
-		const FProgram& Program = Model->Program;
+		const FProgram& Program = Model->GetProgram();
 		check( ParamIndex>=0
 				&& ParamIndex<(int32)Program.Parameters.Num() );
 		check( valueIndex>=0
@@ -1061,7 +1061,7 @@ namespace UE::Mutable::Private
 
 	FProjector FParameters::GetProjectorValue( int32 Index, const FRangeIndex* Pos ) const
     {
-        const FProgram& Program = Model->Program;
+        const FProgram& Program = Model->GetProgram();
 
         // Early out in case of invalid parameters
         if ( Index < 0
@@ -1164,7 +1164,7 @@ namespace UE::Mutable::Private
         Value.type = EProjectorType::Count;
         if (Model)
         {
-            const FProgram& Program = Model->Program;
+            const FProgram& Program = Model->GetProgram();
             const FProjector& Projector = Program.Parameters[ParameterIndex].DefaultValue.Get<FParamProjectorType>();
             Value.type = Projector.type;
         }
@@ -1274,9 +1274,9 @@ namespace UE::Mutable::Private
         check( Index < int32(Parameters->Values.Num()) );
         TSharedPtr<const FModel> Model = Parameters->Model;
         check(Model);
-        check( Index < int32(Model->Program.Ranges.Num()) );
+        check( Index < int32(Model->GetProgram().Ranges.Num()) );
 
-        return Model->Program.Ranges[Index].Name;
+        return Model->GetProgram().Ranges[Index].Name;
     }
 
 
@@ -1286,9 +1286,9 @@ namespace UE::Mutable::Private
         check( Index < int32(Parameters->Values.Num()) );
 		TSharedPtr<const FModel> Model = Parameters->Model;
 		check(Model);
-        check( Index < int32(Model->Program.Ranges.Num()) );
+        check( Index < int32(Model->GetProgram().Ranges.Num()) );
 
-        return Model->Program.Ranges[Index].UID;
+        return Model->GetProgram().Ranges[Index].UID;
     }
 
 
diff --git a/Engine/Plugins/Mutable/Source/MutableRuntime/Private/MuR/System.cpp b/Engine/Plugins/Mutable/Source/MutableRuntime/Private/MuR/System.cpp
--- a/Engine/Plugins/Mutable/Source/MutableRuntime/Private/MuR/System.cpp
+++ b/Engine/Plugins/Mutable/Source/MutableRuntime/Private/MuR/System.cpp
@@ -25,6 +25,7 @@
 #include "MuR/Platform.h"
 #include "MuR/Serialisation.h"
 #include "MuR/MutableRuntimeModule.h"
+#include "MuR/RomManager.h"
 #include "Templates/SharedPointer.h"
 #include "Templates/Tuple.h"
 #include "ProfilingDebugging/CountersTrace.h"
@@ -78,7 +79,7 @@ namespace UE::Mutable::Private
 	TRACE_DECLARE_INT_COUNTER(MutableRuntime_MemBudget,		TEXT("MutableRuntime/MemBudget"));
 
 
-	uint64 FSystem::RomTick = 0;
+	std::atomic<uint32> FSystem::RomTick = 0;
 
 	
     FSystem::FSystem(const FSettings& InSettings)
@@ -121,7 +122,7 @@ namespace UE::Mutable::Private
 	}
 
    
-    TSharedRef<FLiveInstance> FSystem::NewInstance( const TSharedPtr<const FModel>& InModel,
+    TSharedRef<FLiveInstance> FSystem::NewInstance( const TSharedPtr<FModel>& InModel,
 		const TSharedPtr<UE::Mutable::Private::FMeshIdRegistry>& InMeshIdRegistry,
 		const TSharedPtr<UE::Mutable::Private::FImageIdRegistry>& InImageIdRegistry,
 		const TSharedPtr<UE::Mutable::Private::FMaterialIdRegistry>& InMaterialIdRegistry,
@@ -161,7 +162,7 @@ namespace UE::Mutable::Private
 	}
 
 	
-    TSharedRef<FLiveInstance> FSystem::NewBuildInstance(const TSharedPtr<const FModel>& InModel,
+    TSharedRef<FLiveInstance> FSystem::NewBuildInstance(const TSharedPtr<FModel>& InModel,
 		const TSharedPtr<const FParameters>& InParameters,
 	    const TSharedPtr<FExternalResourceProvider>& ExternalResourceProvider)
     {
@@ -182,7 +183,7 @@ namespace UE::Mutable::Private
     	InstanceData->MaterialIdRegistry = MakeShared<FMaterialIdRegistry>();
     	InstanceData->SkeletalMeshIdRegistry = MakeShared<FSkeletalMeshIdRegistry>();
 
-    	InstanceData->Cache->Init(InModel->Program.OpAddress.Num());
+    	InstanceData->Cache->Init(InModel->GetProgram().OpAddress.Num());
     	
     	return InstanceData;
     }
@@ -199,7 +200,7 @@ namespace UE::Mutable::Private
     	CodeRunnerBegin Runner(LiveInstance, InLodMask);
 
     	check(LiveInstance->State >= 0);
-    	FScheduledOpInline Root(LiveInstance->Model->Program.States[LiveInstance->State].Root, 0, 0, true);
+    	FScheduledOpInline Root(LiveInstance->Model->GetProgram().States[LiveInstance->State].Root, 0, 0, true);
     	TManagedPtr<const FInstance> Result = Runner.RunCode(FScheduledOpInline(Root, 0));
     	if (!Result)
     	{
@@ -224,7 +225,7 @@ namespace UE::Mutable::Private
 			return nullptr;
 		}
 
-		const FProgram& Program = LiveInstance->Model->Program;
+		const FProgram& Program = LiveInstance->Model->GetProgram();
 
     	const int32 LiveInstanceState = LiveInstance->State;
     	
@@ -323,7 +324,7 @@ namespace UE::Mutable::Private
 
     	OP::ADDRESS RootAddress = ImageId.GetKey()->Address;
     	
-    	EOpType OpType = LiveInstance->Model->Program.GetOpType(RootAddress);
+    	EOpType OpType = LiveInstance->Model->GetProgram().GetOpType(RootAddress);
     	if (GetOpDataType(OpType) != EDataType::Image)
     	{
     		return UE::Tasks::Launch(TEXT("FSystem::GetImageInvalid"),
@@ -409,7 +410,7 @@ namespace UE::Mutable::Private
 		OP::ADDRESS RootAddress = ImageId.GetKey()->Address;
 
 		const UE::Mutable::Private::FModel* Model = LiveInstance->Model.Get();
-		const UE::Mutable::Private::FProgram& Program = Model->Program;
+		const UE::Mutable::Private::FProgram& Program = Model->GetProgram();
 
 		// TODO: It should be possible to reuse this data if cleared in the correct places only, together with HeapImageDesc.
 		int32 VarValue = CVarClearImageDescCache.GetValueOnAnyThread();
@@ -464,7 +465,7 @@ namespace UE::Mutable::Private
 		
 		OP::ADDRESS RootAddress = ImageId.GetKey()->Address;
 		
-		EOpType OpType = LiveInstance->Model->Program.GetOpType(RootAddress);
+		EOpType OpType = LiveInstance->Model->GetProgram().GetOpType(RootAddress);
 		if (GetOpDataType(OpType) != EDataType::Image)
 		{
 			return UE::Tasks::Launch(TEXT("FSystem::GetImageDescInvalid"),
@@ -572,7 +573,7 @@ namespace UE::Mutable::Private
 
 		OP::ADDRESS RootAddress = MeshId.GetKey()->Address;
 
-		EOpType OpType = LiveInstance->Model->Program.GetOpType(RootAddress);
+		EOpType OpType = LiveInstance->Model->GetProgram().GetOpType(RootAddress);
 		if (GetOpDataType(OpType) != EDataType::Mesh)
 		{
 			return UE::Tasks::Launch(TEXT("FSystem::GetMeshInvalid"),
@@ -650,7 +651,7 @@ namespace UE::Mutable::Private
 
 		OP::ADDRESS RootAddress = SkeletalMeshId.GetKey()->Address;
 		
-		UE::Mutable::Private::EOpType OpType = LiveInstance->Model->Program.GetOpType(RootAddress);
+		UE::Mutable::Private::EOpType OpType = LiveInstance->Model->GetProgram().GetOpType(RootAddress);
 		if (GetOpDataType(OpType) != EDataType::SkeletalMesh)
 		{
 			return UE::Tasks::Launch(TEXT("FSystem::GetSkeletalMeshInvalid"),
@@ -727,7 +728,7 @@ namespace UE::Mutable::Private
 		
 		OP::ADDRESS RootAddress = MaterialId.GetKey()->Address;
 	
-		EOpType OpType = LiveInstance->Model->Program.GetOpType(RootAddress);
+		EOpType OpType = LiveInstance->Model->GetProgram().GetOpType(RootAddress);
 		if (GetOpDataType(OpType) != EDataType::Material)
 		{
 			return UE::Tasks::Launch(TEXT("FSystem::GetMaterialInvalid"),
@@ -806,12 +807,7 @@ namespace UE::Mutable::Private
 
 		if (bClearRoms)
 		{
-			const FProgram& Program = LiveInstance->Model->Program;
-
-			for (int32 RomIndex = 0; RomIndex < Program.Roms.Num(); ++RomIndex)
-			{
-				Program.UnloadRom(RomIndex);
-			}
+			LiveInstance->Model->RomManager.UnloadRoms();
 		}
 		
 		UpdateStats();
@@ -830,7 +826,7 @@ namespace UE::Mutable::Private
 				UniqueDiscreteCoveredCodeVisitor<>(InSystem, InLiveInstance, FSystem::AllLODs),
 		    	RelevantParameters(OutRelevantParameters)
         {
-			const OP::ADDRESS At = InLiveInstance->Model->Program.States[0].Root;
+			const OP::ADDRESS At = InLiveInstance->Model->GetProgram().States[0].Root;
 
             Run(At);
         }
@@ -887,7 +883,7 @@ namespace UE::Mutable::Private
     	const TSharedRef<const FParameters> NewParameters = LiveInstance->Parameters.ToSharedRef();
 
         // check what parameters have changed
-        const FProgram& Program = LiveInstance->Model->Program;
+        const FProgram& Program = LiveInstance->Model->GetProgram();
         const TArray<int>& RuntimeParameters = Program.States[ LiveInstance->State ].m_runtimeParameters;
 
         check( NewParameters->GetCount() == (int32)Program.Parameters.Num() );
@@ -979,7 +975,7 @@ namespace UE::Mutable::Private
 	{
 		FVector4f Result(0,0,0,1);
     	
-		UE::Mutable::Private::EOpType OpType = InLiveInstance->Model->Program.GetOpType(InAddress);
+		UE::Mutable::Private::EOpType OpType = InLiveInstance->Model->GetProgram().GetOpType(InAddress);
 		if (GetOpDataType(OpType) == EDataType::Color)
 		{
 			RunCode(InLiveInstance, InAddress);
@@ -1011,7 +1007,7 @@ namespace UE::Mutable::Private
 	{
 		TManagedPtr<const FImage> Result;
     	
-		UE::Mutable::Private::EOpType OpType = InLiveInstance->Model->Program.GetOpType(InAddress);
+		UE::Mutable::Private::EOpType OpType = InLiveInstance->Model->GetProgram().GetOpType(InAddress);
 		if (GetOpDataType(OpType) == EDataType::Image)
 		{
 			RunCode(InLiveInstance, InAddress, FSystem::AllLODs, uint8(MipsToSkip));
@@ -1029,7 +1025,7 @@ namespace UE::Mutable::Private
 	{
 		TManagedPtr<const FMesh> Result;
     	
-		UE::Mutable::Private::EOpType OpType = InLiveInstance->Model->Program.GetOpType(InRootAddress);
+		UE::Mutable::Private::EOpType OpType = InLiveInstance->Model->GetProgram().GetOpType(InRootAddress);
 		if (GetOpDataType(OpType) == EDataType::Mesh)
 		{
 			const uint8 ExecutionOptions = static_cast<uint8>(MeshContentFilter);
@@ -1048,7 +1044,7 @@ namespace UE::Mutable::Private
 	{
 		TManagedPtr<const FInstance> Result;
     	
-		UE::Mutable::Private::EOpType OpType = InLiveInstance->Model->Program.GetOpType(InAddress);
+		UE::Mutable::Private::EOpType OpType = InLiveInstance->Model->GetProgram().GetOpType(InAddress);
 		if (GetOpDataType(OpType) == EDataType::Instance)
 		{
 			RunCode(InLiveInstance, InAddress);
@@ -1066,9 +1062,9 @@ namespace UE::Mutable::Private
 	{
 		TManagedPtr<const FLayout> Result;
     	
-		if (InLiveInstance->Model->Program.States[0].Root)
+		if (InLiveInstance->Model->GetProgram().States[0].Root)
 		{
-			UE::Mutable::Private::EOpType OpType = InLiveInstance->Model->Program.GetOpType(InAddress);
+			UE::Mutable::Private::EOpType OpType = InLiveInstance->Model->GetProgram().GetOpType(InAddress);
 			if (GetOpDataType(OpType) == EDataType::Layout)
 			{
 				RunCode(InLiveInstance, InAddress);
@@ -1087,9 +1083,9 @@ namespace UE::Mutable::Private
 	{
 		TManagedPtr<const String> Result;
     	
-		if (InLiveInstance->Model->Program.States[0].Root)
+		if (InLiveInstance->Model->GetProgram().States[0].Root)
 		{
-			UE::Mutable::Private::EOpType OpType = InLiveInstance->Model->Program.GetOpType(InAddress);
+			UE::Mutable::Private::EOpType OpType = InLiveInstance->Model->GetProgram().GetOpType(InAddress);
 			if (GetOpDataType(OpType) == EDataType::String)
 			{
 				RunCode(InLiveInstance, InAddress);
@@ -1122,7 +1118,6 @@ namespace UE::Mutable::Private
     {
     	int32 Result = 0;
 
-    	TArray<const FModel*> Models;
     	for (const TWeakPtr<FLiveInstance>& WeakInstance : LiveInstances)
     	{
     		TSharedPtr<FLiveInstance> Instance = WeakInstance.Pin();
@@ -1131,28 +1126,7 @@ namespace UE::Mutable::Private
     			continue;
     		}
 
-    		Models.AddUnique(Instance->Model.Get());
-    	}
-
-    	// Data stored per-model, but related to instance construction
-    	for (const FModel* Model : Models)
-    	{
-    		// Count streamable and currently-loaded resources
-    		const FProgram& Program = Model->Program;
-    		for (const TPair<uint32,TManagedPtr<const FImage>>& ImageLOD : Program.ConstantImageLODsStreamed)
-    		{
-    			if (ImageLOD.Value)
-    			{
-    				Result += ImageLOD.Value->GetDataSize();
-    			}
-    		}
-    		for (const TPair<uint32, TManagedPtr<const FMesh>>& Rom : Program.ConstantMeshesStreamed)
-    		{
-    			if (Rom.Value)
-    			{
-    				Result += Rom.Value->GetDataSize();
-    			}
-    		}
+    		Result +=  Instance->Model->RomManager.GetRomBytes();
     	}
 
     	return Result;
@@ -1172,14 +1146,16 @@ namespace UE::Mutable::Private
 		const uint32 CacheBytes = GetCacheBytes();
 		uint32 HeapBytes = 0;
 	
-		for (const CodeRunner::FScheduledOpData& Data : CurrentRunner->HeapData)
-		{
-			if (Data.Resource)
-			{
-				HeapBytes += Data.Resource->GetDataSize();
-			}
-		}
-
+    	if (CurrentRunner)
+    	{
+    		for (const CodeRunner::FScheduledOpData& Data : CurrentRunner->HeapData)
+    		{
+    			if (Data.Resource)
+    			{
+    				HeapBytes += Data.Resource->GetDataSize();
+    			}
+    		}
+    	}
 
 		// Get allocator counters.
 		const SSIZE_T ImageAllocBytes	 = MemoryCounters::FImageMemoryCounter::Get().load(std::memory_order_relaxed);
@@ -1258,7 +1234,7 @@ namespace UE::Mutable::Private
 		return CurrentBytes > BudgetThresholdBytes;
 	}
 
-	bool FWorkingMemoryManager::EnsureBudgetBelow( uint64 AdditionalMemory, const TSharedPtr<FProgramCache>& InCurrentCache )
+	bool FWorkingMemoryManager::EnsureBudgetBelow(int64 AdditionalMemory, const TSharedPtr<FLiveInstance>& InCurrentLiveInstance)
     {
 		MUTABLE_CPUPROFILER_SCOPE(EnsureBudgetBelow);
 
@@ -1278,18 +1254,6 @@ namespace UE::Mutable::Private
 		// Try to free loaded roms
 		if (!bFinished)
 		{
-			MUTABLE_CPUPROFILER_SCOPE(EnsureBudgetBelow_FreeRoms);
-
-			struct FRomRef
-			{
-				const FModel* Model = nullptr;
-				int32 RomIndex = 0;
-			};
-
-			TArray<TPair<float, FRomRef>> Candidates;
-			Candidates.Reserve(512);
-
-			SSIZE_T RomBytes = 0;
 			for (TWeakPtr<FLiveInstance>& WeakLiveInstance : LiveInstances)
 			{
 				TSharedPtr<FLiveInstance> LiveInstance = WeakLiveInstance.Pin();
@@ -1298,67 +1262,34 @@ namespace UE::Mutable::Private
 					continue;
 				}
 				
-				const FProgram& Program = LiveInstance->Model->Program;
-	
-				check(Program.LoadedMemTrackedRoms.GetMaxIndex() <= Program.Roms.Num());
-
-				for (TSparseArray<uint8>::TConstIterator Iter = Program.LoadedMemTrackedRoms.CreateConstIterator(); Iter; ++Iter)
+				if (LiveInstance == InCurrentLiveInstance)
 				{
-					const int32 RomIndex = Iter.GetIndex();
-
-					const FRomDataRuntime& Rom = Program.Roms[RomIndex];
-					check(Program.IsRomLoaded(RomIndex));
-					check(Rom.ResourceType == (uint32)*Iter);
-
-					RomBytes += Program.GetRomBytes(RomIndex);
-
-					// We cannot unload a rom if some operation is expecting it.
-					const bool bIsRomLocked = Program.PendingOpsPerRom.IsValidIndex(RomIndex) && 
-											  Program.PendingOpsPerRom[RomIndex] > 0;
-					if (!bIsRomLocked)
-					{
-						constexpr float FactorWeight = 100.0f;
-						constexpr float FactorTime = -1.0f;
-						float Priority = FactorWeight * float(Program.RomWeights[RomIndex].Get<0>()) +
-										 FactorTime * float(FSystem::RomTick - Program.RomWeights[RomIndex].Get<1>());
-
-						Candidates.Emplace(Priority, FRomRef{ LiveInstance->Model.Get(), RomIndex });
-					}
+					// Ignore the current live instance.
+					continue;
 				}
-			}
-
-			// Don't sort all candidates, make it a heap in O(N) time. We may not need to visit all elements.
-			auto CompareCandidates = [](const TPair<float, FRomRef>& A, const TPair<float, FRomRef>& B) { return A.Key < B.Key; };
-			Candidates.Heapify(CompareCandidates);
-
-			SSIZE_T CandidatesRomBytes = 0;
-			for (const TPair<float, FRomRef>& Entry : Candidates)
-			{
-				CandidatesRomBytes += Entry.Value.Model->Program.GetRomBytes(Entry.Value.RomIndex);
-			}
-		
-			while (!bFinished && Candidates.Num())
-			{
-				MUTABLE_CPUPROFILER_SCOPE(EnsureBudgetBelow_UnloadRom);
-
-				TPair<float, FRomRef> Candidate = Candidates.HeapTop();
+				
+				const int64 RomRemovedBytes = LiveInstance->Model->RomManager.EnsureBudgetBelow(TotalBytes - BudgetBytes);
+				TotalBytes -= RomRemovedBytes;
 
-				// UE_LOG(LogMutableCore,Log, "Unloading rom because of memory budget: %d.", lowestPriorityRom);
-				int32 UnloadedSize = 0;
-				Candidate.Value.Model->Program.UnloadRom(Candidate.Value.RomIndex, &UnloadedSize);
-				TotalBytes -= UnloadedSize;
 				bFinished = TotalBytes <= BudgetBytes;
-
-				if (!bFinished)
+				if (bFinished)
 				{
-					Candidates.HeapPopDiscard(CompareCandidates, EAllowShrinking::No);
+					break;
 				}
 			}
 		}
+    	
+    	if (!bFinished && InCurrentLiveInstance)
+    	{	
+    		const int64 RomRemovedBytes = InCurrentLiveInstance->Model->RomManager.EnsureBudgetBelow(TotalBytes - BudgetBytes);
+    		TotalBytes -= RomRemovedBytes;
+			
+    		bFinished = TotalBytes <= BudgetBytes;
+    	}
 
 		auto TryRemoveFromCache = [&bFinished, &TotalBytes, Budget = BudgetBytes](FProgramCache& Cache)
 		{
-			int32 ImageRemovedBytes = 0;
+			int64 ImageRemovedBytes = 0;
 
 			for (const TManagedWeakPtr<const FImage>& ImageResult : Cache.ImageResults)
 			{
@@ -1388,7 +1319,7 @@ namespace UE::Mutable::Private
 				return;
 			}
 
-			int32 MeshRemovedBytes = 0;
+			int64 MeshRemovedBytes = 0;
 			for (const TManagedWeakPtr<const FMesh>& MeshResult : Cache.MeshResults)
 			{
 				TManagedPtr<const FMesh> StrongMeshPtr = MeshResult.Pin();
@@ -1427,7 +1358,7 @@ namespace UE::Mutable::Private
 					continue;
 				}
 				
-				if (Instance->Cache == InCurrentCache)
+				if (Instance == InCurrentLiveInstance)
 				{
 					// Ignore the current live instance.
 					continue;
@@ -1447,11 +1378,11 @@ namespace UE::Mutable::Private
 
 		// From the current live instances. It is more involved: we have to make sure any data we want to free is not also
 		// in any cache (0 or 1) position with hit-count > 0.
-		if (!bFinished && InCurrentCache)
+		if (!bFinished && InCurrentLiveInstance)
 		{
 			MUTABLE_CPUPROFILER_SCOPE(EnsureBudgetBelow_FreeCached_Current);
 
-			TryRemoveFromCache(*InCurrentCache);
+			TryRemoveFromCache(*InCurrentLiveInstance->Cache);
 		}
 
 
@@ -1498,7 +1429,7 @@ namespace UE::Mutable::Private
 	{
 		MUTABLE_CPUPROFILER_SCOPE(GetResourceKey);
     	
-		const FProgram& Program = Model.Program;
+		const FProgram& Program = Model.GetProgram();
 
 		// Find the list of relevant parameters
 		const TArray<uint16>* RelevantParams = nullptr;
@@ -1515,7 +1446,7 @@ namespace UE::Mutable::Private
 		// Generate the relevant parameters blob
 		FBitWriterResourceKey Blob(2048*8, true);
 
-		const TArray<FParameterDesc>& ParamDescs = Params.Model->Program.Parameters;
+		const TArray<FParameterDesc>& ParamDescs = Params.Model->GetProgram().Parameters;
 
 		// First make a mask with a bit for each relevant parameter. It will be on for parameters included in the blob.
 		// A parameter will be excluded from the blob if it has the default value, and no multivalues.
@@ -1800,7 +1731,7 @@ namespace UE::Mutable::Private
     	FGeneratedMeshKey NewKey;
     	NewKey.Address = RootAt;
 
-    	const int32* Result = Model.Program.RelevantParameterList.Find(RootAt);
+    	const int32* Result = Model.GetProgram().RelevantParameterList.Find(RootAt);
     	const int32 ParamListIndex = Result ? *Result : INDEX_NONE; 
 		GetResourceKey(Model, Params, ParamListIndex, RootAt, NewKey.ParameterValuesBlob);
 
@@ -1815,7 +1746,7 @@ namespace UE::Mutable::Private
     	FGeneratedImageKey NewKey;
     	NewKey.Address = RootAt;
 
-    	const int32* Result = Model.Program.RelevantParameterList.Find(RootAt);
+    	const int32* Result = Model.GetProgram().RelevantParameterList.Find(RootAt);
     	const int32 ParamListIndex = Result ? *Result : INDEX_NONE; 
     	GetResourceKey(Model, Params, ParamListIndex, RootAt, NewKey.ParameterValuesBlob);
 
@@ -1830,7 +1761,7 @@ namespace UE::Mutable::Private
     	FGeneratedMaterialKey NewKey;
     	NewKey.Address = RootAt;
 
-    	const int32* Result = Model.Program.RelevantParameterList.Find(RootAt);
+    	const int32* Result = Model.GetProgram().RelevantParameterList.Find(RootAt);
     	const int32 ParamListIndex = Result ? *Result : INDEX_NONE; 
     	GetResourceKey(Model, Params, ParamListIndex, RootAt, NewKey.ParameterValuesBlob);
 
@@ -1845,7 +1776,7 @@ namespace UE::Mutable::Private
     	FGeneratedSkeletalMeshKey NewKey;
     	NewKey.Address = RootAt;
 
-    	const int32* Result = Model.Program.RelevantParameterList.Find(RootAt);
+    	const int32* Result = Model.GetProgram().RelevantParameterList.Find(RootAt);
     	const int32 ParamListIndex = Result ? *Result : INDEX_NONE; 
     	GetResourceKey(Model, Params, ParamListIndex, RootAt, NewKey.ParameterValuesBlob);
 
diff --git a/Engine/Plugins/Mutable/Source/MutableTools/Private/MuT/CodeOptimiser.cpp b/Engine/Plugins/Mutable/Source/MutableTools/Private/MuT/CodeOptimiser.cpp
--- a/Engine/Plugins/Mutable/Source/MutableTools/Private/MuT/CodeOptimiser.cpp
+++ b/Engine/Plugins/Mutable/Source/MutableTools/Private/MuT/CodeOptimiser.cpp
@@ -187,11 +187,11 @@ namespace UE::Mutable::Private
 			LinkerOptions.bSeparateImageMips = false;
 
 			TSharedPtr<FModel> Model = MakeShared<FModel>();
-			OP::ADDRESS at = ASTOp::FullLink(SourceCloned, Model->Program, &LinkerOptions);
+			OP::ADDRESS at = ASTOp::FullLink(SourceCloned, Model->GetProgram(), &LinkerOptions);
 
 			FProgram::FState state;
 			state.Root = at;
-			Model->Program.States.Add(state);
+			Model->GetProgram().States.Add(state);
 
 			TSharedPtr<FParameters> LocalParams = FModel::NewParameters(Model);
 
diff --git a/Engine/Plugins/Mutable/Source/MutableTools/Private/MuT/Compiler.cpp b/Engine/Plugins/Mutable/Source/MutableTools/Private/MuT/Compiler.cpp
--- a/Engine/Plugins/Mutable/Source/MutableTools/Private/MuT/Compiler.cpp
+++ b/Engine/Plugins/Mutable/Source/MutableTools/Private/MuT/Compiler.cpp
@@ -541,7 +541,7 @@ namespace UE::Mutable::Private
 
         // Link the Program and generate state data.
 		TSharedPtr<FModel> Result = MakeShared<FModel>();
-        FProgram& Program = Result->Program;
+        FProgram& Program = Result->GetProgram();
 
 		check(Program.Parameters.IsEmpty());
 		Program.Parameters = Parameters;
@@ -662,9 +662,6 @@ namespace UE::Mutable::Private
 
 		// Pack data
 		m_pD->GenerateRoms(Result.Get(), LinkerOptions.AdditionalData );
-
-		Program.PendingOpsPerRom.SetNum(Program.Roms.Num());
-		Program.RomWeights.SetNumZeroed(Program.Roms.Num());
 		
 		// We are not touching the program anymore. Ensure we are not wasting memory.
 		Program.Roms.Shrink();
@@ -716,7 +713,7 @@ namespace UE::Mutable::Private
 
 		uint64 EmbeddedDataBytesLimit = Options->GetPrivate()->EmbeddedDataBytesLimit;
 
-		UE::Mutable::Private::FProgram& Program = Model->Program;
+		UE::Mutable::Private::FProgram& Program = Model->GetProgram();
 
 		// These are used for logging only.
 		int32 NumRoms = 0;
@@ -860,7 +857,7 @@ namespace UE::Mutable::Private
 						check(RomIndex < 0x7fffffff);
 
 						IndexPerMip[MipIndex] = {RomIndex, 1};
-						Program.ConstantImageLODsStreamed.Add(RomIndex, AllMips[MipIndex]);
+						Model->RomManager.SetRom<FImage>(RomIndex, AllMips[MipIndex]);
 
 						Program.Roms.Add(RomData);
 
@@ -960,7 +957,7 @@ namespace UE::Mutable::Private
 					check(RomIndex < 0x7fffffff);
 
 					IndexPerMeshContent[MeshIndex] = { RomIndex, 1 };
-					Program.ConstantMeshesStreamed.Add(RomIndex, AllMeshes[MeshIndex]);
+					Model->RomManager.SetRom<FMesh>(RomIndex, AllMeshes[MeshIndex]);
 
 					Program.Roms.Add(RomData);
 					Program.RomsCompileData.Add(RomDataCompile);
diff --git a/Engine/Plugins/Mutable/Source/MutableValidation/Private/MuV/ValidationUtils.cpp b/Engine/Plugins/Mutable/Source/MutableValidation/Private/MuV/ValidationUtils.cpp
--- a/Engine/Plugins/Mutable/Source/MutableValidation/Private/MuV/ValidationUtils.cpp
+++ b/Engine/Plugins/Mutable/Source/MutableValidation/Private/MuV/ValidationUtils.cpp
@@ -210,11 +210,11 @@ bool TestCustomizableObject(UCustomizableObject& InTargetCustomizableObject, con
 
 		// Roms ---------------------- //
 		{
-			const int32 RomCount =  MutableModel->Program.GetRomCount();
+			const int32 RomCount =  MutableModel->GetProgram().GetRomCount();
 			int64 TotalRomSizeBytes = 0;
 			for (int32 RomIndex = 0; RomIndex < RomCount; RomIndex++)
 			{
-				const uint32 RomByteSize = MutableModel->Program.GetRomSize(RomIndex);
+				const uint32 RomByteSize = MutableModel->GetProgram().GetRomSize(RomIndex);
 				TotalRomSizeBytes += RomByteSize;
 			}
 
