[mutable] Fixed FCustomizableObjectMeshStreamIn tasks leaks.
- In some cases the mutable low priority task that was set as a prerequisite of those tasks was not run (valid case when aborting the streaming operation), which caused the tasks that had that initial task as a dependency unable to run, which kept the FLiveInstance and the FCustomizableObjectMeshStreamIn alive.

#rnx
#Synced-CL 50344094

diff --git a/Engine/Plugins/Mutable/Source/CustomizableObject/Private/MuCO/CustomizableObjectMeshUpdate.cpp b/Engine/Plugins/Mutable/Source/CustomizableObject/Private/MuCO/CustomizableObjectMeshUpdate.cpp
--- a/Engine/Plugins/Mutable/Source/CustomizableObject/Private/MuCO/CustomizableObjectMeshUpdate.cpp
+++ b/Engine/Plugins/Mutable/Source/CustomizableObject/Private/MuCO/CustomizableObjectMeshUpdate.cpp
@@ -50,9 +50,6 @@ FCustomizableObjectMeshStreamIn::FCustomizableObjectMeshStreamIn(
 	OperationData = InMesh->UpdateMeshContext;
 	check(OperationData)
 	
-	// Set up the event that will indicate that the Task_Mutable_UpdateMesh_End was run
-	OperationData->MeshStreamingPrerequisite =  MakeShared<UE::Tasks::FTaskEvent>(TEXT("Task_Mutable_UpdateMesh_End"));
-	
 	OperationData->CurrentFirstLODIdx = CurrentFirstLODIdx;
 	OperationData->PendingFirstLODIdx = PendingFirstLODIdx;
 
@@ -168,45 +165,79 @@ void FCustomizableObjectMeshStreamIn::DoCancelMeshUpdate(const FContext& Context
 
 namespace impl
 {
-	void Task_Mutable_UpdateMesh_End(const TSharedPtr<FMutableUpdateMeshContext> OperationData, TRefCountPtr<FCustomizableObjectMeshStreamIn>& Task, const TSharedRef<UE::Mutable::Private::FLiveInstance>& LiveInstance)
+	void Task_Mutable_UpdateMesh_End(const TSharedPtr<FMutableUpdateMeshContext>& OperationData, TRefCountPtr<FCustomizableObjectMeshStreamIn>& StreamingTask)
 	{
 		MUTABLE_CPUPROFILER_SCOPE(Task_Mutable_UpdateMesh_End);
-
+		
 		// End update
 		const bool bClearRoms = UCustomizableObjectSystem::ShouldClearWorkingMemoryOnUpdateEnd();
-		OperationData->System->EndUpdate(LiveInstance, bClearRoms);
+		OperationData->System->EndUpdate(OperationData->LiveInstance, bClearRoms);
 
-		check(OperationData->MeshStreamingPrerequisite.IsValid())
-		OperationData->MeshStreamingPrerequisite->Trigger();
+		// Stream data (morphs, clothing...).
+		UE::Tasks::TTask StreamDataTask = UE::Tasks::Launch(TEXT("StreamData"),
+		[OperationData]()
+		{
+			FMutableStreamRequest StreamRequest(OperationData->ModelStreamableBulkData);
+	
+			// Get morphs to be streamed.
+			for (UE::Mutable::Private::TManagedPtr<const UE::Mutable::Private::FMesh>& Mesh : OperationData->Meshes)
+			{
+				if (!Mesh)
+				{
+					continue;
+				}
+	
+				LoadMorphTargetsData(StreamRequest, Mesh.ToStrongRef(), OperationData->MorphTargetMeshData);
+				LoadMorphTargetsMetadata(StreamRequest, Mesh.ToStrongRef(), OperationData->MorphTargetMeshData);
+			}
+	
+			// Stream data.
+			UE::Tasks::AddNested(StreamRequest.Stream());
+		},
+		UE::Tasks::ETaskPriority::Inherit);
+
+		// Stream Data should be executed as a child task for this one
+		UE::Tasks::AddNested(StreamDataTask);
+		
+		// Go to next step.
+		UE::Tasks::TTask<void> OnUpdateMeshFinishedTask = UE::Tasks::Launch(TEXT("OnUpdateMeshFinished"),
+			[StreamingTask]()
+				{
+					StreamingTask->OnUpdateMeshFinished();
+				},
+			StreamDataTask,
+			UE::Tasks::ETaskPriority::Inherit,
+			UE::Tasks::EExtendedTaskPriority::Inline);
+		
+		UE::Tasks::AddNested(OnUpdateMeshFinishedTask);
 		
 		TRACE_END_REGION(UE_MUTABLE_UPDATE_MESH_REGION);
 	}
 
 	
 	void Task_Mutable_UpdateMesh_Loop(
-		const TSharedPtr<FMutableUpdateMeshContext> OperationData,
+		const TSharedPtr<FMutableUpdateMeshContext>& OperationData,
 		TRefCountPtr<FCustomizableObjectMeshStreamIn>& Task,
-		const TSharedRef<UE::Mutable::Private::FLiveInstance>& LiveInstance,
 		int32 LODIndex)
 	{
 		MUTABLE_CPUPROFILER_SCOPE(Task_Mutable_UpdateMesh_Loop);
 
 		if (Task->IsCancelled() || LODIndex == OperationData->CurrentFirstLODIdx + OperationData->AssetLODBias)
 		{
-			Task_Mutable_UpdateMesh_End(OperationData, Task, LiveInstance);
+			Task_Mutable_UpdateMesh_End(OperationData, Task);
 			return;
 		}
 
 		const UE::Mutable::Private::FMeshId& MeshID = OperationData->MeshIDs[LODIndex];
 
 		constexpr UE::Mutable::Private::EMeshContentFlags MeshContentFilter = UE::Mutable::Private::EMeshContentFlags::AllFlags;
-		UE::Tasks::TTask<UE::Mutable::Private::TManagedPtr<const UE::Mutable::Private::FMesh>> GetMeshTask = OperationData->System->GetMesh({}, LiveInstance, MeshID, MeshContentFilter);
+		UE::Tasks::TTask<UE::Mutable::Private::TManagedPtr<const UE::Mutable::Private::FMesh>> GetMeshTask = OperationData->System->GetMesh({}, OperationData->LiveInstance, MeshID, MeshContentFilter);
 
 		UE::Tasks::AddNested(UE::Tasks::Launch(TEXT("Task_MutableGetMeshes_GetMesh_Post"), [=]() mutable
 			{
 				OperationData->Meshes[LODIndex] = GetMeshTask.GetResult();
 
-				Task_Mutable_UpdateMesh_Loop(OperationData, Task, LiveInstance, LODIndex + 1);
+				Task_Mutable_UpdateMesh_Loop(OperationData, Task, LODIndex + 1);
 			},
 			GetMeshTask,
 			LowLevelTasks::ETaskPriority::Inherit));
@@ -253,7 +284,7 @@ namespace impl
 		
 		check(Instance);
 
-		Task_Mutable_UpdateMesh_Loop(OperationData, Task, LiveInstance, OperationData->PendingFirstLODIdx + OperationData->AssetLODBias);
+		Task_Mutable_UpdateMesh_Loop(OperationData, Task, OperationData->PendingFirstLODIdx + OperationData->AssetLODBias);
 	}
 	
 } // namespace
@@ -314,41 +345,6 @@ void FCustomizableObjectMeshStreamIn::RequestMeshUpdate(const FContext& Context)
 			});
 	}
 
-	// Stream data (morphs, clothing...).
-	check(OperationData->MeshStreamingPrerequisite.IsValid())
-	
-	UE::Tasks::TTask StreamData = UE::Tasks::Launch(TEXT("StreamData"),
-		[RefThis = TRefCountPtr(this)]()
-		{
-			FMutableStreamRequest StreamRequest(RefThis->OperationData->ModelStreamableBulkData);
-
-			// Get morphs to be streamed.
-			for (UE::Mutable::Private::TManagedPtr<const UE::Mutable::Private::FMesh>& Mesh : RefThis->OperationData->Meshes)
-			{
-				if (!Mesh)
-				{
-					continue;
-				}
-
-				LoadMorphTargetsData(StreamRequest, Mesh.ToStrongRef(), RefThis->OperationData->MorphTargetMeshData);
-				LoadMorphTargetsMetadata(StreamRequest, Mesh.ToStrongRef(), RefThis->OperationData->MorphTargetMeshData);
-			}
-
-			// Stream data.
-			UE::Tasks::AddNested(StreamRequest.Stream());
-		},
-		*OperationData->MeshStreamingPrerequisite,
-		UE::Tasks::ETaskPriority::Inherit);
-
-	// Go to next step.
-	UE::Tasks::Launch(TEXT("OnUpdateMeshFinished"),
-		[RefThis = TRefCountPtr(this)]()
-		{
-			RefThis->OnUpdateMeshFinished();
-		},
-		StreamData,
-		UE::Tasks::ETaskPriority::Inherit,
-		UE::Tasks::EExtendedTaskPriority::Inline);
 	
 	if (IsCancelled() && TaskSynchronization.GetValue() > 0 && bEnableGCHangFix)
 	{
diff --git a/Engine/Plugins/Mutable/Source/CustomizableObject/Private/MuCO/CustomizableObjectMeshUpdate.h b/Engine/Plugins/Mutable/Source/CustomizableObject/Private/MuCO/CustomizableObjectMeshUpdate.h
--- a/Engine/Plugins/Mutable/Source/CustomizableObject/Private/MuCO/CustomizableObjectMeshUpdate.h
+++ b/Engine/Plugins/Mutable/Source/CustomizableObject/Private/MuCO/CustomizableObjectMeshUpdate.h
@@ -44,9 +44,6 @@ struct FMutableUpdateMeshContext
 	/** Serialized data per block. */
 	TMap<uint32, FMorphTargetMeshData> MorphTargetMeshData;
 
-	/** Task setup during the nominal operation of the mesh streaming system. It will be null in non-streaming operation. */
-	TSharedPtr<UE::Tasks::FTaskEvent> MeshStreamingPrerequisite;
-
 	/** SkeletalMesh */
 	/** The resident first LOD resource index.With domain = [0, ResourceState.NumLODs[.NOT THE ASSET LOD INDEX! */ 
 	int32 CurrentFirstLODIdx = INDEX_NONE;
