[Mutable] Retrived check that indicates if a surface edit node is being used correctly.
+Removed uncessary textures from edit nodes.

#rnx

#Synced-CL 50245735

diff --git a/Engine/Plugins/Mutable/Source/CustomizableObjectEditor/Private/MuCOE/GenerateMutableSource/GenerateMutableSourceMaterial.cpp b/Engine/Plugins/Mutable/Source/CustomizableObjectEditor/Private/MuCOE/GenerateMutableSource/GenerateMutableSourceMaterial.cpp
--- a/Engine/Plugins/Mutable/Source/CustomizableObjectEditor/Private/MuCOE/GenerateMutableSource/GenerateMutableSourceMaterial.cpp
+++ b/Engine/Plugins/Mutable/Source/CustomizableObjectEditor/Private/MuCOE/GenerateMutableSource/GenerateMutableSourceMaterial.cpp
@@ -189,6 +189,7 @@ UE::Mutable::Private::Ptr<UE::Mutable::Private::NodeMaterial> GenerateMutableSou
 					Image.ImageNode = PassThroughImagePtr;
 					Image.ImagePropertyIndex = Props.ImagePropertiesIndex;
 					Image.LayoutIndex = -1;
+					Image.bIsPassthrough = true;
 
 					BaseMaterial->ImageParameters.Add(ImageKey, Image);
 				}
@@ -323,6 +324,7 @@ UE::Mutable::Private::Ptr<UE::Mutable::Private::NodeMaterial> GenerateMutableSou
 					Image.ImageNode = ImageNode;
 					Image.ImagePropertyIndex = Props.ImagePropertiesIndex;
 					Image.LayoutIndex = UVLayout;
+					Image.bIsPassthrough = false;
 
 					//Store the image inside the current generated material
 					BaseMaterial->ImageParameters.Add(ImageKey, Image);
@@ -748,6 +750,7 @@ UE::Mutable::Private::Ptr<UE::Mutable::Private::NodeMaterial> GenerateMutableSou
 					ImageData.ImageNode = ImageNode;
 					ImageData.ImagePropertyIndex = Props.ImagePropertiesIndex;
 					ImageData.LayoutIndex = TypedNodeMaterialModify->GetImagePinParameterUVIndex(*InputPin);
+					ImageData.bIsPassthrough = false;
 				}
 
 				// We do not check Null images since it is possible to not have anything linked
@@ -772,6 +775,7 @@ UE::Mutable::Private::Ptr<UE::Mutable::Private::NodeMaterial> GenerateMutableSou
 					ImageData.ImageNode = PassThroughImagePtr;
 					ImageData.ImagePropertyIndex = Props.ImagePropertiesIndex;
 					ImageData.LayoutIndex = -1;
+					ImageData.bIsPassthrough = true;
 				}
 
 				// We do not check Null images since it is possible to not have anything linked
diff --git a/Engine/Plugins/Mutable/Source/CustomizableObjectEditor/Private/MuCOE/GenerateMutableSource/GenerateMutableSourceModifier.cpp b/Engine/Plugins/Mutable/Source/CustomizableObjectEditor/Private/MuCOE/GenerateMutableSource/GenerateMutableSourceModifier.cpp
--- a/Engine/Plugins/Mutable/Source/CustomizableObjectEditor/Private/MuCOE/GenerateMutableSource/GenerateMutableSourceModifier.cpp
+++ b/Engine/Plugins/Mutable/Source/CustomizableObjectEditor/Private/MuCOE/GenerateMutableSource/GenerateMutableSourceModifier.cpp
@@ -317,7 +317,6 @@ UE::Mutable::Private::Ptr<UE::Mutable::Private::NodeModifier> GenerateMutableSou
 			SurfNode->LODs[LODIndex].MeshAdd = AddMeshNode;
 
 			const int32 NumImages = TypedNodeExt->GetNumParameters(EMaterialParameterType::Texture);
-			SurfNode->LODs[LODIndex].Textures.SetNum(NumImages);
 			for (int32 ImageIndex = 0; ImageIndex < NumImages; ++ImageIndex)
 			{
 				UE::Mutable::Private::NodeImagePtr ImageNode;
@@ -354,9 +353,14 @@ UE::Mutable::Private::Ptr<UE::Mutable::Private::NodeModifier> GenerateMutableSou
 					}
 				}
 
+				if (ImageNode)
+				{
+					UE::Mutable::Private::NodeModifierSurfaceEdit::FTexture Texture;
+					Texture.Extend = ImageNode;
+					Texture.MaterialParameterKey = UE::Mutable::Private::FParameterKey(MaterialParameterName, (int8)MaterialParameterLayerIndex);
 
-				SurfNode->LODs[LODIndex].Textures[ImageIndex].Extend = ImageNode;
-				SurfNode->LODs[LODIndex].Textures[ImageIndex].MaterialParameterKey = UE::Mutable::Private::FParameterKey(MaterialParameterName, MaterialParameterLayerIndex);
+					SurfNode->LODs[LODIndex].Textures.Add(Texture);
+				}
 			}
 		}
 
@@ -450,9 +454,8 @@ UE::Mutable::Private::Ptr<UE::Mutable::Private::NodeModifier> GenerateMutableSou
 			GenerationContext.FromLOD = 0;
 			GenerationContext.CurrentLOD = LODIndex;
 
-
 			const int32 NumImages = TypedNodeEdit->GetNumParameters(EMaterialParameterType::Texture);
-			SurfNode->LODs[LODIndex].Textures.SetNum(NumImages);
+
 			for (int32 ImageIndex = 0; ImageIndex < NumImages; ++ImageIndex)
 			{
 				const FNodeMaterialParameterId ImageId = TypedNodeEdit->GetParameterId(EMaterialParameterType::Texture, ImageIndex);
@@ -463,11 +466,11 @@ UE::Mutable::Private::Ptr<UE::Mutable::Private::NodeModifier> GenerateMutableSou
 					//check(ParentMaterialNode->IsImageMutableMode(ImageIndex)); // Ensured at graph time. If it fails, something is wrong.
 
 					const UEdGraphPin* ConnectedImagePin = FollowInputPin(*TypedNodeEdit->GetUsedImagePin(ImageId));
-					const FName MaterialParameterName = TypedNodeEdit->GetParameterName(EMaterialParameterType::Texture, ImageIndex);
+					const FName& MaterialParameterName = TypedNodeEdit->GetParameterName(EMaterialParameterType::Texture, ImageIndex);
 					const int32 MaterialParameterLayerIndex = TypedNodeEdit->GetParameterLayerIndex(EMaterialParameterType::Texture, ImageIndex);
 					check(MaterialParameterLayerIndex <= MAX_int8);
 					
-					UE::Mutable::Private::NodeModifierSurfaceEdit::FTexture& ImagePatch = SurfNode->LODs[LODIndex].Textures[ImageIndex];
+					UE::Mutable::Private::NodeModifierSurfaceEdit::FTexture ImagePatch;
 					ImagePatch.MaterialParameterKey = UE::Mutable::Private::FParameterKey(MaterialParameterName,  (int8)MaterialParameterLayerIndex);
 					// \todo: expose these two options?
 					ImagePatch.PatchBlendType = UE::Mutable::Private::EBlendType::BT_BLEND;
@@ -490,6 +493,7 @@ UE::Mutable::Private::Ptr<UE::Mutable::Private::NodeModifier> GenerateMutableSou
 					FIntPoint GridSize = TypedNodeEdit->Layout->GetGridSize();
 					FVector2f GridSizeF = FVector2f(GridSize);
 					ImagePatch.PatchBlocks.Reserve(TypedNodeEdit->Layout->Blocks.Num());
+
 					for (const FCustomizableObjectLayoutBlock& LayoutBlock : TypedNodeEdit->Layout->Blocks)
 					{
 						FBox2f Rect;
@@ -497,6 +501,8 @@ UE::Mutable::Private::Ptr<UE::Mutable::Private::NodeModifier> GenerateMutableSou
 						Rect.Max = FVector2f(LayoutBlock.Max) / GridSizeF;
 						ImagePatch.PatchBlocks.Add(Rect);
 					}
+
+					SurfNode->LODs[LODIndex].Textures.Add(ImagePatch);
 				}
 			}
 		}
diff --git a/Engine/Plugins/Mutable/Source/MutableTools/Private/MuT/CodeGenerator.cpp b/Engine/Plugins/Mutable/Source/MutableTools/Private/MuT/CodeGenerator.cpp
--- a/Engine/Plugins/Mutable/Source/MutableTools/Private/MuT/CodeGenerator.cpp
+++ b/Engine/Plugins/Mutable/Source/MutableTools/Private/MuT/CodeGenerator.cpp
@@ -874,9 +874,8 @@ namespace UE::Mutable::Private
 		int32 ComponentId = Options.Component ? Options.Component->Id : -1;
 		GetModifiersFor(ComponentId, SurfaceNode->Tags, bModifiersForBeforeOperations, Modifiers);
 
-		// TODO (Max): UE-354034 Fix this function and uncomment it
 		// This pass on the modifiers is only to detect errors that cannot be detected at the point they are applied.
-		//CheckModifiersForSurface(*SurfaceNode, Modifiers, Options.LODIndex);
+		CheckModifiersForSurface(*SurfaceNode, Modifiers, Options.LODIndex);
 		
         // Generate the mesh
 		FMeshGenerationStaticOptions MeshStaticOptions(ComponentId, Options.LODIndex);
@@ -2769,38 +2768,25 @@ namespace UE::Mutable::Private
 				const NodeModifierSurfaceEdit* Edit = static_cast<const NodeModifierSurfaceEdit*>(Mod.Node);
 				const NodeMaterialModify* Material = static_cast<const NodeMaterialModify*>(Node.Material.get());
 
-				bool bAffectsCurrentLOD = Edit->LODs.IsValidIndex(CurrentLOD);
-				if (!bAffectsCurrentLOD)
-				{
-					continue;
-				}
-
-				if (Material->ImageParameters.IsEmpty() || Edit->LODs[CurrentLOD].Textures.IsEmpty())
+				// If this edit is not affected by the current LOD, then ignore the check.
+				// If both arrays are empty means that there are no textures to check
+				// If there is a MeshRemove, this edit does not affect textures
+				if (!Edit->LODs.IsValidIndex(CurrentLOD) || Edit->LODs[CurrentLOD].MeshRemove || (Material->ImageParameters.IsEmpty() && Edit->LODs[CurrentLOD].Textures.IsEmpty()))
 				{
 					continue;
 				}
 
 				bool bAtLeastSomeTexture = false;
 
-				for (TPair<FParameterKey, FImageParameterData> Data : Material->ImageParameters)
+				for (const NodeModifierSurfaceEdit::FTexture& Texture : Edit->LODs[CurrentLOD].Textures)
 				{
-					const NodeModifierSurfaceEdit::FTexture* MatchingEdit = Edit->LODs[CurrentLOD].Textures.FindByPredicate(
-						[&](const NodeModifierSurfaceEdit::FTexture& Candidate)
-						{
-							return (Candidate.MaterialParameterKey == Data.Key);
-						});
+					const FImageParameterData* MatchingEdit = Material->ImageParameters.Find(Texture.MaterialParameterKey);
 
-					if (MatchingEdit)
+					if (!MatchingEdit || MatchingEdit->bIsPassthrough)
 					{
-						bAtLeastSomeTexture = true;
-						break;
+						ErrorLog->Add(FString::Printf(TEXT("Trying to Modify the Passthrough texture [%s]."), *Texture.MaterialParameterKey.ParameterName.ToString()), ELMT_WARNING, Edit->GetMessageContext(), Node.GetMessageContext());
 					}
 				}
-
-				if (!bAtLeastSomeTexture)
-				{
-					ErrorLog->Add(TEXT("A mesh section modifier applies to a section but no texture matches."), ELMT_WARNING, Edit->GetMessageContext(), Node.GetMessageContext());
-				}
 			}
 		}
 	}
diff --git a/Engine/Plugins/Mutable/Source/MutableTools/Public/MuT/NodeMaterialModify.h b/Engine/Plugins/Mutable/Source/MutableTools/Public/MuT/NodeMaterialModify.h
--- a/Engine/Plugins/Mutable/Source/MutableTools/Public/MuT/NodeMaterialModify.h
+++ b/Engine/Plugins/Mutable/Source/MutableTools/Public/MuT/NodeMaterialModify.h
@@ -19,11 +19,9 @@ namespace UE::Mutable::Private
 		Ptr<NodeImage> ImageNode;
 		int32 ImagePropertyIndex = 0;
 		int8 LayoutIndex = 0;
+		bool bIsPassthrough = false;
 
-		bool operator==(const FImageParameterData& Other) const
-		{
-			return ImageNode == Other.ImageNode && ImagePropertyIndex == Other.ImagePropertyIndex && LayoutIndex == Other.LayoutIndex;
-		}
+		bool operator==(const FImageParameterData& Other) const = default;
 	};
 
 	class NodeMaterialModify : public NodeMaterial
