[Mutable] Fix visual artifact when mesh streaming is enabled, and a section from a streamed LOD is disabled.
Missing sections will remain enabled and render a collapsed triangle.
Change the execution thread of the cleanup that may occur after FCustomizableObjectMeshStreamIn::MarkRenderStateDirty to TT_Async.

#rnx

#Synced-CL 50297267

#Synced-CL 50297267

diff --git a/Engine/Plugins/Mutable/Source/CustomizableObject/Private/MuCO/CustomizableObjectMeshUpdate.cpp b/Engine/Plugins/Mutable/Source/CustomizableObject/Private/MuCO/CustomizableObjectMeshUpdate.cpp
--- a/Engine/Plugins/Mutable/Source/CustomizableObject/Private/MuCO/CustomizableObjectMeshUpdate.cpp
+++ b/Engine/Plugins/Mutable/Source/CustomizableObject/Private/MuCO/CustomizableObjectMeshUpdate.cpp
@@ -459,5 +459,5 @@ void FCustomizableObjectMeshStreamIn::MarkRenderStateDirty(const FContext& Conte
 		Abort();
 	}
 	
-	PushTask(Context, CreateResourcesThread, SRA_UPDATE_CALLBACK(DoCreateBuffers), (EThreadType)Context.CurrentThread, SRA_UPDATE_CALLBACK(DoCancel));
+	PushTask(Context, CreateResourcesThread, SRA_UPDATE_CALLBACK(DoCreateBuffers), (EThreadType)TT_Async, SRA_UPDATE_CALLBACK(DoCancel));
 }
diff --git a/Engine/Plugins/Mutable/Source/CustomizableObject/Private/MuCO/UnrealConversionUtils.cpp b/Engine/Plugins/Mutable/Source/CustomizableObject/Private/MuCO/UnrealConversionUtils.cpp
--- a/Engine/Plugins/Mutable/Source/CustomizableObject/Private/MuCO/UnrealConversionUtils.cpp
+++ b/Engine/Plugins/Mutable/Source/CustomizableObject/Private/MuCO/UnrealConversionUtils.cpp
@@ -530,9 +530,9 @@ namespace UnrealConversionUtils
 	{
 		MUTABLE_CPUPROFILER_SCOPE(CopyMutableIndexBuffers);
 
-		const int32 IndexCount = InMutableMesh->GetIndexBuffers().GetElementCount();
+		const int32 MutableIndexCount = InMutableMesh->GetIndexBuffers().GetElementCount();
 
-		if (IndexCount == 0)
+		if (MutableIndexCount == 0)
 		{
 			// Copy indices from an empty buffer
 			UE_LOG(LogMutable, Error, TEXT("UCustomizableInstancePrivateData::BuildSkeletalMeshRenderData is converting an empty mesh."));
@@ -555,11 +555,18 @@ namespace UnrealConversionUtils
 		
 		check(LODResource.MultiSizeIndexContainer.GetDataTypeSize() == ElementSize);
 		
+		// UE-363228: Add a collapsed triangle whenever instead of disabling missing sections
+		const bool bAddCollapsedTriangle = LODResource.RenderSections.Num() != InMutableMesh->Surfaces.Num();
+		const int32 FirstIndex = bAddCollapsedTriangle ? 3 : 0;
+
 		// Don't assume the buffer is empty, otherwise we may add extra elements using Insert(). 
-		LODResource.MultiSizeIndexContainer.GetIndexBuffer()->Empty(IndexCount);
-		LODResource.MultiSizeIndexContainer.GetIndexBuffer()->Insert(0, IndexCount);
+		{
+			const int32 IndexCount = MutableIndexCount + FirstIndex;
+			LODResource.MultiSizeIndexContainer.GetIndexBuffer()->Empty(IndexCount);
+			LODResource.MultiSizeIndexContainer.GetIndexBuffer()->Insert(0, IndexCount);
+		}
 
-		FMemory::Memcpy(LODResource.MultiSizeIndexContainer.GetIndexBuffer()->GetPointerTo(0), DataPtr, IndexCount * ElementSize);
+		FMemory::Memcpy(LODResource.MultiSizeIndexContainer.GetIndexBuffer()->GetPointerTo(FirstIndex), DataPtr, MutableIndexCount * ElementSize);
 		
 		for (int32 SectionIndex = 0; SectionIndex < LODResource.RenderSections.Num(); ++SectionIndex)
 		{
@@ -587,7 +594,7 @@ namespace UnrealConversionUtils
 				Section.NumTriangles = NumTriangles;
 
 				bOutMarkRenderStateDirty |= Section.BaseIndex != IndexBegin;
-				Section.BaseIndex = IndexBegin;
+				Section.BaseIndex = IndexBegin + FirstIndex;
 
 				const int32 VertexBegin = Surface->SubMeshes[0].VertexBegin;
 
@@ -596,8 +603,8 @@ namespace UnrealConversionUtils
 			}
 			else
 			{
-				Section.bDisabled = true;
-				Section.NumTriangles = 0;
+				// Section.bDisabled = true; // UE-363228: Disabling sections after the mesh is generated produces visual artifacts.
+				Section.NumTriangles = 1; // UE-363228: Add collapsed triangle to avoid visual artifacts
 				Section.NumVertices = 0;
 				Section.BaseIndex = 0;
 				Section.BaseVertexIndex = 0;
