[Mutable] Fix cloth crash when streaming is enabled + remove mesh.

#Synced-CL 50394305

diff --git a/Engine/Plugins/Mutable/Source/CustomizableObject/Private/MuCO/CustomizableObjectInstance.cpp b/Engine/Plugins/Mutable/Source/CustomizableObject/Private/MuCO/CustomizableObjectInstance.cpp
--- a/Engine/Plugins/Mutable/Source/CustomizableObject/Private/MuCO/CustomizableObjectInstance.cpp
+++ b/Engine/Plugins/Mutable/Source/CustomizableObject/Private/MuCO/CustomizableObjectInstance.cpp
@@ -4087,6 +4087,14 @@ void BuildOrCopyClothingData(
 			Mappings.Add(Mapping);
 						
 			NumVertices += RenderSection.NumVertices;
+			
+			// Cloth hack. Initialize to something dummy (size is an upper bound). During streaming this data will be set correctly.
+			const UE::Mutable::Private::FCloth& Cloth = LOD->Mesh->ClothSections[SectionIndex];
+			if (Cloth.IsValid())
+			{
+				RenderSection.ClothMappingDataLODs.SetNum(1);
+				RenderSection.ClothMappingDataLODs[0].SetNum(RenderSection.NumVertices);
+			}
 		}
 			
 		RenderLOD.ClothVertexBuffer.SetMetadata(Mappings, Stride, NumVertices);
@@ -4135,12 +4143,7 @@ void BuildOrCopyClothingData(
 			FSkelMeshRenderSection& RenderSection = RenderLOD.RenderSections[SectionIndex];
 			
 			UClothingAssetCommon* Asset = Cloth.ClothingAsset.Get();
-			const uint32 NumClothVertices = Cloth.Data.Num() ? Cloth.Data.Num() : RenderSection.NumVertices;
-			
 			RenderSection.CorrespondClothAssetIndex = ClothingAssets.AddUnique(Asset);
-
-			RenderSection.ClothMappingDataLODs.AddDefaulted(1);
-			RenderSection.ClothMappingDataLODs[0].SetNum(NumClothVertices);
 			
 			FClothingSectionData& ClothingData = RenderSection.ClothingData;
 			ClothingData.AssetGuid = Asset->GetAssetGuid();
diff --git a/Engine/Plugins/Mutable/Source/CustomizableObject/Private/MuCO/UnrealConversionUtils.cpp b/Engine/Plugins/Mutable/Source/CustomizableObject/Private/MuCO/UnrealConversionUtils.cpp
--- a/Engine/Plugins/Mutable/Source/CustomizableObject/Private/MuCO/UnrealConversionUtils.cpp
+++ b/Engine/Plugins/Mutable/Source/CustomizableObject/Private/MuCO/UnrealConversionUtils.cpp
@@ -1922,8 +1922,14 @@ namespace UnrealConversionUtils
 				TArray<FClothBufferIndexMapping> Mappings;
 				Mappings.Add(Mapping);
 			
-				TArrayView<const FMeshToMeshVertData> DataView = TArrayView(MutableMesh.ClothSections[0].Data.GetData(), MutableMesh.ClothSections[0].Data.Num());
+				TArrayView<const FMeshToMeshVertData> DataView = TArrayView(MutableMesh.ClothSections[0].Data.GetData(), NumClothVertices);
 				LODResource.ClothVertexBuffer.Init(DataView, Mappings);
+			
+				// Cloth hack. This data should not be used.
+				FSkelMeshRenderSection& RenderSection = LODResource.RenderSections[0];
+				RenderSection.ClothMappingDataLODs.SetNum(1, EAllowShrinking::No);
+				RenderSection.ClothMappingDataLODs[0].SetNum(NumClothVertices, EAllowShrinking::No);
+				FMemory::Memcpy(&RenderSection.ClothMappingDataLODs[0][0], MutableMesh.ClothSections[0].Data.GetData(), NumClothVertices * sizeof(FMeshToMeshVertData));
 			}
 		}
 		else if (NumSections > 1)
@@ -1954,27 +1960,35 @@ namespace UnrealConversionUtils
 					const uint32 BaseVertexIndex = LODResource.RenderSections[SectionIndex].BaseVertexIndex;
 				
 					FClothBufferIndexMapping& Mapping = Mappings[SectionIndex];
-					if (Cloth.IsValid())
+					
+					const int32 NumClothVertices = Cloth.Data.Num();
+					if (Cloth.IsValid() && NumClothVertices)
 					{
 						Mapping.BaseVertexIndex = BaseVertexIndex; 
 						Mapping.MappingOffset = VertexOffset;
 						Mapping.LODBiasStride = VertexOffset;
-
-						VertexOffset += Cloth.Data.Num();
+					
+						VertexOffset += NumClothVertices;
+	
+						// Cloth hack. This data should not be used.
+						FSkelMeshRenderSection& RenderSection = LODResource.RenderSections[SectionIndex];
+						RenderSection.ClothMappingDataLODs.SetNum(1, EAllowShrinking::No);
+						RenderSection.ClothMappingDataLODs[0].SetNum(NumClothVertices, EAllowShrinking::No);
+						FMemory::Memcpy(&RenderSection.ClothMappingDataLODs[0][0], Cloth.Data.GetData(), NumClothVertices * sizeof(FMeshToMeshVertData));
 					}
 					else
 					{
 						Mapping.BaseVertexIndex = 0; 
 						Mapping.MappingOffset = 0;
 						Mapping.LODBiasStride = 0;
-					}
+					}					
 				}
 
 				LODResource.ClothVertexBuffer.Init(Data, Mappings);
 			}
 		}
 	}
-
+	
 
 	void PropagateBoneUsageFlagsThroughMeshPose(UE::Mutable::Private::FMesh& Mesh)
 	{
