[Mutable] Change realtime morph reconstruction to use compressed data generated at runtime. Use Morph compressed data for all morph operations.

- Moved OpMeshMerge to the Private folder.
- Removed uneeded cloth buffers setup.

#rnx

#Synced-CL 50400930

diff --git a/Engine/Plugins/Mutable/Source/CustomizableObject/Internal/MuCO/CustomizableObjectPrivate.h b/Engine/Plugins/Mutable/Source/CustomizableObject/Internal/MuCO/CustomizableObjectPrivate.h
--- a/Engine/Plugins/Mutable/Source/CustomizableObject/Internal/MuCO/CustomizableObjectPrivate.h
+++ b/Engine/Plugins/Mutable/Source/CustomizableObject/Internal/MuCO/CustomizableObjectPrivate.h
@@ -90,7 +90,6 @@ struct FCustomizableObjectStreameableResourceId
 	{
 		None                  = 0,
 		AssetUserData         = 1,
-		RealTimeMorphTarget   = 2,
 		Socket	              = 4,
 	};
 
@@ -554,48 +553,16 @@ struct FMutableStreamableBlock
 static_assert(sizeof(FMutableStreamableBlock) == 8 * 2);
 
 
-USTRUCT()
-struct FRealTimeMorphStreamable
-{
-	GENERATED_USTRUCT_BODY()
-	
-	UPROPERTY()
-	TArray<FName> NameResolutionMap;
-
-	UPROPERTY()
-	FMutableStreamableBlock Block;
-
-	UPROPERTY()
-	uint32 Size = 0;
-
-	UPROPERTY()
-	uint32 SourceId = 0;
-
-	friend FArchive& operator<<(FArchive& Ar, FRealTimeMorphStreamable& Elem)
-	{
-		Ar << Elem.NameResolutionMap;
-		Ar << Elem.Size;
-		Ar << Elem.Block;
-		Ar << Elem.SourceId;
-
-		return Ar;
-	}
-};
-
 USTRUCT()
 struct FMutableMeshMetadata
 {
 	GENERATED_USTRUCT_BODY()
-	
-	UPROPERTY()
-	uint32 MorphMetadataId = 0;
-	
+
 	UPROPERTY()
 	uint32 SurfaceMetadataId = 0;
 
 	friend FArchive& operator<<(FArchive& Ar, FMutableMeshMetadata& Elem)
 	{
-		Ar << Elem.MorphMetadataId;
 		Ar << Elem.SurfaceMetadataId;
 
 		return Ar;
@@ -624,35 +591,9 @@ struct FMutableSurfaceMetadata
 };
 
 
-USTRUCT()
-struct FMorphTargetVertexData
-{
-	GENERATED_USTRUCT_BODY()
-
-	UPROPERTY()
-	FVector3f PositionDelta = FVector3f::ZeroVector;
-
-	UPROPERTY()
-	FVector3f TangentZDelta = FVector3f::ZeroVector;
-	
-	UPROPERTY()
-	uint32 MorphNameIndex = 0;
-
-	friend FArchive& operator<<(FArchive& Ar, FMorphTargetVertexData& Data)
-	{
-		Ar << Data.PositionDelta;
-		Ar << Data.TangentZDelta;
-		Ar << Data.MorphNameIndex;
-
-		return Ar;
-	}
-};
-static_assert(sizeof(FMorphTargetVertexData) == sizeof(FVector3f)*2 + sizeof(uint32)); // Make sure no padding is present.
-template<> struct TCanBulkSerialize<FMorphTargetVertexData> { enum { Value = true }; };
-
-
 struct FMutableParameterIndex
 {
+
 	FMutableParameterIndex(int32 InIndex, int32 InTypedIndex)
 	{
 		Index = InIndex;
@@ -778,7 +719,6 @@ namespace UE::Mutable::Private
 	{
 		None = 0,
 		Model,
-		RealTimeMorph,
 		
 		DataTypeCount
 	};
@@ -949,9 +889,6 @@ namespace UE::Mutable::Private
 		/** Struct containing map of RomId to RomBytes. */
 		FModelStreamableData ModelStreamableData;
 
-		/** */
-		FModelStreamableData MorphStreamableData;
-		
 		/** List of files to serialize. Each file has a list of binary blocks to be serialized. */
 		TArray<FFile> BulkDataFiles;
 	};
@@ -1006,8 +943,6 @@ struct FModelStreamableBulkData
 {
 	/** Map of Hash to Streaming blocks, used to stream a block of data representing a resource from the BulkData */
 	TMap<uint32, FMutableStreamableBlock> ModelStreamables;
-	
-	TMap<uint32, FRealTimeMorphStreamable> RealTimeMorphStreamables;
 
 	TArray<FByteBulkData> StreamableBulkData;
 
@@ -1028,7 +963,6 @@ struct FModelStreamableBulkData
 	friend FArchive& operator<<(FArchive& Ar, FModelStreamableBulkData& Data)
 	{
 		Ar << Data.ModelStreamables;
-		Ar << Data.RealTimeMorphStreamables;
 		Ar << Data.DDCValues;
 		// Don't serialize FByteBulkData manually, the data will be skipped.
 		
@@ -1552,11 +1486,11 @@ enum class ECustomizableObjectVersions : int32
 	ClothCore,
 	
 	ConvertBoneIdsToBoneNames,
-
+	
 	AddLayerIndexToMatParamTextures,
-
+	
 	ClothCoreSerializationIssue,
-
+	
 	ClothCoreSerializationIssue2,
 
 	MovePhysicsAssetCore,
@@ -1565,6 +1499,8 @@ enum class ECustomizableObjectVersions : int32
 
 	MutableSkeletonToPassthroughObject,
 	
+	UseCompressedMorphsForAllMorphData,
+	
 	// -----<new versions can be added above this line>--------
 	LastCustomizableObjectVersion
 };
diff --git a/Engine/Plugins/Mutable/Source/CustomizableObject/Internal/MuCO/CustomizableObjectSystemPrivate.h b/Engine/Plugins/Mutable/Source/CustomizableObject/Internal/MuCO/CustomizableObjectSystemPrivate.h
--- a/Engine/Plugins/Mutable/Source/CustomizableObject/Internal/MuCO/CustomizableObjectSystemPrivate.h
+++ b/Engine/Plugins/Mutable/Source/CustomizableObject/Internal/MuCO/CustomizableObjectSystemPrivate.h
@@ -57,44 +57,17 @@ struct FMutableImageReference
 	TArray<int32> ConstantImagesNeededToGenerate;
 };
 
+/** 
+ * Reconstruct the final Morph Targets using the mesh compressed morphs.
+ * @param Mesh Mesh containing the compressed data.
+ * @param OutMorphTargets Resulting Morph Targets LOD models filtered by usage.
+ * @param UsageFilter Usage filter applied to the OutMorphTargets.
+ */
+void ReconstructMorphTargetsFromMeshCompressedData(const UE::Mutable::Private::FMesh& Mesh, TArray<FMorphTargetLODModel>& OutMorphTargets, UE::Mutable::Private::EMorphUsageFlags UsageFilter);
 
-/** Serialized Morph data. */
-struct FMorphTargetMeshData
-{
-	/** Names of the Morph Targets. */
-	TArray<FName> NameResolutionMap;
-
-	/** Vertex data. Each vertex contains an index to the NameResolutionMap. */
-	TArray<FMorphTargetVertexData> Data;
-};
-
-
-/** Mapping of FMorphTargetVertexData local names to global names. */
-struct FMappedMorphTargetMeshData
-{
-	/** Index is the local name. Value is the index to a global name table. */
-	TArray<int32> NameResolutionMap;
-
-	/** Pointer to the original data with indices to local names. */
-	const TArray<FMorphTargetVertexData>* DataView = nullptr;
-};
-
-
-/* Reconstruct the final Morph Targets using the global names.
- * @param Mesh Used to know which MappedMorphTargets vertices must be removed.
- * @param GlobalNames
- * @param OutMorphTargets resulting Morph Targets with the indices pointing to GlobalNames. */
-void ReconstructMorphTargets(const UE::Mutable::Private::FMesh& Mesh, const TArray<FName>& GlobalNames, const TMap<uint32, FMappedMorphTargetMeshData>& MappedMorphTargets, TArray<FMorphTargetLODModel>& OutMorphTargets);
-
-
-/** Request the Mutable Data Streamer to load Morph Target blocks. Does not stream them.
- *  @param MutableDataStreamer Streamer to request the blocks.
- *  @param Mesh Mutable mesh to obtain which blocks to load.
- *  @param StreamingResult Key is the block. Value is the container where the data will be deserialized once streamed. */
-void LoadMorphTargetsData(FMutableStreamRequest& MutableDataStreamer, const UE::Mutable::Private::TManagedRef<const UE::Mutable::Private::FMesh>& Mesh, TMap<uint32, FMorphTargetMeshData>& StreamingResult);
-
-void LoadMorphTargetsMetadata(const FMutableStreamRequest& MutableDataStreamer, const UE::Mutable::Private::TManagedRef<const UE::Mutable::Private::FMesh>& Mesh, TMap<uint32, FMorphTargetMeshData>& StreamingResult);
-
+#if !WITH_EDITOR
+void ReconstructMorphTargetsFromMeshCompressedData(const UE::Mutable::Private::FMesh& Mesh, TArray<FMorphTargetCompressedLODModel>& OutMorphTargets, UE::Mutable::Private::EMorphUsageFlags UsageFilter);
+#endif
 /** End a Customizable Object Instance Update. All code paths of an update have to end here. */
 void FinishUpdateGlobal(const TSharedRef<FUpdateContextPrivate>& Context);
 
@@ -247,7 +220,7 @@ struct FSkeletalMeshMorphTargets
 	TArray<FName> RealTimeMorphTargetNames;
 
 	/** First index is the Morph Target (in sync with RealTimeMorphTargetNames). Second index is the LOD. */
-	TArray<TArray<FMorphTargetLODModel>> RealTimeMorphsLODData; 
+	TArray<TArray<FMorphTargetLODModel>> RealTimeMorphsLODData;
 };
 
 
@@ -389,9 +362,7 @@ struct FInstanceUpdateData
 	TArray<FName> ActiveBones;
 	TArray<FName> BoneMaps;
 	
-	/** Key is the Block Id. Value is the LoadAdditionalAssetsAndData read destination. */
-	TMap<uint32, FMorphTargetMeshData> RealTimeMorphTargetMeshData;
-	
+
 	/** Key is the Component Name. Value is the final Morph Target data to copy into the Skeletal Mesh. */
 	TMap<UE::Mutable::Private::FComponentId, FSkeletalMeshMorphTargets> RealTimeMorphTargets;
 
diff --git a/Engine/Plugins/Mutable/Source/CustomizableObject/Internal/MuCO/UnrealConversionUtils.h b/Engine/Plugins/Mutable/Source/CustomizableObject/Internal/MuCO/UnrealConversionUtils.h
--- a/Engine/Plugins/Mutable/Source/CustomizableObject/Internal/MuCO/UnrealConversionUtils.h
+++ b/Engine/Plugins/Mutable/Source/CustomizableObject/Internal/MuCO/UnrealConversionUtils.h
@@ -30,11 +30,13 @@ namespace UE::Mutable::Private
 }
 struct FReferenceSkeleton;
 struct FMutableSurfaceMetadata;
+struct FSkelMeshRenderSection;
 class UModelResources;
 class FSkeletalMeshLODRenderData;
 class USkeletalMesh;
 class USkeleton;
 class USkeletalBodySetup;
+class FMorphTargetVertexInfoBuffers;
 
 namespace UnrealConversionUtils
 {
@@ -157,7 +159,15 @@ namespace UnrealConversionUtils
 		FSkeletalMeshLODRenderData& LODResource
 	);
 
-	void MorphTargetVertexInfoBuffers(FSkeletalMeshLODRenderData& LODResource, const USkeletalMesh& Owner, const UE::Mutable::Private::FMesh& MutableMesh, const TMap<uint32, FMorphTargetMeshData>& MorphTargetMeshData, int32 LODIndex, bool bGenerateCPUMorphTargets);
+	void MorphTargetVertexInfoBuffers(FSkeletalMeshLODRenderData& LODResource, const USkeletalMesh& Owner, const UE::Mutable::Private::FMesh& MutableMesh, int32 LODIndex, bool bGenerateCPUMorphTargets);
+
+
+	CUSTOMIZABLEOBJECT_API void SetMorphData(
+			TNotNull<UE::Mutable::Private::FMesh*> Result, 
+			TNotNull<const USkeletalMesh*> SkeletalMesh,
+			int32 LODIndex, int32 SectionIndex, 
+			uint32 SourceSectionBaseVertexIndex, uint32 SourceSectionNumVertices, uint32 SourceMeshNumVertices, 
+			FMorphTargetVertexInfoBuffers* OptionalMorphTargetVertexBuffer);
 
 	CUSTOMIZABLEOBJECT_API UE::Tasks::FTask ConvertSkeletalMeshFromRuntimeData(USkeletalMesh* SkeletalMesh, int32 LODIndex, int32 SectionIndex, UE::Mutable::Private::TManagedPtr<UE::Mutable::Private::FMesh> Result);
 
diff --git a/Engine/Plugins/Mutable/Source/CustomizableObject/Private/MuCO/CustomizableObject.cpp b/Engine/Plugins/Mutable/Source/CustomizableObject/Private/MuCO/CustomizableObject.cpp
--- a/Engine/Plugins/Mutable/Source/CustomizableObject/Private/MuCO/CustomizableObject.cpp
+++ b/Engine/Plugins/Mutable/Source/CustomizableObject/Private/MuCO/CustomizableObject.cpp
@@ -858,7 +858,6 @@ FString GetDataTypeExtension(UE::Mutable::Private::EStreamableDataType DataType)
 	{
 	case UE::Mutable::Private::EStreamableDataType::None:	return TEXT(".mut");
 	case UE::Mutable::Private::EStreamableDataType::Model: return TEXT("_M.mut");
-	case UE::Mutable::Private::EStreamableDataType::RealTimeMorph: return TEXT("_RTM.mut");
 	default:
 	{
 		unimplemented();
@@ -2123,7 +2122,6 @@ void FModelStreamableBulkData::Serialize(FArchive& Ar, UObject* Owner, bool bCoo
 	MUTABLE_CPUPROFILER_SCOPE(FModelStreamableBulkData::Serialize);
 
 	Ar << ModelStreamables;
-	Ar << RealTimeMorphStreamables;
 
 	if (bCooked)
 	{
@@ -2465,15 +2463,6 @@ namespace UE::Mutable::Private
 			}
 			return;
 		}
-		else if (DataType == EStreamableDataType::RealTimeMorph)
-		{
-			for (int32 BlockIndex = 0; BlockIndex < NumBlocks; ++BlockIndex)
-			{
-				const FBlock& Block = Blocks[BlockIndex];
-				check(Block.Offset + Block.Size <= DestSize);
-				PlatformData->MorphStreamableData.Get(Block.Id, TArrayView64<uint8>(DestData.GetData() + Block.Offset, Block.Size), bDropData);
-			}
-		}
 		else
 		{
 			checkf(false, TEXT("Unknown file DataType found."));
@@ -2588,20 +2577,6 @@ namespace UE::Mutable::Private
 			}
 		}
 
-		// RealTime Morphs. Iterate RealTimeMorph streamables to accumulate their sizes.
-		{
-			// Add RealTimeMorphs to the Default bucket
-			FFileCategory& RealTimeMorphCategory = DefaultBucket.Categories.AddDefaulted_GetRef();
-			RealTimeMorphCategory.Id.DataType = EStreamableDataType::RealTimeMorph;
-
-			const TMap<uint32, FRealTimeMorphStreamable>& RealTimeMorphStreamables = ModelStreamableBulkData.RealTimeMorphStreamables;
-			for (const TPair<uint32, FRealTimeMorphStreamable>& MorphStreamable : RealTimeMorphStreamables)
-			{
-				RealTimeMorphCategory.DataSize += MorphStreamable.Value.Size;
-			}
-
-			DefaultBucket.DataSize += RealTimeMorphCategory.DataSize;
-		}
 		
 		// Limited number of files in each bucket. Find the ideal file distribution between categories based on the accumulated size of their resources.
 		TArray<FFileCategory> Categories;
@@ -2945,27 +2920,6 @@ namespace UE::Mutable::Private
 			}
 		}
 
-		{
-			const TMap<uint32, FRealTimeMorphStreamable>& RealTimeMorphStreamables = ModelStreamableBulkData.RealTimeMorphStreamables;
-
-			const FFileCategoryID RealTimeMorphCategory = { EStreamableDataType::RealTimeMorph, uint16(UE::Mutable::Private::EDataType::None), uint16(EMutableFileFlags::None) };
-
-			for (const TPair<uint32, FRealTimeMorphStreamable>& MorphStreamable : RealTimeMorphStreamables)
-			{
-				const uint32 BlockSize = MorphStreamable.Value.Size;
-
-				FBlock CurrentBlock = 
-				{ 
-					.Id = MorphStreamable.Key, 
-					.SourceId = MorphStreamable.Value.SourceId, 
-					.Size = BlockSize,
-					.Offset = 0 
-				};
-
-				AddNode(RootNode, NumRoms, RealTimeMorphCategory, CurrentBlock);
-			}
-		}
-		
 		// Create Files list
 		for (TPair<FFileCategoryID, FClassifyNode>& Node : RootNode)
 		{
@@ -2995,19 +2949,6 @@ namespace UE::Mutable::Private
 				}
 				break;
 			}
-			case EStreamableDataType::RealTimeMorph:
-			{
-				for (FBlock& Block : File.Blocks)
-				{
-					Block.Offset = SourceOffset;
-					SourceOffset += Block.Size;
-
-					FMutableStreamableBlock& StreamableBlock = ModelStreamableBulkData.RealTimeMorphStreamables[Block.Id].Block;
-					StreamableBlock.FileId = FileIndex;
-					StreamableBlock.Offset = Block.Offset;
-				}
-				break;
-			}
 			default:
 				UE_LOG(LogMutable, Error, TEXT("Unknown DataType found while fixing streaming block files ids."));
 				unimplemented();
@@ -3336,7 +3277,6 @@ void UCustomizableObjectPrivate::LogMemory()
 		{
 			const FModelStreamableBulkData& Bulk = *ModelStreamableData->StreamingData;
 			StreamableData += Bulk.ModelStreamables.GetAllocatedSize();
-			StreamableData += Bulk.RealTimeMorphStreamables.GetAllocatedSize();
 			StreamableData += Bulk.StreamableBulkData.GetAllocatedSize();
 		}
 
diff --git a/Engine/Plugins/Mutable/Source/CustomizableObject/Private/MuCO/CustomizableObjectInstance.cpp b/Engine/Plugins/Mutable/Source/CustomizableObject/Private/MuCO/CustomizableObjectInstance.cpp
--- a/Engine/Plugins/Mutable/Source/CustomizableObject/Private/MuCO/CustomizableObjectInstance.cpp
+++ b/Engine/Plugins/Mutable/Source/CustomizableObject/Private/MuCO/CustomizableObjectInstance.cpp
@@ -4472,44 +4472,6 @@ UE::Tasks::FTask LoadAdditionalAssetsAndData(const TSharedRef<FUpdateContextPriv
 		}
 
 		const FName ComponentName = Context->ComponentNames[Component->Id];
-
-		if (bUpdateMeshes)
-		{
-			// Morphs
-			{
-				// Data
-				for (int32 LODIndex = Context->GetFirstRequestedLOD()[ComponentName]; LODIndex < Component->LODCount; ++LODIndex)
-				{
-					const TSharedRef<FInstanceUpdateData::FLOD>& LOD = Context->InstanceUpdateData.LODs[Component->FirstLOD + LODIndex];
-
-					UE::Mutable::Private::TManagedPtr<const UE::Mutable::Private::FMesh> MutableMesh = LOD->Mesh;
-					if (!MutableMesh)
-					{
-						continue;
-					}
-
-					LoadMorphTargetsData(StreamRequest, MutableMesh.ToStrongRef(), Context->InstanceUpdateData.RealTimeMorphTargetMeshData);
-				}
-
-				// Metadata
-				const int32 FirstLOD = Context->bStreamMeshLODs ?
-					Context->FirstResidentLOD[ComponentName] :
-					Context->GetFirstRequestedLOD()[ComponentName];
-			
-				for (int32 LODIndex = FirstLOD; LODIndex < Component->LODCount; ++LODIndex)
-				{
-					const TSharedRef<FInstanceUpdateData::FLOD>& LOD = Context->InstanceUpdateData.LODs[Component->FirstLOD + LODIndex];
-
-					UE::Mutable::Private::TManagedPtr<const UE::Mutable::Private::FMesh> MutableMesh = LOD->Mesh;
-					if (!MutableMesh)
-					{
-						continue;
-					}
-
-					LoadMorphTargetsMetadata(StreamRequest, MutableMesh.ToStrongRef(), Context->InstanceUpdateData.RealTimeMorphTargetMeshData);
-				}
-			}
-		}
 		
 		for (int32 LODIndex = Context->FirstLODAvailable[ComponentName]; LODIndex < Component->LODCount; ++LODIndex)
 		{
diff --git a/Engine/Plugins/Mutable/Source/CustomizableObject/Private/MuCO/CustomizableObjectMeshUpdate.cpp b/Engine/Plugins/Mutable/Source/CustomizableObject/Private/MuCO/CustomizableObjectMeshUpdate.cpp
--- a/Engine/Plugins/Mutable/Source/CustomizableObject/Private/MuCO/CustomizableObjectMeshUpdate.cpp
+++ b/Engine/Plugins/Mutable/Source/CustomizableObject/Private/MuCO/CustomizableObjectMeshUpdate.cpp
@@ -178,19 +178,7 @@ namespace impl
 		[OperationData]()
 		{
 			FMutableStreamRequest StreamRequest(OperationData->ModelStreamableBulkData);
-	
-			// Get morphs to be streamed.
-			for (UE::Mutable::Private::TManagedPtr<const UE::Mutable::Private::FMesh>& Mesh : OperationData->Meshes)
-			{
-				if (!Mesh)
-				{
-					continue;
-				}
-	
-				LoadMorphTargetsData(StreamRequest, Mesh.ToStrongRef(), OperationData->MorphTargetMeshData);
-				LoadMorphTargetsMetadata(StreamRequest, Mesh.ToStrongRef(), OperationData->MorphTargetMeshData);
-			}
-	
+
 			// Stream data.
 			UE::Tasks::AddNested(StreamRequest.Stream());
 		},
@@ -201,13 +189,13 @@ namespace impl
 		
 		// Go to next step.
 		UE::Tasks::TTask<void> OnUpdateMeshFinishedTask = UE::Tasks::Launch(TEXT("OnUpdateMeshFinished"),
-			[StreamingTask]()
-				{
-					StreamingTask->OnUpdateMeshFinished();
-				},
-			StreamDataTask,
-			UE::Tasks::ETaskPriority::Inherit,
-			UE::Tasks::EExtendedTaskPriority::Inline);
+		[StreamingTask]()
+		{
+			StreamingTask->OnUpdateMeshFinished();
+		},
+		StreamDataTask,
+		UE::Tasks::ETaskPriority::Inherit,
+		UE::Tasks::EExtendedTaskPriority::Inline);
 		
 		UE::Tasks::AddNested(OnUpdateMeshFinishedTask);
 		
@@ -423,7 +411,7 @@ void FCustomizableObjectMeshStreamIn::ConvertMesh(const FContext& Context, bool&
 		UnrealConversionUtils::CopyMutableVertexBuffers(LODResource, MutableMesh.Get(), bNeedsCPUAccess);
 		UnrealConversionUtils::CopyMutableIndexBuffers(LODResource, MutableMesh.Get(), Mesh->SurfaceIDs[LODIndexAbsolute], bOutMarkRenderStateDirty);
 		UnrealConversionUtils::CopyMutableSkinWeightProfilesBuffers(LODResource, *const_cast<UCustomizableObjectSkeletalMesh*>(Mesh), LODIndexAbsolute, MutableMesh.Get(), Mesh->SkinWeightProfileIDs);
-		UnrealConversionUtils::MorphTargetVertexInfoBuffers(LODResource, *Mesh, *MutableMesh.Get(), OperationData->MorphTargetMeshData, LODIndexAbsolute, bGenerateCPUMorphTargetsIfNeeded);
+		UnrealConversionUtils::MorphTargetVertexInfoBuffers(LODResource, *Mesh, *MutableMesh.Get(), LODIndexAbsolute, bGenerateCPUMorphTargetsIfNeeded);
 		UnrealConversionUtils::ClothVertexBuffers(LODResource, *MutableMesh.Get());
 		
 		UnrealConversionUtils::UpdateSkeletalMeshLODRenderDataBuffersSize(LODResource);
diff --git a/Engine/Plugins/Mutable/Source/CustomizableObject/Private/MuCO/CustomizableObjectMeshUpdate.h b/Engine/Plugins/Mutable/Source/CustomizableObject/Private/MuCO/CustomizableObjectMeshUpdate.h
--- a/Engine/Plugins/Mutable/Source/CustomizableObject/Private/MuCO/CustomizableObjectMeshUpdate.h
+++ b/Engine/Plugins/Mutable/Source/CustomizableObject/Private/MuCO/CustomizableObjectMeshUpdate.h
@@ -41,9 +41,6 @@ struct FMutableUpdateMeshContext
 	// Meshes generated per LOD
 	TArray<UE::Mutable::Private::TManagedPtr<const UE::Mutable::Private::FMesh>> Meshes;
 
-	/** Serialized data per block. */
-	TMap<uint32, FMorphTargetMeshData> MorphTargetMeshData;
-
 	/** SkeletalMesh */
 	/** The resident first LOD resource index.With domain = [0, ResourceState.NumLODs[.NOT THE ASSET LOD INDEX! */ 
 	int32 CurrentFirstLODIdx = INDEX_NONE;
diff --git a/Engine/Plugins/Mutable/Source/CustomizableObject/Private/MuCO/CustomizableObjectSystem.cpp b/Engine/Plugins/Mutable/Source/CustomizableObject/Private/MuCO/CustomizableObjectSystem.cpp
--- a/Engine/Plugins/Mutable/Source/CustomizableObject/Private/MuCO/CustomizableObjectSystem.cpp
+++ b/Engine/Plugins/Mutable/Source/CustomizableObject/Private/MuCO/CustomizableObjectSystem.cpp
@@ -1937,148 +1937,53 @@ namespace impl
 
 		FInstanceUpdateData& UpdateData = OperationData->InstanceUpdateData;
 
-		const TMap<uint32, FMorphTargetMeshData>& ResourceIdToMeshDataMap =
-				UpdateData.RealTimeMorphTargetMeshData;
-
-		if (ResourceIdToMeshDataMap.IsEmpty())
-		{
-			return;
-		}
-
-		UCustomizableObject* Object = OperationData->Object.Get();
-		const UModelResources* ModelResources = Object->GetPrivate()->GetModelResources();
-		
-		const TSharedPtr<FModelStreamableBulkData>& ModelStreamableBulkData = Object->GetPrivate()->GetModelStreamableBulkData();
-
 		const int32 NumComponents = OperationData->InstanceUpdateData.Components.Num();
-		check(OperationData->InstanceUpdateData.RealTimeMorphTargets.IsEmpty());
-		OperationData->InstanceUpdateData.RealTimeMorphTargets.Reserve(NumComponents);
 		
-		for (int32 ComponentIndex = 0; ComponentIndex< NumComponents; ++ComponentIndex)
+		UpdateData.RealTimeMorphTargets.Reserve(NumComponents);
+		for (int32 ComponentIndex = 0; ComponentIndex < NumComponents; ++ComponentIndex)
 		{
 			TSharedRef<FInstanceUpdateData::FComponent>& Component = OperationData->InstanceUpdateData.Components[ComponentIndex];
 			check(Component->Id != INDEX_NONE);
 
-			const FName& ComponentName = OperationData->ComponentNames[Component->Id];
-
-			FSkeletalMeshMorphTargets& ComponentMorphTargetsData = OperationData->InstanceUpdateData.RealTimeMorphTargets.Add(Component->Id);
-			
-			ComponentMorphTargetsData.RealTimeMorphTargetNames.Empty();
-
-			TMap<uint32, FMappedMorphTargetMeshData> MorphTargetMeshData;
-			MorphTargetMeshData.Reserve(ResourceIdToMeshDataMap.Num());
+			const FName ComponentName = OperationData->ComponentNames[Component->Id];
 			
-			for (const TPair<uint32, FMorphTargetMeshData>& MorphTargetResource : ResourceIdToMeshDataMap)
-			{
-				FMappedMorphTargetMeshData& MeshData = MorphTargetMeshData.FindOrAdd(MorphTargetResource.Key);
-				MeshData.DataView = &MorphTargetResource.Value.Data;
-
-				const int32 NumMorphNames = MorphTargetResource.Value.NameResolutionMap.Num();
-				MeshData.NameResolutionMap.SetNumUninitialized(NumMorphNames);
-
-				for (int32 NameIndex = 0; NameIndex < NumMorphNames; ++NameIndex)
-				{
-					const int32 ResolvedNameIndex = ComponentMorphTargetsData.RealTimeMorphTargetNames.AddUnique(MorphTargetResource.Value.NameResolutionMap[NameIndex]);
-					MeshData.NameResolutionMap[NameIndex] = ResolvedNameIndex;
-				}
-			}
+			FSkeletalMeshMorphTargets& ComponentMorphTargets = UpdateData.RealTimeMorphTargets.Add(Component->Id);
 
 			// Allocate Morph data for used morphs.
-			TArray<TArray<FMorphTargetLODModel>>& MorphsData = ComponentMorphTargetsData.RealTimeMorphsLODData;
-			const int32 NumMorphs = ComponentMorphTargetsData.RealTimeMorphTargetNames.Num();
-			
 			const int32 NumLODsAvailable = OperationData->NumLODsAvailable[ComponentName];
-
-			MorphsData.SetNum(ComponentMorphTargetsData.RealTimeMorphTargetNames.Num());
-			for (int32 MorphIndex = 0; MorphIndex < NumMorphs; ++MorphIndex)
-			{
-				MorphsData[MorphIndex].SetNum(NumLODsAvailable);
-			}
-			
+	
 			for (int32 LODIndex = OperationData->GetFirstRequestedLOD()[ComponentName]; LODIndex < NumLODsAvailable; ++LODIndex)
 			{
 				const TSharedRef<FInstanceUpdateData::FLOD>& LOD = UpdateData.LODs[Component->FirstLOD + LODIndex];
 
-				if (!LOD->Mesh)
+				if (!(LOD->Mesh && LOD->Mesh->HasMorphs()))
 				{
 					continue;
 				}
 
 				TArray<FMorphTargetLODModel> MorphTargets;
-				ReconstructMorphTargets(*LOD->Mesh.Get(), ComponentMorphTargetsData.RealTimeMorphTargetNames, MorphTargetMeshData, MorphTargets);
+				ReconstructMorphTargetsFromMeshCompressedData(*LOD->Mesh, MorphTargets, UE::Mutable::Private::EMorphUsageFlags::RealTime);
 
-				for (int32 NameIndex = 0; NameIndex < ComponentMorphTargetsData.RealTimeMorphTargetNames.Num(); NameIndex++)
-				{
-					if (MorphTargets.IsValidIndex(NameIndex))
-					{
-						ComponentMorphTargetsData.RealTimeMorphsLODData[NameIndex][LODIndex] = MoveTemp(MorphTargets[NameIndex]);
-					}
-				}
-			}
-
-			const int32 FirstLOD = OperationData->bStreamMeshLODs ?
-				OperationData->FirstLODAvailable[ComponentName] :
-				OperationData->GetFirstRequestedLOD()[ComponentName];
+				const UE::Mutable::Private::FMeshMorph& Morph = LOD->Mesh->Morph;
 
-			// Find which Sections are being used in each LOD (Streamed and Residents).
-			for (int32 LODIndex = FirstLOD; LODIndex < NumLODsAvailable; ++LODIndex)
-			{
-				const TSharedRef<FInstanceUpdateData::FLOD>& LOD = UpdateData.LODs[Component->FirstLOD+LODIndex];
-				check(LOD->Mesh);
-				
-				for (int32 SectionIndex = 0; SectionIndex < LOD->Mesh->Surfaces.Num(); ++SectionIndex)
+				for (int32 LODMorphIndex = 0, NumMorphs = Morph.Names.Num(); LODMorphIndex < NumMorphs; ++LODMorphIndex)
 				{
-					const UE::Mutable::Private::FMeshSurface& Surface = LOD->Mesh->Surfaces[SectionIndex];
-					for (const UE::Mutable::Private::FSurfaceSubMesh& SubMesh : Surface.SubMeshes)
+					if (MorphTargets[LODMorphIndex].Vertices.Num())
 					{
-						if (SubMesh.ExternalId == 0)
-						{
-							continue;	
-						}
+						const FName MorphName = Morph.Names[LODMorphIndex];
+						int32 MorphIndex = ComponentMorphTargets.RealTimeMorphTargetNames.IndexOfByKey(MorphName);
 						
-						const uint32 MorphMetadataId = ModelResources->MeshMetadata[SubMesh.ExternalId].MorphMetadataId;
-
-						FRealTimeMorphStreamable* Result = ModelStreamableBulkData->RealTimeMorphStreamables.Find(MorphMetadataId);
-						if (!Result)
+						if (MorphIndex == INDEX_NONE)
 						{
-							continue;
+							ComponentMorphTargets.RealTimeMorphsLODData.AddDefaulted_GetRef().SetNum(NumLODsAvailable);
+							MorphIndex = ComponentMorphTargets.RealTimeMorphTargetNames.Add(MorphName);
 						}
-
-						for (const FName& MorphName : Result->NameResolutionMap)
-						{
-							const int32 MorphIndex = ComponentMorphTargetsData.RealTimeMorphTargetNames.Find(MorphName);
-							check(MorphIndex != INDEX_NONE);
-							
-							FMorphTargetLODModel& MorphTargetLODModel = ComponentMorphTargetsData.RealTimeMorphsLODData[MorphIndex][LODIndex];
-							MorphTargetLODModel.SectionIndices.Add(SectionIndex);
-						}
-					}
-				}
-			}
-			
-			// Remove empty morph targets
-			for (int32 MorphIndex = 0; MorphIndex < NumMorphs; ++MorphIndex)
-			{
-				const int32 NumLODs = MorphsData[MorphIndex].Num();
-
-				int32 LODIndex = 0;
-				for (; LODIndex < NumLODs; ++LODIndex)
-				{
-					if (!MorphsData[MorphIndex][LODIndex].Vertices.IsEmpty())
-					{
-						break;
+						
+						ComponentMorphTargets.RealTimeMorphsLODData[MorphIndex][LODIndex] = MoveTemp(MorphTargets[LODMorphIndex]);
 					}
 				}
-
-				if (LODIndex >= NumLODs)
-				{
-					MorphsData[MorphIndex].Empty();
-				}
 			}
 		}
-
-		// Free unneeded data memory.
-		UpdateData.RealTimeMorphTargetMeshData.Empty();
 	}
 	
 
@@ -4112,236 +4017,143 @@ int32 UCustomizableObjectSystemPrivate::BlockTillAllRequestsFinished(float TimeL
 }
 
 
-void LoadMorphTargetsData(FMutableStreamRequest& MutableDataStreamer, const UE::Mutable::Private::TManagedRef<const UE::Mutable::Private::FMesh>& Mesh, TMap<uint32, FMorphTargetMeshData>& StreamingResult)
+#if !WITH_EDITOR
+void ReconstructMorphTargetsFromMeshCompressedData(const UE::Mutable::Private::FMesh& Mesh, TArray<FMorphTargetCompressedLODModel>& OutMorphTargets, UE::Mutable::Private::EMorphUsageFlags UsageFilter)
 {
-	MUTABLE_CPUPROFILER_SCOPE(LoadMorphTargetsData);
-
-	if (!CVarEnableRealTimeMorphTargets.GetValueOnAnyThread())
-	{
-		return;
-	}
-
-	const TSharedPtr<FModelStreamableBulkData>& ModelStreamableBulkData = MutableDataStreamer.GetModelStreamableBulkData();
+	MUTABLE_CPUPROFILER_SCOPE(ReconstructMorphTargetsFromMeshCompressedData_ToCompressed);
 	
-	TArray<uint32> RealTimeMorphStreamableBlocksToStream;
-	
-	for (uint64 ResourceId : Mesh->GetStreamedResources())
-	{
-		const FCustomizableObjectStreameableResourceId TypedResourceId = BitCast<FCustomizableObjectStreameableResourceId>(ResourceId);
-			
-		if (TypedResourceId.Type == static_cast<uint8>(FCustomizableObjectStreameableResourceId::EType::RealTimeMorphTarget))
-		{
-			check(TypedResourceId.Id != 0 && TypedResourceId.Id <= TNumericLimits<uint32>::Max());
+	using namespace UE::MorphTargetVertexCodec;
 
-			if (ModelStreamableBulkData->RealTimeMorphStreamables.Contains(static_cast<uint32>(TypedResourceId.Id)))
-			{
-				RealTimeMorphStreamableBlocksToStream.AddUnique(TypedResourceId.Id);
-			}
-			else
-			{
-				UE_LOG(LogMutable, Error, TEXT("Invalid streamed real time morph target data block [%llu] found."), TypedResourceId.Id);
-			}
-		}
-	}
-	
-	for (const int32 BlockId : RealTimeMorphStreamableBlocksToStream)
-	{
-		MUTABLE_CPUPROFILER_SCOPE(RealTimeMorphStreamingRequest_Alloc);
-
-		const FRealTimeMorphStreamable& Streamable = ModelStreamableBulkData->RealTimeMorphStreamables[BlockId];
-		const FMutableStreamableBlock& Block = Streamable.Block; 
-		
-		FMorphTargetMeshData& ReadDestData = StreamingResult.FindOrAdd(BlockId);
+	const int32 NumMorphs = Mesh.Morph.Names.Num();
+	OutMorphTargets.SetNum(NumMorphs);
 
-		// Only request blocks once.
-		if (ReadDestData.Data.Num())
+	for (int32 MorphIndex = 0; MorphIndex < NumMorphs; ++MorphIndex)
+	{
+		if (EnumHasAnyFlags(Mesh.Morph.UsageFlagsPerMorph[MorphIndex], UsageFilter))
 		{
-			continue;
-		}
-		
-		check(Streamable.Size % sizeof(FMorphTargetVertexData) == 0);
-		uint32 NumElems = Streamable.Size / sizeof(FMorphTargetVertexData);
-
-		ReadDestData.Data.SetNumUninitialized(NumElems);
-
-		MutableDataStreamer.AddBlock(
-			Block,
-			UE::Mutable::Private::EStreamableDataType::RealTimeMorph,
-			0,
-			MakeArrayView(reinterpret_cast<uint8*>(ReadDestData.Data.GetData()), ReadDestData.Data.Num() * sizeof(FMorphTargetVertexData)));
-	}
-}
+			const float PositionPrecision = Mesh.Morph.PositionPrecision;
+			const float TangentZPrecision = Mesh.Morph.TangentZPrecision;
+			
+			OutMorphTargets[MorphIndex].PositionPrecision = PositionPrecision;
+			OutMorphTargets[MorphIndex].TangentPrecision = TangentZPrecision;
+			
+			const uint32 BatchStartOffset = Mesh.Morph.BatchStartOffsetPerMorph[MorphIndex];
+			const uint32 MorphNumBatches  = Mesh.Morph.BatchesPerMorph[MorphIndex];
 
+			OutMorphTargets[MorphIndex].PackedDeltaHeaders.SetNumUninitialized(MorphNumBatches);
 
-void LoadMorphTargetsMetadata(const FMutableStreamRequest& MutableDataStreamer, const UE::Mutable::Private::TManagedRef<const UE::Mutable::Private::FMesh>& Mesh, TMap<uint32, FMorphTargetMeshData>& StreamingResult)
-{
-	MUTABLE_CPUPROFILER_SCOPE(LoadMorphTargetsMetadata);
+			int32 MorphDeltaDataBeginInDwords = 0;
+			int32 MorphDeltaDataEndInDwords = 0;
 
-	if (!CVarEnableRealTimeMorphTargets.GetValueOnAnyThread())
-	{
-		return;
-	}
+			if (MorphNumBatches > 0)
+			{
+				// Get the Offset value for the morph traget first batch.
+				MorphDeltaDataBeginInDwords = Mesh.MorphDataBuffer[BatchStartOffset * NumBatchHeaderDwords] / sizeof(uint32);
 
-	const TSharedPtr<FModelStreamableBulkData>& ModelStreamableBulkData = MutableDataStreamer.GetModelStreamableBulkData();
+				for (int32 BatchHeaderIndex = 0; BatchHeaderIndex < MorphNumBatches; ++BatchHeaderIndex)
+				{
+					TConstArrayView<uint32> BatchHeaderData(
+							&Mesh.MorphDataBuffer[(BatchStartOffset + BatchHeaderIndex) * NumBatchHeaderDwords], 
+							NumBatchHeaderDwords);
 
-	const TArray<uint64>& StreamedResources = Mesh->GetStreamedResources();
+					FDeltaBatchHeader& Header = OutMorphTargets[MorphIndex].PackedDeltaHeaders[BatchHeaderIndex];
+					ReadHeader(Header, BatchHeaderData);
+				}	
 
-	for (uint64 ResourceId : StreamedResources)
-	{
-		FCustomizableObjectStreameableResourceId TypedResourceId = BitCast<FCustomizableObjectStreameableResourceId>(ResourceId);
-						
-		if (TypedResourceId.Type == static_cast<uint8>(FCustomizableObjectStreameableResourceId::EType::RealTimeMorphTarget))
-		{
-			check(TypedResourceId.Id != 0 && TypedResourceId.Id <= TNumericLimits<uint32>::Max());
+				MorphDeltaDataEndInDwords = 
+						OutMorphTargets[MorphIndex].PackedDeltaHeaders.Last().DataOffset / sizeof(uint32) +
+						CalculateBatchDwords(OutMorphTargets[MorphIndex].PackedDeltaHeaders.Last()); 
 
-			if (ModelStreamableBulkData->RealTimeMorphStreamables.Contains(static_cast<uint32>(TypedResourceId.Id)))
-			{
-				FMorphTargetMeshData& ReadDestData = StreamingResult.FindOrAdd(TypedResourceId.Id);
-				ReadDestData.NameResolutionMap = ModelStreamableBulkData->RealTimeMorphStreamables[TypedResourceId.Id].NameResolutionMap;
-			}
-			else
-			{
-				UE_LOG(LogMutable, Error, TEXT("Invalid streamed real time morph target data block [%llu] found."), TypedResourceId.Id);
+				// Fixup header offset.
+				for (int32 BatchHeaderIndex = 0; BatchHeaderIndex < MorphNumBatches; ++BatchHeaderIndex)
+				{
+					FDeltaBatchHeader& Header = OutMorphTargets[MorphIndex].PackedDeltaHeaders[BatchHeaderIndex];
+					Header.DataOffset -= MorphDeltaDataBeginInDwords * sizeof(uint32);
+				}	
 			}
+
+			OutMorphTargets[MorphIndex].PackedDeltaData.SetNumUninitialized(MorphDeltaDataEndInDwords - MorphDeltaDataBeginInDwords);
+			FMemory::Memcpy(
+					OutMorphTargets[MorphIndex].PackedDeltaData.GetData(),
+					Mesh.MorphDataBuffer.GetData() + MorphDeltaDataBeginInDwords,
+					(MorphDeltaDataEndInDwords - MorphDeltaDataBeginInDwords) * sizeof(uint32));
 		}
 	}
 }
+#endif
 
-
-void ReconstructMorphTargets(const UE::Mutable::Private::FMesh& Mesh, const TArray<FName>& GlobalNames, const TMap<uint32, FMappedMorphTargetMeshData>& MappedMorphTargets, TArray<FMorphTargetLODModel>& OutMorphTargets)
+void ReconstructMorphTargetsFromMeshCompressedData(const UE::Mutable::Private::FMesh& Mesh, TArray<FMorphTargetLODModel>& OutMorphTargets, UE::Mutable::Private::EMorphUsageFlags UsageFilter)
 {
-	const UE::Mutable::Private::FMeshBufferSet& MeshSet = Mesh.GetVertexBuffers();
-
-	int32 VertexMorphsInfoIndexAndCountBufferIndex, VertexMorphsInfoIndexAndCountBufferChannel;
-	MeshSet.FindChannel(UE::Mutable::Private::EMeshBufferSemantic::Other, 0, &VertexMorphsInfoIndexAndCountBufferIndex, &VertexMorphsInfoIndexAndCountBufferChannel);
-
-	int32 VertexMorphsResourceIdBufferIndex, VertexMorphsResourceIdBufferChannel;
-	MeshSet.FindChannel(UE::Mutable::Private::EMeshBufferSemantic::Other, 1, &VertexMorphsResourceIdBufferIndex, &VertexMorphsResourceIdBufferChannel);
-
-	if (VertexMorphsInfoIndexAndCountBufferIndex < 0 || VertexMorphsResourceIdBufferIndex < 0)
-	{
-		return;
-	}
-
-	OutMorphTargets.SetNum(GlobalNames.Num());
-
-	const uint32* const VertexMorphsInfoIndexAndCountBuffer = reinterpret_cast<const uint32*>(MeshSet.GetBufferData(VertexMorphsInfoIndexAndCountBufferIndex));
-	TArrayView<const uint32> VertexMorphsInfoIndexAndCountView(VertexMorphsInfoIndexAndCountBuffer, MeshSet.GetElementCount());
-
-	const uint32* const VertexMorphsResourceIdBuffer = reinterpret_cast<const uint32*>(MeshSet.GetBufferData(VertexMorphsResourceIdBufferIndex));
-	TArrayView<const uint32> VertexMorphsResourceIdView(VertexMorphsResourceIdBuffer, MeshSet.GetElementCount());
-
-	TArray<int32> SectionMorphTargetVerticesCount;
-	SectionMorphTargetVerticesCount.SetNumZeroed(GlobalNames.Num());
+	MUTABLE_CPUPROFILER_SCOPE(ReconstructMorphTargetsFromMeshCompressedData_ToUncompressed);
 	
-	const int32 SurfaceCount = Mesh.GetSurfaceCount();
-	for (int32 Section = 0; Section < SurfaceCount; ++Section)
-	{
-		// Reset SectionMorphTargets.
-		for (int32& Elem : SectionMorphTargetVerticesCount)
-		{
-			Elem = 0;
-		}
+	using namespace UE::MorphTargetVertexCodec;
 
-		int32 FirstVertex, VerticesCount, FirstIndex, IndiciesCount, UnusedBoneIndex, UnusedBoneCount;
-		Mesh.GetSurface(Section, FirstVertex, VerticesCount, FirstIndex, IndiciesCount, UnusedBoneIndex, UnusedBoneCount);
+	const int32 NumMorphs = Mesh.Morph.Names.Num();
+	OutMorphTargets.SetNum(NumMorphs);
 
-		for (int32 VertexIdx = FirstVertex; VertexIdx < FirstVertex + VerticesCount;)
+	for (int32 MorphIndex = 0; MorphIndex < NumMorphs; ++MorphIndex)
+	{
+		if (EnumHasAnyFlags(Mesh.Morph.UsageFlagsPerMorph[MorphIndex], UsageFilter))
 		{
-			// Find a span with the same VertexMorphResourceId to amortise the cost of finding 
-			// in the loaded resources map. It is expected to find large consecutive mesh sections pointing to
-			// the same loaded resource.
+			OutMorphTargets[MorphIndex].SectionIndices = Mesh.Morph.SurfacesInUsePerMorph[MorphIndex];
 
-			const int32 SpanStart = VertexIdx++;
-			const uint32 CurrentResourceId = VertexMorphsResourceIdView[SpanStart];
+			const float PositionPrecision = Mesh.Morph.PositionPrecision;
+			const float TangentZPrecision = Mesh.Morph.TangentZPrecision;
+			
+			const int32 BatchStartOffset = Mesh.Morph.BatchStartOffsetPerMorph[MorphIndex];
+			const int32 MorphNumBatches  = Mesh.Morph.BatchesPerMorph[MorphIndex];
 
-			// Vertex with no morphs are marked with 0, skip vertex if the case.
-			if (CurrentResourceId == 0)
+			int32 NumMorphVertices = 0;
+			for (int32 BatchHeaderIndex = 0; BatchHeaderIndex < MorphNumBatches; ++BatchHeaderIndex)
 			{
-				continue;
-			}
+				TConstArrayView<uint32> BatchHeaderData(
+						&Mesh.MorphDataBuffer[(BatchStartOffset + BatchHeaderIndex) * NumBatchHeaderDwords], 
+						NumBatchHeaderDwords);
 
-			for (; VertexIdx < FirstVertex + VerticesCount; ++VertexIdx)
-			{
-				const int32 VertexResourceId = VertexMorphsResourceIdView[VertexIdx];
-				// we can skip vertices with no morph without breaking the span.
-				if (VertexResourceId == 0)
-				{
-					continue;
-				}
+				FDeltaBatchHeader BatchHeader;
+				ReadHeader(BatchHeader, BatchHeaderData);
 
-				if (CurrentResourceId != VertexResourceId)
-				{
-					break;
-				}
+				NumMorphVertices += BatchHeader.NumElements;
 			}
-			const int32 SpanEnd = VertexIdx;
-
-			const FMappedMorphTargetMeshData* MorphTargetReconstructionData = MappedMorphTargets.Find(CurrentResourceId);
-
-			if (!MorphTargetReconstructionData)
+			
+			OutMorphTargets[MorphIndex].Vertices.SetNumUninitialized(NumMorphVertices);
+			OutMorphTargets[MorphIndex].NumVertices = NumMorphVertices;
+			
+			int32 MorphVertexIndex = 0;
+			for (int32 BatchHeaderIndex = 0; BatchHeaderIndex < MorphNumBatches; ++BatchHeaderIndex)
 			{
-				ensureMsgf(false, TEXT("Needed realtime morph reconstruction data was not loaded properly. Some realtime morphs may not work correctly."));
-				continue;
-			}
+				TConstArrayView<uint32> BatchHeaderData(
+						&Mesh.MorphDataBuffer[(BatchStartOffset + BatchHeaderIndex) * NumBatchHeaderDwords], 
+						NumBatchHeaderDwords);
 
-			const TArray<FMorphTargetVertexData>& SpanMorphData = *MorphTargetReconstructionData->DataView;
-			const int32 NumNamesInResolutionMap = MorphTargetReconstructionData->NameResolutionMap.Num();
+				FDeltaBatchHeader BatchHeader;
+				ReadHeader(BatchHeader, BatchHeaderData);
 
-			for (int32 SpanVertexIdx = SpanStart; SpanVertexIdx < SpanEnd; ++SpanVertexIdx)
-			{
-				const uint32 MorphOffsetAndCount = VertexMorphsInfoIndexAndCountView[SpanVertexIdx];
-				if (MorphOffsetAndCount == 0)
+				if (BatchHeader.NumElements == 0)
 				{
 					continue;
 				}
 
-				// See encoding in GenerateMutableSourceMesh.cpp.
-				constexpr uint32 Log2MaxNumVerts = 23;
-				
-				TArrayView<const FMorphTargetVertexData> MorphsVertexDataView = MakeArrayView(
-						SpanMorphData.GetData() + (MorphOffsetAndCount & ((1 << Log2MaxNumVerts) - 1)), 
-						MorphOffsetAndCount >> Log2MaxNumVerts);
-
-				for (const FMorphTargetVertexData& SourceVertex : MorphsVertexDataView)
-				{
-					if (SourceVertex.MorphNameIndex >= (uint32)NumNamesInResolutionMap)
-					{
-						ensureMsgf(false, TEXT("Invalid real-time morphs names found in instance vertices. Some morph may not work as expected."));
-						continue;
-					}
-
-					const uint32 ResolvedNameIndex = MorphTargetReconstructionData->NameResolutionMap[SourceVertex.MorphNameIndex];
-
-					FMorphTargetLODModel& DestMorphLODModel = OutMorphTargets[ResolvedNameIndex];
+				TArray<FQuantizedDelta, TInlineAllocator<UE::MorphTargetVertexCodec::BatchSize>> QuantizedDeltas;
+				QuantizedDeltas.SetNumUninitialized(BatchHeader.NumElements);
 
-					DestMorphLODModel.Vertices.Emplace(
-							FMorphTargetDelta 
-							{ 
-								SourceVertex.PositionDelta, 
-								SourceVertex.TangentZDelta, 
-								static_cast<uint32>(SpanVertexIdx) 
-							});
+				TConstArrayView<uint32> Data(
+						&Mesh.MorphDataBuffer[BatchHeader.DataOffset / sizeof(uint32)], 
+						CalculateBatchDwords(BatchHeader));
 
-					++SectionMorphTargetVerticesCount[ResolvedNameIndex];
+				ReadQuantizedDeltas(QuantizedDeltas, BatchHeader, Data);
+				for (FQuantizedDelta& QuantizedDelta : QuantizedDeltas)
+				{
+					FMorphTargetDelta& Delta = OutMorphTargets[MorphIndex].Vertices[MorphVertexIndex++];
+					DequantizeDelta(Delta, BatchHeader.bTangents, QuantizedDelta, Mesh.Morph.PositionPrecision, Mesh.Morph.TangentZPrecision);
 				}
 			}
-		}
-
-		const int32 SectionMorphTargetsNum = SectionMorphTargetVerticesCount.Num();
-		for (int32 MorphIdx = 0; MorphIdx < SectionMorphTargetsNum; ++MorphIdx)
-		{
-			if (SectionMorphTargetVerticesCount[MorphIdx] > 0)
-			{
-				FMorphTargetLODModel& MorphTargetLodModel = OutMorphTargets[MorphIdx];
-
-				MorphTargetLodModel.NumVertices += SectionMorphTargetVerticesCount[MorphIdx];
-			}
+		
+			check(MorphVertexIndex == NumMorphVertices);
 		}
 	}
 }
 
-
 bool IsStreamingEnabled(const UCustomizableObject& Object, const int32 State)
 {
 	const UModelResources& ModelResources = Object.GetPrivate()->GetModelResourcesChecked();
diff --git a/Engine/Plugins/Mutable/Source/CustomizableObject/Private/MuCO/UnrealConversionUtils.cpp b/Engine/Plugins/Mutable/Source/CustomizableObject/Private/MuCO/UnrealConversionUtils.cpp
--- a/Engine/Plugins/Mutable/Source/CustomizableObject/Private/MuCO/UnrealConversionUtils.cpp
+++ b/Engine/Plugins/Mutable/Source/CustomizableObject/Private/MuCO/UnrealConversionUtils.cpp
@@ -900,128 +900,74 @@ namespace UnrealConversionUtils
 		}
 	}
 
-	void MorphTargetVertexInfoBuffers(FSkeletalMeshLODRenderData& LODResource, const USkeletalMesh& Owner, const UE::Mutable::Private::FMesh& MutableMesh, const TMap<uint32, FMorphTargetMeshData>& MorphTargetMeshData, int32 LODIndex, bool bGenerateCPUMorphTargetsIfNeeded)
+	void MorphTargetVertexInfoBuffers(FSkeletalMeshLODRenderData& LODResource, const USkeletalMesh& Owner, const UE::Mutable::Private::FMesh& MutableMesh, int32 LODIndex, bool bGenerateCPUMorphTargetsIfNeeded)
 	{
 		if (Owner.GetMorphTargets().IsEmpty())
 		{
 			return;
 		}
-		
-		// Get the global names.
-		const TMap<FName, int32>& IndexMap = Owner.GetMorphTargetIndexMap();
 	
-		int32 MaxIndex = 0;
-		for (const TPair<FName, int32>& Pair : IndexMap)
-		{
-			MaxIndex = FMath::Max(MaxIndex, Pair.Value);
-		}
-
-		TArray<FName> GlobalNames;
-		GlobalNames.SetNum(MaxIndex + 1);
-
-		for (const TPair<FName, int32>& Pair : IndexMap)
+		if (!MutableMesh.HasMorphs())
 		{
-			GlobalNames[Pair.Value] = Pair.Key;
-		}
-
-		// Map the local names to the global names.
-		TMap<uint32, FMappedMorphTargetMeshData> GlobalMorphTargets; // Key is the block id. Value data needed to reconstruct the Morp Targets.
-		for (const TPair<uint32, FMorphTargetMeshData>& Pair : MorphTargetMeshData)
-		{
-			FMappedMorphTargetMeshData& MappedMorphsData = GlobalMorphTargets.FindOrAdd(Pair.Key);
-
-			// Remapping to global names.
-			MappedMorphsData.NameResolutionMap.SetNum(Pair.Value.NameResolutionMap.Num());
-			for (int32 NameIndex = 0; NameIndex < Pair.Value.NameResolutionMap.Num(); ++NameIndex)
-			{
-				MappedMorphsData.NameResolutionMap[NameIndex] = GlobalNames.Find(Pair.Value.NameResolutionMap[NameIndex]);
-			}
-
-			// Pointer to the serialized data
-			MappedMorphsData.DataView = &Pair.Value.Data;
+			return;
 		}
 
+		// Get the global names map.
+		const TMap<FName, int32>& IndexMap = Owner.GetMorphTargetIndexMap();
+	
 		// Reconstruct the final Morph Targets using the global names.
-		TArray<FMorphTargetLODModel> MorphTargetLODs;
-		ReconstructMorphTargets(MutableMesh, GlobalNames, GlobalMorphTargets, MorphTargetLODs); // FMorphTargetLODModel::SectionIndices not initialized here
-
-		TArray<const FMorphTargetLODModel*> MorphsToInit;
-		MorphsToInit.Reserve(MorphTargetLODs.Num());
-		Algo::Transform(MorphTargetLODs, MorphsToInit, [](const FMorphTargetLODModel& Morph) { return &Morph; });
-
-		const float ErrorTolerance = Owner.GetLODInfo(LODIndex)->MorphTargetPositionErrorTolerance;
-		const int32 NumVertices = LODResource.StaticVertexBuffers.StaticMeshVertexBuffer.GetNumVertices();
 
 		const bool bGenerateCPUMorphTargets = 
 				bGenerateCPUMorphTargetsIfNeeded && !FMorphTargetVertexInfoBuffers::IsPlatformShaderSupported(GMaxRHIShaderPlatform);
 
-#if !WITH_EDITOR
+#if WITH_EDITOR
+		TArray<FMorphTargetLODModel> MorphTargetLODs;
+		ReconstructMorphTargetsFromMeshCompressedData(MutableMesh, MorphTargetLODs, UE::Mutable::Private::EMorphUsageFlags::RealTime);
+#else
 		TArray<FMorphTargetCompressedLODModel> CompressedMorphs;
-		CompressedMorphs.SetNum(MorphTargetLODs.Num());
+		ReconstructMorphTargetsFromMeshCompressedData(MutableMesh, CompressedMorphs, UE::Mutable::Private::EMorphUsageFlags::RealTime);
+#endif
 		
-		TBitArray<> SectionNeedsTangents;
-		SectionNeedsTangents.Init(false, LODResource.RenderSections.Num());
-		for (int32 SectionIndex = 0; SectionIndex < SectionNeedsTangents.Num(); ++SectionIndex)
-		{
-			SectionNeedsTangents[SectionIndex] = FMorphTargetVertexInfoBuffers::NeedsTangents(
-					LODResource.RenderSections[SectionIndex].bRecomputeTangent);
-		}
-
-		const bool bAllRenderSectionsNeedTangents = SectionNeedsTangents.CountSetBits() == SectionNeedsTangents.Num();
-
-		TBitArray<> VertexNeedsTangentsStorage;
-		TBitArray<>* VertexNeedsTangents = nullptr;
+#if WITH_EDITOR
+		TArray<FMorphTargetLODModel*> MorphsToInit;
+		MorphsToInit.Init(nullptr, Owner.GetMorphTargets().Num());
 
-		// When generating CPU Morph targets we need to generate tangents.
-		if (!(bGenerateCPUMorphTargets || bAllRenderSectionsNeedTangents))
+		for (int32 MorphIndex = MutableMesh.Morph.Names.Num() - 1; MorphIndex >= 0; --MorphIndex)
 		{
-			VertexNeedsTangentsStorage.Init(false, NumVertices);
-			for (int32 SectionIndex = 0; SectionIndex < SectionNeedsTangents.Num(); ++SectionIndex)
+			const int32* FoundIndex = IndexMap.Find(MutableMesh.Morph.Names[MorphIndex]);
+
+			if (FoundIndex && MorphTargetLODs[MorphIndex].NumVertices)
 			{
-				const FSkelMeshRenderSection& RenderSection = LODResource.RenderSections[SectionIndex];
-				if (SectionNeedsTangents[SectionIndex])
-				{
-					VertexNeedsTangentsStorage.SetRange(RenderSection.BaseVertexIndex, RenderSection.NumVertices, true);
-				}
+				MorphsToInit[*FoundIndex] = &MorphTargetLODs[MorphIndex];
 			}
-
-			VertexNeedsTangents = &VertexNeedsTangentsStorage;
 		}
 
-		for (int32 Index = 0, NumMorphs = MorphTargetLODs.Num(); Index < NumMorphs; ++Index)
-		{
-			using namespace UE::MorphTargetVertexCodec;
-
-			CompressedMorphs[Index].PositionPrecision = ComputePositionPrecision(ErrorTolerance);
-			CompressedMorphs[Index].TangentPrecision = ComputeTangentPrecision();
-
-			Encode(MorphTargetLODs[Index].Vertices, 
-					VertexNeedsTangents,
-					CompressedMorphs[Index].PositionPrecision,
-					CompressedMorphs[Index].TangentPrecision,
-					CompressedMorphs[Index].PackedDeltaHeaders,
-					CompressedMorphs[Index].PackedDeltaData);
+		const float ErrorTolerance = Owner.GetLODInfo(LODIndex)->MorphTargetPositionErrorTolerance;
+		const int32 NumVertices = LODResource.StaticVertexBuffers.StaticMeshVertexBuffer.GetNumVertices();
 
-			// The uncompressed data for this Morph is no longer needed, free it.
-			MorphTargetLODs[Index].Vertices.Empty();
-		}
-#endif
-		
-#if WITH_EDITOR
 		const TBitArray UsesBuiltinMorphTargetCompression(true, MorphTargetLODs.Num());
 		LODResource.MorphTargetVertexInfoBuffers.InitMorphResourcesStreaming(GMaxRHIShaderPlatform, LODResource.RenderSections, MorphsToInit, NumVertices, ErrorTolerance);
 #else
 		TArray<FVector4f> MorphTargetMinimumValues;
 		TArray<FVector4f> MorphTargetMaximumValues;
 
-		MorphTargetMinimumValues.SetNum(MorphTargetLODs.Num());
-		MorphTargetMaximumValues.SetNum(MorphTargetLODs.Num());
+		MorphTargetMinimumValues.SetNum(Owner.GetMorphTargets().Num());
+		MorphTargetMaximumValues.SetNum(Owner.GetMorphTargets().Num());
 	
-		ComputeMorphTargetsMinAndMaxValues(MorphsToInit, MorphTargetMinimumValues, MorphTargetMaximumValues);
+		TArray<FMorphTargetCompressedLODModel*> CompressedMorphsToInit;
+		CompressedMorphsToInit.Init(nullptr, Owner.GetMorphTargets().Num());
 
-		TArray<const FMorphTargetCompressedLODModel*> CompressedMorphsToInit;
-		CompressedMorphsToInit.Reserve(CompressedMorphs.Num());
-		Algo::Transform(CompressedMorphs, CompressedMorphsToInit, [](const FMorphTargetCompressedLODModel& Morph) { return &Morph; });
+		for (int32 MorphIndex = MutableMesh.Morph.Names.Num() - 1; MorphIndex >= 0; --MorphIndex)
+		{
+			const int32* FoundIndex = IndexMap.Find(MutableMesh.Morph.Names[MorphIndex]);
+
+			if (FoundIndex)
+			{
+				CompressedMorphsToInit[*FoundIndex] = &CompressedMorphs[MorphIndex];
+				MorphTargetMinimumValues[*FoundIndex] = MutableMesh.Morph.MinimumValuePerMorph[MorphIndex]; 
+				MorphTargetMaximumValues[*FoundIndex] = MutableMesh.Morph.MaximumValuePerMorph[MorphIndex];
+			}
+		}
 
 		LODResource.MorphTargetVertexInfoBuffers.InitMorphResourcesStreaming(GMaxRHIShaderPlatform, CompressedMorphsToInit, MorphTargetMinimumValues, MorphTargetMaximumValues);
 #endif
@@ -1030,23 +976,414 @@ namespace UnrealConversionUtils
 #if !WITH_EDITOR
 			for (int32 Index = 0, NumMorphs = Owner.GetMorphTargets().Num(); Index < NumMorphs; ++Index)
 			{
-				UMorphTarget* MorphLOD = Owner.GetMorphTargets()[Index];
-				check(MorphLOD->GetCompressedLODModels().IsValidIndex(LODIndex));
-				MorphLOD->GetCompressedLODModels()[LODIndex] = MoveTemp(CompressedMorphs[Index]);
+				if (CompressedMorphsToInit[Index])
+				{
+					UMorphTarget* MorphLOD = Owner.GetMorphTargets()[Index];
+					check(MorphLOD->GetCompressedLODModels().IsValidIndex(LODIndex));
+					MorphLOD->GetCompressedLODModels()[LODIndex] = MoveTemp(*CompressedMorphsToInit[Index]);
+				}
 			}
 #else
 			for (int32 Index = 0, NumMorphs = Owner.GetMorphTargets().Num(); Index < NumMorphs; ++Index)
 			{
-				UMorphTarget* MorphLOD = Owner.GetMorphTargets()[Index];
+				if (MorphsToInit[Index])
+				{
+					UMorphTarget* MorphLOD = Owner.GetMorphTargets()[Index];
 
-				check(MorphLOD->GetMorphLODModels().IsValidIndex(LODIndex));
-				MorphLOD->GetMorphLODModels()[LODIndex].NumVertices = MorphTargetLODs[Index].Vertices.Num();
-				MorphLOD->GetMorphLODModels()[LODIndex].Vertices    = MoveTemp(MorphTargetLODs[Index].Vertices);
+					check(MorphLOD->GetMorphLODModels().IsValidIndex(LODIndex));
+					MorphLOD->GetMorphLODModels()[LODIndex].NumVertices = MorphsToInit[Index]->Vertices.Num();
+					MorphLOD->GetMorphLODModels()[LODIndex].Vertices    = MoveTemp(MorphsToInit[Index]->Vertices);
+				}
 			}
 #endif
 		}
 	}
 
+	void SetMorphData(
+			TNotNull<UE::Mutable::Private::FMesh*> Result, 
+			TNotNull<const USkeletalMesh*> SkeletalMesh, 
+			int32 LODIndex, int32 SectionIndex, 
+			uint32 SourceSectionBaseVertexIndex, uint32 SourceSectionNumVertices, uint32 SourceMeshNumVertices, 
+			FMorphTargetVertexInfoBuffers* OptionalMorphTargetVertexBuffer)
+	{
+		using namespace UE::MorphTargetVertexCodec;
+		const FSkeletalMeshLODInfo* SkeletalMeshLODInfo = SkeletalMesh->GetLODInfo(LODIndex);
+
+		const TArray<TObjectPtr<UMorphTarget>>& MorphTargets = SkeletalMesh->GetMorphTargets();
+		for (UMorphTarget* MorphTarget : MorphTargets)
+		{
+			if (!MorphTarget)
+			{
+				continue;
+			}
+			
+			TArray<FMorphTargetLODModel>& MorphLODModels = MorphTarget->GetMorphLODModels();
+
+			if (MorphLODModels.IsValidIndex(LODIndex) && MorphLODModels[LODIndex].SectionIndices.Contains(SectionIndex))
+			{
+				Result->Morph.Names.Add(MorphTarget->GetFName());
+				Result->Morph.UsageFlagsPerMorph.Add(UE::Mutable::Private::EMorphUsageFlags::None);
+			
+				// Init with one surface. 
+				Result->Morph.SurfacesInUsePerMorph.Emplace_GetRef().Init(0, 1);
+			}
+		}
+
+		TArray<uint32> RuntimeCompressedMorphData;
+		if (!OptionalMorphTargetVertexBuffer)
+		{
+			MUTABLE_CPUPROFILER_SCOPE(CompressMorphs)
+
+			// In Editor morph data is not compressed and is located in FMorphTargetLODModel.	
+			TArray<const FMorphTargetLODModel*> MorphTargetLODs;
+			MorphTargetLODs.Reserve(MorphTargets.Num());
+
+			for (UMorphTarget* MorphTarget : MorphTargets)
+			{
+				if (!MorphTarget)
+				{
+					continue;
+				}
+
+				TArray<FMorphTargetLODModel>& MorphLODModels = MorphTarget->GetMorphLODModels();
+				if (MorphLODModels.IsValidIndex(LODIndex) && MorphLODModels[LODIndex].SectionIndices.Contains(SectionIndex))
+				{
+					MorphTargetLODs.Add(&MorphLODModels[LODIndex]);
+				}
+			}
+
+			check(MorphTargetLODs.Num() == Result->Morph.Names.Num());
+			
+			Result->Morph.BatchStartOffsetPerMorph.Empty(MorphTargetLODs.Num());
+			Result->Morph.BatchesPerMorph.Empty(MorphTargetLODs.Num());
+			Result->Morph.MaximumValuePerMorph.Empty(MorphTargetLODs.Num());
+			Result->Morph.MinimumValuePerMorph.Empty(MorphTargetLODs.Num());
+
+			// NOTE: Here we are compressing the morphs for all sections. This could be avoided pre-filtering
+			// the morph vertices in the section. This is done this way so the post-compress filtering is applied
+			// in editor to mimic what will happen with cooked data.	
+			
+			const float PositionPrecision = ComputePositionPrecision(SkeletalMeshLODInfo->MorphTargetPositionErrorTolerance);
+			const float TangentZPrecision = ComputeTangentPrecision();
+
+			TArray<FDeltaBatchHeader> BatchHeaders;
+			TArray<uint32> BitstreamData;
+		
+			for (const FMorphTargetLODModel* MorphModel : MorphTargetLODs)
+			{
+				if (!MorphModel)
+				{
+					continue;
+				}
+				
+				uint32 BatchStartOffset = BatchHeaders.Num();
+				
+				FVector4f MinimumValues;
+				FVector4f MaximumValues;
+
+				ComputeMorphTargetsMinAndMaxValues(
+						TConstArrayView<const FMorphTargetLODModel*>(&MorphModel, 1),
+						TArrayView<FVector4f>(&MinimumValues, 1),
+						TArrayView<FVector4f>(&MaximumValues, 1));
+
+				// Encode the actual morph vertex info into the quantized bitstream. NeedsTangents set to null 
+				// as we always need tangents.
+				Encode(MorphModel->Vertices, nullptr, PositionPrecision, TangentZPrecision, BatchHeaders, BitstreamData);
+			
+				const uint32 MorphNumBatches = BatchHeaders.Num() - BatchStartOffset;
+				Result->Morph.BatchStartOffsetPerMorph.Add(BatchStartOffset);
+				Result->Morph.BatchesPerMorph.Add(MorphNumBatches);
+				Result->Morph.MaximumValuePerMorph.Add(MaximumValues);
+				Result->Morph.MinimumValuePerMorph.Add(MinimumValues);
+			}
+		
+			Result->Morph.PositionPrecision = PositionPrecision;
+			Result->Morph.TangentZPrecision = TangentZPrecision;
+			Result->Morph.NumTotalBatches = BatchHeaders.Num();
+			
+			// Fix batch headers and write them packed.
+			for (FDeltaBatchHeader& BatchHeader : BatchHeaders)
+			{
+				BatchHeader.DataOffset += BatchHeaders.Num() * NumBatchHeaderDwords*sizeof(uint32);
+
+				TStaticArray<uint32, NumBatchHeaderDwords> HeaderData;
+				WriteHeader(BatchHeader, HeaderData);
+				RuntimeCompressedMorphData.Append(HeaderData);
+			}
+
+			RuntimeCompressedMorphData.Append(BitstreamData);
+		}
+
+		const TConstArrayView<uint32> SourceMorphDataView = !OptionalMorphTargetVertexBuffer 
+				? TConstArrayView<uint32>(RuntimeCompressedMorphData)
+				: TConstArrayView<uint32>(OptionalMorphTargetVertexBuffer->GetData(), OptionalMorphTargetVertexBuffer->GetMorphDataSizeInBytes()/sizeof(uint32));
+
+		if (SourceMorphDataView.GetData())
+		{	
+			MUTABLE_CPUPROFILER_SCOPE(ExtractSectionMorphs)
+			UE::Mutable::Private::FMeshMorph& Morph = Result->Morph;
+		
+			if (OptionalMorphTargetVertexBuffer)
+			{
+				const int32 NumMorphs = OptionalMorphTargetVertexBuffer->GetNumMorphs();
+				Morph.MaximumValuePerMorph.SetNumUninitialized(NumMorphs);
+				Morph.MinimumValuePerMorph.SetNumUninitialized(NumMorphs);
+				Morph.BatchStartOffsetPerMorph.SetNumUninitialized(NumMorphs);
+				Morph.BatchesPerMorph.SetNumUninitialized(NumMorphs);
+
+				for (int32 MorphIndex = 0; MorphIndex < NumMorphs; ++MorphIndex)
+				{
+					Morph.MaximumValuePerMorph[MorphIndex] = OptionalMorphTargetVertexBuffer->GetMaximumMorphScale(MorphIndex);
+					Morph.MinimumValuePerMorph[MorphIndex] = OptionalMorphTargetVertexBuffer->GetMinimumMorphScale(MorphIndex);
+					Morph.BatchStartOffsetPerMorph[MorphIndex] = OptionalMorphTargetVertexBuffer->GetBatchStartOffset(MorphIndex);
+					Morph.BatchesPerMorph[MorphIndex] = OptionalMorphTargetVertexBuffer->GetNumBatches(MorphIndex);
+					Morph.NumTotalBatches = OptionalMorphTargetVertexBuffer->GetNumBatches();
+					Morph.PositionPrecision = OptionalMorphTargetVertexBuffer->GetPositionPrecision();
+					Morph.TangentZPrecision = OptionalMorphTargetVertexBuffer->GetTangentZPrecision();
+				}
+			}
+
+			// Strip all vertices not in this section.
+			const int32 SectionVertexRangeBegin = SourceSectionBaseVertexIndex;
+			const int32 SectionVertexRangeEnd   = SourceSectionBaseVertexIndex + SourceSectionNumVertices;
+
+			uint32 TotalNumBatches = 0;
+			uint32 CumulativeDataOffsetInDwords = 0;
+
+			int32 NumMorphs = Morph.BatchStartOffsetPerMorph.Num();
+			for (int32 MorphIndex = 0; MorphIndex < NumMorphs; ++MorphIndex)
+			{	
+				const int32 SourceMorphNumBatches = Morph.BatchesPerMorph[MorphIndex];
+				const int32 SourceBatchOffset = Morph.BatchStartOffsetPerMorph[MorphIndex]; 
+
+				if (SourceMorphNumBatches == 0)
+				{
+					continue;
+				}
+
+				int32 BatchIndex = 0;
+				for ( ; BatchIndex < SourceMorphNumBatches; ++BatchIndex)
+				{
+					FDeltaBatchHeader BatchHeader;
+					TConstArrayView<uint32> BatchHeaderDataView(
+							SourceMorphDataView.GetData() + (SourceBatchOffset + BatchIndex) * NumBatchHeaderDwords, 
+							NumBatchHeaderDwords); 
+
+					ReadHeader(BatchHeader, BatchHeaderDataView);
+
+					if (BatchHeader.IndexMin >= (uint32)SectionVertexRangeBegin)
+					{
+						break;
+					}
+				}
+
+				// The previous batch could have deltas in the section range.
+				int32 FirstSectionBatchIndex = FMath::Max(0, BatchIndex - 1);
+
+				for ( ; BatchIndex < SourceMorphNumBatches; ++BatchIndex)
+				{
+					FDeltaBatchHeader BatchHeader;
+					TConstArrayView<uint32> BatchHeaderDataView(
+							SourceMorphDataView.GetData() + (SourceBatchOffset + BatchIndex) * NumBatchHeaderDwords, 
+							NumBatchHeaderDwords); 
+
+					ReadHeader(BatchHeader, BatchHeaderDataView);
+
+					if (BatchHeader.IndexMin >= (uint32)SectionVertexRangeEnd)
+					{
+						break;
+					}
+				}
+				
+				int32 LastSectionBatchIndex = FMath::Max(0, BatchIndex - 1);
+
+				uint32 NumSectionMorphBatches = LastSectionBatchIndex + 1 - FirstSectionBatchIndex;
+				
+				TConstArrayView<uint32> SectionMorphHeadersView(
+							SourceMorphDataView.GetData() + (SourceBatchOffset + FirstSectionBatchIndex) * NumBatchHeaderDwords, 
+							NumSectionMorphBatches * NumBatchHeaderDwords);
+
+				Result->MorphDataBuffer.Append(SectionMorphHeadersView);
+
+				Result->Morph.BatchStartOffsetPerMorph[MorphIndex] = CumulativeDataOffsetInDwords / NumBatchHeaderDwords;
+				Result->Morph.BatchesPerMorph[MorphIndex] = NumSectionMorphBatches;
+				TotalNumBatches += NumSectionMorphBatches;
+				CumulativeDataOffsetInDwords += SectionMorphHeadersView.Num();
+			}
+			
+			Result->Morph.NumTotalBatches = TotalNumBatches;
+
+			auto TrimBatchDeltasNotInRange = [](
+					int32 VertexRangeBegin, int32 VertexRangeEnd,
+					FDeltaBatchHeader& InOutBatchHeader,
+					TConstArrayView<uint32> DeltasData,
+					TArrayView<uint32> OutBatchDeltasData)
+			{
+				TConstArrayView<uint32> BatchDeltasData(
+						DeltasData.GetData() + (InOutBatchHeader.DataOffset / sizeof(uint32)), 
+						CalculateBatchDwords(InOutBatchHeader));
+			
+				TArray<FQuantizedDelta, TInlineAllocator<UE::MorphTargetVertexCodec::BatchSize>> QuantizedDeltas;
+				QuantizedDeltas.SetNumUninitialized(UE::MorphTargetVertexCodec::BatchSize);
+
+				ReadQuantizedDeltas(QuantizedDeltas, InOutBatchHeader, BatchDeltasData);
+
+				int32 NewIndexMin = InOutBatchHeader.IndexMin;
+				uint32 DeltaIndex = 0;
+				for ( ; DeltaIndex < InOutBatchHeader.NumElements; ++DeltaIndex)
+				{
+					int32 DeltaVertexIdx = QuantizedDeltas[DeltaIndex].Index;
+
+					if (DeltaVertexIdx >= VertexRangeBegin)
+					{
+						NewIndexMin = DeltaVertexIdx;
+						break;
+					}
+				}
+
+				int32 DeltaIndexBegin = DeltaIndex;
+
+				for ( ; DeltaIndex < InOutBatchHeader.NumElements; ++DeltaIndex)
+				{
+					int32 DeltaVertexIdx = QuantizedDeltas[DeltaIndex].Index;
+					
+					if (DeltaVertexIdx >= VertexRangeEnd)
+					{
+						break;
+					}
+				}
+
+				InOutBatchHeader.IndexMin    = NewIndexMin;
+				InOutBatchHeader.NumElements = DeltaIndex - DeltaIndexBegin;
+
+				check(InOutBatchHeader.NumElements * sizeof(FQuantizedDelta) <= QuantizedDeltas.NumBytes());
+				FMemory::Memmove(
+						QuantizedDeltas.GetData(), 
+						QuantizedDeltas.GetData() + DeltaIndexBegin,
+						InOutBatchHeader.NumElements * sizeof(FQuantizedDelta));
+				
+				WriteQuantizedDeltas(QuantizedDeltas, InOutBatchHeader, OutBatchDeltasData);
+			};
+
+			for (int32 MorphIndex = 0; MorphIndex < NumMorphs; ++MorphIndex)
+			{
+				const int32 MorphNumBatches = Result->Morph.BatchesPerMorph[MorphIndex];
+				const int32 BatchOffset = Result->Morph.BatchStartOffsetPerMorph[MorphIndex]; 
+		
+				if (MorphNumBatches == 0)
+				{
+					continue;
+				}
+
+				TArrayView<uint32> MorphBatchHeadersView(
+						Result->MorphDataBuffer.GetData() + BatchOffset * NumBatchHeaderDwords,
+						MorphNumBatches * NumBatchHeaderDwords);
+
+				// Batches still point to the source data.
+				FDeltaBatchHeader FirstBatchHeader;
+				TArrayView<uint32> FirstBatchHeaderDataView(
+						MorphBatchHeadersView.GetData(), NumBatchHeaderDwords);
+				ReadHeader(FirstBatchHeader, FirstBatchHeaderDataView);
+
+				FDeltaBatchHeader LastBatchHeader;
+				TArrayView<uint32> LastBatchHeaderDataView(
+						MorphBatchHeadersView.GetData() + (MorphNumBatches - 1)*NumBatchHeaderDwords, NumBatchHeaderDwords);
+				ReadHeader(LastBatchHeader, LastBatchHeaderDataView);
+
+				FDeltaBatchHeader FirstFullBatchHeader;
+				FDeltaBatchHeader LastFullBatchHeader;
+				
+				TConstArrayView<uint32> SourceFullBatchesDataView;
+				if (MorphNumBatches > 2)
+				{
+					TArrayView<uint32> FirstFullBatchHeaderDataView(
+							MorphBatchHeadersView.GetData() + 1*NumBatchHeaderDwords, NumBatchHeaderDwords);
+					ReadHeader(FirstFullBatchHeader, FirstFullBatchHeaderDataView);
+
+					TArrayView<uint32> LastFullBatchHeaderDataView(
+							MorphBatchHeadersView.GetData() + (MorphNumBatches - 2)*NumBatchHeaderDwords, NumBatchHeaderDwords);
+					ReadHeader(LastFullBatchHeader, LastFullBatchHeaderDataView);
+
+					SourceFullBatchesDataView = TConstArrayView<uint32>(
+						SourceMorphDataView.GetData() + (FirstFullBatchHeader.DataOffset / sizeof(uint32)),
+						((LastFullBatchHeader.DataOffset - FirstFullBatchHeader.DataOffset) / sizeof(uint32)) + CalculateBatchDwords(LastFullBatchHeader));
+				}
+
+				TArray<uint32, TInlineAllocator<BatchSize*sizeof(FQuantizedDelta)/sizeof(uint32)>> PartialBatchDeltaStorage;
+				uint32 PartialBatchDeltaStorageDwords = FMath::Max(CalculateBatchDwords(FirstBatchHeader), CalculateBatchDwords(LastBatchHeader));
+				PartialBatchDeltaStorage.SetNumUninitialized(PartialBatchDeltaStorageDwords);
+
+				// We are potentially reallocating, views will be invalidated after the reserve or any append.
+				// TODO: We can know an accurate upper-bound of the final size beforehand only looking at source 
+				// headers, pre-compute the size and only make one allocation.
+				Result->MorphDataBuffer.Reserve(
+						MorphBatchHeadersView.Num() + SourceFullBatchesDataView.Num() + 
+						CalculateBatchDwords(FirstBatchHeader) + CalculateBatchDwords(LastBatchHeader));
+
+				TrimBatchDeltasNotInRange(
+						SectionVertexRangeBegin, SectionVertexRangeEnd, FirstBatchHeader, SourceMorphDataView, PartialBatchDeltaStorage);
+				
+				FirstBatchHeaderDataView = TArrayView<uint32>(
+						Result->MorphDataBuffer.GetData() + BatchOffset*NumBatchHeaderDwords, NumBatchHeaderDwords);
+				WriteHeader(FirstBatchHeader, FirstBatchHeaderDataView);			
+
+				// Batch size may have changed, CalculateBatchDwords value cannot be cached from previous query. 
+				Result->MorphDataBuffer.Append(
+						TArrayView<uint32>(PartialBatchDeltaStorage.GetData(), CalculateBatchDwords(FirstBatchHeader)));
+			
+				// SourceFullBatchesDataView is expected to be empty in some cases.
+				Result->MorphDataBuffer.Append(SourceFullBatchesDataView);
+
+				if (MorphNumBatches > 1)
+				{
+					PartialBatchDeltaStorage.SetNumUninitialized(CalculateBatchDwords(LastBatchHeader), EAllowShrinking::No);
+
+					TrimBatchDeltasNotInRange(
+							SectionVertexRangeBegin, SectionVertexRangeEnd, LastBatchHeader, SourceMorphDataView, PartialBatchDeltaStorage);
+					
+					LastBatchHeaderDataView = TArrayView<uint32>(
+							Result->MorphDataBuffer.GetData() + (BatchOffset + MorphNumBatches - 1)*NumBatchHeaderDwords, NumBatchHeaderDwords);
+					WriteHeader(LastBatchHeader, LastBatchHeaderDataView);					
+
+					// Batch size may have changed, CalculateBatchDwords value cannot be cached from previous query. 
+					Result->MorphDataBuffer.Append(
+							TArrayView<uint32>(PartialBatchDeltaStorage.GetData(), CalculateBatchDwords(LastBatchHeader)));
+				}
+
+				int32 NumMorphBatches = Result->Morph.BatchesPerMorph[MorphIndex]; 
+
+				// Batch header fixup.
+				TArrayView<uint32> BatchHeadersData(
+						Result->MorphDataBuffer.GetData() + Result->Morph.BatchStartOffsetPerMorph[MorphIndex] * NumBatchHeaderDwords, 
+						NumMorphBatches * NumBatchHeaderDwords);
+				
+				for (int32 BatchIndex = 0; BatchIndex < NumMorphBatches; ++BatchIndex)
+				{
+					TArrayView<uint32> BatchHeaderData(
+							BatchHeadersData.GetData() + BatchIndex*NumBatchHeaderDwords, NumBatchHeaderDwords);
+
+					FDeltaBatchHeader BatchHeader;
+					ReadHeader(BatchHeader, BatchHeaderData);
+
+					if (BatchHeader.NumElements != 0)
+					{
+						check(BatchHeader.IndexMin >= (uint32)SectionVertexRangeBegin);
+						BatchHeader.IndexMin = BatchHeader.IndexMin - SectionVertexRangeBegin;
+						BatchHeader.DataOffset = CumulativeDataOffsetInDwords * sizeof(uint32);
+						CumulativeDataOffsetInDwords += CalculateBatchDwords(BatchHeader);
+					}
+					else
+					{
+						BatchHeader.IndexMin   = 0;
+						BatchHeader.DataOffset = CumulativeDataOffsetInDwords * sizeof(uint32);
+					}
+
+					WriteHeader(BatchHeader, BatchHeaderData);
+				}
+			}
+		}
+	}
+
 	
 	UE::Tasks::FTask ConvertSkeletalMeshFromRuntimeData(USkeletalMesh* SkeletalMesh, int32 LODIndex, int32 SectionIndex, UE::Mutable::Private::TManagedPtr<UE::Mutable::Private::FMesh> Result)
 	{
@@ -1477,376 +1814,11 @@ namespace UnrealConversionUtils
 					}
 				}
 
-				// Morphs
-
-				using namespace UE::MorphTargetVertexCodec;
-				const FSkeletalMeshLODInfo* SkeletalMeshLODInfo = SkeletalMesh->GetLODInfo(LODIndex);
-
-				const TArray<TObjectPtr<UMorphTarget>>& MorphTargets = SkeletalMesh->GetMorphTargets();
-				for (UMorphTarget* MorphTarget : MorphTargets)
-				{
-					Result->Morph.Names.Add(MorphTarget->GetFName());
-				}
-
-				const FMorphTargetVertexInfoBuffers& MorphTargetBuffers = Context->DataLOD->MorphTargetVertexInfoBuffers;
-				TArray<uint32> RuntimeCompressedMorphData;
-				if (MorphTargetBuffers.GetMorphDataSizeInBytes() == 0)
-				{
-					MUTABLE_CPUPROFILER_SCOPE(CompressMorphs)
-
-					// In Editor morph data is not compressed and is located in FMorphTargetLODModel.	
-					TArray<const FMorphTargetLODModel*> MorphTargetLODs;
-					MorphTargetLODs.Reserve(MorphTargets.Num());
-	
-					for (UMorphTarget* MorphTarget : MorphTargets)
-					{
-						TArray<FMorphTargetLODModel>& MorphLODModels = MorphTarget->GetMorphLODModels();
-	
-						const FMorphTargetLODModel* MorphTargetLOD = MorphLODModels.IsValidIndex(LODIndex) 
-								? &MorphLODModels[LODIndex] 
-								: nullptr;
-						MorphTargetLODs.Add(MorphTargetLOD);
-					}
-
-					check(MorphTargetLODs.Num() == Result->Morph.Names.Num());
-					
-					Result->Morph.BatchStartOffsetPerMorph.Empty(MorphTargetLODs.Num());
-					Result->Morph.BatchesPerMorph.Empty(MorphTargetLODs.Num());
-					Result->Morph.MaximumValuePerMorph.Empty(MorphTargetLODs.Num());
-					Result->Morph.MinimumValuePerMorph.Empty(MorphTargetLODs.Num());
-
-					// NOTE: Here we are compressing the morphs for all sections. This could be avoided pre-filtering
-					// the morph vertices in the section. This is done this way so the post-compress filtering is applied
-					// in editor to mimic what will happen with cooked data.
-
-					const int32 NumOriginalVertices = 
-							Context->OriginalLOD->RenderSections.Last().BaseVertexIndex +
-							Context->OriginalLOD->RenderSections.Last().NumVertices;
-
-					// From the Mutable point of view, all morphs need tangents.
-					TBitArray<> VertexNeedsTangents;
-					VertexNeedsTangents.Init(true, NumOriginalVertices);
-			
-					const float PositionPrecision = ComputePositionPrecision(SkeletalMeshLODInfo->MorphTargetPositionErrorTolerance);
-					const float TangentZPrecision = ComputeTangentPrecision();
-
-					TArray<FDeltaBatchHeader> BatchHeaders;
-					TArray<uint32> BitstreamData;
-				
-					for (const FMorphTargetLODModel* MorphModel : MorphTargetLODs)
-					{
-						if (!MorphModel)
-						{
-							continue;
-						}
-						
-						uint32 BatchStartOffset = BatchHeaders.Num();
-						
-						FVector4f MinimumValues;
-						FVector4f MaximumValues;
-
-						ComputeMorphTargetsMinAndMaxValues(
-								TConstArrayView<const FMorphTargetLODModel*>(&MorphModel, 1),
-								TArrayView<FVector4f>(&MinimumValues, 1),
-								TArrayView<FVector4f>(&MaximumValues, 1));
-
-						// Encode the actual morph vertex info into the quantized bitstream.
-						Encode(MorphModel->Vertices, &VertexNeedsTangents, PositionPrecision, TangentZPrecision, BatchHeaders, BitstreamData);
-					
-						const uint32 MorphNumBatches = BatchHeaders.Num() - BatchStartOffset;
-						Result->Morph.BatchStartOffsetPerMorph.Add(BatchStartOffset);
-						Result->Morph.BatchesPerMorph.Add(MorphNumBatches);
-						Result->Morph.MaximumValuePerMorph.Add(MaximumValues);
-						Result->Morph.MinimumValuePerMorph.Add(MinimumValues);
-					}
-				
-					Result->Morph.PositionPrecision = PositionPrecision;
-					Result->Morph.TangentZPrecision = TangentZPrecision;
-					Result->Morph.NumTotalBatches = BatchHeaders.Num();
-					
-					// Fix batch headers and write them packed.
-					for (FDeltaBatchHeader& BatchHeader : BatchHeaders)
-					{
-						BatchHeader.DataOffset += BatchHeaders.Num() * NumBatchHeaderDwords*sizeof(uint32);
-
-						TStaticArray<uint32, NumBatchHeaderDwords> HeaderData;
-						WriteHeader(BatchHeader, HeaderData);
-						RuntimeCompressedMorphData.Append(HeaderData);
-					}
-
-					RuntimeCompressedMorphData.Append(BitstreamData);
-				}
-
-				const TConstArrayView<uint32> SourceMorphDataView = !RuntimeCompressedMorphData.IsEmpty() 
-						? TConstArrayView<uint32>(RuntimeCompressedMorphData)
-						: TConstArrayView<uint32>(MorphTargetBuffers.GetData(), MorphTargetBuffers.GetMorphDataSizeInBytes()/sizeof(uint32));
-
-				if (SourceMorphDataView.GetData())
-				{	
-					MUTABLE_CPUPROFILER_SCOPE(ExtractSectionMorphs)
-					UE::Mutable::Private::FMeshMorph& Morph = Result->Morph;
-				
-					if (RuntimeCompressedMorphData.IsEmpty())
-					{
-						const int32 NumMorphs = MorphTargetBuffers.GetNumMorphs();
-						Morph.MaximumValuePerMorph.SetNumUninitialized(NumMorphs);
-						Morph.MinimumValuePerMorph.SetNumUninitialized(NumMorphs);
-						Morph.BatchStartOffsetPerMorph.SetNumUninitialized(NumMorphs);
-						Morph.BatchesPerMorph.SetNumUninitialized(NumMorphs);
-
-						for (int32 MorphIndex = 0; MorphIndex < NumMorphs; ++MorphIndex)
-						{
-							Morph.MaximumValuePerMorph[MorphIndex] = MorphTargetBuffers.GetMaximumMorphScale(MorphIndex);
-							Morph.MinimumValuePerMorph[MorphIndex] = MorphTargetBuffers.GetMinimumMorphScale(MorphIndex);
-							Morph.BatchStartOffsetPerMorph[MorphIndex] = MorphTargetBuffers.GetBatchStartOffset(MorphIndex);
-							Morph.BatchesPerMorph[MorphIndex] = MorphTargetBuffers.GetNumBatches(MorphIndex);
-							Morph.NumTotalBatches = MorphTargetBuffers.GetNumBatches();
-							Morph.PositionPrecision = MorphTargetBuffers.GetPositionPrecision();
-							Morph.TangentZPrecision = MorphTargetBuffers.GetTangentZPrecision();
-						}
-					}
-
-					// Strip all vertices not in this section.
-					const int32 SectionVertexRangeBegin = Section.BaseVertexIndex;
-					const int32 SectionVertexRangeEnd   = Section.BaseVertexIndex + Section.NumVertices;
-
-					uint32 CumulativeDataOffsetInDwords = 0;
-
-					int32 NumMorphs = Morph.BatchStartOffsetPerMorph.Num();
-					for (int32 MorphIndex = 0; MorphIndex < NumMorphs; ++MorphIndex)
-					{	
-						const int32 SourceMorphNumBatches = Morph.BatchesPerMorph[MorphIndex];
-						const int32 SourceBatchOffset = Morph.BatchStartOffsetPerMorph[MorphIndex]; 
-
-						if (SourceMorphNumBatches == 0)
-						{
-							continue;
-						}
-
-						int32 BatchIndex = 0;
-						for ( ; BatchIndex < SourceMorphNumBatches; ++BatchIndex)
-						{
-							FDeltaBatchHeader BatchHeader;
-							TConstArrayView<uint32> BatchHeaderDataView(
-									SourceMorphDataView.GetData() + (SourceBatchOffset + BatchIndex) * NumBatchHeaderDwords, 
-									NumBatchHeaderDwords); 
-
-							ReadHeader(BatchHeader, BatchHeaderDataView);
-
-							if (BatchHeader.IndexMin >= (uint32)SectionVertexRangeBegin)
-							{
-								break;
-							}
-						}
-
-						// The previous batch could have deltas in the section range.
-						int32 FirstSectionBatchIndex = FMath::Max(0, BatchIndex - 1);
-
-						for ( ; BatchIndex < SourceMorphNumBatches; ++BatchIndex)
-						{
-							FDeltaBatchHeader BatchHeader;
-							TConstArrayView<uint32> BatchHeaderDataView(
-									SourceMorphDataView.GetData() + (SourceBatchOffset + BatchIndex) * NumBatchHeaderDwords, 
-									NumBatchHeaderDwords); 
-
-							ReadHeader(BatchHeader, BatchHeaderDataView);
-
-							if (BatchHeader.IndexMin >= (uint32)SectionVertexRangeEnd)
-							{
-								break;
-							}
-						}
-						
-						int32 LastSectionBatchIndex = FMath::Max(0, BatchIndex - 1);
-
-						uint32 NumSectionMorphBatches = LastSectionBatchIndex + 1 - FirstSectionBatchIndex;
-						
-						TConstArrayView<uint32> SectionMorphHeadersView(
-									SourceMorphDataView.GetData() + (SourceBatchOffset + FirstSectionBatchIndex) * NumBatchHeaderDwords, 
-									NumSectionMorphBatches * NumBatchHeaderDwords);
-
-						Result->MorphDataBuffer.Append(SectionMorphHeadersView);
-
-						Result->Morph.BatchStartOffsetPerMorph[MorphIndex] = CumulativeDataOffsetInDwords / NumBatchHeaderDwords;
-						Result->Morph.BatchesPerMorph[MorphIndex] = NumSectionMorphBatches;
-						CumulativeDataOffsetInDwords += SectionMorphHeadersView.Num();
-					}
-
-					auto TrimBatchDeltasNotInRange = [](
-							int32 VertexRangeBegin, int32 VertexRangeEnd,
-							FDeltaBatchHeader& InOutBatchHeader,
-							TConstArrayView<uint32> DeltasData,
-							TArrayView<uint32> OutBatchDeltasData)
-					{
-						TConstArrayView<uint32> BatchDeltasData(
-								DeltasData.GetData() + (InOutBatchHeader.DataOffset / sizeof(uint32)), 
-								CalculateBatchDwords(InOutBatchHeader));
-					
-						TArray<FQuantizedDelta, TInlineAllocator<UE::MorphTargetVertexCodec::BatchSize>> QuantizedDeltas;
-						QuantizedDeltas.SetNumUninitialized(UE::MorphTargetVertexCodec::BatchSize);
-
-						ReadQuantizedDeltas(QuantizedDeltas, InOutBatchHeader, BatchDeltasData);
-
-						int32 NewIndexMin = InOutBatchHeader.IndexMin;
-						uint32 DeltaIndex = 0;
-						for ( ; DeltaIndex < InOutBatchHeader.NumElements; ++DeltaIndex)
-						{
-							int32 DeltaVertexIdx = QuantizedDeltas[DeltaIndex].Index;
-
-							if (DeltaVertexIdx >= VertexRangeBegin)
-							{
-								NewIndexMin = DeltaVertexIdx;
-								break;
-							}
-						}
-
-						int32 DeltaIndexBegin = DeltaIndex;
-
-						for ( ; DeltaIndex < InOutBatchHeader.NumElements; ++DeltaIndex)
-						{
-							int32 DeltaVertexIdx = QuantizedDeltas[DeltaIndex].Index;
-							
-							if (DeltaVertexIdx >= VertexRangeEnd)
-							{
-								break;
-							}
-						}
-
-						InOutBatchHeader.IndexMin    = NewIndexMin;
-						InOutBatchHeader.NumElements = DeltaIndex - DeltaIndexBegin;
-
-						check(InOutBatchHeader.NumElements * sizeof(FQuantizedDelta) <= QuantizedDeltas.NumBytes());
-						FMemory::Memmove(
-								QuantizedDeltas.GetData(), 
-								QuantizedDeltas.GetData() + DeltaIndexBegin,
-								InOutBatchHeader.NumElements * sizeof(FQuantizedDelta));
-						
-						WriteQuantizedDeltas(QuantizedDeltas, InOutBatchHeader, OutBatchDeltasData);
-					};
-
-					for (int32 MorphIndex = 0; MorphIndex < NumMorphs; ++MorphIndex)
-					{
-						const int32 MorphNumBatches = Result->Morph.BatchesPerMorph[MorphIndex];
-						const int32 BatchOffset = Result->Morph.BatchStartOffsetPerMorph[MorphIndex]; 
-				
-						if (MorphNumBatches == 0)
-						{
-							continue;
-						}
-
-						TArrayView<uint32> MorphBatchHeadersView(
-								Result->MorphDataBuffer.GetData() + BatchOffset * NumBatchHeaderDwords,
-								MorphNumBatches * NumBatchHeaderDwords);
-
-						// Batches still point to the source data.
-						FDeltaBatchHeader FirstBatchHeader;
-						TArrayView<uint32> FirstBatchHeaderDataView(
-								MorphBatchHeadersView.GetData(), NumBatchHeaderDwords);
-						ReadHeader(FirstBatchHeader, FirstBatchHeaderDataView);
-
-						FDeltaBatchHeader LastBatchHeader;
-						TArrayView<uint32> LastBatchHeaderDataView(
-								MorphBatchHeadersView.GetData() + (MorphNumBatches - 1)*NumBatchHeaderDwords, NumBatchHeaderDwords);
-						ReadHeader(LastBatchHeader, LastBatchHeaderDataView);
-
-						FDeltaBatchHeader FirstFullBatchHeader;
-						FDeltaBatchHeader LastFullBatchHeader;
-						
-						TConstArrayView<uint32> SourceFullBatchesDataView;
-						if (MorphNumBatches > 2)
-						{
-							TArrayView<uint32> FirstFullBatchHeaderDataView(
-									MorphBatchHeadersView.GetData() + 1*NumBatchHeaderDwords, NumBatchHeaderDwords);
-							ReadHeader(FirstFullBatchHeader, FirstFullBatchHeaderDataView);
-
-							TArrayView<uint32> LastFullBatchHeaderDataView(
-									MorphBatchHeadersView.GetData() + (MorphNumBatches - 2)*NumBatchHeaderDwords, NumBatchHeaderDwords);
-							ReadHeader(LastFullBatchHeader, LastFullBatchHeaderDataView);
-
-							SourceFullBatchesDataView = TConstArrayView<uint32>(
-								SourceMorphDataView.GetData() + (FirstFullBatchHeader.DataOffset / sizeof(uint32)),
-								((LastFullBatchHeader.DataOffset - FirstFullBatchHeader.DataOffset) / sizeof(uint32)) + CalculateBatchDwords(LastFullBatchHeader));
-						}
-
-						TArray<uint32, TInlineAllocator<BatchSize*sizeof(FQuantizedDelta)/sizeof(uint32)>> PartialBatchDeltaStorage;
-						uint32 PartialBatchDeltaStorageDwords = FMath::Max(CalculateBatchDwords(FirstBatchHeader), CalculateBatchDwords(LastBatchHeader));
-						PartialBatchDeltaStorage.SetNumUninitialized(PartialBatchDeltaStorageDwords);
-
-						// We are potentially reallocating, views will be invalidated after the reserve or any append.
-						// TODO: We can know an accurate upper-bound of the final size beforehand only looking at source 
-						// headers, pre-compute the size and only make one allocation.
-						Result->MorphDataBuffer.Reserve(
-								MorphBatchHeadersView.Num() + SourceFullBatchesDataView.Num() + 
-								CalculateBatchDwords(FirstBatchHeader) + CalculateBatchDwords(LastBatchHeader));
-
-						TrimBatchDeltasNotInRange(
-								SectionVertexRangeBegin, SectionVertexRangeEnd, FirstBatchHeader, SourceMorphDataView, PartialBatchDeltaStorage);
-						
-						FirstBatchHeaderDataView = TArrayView<uint32>(
-								Result->MorphDataBuffer.GetData() + BatchOffset*NumBatchHeaderDwords, NumBatchHeaderDwords);
-						WriteHeader(FirstBatchHeader, FirstBatchHeaderDataView);			
-	
-						// Batch size may have changed, CalculateBatchDwords value cannot be cached from previous query. 
-						Result->MorphDataBuffer.Append(
-								TArrayView<uint32>(PartialBatchDeltaStorage.GetData(), CalculateBatchDwords(FirstBatchHeader)));
-					
-						// SourceFullBatchesDataView is expected to be empty in some cases.
-						Result->MorphDataBuffer.Append(SourceFullBatchesDataView);
-
-						if (MorphNumBatches > 1)
-						{
-							PartialBatchDeltaStorage.SetNumUninitialized(CalculateBatchDwords(LastBatchHeader), EAllowShrinking::No);
-
-							TrimBatchDeltasNotInRange(
-									SectionVertexRangeBegin, SectionVertexRangeEnd, LastBatchHeader, SourceMorphDataView, PartialBatchDeltaStorage);
-							
-							LastBatchHeaderDataView = TArrayView<uint32>(
-									Result->MorphDataBuffer.GetData() + (BatchOffset + MorphNumBatches - 1)*NumBatchHeaderDwords, NumBatchHeaderDwords);
-							WriteHeader(LastBatchHeader, LastBatchHeaderDataView);					
-	
-							// Batch size may have changed, CalculateBatchDwords value cannot be cached from previous query. 
-							Result->MorphDataBuffer.Append(
-									TArrayView<uint32>(PartialBatchDeltaStorage.GetData(), CalculateBatchDwords(LastBatchHeader)));
-						}
-
-						int32 NumMorphBatches = Result->Morph.BatchesPerMorph[MorphIndex]; 
-
-						// Batch header fixup.
-						TArrayView<uint32> BatchHeadersData(
-								Result->MorphDataBuffer.GetData() + Result->Morph.BatchStartOffsetPerMorph[MorphIndex] * NumBatchHeaderDwords, 
-								NumMorphBatches * NumBatchHeaderDwords);
-						
-						for (int32 BatchIndex = 0; BatchIndex < NumMorphBatches; ++BatchIndex)
-						{
-							TArrayView<uint32> BatchHeaderData(
-									BatchHeadersData.GetData() + BatchIndex*NumBatchHeaderDwords, NumBatchHeaderDwords);
-
-							FDeltaBatchHeader BatchHeader;
-							ReadHeader(BatchHeader, BatchHeaderData);
-
-							if (BatchHeader.NumElements != 0)
-							{
-								check(BatchHeader.IndexMin >= (uint32)SectionVertexRangeBegin);
-								BatchHeader.IndexMin = BatchHeader.IndexMin - SectionVertexRangeBegin;
-								BatchHeader.DataOffset = CumulativeDataOffsetInDwords * sizeof(uint32);
-								CumulativeDataOffsetInDwords += CalculateBatchDwords(BatchHeader);
-							}
-							else
-							{
-								BatchHeader.IndexMin   = 0;
-								BatchHeader.DataOffset = CumulativeDataOffsetInDwords;
-							}
-
-							WriteHeader(BatchHeader, BatchHeaderData);
-						}
-					}
-				}
-	
-				// Make sure data can be loaded 4 Dwords at a time. Probably not needed for Mutable.
-				if (Result->MorphDataBuffer.Num())
-				{
-					Result->MorphDataBuffer.Append({0, 0, 0});
-				}
+				SetMorphData(
+						Result.Get(), SkeletalMesh.Get(), 
+						LODIndex, SectionIndex, 
+						(int32)Section.BaseVertexIndex, (int32)Section.NumVertices, (int32)Context->OriginalLOD->GetNumVertices(), 
+						nullptr);
 					
 				Result->EnsureSurfaceData();
 
diff --git a/Engine/Plugins/Mutable/Source/CustomizableObjectEditor/Private/MuCOE/CustomizableObjectCompiler.cpp b/Engine/Plugins/Mutable/Source/CustomizableObjectEditor/Private/MuCOE/CustomizableObjectCompiler.cpp
--- a/Engine/Plugins/Mutable/Source/CustomizableObjectEditor/Private/MuCOE/CustomizableObjectCompiler.cpp
+++ b/Engine/Plugins/Mutable/Source/CustomizableObjectEditor/Private/MuCOE/CustomizableObjectCompiler.cpp
@@ -7,6 +7,7 @@
 #include "ClothConfig.h"
 #include "HAL/RunnableThread.h"
 #include "Interfaces/ITargetPlatformManagerModule.h"
+
 #include "MessageLogModule.h"
 #include "GenerateMutableSource/GenerateMutableSourceComponent.h"
 #include "MuCO/CustomizableObject.h"
@@ -1487,14 +1488,6 @@ void FCustomizableObjectCompiler::FinishLoadingModelDataFromDDC()
 							}
 							break;
 						}
-						case UE::Mutable::Private::EStreamableDataType::RealTimeMorph:
-						{
-							for (const UE::Mutable::Private::FBlock& Block : File->Blocks)
-							{
-								PlatformData->MorphStreamableData.Set(Block.Id, TempData.GetData() + Block.Offset, Block.Size);
-							}
-							break;
-						}
 
 						default:
 							unimplemented();
@@ -1678,32 +1671,6 @@ void FCustomizableObjectCompiler::FinishCompilationTask()
 		ModelResources->StreamedResourceDataEditor = MoveTemp(CompilationContext->StreamedResourceData);
 
 		TSharedPtr<FModelStreamableBulkData> ModelStreamableBulkData = PlatformData->ModelStreamableBulkData;
-		ModelStreamableBulkData->RealTimeMorphStreamables.Empty(32);
-
-		uint64 RealTimeMorphDataOffsetInBytes = 0;
-		for (TPair<uint32, FRealTimeMorphMeshData>& MeshData : CompilationContext->RealTimeMorphTargetPerMeshData)
-		{
-			const uint32 DataSizeInBytes = (uint32)MeshData.Value.Data.Num() * sizeof(FMorphTargetVertexData);
-			FRealTimeMorphStreamable& ResourceMeshData = ModelStreamableBulkData->RealTimeMorphStreamables.FindOrAdd(MeshData.Key);
-
-			check(ResourceMeshData.NameResolutionMap.IsEmpty());
-			check(ResourceMeshData.Size == 0);
-
-			ResourceMeshData.NameResolutionMap = MeshData.Value.NameResolutionMap;
-			ResourceMeshData.Size = DataSizeInBytes;
-			EMutableFileFlags Flags = EMutableFileFlags::None;
-			ResourceMeshData.Block = FMutableStreamableBlock{ .FileId = uint32(0), .Flags = uint16(Flags), .Offset = RealTimeMorphDataOffsetInBytes };
-			ResourceMeshData.SourceId = MeshData.Value.SourceId;
-
-			RealTimeMorphDataOffsetInBytes += DataSizeInBytes;
-
-			PlatformData->MorphStreamableData.Set(
-				MeshData.Key,
-				reinterpret_cast<uint8*>(MeshData.Value.Data.GetData()),
-				DataSizeInBytes);
-
-			MeshData.Value.Data.Empty();
-		}
 		
 		ModelResources->SurfaceMetadata = MoveTemp(CompilationContext->SurfaceMetadata);
 		ModelResources->MeshMetadata = MoveTemp(CompilationContext->MeshMetadata);
diff --git a/Engine/Plugins/Mutable/Source/CustomizableObjectEditor/Private/MuCOE/GenerateMutableSource/GenerateMutableSource.h b/Engine/Plugins/Mutable/Source/CustomizableObjectEditor/Private/MuCOE/GenerateMutableSource/GenerateMutableSource.h
--- a/Engine/Plugins/Mutable/Source/CustomizableObjectEditor/Private/MuCOE/GenerateMutableSource/GenerateMutableSource.h
+++ b/Engine/Plugins/Mutable/Source/CustomizableObjectEditor/Private/MuCOE/GenerateMutableSource/GenerateMutableSource.h
@@ -366,18 +366,6 @@ struct FPoseBoneData
 	TArray<FTransform> ArrayTransform;
 };
 
-struct FRealTimeMorphMeshData
-{
-	TArray<FName> NameResolutionMap;
-	TArray<FMorphTargetVertexData> Data;
-
-	/* Used to group data when generating bulk data files.
-	 * This property should not be taken into consideration when comparing structs.
-	 */
-	uint32 SourceId = 0;
-};
-
-
 struct FGroupProjectorTempData
 {
 	class UCustomizableObjectNodeGroupProjectorParameter* CustomizableObjectNodeGroupProjectorParameter;
@@ -527,9 +515,6 @@ struct FMutableCompilationContext
 	/** Global morph selection overrides. */
 	TArray<FRealTimeMorphSelectionOverride> RealTimeMorphTargetsOverrides;
 
-	/** Data used for MorphTarget reconstruction. */
-	TMap<uint32, FRealTimeMorphMeshData> RealTimeMorphTargetPerMeshData;
-	
 	/** Stores the physics assets gathered from the SkeletalMesh nodes during compilation, to be used in mesh generation in-game */
 	TArray<TSoftObjectPtr<UPhysicsAsset>> PhysicsAssets;
 
diff --git a/Engine/Plugins/Mutable/Source/CustomizableObjectEditor/Private/MuCOE/GenerateMutableSource/GenerateMutableSourceMesh.cpp b/Engine/Plugins/Mutable/Source/CustomizableObjectEditor/Private/MuCOE/GenerateMutableSource/GenerateMutableSourceMesh.cpp
--- a/Engine/Plugins/Mutable/Source/CustomizableObjectEditor/Private/MuCOE/GenerateMutableSource/GenerateMutableSourceMesh.cpp
+++ b/Engine/Plugins/Mutable/Source/CustomizableObjectEditor/Private/MuCOE/GenerateMutableSource/GenerateMutableSourceMesh.cpp
@@ -157,16 +157,6 @@ void SetBuffersFormat(FMutableCompilationContext& Context,
 
 	MeshContext.NumBaseBuffers = MutableBufferCount;
 
-	if (bHasRealTimeMorphs)
-	{
-		MutableBufferCount += 2;
-	}
-
-	if (bHasClothing)
-	{
-		MutableBufferCount += 2;
-	}
-
 	if (bHasSkinWeightProfiles)
 	{
 		MutableBufferCount += LODModel.SkinWeightProfiles.Num();
@@ -215,72 +205,7 @@ void SetBuffersFormat(FMutableCompilationContext& Context,
 
 	// Optional buffers that might be removed when converting the mesh
 
-	// MorphTarget vertex tracking info buffers
-	if (bHasRealTimeMorphs)
-	{
-		using namespace UE::Mutable::Private;
-		{
-			const int32 ElementSize = sizeof(uint32);
-			constexpr int32 ChannelCount = 1;
-			const EMeshBufferSemantic Semantics[ChannelCount] = { EMeshBufferSemantic::Other };
-			const int32 SemanticIndices[ChannelCount] = { 0 };
-			const EMeshBufferFormat Formats[ChannelCount] = { EMeshBufferFormat::UInt32 };
-			const int32 Components[ChannelCount] = { 1 };
-			const int32 Offsets[ChannelCount] = { 0 };
-
-			VertexBuffers.SetBuffer(CurrentVertexBuffer, ElementSize, ChannelCount, Semantics, SemanticIndices, Formats, Components, Offsets);
-			check(!VertexBuffers.Buffers[CurrentVertexBuffer].HasPadding());
-			++CurrentVertexBuffer;
-		}
-
-		{
-			const int32 ElementSize = sizeof(uint32);
-			constexpr int32 ChannelCount = 1;
-			const EMeshBufferSemantic Semantics[ChannelCount] = { EMeshBufferSemantic::Other };
-			const int32 SemanticIndices[ChannelCount] = { 1 };
-			const EMeshBufferFormat Formats[ChannelCount] = { EMeshBufferFormat::UInt32 };
-			const int32 Components[ChannelCount] = { 1 };
-			const int32 Offsets[ChannelCount] = { 0 };
-
-			VertexBuffers.SetBuffer(CurrentVertexBuffer, ElementSize, ChannelCount, Semantics, SemanticIndices, Formats, Components, Offsets);
-			check(!VertexBuffers.Buffers[CurrentVertexBuffer].HasPadding());
-			++CurrentVertexBuffer;
-		}
-	}
-
-	//Clothing Data Buffer.
-	if (bHasClothing)
-	{
-		{
-			using namespace UE::Mutable::Private;
-			const int32 ElementSize = sizeof(int32);
-			constexpr int32 ChannelCount = 1;
-			const EMeshBufferSemantic Semantics[ChannelCount] = { EMeshBufferSemantic::Other };
-			const int32 SemanticIndices[ChannelCount] = { 2 };
-			const EMeshBufferFormat Formats[ChannelCount] = { EMeshBufferFormat::Int32 };
-			const int32 Components[ChannelCount] = { 1 };
-			const int32 Offsets[ChannelCount] = { 0 };
-
-			VertexBuffers.SetBuffer(CurrentVertexBuffer, ElementSize, ChannelCount, Semantics, SemanticIndices, Formats, Components, Offsets);
-			check(!VertexBuffers.Buffers[CurrentVertexBuffer].HasPadding());
-			++CurrentVertexBuffer;
-		}
 
-		{
-			using namespace UE::Mutable::Private;
-			const int32 ElementSize = sizeof(uint32);
-			constexpr int32 ChannelCount = 1;
-			const EMeshBufferSemantic Semantics[ChannelCount] = { EMeshBufferSemantic::Other };
-			const int32 SemanticIndices[ChannelCount] = { 3 };
-			const EMeshBufferFormat Formats[ChannelCount] = { EMeshBufferFormat::UInt32 };
-			const int32 Components[ChannelCount] = { 1 };
-			const int32 Offsets[ChannelCount] = { 0 };
-
-			VertexBuffers.SetBuffer(CurrentVertexBuffer, ElementSize, ChannelCount, Semantics, SemanticIndices, Formats, Components, Offsets);
-			check(!VertexBuffers.Buffers[CurrentVertexBuffer].HasPadding());
-			++CurrentVertexBuffer;
-		}
-	}
 
 	if (bHasSkinWeightProfiles)
 	{
@@ -664,6 +589,83 @@ bool IsSkeletalMeshCompatibleWithRefSkeleton(FSkeletalMeshComponentInfo& Compone
 }
 
 
+TArray<FName> GetUsedRealTimeMorphsNames(const FMutableSourceMeshData& SourceMeshData, const TArray<FRealTimeMorphSelectionOverride>& Overrides)
+{
+	TArray<FName> Result;
+	
+	if (SourceMeshData.TableReferenceSkeletalMesh.IsNull())
+	{	
+		const USkeletalMesh* SourceSkeletalMesh = Cast<USkeletalMesh>(UE::Mutable::Private::LoadObject(SourceMeshData.Mesh));
+		
+		if (SourceSkeletalMesh)
+		{
+			if (SourceMeshData.bUseAllRealTimeMorphs)
+			{
+				const TArray<UMorphTarget*>& SkeletalMeshMorphTargets = SourceSkeletalMesh->GetMorphTargets();
+
+				for (const UMorphTarget* MorphTarget : SkeletalMeshMorphTargets)
+				{
+					check(MorphTarget);
+					Result.Add(MorphTarget->GetFName());
+				}
+			}
+			else
+			{
+				for (const FString& MorphName : SourceMeshData.UsedRealTimeMorphTargetNames)
+				{
+					Result.Emplace(*MorphName);
+				}
+			}
+
+			for (const FRealTimeMorphSelectionOverride& MorphTargetOverride : Overrides)
+			{
+				const ECustomizableObjectSelectionOverride OverrideValue = Invoke([&]() -> ECustomizableObjectSelectionOverride
+				{
+					const int32 FoundMeshIndex = MorphTargetOverride.SkeletalMeshes.IndexOfByPredicate(
+					[Name = SourceSkeletalMesh->GetFName()](const FSkeletalMeshMorphTargetOverride& Elem)
+					{
+						return Name == Elem.SkeletalMeshName;
+					});
+
+					if (FoundMeshIndex != INDEX_NONE)
+					{
+						return MorphTargetOverride.SkeletalMeshes[FoundMeshIndex].SelectionOverride;
+					}
+
+					return MorphTargetOverride.SelectionOverride;
+				});
+
+				if (OverrideValue == ECustomizableObjectSelectionOverride::Enable)
+				{
+					Result.AddUnique(MorphTargetOverride.MorphName);
+				}
+				else if (OverrideValue == ECustomizableObjectSelectionOverride::Disable)
+				{
+					Result.Remove(MorphTargetOverride.MorphName);
+				}
+			}
+		}
+	}
+	else
+	{
+		const USkeletalMesh* TableReferenceSkeletalMesh = Cast<USkeletalMesh>(UE::Mutable::Private::LoadObject(SourceMeshData.TableReferenceSkeletalMesh));
+
+		if (TableReferenceSkeletalMesh)
+		{
+			for (const UMorphTarget* MorphTarget : TableReferenceSkeletalMesh->GetMorphTargets())
+			{
+				if (MorphTarget)
+				{
+					Result.Add(MorphTarget->GetFName());
+				}
+			}
+		}
+	}
+
+
+	return Result;
+}
+
 TArray<TTuple<UPhysicsAsset*, int32>> GetPhysicsAssetsFromAnimInstance(const TSoftClassPtr<UAnimInstance>& AnimInstance)
 {
 	// TODO: Consider caching the result in the GenerationContext.
@@ -1365,310 +1367,12 @@ UE::Tasks::FTask ConvertSkeletalMeshToMutable(TSharedRef<FMeshConversionContext>
 
 	// RealTimeMorphs and Clothing buffers might be removed if unused. Convert them before launching the tasks to convert non-optional buffers.
 
-	// RealTimeMorphs
-	uint32 MorphUniqueHash = 0;
 	if (!bIgnoreMorphs && Context.Options.bRealTimeMorphTargetsEnabled && SkeletalMesh->GetMorphTargets().Num())
 	{
-
-		// This call involves resolving every TObjectPtr<UMorphTarget> to a UMorphTarget*, so
-		// cache the result here to avoid calling it repeatedly.
-		const TArray<UMorphTarget*>& SkeletalMeshMorphTargets = SkeletalMesh->GetMorphTargets();
-
-		// Find real-time MorphTargets to be used.
-		TArray<const UMorphTarget*> UsedMorphTargets;
-		UsedMorphTargets.Reserve(SkeletalMeshMorphTargets.Num());
-
-		// Add SkeletalMesh node used defined real-time morph targets to a temporal array where
-		// the actual to be used real-time morphs names will be placed.        
-		TArray<FName> UsedMorphTargetsNames = Invoke([&]()
-			{
-				TArray<FName> MorphTargetsNames;
-				MorphTargetsNames.Reserve(SkeletalMeshMorphTargets.Num());
-
-				if (!Source.TableReferenceSkeletalMesh)
-				{
-					if (Source.bUseAllRealTimeMorphs)
-					{
-						for (const UMorphTarget* MorphTarget : SkeletalMeshMorphTargets)
-						{
-							check(MorphTarget);
-							MorphTargetsNames.Add(MorphTarget->GetFName());
-						}
-					}
-					else
-					{
-						for (const FString& MorphName : Source.UsedRealTimeMorphTargetNames)
-						{
-							MorphTargetsNames.Emplace(*MorphName);
-						}
-					}
-				}
-				else
-				{
-					const USkeletalMesh* TableReferenceSkeletalMesh = Cast<USkeletalMesh>(Source.TableReferenceSkeletalMesh.Get());
-					check(TableReferenceSkeletalMesh);
-
-					for (const UMorphTarget* MorphTarget : SkeletalMeshMorphTargets)
-					{
-						check(MorphTarget && TableReferenceSkeletalMesh);
-						if (TableReferenceSkeletalMesh->FindMorphTarget(MorphTarget->GetFName()))
-						{
-							MorphTargetsNames.Add(MorphTarget->GetFName());
-						}
-					}
-				}
-
-				return MorphTargetsNames;
-			});
-
-		TArray<FRealTimeMorphSelectionOverride>& RealTimeMorphTargetOverrides = Context.RealTimeMorphTargetsOverrides;
-		for (FRealTimeMorphSelectionOverride& MorphTargetOverride : RealTimeMorphTargetOverrides)
-		{
-			const ECustomizableObjectSelectionOverride OverrideValue =
-				Invoke([&]() -> ECustomizableObjectSelectionOverride
-					{
-						const int32 FoundMeshIndex = MorphTargetOverride.SkeletalMeshes.IndexOfByPredicate(
-							[Name = SkeletalMesh->GetFName()](const FSkeletalMeshMorphTargetOverride& Elem)
-							{
-								return Name == Elem.SkeletalMeshName;
-							});
-
-						if (FoundMeshIndex != INDEX_NONE)
-						{
-							return MorphTargetOverride.SkeletalMeshes[FoundMeshIndex].SelectionOverride;
-						}
-
-						return MorphTargetOverride.SelectionOverride;
-					});
-
-			if (OverrideValue == ECustomizableObjectSelectionOverride::Enable)
-			{
-				UsedMorphTargetsNames.AddUnique(MorphTargetOverride.MorphName);
-			}
-			else if (OverrideValue == ECustomizableObjectSelectionOverride::Disable)
-			{
-				UsedMorphTargetsNames.Remove(MorphTargetOverride.MorphName);
-			}
-		}
-
-
-
-		for (const UMorphTarget* MorphTarget : SkeletalMeshMorphTargets)
-		{
-			if (!MorphTarget)
-			{
-				continue;
-			}
-
-			const bool bHasToBeAdded = UsedMorphTargetsNames.Contains(MorphTarget->GetFName());
-			if (bHasToBeAdded)
-			{
-				UsedMorphTargets.Add(MorphTarget);
-			}
-		}
-
-		if (UsedMorphTargets.Num())
-		{
-			MUTABLE_CPUPROFILER_SCOPE(RealTimeMorphTargetProcessing);
-
-			using namespace UE::Mutable::Private;
-			check(MutableMesh->GetVertexBuffers().GetBufferCount() >= NextBufferIndex + 2);
-
-			// MorphTarget block offset buffer
-			check(MutableMesh->GetVertexBuffers().Buffers[NextBufferIndex].Channels.Num() == 1);
-			check(MutableMesh->GetVertexBuffers().Buffers[NextBufferIndex].Channels[0].Semantic == EMeshBufferSemantic::Other);
-			check(MutableMesh->GetVertexBuffers().Buffers[NextBufferIndex].Channels[0].SemanticIndex == 0);
-			check(MutableMesh->GetVertexBuffers().Buffers[NextBufferIndex].Channels[0].Format == EMeshBufferFormat::UInt32);
-
-			// MorphTarget vertex morph count.
-
-			// MorphTarget vertex block id.
-			check(MutableMesh->GetVertexBuffers().Buffers[NextBufferIndex + 1].Channels.Num() == 1);
-			check(MutableMesh->GetVertexBuffers().Buffers[NextBufferIndex + 1].Channels[0].Semantic == EMeshBufferSemantic::Other);
-			check(MutableMesh->GetVertexBuffers().Buffers[NextBufferIndex + 1].Channels[0].SemanticIndex == 1);
-			check(MutableMesh->GetVertexBuffers().Buffers[NextBufferIndex + 1].Channels[0].Format == EMeshBufferFormat::UInt32);
-
-			// Setup MorphTarget reconstruction data.
-			TArrayView<uint32> VertexMorphsOffsetAndCountBufferView(reinterpret_cast<uint32*>(
-				MutableMesh->GetVertexBuffers().GetBufferData(NextBufferIndex)), VertexCount);
-			TArrayView<uint32> VertexMorphsResourceIdBufferView(reinterpret_cast<uint32*>(
-				MutableMesh->GetVertexBuffers().GetBufferData(NextBufferIndex + 1)), VertexCount);
-
-			for (uint32& Elem : VertexMorphsOffsetAndCountBufferView)
-			{
-				Elem = 0;
-			}
-
-			constexpr uint32 InvalidResourceId = 0;
-			for (uint32& Elem : VertexMorphsResourceIdBufferView)
-			{
-				Elem = InvalidResourceId;
-			}
-
-
-			TArray<FMorphTargetVertexData> MorphsMeshData;
-			MorphsMeshData.Reserve(32);
-
-			TArray<FName> NameResolutionMap;
-			NameResolutionMap.Reserve(32);
-
-			TArray<FMorphTargetVertexData> MorphsUsed;
-			for (uint32 VertexIdx = 0; VertexIdx < VertexCount; ++VertexIdx)
-			{
-				MorphsUsed.Reset(UsedMorphTargets.Num());
-
-				for (const UMorphTarget* MorphTarget : UsedMorphTargets)
-				{
-					if (!MorphTarget)
-					{
-						continue;
-					}
-
-					const TArray<FMorphTargetLODModel>& MorphLODModels = MorphTarget->GetMorphLODModels();
-
-					if (LODIndex >= MorphLODModels.Num() || !MorphLODModels[LODIndex].SectionIndices.Contains(SectionIndex))
-					{
-						continue;
-					}
-
-					// The vertices should be sorted by SourceIdx
-					check(MorphLODModels[LODIndex].Vertices.Num() < 2 || MorphLODModels[LODIndex].Vertices[0].SourceIdx < MorphLODModels[LODIndex].Vertices.Last().SourceIdx);
-
-					const int32 VertexFoundIndex = Algo::BinarySearchBy(MorphLODModels[LODIndex].Vertices, (uint32)VertexIdx + MeshSection.BaseVertexIndex,
-						[](const FMorphTargetDelta& Element) { return Element.SourceIdx; });
-
-					if (VertexFoundIndex == INDEX_NONE)
-					{
-						continue;
-					}
-
-					const FMorphTargetDelta& VertexFound = MorphLODModels[LODIndex].Vertices[VertexFoundIndex];
-					const FName MorphTargetName = MorphTarget->GetFName();
-
-					int32 MorphTargetNameIndex = NameResolutionMap.Find(MorphTargetName);
-
-					MorphTargetNameIndex = MorphTargetNameIndex != INDEX_NONE
-						? MorphTargetNameIndex
-						: NameResolutionMap.Emplace(MorphTargetName);
-
-					MorphsUsed.Emplace(FMorphTargetVertexData{ VertexFound.PositionDelta, VertexFound.TangentZDelta, (uint32)MorphTargetNameIndex });
-				}
-
-				if (MorphsUsed.Num())
-				{
-					// TODO: Those constants should be stored somewhere accessible to the decoding code.
-					constexpr uint32 Log2MaxNumVerts = 23;
-					constexpr uint32 Log2MaxNumMorphs = 32 - Log2MaxNumVerts;
-
-					checkf(MorphsMeshData.Num() < (1 << Log2MaxNumVerts) - 1, TEXT("Maximum number of vertices with morphs per mesh reached."));
-					checkf(MorphsUsed.Num() < (1 << Log2MaxNumMorphs) - 1, TEXT("Maximum number of morphs per vertex reached."));
-
-					uint32 EncodedVertexOffsetAndCount =
-						(uint32(MorphsMeshData.Num()) & ((1 << Log2MaxNumVerts) - 1)) | (uint32(MorphsUsed.Num()) << Log2MaxNumVerts);
-					VertexMorphsOffsetAndCountBufferView[VertexIdx] = EncodedVertexOffsetAndCount;
-
-					MorphsMeshData.Append(MorphsUsed);
-				}
-			}
-
-			// Only commit the morph if there is data.
-			if (MorphsMeshData.Num())
-			{
-				const uint32 DataHash = CityHash32(
-					reinterpret_cast<const char*>(MorphsMeshData.GetData()), MorphsMeshData.Num() * sizeof(FMorphTargetVertexData));
-
-				MorphUniqueHash = DataHash == InvalidResourceId ? InvalidResourceId + 1 : DataHash;
-
-				const FRealTimeMorphMeshData* FoundHash = Context.RealTimeMorphTargetPerMeshData.Find(MorphUniqueHash);
-
-				bool bIsDataAlreadyCollected = false;
-				if (FoundHash)
-				{
-					bIsDataAlreadyCollected =
-						FoundHash->NameResolutionMap == NameResolutionMap &&
-						FoundHash->Data.Num() == MorphsMeshData.Num() &&
-						FMemory::Memcmp(FoundHash->Data.GetData(), MorphsMeshData.GetData(),
-							MorphsMeshData.Num() * sizeof(FMorphTargetVertexData)) == 0;
-				}
-
-				// NOTE: This way of unique hash generation guarantees all valid values can be used but given its 
-				// sequential nature a cascade of changes can occur if new meshes are added. Not many hash collisions 
-				// are expected so it should not be problematic.
-				if (FoundHash && !bIsDataAlreadyCollected)
-				{
-					uint32 NumTries = 0;
-					for (; NumTries < TNumericLimits<uint32>::Max(); ++NumTries)
-					{
-						FoundHash = Context.RealTimeMorphTargetPerMeshData.Find(MorphUniqueHash);
-
-						if (!FoundHash)
-						{
-							break;
-						}
-
-						bIsDataAlreadyCollected =
-							FoundHash->NameResolutionMap == NameResolutionMap &&
-							FoundHash->Data.Num() == MorphsMeshData.Num() &&
-							FMemory::Memcmp(FoundHash->Data.GetData(), MorphsMeshData.GetData(),
-								MorphsMeshData.Num() * sizeof(FMorphTargetVertexData)) == 0;
-
-						if (bIsDataAlreadyCollected)
-						{
-							break;
-						}
-
-						MorphUniqueHash = MorphUniqueHash + 1 == InvalidResourceId ? InvalidResourceId + 1 : MorphUniqueHash + 1;
-					}
-
-					if (NumTries == TNumericLimits<uint32>::Max())
-					{
-						UE_LOG(LogMutable, Warning, TEXT("Maximum number of meshes with real time morphs reached, some morphs may not work as expected."));
-						MorphUniqueHash = InvalidResourceId;
-					}
-				}
-
-				if (MorphUniqueHash != InvalidResourceId)
-				{
-					// Add the data block reference to the mesh vertices.
-					for (uint32 VertexIdx = 0; VertexIdx < VertexCount; ++VertexIdx)
-					{
-						// Keep as invalid the vertices with no morphs, this way we can check if a vertex has morphs looking at either
-						// count and offset (Notice that if count is 0 then offset must be 0 as well) or the hash.
-						// This is useful because those two buffers are iterated in different phases when reconstructing the data back after
-						// mesh generation, but we need to know if a vertex has morphs in both phases. 
-						if (VertexMorphsOffsetAndCountBufferView[VertexIdx] != 0)
-						{
-							VertexMorphsResourceIdBufferView[VertexIdx] = MorphUniqueHash;
-						}
-					}
-
-					FCustomizableObjectStreameableResourceId StreamedMorphResource;
-					StreamedMorphResource.Id = MorphUniqueHash;
-					StreamedMorphResource.Type = static_cast<uint8>(FCustomizableObjectStreameableResourceId::EType::RealTimeMorphTarget);
-
-					MutableMesh->AddStreamedResource(BitCast<uint64>(StreamedMorphResource));
-
-					if (!bIsDataAlreadyCollected)
-					{
-						FRealTimeMorphMeshData& RealTimeMorphMeshData = Context.RealTimeMorphTargetPerMeshData.FindOrAdd(MorphUniqueHash);
-
-						check(RealTimeMorphMeshData.NameResolutionMap.IsEmpty());
-						check(RealTimeMorphMeshData.Data.IsEmpty());
-
-						RealTimeMorphMeshData.NameResolutionMap = MoveTemp(NameResolutionMap);
-						RealTimeMorphMeshData.Data = MoveTemp(MorphsMeshData);
-						RealTimeMorphMeshData.SourceId = MeshId;
-					}
-				}
-			}
-
-			NextBufferIndex += 2;
-		}
-		else
-		{
-			MutableMesh->GetVertexBuffers().RemoveBuffer(NextBufferIndex + 1);
-			MutableMesh->GetVertexBuffers().RemoveBuffer(NextBufferIndex);
-		}
+		UnrealConversionUtils::SetMorphData(
+				MutableMesh.Get(), SkeletalMesh, 
+				LODIndex, SectionIndex, Section.BaseVertexIndex, Section.NumVertices, LODModel.NumVertices,
+				nullptr);
 	}
 
 	if (!bIgnoreSkinning)
@@ -1892,7 +1596,6 @@ UE::Tasks::FTask ConvertSkeletalMeshToMutable(TSharedRef<FMeshConversionContext>
 		FMutableMeshMetadata MeshMetadata;
 		FMemory::Memzero(MeshMetadata);
 
-		MeshMetadata.MorphMetadataId = MorphUniqueHash;
 		MeshMetadata.SurfaceMetadataId = AddUniqueSurfaceMetadata(Source.Metadata, Context);
 
 		auto HashMeshMetadataFunc = [](const FMutableMeshMetadata& Data) -> uint32
@@ -2853,131 +2556,42 @@ UE::Mutable::Private::NodeMeshPtr GenerateMorphMesh(const UEdGraphPin* Pin,
 	// Target
 	UE::Mutable::Private::NodeMeshPtr BaseSourceMesh = SourceNode;
 
-	UE::Mutable::Private::TManagedPtr<UE::Mutable::Private::FMesh> MorphedSourceMesh;
-
-	bool bSuccess = false;
-	
-	if (const UCustomizableObjectNodeTable* TypedNodeTable = Cast<UCustomizableObjectNodeTable>(Pin->GetOwningNode()))
+	if (UCustomizableObjectNodeMeshMorph* TypedMorphNode = Cast<UCustomizableObjectNodeMeshMorph>(TypedNodeMorphs[MorphIndex].OwningNode))
 	{
-		TObjectPtr<UDataTable> Table = UE::Mutable::Private::LoadObject(TypedNodeTable->Table);
-		TObjectPtr<UScriptStruct> Structure = UE::Mutable::Private::LoadObject(TypedNodeTable->Structure);
-		const FString TableName = Table ? GetNameSafe(Table).ToLower() : GetNameSafe(Structure).ToLower();
-		const uint32 TableId = CityHash32(reinterpret_cast<const char*>(*TableName), TableName.Len() * sizeof(FString::ElementType));
-
-		UDataTable* DataTable = GetDataTable(TypedNodeTable, GenerationContext);
-
-		// Generate a new Column for each morph
-		TArray<uint32> RowIds;
-		const TArray<FName>& RowNames = GetRowsToCompile(*DataTable, *TypedNodeTable, GenerationContext, RowIds);
-		int32 NumRows = RowNames.Num();
-
-		// Should exist
-		FGeneratedMutableDataTableKey MutableTableKey(DataTable->GetName(), TypedNodeTable->VersionColumn, TypedNodeTable->CompilationFilterOptions);
-		UE::Mutable::Private::Ptr<UE::Mutable::Private::FTable> GeneratedTable = GenerationContext.GeneratedTables[MutableTableKey].GeneratedTable;
-		check(GeneratedTable);
-
-		FString ColumnName = TableColumnName + TypedNodeMorphs[MorphIndex].MorphTargetName;
-		int32 ColumnIndex = INDEX_NONE;
-
-		for (int32 RowIndex = 0; RowIndex < NumRows; ++RowIndex)
+		Result->bReshapeSkeleton = TypedMorphNode->bReshapeSkeleton;
+		Result->bReshapePhysicsVolumes = TypedMorphNode->bReshapePhysicsVolumes;
 		{
-			const FName RowName = RowNames[RowIndex];
-			const uint32 RowId = RowIds[RowIndex];
+			//const UEdGraphPin* ConnectedPin = FollowInputPin(*TypedMorphNode->MeshPin());
+			const UEdGraphPin* SourceMeshPin = Pin ? FindMeshBaseSource(*Pin, false, GenerationContext.MacroNodesStack) : nullptr;
+			const UEdGraphNode* SkeletalMeshNode = SourceMeshPin ? SourceMeshPin->GetOwningNode() : nullptr;
 
-			ColumnIndex = GeneratedTable->FindColumn(ColumnName);
+			TArray<TTuple<USkeletalMesh*, TSoftClassPtr<UAnimInstance>>> SkeletalMeshesToDeform = 
+					GetSkeletalMeshesInfoForReshapeSelection(SkeletalMeshNode, SourceMeshPin, GenerationContext);
 
-			if (ColumnIndex == INDEX_NONE)
+			bool bWarningFound = false;
+			if (TypedMorphNode->bReshapeSkeleton)
 			{
-				ColumnIndex = GeneratedTable->AddColumn(ColumnName, UE::Mutable::Private::ETableColumnType::Mesh);
+				bWarningFound = GetAndValidateReshapeBonesToDeform(
+					Result->BonesToDeform, TypedMorphNode->BonesToDeform, SkeletalMeshesToDeform, TypedMorphNode, TypedMorphNode->SelectionMethod, GenerationContext);
 			}
 
-			UE::Mutable::Private::TManagedPtr<UE::Mutable::Private::FMesh> MorphedSourceTableMesh = BuildMorphedMutableMesh(Pin, TypedNodeMorphs[MorphIndex].MorphTargetName, GenerationContext, bOnlyConnectedLOD, RowName);
-			GeneratedTable->SetCell(ColumnIndex, RowId, MorphedSourceTableMesh);
-		}
-
-		if (ColumnIndex > INDEX_NONE)
-		{
-			bSuccess = true;
-
-			UE::Mutable::Private::Ptr<UE::Mutable::Private::NodeMeshTable> MorphedSourceMeshNodeTable = new UE::Mutable::Private::NodeMeshTable;
-			MorphedSourceMeshNodeTable->Table = GeneratedTable;
-			MorphedSourceMeshNodeTable->ColumnName = ColumnName;
-			MorphedSourceMeshNodeTable->ParameterName = TypedNodeTable->ParameterName;
-			MorphedSourceMeshNodeTable->SetMessageContext( MorphNode );
-			MorphedSourceMeshNodeTable->SourceDataDescriptor.SourceId = TableId;
-			MorphedSourceMeshNodeTable->SourceDataDescriptor.OptionalMaxLODSize = 0;
-
-			UE::Mutable::Private::Ptr<UE::Mutable::Private::NodeMeshMakeMorph> Morph = new UE::Mutable::Private::NodeMeshMakeMorph;
-			Morph->Base = BaseSourceMesh;
-			Morph->Target = MorphedSourceMeshNodeTable;
-			Morph->bOnlyPositionAndNormal = true;
-			Morph->SetMessageContext(MorphNode);
-
-			Result->Morph = Morph;
-		}
-	}
-	else
-	{
-		MorphedSourceMesh = BuildMorphedMutableMesh(Pin, TypedNodeMorphs[MorphIndex].MorphTargetName, GenerationContext, bOnlyConnectedLOD);
-		bSuccess = true;
-
-		if (MorphedSourceMesh)
-		{
-			UE::Mutable::Private::Ptr<UE::Mutable::Private::NodeMeshConstant> MorphedSourceMeshNode = new UE::Mutable::Private::NodeMeshConstant;
-			MorphedSourceMeshNode->Value = MorphedSourceMesh;
-			MorphedSourceMeshNode->SetMessageContext(MorphNode);
-
-			UE::Mutable::Private::Ptr<UE::Mutable::Private::NodeMeshMakeMorph> Morph = new UE::Mutable::Private::NodeMeshMakeMorph;
-			Morph->Base = BaseSourceMesh;
-			Morph->Target = MorphedSourceMeshNode;
-			Morph->bOnlyPositionAndNormal = true;
-			Morph->SetMessageContext(MorphNode);
-
-			Result->Morph = Morph;
-
-			if (UCustomizableObjectNodeMeshMorph* TypedMorphNode = Cast<UCustomizableObjectNodeMeshMorph>(TypedNodeMorphs[MorphIndex].OwningNode))
+			if (TypedMorphNode->bReshapePhysicsVolumes)
 			{
-				Result->bReshapeSkeleton = TypedMorphNode->bReshapeSkeleton;
-				Result->bReshapePhysicsVolumes = TypedMorphNode->bReshapePhysicsVolumes;
-				{
-					//const UEdGraphPin* ConnectedPin = FollowInputPin(*TypedMorphNode->MeshPin());
-					const UEdGraphPin* SourceMeshPin = Pin ? FindMeshBaseSource(*Pin, false, GenerationContext.MacroNodesStack) : nullptr;
-					const UEdGraphNode* SkeletalMeshNode = SourceMeshPin ? SourceMeshPin->GetOwningNode() : nullptr;
-
-					TArray<TTuple<USkeletalMesh*, TSoftClassPtr<UAnimInstance>>> SkeletalMeshesToDeform = 
-							GetSkeletalMeshesInfoForReshapeSelection(SkeletalMeshNode, SourceMeshPin, GenerationContext);
-
-					bool bWarningFound = false;
-					if (TypedMorphNode->bReshapeSkeleton)
-					{
-						bWarningFound = GetAndValidateReshapeBonesToDeform(
-							Result->BonesToDeform, TypedMorphNode->BonesToDeform, SkeletalMeshesToDeform, TypedMorphNode, TypedMorphNode->SelectionMethod, GenerationContext);
-					}
-
-					if (TypedMorphNode->bReshapePhysicsVolumes)
-					{
-						bWarningFound = bWarningFound || GetAndValidateReshapePhysicsToDeform(
-							Result->PhysicsToDeform, 
-							TypedMorphNode->PhysicsBodiesToDeform, 
-							SkeletalMeshesToDeform,
-							TypedMorphNode->PhysicsSelectionMethod,
-							TypedMorphNode, GenerationContext);
-					}
-						
-					if (bWarningFound)
-					{
-						TypedMorphNode->SetRefreshNodeWarning();
-					}
-				}
+				bWarningFound = bWarningFound || GetAndValidateReshapePhysicsToDeform(
+					Result->PhysicsToDeform, 
+					TypedMorphNode->PhysicsBodiesToDeform, 
+					SkeletalMeshesToDeform,
+					TypedMorphNode->PhysicsSelectionMethod,
+					TypedMorphNode, GenerationContext);
+			}
+				
+			if (bWarningFound)
+			{
+				TypedMorphNode->SetRefreshNodeWarning();
 			}
 		}
 	}
 
-	if(!bSuccess)
-	{
-		GenerationContext.Log(LOCTEXT("MorphGenerationFailed", "Failed to generate morph target."), MorphNode);
-	}
-
 	return Result;
 }
 
@@ -3163,10 +2777,11 @@ UE::Mutable::Private::Ptr<UE::Mutable::Private::NodeMesh> GenerateMutableSourceM
 			// Layouts
 			// 
 			// Generate Layouts when using manual LODs or when processing the base Automatic LOD (the other automatic LODs will use the base layouts).
-			const bool bGenerateLayouts = GenerationContext.CurrentAutoLODStrategy != ECustomizableObjectAutomaticLODStrategy::AutomaticFromMesh
-				|| (GenerationContext.CurrentLOD == GenerationContext.FromLOD || GenerationContext.CurrentLOD <= GenerationContext.FirstLODAvailable[GenerationContext.CurrentSkeletalMeshComponent]);
+			const bool bGenerateLayouts = 
+					GenerationContext.CurrentAutoLODStrategy != ECustomizableObjectAutomaticLODStrategy::AutomaticFromMesh || 
+					(GenerationContext.CurrentLOD == GenerationContext.FromLOD || GenerationContext.CurrentLOD <= GenerationContext.FirstLODAvailable[GenerationContext.CurrentSkeletalMeshComponent]);
 
-			if(bGenerateLayouts)
+			if (bGenerateLayouts)
 			{
 				const FLayoutGenerationFlags& LayoutFlags = GenerationContext.LayoutGenerationFlags.Last();
 
@@ -3201,9 +2816,11 @@ UE::Mutable::Private::Ptr<UE::Mutable::Private::NodeMesh> GenerateMutableSourceM
 
 						MeshNode->Layouts[LayoutIndex] = LayoutNode;
 					}
-				}
+				}	
 			}
 
+			MeshNode->RealTimeMorphNames = GetUsedRealTimeMorphsNames(Source, GenerationContext.CompilationContext->RealTimeMorphTargetsOverrides);
+
 			const FString MeshName = TypedNodeSkel->SkeletalMesh.GetLongPackageName().ToLower();
 			MeshNode->SourceDataDescriptor.SourceId = CityHash32(reinterpret_cast<const char*>(*MeshName), MeshName.Len() * sizeof(FString::ElementType));
 			MeshNode->SourceDataDescriptor.OptionalMaxLODSize = 0;
@@ -3762,7 +3379,7 @@ UE::Mutable::Private::Ptr<UE::Mutable::Private::NodeMesh> GenerateMutableSourceM
 							Source.bOnlyConnectedLOD = bOnlyConnectedLOD;
 							GetLODAndSectionForAutomaticLODs(*GenerationContext.CompilationContext, Source, *DefaultSkeletalMesh, LODIndex, SectionIndex);
 
-							MutableColumnName = TypedNodeTable->GenerateSkeletalMeshMutableColumName(ColumnName, LODIndex, SectionIndex);
+							MutableColumnName = TypedNodeTable->GenerateSkeletalMeshMutableColumName(ColumnName, LODIndex, SectionIndex);	
 						}
 						else
 						{
@@ -3779,12 +3396,13 @@ UE::Mutable::Private::Ptr<UE::Mutable::Private::NodeMesh> GenerateMutableSourceM
 						Source.LODOffset = GenerationContext.CurrentLOD - GenerationContext.FromLOD;
 						Source.bOnlyConnectedLOD = bOnlyConnectedLOD; // True?
 						bSuccess = GenerateTableColumn(TypedNodeTable, Pin, GeneratedTable, ColumnName, Property, Source, GenerationContext);
-
+							
 						if (!bSuccess)
 						{
 							FString Msg = FString::Printf(TEXT("Failed to generate the mutable table column [%s]"), *MutableColumnName);
 							GenerationContext.Log(FText::FromString(Msg), Node);
 						}
+
 					}
 
 					if (bSuccess)
@@ -3844,6 +3462,13 @@ UE::Mutable::Private::Ptr<UE::Mutable::Private::NodeMesh> GenerateMutableSourceM
 								}
 							}
 
+							// Set real-time morph names.
+							{
+								FMutableSourceMeshData Source;
+								Source.TableReferenceSkeletalMesh = DefaultSkeletalMesh;
+								MeshTableNode->RealTimeMorphNames = GetUsedRealTimeMorphsNames(Source, {});
+							}
+
 							// Applying Mesh Morph Nodes
 							if (DefaultSkeletalMesh && GenerationContext.MeshMorphStack.Num())
 							{
@@ -3882,6 +3507,9 @@ UE::Mutable::Private::Ptr<UE::Mutable::Private::NodeMesh> GenerateMutableSourceM
 			BreakNode->bAutomaticLODs = GenerationContext.CurrentAutoLODStrategy == ECustomizableObjectAutomaticLODStrategy::AutomaticFromMesh;
 		}
 		
+		TSoftObjectPtr<UStreamableRenderAsset> NodeMesh = TypedNodeParam->GetMesh();
+		FMutableSourceMeshData Source;
+		Source.Mesh = NodeMesh;
 		// Layouts
 		{
 			const FLayoutGenerationFlags& LayoutFlags = GenerationContext.LayoutGenerationFlags.Last();
@@ -3919,8 +3547,7 @@ UE::Mutable::Private::Ptr<UE::Mutable::Private::NodeMesh> GenerateMutableSourceM
 				}
 			}
 
-			TSoftObjectPtr<UStreamableRenderAsset> LayoutMesh = TypedNodeParam->GetMesh();
-			if (!LayoutMesh.IsNull())
+			if (!Source.Mesh.IsNull())
 			{
 				const EMutableMeshConversionFlags ModifiersMeshFlags =
 					EMutableMeshConversionFlags::IgnoreSkinning |
@@ -3931,14 +3558,16 @@ UE::Mutable::Private::Ptr<UE::Mutable::Private::NodeMesh> GenerateMutableSourceM
 				GenerationContext.MeshGenerationFlags.Push(ModifiersMeshFlags);
 
 				// Get the base mesh
-				FMutableSourceMeshData Source;
-				Source.Mesh = LayoutMesh;
 				BreakNode->ReferenceMesh = ConvertSkeletalMeshToMutable(Source, BaseLODIndex, BaseSectionIndex, GenerationContext, Node);
 
 				GenerationContext.MeshGenerationFlags.Pop();
+
 			}
 		}
 
+		Source.bUseAllRealTimeMorphs = true; // TODO: Remove, this should be selected similarly to the constant case, meanwhile use all morphs from the Reference Mesh.
+		BreakNode->RealTimeMorphNames = GetUsedRealTimeMorphsNames(Source, GenerationContext.CompilationContext->RealTimeMorphTargetsOverrides);
+
 		// For constants we also have:
 		// animbp instance?
 		// gameplay tags?
diff --git a/Engine/Plugins/Mutable/Source/CustomizableObjectEditor/Private/MuCOE/SMutableCodeViewer.cpp b/Engine/Plugins/Mutable/Source/CustomizableObjectEditor/Private/MuCOE/SMutableCodeViewer.cpp
--- a/Engine/Plugins/Mutable/Source/CustomizableObjectEditor/Private/MuCOE/SMutableCodeViewer.cpp
+++ b/Engine/Plugins/Mutable/Source/CustomizableObjectEditor/Private/MuCOE/SMutableCodeViewer.cpp
@@ -1573,7 +1573,7 @@ void SMutableCodeViewer::GenerateElementRecursive(const int32& InStateIndex, UE:
 			TagCount = InProgram.ConstantUInt32Lists[Args.Tags.ListAddress].Num();
 		}
 
-		int32 ResourceCount = EnumHasAnyFlags(Args.Flags, OpEnumFlags::HasResources) ? 1 : 0;;
+		int32 ResourceCount = EnumHasAnyFlags(Args.Flags, OpEnumFlags::HasResources) ? 1 : 0;
 		if (EnumHasAnyFlags(Args.Flags, OpEnumFlags::IsResourceList) &&
 			ensure(InProgram.ConstantUInt64Lists.IsValidIndex(Args.ResourceIds.ListAddress)))
 		{
@@ -1582,6 +1582,15 @@ void SMutableCodeViewer::GenerateElementRecursive(const int32& InStateIndex, UE:
 					: 0;
 		}
 
+		int32 RealTimeMorphNamesCount = EnumHasAnyFlags(Args.Flags, OpEnumFlags::HasRealTimeMorphNames) ? 1 : 0;
+		if (EnumHasAnyFlags(Args.Flags, OpEnumFlags::IsRealTimeMorphNamesList) &&
+			ensure(InProgram.ConstantUInt32Lists.IsValidIndex(Args.RealTimeMorphNames.ListAddress)))
+		{
+			RealTimeMorphNamesCount = InProgram.ConstantUInt32Lists.IsValidIndex(Args.RealTimeMorphNames.ListAddress) 
+					? InProgram.ConstantUInt32Lists[Args.RealTimeMorphNames.ListAddress].Num() 
+					: 0;
+		}
+
 		const int32 SkeletonCount = Args.SkeletonId != UE::Mutable::Private::PASSTHROUGH_ID_INVALID ? 1 : 0;
 		const int32 PhysicsAssetCount = Args.PhysicsAssetId != UE::Mutable::Private::PASSTHROUGH_ID_INVALID ? 1 : 0;
 
@@ -1592,8 +1601,8 @@ void SMutableCodeViewer::GenerateElementRecursive(const int32& InStateIndex, UE:
 			AdditionalPhysicsAssetCount = InProgram.ConstantUInt32Lists[Args.AdditionalPhysicsAssetsIds.ListAddress].Num();
 		}
 
-		FString Caption = FString::Printf(TEXT("add %d Tags, %d Resources, %d Skeleton, %d PhysicsAsset, and %d Additional PhysicsAssets to "),
-			TagCount, ResourceCount, SkeletonCount, PhysicsAssetCount, AdditionalPhysicsAssetCount);
+		FString Caption = FString::Printf(TEXT("add %d Tags, %d Resources, %d Skeleton, %d PhysicsAsset, %d RealTimeMorphs and %d Additional PhysicsAssets to "),
+			TagCount, ResourceCount, SkeletonCount, PhysicsAssetCount, RealTimeMorphNamesCount, AdditionalPhysicsAssetCount);
 		AddOpFunc(Args.Source, Caption);
 
 		break;
diff --git a/Engine/Plugins/Mutable/Source/MutableRuntime/Internal/MuR/Mesh.h b/Engine/Plugins/Mutable/Source/MutableRuntime/Internal/MuR/Mesh.h
--- a/Engine/Plugins/Mutable/Source/MutableRuntime/Internal/MuR/Mesh.h
+++ b/Engine/Plugins/Mutable/Source/MutableRuntime/Internal/MuR/Mesh.h
@@ -60,6 +60,23 @@ namespace UE::Mutable::Private
 	MUTABLE_DEFINE_POD_SERIALISABLE(FSurfaceSubMesh);
 	MUTABLE_DEFINE_POD_VECTOR_SERIALISABLE(FSurfaceSubMesh);
 
+	enum class EMorphUsageFlags : uint8
+	{
+		None     = 0,
+
+		Baked    = 1 << 0,
+		RealTime = 1 << 1,
+		External = 1 << 2,
+
+		Merged   = RealTime | External,
+
+		AllFlags = 0xFF,
+	};
+
+	ENUM_CLASS_FLAGS(EMorphUsageFlags);
+	MUTABLE_DEFINE_POD_SERIALISABLE(EMorphUsageFlags);
+	MUTABLE_DEFINE_POD_VECTOR_SERIALISABLE(EMorphUsageFlags);
+
 	struct FMeshSurface
 	{
 		TArray<FSurfaceSubMesh, TInlineAllocator<1>> SubMeshes; 
@@ -107,9 +124,18 @@ namespace UE::Mutable::Private
 		TArray<uint32> BatchStartOffsetPerMorph;
 		TArray<uint32> BatchesPerMorph;
 		
+		TArray<TArray<int32, TInlineAllocator<1>>> SurfacesInUsePerMorph;
+
+		TArray<EMorphUsageFlags> UsageFlagsPerMorph;
+
 		uint32 NumTotalBatches = 0;
 		float PositionPrecision = 0.0f; // TODO MORPHS Due to merging, this must be per block
 		float TangentZPrecision = 0.0f; // TODO MORPHS Due to merging, this must be per block
+
+		bool operator==(const FMeshMorph&) const = default;
+
+		void Serialise(FOutputArchive&) const;
+		void Unserialise(FInputArchive&);
 	};
 
 	
@@ -194,7 +220,8 @@ namespace UE::Mutable::Private
 		WithAdditionalPhysics = 1 << 15,
 		WithStreamedResources = 1 << 16,
 		WithMorphData = 1 << 17,
-		WithClothData = 1 << 18,
+		WithMorphMetaData = 1 << 18,
+		WithClothData = 1 << 19,
 
 		AllFlags = 0xFFFFFFFF
 	};
@@ -597,8 +624,10 @@ namespace UE::Mutable::Private
 				bEqual &= *AdditionalPhysicsBodies[i] == *Other.AdditionalPhysicsBodies[i];
 			}
 
-			if (bEqual) bEqual = ClothSections == Other.ClothSections;
-			
+			bEqual = bEqual && ClothSections == Other.ClothSections;
+			bEqual = bEqual && Morph == Other.Morph;
+			bEqual = bEqual && MorphDataBuffer == Other.MorphDataBuffer;
+
 			return bEqual;
 		}
 
diff --git a/Engine/Plugins/Mutable/Source/MutableRuntime/Internal/MuR/Model.h b/Engine/Plugins/Mutable/Source/MutableRuntime/Internal/MuR/Model.h
--- a/Engine/Plugins/Mutable/Source/MutableRuntime/Internal/MuR/Model.h
+++ b/Engine/Plugins/Mutable/Source/MutableRuntime/Internal/MuR/Model.h
@@ -880,6 +880,7 @@ namespace UE::Mutable::Private
 					Result->Tags = MetaDataMesh->Tags; 
 					Result->SkeletonObjects = MetaDataMesh->SkeletonObjects;
 					Result->StreamedResources = MetaDataMesh->StreamedResources; 
+					Result->Morph = MetaDataMesh->Morph;
 				}
 			}
 
diff --git a/Engine/Plugins/Mutable/Source/MutableRuntime/Internal/MuR/OpMeshMerge.h b/Engine/Plugins/Mutable/Source/MutableRuntime/Internal/MuR/OpMeshMerge.h
--- a/Engine/Plugins/Mutable/Source/MutableRuntime/Internal/MuR/OpMeshMerge.h
+++ b/Engine/Plugins/Mutable/Source/MutableRuntime/Internal/MuR/OpMeshMerge.h
@@ -2,1483 +2,26 @@
 
 #pragma once
 
-#include "MuR/MeshPrivate.h"
-#include "MuR/Platform.h"
-#include "MuR/MutableMath.h"
-#include "MuR/OpMeshFormat.h"
-#include "MuR/MutableTrace.h"
-#include "MuR/MutableRuntimeModule.h"
+#include "MuR/ManagedPointer.h"
 
-#include "Engine/SkeletalMesh.h"
-#include "Rendering/MorphTargetVertexCodec.h"
-#include "Rendering/SkeletalMeshVertexBuffer.h"
-
-#include "UObject/StrongObjectPtr.h"
+#define UE_API MUTABLERUNTIME_API 
 
 namespace UE::Mutable::Private
 {
+	class FMesh;
+	class FSkeleton;
+
 	struct FMeshMergeScratchMeshes
 	{
 		TManagedPtr<FMesh> FirstReformat;
 		TManagedPtr<FMesh> SecondReformat;
 	};
 
-	
-	inline void FixMorphIndices(FMesh& Result, int32 VertexIndexOffset)
-	{
-        MUTABLE_CPUPROFILER_SCOPE(CompressedMorphs_Merge_FixMorphIndices)
-
-		using namespace UE::MorphTargetVertexCodec;
-
-		if (VertexIndexOffset == 0)
-		{
-			return;
-		}
-
-		int32 NumMorphs = Result.Morph.Names.Num();
-		for (int32 MorphIndex = 0; MorphIndex < NumMorphs; ++MorphIndex)
-		{
-			TArrayView<uint32> MorphHeadersDataView = TArrayView<uint32>(
-					Result.MorphDataBuffer.GetData() + Result.Morph.BatchStartOffsetPerMorph[MorphIndex],
-					Result.Morph.BatchesPerMorph[MorphIndex]);
-
-			int32 NumBatches = Result.Morph.BatchesPerMorph[MorphIndex];
-			for (int32 BatchIndex = 0; BatchIndex < NumBatches; ++BatchIndex)
-			{
-				TArrayView<uint32> HeaderDataView = TArrayView<uint32>(
-						MorphHeadersDataView.GetData() + BatchIndex*NumBatchHeaderDwords,
-						NumBatchHeaderDwords);
-
-				FDeltaBatchHeader BatchHeader;
-				ReadHeader(BatchHeader, HeaderDataView);
-
-				BatchHeader.IndexMin += VertexIndexOffset;
-
-				WriteHeader(BatchHeader, HeaderDataView);
-			}
-		}
-	}
-
-	
-	inline void MergeMorphs(FMesh& Result, const FMesh& MeshA, const FMesh& MeshB)
-	{
-        MUTABLE_CPUPROFILER_SCOPE(CompressedMorphs_Merge_Merge)
-
-		using namespace UE::MorphTargetVertexCodec;
-		
-		// Compute memory requirements.	
-		Result.Morph.Names.Reserve(MeshA.Morph.Names.Num() + MeshB.Morph.Names.Num());
-		Result.Morph.Names.Append(MeshA.Morph.Names);
-
-		TArray<int32, TInlineAllocator<128>> ResultToBMorphNameMap; 
-		ResultToBMorphNameMap.Init(-1, MeshB.Morph.Names.Num() + MeshA.Morph.Names.Num());
-
-		// MeshA morphs map is the identity, compute only the MorphB morphs map.
-		for (int32 Index = 0; Index < MeshB.Morph.Names.Num(); ++Index)
-		{
-			int32 MappedIndex = Result.Morph.Names.AddUnique(MeshB.Morph.Names[Index]);
-			ResultToBMorphNameMap[MappedIndex] = Index;
-		}
-	
-		const uint32 MergedDataBufferSizeInDwords = MeshA.MorphDataBuffer.Num() + MeshB.MorphDataBuffer.Num();  
-		
-		Result.MorphDataBuffer.SetNum(MergedDataBufferSizeInDwords);
-		Result.Morph.MinimumValuePerMorph.SetNumUninitialized(Result.Morph.Names.Num());
-		Result.Morph.MaximumValuePerMorph.SetNumUninitialized(Result.Morph.Names.Num());
-		Result.Morph.BatchesPerMorph.SetNumUninitialized(Result.Morph.Names.Num());
-		Result.Morph.BatchStartOffsetPerMorph.SetNumUninitialized(Result.Morph.Names.Num());
-		
-		Result.Morph.PositionPrecision = MeshA.Morph.PositionPrecision;
-		Result.Morph.TangentZPrecision = MeshA.Morph.TangentZPrecision;
-
-		const int32 MeshANumMorphs = MeshA.Morph.Names.Num();
-		const int32 ResultNumMorphs = Result.Morph.Names.Num();
-
-		// Merge headers.
-		int32 BatchAccumulatedOffsetInDwords = 0;
-		int32 MorphIndex = 0;
-		for (; MorphIndex < MeshANumMorphs; ++MorphIndex)
-		{
-			Result.Morph.MaximumValuePerMorph[MorphIndex] = MeshA.Morph.MaximumValuePerMorph[MorphIndex];
-			Result.Morph.MinimumValuePerMorph[MorphIndex] = MeshA.Morph.MinimumValuePerMorph[MorphIndex];
-
-			Result.Morph.BatchStartOffsetPerMorph[MorphIndex] = BatchAccumulatedOffsetInDwords / NumBatchHeaderDwords;
-			Result.Morph.BatchesPerMorph[MorphIndex] = MeshA.Morph.BatchesPerMorph[MorphIndex];
-
-			FMemory::Memcpy(
-					Result.MorphDataBuffer.GetData() + BatchAccumulatedOffsetInDwords, 
-					MeshA.MorphDataBuffer.GetData() + MeshA.Morph.BatchStartOffsetPerMorph[MorphIndex]*NumBatchHeaderDwords,
-					MeshA.Morph.BatchesPerMorph[MorphIndex] * NumBatchHeaderDwords*sizeof(uint32));
-
-			BatchAccumulatedOffsetInDwords += MeshA.Morph.BatchesPerMorph[MorphIndex] * NumBatchHeaderDwords;
-			
-			int32 MeshBMorphIndex = ResultToBMorphNameMap[MorphIndex];
-			if (MeshBMorphIndex >= 0)
-			{
-				Result.Morph.BatchesPerMorph[MorphIndex] += MeshB.Morph.BatchesPerMorph[MeshBMorphIndex];
-
-				const FVector4f& SourceMaxValue = MeshB.Morph.MaximumValuePerMorph[MeshBMorphIndex];
-				const FVector4f& SourceMinValue = MeshB.Morph.MaximumValuePerMorph[MeshBMorphIndex];
-
-				Result.Morph.MaximumValuePerMorph[MorphIndex] = 
-						SourceMaxValue.ComponentMax(Result.Morph.MaximumValuePerMorph[MorphIndex]);
-				Result.Morph.MinimumValuePerMorph[MorphIndex] = 
-						SourceMinValue.ComponentMin(Result.Morph.MinimumValuePerMorph[MorphIndex]);
-
-				FMemory::Memcpy(
-						Result.MorphDataBuffer.GetData() + BatchAccumulatedOffsetInDwords, 
-						MeshB.MorphDataBuffer.GetData() + MeshB.Morph.BatchStartOffsetPerMorph[MeshBMorphIndex]*NumBatchHeaderDwords,
-						MeshB.Morph.BatchesPerMorph[MeshBMorphIndex] * NumBatchHeaderDwords*sizeof(uint32));
-				
-				BatchAccumulatedOffsetInDwords += MeshB.Morph.BatchesPerMorph[MorphIndex] * NumBatchHeaderDwords;
-			}
-		}
-
-		for (; MorphIndex < ResultNumMorphs; ++MorphIndex)
-		{
-			int32 MeshBMorphIndex = ResultToBMorphNameMap[MorphIndex];
-			
-			Result.Morph.BatchesPerMorph[MorphIndex] = MeshB.Morph.BatchesPerMorph[MeshBMorphIndex]; 
-			Result.Morph.BatchStartOffsetPerMorph[MorphIndex] = BatchAccumulatedOffsetInDwords / NumBatchHeaderDwords;
-
-			Result.Morph.MaximumValuePerMorph[MorphIndex] = MeshB.Morph.MaximumValuePerMorph[MeshBMorphIndex];
-			Result.Morph.MinimumValuePerMorph[MorphIndex] = MeshB.Morph.MinimumValuePerMorph[MeshBMorphIndex];
-
-			FMemory::Memcpy(
-					Result.MorphDataBuffer.GetData() + BatchAccumulatedOffsetInDwords, 
-					MeshB.MorphDataBuffer.GetData() + MeshB.Morph.BatchStartOffsetPerMorph[MeshBMorphIndex]*NumBatchHeaderDwords,
-					MeshB.Morph.BatchesPerMorph[MeshBMorphIndex] * NumBatchHeaderDwords*sizeof(uint32));
-
-			BatchAccumulatedOffsetInDwords += Result.Morph.BatchesPerMorph[MorphIndex] * NumBatchHeaderDwords;
-		}
-
-		// Merge morph data and fix up headers.
-
-		const uint32 MeshANumVertices = MeshA.GetVertexCount();
-
-		// Add Morphs in MeshA and common morphs with MeshB.
-		MorphIndex = 0;
-		for (; MorphIndex < MeshANumMorphs; ++MorphIndex)
-		{		
-			TArrayView<uint32> MorphHeadersDataView = TArrayView<uint32>(
-					Result.MorphDataBuffer.GetData() + Result.Morph.BatchStartOffsetPerMorph[MorphIndex]*NumBatchHeaderDwords,
-					Result.Morph.BatchesPerMorph[MorphIndex] * NumBatchHeaderDwords);
-
-			int32 BatchIndex = 0;
-			{
-				int32 NumBatches = MeshA.Morph.BatchesPerMorph[MorphIndex];
-				for ( ; BatchIndex < NumBatches; ++BatchIndex)
-				{
-					TArrayView<uint32> HeaderDataView = TArrayView<uint32>(
-							MorphHeadersDataView.GetData() + BatchIndex*NumBatchHeaderDwords,
-							NumBatchHeaderDwords);
-
-					FDeltaBatchHeader BatchHeader;
-					ReadHeader(BatchHeader, HeaderDataView);
-
-					const int32 SourceDataOffset = BatchHeader.DataOffset;
-					BatchHeader.DataOffset = BatchAccumulatedOffsetInDwords * sizeof(uint32);
-
-					WriteHeader(BatchHeader, HeaderDataView);
-
-					const int32 NumBatchDwords = CalculateBatchDwords(BatchHeader);
-					FMemory::Memcpy(
-							Result.MorphDataBuffer.GetData() + BatchHeader.DataOffset / sizeof(uint32),
-							MeshA.MorphDataBuffer.GetData() + SourceDataOffset / sizeof(uint32),
-							NumBatchDwords * sizeof(uint32));
-
-					BatchAccumulatedOffsetInDwords += NumBatchDwords;
-				}
-			}
-
-			int32 MeshBMorphIndex = ResultToBMorphNameMap[MorphIndex];
-			if (MeshBMorphIndex >= 0)
-			{
-				int32 NumBatches = MeshB.Morph.BatchesPerMorph[MeshBMorphIndex];
-
-				for (; BatchIndex < NumBatches; ++BatchIndex)
-				{
-					TArrayView<uint32> HeaderDataView = TArrayView<uint32>(
-							MorphHeadersDataView.GetData() + BatchIndex*NumBatchHeaderDwords,
-							NumBatchHeaderDwords);
-
-					FDeltaBatchHeader BatchHeader;
-					ReadHeader(BatchHeader, HeaderDataView);
-
-					const int32 SourceDataOffset = BatchHeader.DataOffset;
-					BatchHeader.DataOffset = BatchAccumulatedOffsetInDwords * sizeof(uint32);
-					BatchHeader.IndexMin += MeshANumVertices;
-
-					WriteHeader(BatchHeader, HeaderDataView);
-
-					const int32 NumBatchDwords = CalculateBatchDwords(BatchHeader);
-
-					check(Result.MorphDataBuffer.Num() >= BatchHeader.DataOffset / sizeof(uint32) + NumBatchDwords);
-					FMemory::Memcpy(
-							Result.MorphDataBuffer.GetData() + BatchHeader.DataOffset / sizeof(uint32),
-							MeshB.MorphDataBuffer.GetData() + SourceDataOffset / sizeof(uint32),
-							NumBatchDwords * sizeof(uint32));
-					
-					BatchAccumulatedOffsetInDwords += NumBatchDwords;
-				}
-			}
-		}
-
-		// Add Morphs in MeshB that are not present in MeshA.
-		for (; MorphIndex < ResultNumMorphs; ++MorphIndex)
-		{
-			TArrayView<uint32> MorphHeadersDataView = TArrayView<uint32>(
-					Result.MorphDataBuffer.GetData() + Result.Morph.BatchStartOffsetPerMorph[MorphIndex]*NumBatchHeaderDwords,
-					Result.Morph.BatchesPerMorph[MorphIndex]*NumBatchHeaderDwords);
-
-			int32 BatchIndex = 0;
-			{
-				int32 NumBatches = Result.Morph.BatchesPerMorph[MorphIndex];
-				for (; BatchIndex < NumBatches; ++BatchIndex)
-				{
-					TArrayView<uint32> HeaderDataView = TArrayView<uint32>(
-							MorphHeadersDataView.GetData() + BatchIndex*NumBatchHeaderDwords,
-							NumBatchHeaderDwords);
-
-					FDeltaBatchHeader BatchHeader;
-					ReadHeader(BatchHeader, HeaderDataView);
-
-					const int32 SourceDataOffset = BatchHeader.DataOffset;
-					BatchHeader.DataOffset = BatchAccumulatedOffsetInDwords * sizeof(uint32);
-					BatchHeader.IndexMin += MeshANumVertices;
-
-					WriteHeader(BatchHeader, HeaderDataView);
-
-					const int32 NumBatchDwords = CalculateBatchDwords(BatchHeader);
-				
-					check(Result.MorphDataBuffer.Num() >= BatchHeader.DataOffset / sizeof(uint32) + NumBatchDwords);
-					FMemory::Memcpy(
-							Result.MorphDataBuffer.GetData() + BatchHeader.DataOffset / sizeof(uint32),
-							MeshB.MorphDataBuffer.GetData() + SourceDataOffset / sizeof(uint32),
-							NumBatchDwords * sizeof(uint32));
-
-					BatchAccumulatedOffsetInDwords += NumBatchDwords;
-				}
-			}
-		}
-	}
-
-	inline void MergeSkeletons(FSkeleton& Base, const FSkeleton& Other, TArray<int32>& OutRemappedOtherBoneIndices)
-	{
-		const int32 NumBonesBaseBeforeMerge = Base.GetNumBones();
-		const int32 NumBonesOther = Other.GetNumBones();
-		OutRemappedOtherBoneIndices.SetNumUninitialized(NumBonesOther);
-
-		// Merge pSecond and build the remap array 
-		for (int32 Index = 0; Index < NumBonesOther; ++Index)
-		{
-			OutRemappedOtherBoneIndices[Index] = Base.FindOrAddBone(Other.BoneNames[Index], Other.BoneParents[Index]);
-		}
-
-		// Remap Parent bone indices to new Skeleton
-		for (int32 Index = NumBonesBaseBeforeMerge; Index < Base.BoneParents.Num(); ++Index)
-		{
-			int16 OtherBoneIndex = Base.BoneParents[Index];
-			if (OtherBoneIndex != INDEX_NONE)
-			{
-				Base.BoneParents[Index] = OutRemappedOtherBoneIndices[OtherBoneIndex];
-			}
-		}
-	}
-
-	
-	/** Copy cloth from MeshA. MeshB only required to for data integrity. */
-	inline void CopyClothFirst(FMesh& Result, const FMesh& MeshA, const FMesh& MeshB)
-	{
-		Result.ClothSections[0].ClothingAsset = MeshA.ClothSections[0].ClothingAsset;
-			
-		const uint32 NumClothVerticesA = MeshA.ClothSections[0].Data.Num();
-		if (!NumClothVerticesA)
-		{
-			return;
-		}
-
-		const int32 NumInfluences = MeshA.ClothSections[0].Data.Num() / MeshA.VertexBuffers.GetElementCount();
-		const uint32 NumClothVerticesB = NumInfluences * MeshB.VertexBuffers.GetElementCount();
-		
-		Result.ClothSections[0].Data.Reserve(NumClothVerticesA + NumClothVerticesB);
-
-		Result.ClothSections[0].Data.Append(MeshA.ClothSections[0].Data);
-		
-		// Add MeshB vertices without simulation. See FMeshToMeshVertData::SourceMeshVertIndices.
-		for (uint32 VertexIndex = 0; VertexIndex < NumClothVerticesB; ++VertexIndex)
-		{
-			FMeshToMeshVertData& Data = Result.ClothSections[0].Data.AddZeroed_GetRef();
-			Data.SourceMeshVertIndices[3] = static_cast<uint16>(0xffff);
-		}
-	}
-	
-	
-	/** Copy cloth from MeshB. MeshA only required to for data integrity. */
-	inline void CopyClothSecond(FMesh& Result, const FMesh& MeshA, const FMesh& MeshB)
-	{
-		Result.ClothSections[0].ClothingAsset = MeshB.ClothSections[0].ClothingAsset;
-
-		const uint32 NumClothVerticesB = MeshB.ClothSections[0].Data.Num();
-		if (!NumClothVerticesB)
-		{
-			return;
-		}
-		
-		const int32 NumInfluences = MeshB.ClothSections[0].Data.Num() / MeshB.VertexBuffers.GetElementCount();
-		const uint32 NumClothVerticesA =  NumInfluences * MeshA.VertexBuffers.GetElementCount();
-
-		Result.ClothSections[0].Data.Reserve(NumClothVerticesA + NumClothVerticesB);
-		
-		// Add MeshA vertices without simulation. See FMeshToMeshVertData::SourceMeshVertIndices.
-		for (uint32 VertexIndex = 0; VertexIndex < NumClothVerticesA; ++VertexIndex)
-		{
-			FMeshToMeshVertData& Data = Result.ClothSections[0].Data.AddZeroed_GetRef();
-			Data.SourceMeshVertIndices[3] = static_cast<uint16>(0xffff);					
-		}
-		
-		Result.ClothSections[0].Data.Append(MeshB.ClothSections[0].Data);
-	}
-	
-	
-	//---------------------------------------------------------------------------------------------
-	//! Merge two meshes into one new mesh
-	//---------------------------------------------------------------------------------------------
-	inline void MeshMerge(FMesh* Result, const FMesh* pFirst, const FMesh* pSecond, bool bMergeSurfaces, FMeshMergeScratchMeshes& ScratchMeshes)
-	{
-		MUTABLE_CPUPROFILER_SCOPE(MeshMerge);
-
-		// Should never happen, but fixes static analysis warnings.
-		if (!(pFirst && pSecond))
-		{
-			return;
-		}
-
-		// Indices
-		//-----------------
-		if (pFirst->GetIndexBuffers().GetBufferCount() > 0)
-		{
-			MUTABLE_CPUPROFILER_SCOPE(Indices);
-
-			const int32 FirstCount = pFirst->GetIndexBuffers().GetElementCount();
-			const int32 SecondCount = pSecond->GetIndexBuffers().GetElementCount();
-
-			if (pFirst->IndexBuffers.IsDescriptor() || pSecond->IndexBuffers.IsDescriptor())
-			{
-				EnumAddFlags(Result->IndexBuffers.Flags, EMeshBufferSetFlags::IsDescriptor);
-			}
-			
-			Result->GetIndexBuffers().SetElementCount(FirstCount + SecondCount);
-
-			check(pFirst->GetIndexBuffers().GetBufferCount() <= 1);
-			check(pSecond->GetIndexBuffers().GetBufferCount() <= 1);
-			Result->GetIndexBuffers().SetBufferCount(1);
-
-			FMeshBuffer& ResultIndexBuffer = Result->GetIndexBuffers().Buffers[0];
-
-			const FMeshBuffer& FirstIndexBuffer = pFirst->GetIndexBuffers().Buffers[0];
-			const FMeshBuffer& SecondIndexBuffer = pSecond->GetIndexBuffers().Buffers[0];
-
-			// Avoid unused variable warnings
-			(void)FirstIndexBuffer;
-			(void)SecondIndexBuffer;
-
-			// This will be changed below if need to change the format of the index buffers.
-			EMeshBufferFormat IndexBufferFormat = EMeshBufferFormat::None;
-
-			if (FirstCount && SecondCount)
-			{
-				check(!FirstIndexBuffer.Channels.IsEmpty());
-
-				// We need to know the total number of vertices in case we need to adjust the index buffer format.
-				const uint64 TotalVertexCount = pFirst->GetVertexBuffers().GetElementCount() + pSecond->GetVertexBuffers().GetElementCount();
-				const uint64 MaxValueBits = GetMeshFormatData(pFirst->GetIndexBuffers().Buffers[0].Channels[0].Format).MaxValueBits;
-				const uint64 MaxSupportedVertices = uint64(1) << MaxValueBits;
-				
-				if (TotalVertexCount > MaxSupportedVertices)
-				{
-					IndexBufferFormat = TotalVertexCount > MAX_uint16 ? EMeshBufferFormat::UInt32 : EMeshBufferFormat::UInt16;
-				}
-			}
-			
-			if (IndexBufferFormat != EMeshBufferFormat::None)
-			{
-				// We only support vertex indices in case of having to change the format.
-				check(FirstIndexBuffer.Channels.Num() == 1);
-
-				ResultIndexBuffer.Channels.SetNum(1);
-				ResultIndexBuffer.Channels[0].Semantic = EMeshBufferSemantic::VertexIndex;
-				ResultIndexBuffer.Channels[0].Format = IndexBufferFormat;
-				ResultIndexBuffer.Channels[0].ComponentCount = 1;
-				ResultIndexBuffer.Channels[0].SemanticIndex = 0;
-				ResultIndexBuffer.Channels[0].Offset = 0;
-				ResultIndexBuffer.ElementSize = GetMeshFormatData(IndexBufferFormat).SizeInBytes;
-			}
-			else if (FirstCount)
-			{
-				ResultIndexBuffer.Channels = FirstIndexBuffer.Channels;
-				ResultIndexBuffer.ElementSize = FirstIndexBuffer.ElementSize;
-			}
-			else if (SecondCount)
-			{
-				ResultIndexBuffer.Channels = SecondIndexBuffer.Channels;
-				ResultIndexBuffer.ElementSize = SecondIndexBuffer.ElementSize;
-			}
-
-
-			check(ResultIndexBuffer.Channels.Num() == 1);
-			check(ResultIndexBuffer.Channels[0].Semantic == EMeshBufferSemantic::VertexIndex);
-
-			if (!Result->IndexBuffers.IsDescriptor())
-			{
-
-				ResultIndexBuffer.Data.SetNum(ResultIndexBuffer.ElementSize * (FirstCount + SecondCount));
-			
-				if (!ResultIndexBuffer.Data.IsEmpty())
-				{
-					if (FirstCount)
-					{
-						if (IndexBufferFormat == EMeshBufferFormat::None
-							|| IndexBufferFormat == FirstIndexBuffer.Channels[0].Format)
-						{
-							FMemory::Memcpy(&ResultIndexBuffer.Data[0],
-								&FirstIndexBuffer.Data[0],
-								FirstIndexBuffer.ElementSize * FirstCount);
-						}
-						else
-						{
-							// Conversion required
-							const uint8_t* pSource = &FirstIndexBuffer.Data[0];
-							uint8_t* pDest = &ResultIndexBuffer.Data[0];
-							switch (IndexBufferFormat)
-							{
-							case EMeshBufferFormat::UInt32:
-							{
-								switch (FirstIndexBuffer.Channels[0].Format)
-								{
-								case EMeshBufferFormat::UInt16:
-								{
-									for (int32 v = 0; v < FirstCount; ++v)
-									{
-										*(uint32_t*)pDest = *(const uint16*)pSource;
-										pSource += FirstIndexBuffer.ElementSize;
-										pDest += ResultIndexBuffer.ElementSize;
-									}
-									break;
-								}
-
-								case EMeshBufferFormat::UInt8:
-								{
-									for (int32 v = 0; v < FirstCount; ++v)
-									{
-										*(uint32_t*)pDest = *(const uint8_t*)pSource;
-										pSource += FirstIndexBuffer.ElementSize;
-										pDest += ResultIndexBuffer.ElementSize;
-									}
-									break;
-								}
-
-								default:
-									checkf(false, TEXT("Format not supported."));
-									break;
-								}
-								break;
-							}
-
-							case EMeshBufferFormat::UInt16:
-							{
-								switch (FirstIndexBuffer.Channels[0].Format)
-								{
-
-								case EMeshBufferFormat::UInt8:
-								{
-									for (int32 v = 0; v < FirstCount; ++v)
-									{
-										*(uint16*)pDest = *(const uint8_t*)pSource;
-										pSource += FirstIndexBuffer.ElementSize;
-										pDest += ResultIndexBuffer.ElementSize;
-									}
-									break;
-								}
-
-								default:
-									checkf(false, TEXT("Format not supported."));
-									break;
-								}
-								break;
-							}
-
-							default:
-								checkf(false, TEXT("Format not supported."));
-								break;
-							}
-						}
-					}
-
-					if (SecondCount)
-					{
-						const uint8_t* pSource = &SecondIndexBuffer.Data[0];
-						uint8_t* pDest = &ResultIndexBuffer.Data[ResultIndexBuffer.ElementSize * FirstCount];
-
-						uint32_t firstVertexCount = pFirst->GetVertexBuffers().GetElementCount();
-
-						if (IndexBufferFormat == EMeshBufferFormat::None
-							|| IndexBufferFormat == SecondIndexBuffer.Channels[0].Format)
-						{
-							switch (SecondIndexBuffer.Channels[0].Format)
-							{
-							case EMeshBufferFormat::Int32:
-							case EMeshBufferFormat::UInt32:
-							case EMeshBufferFormat::NInt32:
-							case EMeshBufferFormat::NUInt32:
-							{
-								for (int32 v = 0; v < SecondCount; ++v)
-								{
-									*(uint32_t*)pDest = firstVertexCount + *(const uint32_t*)pSource;
-									pSource += SecondIndexBuffer.ElementSize;
-									pDest += ResultIndexBuffer.ElementSize;
-								}
-								break;
-							}
-
-							case EMeshBufferFormat::Int16:
-							case EMeshBufferFormat::UInt16:
-							case EMeshBufferFormat::NInt16:
-							case EMeshBufferFormat::NUInt16:
-							{
-								for (int32 v = 0; v < SecondCount; ++v)
-								{
-									*(uint16*)pDest = uint16(firstVertexCount) + *(const uint16*)pSource;
-									pSource += SecondIndexBuffer.ElementSize;
-									pDest += ResultIndexBuffer.ElementSize;
-								}
-								break;
-							}
-
-							case EMeshBufferFormat::Int8:
-							case EMeshBufferFormat::UInt8:
-							case EMeshBufferFormat::NInt8:
-							case EMeshBufferFormat::NUInt8:
-							{
-								for (int32 v = 0; v < SecondCount; ++v)
-								{
-									*(uint8_t*)pDest = uint8_t(firstVertexCount) + *(const uint8_t*)pSource;
-									pSource += SecondIndexBuffer.ElementSize;
-									pDest += ResultIndexBuffer.ElementSize;
-								}
-								break;
-							}
-
-							default:
-								checkf(false, TEXT("Format not supported."));
-								break;
-							}
-						}
-						else
-						{
-							// Format conversion required
-							switch (IndexBufferFormat)
-							{
-
-							case EMeshBufferFormat::UInt32:
-							{
-								switch (SecondIndexBuffer.Channels[0].Format)
-								{
-								case EMeshBufferFormat::Int16:
-								case EMeshBufferFormat::UInt16:
-								case EMeshBufferFormat::NInt16:
-								case EMeshBufferFormat::NUInt16:
-								{
-									for (int32 v = 0; v < SecondCount; ++v)
-									{
-										*(uint32_t*)pDest = uint32_t(firstVertexCount) + *(const uint16*)pSource;
-										pSource += SecondIndexBuffer.ElementSize;
-										pDest += ResultIndexBuffer.ElementSize;
-									}
-									break;
-								}
-
-								case EMeshBufferFormat::Int8:
-								case EMeshBufferFormat::UInt8:
-								case EMeshBufferFormat::NInt8:
-								case EMeshBufferFormat::NUInt8:
-								{
-									for (int32 v = 0; v < SecondCount; ++v)
-									{
-										*(uint32_t*)pDest = uint32_t(firstVertexCount) + *(const uint8_t*)pSource;
-										pSource += SecondIndexBuffer.ElementSize;
-										pDest += ResultIndexBuffer.ElementSize;
-									}
-									break;
-								}
-
-								default:
-									checkf(false, TEXT("Format not supported."));
-									break;
-								}
-
-								break;
-							}
-
-							case EMeshBufferFormat::UInt16:
-							{
-								switch (SecondIndexBuffer.Channels[0].Format)
-								{
-								case EMeshBufferFormat::Int8:
-								case EMeshBufferFormat::UInt8:
-								case EMeshBufferFormat::NInt8:
-								case EMeshBufferFormat::NUInt8:
-								{
-									for (int32 v = 0; v < SecondCount; ++v)
-									{
-										*(uint16*)pDest = uint16(firstVertexCount) + *(const uint8_t*)pSource;
-										pSource += SecondIndexBuffer.ElementSize;
-										pDest += ResultIndexBuffer.ElementSize;
-									}
-									break;
-								}
-
-								default:
-									checkf(false, TEXT("Format not supported."));
-									break;
-								}
-
-								break;
-							}
-
-							default:
-								checkf(false, TEXT("Format not supported."));
-								break;
-
-							}
-						}
-					}
-				}
-			}
-		}
-
-
-		// Layouts
-		//-----------------
-		{
-			MUTABLE_CPUPROFILER_SCOPE(Layouts);
-
-			int32 ResultLayoutCount = FMath::Max(pFirst->Layouts.Num(),pSecond->Layouts.Num());
-			Result->Layouts.SetNum(ResultLayoutCount);
-			for (int32 LayoutIndex = 0; LayoutIndex < ResultLayoutCount; ++LayoutIndex)
-			{
-				TManagedPtr<FLayout> pR;
-				
-				if (LayoutIndex < pFirst->Layouts.Num())
-				{
-					const FLayout* pF = pFirst->Layouts[LayoutIndex].Get();
-					pR = pF->Clone();
-				}
-
-				if (LayoutIndex < pSecond->Layouts.Num())
-				{
-					const FLayout* pS = pSecond->Layouts[LayoutIndex].Get();
-					if (!pR)
-					{
-						pR = pS->Clone();
-					}
-					else
-					{
-						pR->Blocks.Append(pS->Blocks);
-					}
-				}
-
-				Result->Layouts[LayoutIndex] = pR;
-			}
-		}
-
-
-		// Skeleton
-		//---------------------------
-
-		// Add SkeletonObjects
-		Result->SkeletonObjects = pFirst->SkeletonObjects;
-
-		for (const TPassthroughObjectPtr<USkeleton>& SkeletonObject : pSecond->SkeletonObjects)
-		{
-			Result->SkeletonObjects.AddUnique(SkeletonObject);
-		}
-
-		TManagedPtr<const FSkeleton> FirstSkeleton = pFirst->GetSkeleton();
-		TManagedPtr<const FSkeleton> SecondSkeleton = pSecond->GetSkeleton();
-
-		// Do they have the same skeleton?
-		bool bMergeSkeletons = FirstSkeleton != SecondSkeleton;
-
-		// Are they different skeletons but with the same data?
-		if (bMergeSkeletons && FirstSkeleton && SecondSkeleton)
-		{
-			bMergeSkeletons = !(*FirstSkeleton == *SecondSkeleton);
-		}
-
-		TArray<int32> RemappedSecondBoneIndices;
-
-		if (bMergeSkeletons)
-		{
-			MUTABLE_CPUPROFILER_SCOPE(MergeSkeleton);
-
-			TManagedPtr<FSkeleton> ResultSkeleton = FirstSkeleton ? FirstSkeleton->Clone() : MakeManaged<FSkeleton>();
-			Result->SetSkeleton(ResultSkeleton);
-
-			MergeSkeletons(*ResultSkeleton.Get(), *SecondSkeleton.Get(), RemappedSecondBoneIndices);
-		}
-		else
-		{
-			Result->SetSkeleton(FirstSkeleton);
-
-			const int32 NumBonesFirst = FirstSkeleton ? FirstSkeleton->GetNumBones() : 0;
-			RemappedSecondBoneIndices.SetNumUninitialized(NumBonesFirst);
-			for (int32 Index = 0; Index < NumBonesFirst; ++Index)
-			{
-				RemappedSecondBoneIndices[Index] = Index;
-			}
-		}
-
-
-		// Surfaces
-		//---------------------------
-		
-		// Remap bone indices if we merge surfaces since bonemaps will be merged too.
-		bool bRemapBoneIndices = false;
-		TArray<uint16> RemappedBoneMapIndices;
-
-		const int32 NumBonesInSecondBoneMap = pSecond->BoneMap.Num();
-
-		// Used to know the format of the bone index buffer
-		uint32 MaxNumBonesInBoneMaps = 0;
-
-		{
-			MUTABLE_CPUPROFILER_SCOPE(Surfaces);
-			
-			const int32 NumFirstBonesInBoneMap = pFirst->BoneMap.Num();
-			Result->BoneMap = pFirst->BoneMap;
-
-			// Copy BoneMap to remap bones when merging skeletons
-			TArray<FBoneIdOrIndex> SecondBoneMap = pSecond->BoneMap;
+	/** Merge two meshes into one new mesh */
+	void MeshMerge(FMesh* Result, const FMesh* pFirst, const FMesh* pSecond, bool bMergeSurfaces, FMeshMergeScratchMeshes& ScratchMeshes);
 
-			if (bMergeSkeletons)
-			{
-				// Remap BoneMap to new skeleton
-				for (int32 BoneIndex = 0; BoneIndex < NumBonesInSecondBoneMap; ++BoneIndex)
-				{
-					SecondBoneMap[BoneIndex].Index = RemappedSecondBoneIndices[SecondBoneMap[BoneIndex].Index];
-				}
-			}
-
-			if (bMergeSurfaces)
-			{
-				// Merge BoneMaps
-				RemappedBoneMapIndices.Reserve(NumBonesInSecondBoneMap);
-
-				for (uint16 Index = 0; Index < NumBonesInSecondBoneMap; ++Index)
-				{
-					const int32 BoneMapIndex = Result->BoneMap.AddUnique(SecondBoneMap[Index]);
-					RemappedBoneMapIndices.Add(BoneMapIndex);
-
-					bRemapBoneIndices = bRemapBoneIndices || BoneMapIndex != Index;
-				}
-
-				FMeshSurface& NewSurface = Result->Surfaces.AddDefaulted_GetRef();
-				Result->ClothSections.AddDefaulted();
-				NewSurface.BoneMapCount = Result->BoneMap.Num();
-
-				int32 NumFirstSubMeshes = 0;
-				for (const FMeshSurface& Surf : pFirst->Surfaces)
-				{
-					NewSurface.SubMeshes.Append(Surf.SubMeshes);
-					NumFirstSubMeshes += Surf.SubMeshes.Num();
-				}
-
-				for (const FMeshSurface& Surf : pSecond->Surfaces)
-				{
-					NewSurface.SubMeshes.Append(Surf.SubMeshes);
-				}
-
-				// Fix surface Submesh ranges.
-				if (NumFirstSubMeshes > 0)
-				{
-					const int32 NumResultSubMeshes = NewSurface.SubMeshes.Num();
-				
-					const FSurfaceSubMesh LastFromFirstMesh = pFirst->Surfaces.Last().SubMeshes.Last();
-	
-					for (int32 SecondSubMeshIndex = NumFirstSubMeshes; SecondSubMeshIndex < NumResultSubMeshes; ++SecondSubMeshIndex)
-					{
-						NewSurface.SubMeshes[SecondSubMeshIndex].VertexBegin += LastFromFirstMesh.VertexEnd;	
-						NewSurface.SubMeshes[SecondSubMeshIndex].VertexEnd += LastFromFirstMesh.VertexEnd;	
-						NewSurface.SubMeshes[SecondSubMeshIndex].IndexBegin += LastFromFirstMesh.IndexEnd;	
-						NewSurface.SubMeshes[SecondSubMeshIndex].IndexEnd += LastFromFirstMesh.IndexEnd;	
-					}
-				}
-			}
-			else
-			{
-				// Add the BoneMap of the second mesh
-				Result->BoneMap.Append(SecondBoneMap);
-
-				// Add pFirst surfaces
-				Result->Surfaces = pFirst->Surfaces;
-				
-				const int32 FirstVertexEnd = pFirst->GetVertexCount();
-				const int32 FirstIndexEnd = pFirst->GetIndexCount();
-
-				check(pSecond->Surfaces.Num() == 1);
-				FMeshSurface& NewSurface = Result->Surfaces.Add_GetRef(pSecond->Surfaces[0]);
-				
-				for (FSurfaceSubMesh& SubMesh : NewSurface.SubMeshes)
-				{
-					SubMesh.VertexBegin += FirstVertexEnd;
-					SubMesh.VertexEnd += FirstVertexEnd;
-					SubMesh.IndexBegin += FirstIndexEnd;
-					SubMesh.IndexEnd += FirstIndexEnd;
-				}
-
-				NewSurface.BoneMapIndex += NumFirstBonesInBoneMap;
-				
-				Result->ClothSections.SetNum(Result->Surfaces.Num());
-			}
-
-			for (const FMeshSurface& Surface : Result->Surfaces)
-			{
-				MaxNumBonesInBoneMaps = FMath::Max(MaxNumBonesInBoneMaps, Surface.BoneMapCount);
-			}
-
-			Result->BoneMap.Shrink();
-		}
-
-
-		// Pose
-		//---------------------------
-		if (TManagedPtr<const FSkeleton> ResultSkeleton = Result->GetSkeleton())
-		{
-			MUTABLE_CPUPROFILER_SCOPE(Pose);
-
-			Result->BonePoses.Reserve(ResultSkeleton->BoneNames.Num());
-
-			// Copy poses from the first mesh
-			Result->BonePoses = pFirst->BonePoses;
-
-			// Add or override bone poses
-			for (const FMesh::FBonePose& SecondBonePose : pSecond->BonePoses)
-			{
-				const int32 RemappedBoneIndex = RemappedSecondBoneIndices[SecondBonePose.BoneId.Index];
-				const int32 ResultBonePoseIndex = Result->FindBonePoseByBoneIndex(RemappedBoneIndex);
-
-				if (ResultBonePoseIndex != INDEX_NONE)
-				{
-					FMesh::FBonePose& ResultBonePose = Result->BonePoses[ResultBonePoseIndex];
-
-					// TODO: Not sure how to tune this priority, review it.
-					// For now use a similar strategy as before. 
-					auto ComputeBoneMergePriority = [](const FMesh::FBonePose& BonePose)
-					{
-						return (EnumHasAnyFlags(BonePose.BoneUsageFlags, EBoneUsageFlags::Skinning) ? 1 : 0) +
-							(EnumHasAnyFlags(BonePose.BoneUsageFlags, EBoneUsageFlags::Reshaped) ? 1 : 0);
-					};
-
-					const int32 ResultBonePriority = ComputeBoneMergePriority(ResultBonePose);
-					const int32 SecondBonePriority = ComputeBoneMergePriority(SecondBonePose);
-
-					if (ResultBonePriority < SecondBonePriority ||
-						(ResultBonePriority == SecondBonePriority && ResultBonePose.BonePosePriority < SecondBonePose.BonePosePriority))
-					{
-						//ResultBonePose.BoneName = SecondBonePose.BoneName;
-						ResultBonePose.BoneTransform = SecondBonePose.BoneTransform;
-						ResultBonePose.BonePosePriority = SecondBonePose.BonePosePriority;
-
-						// Merge usage flags
-						EnumAddFlags(ResultBonePose.BoneUsageFlags, SecondBonePose.BoneUsageFlags);
-					}
-				}
-				else
-				{
-					FMesh::FBonePose NewPose = SecondBonePose;
-					NewPose.BoneId.Index = RemappedBoneIndex;
-					Result->BonePoses.Add(NewPose);
-				}
-			}
-
-			Result->BonePoses.Shrink();
-		}
-
-
-		// PhysicsBodies
-		//---------------------------
-		{
-			MUTABLE_CPUPROFILER_SCOPE(PhysicsBodies);
-
-			// Append PhysicsAssets
-			Result->PhysicsAssets = pFirst->PhysicsAssets;
-
-			for (const TPassthroughObjectPtr<UPhysicsAsset>& PhysicsAsset : pSecond->PhysicsAssets)
-			{
-				Result->PhysicsAssets.AddUnique(PhysicsAsset);
-			}
-
-			// Appends InPhysicsBody to OutPhysicsBody removing Bodies that are equal and their properies are identical.	
-			auto AppendPhysicsBodiesUnique = [](FPhysicsBody& OutPhysicsBody, const FPhysicsBody& InPhysicsBody) -> bool
-			{
-				TArray<FName>& OutBones = OutPhysicsBody.BodiesBoneNames;
-				TArray<FPhysicsBodyAggregate>& OutBodies = OutPhysicsBody.Bodies;
-
-				const TArray<FName>& InBones = InPhysicsBody.BodiesBoneNames;
-				const TArray<FPhysicsBodyAggregate>& InBodies = InPhysicsBody.Bodies;
-
-				const int32 InBodyCount = InPhysicsBody.GetBodyCount();
-				const int32 OutBodyCount = OutPhysicsBody.GetBodyCount();
-
-				bool bModified = false;
-
-				for (int32 InBodyIndex = 0; InBodyIndex < InBodyCount; ++InBodyIndex)
-				{
-					const int32 FoundIndex = OutBones.Find(InBones[InBodyIndex]);
-
-					if (FoundIndex == INDEX_NONE)
-					{
-						OutBones.Add(InBones[InBodyIndex]);
-						OutBodies.Add(InBodies[InBodyIndex]);
-
-						bModified |= true;
-
-						continue;
-					}
-
-					for (const FSphereBody& Body : InBodies[InBodyIndex].Spheres)
-					{
-						const int32 NumBeforeAddition = OutBodies[FoundIndex].Spheres.Num();
-						bModified |= NumBeforeAddition == OutBodies[FoundIndex].Spheres.AddUnique(Body);
-					}
-
-					for (const FBoxBody& Body : InBodies[InBodyIndex].Boxes)
-					{
-						const int32 NumBeforeAddition = OutBodies[FoundIndex].Boxes.Num();
-						bModified |= NumBeforeAddition == OutBodies[FoundIndex].Boxes.AddUnique(Body);
-					}
-
-					for (const FSphylBody& Body : InBodies[InBodyIndex].Sphyls)
-					{
-						const int32 NumBeforeAddition = OutBodies[FoundIndex].Sphyls.Num();
-						bModified |= NumBeforeAddition == OutBodies[FoundIndex].Sphyls.AddUnique(Body);
-					}
-
-					for (const FTaperedCapsuleBody& Body : InBodies[InBodyIndex].TaperedCapsules)
-					{
-						const int32 NumBeforeAddition = OutBodies[FoundIndex].TaperedCapsules.Num();
-						bModified |= NumBeforeAddition == OutBodies[FoundIndex].TaperedCapsules.AddUnique(Body);
-					}
-
-					for (const FConvexBody& Body : InBodies[InBodyIndex].Convex)
-					{
-						const int32 NumBeforeAddition = OutBodies[FoundIndex].Convex.Num();
-						bModified |= NumBeforeAddition == OutBodies[FoundIndex].Convex.AddUnique(Body);
-					}
-				}
-
-				return bModified;
-			};
-
-			TTuple<TManagedPtr<const FPhysicsBody>, bool> SharedResultPhysicsBody = Invoke([&]()
-				-> TTuple<TManagedPtr<const FPhysicsBody>, bool>
-			{
-				if (pFirst->GetPhysicsBody() == pSecond->GetPhysicsBody())
-				{
-					return MakeTuple(pFirst->GetPhysicsBody(), true);
-				}
-
-				if (pFirst->GetPhysicsBody() && !pSecond->GetPhysicsBody())
-				{
-					return MakeTuple(pFirst->GetPhysicsBody(), true);
-				}
-
-				if (!pFirst->GetPhysicsBody() && pSecond->GetPhysicsBody())
-				{
-					return MakeTuple(pSecond->GetPhysicsBody(), true);
-				}
-
-				return MakeTuple(nullptr, false);
-			});
-
-			if (SharedResultPhysicsBody.Get<1>())
-			{
-				// Only one or non of the meshes has physics, share the result.
-				Result->SetPhysicsBody(SharedResultPhysicsBody.Get<0>());
-			}
-			else
-			{
-				check(pFirst->GetPhysicsBody() && pSecond->GetPhysicsBody());
-
-				TManagedPtr<FPhysicsBody> MergedResultPhysicsBody = pFirst->GetPhysicsBody()->Clone();
-
-				MergedResultPhysicsBody->bBodiesModified =
-					AppendPhysicsBodiesUnique(*MergedResultPhysicsBody, *pSecond->GetPhysicsBody()) ||
-					pFirst->GetPhysicsBody()->bBodiesModified || pSecond->GetPhysicsBody()->bBodiesModified;
-
-				Result->SetPhysicsBody(MergedResultPhysicsBody);
-			}
-
-			// Additional physics bodies.
-			const int32 MaxAdditionalPhysicsResultNum = pFirst->AdditionalPhysicsBodies.Num() + pSecond->AdditionalPhysicsBodies.Num();
-
-			Result->AdditionalPhysicsBodies.Reserve(MaxAdditionalPhysicsResultNum);
-			Result->AdditionalPhysicsBodies.Append(pFirst->AdditionalPhysicsBodies);
-
-			Result->AdditionalPhysicsAssets.Reserve(MaxAdditionalPhysicsResultNum);
-			Result->AdditionalPhysicsAssets.Append(pFirst->AdditionalPhysicsAssets);
-			
-			// Not very many additional bodies expected, do a quadratic search to have unique bodies based on external id.
-			const int32 NumSecondAdditionalBodies = pSecond->AdditionalPhysicsBodies.Num();
-			check(pSecond->AdditionalPhysicsAssets.Num() == NumSecondAdditionalBodies);
-
-			for (int32 Index = 0; Index < NumSecondAdditionalBodies; ++Index)
-			{
-				const int32 CustomIdToMerge = pSecond->AdditionalPhysicsBodies[Index]->CustomId;
-
-				const TManagedPtr<const FPhysicsBody>* Found = pFirst->AdditionalPhysicsBodies.FindByPredicate(
-					[CustomIdToMerge](const TManagedPtr<const FPhysicsBody>& Body) { return CustomIdToMerge == Body->CustomId; });
-
-				// TODO: current usages do not expect collisions, but same Id collision with bodies modified in differnet ways
-				// may need to be contemplated at some point.
-				if (!Found)
-				{
-					Result->AdditionalPhysicsBodies.Add(pSecond->AdditionalPhysicsBodies[Index]);
-					Result->AdditionalPhysicsAssets.Add(pSecond->AdditionalPhysicsAssets[Index]);
-				}
-			}
-		}
-
-		// This affects both vertex IDs and layout block ids.
-		bool bNeedsExplicitIds = pFirst->MeshIDPrefix != pSecond->MeshIDPrefix;
-
-		// These two extra checks are necessary for corner cases of meshes merging with fragments of themselves that
-		// undergo different operations.
-		if (!bNeedsExplicitIds)
-		{
-			UntypedMeshBufferIteratorConst FirstIDs(pFirst->VertexBuffers, EMeshBufferSemantic::VertexIndex, 0);
-			UntypedMeshBufferIteratorConst SecondIDs(pSecond->VertexBuffers, EMeshBufferSemantic::VertexIndex, 0);
-			bNeedsExplicitIds = FirstIDs.GetFormat() != SecondIDs.GetFormat();
-		}
-		if (!bNeedsExplicitIds)
-		{
-			UntypedMeshBufferIteratorConst FirstIDs(pFirst->VertexBuffers, EMeshBufferSemantic::LayoutBlock, 0);
-			UntypedMeshBufferIteratorConst SecondIDs(pSecond->VertexBuffers, EMeshBufferSemantic::LayoutBlock, 0);
-			bNeedsExplicitIds = FirstIDs.GetFormat() != SecondIDs.GetFormat();
-		}
-
-		if (!bNeedsExplicitIds)
-		{
-			// This is needed in case a mesh merges with itself.
-			Result->MeshIDPrefix = pFirst->MeshIDPrefix;
-		}
-
-		// Vertices
-		//-----------------
-		{
-            MUTABLE_CPUPROFILER_SCOPE(Vertices);
-
-			const int32 FirstBufferCount = pFirst->VertexBuffers.Buffers.Num();
-			const int32 SecondBufferCount = pSecond->VertexBuffers.Buffers.Num();
-			const int32 FirstVertexCount = pFirst->GetVertexBuffers().GetElementCount();
-			const int32 SecondVertexCount = pSecond->GetVertexBuffers().GetElementCount();
-
-			// Check if the format of the BoneIndex buffer has to change
-			bool bChangeBoneIndicesFormat = false;
-			EMeshBufferFormat BoneIndexFormat = MaxNumBonesInBoneMaps > MAX_uint8 ? EMeshBufferFormat::UInt16 : EMeshBufferFormat::UInt8;
-			bChangeBoneIndicesFormat = pFirst->GetVertexBuffers().HasAnySemanticWithDifferentFormat(EMeshBufferSemantic::BoneIndices, BoneIndexFormat);
-			if (!bChangeBoneIndicesFormat)
-			{
-				bChangeBoneIndicesFormat = pSecond->GetVertexBuffers().HasAnySemanticWithDifferentFormat(EMeshBufferSemantic::BoneIndices, BoneIndexFormat);
-			}
-
-			// Step 1: Set up the initial vertex buffer structure of the result mesh.
-			//-----------------------------------------------------------------------
-			{
-				MUTABLE_CPUPROFILER_SCOPE(ResultFormatSetup);
-
-				// Start with the structure of the first mesh
-				Result->GetVertexBuffers().SetBufferCount(FirstBufferCount);
-				for (int32 BufferIndex = 0; BufferIndex < FirstBufferCount; ++BufferIndex)
-				{
-					Result->VertexBuffers.Buffers[BufferIndex].Channels = pFirst->VertexBuffers.Buffers[BufferIndex].Channels;
-					Result->VertexBuffers.Buffers[BufferIndex].ElementSize = pFirst->VertexBuffers.Buffers[BufferIndex].ElementSize;
-				}
-
-				// See if we need to add additional buffers from the second mesh (like vertex colours or additional UV Channels)
-				// This is a bit ad-hoc: we only add buffers containing all new channels
-				for (const FMeshBuffer& SecondBuf : pSecond->GetVertexBuffers().Buffers)
-				{
-					bool bSomeChannel = false;
-					bool bAllNewChannels = true;
-					for (const FMeshBufferChannel& SecondChan : SecondBuf.Channels)
-					{
-						// Skip system buffers
-						if (SecondChan.Semantic == EMeshBufferSemantic::VertexIndex
-							||
-							SecondChan.Semantic == EMeshBufferSemantic::LayoutBlock)
-						{
-							continue;
-						}
-
-						bSomeChannel = true;
-
-						int32 FoundBuffer = -1;
-						int32 FoundChannel = -1;
-						pFirst->GetVertexBuffers().FindChannel(SecondChan.Semantic, SecondChan.SemanticIndex, &FoundBuffer, &FoundChannel);
-						if (FoundBuffer >= 0)
-						{
-							// There's at least one channel that already exists in the first mesh. Don't add the buffer.
-							bAllNewChannels = false;
-							continue;
-						}
-
-						// If there are additional UV channels try to add them
-						if (!bAllNewChannels && SecondChan.Semantic == EMeshBufferSemantic::TexCoords
-							&&
-							SecondChan.SemanticIndex > 0)
-						{
-							// Add additional UV channels if the previous one is found.
-							FMeshBufferSet& ResultVertexBuffers = Result->GetVertexBuffers();
-							ResultVertexBuffers.FindChannel(EMeshBufferSemantic::TexCoords, SecondChan.SemanticIndex - 1, &FoundBuffer, &FoundChannel);
-
-							if (FoundBuffer >= 0)
-							{
-								FMeshBuffer& Buffer = ResultVertexBuffers.Buffers[FoundBuffer];
-								Buffer.Channels.Insert(SecondChan, FoundChannel + 1);
-
-								ResultVertexBuffers.UpdateOffsets(FoundBuffer);
-							}
-						}
-					}
-
-					if (bSomeChannel && bAllNewChannels)
-					{
-						int32 NewBufferIndex = Result->GetVertexBuffers().Buffers.Emplace();
-						Result->VertexBuffers.Buffers[NewBufferIndex].Channels = SecondBuf.Channels;
-						Result->VertexBuffers.Buffers[NewBufferIndex].ElementSize = SecondBuf.ElementSize;
-					}
-				}
-
-				// See if we need to enlarge the components of any of the result channels
-				int32 ResultBufferCount = Result->GetVertexBuffers().GetBufferCount();
-				for (int32 BufferIndex = 0; BufferIndex < FMath::Min(ResultBufferCount, FirstBufferCount); ++BufferIndex)
-				{
-					// Expand component counts in vertex channels of the format mesh
-					FMeshBuffer& result = Result->GetVertexBuffers().Buffers[BufferIndex];
-
-					bool bResetOffsets = false;
-					for (int32 c = 0; c < result.Channels.Num(); ++c)
-					{
-						int32 sb = -1;
-						int32 sc = -1;
-						pSecond->GetVertexBuffers().FindChannel
-						(
-							result.Channels[c].Semantic,
-							result.Channels[c].SemanticIndex,
-							&sb, &sc
-						);
-
-						if (sb >= 0)
-						{
-							const FMeshBuffer& second = pSecond->GetVertexBuffers().Buffers[sb];
-
-							if (second.Channels[sc].ComponentCount>result.Channels[c].ComponentCount)
-							{
-								result.Channels[c].ComponentCount = second.Channels[sc].ComponentCount;
-								bResetOffsets = true;
-							}
-						}
-					}
-
-					// Reset the channel offsets if necessary
-					if (bResetOffsets)
-					{
-						Result->GetVertexBuffers().UpdateOffsets(BufferIndex);
-					}
-				}
-
-
-				// Change the format of the bone indices buffer
-				if (bChangeBoneIndicesFormat)
-				{
-					// Iterate all vertex buffers and update the format
-					FMeshBufferSet& VertexBuffers = Result->GetVertexBuffers();
-					for (int32 VertexBufferIndex = 0; VertexBufferIndex < VertexBuffers.Buffers.Num(); ++VertexBufferIndex)
-					{
-						FMeshBuffer& result = VertexBuffers.Buffers[VertexBufferIndex];
-
-						const int32 ChannelsCount = VertexBuffers.GetBufferChannelCount(VertexBufferIndex);
-						for (int32 ChannelIndex = 0; ChannelIndex < ChannelsCount; ++ChannelIndex)
-						{
-							if (result.Channels[ChannelIndex].Semantic == EMeshBufferSemantic::BoneIndices)
-							{
-								result.Channels[ChannelIndex].Format = BoneIndexFormat;
-
-								// Reset offsets
-								int32 offset = 0;
-								for (int32 AuxChannelIndex = 0; AuxChannelIndex < ChannelsCount; ++AuxChannelIndex)
-								{
-									result.Channels[AuxChannelIndex].Offset = (uint8_t)offset;
-									offset += result.Channels[AuxChannelIndex].ComponentCount
-										*
-										GetMeshFormatData(result.Channels[AuxChannelIndex].Format).SizeInBytes;
-								}
-								result.ElementSize = offset;
-							}
-						}
-					}
-				}
-
-				// Manage vertex IDs
-				if (bNeedsExplicitIds)
-				{
-					// Make sure the result format is suitable for the explicit IDs
-					Result->MakeIdsExplicit();
-				}
-			}
-
-
-			// Step 2: Fill the result buffers
-			//-----------------------------------------------------------------------
-			if (pFirst->VertexBuffers.IsDescriptor() || pSecond->VertexBuffers.IsDescriptor())
-			{
-				EnumAddFlags(Result->VertexBuffers.Flags, EMeshBufferSetFlags::IsDescriptor);
-			}
-			Result->VertexBuffers.SetElementCount(FirstVertexCount + SecondVertexCount);
-			
-			if (!Result->VertexBuffers.IsDescriptor())
-			{
-				MUTABLE_CPUPROFILER_SCOPE(ResultFill);
-				
-				// We have the final result vertex buffers structure: allocate the memory for them.
-				int32 ResultBufferCount = Result->GetVertexBuffers().GetBufferCount();
-				for (int32 ResultBufferIndex = 0; ResultBufferIndex < ResultBufferCount; ++ResultBufferIndex)
-				{
-					FMeshBuffer& ResultBuffer = Result->VertexBuffers.Buffers[ResultBufferIndex];
-
-					// TODO: Don't assume the buffer order in First and Second matches Result, for more opportunities of fast-path
-					bool bFirstFastPath = (ResultBufferIndex < FirstBufferCount)
-						&&
-						pFirst->VertexBuffers.HasSameFormat(ResultBufferIndex, Result->VertexBuffers, ResultBufferIndex);
-
-					int32 FirstResultBufferSize = ResultBuffer.ElementSize * FirstVertexCount;
-					uint8* Destination = ResultBuffer.Data.GetData();
-					if (bFirstFastPath)
-					{
-						MUTABLE_CPUPROFILER_SCOPE(FirstFastPath);
-
-						const FMeshBuffer& FirstBuffer = pFirst->VertexBuffers.Buffers[ResultBufferIndex];
-						check(FirstResultBufferSize == FirstBuffer.Data.Num());
-						FMemory::Memcpy(Destination, FirstBuffer.Data.GetData(), FirstResultBufferSize);
-					}
-					else
-					{
-						MUTABLE_CPUPROFILER_SCOPE(FirstSlowPath);
-
-						int32 OffsetElements = 0;
-						MeshFormatBuffer(pFirst->VertexBuffers, ResultBuffer, OffsetElements, true, pFirst->MeshIDPrefix);
-					}
-
-
-					bool bSecondFastPath = (ResultBufferIndex < SecondBufferCount)
-						&&
-						pSecond->VertexBuffers.HasSameFormat(ResultBufferIndex, Result->VertexBuffers, ResultBufferIndex);
-
-					int32 SecondResultBufferSize = ResultBuffer.ElementSize * SecondVertexCount;
-					Destination = ResultBuffer.Data.GetData() + FirstResultBufferSize;
-					if (bSecondFastPath)
-					{
-						MUTABLE_CPUPROFILER_SCOPE(SecondFastPath);
-
-						const FMeshBuffer& SecondBuffer = pSecond->VertexBuffers.Buffers[ResultBufferIndex];
-						check(SecondResultBufferSize == SecondBuffer.Data.Num());
-						FMemory::Memcpy(Destination, SecondBuffer.Data.GetData(), SecondResultBufferSize);
-					}
-					else
-					{
-						MUTABLE_CPUPROFILER_SCOPE(SecondSlowPath);
-
-						int32 OffsetElements = FirstVertexCount;
-						MeshFormatBuffer(pSecond->VertexBuffers, ResultBuffer, OffsetElements, true, pSecond->MeshIDPrefix);
-					}
-				}
-
-				// TODO: This still needs to be optimized
-				if (bRemapBoneIndices)
-				{
-					MUTABLE_CPUPROFILER_SCOPE(RemapBones);
-
-					// We need to remap the bone indices of the second mesh vertices that we already copied to result
-					check(!RemappedBoneMapIndices.IsEmpty());
-
-					// Iterate all vertex buffers and update the format
-					FMeshBufferSet& VertexBuffers = Result->GetVertexBuffers();
-					for (int32 VertexBufferIndex = 0; VertexBufferIndex < VertexBuffers.Buffers.Num(); ++VertexBufferIndex)
-					{
-						FMeshBuffer& ResultBuffer = VertexBuffers.Buffers[VertexBufferIndex];
-
-						const int32 ElemSize = VertexBuffers.GetElementSize(VertexBufferIndex);
-						const int32 FirstSize = FirstVertexCount * ElemSize;
-
-						const int32 ChannelsCount = VertexBuffers.GetBufferChannelCount(VertexBufferIndex);
-						for (int32 ChannelIndex = 0; ChannelIndex < ChannelsCount; ++ChannelIndex)
-						{
-							if (ResultBuffer.Channels[ChannelIndex].Semantic != EMeshBufferSemantic::BoneIndices)
-							{
-								continue;
-							}
-
-							int32 ResultOffset = FirstSize + ResultBuffer.Channels[ChannelIndex].Offset;
-
-							const int32 NumComponents = ResultBuffer.Channels[ChannelIndex].ComponentCount;
-
-							// Bone indices may need remapping
-							for (int32 VertexIndex = 0; VertexIndex < pSecond->GetVertexCount(); ++VertexIndex)
-							{
-								switch (BoneIndexFormat)
-								{
-								case EMeshBufferFormat::Int8:
-								case EMeshBufferFormat::UInt8:
-								{
-									uint8* pD = reinterpret_cast<uint8*>(&ResultBuffer.Data[ResultOffset]);
-
-									for (int32 ComponentIndex = 0; ComponentIndex < NumComponents; ++ComponentIndex)
-									{
-										uint8 BoneMapIndex = pD[ComponentIndex];
-
-										// be defensive
-										if (BoneMapIndex < NumBonesInSecondBoneMap)
-										{
-											pD[ComponentIndex] = (uint8)RemappedBoneMapIndices[BoneMapIndex];
-										}
-										else
-										{
-											pD[ComponentIndex] = 0;
-										}
-									}
-
-									ResultOffset += ElemSize;
-									break;
-								}
-
-								case EMeshBufferFormat::Int16:
-								case EMeshBufferFormat::UInt16:
-								{
-									uint16* pD = reinterpret_cast<uint16*>(&ResultBuffer.Data[ResultOffset]);
-
-									for (int32 ComponentIndex = 0; ComponentIndex < NumComponents; ++ComponentIndex)
-									{
-										uint16 BoneMapIndex = pD[ComponentIndex];
-
-										// be defensive
-										if (BoneMapIndex < NumBonesInSecondBoneMap)
-										{
-											pD[ComponentIndex] = (uint16)RemappedBoneMapIndices[BoneMapIndex];
-										}
-										else
-										{
-											pD[ComponentIndex] = 0;
-										}
-									}
-
-									ResultOffset += ElemSize;
-									break;
-								}
-
-								case EMeshBufferFormat::Int32:
-								case EMeshBufferFormat::UInt32:
-								{
-									// Unreal does not support 32 bit bone indices
-									checkf(false, TEXT("Format not supported."));
-									break;
-								}
-
-								default:
-									checkf(false, TEXT("Format not supported."));
-								}
-							}
-						}
-					}
-				}
-			}
-		}
-
-
-		// Tags
-		Result->Tags = pFirst->Tags;
-
-		for (const FString& SecondTag : pSecond->Tags)
-		{
-			Result->Tags.AddUnique(SecondTag);
-		}
-
-
-		// Streamed Resources
-		Result->StreamedResources = pFirst->StreamedResources;
-
-		const int32 NumStreamedResources = pSecond->StreamedResources.Num();
-		for (int32 Index = 0; Index < NumStreamedResources; ++Index)
-		{
-			Result->StreamedResources.AddUnique(pSecond->StreamedResources[Index]);
-		}
-
-		// Skeletal Meshes
-		Result->SkeletalMeshes = pFirst->SkeletalMeshes;
-
-		const int32 NumSkeletalMeshes = pSecond->SkeletalMeshes.Num();
-		for (int32 Index = 0; Index < NumSkeletalMeshes; ++Index)
-		{
-			Result->SkeletalMeshes.AddUnique(pSecond->SkeletalMeshes[Index]);
-		}
-
-		// Morphs
-		if (pFirst->HasMorphs() && !pSecond->HasMorphs())
-		{
-			Result->Morph = pFirst->Morph;
-			Result->MorphDataBuffer = pFirst->MorphDataBuffer;
-		}
-		else if (!pFirst->HasMorphs() && pSecond->HasMorphs())
-		{
-			Result->Morph = pSecond->Morph;
-			Result->MorphDataBuffer = pSecond->MorphDataBuffer;
-
-			FixMorphIndices(*Result, pFirst->GetVertexCount());
-		}
-		else if (pFirst->HasMorphs() && pSecond->HasMorphs())
-		{
-			MergeMorphs(*Result, *pFirst, *pSecond);	
-		}
-
-		// Cloth
-		if (bMergeSurfaces)
-		{
-			if (pFirst->ClothSections[0].IsValid() && !pSecond->ClothSections[0].IsValid())
-			{
-				CopyClothFirst(*Result, *pFirst, *pSecond);	
-			}
-			else if (!pFirst->ClothSections[0].IsValid() && pSecond->ClothSections[0].IsValid())
-			{
-				CopyClothSecond(*Result, *pFirst, *pSecond);
-			}
-			else if (pFirst->ClothSections[0].IsValid() && pSecond->ClothSections[0].IsValid())
-			{
-				UE_LOG(LogMutableCore, Error, TEXT("Merging two mesh sections with cloth is not supported."));
-				CopyClothFirst(*Result, *pFirst, *pSecond);
-			}
-		}
-		else
-		{
-			for (int32 SectionIndexFirst = 0; SectionIndexFirst < pFirst->ClothSections.Num(); ++SectionIndexFirst)
-			{
-				Result->ClothSections[SectionIndexFirst] = pFirst->ClothSections[SectionIndexFirst];
-			}
-
-			for (int32 SectionIndexSecond = 0; SectionIndexSecond < pSecond->ClothSections.Num(); ++SectionIndexSecond)
-			{
-				Result->ClothSections[pFirst->Surfaces.Num() + SectionIndexSecond] = pSecond->ClothSections[SectionIndexSecond];
-			}
-		}
-	}
+	/** Merge two skeletons, OutRemappedOtherBoneIndices containts the map from bone in Other to bone in Base */
+	UE_API void MergeSkeletons(FSkeleton& Base, const FSkeleton& Other, TArray<int32>& OutRemappedOtherBoneIndices);
 }
+
+#undef UE_API
diff --git a/Engine/Plugins/Mutable/Source/MutableRuntime/Internal/MuR/Operations.h b/Engine/Plugins/Mutable/Source/MutableRuntime/Internal/MuR/Operations.h
--- a/Engine/Plugins/Mutable/Source/MutableRuntime/Internal/MuR/Operations.h
+++ b/Engine/Plugins/Mutable/Source/MutableRuntime/Internal/MuR/Operations.h
@@ -1092,18 +1092,20 @@ namespace UE::Mutable::Private
         {
             enum class EnumFlags : uint16
             {
-                None			= 0,
-                HasTags			= 1 << 0,
-				IsTagList		= 1 << 1,
-                HasResources	= 1 << 2,
-                IsResourceList	= 1 << 3,
+				None = 0,
+				HasTags = 1 << 0,
+				IsTagList = 1 << 1,
+				HasResources = 1 << 2,
+				IsResourceList = 1 << 3,
 				HasAdditionalPhysicsAsset = 1 << 6,
-                IsAdditionalPhysicsAssetList = 1 << 7
+				IsAdditionalPhysicsAssetList = 1 << 7,
+				HasRealTimeMorphNames = 1 << 8,
+				IsRealTimeMorphNamesList = 1 << 9,	
             };
 
-            union TagAddressOrUInt32ListAddress
+            union StringAddressOrUInt32ListAddress
             {
-                ADDRESS TagAddress;
+                ADDRESS StringAddress;
 				LIST_UINT32_ADDRESS ListAddress;
             };
 
@@ -1123,12 +1125,13 @@ namespace UE::Mutable::Private
 			int16 BonePosePriority;
             ADDRESS Source;
             PASSTHROUGH_ID SkeletonId;
-            TagAddressOrUInt32ListAddress Tags;
+            StringAddressOrUInt32ListAddress Tags;
             ResourceIdOrUInt64ListAddress ResourceIds;
 			PASSTHROUGH_ID PhysicsAssetId; // Main Physics Asset of the mesh
 			PassthroughIDOrUInt32ListAddress AdditionalPhysicsAssetsIds; // Physics Assets used by AnimBPs
+            StringAddressOrUInt32ListAddress RealTimeMorphNames;
         };
-        static_assert(sizeof(MeshAddMetadataArgs) == sizeof(uint64)*4, "MeshAddMetadataArgs has an unexpected size.");
+        static_assert(sizeof(MeshAddMetadataArgs) == sizeof(uint64)*5, "MeshAddMetadataArgs has an unexpected size.");
 
 
         struct InstanceAddArgs
diff --git a/Engine/Plugins/Mutable/Source/MutableRuntime/Private/MuR/CodeRunner.cpp b/Engine/Plugins/Mutable/Source/MutableRuntime/Private/MuR/CodeRunner.cpp
--- a/Engine/Plugins/Mutable/Source/MutableRuntime/Private/MuR/CodeRunner.cpp
+++ b/Engine/Plugins/Mutable/Source/MutableRuntime/Private/MuR/CodeRunner.cpp
@@ -2860,9 +2860,9 @@ namespace UE::Mutable::Private
 						}
 					}
 					else if(EnumHasAnyFlags(OpArgs.Flags, OpEnumFlags::HasTags) &&
-						Program.ConstantStrings.IsValidIndex(OpArgs.Tags.TagAddress))
+						Program.ConstantStrings.IsValidIndex(OpArgs.Tags.StringAddress))
 					{
-						Result->Tags.Add(Program.ConstantStrings[OpArgs.Tags.TagAddress]);
+						Result->Tags.Add(Program.ConstantStrings[OpArgs.Tags.StringAddress]);
 					}
 					else
 					{
@@ -2921,6 +2921,46 @@ namespace UE::Mutable::Private
 					{
 						check(!EnumHasAnyFlags(OpArgs.Flags, OpEnumFlags::IsAdditionalPhysicsAssetList));
 					}
+
+					check(Result->Morph.Names.Num() == Result->Morph.UsageFlagsPerMorph.Num());
+					
+					if (EnumHasAnyFlags(OpArgs.Flags, OpEnumFlags::IsRealTimeMorphNamesList) &&
+						Program.ConstantUInt32Lists.IsValidIndex(OpArgs.RealTimeMorphNames.ListAddress))
+					{
+						const TArray<uint32>& ConstantRealTimeMorphList = Program.ConstantUInt32Lists[OpArgs.RealTimeMorphNames.ListAddress];
+
+						for (uint32 MorphNameIndex : ConstantRealTimeMorphList)
+						{
+							if (ensureAlways(Program.ConstantStrings.IsValidIndex(MorphNameIndex)))
+							{
+								FName RealTimeMorphName = FName(Program.ConstantStrings[MorphNameIndex]);
+
+								int32 FoundIndex = Result->Morph.Names.IndexOfByKey(RealTimeMorphName);
+								if (FoundIndex != INDEX_NONE)
+								{
+									check(Result->Morph.UsageFlagsPerMorph.IsValidIndex(FoundIndex));
+									EnumAddFlags(Result->Morph.UsageFlagsPerMorph[FoundIndex], EMorphUsageFlags::RealTime);
+								}
+							}
+						}
+					}
+					else if (EnumHasAnyFlags(OpArgs.Flags, OpEnumFlags::HasRealTimeMorphNames) &&
+						Program.ConstantStrings.IsValidIndex(OpArgs.RealTimeMorphNames.StringAddress))
+					{
+						FName RealTimeMorphName = FName(Program.ConstantStrings[OpArgs.RealTimeMorphNames.StringAddress]);
+						int32 FoundIndex = Result->Morph.Names.IndexOfByKey(RealTimeMorphName);
+
+						if (FoundIndex != INDEX_NONE)
+						{
+							check(Result->Morph.UsageFlagsPerMorph.IsValidIndex(FoundIndex));
+							EnumAddFlags(Result->Morph.UsageFlagsPerMorph[FoundIndex], EMorphUsageFlags::RealTime);
+						}
+					}
+					else
+					{
+						check(!EnumHasAnyFlags(OpArgs.Flags, OpEnumFlags::HasRealTimeMorphNames | OpEnumFlags::IsRealTimeMorphNamesList));
+					}
+
 				}
 
 				StoreMesh(Item, Result);
diff --git a/Engine/Plugins/Mutable/Source/MutableRuntime/Private/MuR/Mesh.cpp b/Engine/Plugins/Mutable/Source/MutableRuntime/Private/MuR/Mesh.cpp
--- a/Engine/Plugins/Mutable/Source/MutableRuntime/Private/MuR/Mesh.cpp
+++ b/Engine/Plugins/Mutable/Source/MutableRuntime/Private/MuR/Mesh.cpp
@@ -179,9 +179,13 @@ void FMesh::CopyFrom(const FMesh& From, EMeshCopyFlags InFlags)
 		AdditionalPhysicsBodies = From.AdditionalPhysicsBodies;
 	}
 
-	if (EnumHasAnyFlags(InFlags, EMeshCopyFlags::WithMorphData))
+	if (EnumHasAnyFlags(InFlags, EMeshCopyFlags::WithMorphMetaData))
 	{
 		Morph = From.Morph;
+	}
+
+	if (EnumHasAnyFlags(InFlags, EMeshCopyFlags::WithMorphData))
+	{
 		MorphDataBuffer = From.MorphDataBuffer;	
 	}
 	
@@ -200,7 +204,7 @@ uint32 FMesh::GetId() const
 
 bool FMesh::HasMorphs() const
 {
-	return Morph.NumTotalBatches != 0;
+	return MorphDataBuffer.Num() && Morph.NumTotalBatches != 0;
 }
 
 
@@ -908,6 +912,9 @@ void FMesh::ResetBufferIndices()
 MUTABLE_IMPLEMENT_POD_SERIALISABLE(FSurfaceSubMesh);
 MUTABLE_IMPLEMENT_POD_VECTOR_SERIALISABLE(FSurfaceSubMesh);
 
+MUTABLE_IMPLEMENT_POD_SERIALISABLE(EMorphUsageFlags);
+MUTABLE_IMPLEMENT_POD_VECTOR_SERIALISABLE(EMorphUsageFlags);
+
 void FMeshSurface::Serialise(FOutputArchive& Arch) const
 {
 	Arch << SubMeshes;
@@ -928,6 +935,40 @@ void FMeshSurface::Unserialise(FInputArchive& Arch)
 }
 
 
+void FMeshMorph::Serialise(FOutputArchive& Arch) const
+{
+	Arch << Names;
+	Arch << MaximumValuePerMorph;
+	Arch << MinimumValuePerMorph;
+	Arch << BatchStartOffsetPerMorph;
+	Arch << BatchesPerMorph;
+
+	Arch << SurfacesInUsePerMorph;
+	Arch << UsageFlagsPerMorph;
+
+	Arch << NumTotalBatches;
+	Arch << PositionPrecision;
+	Arch << TangentZPrecision;
+}
+
+
+void FMeshMorph::Unserialise(FInputArchive& Arch)
+{
+	Arch >> Names;
+	Arch >> MaximumValuePerMorph;
+	Arch >> MinimumValuePerMorph;
+	Arch >> BatchStartOffsetPerMorph;
+	Arch >> BatchesPerMorph;
+
+	Arch >> SurfacesInUsePerMorph;
+	Arch >> UsageFlagsPerMorph;
+
+	Arch >> NumTotalBatches;
+	Arch >> PositionPrecision;
+	Arch >> TangentZPrecision;
+}
+
+
 void FMesh::FBonePose::Serialise(FOutputArchive& Arch) const
 {
 	Arch << BoneId;
@@ -967,6 +1008,9 @@ void FMesh::Serialise(FOutputArchive& Arch) const
 	Arch << MeshIDPrefix;
 	
 	Arch << ClothSections;
+
+	Arch << Morph;
+	Arch << MorphDataBuffer;
 }
 
 
@@ -993,6 +1037,9 @@ void FMesh::Unserialise(FInputArchive& Arch)
 	Arch >> MeshIDPrefix;
 	
 	Arch >> ClothSections;
+
+	Arch >> Morph;
+	Arch >> MorphDataBuffer;
 }
 
 
@@ -1037,6 +1084,9 @@ bool FMesh::IsSimilar(const FMesh& Other) const
 	bEqual = bEqual && (Surfaces == Other.Surfaces);
 	bEqual = bEqual && (Tags == Other.Tags);
 
+	bEqual = bEqual && (Morph == Other.Morph);
+	bEqual = bEqual && (MorphDataBuffer == Other.MorphDataBuffer);
+
 	// Special comparison for vertex buffers
 	if (bEqual)
 	{
diff --git a/Engine/Plugins/Mutable/Source/MutableRuntime/Private/MuR/OpMeshMerge.cpp b/Engine/Plugins/Mutable/Source/MutableRuntime/Private/MuR/OpMeshMerge.cpp
new file mode 100644
--- /dev/null
+++ b/Engine/Plugins/Mutable/Source/MutableRuntime/Private/MuR/OpMeshMerge.cpp
@@ -0,0 +1,1715 @@
+// Copyright Epic Games, Inc. All Rights Reserved.
+
+#include "MuR/OpMeshMerge.h"
+
+#include "MuR/MeshPrivate.h"
+#include "MuR/Platform.h"
+#include "MuR/MutableMath.h"
+#include "MuR/OpMeshFormat.h"
+#include "MuR/MutableTrace.h"
+#include "MuR/MutableRuntimeModule.h"
+
+#include "Engine/SkeletalMesh.h"
+#include "Rendering/MorphTargetVertexCodec.h"
+
+#include "UObject/StrongObjectPtr.h"
+
+#include "Algo/Count.h"
+
+namespace UE::Mutable::Private
+{
+namespace OpMeshMergeMorphsInternal
+{
+	uint32 ComputeMorphAllocationRequirementsInDwords(TConstArrayView<uint32> MorphHeaderData)
+	{
+		using namespace UE::MorphTargetVertexCodec;
+		
+		const int32 NumBatches = MorphHeaderData.Num() / NumBatchHeaderDwords;
+
+		if (NumBatches <= 0)
+		{
+			return 0;
+		}
+
+		// First Dword in the header data correspond to the Offset in Bytes. See MorphTargetVertexCodec::WriteHeader()
+		const uint32 DataOffsetFirstBatchInDwords = MorphHeaderData[0] / sizeof(uint32);
+		const uint32 DataOffsetLastBatchInDwords  = MorphHeaderData[(NumBatches - 1) * NumBatchHeaderDwords] / sizeof(uint32);
+
+		// We need the last header to compute its size, this could be approximated by an upper bound.
+		TConstArrayView<uint32> LastHeaderData(
+				&MorphHeaderData[(NumBatches - 1) * NumBatchHeaderDwords], 
+				MorphTargetVertexCodec::NumBatchHeaderDwords);
+			
+		FDeltaBatchHeader LastBatchHeader;
+		UE::MorphTargetVertexCodec::ReadHeader(LastBatchHeader, LastHeaderData);
+
+		check(LastBatchHeader.DataOffset == DataOffsetLastBatchInDwords * sizeof(uint32));
+
+		return CalculateBatchDwords(LastBatchHeader) + DataOffsetLastBatchInDwords - DataOffsetFirstBatchInDwords + MorphHeaderData.Num();
+	}
+
+	void FixMorphIndicesSurfaceMetadataAndApplyFilter(FMesh& Result, int32 VertexIndexOffset, int32 SurfacesOffset, EMorphUsageFlags UsageFilter)
+	{
+        MUTABLE_CPUPROFILER_SCOPE(CompressedMorphs_Merge_FixMorphIndices)
+
+		using namespace UE::MorphTargetVertexCodec;
+
+		if (SurfacesOffset != 0)
+		{
+			for (TArrayView<int32> SurfacesInUse : Result.Morph.SurfacesInUsePerMorph)
+			{
+				for (int32& SurfaceIndex : SurfacesInUse)
+				{
+					SurfaceIndex += SurfacesOffset;
+				}
+			}
+		}
+	
+		int32 NumMorphsToRemove = Algo::CountIf(Result.Morph.UsageFlagsPerMorph, 
+				[UsageFilter](EMorphUsageFlags Flags) { return !EnumHasAnyFlags(Flags, UsageFilter); });
+
+		int32 NumMorphs = Result.Morph.Names.Num();
+
+		if (NumMorphsToRemove != 0)
+		{
+			// Remove batches that don't pass the filter keeping original offsets for the next pass.
+			Result.Morph.NumTotalBatches = 0;
+			uint32 AccumulatedDataOffsetInDwords = 0;
+			for (int32 MorphIndex = 0; MorphIndex < NumMorphs; ++MorphIndex)
+			{
+				if (EnumHasAnyFlags(Result.Morph.UsageFlagsPerMorph[MorphIndex], UsageFilter))
+				{
+					int32 NumBatches = Result.Morph.BatchesPerMorph[MorphIndex];
+					
+					FMemory::Memmove(
+							Result.MorphDataBuffer.GetData() + AccumulatedDataOffsetInDwords,
+							Result.MorphDataBuffer.GetData() + Result.Morph.BatchStartOffsetPerMorph[MorphIndex] * NumBatchHeaderDwords,
+							NumBatches * NumBatchHeaderDwords * sizeof(uint32));
+
+					Result.Morph.BatchStartOffsetPerMorph[MorphIndex] = Result.Morph.NumTotalBatches;
+					Result.Morph.NumTotalBatches += NumBatches;
+					
+					AccumulatedDataOffsetInDwords += NumBatches * NumBatchHeaderDwords;
+				}
+				else
+				{
+					Result.Morph.BatchesPerMorph[MorphIndex] = 0;
+					Result.Morph.BatchStartOffsetPerMorph[MorphIndex] = Result.Morph.NumTotalBatches;
+				}
+				
+			}
+
+			// Move the data of the remaining batches and patch the offsets.
+			for (int32 MorphIndex = 0; MorphIndex < NumMorphs; ++MorphIndex)
+			{
+				int32 NumBatches = Result.Morph.BatchesPerMorph[MorphIndex];
+
+				if (NumBatches == 0)
+				{
+					continue;
+				}
+
+				int32 BatchStartOffset = Result.Morph.BatchStartOffsetPerMorph[MorphIndex];
+
+				int32 MorphDeltasDataBeginInBytes = Result.MorphDataBuffer[BatchStartOffset * NumBatchHeaderDwords];
+				int32 OldToNewDeltaOffsetInBytes = MorphDeltasDataBeginInBytes - (AccumulatedDataOffsetInDwords * sizeof(uint32));
+
+				check(OldToNewDeltaOffsetInBytes >= 0);
+
+				for (int32 BatchIndex = 0; BatchIndex < NumBatches; ++BatchIndex)
+				{
+					Result.MorphDataBuffer[(BatchStartOffset + BatchIndex) * NumBatchHeaderDwords] -= OldToNewDeltaOffsetInBytes; 
+				}
+				
+				int32 LastBatchDataOffsetInBytes = Result.MorphDataBuffer[(BatchStartOffset + NumBatches - 1) * NumBatchHeaderDwords];
+
+				TArrayView<uint32> LastHeaderDataView = TArrayView<uint32>(
+						&Result.MorphDataBuffer[(BatchStartOffset + NumBatches - 1) * NumBatchHeaderDwords],
+						NumBatchHeaderDwords);
+
+				FDeltaBatchHeader LastBatchHeader;
+				UE::MorphTargetVertexCodec::ReadHeader(LastBatchHeader, LastHeaderDataView);
+
+				int32 MorphDeltaDataSizeInBytes = 
+					(LastBatchDataOffsetInBytes + CalculateBatchDwords(LastBatchHeader) * sizeof(uint32)) -
+					(MorphDeltasDataBeginInBytes - OldToNewDeltaOffsetInBytes);
+
+				FMemory::Memmove(
+						Result.MorphDataBuffer.GetData() + AccumulatedDataOffsetInDwords,
+						Result.MorphDataBuffer.GetData() + (MorphDeltasDataBeginInBytes / sizeof(uint32)),
+						MorphDeltaDataSizeInBytes);	
+
+				AccumulatedDataOffsetInDwords += MorphDeltaDataSizeInBytes / sizeof(uint32);
+			}
+		
+			// Free unused data.
+			Result.MorphDataBuffer.SetNum(AccumulatedDataOffsetInDwords, EAllowShrinking::Yes);
+		}
+
+		if (VertexIndexOffset != 0)
+		{
+			for (int32 MorphIndex = 0; MorphIndex < NumMorphs; ++MorphIndex)
+			{
+				TArrayView<uint32> MorphHeadersDataView = TArrayView<uint32>(
+						Result.MorphDataBuffer.GetData() + Result.Morph.BatchStartOffsetPerMorph[MorphIndex]*NumBatchHeaderDwords,
+						Result.Morph.BatchesPerMorph[MorphIndex]*NumBatchHeaderDwords);
+
+				int32 NumBatches = Result.Morph.BatchesPerMorph[MorphIndex];
+				for (int32 BatchIndex = 0; BatchIndex < NumBatches; ++BatchIndex)
+				{
+					TArrayView<uint32> HeaderDataView = TArrayView<uint32>(
+							MorphHeadersDataView.GetData() + BatchIndex*NumBatchHeaderDwords,
+							NumBatchHeaderDwords);
+
+					FDeltaBatchHeader BatchHeader;
+					UE::MorphTargetVertexCodec::ReadHeader(BatchHeader, HeaderDataView);
+
+					BatchHeader.IndexMin += VertexIndexOffset;
+
+					UE::MorphTargetVertexCodec::WriteHeader(BatchHeader, HeaderDataView);
+				}
+			}
+		}
+
+	}
+
+	void MergeMorphs(FMesh& Result, const FMesh& MeshA, const FMesh& MeshB, bool bMergeSurfaces, EMorphUsageFlags UsageFilter)
+	{
+        MUTABLE_CPUPROFILER_SCOPE(CompressedMorphs_Merge_Merge)
+
+		using namespace UE::MorphTargetVertexCodec;
+		
+		// Compute memory requirements.	
+		Result.Morph.Names.Reserve(MeshA.Morph.Names.Num() + MeshB.Morph.Names.Num());
+		Result.Morph.Names.Append(MeshA.Morph.Names);
+
+		TArray<int32, TInlineAllocator<128>> ResultToBMorphNameMap; 
+		ResultToBMorphNameMap.Init(INDEX_NONE, MeshB.Morph.Names.Num() + MeshA.Morph.Names.Num());
+
+		// MeshA morphs map is the identity, compute only the MorphB morphs map.
+		for (int32 Index = 0; Index < MeshB.Morph.Names.Num(); ++Index)
+		{
+			int32 MappedIndex = Result.Morph.Names.AddUnique(MeshB.Morph.Names[Index]);
+			ResultToBMorphNameMap[MappedIndex] = Index;
+		}
+
+		const int32 MeshANumMorphs = MeshA.Morph.Names.Num();
+		const int32 ResultNumMorphs = Result.Morph.Names.Num();
+
+		Result.Morph.MinimumValuePerMorph.SetNumUninitialized(ResultNumMorphs);
+		Result.Morph.MaximumValuePerMorph.SetNumUninitialized(ResultNumMorphs);
+		Result.Morph.BatchesPerMorph.SetNumZeroed(ResultNumMorphs);
+		Result.Morph.BatchStartOffsetPerMorph.SetNum(ResultNumMorphs);
+		Result.Morph.SurfacesInUsePerMorph.SetNum(ResultNumMorphs);
+		Result.Morph.UsageFlagsPerMorph.Init(EMorphUsageFlags::None, ResultNumMorphs);
+
+		// Surface in use and usages.
+		{
+			int32 MeshBSurfaceOffset = bMergeSurfaces ? 0 : MeshA.Surfaces.Num();
+			int32 MorphIndex = 0;
+			for (; MorphIndex < MeshANumMorphs; ++MorphIndex)
+			{
+				if (EnumHasAnyFlags(MeshA.Morph.UsageFlagsPerMorph[MorphIndex], UsageFilter))
+				{
+					Result.Morph.SurfacesInUsePerMorph[MorphIndex] = MeshA.Morph.SurfacesInUsePerMorph[MorphIndex];
+					EnumAddFlags(Result.Morph.UsageFlagsPerMorph[MorphIndex], MeshA.Morph.UsageFlagsPerMorph[MorphIndex]);
+				}
+
+				int32 MorphIndexInB = ResultToBMorphNameMap[MorphIndex];
+				if (MorphIndexInB != INDEX_NONE && EnumHasAnyFlags(MeshB.Morph.UsageFlagsPerMorph[MorphIndexInB], UsageFilter))
+				{
+					EnumAddFlags(Result.Morph.UsageFlagsPerMorph[MorphIndex], MeshB.Morph.UsageFlagsPerMorph[MorphIndexInB]);
+
+					for (int32 SurfaceInUse : MeshB.Morph.SurfacesInUsePerMorph[MorphIndexInB])
+					{
+						SurfaceInUse += MeshBSurfaceOffset;
+						Result.Morph.SurfacesInUsePerMorph[MorphIndex].AddUnique(SurfaceInUse);
+					}
+				}
+			}
+
+			for (; MorphIndex < ResultNumMorphs; ++MorphIndex)
+			{
+				int32 MorphIndexInB = ResultToBMorphNameMap[MorphIndex];
+			
+				if (EnumHasAnyFlags(MeshB.Morph.UsageFlagsPerMorph[MorphIndexInB], UsageFilter))
+				{
+					check(MorphIndexInB != INDEX_NONE);
+
+					Result.Morph.SurfacesInUsePerMorph[MorphIndex] = MeshB.Morph.SurfacesInUsePerMorph[MorphIndexInB];
+					EnumAddFlags(Result.Morph.UsageFlagsPerMorph[MorphIndex], MeshB.Morph.UsageFlagsPerMorph[MorphIndexInB]);
+
+					if (MeshBSurfaceOffset != 0)
+					{
+						for (int32& SurfaceInUse : Result.Morph.SurfacesInUsePerMorph[MorphIndex])
+						{
+							SurfaceInUse += MeshBSurfaceOffset;
+						}
+					}
+				}
+			}
+		}
+
+		// Compute memory allocation requirements and NumBatches. 
+		{
+			uint32 MergedDataBufferSizeInDwords = 0; 
+			Result.Morph.NumTotalBatches =  0;
+			for (int32 MorphIndex = MeshA.Morph.Names.Num() - 1; MorphIndex >= 0; --MorphIndex)
+			{	
+				if (EnumHasAnyFlags(MeshA.Morph.UsageFlagsPerMorph[MorphIndex], UsageFilter))
+				{
+					int32 NumBatches = MeshA.Morph.BatchesPerMorph[MorphIndex];
+					Result.Morph.NumTotalBatches += NumBatches;
+
+					TConstArrayView<uint32> HeadersDataView(
+							MeshA.MorphDataBuffer.GetData() + MeshA.Morph.BatchStartOffsetPerMorph[MorphIndex]*NumBatchHeaderDwords,
+							NumBatchHeaderDwords*NumBatches);
+
+					MergedDataBufferSizeInDwords += ComputeMorphAllocationRequirementsInDwords(HeadersDataView);
+				}
+			}
+
+			for (int32 MorphIndex = MeshB.Morph.Names.Num() - 1; MorphIndex >= 0; --MorphIndex)
+			{
+				if (EnumHasAnyFlags(MeshB.Morph.UsageFlagsPerMorph[MorphIndex], UsageFilter))
+				{
+					int32 NumBatches = MeshB.Morph.BatchesPerMorph[MorphIndex];
+					Result.Morph.NumTotalBatches += NumBatches;
+
+					TConstArrayView<uint32> HeadersDataView(
+							MeshB.MorphDataBuffer.GetData() + MeshB.Morph.BatchStartOffsetPerMorph[MorphIndex]*NumBatchHeaderDwords,
+							NumBatchHeaderDwords*NumBatches);
+
+					MergedDataBufferSizeInDwords += ComputeMorphAllocationRequirementsInDwords(HeadersDataView);
+				}
+			}
+
+			Result.MorphDataBuffer.SetNumUninitialized(MergedDataBufferSizeInDwords);
+		}
+
+		Result.Morph.PositionPrecision = MeshA.Morph.PositionPrecision;
+		Result.Morph.TangentZPrecision = MeshA.Morph.TangentZPrecision;
+
+		// Merge headers.
+		int32 BatchAccumulatedOffsetInDwords = 0;
+		int32 MorphIndex = 0;
+		for (; MorphIndex < MeshANumMorphs; ++MorphIndex)
+		{
+			Result.Morph.MaximumValuePerMorph[MorphIndex] = MeshA.Morph.MaximumValuePerMorph[MorphIndex];
+			Result.Morph.MinimumValuePerMorph[MorphIndex] = MeshA.Morph.MinimumValuePerMorph[MorphIndex];
+			Result.Morph.BatchStartOffsetPerMorph[MorphIndex] = BatchAccumulatedOffsetInDwords / NumBatchHeaderDwords;	
+			Result.Morph.BatchesPerMorph[MorphIndex] = 0;
+
+			if (EnumHasAnyFlags(MeshA.Morph.UsageFlagsPerMorph[MorphIndex], UsageFilter))
+			{
+				Result.Morph.BatchesPerMorph[MorphIndex] = MeshA.Morph.BatchesPerMorph[MorphIndex];
+				if (Result.MorphDataBuffer.Num())
+				{
+					FMemory::Memcpy(
+							Result.MorphDataBuffer.GetData() + BatchAccumulatedOffsetInDwords, 
+							MeshA.MorphDataBuffer.GetData() + MeshA.Morph.BatchStartOffsetPerMorph[MorphIndex]*NumBatchHeaderDwords,
+							MeshA.Morph.BatchesPerMorph[MorphIndex] * NumBatchHeaderDwords*sizeof(uint32));
+				}
+			
+				BatchAccumulatedOffsetInDwords += MeshA.Morph.BatchesPerMorph[MorphIndex] * NumBatchHeaderDwords;
+			}
+
+			int32 MeshBMorphIndex = ResultToBMorphNameMap[MorphIndex];
+			if (MeshBMorphIndex >= 0)
+			{
+				if (EnumHasAnyFlags(MeshB.Morph.UsageFlagsPerMorph[MeshBMorphIndex], UsageFilter))
+				{
+					Result.Morph.BatchesPerMorph[MorphIndex] += MeshB.Morph.BatchesPerMorph[MeshBMorphIndex];
+
+					const FVector4f& SourceMaxValue = MeshB.Morph.MaximumValuePerMorph[MeshBMorphIndex];
+					const FVector4f& SourceMinValue = MeshB.Morph.MinimumValuePerMorph[MeshBMorphIndex];
+
+					Result.Morph.MaximumValuePerMorph[MorphIndex] = 
+							SourceMaxValue.ComponentMax(Result.Morph.MaximumValuePerMorph[MorphIndex]);
+					Result.Morph.MinimumValuePerMorph[MorphIndex] = 
+							SourceMinValue.ComponentMin(Result.Morph.MinimumValuePerMorph[MorphIndex]);
+
+					if (Result.MorphDataBuffer.Num())
+					{
+						FMemory::Memcpy(
+								Result.MorphDataBuffer.GetData() + BatchAccumulatedOffsetInDwords, 
+								MeshB.MorphDataBuffer.GetData() + MeshB.Morph.BatchStartOffsetPerMorph[MeshBMorphIndex]*NumBatchHeaderDwords,
+								MeshB.Morph.BatchesPerMorph[MeshBMorphIndex] * NumBatchHeaderDwords*sizeof(uint32));
+					}
+
+					BatchAccumulatedOffsetInDwords += MeshB.Morph.BatchesPerMorph[MeshBMorphIndex] * NumBatchHeaderDwords;
+				}
+			}
+		}
+
+		for (; MorphIndex < ResultNumMorphs; ++MorphIndex)
+		{
+			int32 MeshBMorphIndex = ResultToBMorphNameMap[MorphIndex];
+			
+			Result.Morph.BatchesPerMorph[MorphIndex] = 0;
+			Result.Morph.BatchStartOffsetPerMorph[MorphIndex] = 0;
+
+			if (EnumHasAnyFlags(MeshB.Morph.UsageFlagsPerMorph[MeshBMorphIndex], UsageFilter))
+			{
+				Result.Morph.BatchesPerMorph[MorphIndex] = MeshB.Morph.BatchesPerMorph[MeshBMorphIndex]; 
+				Result.Morph.BatchStartOffsetPerMorph[MorphIndex] = BatchAccumulatedOffsetInDwords / NumBatchHeaderDwords;
+
+				Result.Morph.MaximumValuePerMorph[MorphIndex] = MeshB.Morph.MaximumValuePerMorph[MeshBMorphIndex];
+				Result.Morph.MinimumValuePerMorph[MorphIndex] = MeshB.Morph.MinimumValuePerMorph[MeshBMorphIndex];
+
+				if (Result.MorphDataBuffer.Num())
+				{
+					FMemory::Memcpy(
+							Result.MorphDataBuffer.GetData() + BatchAccumulatedOffsetInDwords, 
+							MeshB.MorphDataBuffer.GetData() + MeshB.Morph.BatchStartOffsetPerMorph[MeshBMorphIndex]*NumBatchHeaderDwords,
+							MeshB.Morph.BatchesPerMorph[MeshBMorphIndex] * NumBatchHeaderDwords*sizeof(uint32));
+				}
+
+				BatchAccumulatedOffsetInDwords += MeshB.Morph.BatchesPerMorph[MeshBMorphIndex] * NumBatchHeaderDwords;
+			}
+		}
+
+		// Merge morph data and fix up headers.
+
+		if (Result.MorphDataBuffer.Num())
+		{
+			const uint32 MeshANumVertices = MeshA.GetVertexCount();
+
+			// Add Morphs in MeshA and common morphs with MeshB.
+			MorphIndex = 0;
+			for (; MorphIndex < MeshANumMorphs; ++MorphIndex)
+			{
+				TArrayView<uint32> MorphHeadersDataView = TArrayView<uint32>(
+						Result.MorphDataBuffer.GetData() + Result.Morph.BatchStartOffsetPerMorph[MorphIndex]*NumBatchHeaderDwords,
+						Result.Morph.BatchesPerMorph[MorphIndex] * NumBatchHeaderDwords);
+
+				int32 BatchIndex = 0;
+				
+				if (EnumHasAnyFlags(MeshA.Morph.UsageFlagsPerMorph[MorphIndex], UsageFilter))
+				{
+					int32 NumBatches = MeshA.Morph.BatchesPerMorph[MorphIndex];
+					for ( ; BatchIndex < NumBatches; ++BatchIndex)
+					{
+						TArrayView<uint32> HeaderDataView = TArrayView<uint32>(
+								MorphHeadersDataView.GetData() + BatchIndex*NumBatchHeaderDwords,
+								NumBatchHeaderDwords);
+
+						FDeltaBatchHeader BatchHeader;
+						ReadHeader(BatchHeader, HeaderDataView);
+
+						const int32 SourceDataOffset = BatchHeader.DataOffset;
+						BatchHeader.DataOffset = BatchAccumulatedOffsetInDwords * sizeof(uint32);
+
+						WriteHeader(BatchHeader, HeaderDataView);
+
+						const int32 NumBatchDwords = CalculateBatchDwords(BatchHeader);
+						FMemory::Memcpy(
+								Result.MorphDataBuffer.GetData() + (BatchHeader.DataOffset / sizeof(uint32)),
+								MeshA.MorphDataBuffer.GetData() + (SourceDataOffset / sizeof(uint32)),
+								NumBatchDwords * sizeof(uint32));
+
+						BatchAccumulatedOffsetInDwords += NumBatchDwords;
+					}
+				}
+
+				int32 MeshBMorphIndex = ResultToBMorphNameMap[MorphIndex];
+				if (MeshBMorphIndex >= 0)
+				{
+					if (EnumHasAnyFlags(MeshB.Morph.UsageFlagsPerMorph[MeshBMorphIndex], UsageFilter))
+					{
+						int32 NumBatches = MeshB.Morph.BatchesPerMorph[MeshBMorphIndex];
+
+						for (int32 MeshBBatchIndex = 0; MeshBBatchIndex < NumBatches; ++MeshBBatchIndex, ++BatchIndex)
+						{
+							TArrayView<uint32> HeaderDataView = TArrayView<uint32>(
+									MorphHeadersDataView.GetData() + BatchIndex*NumBatchHeaderDwords,
+									NumBatchHeaderDwords);
+
+							FDeltaBatchHeader BatchHeader;
+							ReadHeader(BatchHeader, HeaderDataView);
+
+							const int32 SourceDataOffset = BatchHeader.DataOffset;
+							BatchHeader.DataOffset = BatchAccumulatedOffsetInDwords * sizeof(uint32);
+							BatchHeader.IndexMin += MeshANumVertices;
+
+							WriteHeader(BatchHeader, HeaderDataView);
+
+							const int32 NumBatchDwords = CalculateBatchDwords(BatchHeader);
+
+							check(Result.MorphDataBuffer.Num() >= BatchHeader.DataOffset / sizeof(uint32) + NumBatchDwords);
+							FMemory::Memcpy(
+									Result.MorphDataBuffer.GetData() + (BatchHeader.DataOffset / sizeof(uint32)),
+									MeshB.MorphDataBuffer.GetData() + (SourceDataOffset / sizeof(uint32)),
+									NumBatchDwords * sizeof(uint32));
+							
+							BatchAccumulatedOffsetInDwords += NumBatchDwords;
+						}
+					}
+				}
+			}
+
+			// Add Morphs in MeshB that are not present in MeshA.
+			for (; MorphIndex < ResultNumMorphs; ++MorphIndex)
+			{
+				TArrayView<uint32> MorphHeadersDataView = TArrayView<uint32>(
+						Result.MorphDataBuffer.GetData() + Result.Morph.BatchStartOffsetPerMorph[MorphIndex]*NumBatchHeaderDwords,
+						Result.Morph.BatchesPerMorph[MorphIndex]*NumBatchHeaderDwords);
+
+				int32 MeshBMorphIndex = ResultToBMorphNameMap[MorphIndex];
+				int32 BatchIndex = 0;
+				if (EnumHasAnyFlags(MeshB.Morph.UsageFlagsPerMorph[MeshBMorphIndex], UsageFilter))
+				{
+					int32 NumBatches = Result.Morph.BatchesPerMorph[MorphIndex];
+					for (; BatchIndex < NumBatches; ++BatchIndex)
+					{
+						TArrayView<uint32> HeaderDataView = TArrayView<uint32>(
+								MorphHeadersDataView.GetData() + BatchIndex*NumBatchHeaderDwords,
+								NumBatchHeaderDwords);
+
+						FDeltaBatchHeader BatchHeader;
+						ReadHeader(BatchHeader, HeaderDataView);
+
+						const int32 SourceDataOffset = BatchHeader.DataOffset;
+						BatchHeader.DataOffset = BatchAccumulatedOffsetInDwords * sizeof(uint32);
+						BatchHeader.IndexMin += MeshANumVertices;
+
+						WriteHeader(BatchHeader, HeaderDataView);
+
+						const int32 NumBatchDwords = CalculateBatchDwords(BatchHeader);
+					
+						check(Result.MorphDataBuffer.Num() >= BatchHeader.DataOffset / sizeof(uint32) + NumBatchDwords);
+						FMemory::Memcpy(
+								Result.MorphDataBuffer.GetData() + (BatchHeader.DataOffset / sizeof(uint32)),
+								MeshB.MorphDataBuffer.GetData() + (SourceDataOffset / sizeof(uint32)),
+								NumBatchDwords * sizeof(uint32));
+
+						BatchAccumulatedOffsetInDwords += NumBatchDwords;
+					}
+				}
+			}
+
+			check(BatchAccumulatedOffsetInDwords == Result.MorphDataBuffer.Num());
+		}
+	}
+} // namespace OpMeshMergeMorphsInternal
+
+namespace OpMeshMergeClothInternal
+{
+	/** Copy cloth from MeshA. MeshB only required to for data integrity. */
+	void CopyClothFirst(FMesh& Result, const FMesh& MeshA, const FMesh& MeshB)
+	{
+		Result.ClothSections[0].ClothingAsset = MeshA.ClothSections[0].ClothingAsset;
+			
+		const uint32 NumClothVerticesA = MeshA.ClothSections[0].Data.Num();
+		if (!NumClothVerticesA)
+		{
+			return;
+		}
+
+		const int32 NumInfluences = MeshA.ClothSections[0].Data.Num() / MeshA.VertexBuffers.GetElementCount();
+		const uint32 NumClothVerticesB = NumInfluences * MeshB.VertexBuffers.GetElementCount();
+		
+		Result.ClothSections[0].Data.Reserve(NumClothVerticesA + NumClothVerticesB);
+
+		Result.ClothSections[0].Data.Append(MeshA.ClothSections[0].Data);
+		
+		// Add MeshB vertices without simulation. See FMeshToMeshVertData::SourceMeshVertIndices.
+		for (uint32 VertexIndex = 0; VertexIndex < NumClothVerticesB; ++VertexIndex)
+		{
+			FMeshToMeshVertData& Data = Result.ClothSections[0].Data.AddZeroed_GetRef();
+			Data.SourceMeshVertIndices[3] = static_cast<uint16>(0xffff);
+		}
+	}
+	
+	
+	/** Copy cloth from MeshB. MeshA only required to for data integrity. */
+	void CopyClothSecond(FMesh& Result, const FMesh& MeshA, const FMesh& MeshB)
+	{
+		Result.ClothSections[0].ClothingAsset = MeshB.ClothSections[0].ClothingAsset;
+
+		const uint32 NumClothVerticesB = MeshB.ClothSections[0].Data.Num();
+		if (!NumClothVerticesB)
+		{
+			return;
+		}
+		
+		const int32 NumInfluences = MeshB.ClothSections[0].Data.Num() / MeshB.VertexBuffers.GetElementCount();
+		const uint32 NumClothVerticesA =  NumInfluences * MeshA.VertexBuffers.GetElementCount();
+
+		Result.ClothSections[0].Data.Reserve(NumClothVerticesA + NumClothVerticesB);
+		
+		// Add MeshA vertices without simulation. See FMeshToMeshVertData::SourceMeshVertIndices.
+		for (uint32 VertexIndex = 0; VertexIndex < NumClothVerticesA; ++VertexIndex)
+		{
+			FMeshToMeshVertData& Data = Result.ClothSections[0].Data.AddZeroed_GetRef();
+			Data.SourceMeshVertIndices[3] = static_cast<uint16>(0xffff);					
+		}
+		
+		Result.ClothSections[0].Data.Append(MeshB.ClothSections[0].Data);
+	}
+} // namespace OpMeshMergeClothInternal
+
+	void MergeSkeletons(FSkeleton& Base, const FSkeleton& Other, TArray<int32>& OutRemappedOtherBoneIndices)
+	{
+		const int32 NumBonesBaseBeforeMerge = Base.GetNumBones();
+		const int32 NumBonesOther = Other.GetNumBones();
+		OutRemappedOtherBoneIndices.SetNumUninitialized(NumBonesOther);
+
+		// Merge pSecond and build the remap array 
+		for (int32 Index = 0; Index < NumBonesOther; ++Index)
+		{
+			OutRemappedOtherBoneIndices[Index] = Base.FindOrAddBone(Other.BoneNames[Index], Other.BoneParents[Index]);
+		}
+
+		// Remap Parent bone indices to new Skeleton
+		for (int32 Index = NumBonesBaseBeforeMerge; Index < Base.BoneParents.Num(); ++Index)
+		{
+			int16 OtherBoneIndex = Base.BoneParents[Index];
+			if (OtherBoneIndex != INDEX_NONE)
+			{
+				Base.BoneParents[Index] = OutRemappedOtherBoneIndices[OtherBoneIndex];
+			}
+		}
+	}
+	
+	void MeshMerge(FMesh* Result, const FMesh* pFirst, const FMesh* pSecond, bool bMergeSurfaces, FMeshMergeScratchMeshes& ScratchMeshes)
+	{
+		MUTABLE_CPUPROFILER_SCOPE(MeshMerge);
+
+		// Should never happen, but fixes static analysis warnings.
+		if (!(pFirst && pSecond))
+		{
+			return;
+		}
+
+		// Indices
+		//-----------------
+		if (pFirst->GetIndexBuffers().GetBufferCount() > 0)
+		{
+			MUTABLE_CPUPROFILER_SCOPE(Indices);
+
+			const int32 FirstCount = pFirst->GetIndexBuffers().GetElementCount();
+			const int32 SecondCount = pSecond->GetIndexBuffers().GetElementCount();
+
+			if (pFirst->IndexBuffers.IsDescriptor() || pSecond->IndexBuffers.IsDescriptor())
+			{
+				EnumAddFlags(Result->IndexBuffers.Flags, EMeshBufferSetFlags::IsDescriptor);
+			}
+			
+			Result->GetIndexBuffers().SetElementCount(FirstCount + SecondCount);
+
+			check(pFirst->GetIndexBuffers().GetBufferCount() <= 1);
+			check(pSecond->GetIndexBuffers().GetBufferCount() <= 1);
+			Result->GetIndexBuffers().SetBufferCount(1);
+
+			FMeshBuffer& ResultIndexBuffer = Result->GetIndexBuffers().Buffers[0];
+
+			const FMeshBuffer& FirstIndexBuffer = pFirst->GetIndexBuffers().Buffers[0];
+			const FMeshBuffer& SecondIndexBuffer = pSecond->GetIndexBuffers().Buffers[0];
+
+			// Avoid unused variable warnings
+			(void)FirstIndexBuffer;
+			(void)SecondIndexBuffer;
+
+			// This will be changed below if need to change the format of the index buffers.
+			EMeshBufferFormat IndexBufferFormat = EMeshBufferFormat::None;
+
+			if (FirstCount && SecondCount)
+			{
+				check(!FirstIndexBuffer.Channels.IsEmpty());
+
+				// We need to know the total number of vertices in case we need to adjust the index buffer format.
+				const uint64 TotalVertexCount = pFirst->GetVertexBuffers().GetElementCount() + pSecond->GetVertexBuffers().GetElementCount();
+				const uint64 MaxValueBits = GetMeshFormatData(pFirst->GetIndexBuffers().Buffers[0].Channels[0].Format).MaxValueBits;
+				const uint64 MaxSupportedVertices = uint64(1) << MaxValueBits;
+				
+				if (TotalVertexCount > MaxSupportedVertices)
+				{
+					IndexBufferFormat = TotalVertexCount > MAX_uint16 ? EMeshBufferFormat::UInt32 : EMeshBufferFormat::UInt16;
+				}
+			}
+			
+			if (IndexBufferFormat != EMeshBufferFormat::None)
+			{
+				// We only support vertex indices in case of having to change the format.
+				check(FirstIndexBuffer.Channels.Num() == 1);
+
+				ResultIndexBuffer.Channels.SetNum(1);
+				ResultIndexBuffer.Channels[0].Semantic = EMeshBufferSemantic::VertexIndex;
+				ResultIndexBuffer.Channels[0].Format = IndexBufferFormat;
+				ResultIndexBuffer.Channels[0].ComponentCount = 1;
+				ResultIndexBuffer.Channels[0].SemanticIndex = 0;
+				ResultIndexBuffer.Channels[0].Offset = 0;
+				ResultIndexBuffer.ElementSize = GetMeshFormatData(IndexBufferFormat).SizeInBytes;
+			}
+			else if (FirstCount)
+			{
+				ResultIndexBuffer.Channels = FirstIndexBuffer.Channels;
+				ResultIndexBuffer.ElementSize = FirstIndexBuffer.ElementSize;
+			}
+			else if (SecondCount)
+			{
+				ResultIndexBuffer.Channels = SecondIndexBuffer.Channels;
+				ResultIndexBuffer.ElementSize = SecondIndexBuffer.ElementSize;
+			}
+
+
+			check(ResultIndexBuffer.Channels.Num() == 1);
+			check(ResultIndexBuffer.Channels[0].Semantic == EMeshBufferSemantic::VertexIndex);
+
+			if (!Result->IndexBuffers.IsDescriptor())
+			{
+
+				ResultIndexBuffer.Data.SetNum(ResultIndexBuffer.ElementSize * (FirstCount + SecondCount));
+			
+				if (!ResultIndexBuffer.Data.IsEmpty())
+				{
+					if (FirstCount)
+					{
+						if (IndexBufferFormat == EMeshBufferFormat::None
+							|| IndexBufferFormat == FirstIndexBuffer.Channels[0].Format)
+						{
+							FMemory::Memcpy(&ResultIndexBuffer.Data[0],
+								&FirstIndexBuffer.Data[0],
+								FirstIndexBuffer.ElementSize * FirstCount);
+						}
+						else
+						{
+							// Conversion required
+							const uint8_t* pSource = &FirstIndexBuffer.Data[0];
+							uint8_t* pDest = &ResultIndexBuffer.Data[0];
+							switch (IndexBufferFormat)
+							{
+							case EMeshBufferFormat::UInt32:
+							{
+								switch (FirstIndexBuffer.Channels[0].Format)
+								{
+								case EMeshBufferFormat::UInt16:
+								{
+									for (int32 v = 0; v < FirstCount; ++v)
+									{
+										*(uint32_t*)pDest = *(const uint16*)pSource;
+										pSource += FirstIndexBuffer.ElementSize;
+										pDest += ResultIndexBuffer.ElementSize;
+									}
+									break;
+								}
+
+								case EMeshBufferFormat::UInt8:
+								{
+									for (int32 v = 0; v < FirstCount; ++v)
+									{
+										*(uint32_t*)pDest = *(const uint8_t*)pSource;
+										pSource += FirstIndexBuffer.ElementSize;
+										pDest += ResultIndexBuffer.ElementSize;
+									}
+									break;
+								}
+
+								default:
+									checkf(false, TEXT("Format not supported."));
+									break;
+								}
+								break;
+							}
+
+							case EMeshBufferFormat::UInt16:
+							{
+								switch (FirstIndexBuffer.Channels[0].Format)
+								{
+
+								case EMeshBufferFormat::UInt8:
+								{
+									for (int32 v = 0; v < FirstCount; ++v)
+									{
+										*(uint16*)pDest = *(const uint8_t*)pSource;
+										pSource += FirstIndexBuffer.ElementSize;
+										pDest += ResultIndexBuffer.ElementSize;
+									}
+									break;
+								}
+
+								default:
+									checkf(false, TEXT("Format not supported."));
+									break;
+								}
+								break;
+							}
+
+							default:
+								checkf(false, TEXT("Format not supported."));
+								break;
+							}
+						}
+					}
+
+					if (SecondCount)
+					{
+						const uint8_t* pSource = &SecondIndexBuffer.Data[0];
+						uint8_t* pDest = &ResultIndexBuffer.Data[ResultIndexBuffer.ElementSize * FirstCount];
+
+						uint32_t firstVertexCount = pFirst->GetVertexBuffers().GetElementCount();
+
+						if (IndexBufferFormat == EMeshBufferFormat::None
+							|| IndexBufferFormat == SecondIndexBuffer.Channels[0].Format)
+						{
+							switch (SecondIndexBuffer.Channels[0].Format)
+							{
+							case EMeshBufferFormat::Int32:
+							case EMeshBufferFormat::UInt32:
+							case EMeshBufferFormat::NInt32:
+							case EMeshBufferFormat::NUInt32:
+							{
+								for (int32 v = 0; v < SecondCount; ++v)
+								{
+									*(uint32_t*)pDest = firstVertexCount + *(const uint32_t*)pSource;
+									pSource += SecondIndexBuffer.ElementSize;
+									pDest += ResultIndexBuffer.ElementSize;
+								}
+								break;
+							}
+
+							case EMeshBufferFormat::Int16:
+							case EMeshBufferFormat::UInt16:
+							case EMeshBufferFormat::NInt16:
+							case EMeshBufferFormat::NUInt16:
+							{
+								for (int32 v = 0; v < SecondCount; ++v)
+								{
+									*(uint16*)pDest = uint16(firstVertexCount) + *(const uint16*)pSource;
+									pSource += SecondIndexBuffer.ElementSize;
+									pDest += ResultIndexBuffer.ElementSize;
+								}
+								break;
+							}
+
+							case EMeshBufferFormat::Int8:
+							case EMeshBufferFormat::UInt8:
+							case EMeshBufferFormat::NInt8:
+							case EMeshBufferFormat::NUInt8:
+							{
+								for (int32 v = 0; v < SecondCount; ++v)
+								{
+									*(uint8_t*)pDest = uint8_t(firstVertexCount) + *(const uint8_t*)pSource;
+									pSource += SecondIndexBuffer.ElementSize;
+									pDest += ResultIndexBuffer.ElementSize;
+								}
+								break;
+							}
+
+							default:
+								checkf(false, TEXT("Format not supported."));
+								break;
+							}
+						}
+						else
+						{
+							// Format conversion required
+							switch (IndexBufferFormat)
+							{
+
+							case EMeshBufferFormat::UInt32:
+							{
+								switch (SecondIndexBuffer.Channels[0].Format)
+								{
+								case EMeshBufferFormat::Int16:
+								case EMeshBufferFormat::UInt16:
+								case EMeshBufferFormat::NInt16:
+								case EMeshBufferFormat::NUInt16:
+								{
+									for (int32 v = 0; v < SecondCount; ++v)
+									{
+										*(uint32_t*)pDest = uint32_t(firstVertexCount) + *(const uint16*)pSource;
+										pSource += SecondIndexBuffer.ElementSize;
+										pDest += ResultIndexBuffer.ElementSize;
+									}
+									break;
+								}
+
+								case EMeshBufferFormat::Int8:
+								case EMeshBufferFormat::UInt8:
+								case EMeshBufferFormat::NInt8:
+								case EMeshBufferFormat::NUInt8:
+								{
+									for (int32 v = 0; v < SecondCount; ++v)
+									{
+										*(uint32_t*)pDest = uint32_t(firstVertexCount) + *(const uint8_t*)pSource;
+										pSource += SecondIndexBuffer.ElementSize;
+										pDest += ResultIndexBuffer.ElementSize;
+									}
+									break;
+								}
+
+								default:
+									checkf(false, TEXT("Format not supported."));
+									break;
+								}
+
+								break;
+							}
+
+							case EMeshBufferFormat::UInt16:
+							{
+								switch (SecondIndexBuffer.Channels[0].Format)
+								{
+								case EMeshBufferFormat::Int8:
+								case EMeshBufferFormat::UInt8:
+								case EMeshBufferFormat::NInt8:
+								case EMeshBufferFormat::NUInt8:
+								{
+									for (int32 v = 0; v < SecondCount; ++v)
+									{
+										*(uint16*)pDest = uint16(firstVertexCount) + *(const uint8_t*)pSource;
+										pSource += SecondIndexBuffer.ElementSize;
+										pDest += ResultIndexBuffer.ElementSize;
+									}
+									break;
+								}
+
+								default:
+									checkf(false, TEXT("Format not supported."));
+									break;
+								}
+
+								break;
+							}
+
+							default:
+								checkf(false, TEXT("Format not supported."));
+								break;
+
+							}
+						}
+					}
+				}
+			}
+		}
+
+
+		// Layouts
+		//-----------------
+		{
+			MUTABLE_CPUPROFILER_SCOPE(Layouts);
+
+			int32 ResultLayoutCount = FMath::Max(pFirst->Layouts.Num(),pSecond->Layouts.Num());
+			Result->Layouts.SetNum(ResultLayoutCount);
+			for (int32 LayoutIndex = 0; LayoutIndex < ResultLayoutCount; ++LayoutIndex)
+			{
+				TManagedPtr<FLayout> pR;
+				
+				if (LayoutIndex < pFirst->Layouts.Num())
+				{
+					const FLayout* pF = pFirst->Layouts[LayoutIndex].Get();
+					pR = pF->Clone();
+				}
+
+				if (LayoutIndex < pSecond->Layouts.Num())
+				{
+					const FLayout* pS = pSecond->Layouts[LayoutIndex].Get();
+					if (!pR)
+					{
+						pR = pS->Clone();
+					}
+					else
+					{
+						pR->Blocks.Append(pS->Blocks);
+					}
+				}
+
+				Result->Layouts[LayoutIndex] = pR;
+			}
+		}
+
+
+		// Skeleton
+		//---------------------------
+
+		// Add SkeletonObjects
+		Result->SkeletonObjects = pFirst->SkeletonObjects;
+
+		for (const TPassthroughObjectPtr<USkeleton>& SkeletonObject : pSecond->SkeletonObjects)
+		{
+			Result->SkeletonObjects.AddUnique(SkeletonObject);
+		}
+
+		TManagedPtr<const FSkeleton> FirstSkeleton = pFirst->GetSkeleton();
+		TManagedPtr<const FSkeleton> SecondSkeleton = pSecond->GetSkeleton();
+
+		// Do they have the same skeleton?
+		bool bMergeSkeletons = FirstSkeleton != SecondSkeleton;
+
+		// Are they different skeletons but with the same data?
+		if (bMergeSkeletons && FirstSkeleton && SecondSkeleton)
+		{
+			bMergeSkeletons = !(*FirstSkeleton == *SecondSkeleton);
+		}
+
+		TArray<int32> RemappedSecondBoneIndices;
+
+		if (bMergeSkeletons)
+		{
+			MUTABLE_CPUPROFILER_SCOPE(MergeSkeleton);
+
+			TManagedPtr<FSkeleton> ResultSkeleton = FirstSkeleton ? FirstSkeleton->Clone() : MakeManaged<FSkeleton>();
+			Result->SetSkeleton(ResultSkeleton);
+
+			MergeSkeletons(*ResultSkeleton.Get(), *SecondSkeleton.Get(), RemappedSecondBoneIndices);
+		}
+		else
+		{
+			Result->SetSkeleton(FirstSkeleton);
+
+			const int32 NumBonesFirst = FirstSkeleton ? FirstSkeleton->GetNumBones() : 0;
+			RemappedSecondBoneIndices.SetNumUninitialized(NumBonesFirst);
+			for (int32 Index = 0; Index < NumBonesFirst; ++Index)
+			{
+				RemappedSecondBoneIndices[Index] = Index;
+			}
+		}
+
+
+		// Surfaces
+		//---------------------------
+		
+		// Remap bone indices if we merge surfaces since bonemaps will be merged too.
+		bool bRemapBoneIndices = false;
+		TArray<uint16> RemappedBoneMapIndices;
+
+		const int32 NumBonesInSecondBoneMap = pSecond->BoneMap.Num();
+
+		// Used to know the format of the bone index buffer
+		uint32 MaxNumBonesInBoneMaps = 0;
+
+		{
+			MUTABLE_CPUPROFILER_SCOPE(Surfaces);
+			
+			const int32 NumFirstBonesInBoneMap = pFirst->BoneMap.Num();
+			Result->BoneMap = pFirst->BoneMap;
+
+			// Copy BoneMap to remap bones when merging skeletons
+			TArray<FBoneIdOrIndex> SecondBoneMap = pSecond->BoneMap;
+
+			if (bMergeSkeletons)
+			{
+				// Remap BoneMap to new skeleton
+				for (int32 BoneIndex = 0; BoneIndex < NumBonesInSecondBoneMap; ++BoneIndex)
+				{
+					SecondBoneMap[BoneIndex].Index = RemappedSecondBoneIndices[SecondBoneMap[BoneIndex].Index];
+				}
+			}
+
+			if (bMergeSurfaces)
+			{
+				// Merge BoneMaps
+				RemappedBoneMapIndices.Reserve(NumBonesInSecondBoneMap);
+
+				for (uint16 Index = 0; Index < NumBonesInSecondBoneMap; ++Index)
+				{
+					const int32 BoneMapIndex = Result->BoneMap.AddUnique(SecondBoneMap[Index]);
+					RemappedBoneMapIndices.Add(BoneMapIndex);
+
+					bRemapBoneIndices = bRemapBoneIndices || BoneMapIndex != Index;
+				}
+
+				FMeshSurface& NewSurface = Result->Surfaces.AddDefaulted_GetRef();
+				Result->ClothSections.AddDefaulted();
+				NewSurface.BoneMapCount = Result->BoneMap.Num();
+
+				int32 NumFirstSubMeshes = 0;
+				for (const FMeshSurface& Surf : pFirst->Surfaces)
+				{
+					NewSurface.SubMeshes.Append(Surf.SubMeshes);
+					NumFirstSubMeshes += Surf.SubMeshes.Num();
+				}
+
+				for (const FMeshSurface& Surf : pSecond->Surfaces)
+				{
+					NewSurface.SubMeshes.Append(Surf.SubMeshes);
+				}
+
+				// Fix surface Submesh ranges.
+				if (NumFirstSubMeshes > 0)
+				{
+					const int32 NumResultSubMeshes = NewSurface.SubMeshes.Num();
+				
+					const FSurfaceSubMesh LastFromFirstMesh = pFirst->Surfaces.Last().SubMeshes.Last();
+	
+					for (int32 SecondSubMeshIndex = NumFirstSubMeshes; SecondSubMeshIndex < NumResultSubMeshes; ++SecondSubMeshIndex)
+					{
+						NewSurface.SubMeshes[SecondSubMeshIndex].VertexBegin += LastFromFirstMesh.VertexEnd;	
+						NewSurface.SubMeshes[SecondSubMeshIndex].VertexEnd += LastFromFirstMesh.VertexEnd;	
+						NewSurface.SubMeshes[SecondSubMeshIndex].IndexBegin += LastFromFirstMesh.IndexEnd;	
+						NewSurface.SubMeshes[SecondSubMeshIndex].IndexEnd += LastFromFirstMesh.IndexEnd;	
+					}
+				}
+			}
+			else
+			{
+				// Add the BoneMap of the second mesh
+				Result->BoneMap.Append(SecondBoneMap);
+
+				// Add pFirst surfaces
+				Result->Surfaces = pFirst->Surfaces;
+				
+				const int32 FirstVertexEnd = pFirst->GetVertexCount();
+				const int32 FirstIndexEnd = pFirst->GetIndexCount();
+
+				check(pSecond->Surfaces.Num() == 1);
+				FMeshSurface& NewSurface = Result->Surfaces.Add_GetRef(pSecond->Surfaces[0]);
+				
+				for (FSurfaceSubMesh& SubMesh : NewSurface.SubMeshes)
+				{
+					SubMesh.VertexBegin += FirstVertexEnd;
+					SubMesh.VertexEnd += FirstVertexEnd;
+					SubMesh.IndexBegin += FirstIndexEnd;
+					SubMesh.IndexEnd += FirstIndexEnd;
+				}
+
+				NewSurface.BoneMapIndex += NumFirstBonesInBoneMap;
+				
+				Result->ClothSections.SetNum(Result->Surfaces.Num());
+			}
+
+			for (const FMeshSurface& Surface : Result->Surfaces)
+			{
+				MaxNumBonesInBoneMaps = FMath::Max(MaxNumBonesInBoneMaps, Surface.BoneMapCount);
+			}
+
+			Result->BoneMap.Shrink();
+		}
+
+
+		// Pose
+		//---------------------------
+		if (TManagedPtr<const FSkeleton> ResultSkeleton = Result->GetSkeleton())
+		{
+			MUTABLE_CPUPROFILER_SCOPE(Pose);
+
+			Result->BonePoses.Reserve(ResultSkeleton->BoneNames.Num());
+
+			// Copy poses from the first mesh
+			Result->BonePoses = pFirst->BonePoses;
+
+			// Add or override bone poses
+			for (const FMesh::FBonePose& SecondBonePose : pSecond->BonePoses)
+			{
+				const int32 RemappedBoneIndex = RemappedSecondBoneIndices[SecondBonePose.BoneId.Index];
+				const int32 ResultBonePoseIndex = Result->FindBonePoseByBoneIndex(RemappedBoneIndex);
+
+				if (ResultBonePoseIndex != INDEX_NONE)
+				{
+					FMesh::FBonePose& ResultBonePose = Result->BonePoses[ResultBonePoseIndex];
+
+					// TODO: Not sure how to tune this priority, review it.
+					// For now use a similar strategy as before. 
+					auto ComputeBoneMergePriority = [](const FMesh::FBonePose& BonePose)
+					{
+						return (EnumHasAnyFlags(BonePose.BoneUsageFlags, EBoneUsageFlags::Skinning) ? 1 : 0) +
+							(EnumHasAnyFlags(BonePose.BoneUsageFlags, EBoneUsageFlags::Reshaped) ? 1 : 0);
+					};
+
+					const int32 ResultBonePriority = ComputeBoneMergePriority(ResultBonePose);
+					const int32 SecondBonePriority = ComputeBoneMergePriority(SecondBonePose);
+
+					if (ResultBonePriority < SecondBonePriority ||
+						(ResultBonePriority == SecondBonePriority && ResultBonePose.BonePosePriority < SecondBonePose.BonePosePriority))
+					{
+						//ResultBonePose.BoneName = SecondBonePose.BoneName;
+						ResultBonePose.BoneTransform = SecondBonePose.BoneTransform;
+						ResultBonePose.BonePosePriority = SecondBonePose.BonePosePriority;
+
+						// Merge usage flags
+						EnumAddFlags(ResultBonePose.BoneUsageFlags, SecondBonePose.BoneUsageFlags);
+					}
+				}
+				else
+				{
+					FMesh::FBonePose NewPose = SecondBonePose;
+					NewPose.BoneId.Index = RemappedBoneIndex;
+					Result->BonePoses.Add(NewPose);
+				}
+			}
+
+			Result->BonePoses.Shrink();
+		}
+
+
+		// PhysicsBodies
+		//---------------------------
+		{
+			MUTABLE_CPUPROFILER_SCOPE(PhysicsBodies);
+			
+			// Append PhysicsAssets
+			Result->PhysicsAssets = pFirst->PhysicsAssets;
+
+			for (const TPassthroughObjectPtr<UPhysicsAsset>& PhysicsAsset : pSecond->PhysicsAssets)
+			{
+				Result->PhysicsAssets.AddUnique(PhysicsAsset);
+			}
+
+			// Appends InPhysicsBody to OutPhysicsBody removing Bodies that are equal and their properties are identical.
+			auto AppendPhysicsBodiesUnique = [](FPhysicsBody& OutPhysicsBody, const FPhysicsBody& InPhysicsBody) -> bool
+			{
+				TArray<FName>& OutBones = OutPhysicsBody.BodiesBoneNames;
+				TArray<FPhysicsBodyAggregate>& OutBodies = OutPhysicsBody.Bodies;
+
+				const TArray<FName>& InBones = InPhysicsBody.BodiesBoneNames;
+				const TArray<FPhysicsBodyAggregate>& InBodies = InPhysicsBody.Bodies;
+
+				const int32 InBodyCount = InPhysicsBody.GetBodyCount();
+				const int32 OutBodyCount = OutPhysicsBody.GetBodyCount();
+
+				bool bModified = false;
+
+				for (int32 InBodyIndex = 0; InBodyIndex < InBodyCount; ++InBodyIndex)
+				{
+					const int32 FoundIndex = OutBones.Find(InBones[InBodyIndex]);
+
+					if (FoundIndex == INDEX_NONE)
+					{
+						OutBones.Add(InBones[InBodyIndex]);
+						OutBodies.Add(InBodies[InBodyIndex]);
+
+						bModified |= true;
+
+						continue;
+					}
+
+					for (const FSphereBody& Body : InBodies[InBodyIndex].Spheres)
+					{
+						const int32 NumBeforeAddition = OutBodies[FoundIndex].Spheres.Num();
+						bModified |= NumBeforeAddition == OutBodies[FoundIndex].Spheres.AddUnique(Body);
+					}
+
+					for (const FBoxBody& Body : InBodies[InBodyIndex].Boxes)
+					{
+						const int32 NumBeforeAddition = OutBodies[FoundIndex].Boxes.Num();
+						bModified |= NumBeforeAddition == OutBodies[FoundIndex].Boxes.AddUnique(Body);
+					}
+
+					for (const FSphylBody& Body : InBodies[InBodyIndex].Sphyls)
+					{
+						const int32 NumBeforeAddition = OutBodies[FoundIndex].Sphyls.Num();
+						bModified |= NumBeforeAddition == OutBodies[FoundIndex].Sphyls.AddUnique(Body);
+					}
+
+					for (const FTaperedCapsuleBody& Body : InBodies[InBodyIndex].TaperedCapsules)
+					{
+						const int32 NumBeforeAddition = OutBodies[FoundIndex].TaperedCapsules.Num();
+						bModified |= NumBeforeAddition == OutBodies[FoundIndex].TaperedCapsules.AddUnique(Body);
+					}
+
+					for (const FConvexBody& Body : InBodies[InBodyIndex].Convex)
+					{
+						const int32 NumBeforeAddition = OutBodies[FoundIndex].Convex.Num();
+						bModified |= NumBeforeAddition == OutBodies[FoundIndex].Convex.AddUnique(Body);
+					}
+				}
+
+				return bModified;
+			};
+
+			TTuple<TManagedPtr<const FPhysicsBody>, bool> SharedResultPhysicsBody = Invoke([&]()
+				-> TTuple<TManagedPtr<const FPhysicsBody>, bool>
+			{
+				if (pFirst->GetPhysicsBody() == pSecond->GetPhysicsBody())
+				{
+					return MakeTuple(pFirst->GetPhysicsBody(), true);
+				}
+
+				if (pFirst->GetPhysicsBody() && !pSecond->GetPhysicsBody())
+				{
+					return MakeTuple(pFirst->GetPhysicsBody(), true);
+				}
+
+				if (!pFirst->GetPhysicsBody() && pSecond->GetPhysicsBody())
+				{
+					return MakeTuple(pSecond->GetPhysicsBody(), true);
+				}
+
+				return MakeTuple(nullptr, false);
+			});
+
+			if (SharedResultPhysicsBody.Get<1>())
+			{
+				// Only one or non of the meshes has physics, share the result.
+				Result->SetPhysicsBody(SharedResultPhysicsBody.Get<0>());
+			}
+			else
+			{
+				check(pFirst->GetPhysicsBody() && pSecond->GetPhysicsBody());
+
+				TManagedPtr<FPhysicsBody> MergedResultPhysicsBody = pFirst->GetPhysicsBody()->Clone();
+
+				MergedResultPhysicsBody->bBodiesModified =
+					AppendPhysicsBodiesUnique(*MergedResultPhysicsBody, *pSecond->GetPhysicsBody()) ||
+					pFirst->GetPhysicsBody()->bBodiesModified || pSecond->GetPhysicsBody()->bBodiesModified;
+
+				Result->SetPhysicsBody(MergedResultPhysicsBody);
+			}
+
+			// Additional physics bodies.
+			const int32 MaxAdditionalPhysicsResultNum = pFirst->AdditionalPhysicsBodies.Num() + pSecond->AdditionalPhysicsBodies.Num();
+
+			Result->AdditionalPhysicsBodies.Reserve(MaxAdditionalPhysicsResultNum);
+			Result->AdditionalPhysicsBodies.Append(pFirst->AdditionalPhysicsBodies);
+
+			Result->AdditionalPhysicsAssets.Reserve(MaxAdditionalPhysicsResultNum);
+			Result->AdditionalPhysicsAssets.Append(pFirst->AdditionalPhysicsAssets);
+			
+			// Not very many additional bodies expected, do a quadratic search to have unique bodies based on external id.
+			const int32 NumSecondAdditionalBodies = pSecond->AdditionalPhysicsBodies.Num();
+			check(pSecond->AdditionalPhysicsAssets.Num() == NumSecondAdditionalBodies);
+
+			for (int32 Index = 0; Index < NumSecondAdditionalBodies; ++Index)
+			{
+				const int32 CustomIdToMerge = pSecond->AdditionalPhysicsBodies[Index]->CustomId;
+
+				const TManagedPtr<const FPhysicsBody>* Found = pFirst->AdditionalPhysicsBodies.FindByPredicate(
+					[CustomIdToMerge](const TManagedPtr<const FPhysicsBody>& Body) { return CustomIdToMerge == Body->CustomId; });
+
+				// TODO: current usages do not expect collisions, but same Id collision with bodies modified in differnet ways
+				// may need to be contemplated at some point.
+				if (!Found)
+				{
+					Result->AdditionalPhysicsBodies.Add(pSecond->AdditionalPhysicsBodies[Index]);
+					Result->AdditionalPhysicsAssets.Add(pSecond->AdditionalPhysicsAssets[Index]);
+				}
+			}
+		}
+
+		// This affects both vertex IDs and layout block ids.
+		bool bNeedsExplicitIds = pFirst->MeshIDPrefix != pSecond->MeshIDPrefix;
+
+		// These two extra checks are necessary for corner cases of meshes merging with fragments of themselves that
+		// undergo different operations.
+		if (!bNeedsExplicitIds)
+		{
+			UntypedMeshBufferIteratorConst FirstIDs(pFirst->VertexBuffers, EMeshBufferSemantic::VertexIndex, 0);
+			UntypedMeshBufferIteratorConst SecondIDs(pSecond->VertexBuffers, EMeshBufferSemantic::VertexIndex, 0);
+			bNeedsExplicitIds = FirstIDs.GetFormat() != SecondIDs.GetFormat();
+		}
+		if (!bNeedsExplicitIds)
+		{
+			UntypedMeshBufferIteratorConst FirstIDs(pFirst->VertexBuffers, EMeshBufferSemantic::LayoutBlock, 0);
+			UntypedMeshBufferIteratorConst SecondIDs(pSecond->VertexBuffers, EMeshBufferSemantic::LayoutBlock, 0);
+			bNeedsExplicitIds = FirstIDs.GetFormat() != SecondIDs.GetFormat();
+		}
+
+		if (!bNeedsExplicitIds)
+		{
+			// This is needed in case a mesh merges with itself.
+			Result->MeshIDPrefix = pFirst->MeshIDPrefix;
+		}
+
+		// Vertices
+		//-----------------
+		{
+            MUTABLE_CPUPROFILER_SCOPE(Vertices);
+
+			const int32 FirstBufferCount = pFirst->VertexBuffers.Buffers.Num();
+			const int32 SecondBufferCount = pSecond->VertexBuffers.Buffers.Num();
+			const int32 FirstVertexCount = pFirst->GetVertexBuffers().GetElementCount();
+			const int32 SecondVertexCount = pSecond->GetVertexBuffers().GetElementCount();
+
+			// Check if the format of the BoneIndex buffer has to change
+			bool bChangeBoneIndicesFormat = false;
+			EMeshBufferFormat BoneIndexFormat = MaxNumBonesInBoneMaps > MAX_uint8 ? EMeshBufferFormat::UInt16 : EMeshBufferFormat::UInt8;
+			bChangeBoneIndicesFormat = pFirst->GetVertexBuffers().HasAnySemanticWithDifferentFormat(EMeshBufferSemantic::BoneIndices, BoneIndexFormat);
+			if (!bChangeBoneIndicesFormat)
+			{
+				bChangeBoneIndicesFormat = pSecond->GetVertexBuffers().HasAnySemanticWithDifferentFormat(EMeshBufferSemantic::BoneIndices, BoneIndexFormat);
+			}
+
+			// Step 1: Set up the initial vertex buffer structure of the result mesh.
+			//-----------------------------------------------------------------------
+			{
+				MUTABLE_CPUPROFILER_SCOPE(ResultFormatSetup);
+
+				// Start with the structure of the first mesh
+				Result->GetVertexBuffers().SetBufferCount(FirstBufferCount);
+				for (int32 BufferIndex = 0; BufferIndex < FirstBufferCount; ++BufferIndex)
+				{
+					Result->VertexBuffers.Buffers[BufferIndex].Channels = pFirst->VertexBuffers.Buffers[BufferIndex].Channels;
+					Result->VertexBuffers.Buffers[BufferIndex].ElementSize = pFirst->VertexBuffers.Buffers[BufferIndex].ElementSize;
+				}
+
+				// See if we need to add additional buffers from the second mesh (like vertex colours or additional UV Channels)
+				// This is a bit ad-hoc: we only add buffers containing all new channels
+				for (const FMeshBuffer& SecondBuf : pSecond->GetVertexBuffers().Buffers)
+				{
+					bool bSomeChannel = false;
+					bool bAllNewChannels = true;
+					for (const FMeshBufferChannel& SecondChan : SecondBuf.Channels)
+					{
+						// Skip system buffers
+						if (SecondChan.Semantic == EMeshBufferSemantic::VertexIndex
+							||
+							SecondChan.Semantic == EMeshBufferSemantic::LayoutBlock)
+						{
+							continue;
+						}
+
+						bSomeChannel = true;
+
+						int32 FoundBuffer = -1;
+						int32 FoundChannel = -1;
+						pFirst->GetVertexBuffers().FindChannel(SecondChan.Semantic, SecondChan.SemanticIndex, &FoundBuffer, &FoundChannel);
+						if (FoundBuffer >= 0)
+						{
+							// There's at least one channel that already exists in the first mesh. Don't add the buffer.
+							bAllNewChannels = false;
+							continue;
+						}
+
+						// If there are additional UV channels try to add them
+						if (!bAllNewChannels && SecondChan.Semantic == EMeshBufferSemantic::TexCoords
+							&&
+							SecondChan.SemanticIndex > 0)
+						{
+							// Add additional UV channels if the previous one is found.
+							FMeshBufferSet& ResultVertexBuffers = Result->GetVertexBuffers();
+							ResultVertexBuffers.FindChannel(EMeshBufferSemantic::TexCoords, SecondChan.SemanticIndex - 1, &FoundBuffer, &FoundChannel);
+
+							if (FoundBuffer >= 0)
+							{
+								FMeshBuffer& Buffer = ResultVertexBuffers.Buffers[FoundBuffer];
+								Buffer.Channels.Insert(SecondChan, FoundChannel + 1);
+
+								ResultVertexBuffers.UpdateOffsets(FoundBuffer);
+							}
+						}
+					}
+
+					if (bSomeChannel && bAllNewChannels)
+					{
+						int32 NewBufferIndex = Result->GetVertexBuffers().Buffers.Emplace();
+						Result->VertexBuffers.Buffers[NewBufferIndex].Channels = SecondBuf.Channels;
+						Result->VertexBuffers.Buffers[NewBufferIndex].ElementSize = SecondBuf.ElementSize;
+					}
+				}
+
+				// See if we need to enlarge the components of any of the result channels
+				int32 ResultBufferCount = Result->GetVertexBuffers().GetBufferCount();
+				for (int32 BufferIndex = 0; BufferIndex < FMath::Min(ResultBufferCount, FirstBufferCount); ++BufferIndex)
+				{
+					// Expand component counts in vertex channels of the format mesh
+					FMeshBuffer& result = Result->GetVertexBuffers().Buffers[BufferIndex];
+
+					bool bResetOffsets = false;
+					for (int32 c = 0; c < result.Channels.Num(); ++c)
+					{
+						int32 sb = -1;
+						int32 sc = -1;
+						pSecond->GetVertexBuffers().FindChannel
+						(
+							result.Channels[c].Semantic,
+							result.Channels[c].SemanticIndex,
+							&sb, &sc
+						);
+
+						if (sb >= 0)
+						{
+							const FMeshBuffer& second = pSecond->GetVertexBuffers().Buffers[sb];
+
+							if (second.Channels[sc].ComponentCount>result.Channels[c].ComponentCount)
+							{
+								result.Channels[c].ComponentCount = second.Channels[sc].ComponentCount;
+								bResetOffsets = true;
+							}
+						}
+					}
+
+					// Reset the channel offsets if necessary
+					if (bResetOffsets)
+					{
+						Result->GetVertexBuffers().UpdateOffsets(BufferIndex);
+					}
+				}
+
+
+				// Change the format of the bone indices buffer
+				if (bChangeBoneIndicesFormat)
+				{
+					// Iterate all vertex buffers and update the format
+					FMeshBufferSet& VertexBuffers = Result->GetVertexBuffers();
+					for (int32 VertexBufferIndex = 0; VertexBufferIndex < VertexBuffers.Buffers.Num(); ++VertexBufferIndex)
+					{
+						FMeshBuffer& result = VertexBuffers.Buffers[VertexBufferIndex];
+
+						const int32 ChannelsCount = VertexBuffers.GetBufferChannelCount(VertexBufferIndex);
+						for (int32 ChannelIndex = 0; ChannelIndex < ChannelsCount; ++ChannelIndex)
+						{
+							if (result.Channels[ChannelIndex].Semantic == EMeshBufferSemantic::BoneIndices)
+							{
+								result.Channels[ChannelIndex].Format = BoneIndexFormat;
+
+								// Reset offsets
+								int32 offset = 0;
+								for (int32 AuxChannelIndex = 0; AuxChannelIndex < ChannelsCount; ++AuxChannelIndex)
+								{
+									result.Channels[AuxChannelIndex].Offset = (uint8_t)offset;
+									offset += result.Channels[AuxChannelIndex].ComponentCount
+										*
+										GetMeshFormatData(result.Channels[AuxChannelIndex].Format).SizeInBytes;
+								}
+								result.ElementSize = offset;
+							}
+						}
+					}
+				}
+
+				// Manage vertex IDs
+				if (bNeedsExplicitIds)
+				{
+					// Make sure the result format is suitable for the explicit IDs
+					Result->MakeIdsExplicit();
+				}
+			}
+
+
+			// Step 2: Fill the result buffers
+			//-----------------------------------------------------------------------
+			if (pFirst->VertexBuffers.IsDescriptor() || pSecond->VertexBuffers.IsDescriptor())
+			{
+				EnumAddFlags(Result->VertexBuffers.Flags, EMeshBufferSetFlags::IsDescriptor);
+			}
+			Result->VertexBuffers.SetElementCount(FirstVertexCount + SecondVertexCount);
+			
+			if (!Result->VertexBuffers.IsDescriptor())
+			{
+				MUTABLE_CPUPROFILER_SCOPE(ResultFill);
+				
+				// We have the final result vertex buffers structure: allocate the memory for them.
+				int32 ResultBufferCount = Result->GetVertexBuffers().GetBufferCount();
+				for (int32 ResultBufferIndex = 0; ResultBufferIndex < ResultBufferCount; ++ResultBufferIndex)
+				{
+					FMeshBuffer& ResultBuffer = Result->VertexBuffers.Buffers[ResultBufferIndex];
+
+					// TODO: Don't assume the buffer order in First and Second matches Result, for more opportunities of fast-path
+					bool bFirstFastPath = (ResultBufferIndex < FirstBufferCount)
+						&&
+						pFirst->VertexBuffers.HasSameFormat(ResultBufferIndex, Result->VertexBuffers, ResultBufferIndex);
+
+					int32 FirstResultBufferSize = ResultBuffer.ElementSize * FirstVertexCount;
+					uint8* Destination = ResultBuffer.Data.GetData();
+					if (bFirstFastPath)
+					{
+						MUTABLE_CPUPROFILER_SCOPE(FirstFastPath);
+
+						const FMeshBuffer& FirstBuffer = pFirst->VertexBuffers.Buffers[ResultBufferIndex];
+						check(FirstResultBufferSize == FirstBuffer.Data.Num());
+						FMemory::Memcpy(Destination, FirstBuffer.Data.GetData(), FirstResultBufferSize);
+					}
+					else
+					{
+						MUTABLE_CPUPROFILER_SCOPE(FirstSlowPath);
+
+						int32 OffsetElements = 0;
+						MeshFormatBuffer(pFirst->VertexBuffers, ResultBuffer, OffsetElements, true, pFirst->MeshIDPrefix);
+					}
+
+
+					bool bSecondFastPath = (ResultBufferIndex < SecondBufferCount)
+						&&
+						pSecond->VertexBuffers.HasSameFormat(ResultBufferIndex, Result->VertexBuffers, ResultBufferIndex);
+
+					int32 SecondResultBufferSize = ResultBuffer.ElementSize * SecondVertexCount;
+					Destination = ResultBuffer.Data.GetData() + FirstResultBufferSize;
+					if (bSecondFastPath)
+					{
+						MUTABLE_CPUPROFILER_SCOPE(SecondFastPath);
+
+						const FMeshBuffer& SecondBuffer = pSecond->VertexBuffers.Buffers[ResultBufferIndex];
+						check(SecondResultBufferSize == SecondBuffer.Data.Num());
+						FMemory::Memcpy(Destination, SecondBuffer.Data.GetData(), SecondResultBufferSize);
+					}
+					else
+					{
+						MUTABLE_CPUPROFILER_SCOPE(SecondSlowPath);
+
+						int32 OffsetElements = FirstVertexCount;
+						MeshFormatBuffer(pSecond->VertexBuffers, ResultBuffer, OffsetElements, true, pSecond->MeshIDPrefix);
+					}
+				}
+
+				// TODO: This still needs to be optimized
+				if (bRemapBoneIndices)
+				{
+					MUTABLE_CPUPROFILER_SCOPE(RemapBones);
+
+					// We need to remap the bone indices of the second mesh vertices that we already copied to result
+					check(!RemappedBoneMapIndices.IsEmpty());
+
+					// Iterate all vertex buffers and update the format
+					FMeshBufferSet& VertexBuffers = Result->GetVertexBuffers();
+					for (int32 VertexBufferIndex = 0; VertexBufferIndex < VertexBuffers.Buffers.Num(); ++VertexBufferIndex)
+					{
+						FMeshBuffer& ResultBuffer = VertexBuffers.Buffers[VertexBufferIndex];
+
+						const int32 ElemSize = VertexBuffers.GetElementSize(VertexBufferIndex);
+						const int32 FirstSize = FirstVertexCount * ElemSize;
+
+						const int32 ChannelsCount = VertexBuffers.GetBufferChannelCount(VertexBufferIndex);
+						for (int32 ChannelIndex = 0; ChannelIndex < ChannelsCount; ++ChannelIndex)
+						{
+							if (ResultBuffer.Channels[ChannelIndex].Semantic != EMeshBufferSemantic::BoneIndices)
+							{
+								continue;
+							}
+
+							int32 ResultOffset = FirstSize + ResultBuffer.Channels[ChannelIndex].Offset;
+
+							const int32 NumComponents = ResultBuffer.Channels[ChannelIndex].ComponentCount;
+
+							// Bone indices may need remapping
+							for (int32 VertexIndex = 0; VertexIndex < pSecond->GetVertexCount(); ++VertexIndex)
+							{
+								switch (BoneIndexFormat)
+								{
+								case EMeshBufferFormat::Int8:
+								case EMeshBufferFormat::UInt8:
+								{
+									uint8* pD = reinterpret_cast<uint8*>(&ResultBuffer.Data[ResultOffset]);
+
+									for (int32 ComponentIndex = 0; ComponentIndex < NumComponents; ++ComponentIndex)
+									{
+										uint8 BoneMapIndex = pD[ComponentIndex];
+
+										// be defensive
+										if (BoneMapIndex < NumBonesInSecondBoneMap)
+										{
+											pD[ComponentIndex] = (uint8)RemappedBoneMapIndices[BoneMapIndex];
+										}
+										else
+										{
+											pD[ComponentIndex] = 0;
+										}
+									}
+
+									ResultOffset += ElemSize;
+									break;
+								}
+
+								case EMeshBufferFormat::Int16:
+								case EMeshBufferFormat::UInt16:
+								{
+									uint16* pD = reinterpret_cast<uint16*>(&ResultBuffer.Data[ResultOffset]);
+
+									for (int32 ComponentIndex = 0; ComponentIndex < NumComponents; ++ComponentIndex)
+									{
+										uint16 BoneMapIndex = pD[ComponentIndex];
+
+										// be defensive
+										if (BoneMapIndex < NumBonesInSecondBoneMap)
+										{
+											pD[ComponentIndex] = (uint16)RemappedBoneMapIndices[BoneMapIndex];
+										}
+										else
+										{
+											pD[ComponentIndex] = 0;
+										}
+									}
+
+									ResultOffset += ElemSize;
+									break;
+								}
+
+								case EMeshBufferFormat::Int32:
+								case EMeshBufferFormat::UInt32:
+								{
+									// Unreal does not support 32 bit bone indices
+									checkf(false, TEXT("Format not supported."));
+									break;
+								}
+
+								default:
+									checkf(false, TEXT("Format not supported."));
+								}
+							}
+						}
+					}
+				}
+			}
+		}
+
+
+		// Tags
+		Result->Tags = pFirst->Tags;
+
+		for (const FString& SecondTag : pSecond->Tags)
+		{
+			Result->Tags.AddUnique(SecondTag);
+		}
+
+
+		// Streamed Resources
+		Result->StreamedResources = pFirst->StreamedResources;
+
+		const int32 NumStreamedResources = pSecond->StreamedResources.Num();
+		for (int32 Index = 0; Index < NumStreamedResources; ++Index)
+		{
+			Result->StreamedResources.AddUnique(pSecond->StreamedResources[Index]);
+		}
+
+		// Skeletal Meshes
+		Result->SkeletalMeshes = pFirst->SkeletalMeshes;
+
+		const int32 NumSkeletalMeshes = pSecond->SkeletalMeshes.Num();
+		for (int32 Index = 0; Index < NumSkeletalMeshes; ++Index)
+		{
+			Result->SkeletalMeshes.AddUnique(pSecond->SkeletalMeshes[Index]);
+		}
+
+		// Mesh morphs.
+		constexpr EMorphUsageFlags MorphUsageFilter = EMorphUsageFlags::Merged;
+		if (pFirst->HasMorphs() && !pSecond->HasMorphs())
+		{
+			Result->Morph = pFirst->Morph;
+			Result->MorphDataBuffer = pFirst->MorphDataBuffer;
+			OpMeshMergeMorphsInternal::FixMorphIndicesSurfaceMetadataAndApplyFilter(*Result, 0, 0, MorphUsageFilter);
+		}
+		else if (!pFirst->HasMorphs() && pSecond->HasMorphs())
+		{
+			Result->Morph = pSecond->Morph;
+			Result->MorphDataBuffer = pSecond->MorphDataBuffer;
+
+			int32 SurfaceOffset = bMergeSurfaces ? 0 : pFirst->Surfaces.Num();
+			OpMeshMergeMorphsInternal::FixMorphIndicesSurfaceMetadataAndApplyFilter(*Result, pFirst->GetVertexCount(), SurfaceOffset, MorphUsageFilter);
+		}
+		else if (pFirst->HasMorphs() && pSecond->HasMorphs())
+		{
+			OpMeshMergeMorphsInternal::MergeMorphs(*Result, *pFirst, *pSecond, bMergeSurfaces, MorphUsageFilter);
+		}
+
+		// Cloth
+		if (bMergeSurfaces)
+		{
+			if (pFirst->ClothSections[0].IsValid() && !pSecond->ClothSections[0].IsValid())
+			{
+				OpMeshMergeClothInternal::CopyClothFirst(*Result, *pFirst, *pSecond);	
+			}
+			else if (!pFirst->ClothSections[0].IsValid() && pSecond->ClothSections[0].IsValid())
+			{
+				OpMeshMergeClothInternal::CopyClothSecond(*Result, *pFirst, *pSecond);
+			}
+			else if (pFirst->ClothSections[0].IsValid() && pSecond->ClothSections[0].IsValid())
+			{
+				UE_LOG(LogMutableCore, Error, TEXT("Merging two mesh sections with cloth is not supported."));
+				OpMeshMergeClothInternal::CopyClothFirst(*Result, *pFirst, *pSecond);
+			}
+		}
+		else
+		{
+			for (int32 SectionIndexFirst = 0; SectionIndexFirst < pFirst->ClothSections.Num(); ++SectionIndexFirst)
+			{
+				Result->ClothSections[SectionIndexFirst] = pFirst->ClothSections[SectionIndexFirst];
+			}
+
+			for (int32 SectionIndexSecond = 0; SectionIndexSecond < pSecond->ClothSections.Num(); ++SectionIndexSecond)
+			{
+				Result->ClothSections[pFirst->Surfaces.Num() + SectionIndexSecond] = pSecond->ClothSections[SectionIndexSecond];
+			}
+		}
+	}	
+}
diff --git a/Engine/Plugins/Mutable/Source/MutableRuntime/Private/MuR/OpMeshRemove.cpp b/Engine/Plugins/Mutable/Source/MutableRuntime/Private/MuR/OpMeshRemove.cpp
--- a/Engine/Plugins/Mutable/Source/MutableRuntime/Private/MuR/OpMeshRemove.cpp
+++ b/Engine/Plugins/Mutable/Source/MutableRuntime/Private/MuR/OpMeshRemove.cpp
@@ -12,142 +12,109 @@
 
 namespace UE::Mutable::Private
 {
-	void MeshRemoveMeshMorphs(FMesh& Result, TConstArrayView<int32> VertexIndexRemap)
+
+	void RemoveMorphBatchVertices(
+		UE::MorphTargetVertexCodec::FDeltaBatchHeader& InOutBatchHeader, 
+		TArrayView<uint32> DataView, 
+		//int32 BatchVertexRangeBegin, 
+		//int32 BatchVertexRangeEnd, 
+		TConstArrayView<int32> VertexRemap,
+		int32 NumResultVertices
+	)
 	{
-		MUTABLE_CPUPROFILER_SCOPE(CompressedMorphs_MeshRemoveMesh)
+		UE::MorphTargetVertexCodec::FQuantizedDelta QuantizedDeltasReadBuffer[UE::MorphTargetVertexCodec::BatchSize];
 
-		using namespace UE::MorphTargetVertexCodec;
+		TArrayView<uint32> BatchDataView = TArrayView<uint32>(
+				DataView.GetData() + (InOutBatchHeader.DataOffset/sizeof(uint32)), 
+				CalculateBatchDwords(InOutBatchHeader)); 
 
-		if (!Result.HasMorphs())
-		{
-			return;
-		}
-
-		auto GetBatchHeader = [](TArrayView<uint32> HeadersDataView, int32 BatchIndex) -> FDeltaBatchHeader
-		{
-			TArrayView<uint32> HeaderDataView = TArrayView<uint32>(
-						HeadersDataView.GetData() + BatchIndex*NumBatchHeaderDwords,
-						NumBatchHeaderDwords);
+		check(InOutBatchHeader.NumElements <= UE::MorphTargetVertexCodec::BatchSize);
+		TArrayView<UE::MorphTargetVertexCodec::FQuantizedDelta> BatchQuantizedDeltas(
+				QuantizedDeltasReadBuffer, InOutBatchHeader.NumElements);
 
-			FDeltaBatchHeader BatchHeader;
-			ReadHeader(BatchHeader, HeaderDataView);
+		ReadQuantizedDeltas(BatchQuantizedDeltas, InOutBatchHeader, BatchDataView); 
+		
+		const int32 NumDeltas = InOutBatchHeader.NumElements;
 
-			return BatchHeader;
-		};
+		check(InOutBatchHeader.NumElements <= UE::MorphTargetVertexCodec::BatchSize);
 
-		auto StoreBatchHeader = [](TArrayView<uint32> HeadersDataView, int32 BatchIndex, const FDeltaBatchHeader& BatchHeader) -> void
+		InOutBatchHeader.IndexMin = 0;
+		
+		int32 DeltaIndex = 0;
+		for ( ; DeltaIndex < NumDeltas; ++DeltaIndex)
 		{
-			TArrayView<uint32> HeaderDataView = TArrayView<uint32>(
-						HeadersDataView.GetData() + BatchIndex*NumBatchHeaderDwords,
-						NumBatchHeaderDwords);
+			check((int32)BatchQuantizedDeltas[DeltaIndex].Index < VertexRemap.Num());
+			
+			int32 RemappedDeltaVertexIdx = VertexRemap[BatchQuantizedDeltas[DeltaIndex].Index];
+			if (RemappedDeltaVertexIdx >= 0)
+			{
+				check(RemappedDeltaVertexIdx < NumResultVertices);
+				InOutBatchHeader.IndexMin = (uint32)RemappedDeltaVertexIdx;
+				break;
+			}
+		}
 
-			WriteHeader(BatchHeader, HeaderDataView);
-		};
+		UE::MorphTargetVertexCodec::FQuantizedDelta QuantizedDeltasKeepBuffer[UE::MorphTargetVertexCodec::BatchSize];
 
-		int32 NumMorphs = Result.Morph.Names.Num();
-		for (int32 MorphIndex = 0; MorphIndex < NumMorphs; ++MorphIndex)
+		uint32 DestDeltaIndex = 0;
+		for ( ; DeltaIndex < NumDeltas; ++DeltaIndex)
 		{
-			auto ProcessBatchData =
-			[
-				VertexIndexRemap
-			](
-				FDeltaBatchHeader& InOutBatchHeader,
-				TArrayView<uint32> DataView,
-				int32 BatchVertexRangeBegin,
-				int32 BatchVertexRangeEnd
-			)
+			check((int32)BatchQuantizedDeltas[DeltaIndex].Index < VertexRemap.Num());
+
+			int32 RemappedDeltaVertexIdx = VertexRemap[BatchQuantizedDeltas[DeltaIndex].Index];
+			if (RemappedDeltaVertexIdx >= 0)
 			{
-				TArray<FQuantizedDelta, TInlineAllocator<UE::MorphTargetVertexCodec::BatchSize*2>> QuantizedDeltasStorage;
+				check(RemappedDeltaVertexIdx < NumResultVertices);
 
-				// Only check the vertices beforehand if the range is within a small range, otherwise 
-				// read the quantized deltas and only check the vertices in the batch.
-				int32 NumRemovedVerticesInRange = -1;
-				constexpr int32 MaxNumVerticesToCheck = 256;
-				if (BatchVertexRangeEnd - BatchVertexRangeBegin < MaxNumVerticesToCheck)
+				QuantizedDeltasKeepBuffer[DestDeltaIndex++] = UE::MorphTargetVertexCodec::FQuantizedDelta
 				{
-					NumRemovedVerticesInRange = 0;
-					for (int32 Index = BatchVertexRangeBegin; Index < BatchVertexRangeEnd; ++Index)
-					{
-						NumRemovedVerticesInRange += static_cast<int32>(VertexIndexRemap[Index] < 0);
-					}
-				}
+					.Position = BatchQuantizedDeltas[DeltaIndex].Position,
+					.TangentZ = BatchQuantizedDeltas[DeltaIndex].TangentZ,
+					.Index    = (uint32)RemappedDeltaVertexIdx,
+				};	
+			}
+		}
+		
+		InOutBatchHeader.NumElements = DestDeltaIndex;
+		check(DestDeltaIndex <= UE::MorphTargetVertexCodec::BatchSize);
+		WriteQuantizedDeltas(TArrayView(QuantizedDeltasKeepBuffer, DestDeltaIndex), InOutBatchHeader, BatchDataView);	
 
-				if (NumRemovedVerticesInRange == 0)
-				{
-					InOutBatchHeader.IndexMin = VertexIndexRemap[InOutBatchHeader.IndexMin];
-				}
-				else if (NumRemovedVerticesInRange == BatchVertexRangeEnd - BatchVertexRangeBegin)
-				{
-					InOutBatchHeader.IndexMin    = 0;
-					InOutBatchHeader.NumElements = 0;
-				}
-				else
-				{
-					TArrayView<uint32> BatchDataView = TArrayView<uint32>(
-							DataView.GetData() + InOutBatchHeader.DataOffset/sizeof(uint32), 
-							CalculateBatchDwords(InOutBatchHeader)); 
+#if 0 // Check data integrity.
+		{
+			TArray<UE::MorphTargetVertexCodec::FQuantizedDelta, TInlineAllocator<UE::MorphTargetVertexCodec::BatchSize>> TestQuantizedDeltasStorage;
+			TestQuantizedDeltasStorage.SetNum(UE::MorphTargetVertexCodec::BatchSize);
 
-					TArrayView<FQuantizedDelta> BatchQuantizedDeltas(
-							QuantizedDeltasStorage.GetData(), InOutBatchHeader.NumElements);
+			check(InOutBatchHeader.NumElements <= UE::MorphTargetVertexCodec::BatchSize); 
 
-					ReadQuantizedDeltas(BatchQuantizedDeltas, InOutBatchHeader, BatchDataView); 
-					
-					const int32 NumDeltas = InOutBatchHeader.NumElements;
-		
-					int32 NumDeltaVerticesRemovedInRange = 0;
-					for (int32 DeltaIndex = 0; DeltaIndex < NumDeltas; ++DeltaIndex)
-					{
-						NumDeltaVerticesRemovedInRange += static_cast<int32>(VertexIndexRemap[BatchQuantizedDeltas[DeltaIndex].Index] < 0);
-					}
+			TArrayView<UE::MorphTargetVertexCodec::FQuantizedDelta> TestBatchQuantizedDeltas(
+					TestQuantizedDeltasStorage.GetData(), InOutBatchHeader.NumElements);
 
-					if (NumDeltaVerticesRemovedInRange == 0)
-					{
-						InOutBatchHeader.IndexMin = (uint32)VertexIndexRemap[InOutBatchHeader.IndexMin];
-					}
-					else if (NumDeltaVerticesRemovedInRange == BatchVertexRangeEnd - BatchVertexRangeBegin)
-					{
-						InOutBatchHeader.IndexMin	= 0;
-						InOutBatchHeader.NumElements = 0;
-					}
-					else
-					{
-						check(InOutBatchHeader.NumElements <= UE::MorphTargetVertexCodec::BatchSize);
-						TArrayView<FQuantizedDelta> KeptQuantizedDeltas(
-								QuantizedDeltasStorage.GetData() + UE::MorphTargetVertexCodec::BatchSize, 
-								InOutBatchHeader.NumElements);
+			ReadQuantizedDeltas(TestBatchQuantizedDeltas, InOutBatchHeader, BatchDataView); 
 
-						int32 DeltaIndex = 0;
-						for ( ; DeltaIndex < NumDeltas; ++DeltaIndex)
-						{
-							int32 RemappedDeltaVertexIdx = VertexIndexRemap[BatchQuantizedDeltas[DeltaIndex].Index];
-							if (RemappedDeltaVertexIdx >= 0)
-							{
-								InOutBatchHeader.IndexMin = (uint32)RemappedDeltaVertexIdx;
-								break;
-							}
-						}
+			for (int32 TestDeltaIndex = 0; TestDeltaIndex < InOutBatchHeader.NumElements; ++TestDeltaIndex)
+			{
+				check(TestBatchQuantizedDeltas[TestDeltaIndex].Index < NumResultVertices);		
+			}
+		}
+#endif
 
-						uint32 DestDeltaIndex = 0;
-						for ( ; DeltaIndex < NumDeltas; ++DeltaIndex)
-						{
-							int32 RemappedDeltaVertexIdx = VertexIndexRemap[BatchQuantizedDeltas[DeltaIndex].Index];
-							if (RemappedDeltaVertexIdx >= 0)
-							{
-								KeptQuantizedDeltas[DestDeltaIndex++] = FQuantizedDelta
-								{
-									.Position = BatchQuantizedDeltas[DeltaIndex].Position,
-									.TangentZ = BatchQuantizedDeltas[DeltaIndex].TangentZ,
-									.Index    = (uint32)RemappedDeltaVertexIdx,
-								};	
-							}
-						}
+	}
 
-						InOutBatchHeader.NumElements = DestDeltaIndex;
-						WriteQuantizedDeltas(KeptQuantizedDeltas, InOutBatchHeader, BatchDataView);
-					}
-				}
-			};
 
+	void MeshRemoveMeshMorphs(FMesh& Result, TConstArrayView<int32> VertexIndexRemap)
+	{
+		MUTABLE_CPUPROFILER_SCOPE(CompressedMorphs_MeshRemoveMesh)
+
+		using namespace UE::MorphTargetVertexCodec;
+
+		if (!Result.HasMorphs())
+		{
+			return;
+		}
+
+		int32 NumMorphs = Result.Morph.Names.Num();
+		for (int32 MorphIndex = 0; MorphIndex < NumMorphs; ++MorphIndex)
+		{
 			const int32 NumBatches = Result.Morph.BatchesPerMorph[MorphIndex];
 			
 			if (NumBatches == 0)
@@ -157,31 +124,71 @@ namespace UE::Mutable::Private
 
 			int32 BatchStartOffset = Result.Morph.BatchStartOffsetPerMorph[MorphIndex];
 
-			TArrayView<uint32> DataView = MakeArrayView(Result.MorphDataBuffer); 
+			TArrayView<uint32> DataView = Result.MorphDataBuffer; 
+			uint32* MorphDataPtr = DataView.GetData();
+
+			for (int32 BatchIndex = 0; BatchIndex < NumBatches; ++BatchIndex)
+			{
+				FDeltaBatchHeader BatchHeader;
+				ReadHeader(BatchHeader, TArrayView<uint32>(MorphDataPtr + ((BatchStartOffset + BatchIndex) * NumBatchHeaderDwords), NumBatchHeaderDwords));
+				RemoveMorphBatchVertices(BatchHeader, DataView, VertexIndexRemap, Result.GetVertexCount());	
 
-			FDeltaBatchHeader BatchHeaderStorage[2];
+				WriteHeader(BatchHeader, TArrayView<uint32>(MorphDataPtr + ((BatchStartOffset + BatchIndex) * NumBatchHeaderDwords), NumBatchHeaderDwords));
+			}
 
-			FDeltaBatchHeader* CurrBatchHeader = &BatchHeaderStorage[0];
-			FDeltaBatchHeader* NextBatchHeader = &BatchHeaderStorage[1];
+			// TODO: 
+			// The following implementation is better as it can skip some work by checking the range 
+			// of indices for a batch only looking at 2 consecutive headers. This should allow to remap all
+			// morph vertices indices just changing the batch MinIndex for batches that don't remove vertices 
+			// but this approach does not seam to be working. The code above remaps all batches changing the 
+			// deltas, but it world be worth investigate why the more optimal approach is not working and 
+			// reimplement it.
+			//
+			// Additionally, this code does not remove data from the batches, it only modifies the number of 
+			// valid deltas in the batches. As a memory optimization, the unused data should be removed.
 			
-			uint32* MorphDataPtr = DataView.GetData();
+			//FDeltaBatchHeader BatchHeaderStorage[2];
+
+			//FDeltaBatchHeader* CurrBatchHeader = &BatchHeaderStorage[0];
+			//FDeltaBatchHeader* NextBatchHeader = &BatchHeaderStorage[1];
+			//
+			//int32 CurrBatchIndex = 0;
+			// Find first non empty batch.
+			//for (; CurrBatchIndex < NumBatches; ++CurrBatchIndex)
+			//{
+			//	ReadHeader(*CurrBatchHeader, TArrayView<uint32>(MorphDataPtr + ((BatchStartOffset + CurrBatchIndex) * NumBatchHeaderDwords), NumBatchHeaderDwords));
+
+			//	if (CurrBatchHeader->NumElements)
+			//	{
+			//		break;
+			//	}
+			//}
+
+			//const int32 NumSourceVertices = VertexIndexRemap.Num();
+			//for (int32 NextBatchIndex = CurrBatchIndex + 1; NextBatchIndex < NumBatches; ++NextBatchIndex)
+			//{
+			//	// Find next non empty batch.
+			//	for (; NextBatchIndex < NumBatches; ++NextBatchIndex)
+			//	{
+			//		ReadHeader(*NextBatchHeader, TArrayView<uint32>(MorphDataPtr + ((BatchStartOffset + NextBatchIndex) * NumBatchHeaderDwords), NumBatchHeaderDwords));
+
+			//		if (NextBatchHeader->NumElements)
+			//		{
+			//			break;
+			//		}
+			//	}
+		
+			//	int32 BatchRangeBegin = CurrBatchHeader->IndexMin;
+			//	int32 BatchRangeEnd   = NextBatchIndex < NumBatches ? NextBatchHeader->IndexMin : NumSourceVertices; 	
 
-			ReadHeader(*CurrBatchHeader, TArrayView<uint32>(MorphDataPtr + (BatchStartOffset + 0) * NumBatchHeaderDwords, NumBatchHeaderDwords));
+			//	RemoveMorphBatchVertices(*CurrBatchHeader, DataView, BatchRangeBegin, BatchRangeEnd, VertexIndexRemap);	
 
-			for (int32 BatchIndex = 0; BatchIndex < NumBatches - 1; ++BatchIndex)
-			{
-				ReadHeader(*NextBatchHeader, TArrayView<uint32>(MorphDataPtr + (BatchStartOffset + BatchIndex + 1) * NumBatchHeaderDwords, NumBatchHeaderDwords));
-				
-				ProcessBatchData(*CurrBatchHeader, DataView, CurrBatchHeader->IndexMin, NextBatchHeader->IndexMin);
-				
-				WriteHeader(*CurrBatchHeader, TArrayView<uint32>(MorphDataPtr + (BatchStartOffset + BatchIndex) * NumBatchHeaderDwords, NumBatchHeaderDwords));
-				::Swap(CurrBatchHeader, NextBatchHeader);
-			}
+			//	WriteHeader(*CurrBatchHeader, TArrayView<uint32>(MorphDataPtr + ((BatchStartOffset + CurrBatchIndex) * NumBatchHeaderDwords), NumBatchHeaderDwords));
 
-			const int32 NumVertices = VertexIndexRemap.Num();
-			ProcessBatchData(*CurrBatchHeader, DataView, CurrBatchHeader->IndexMin, NumVertices);
-			WriteHeader(*CurrBatchHeader, TArrayView<uint32>(MorphDataPtr + (BatchStartOffset + NumBatches - 1) * NumBatchHeaderDwords, NumBatchHeaderDwords));
-		}
+			//	::Swap(CurrBatchHeader, NextBatchHeader);
+			//	CurrBatchIndex = NextBatchIndex;
+			//}
+		}	
 	}
 
     void MeshRemoveRecreateSurface(FMesh* Result, const TBitArray<>& UsedVertices, const TBitArray<>& UsedFaces)
@@ -221,16 +228,53 @@ namespace UE::Mutable::Private
                 PrevMeshRangeIndexEnd = MeshRangeIndexEnd;
             }
         }
-		
+
 		// Remove Empty surfaces but always keep the first one. 
 		// The previous step has eliminated empty submeshes, so it is only needed to check if the surface has 
 		// any submesh.
-		for (int32 I = Result->Surfaces.Num() - 1; I >= 1; --I)
+		
+		TArray<int32, TInlineAllocator<16>> SurfaceIndexRemap;
+		SurfaceIndexRemap.Init(-1, Result->Surfaces.Num());
+
+		int32 NumRemovedSurfaces = 0;
+		for (int32 I = 0; I < Result->Surfaces.Num(); ++I)
+		{
+			bool bSurfaceRemoved = !Result->Surfaces[I].SubMeshes.Num();
+			NumRemovedSurfaces += static_cast<int32>(bSurfaceRemoved);
+
+			if (!bSurfaceRemoved)
+			{
+				SurfaceIndexRemap[I] = I - NumRemovedSurfaces; 
+			}
+		}
+
+		if (NumRemovedSurfaces > 0)
 		{
-			if (!Result->Surfaces[I].SubMeshes.Num())
+			for (int32 I = Result->Surfaces.Num() - 1; I >= 1; --I)
+			{
+				if (SurfaceIndexRemap[I] < 0)
+				{
+					Result->Surfaces.RemoveAt(I, EAllowShrinking::No);
+					Result->ClothSections.RemoveAt(I, EAllowShrinking::No);
+				}
+			}
+
+			// Remove and remap surface indces in use from Morph.
+			for (int32 MorphIndex = Result->Morph.SurfacesInUsePerMorph.Num() - 1; MorphIndex >= 0; --MorphIndex)
 			{
-				Result->Surfaces.RemoveAt(I, EAllowShrinking::No);
-				Result->ClothSections.RemoveAt(I, EAllowShrinking::No);
+				for (int32 SurfIndex = SurfaceIndexRemap.Num() - 1; SurfIndex >= 0; --SurfIndex)
+				{
+					if (SurfaceIndexRemap[SurfIndex] < 0)
+					{
+						Result->Morph.SurfacesInUsePerMorph[MorphIndex].Remove(SurfIndex);
+					}
+				}
+
+				for (int32& SurfaceInUseIndex : Result->Morph.SurfacesInUsePerMorph[MorphIndex]) 
+				{
+					check(SurfaceIndexRemap[SurfaceInUseIndex] >= 0);
+					SurfaceInUseIndex = SurfaceIndexRemap[SurfaceInUseIndex];
+				}
 			}
 		}
 
@@ -477,12 +521,7 @@ namespace UE::Mutable::Private
 		}
 
 		TArray<int32> UsedVerticesMap;
-#if DO_CHECK
-        UsedVerticesMap.Init(-1, NumVertices);
-#else
-        // This data will only be accessed by indices that have been mapped, No need to initialize. 
-		UsedVerticesMap.SetNumUninitialized(NumVertices);
-#endif
+		UsedVerticesMap.Init(-1, NumVertices);
     	
 		FMeshBufferSet& VertexBufferSet = Result->GetVertexBuffers();
 
diff --git a/Engine/Plugins/Mutable/Source/MutableTools/Internal/MuT/NodeMeshConstant.h b/Engine/Plugins/Mutable/Source/MutableTools/Internal/MuT/NodeMeshConstant.h
--- a/Engine/Plugins/Mutable/Source/MutableTools/Internal/MuT/NodeMeshConstant.h
+++ b/Engine/Plugins/Mutable/Source/MutableTools/Internal/MuT/NodeMeshConstant.h
@@ -26,6 +26,8 @@ namespace UE::Mutable::Private
 		UE::Mutable::Private::TManagedPtr<FMesh> Value;
 
 		TArray<Ptr<NodeLayout>> Layouts;
+		
+		TArray<FName> RealTimeMorphNames;
 
 	public:
 
diff --git a/Engine/Plugins/Mutable/Source/MutableTools/Internal/MuT/NodeMeshSkeletalMeshBreak.h b/Engine/Plugins/Mutable/Source/MutableTools/Internal/MuT/NodeMeshSkeletalMeshBreak.h
--- a/Engine/Plugins/Mutable/Source/MutableTools/Internal/MuT/NodeMeshSkeletalMeshBreak.h
+++ b/Engine/Plugins/Mutable/Source/MutableTools/Internal/MuT/NodeMeshSkeletalMeshBreak.h
@@ -19,6 +19,8 @@ namespace UE::Mutable::Private
 
 		TArray<Ptr<NodeLayout>> Layouts;
 
+		TArray<FName> RealTimeMorphNames;
+
 		uint8 LODIndex = 0;
 
 		uint8 SectionIndex = 0;
diff --git a/Engine/Plugins/Mutable/Source/MutableTools/Internal/MuT/NodeMeshTable.h b/Engine/Plugins/Mutable/Source/MutableTools/Internal/MuT/NodeMeshTable.h
--- a/Engine/Plugins/Mutable/Source/MutableTools/Internal/MuT/NodeMeshTable.h
+++ b/Engine/Plugins/Mutable/Source/MutableTools/Internal/MuT/NodeMeshTable.h
@@ -27,6 +27,7 @@ namespace UE::Mutable::Private
 		FString DefaultRowName;
 
 		TArray<Ptr<NodeLayout>> Layouts;
+		TArray<FName> RealTimeMorphNames;
 
 		/** */
 		FSourceDataDescriptor SourceDataDescriptor;
diff --git a/Engine/Plugins/Mutable/Source/MutableTools/Private/MuT/ASTOpConstantResource.cpp b/Engine/Plugins/Mutable/Source/MutableTools/Private/MuT/ASTOpConstantResource.cpp
--- a/Engine/Plugins/Mutable/Source/MutableTools/Private/MuT/ASTOpConstantResource.cpp
+++ b/Engine/Plugins/Mutable/Source/MutableTools/Private/MuT/ASTOpConstantResource.cpp
@@ -114,8 +114,9 @@ namespace UE::Mutable::Private
 						EMeshCopyFlags::WithSurfaces      | 
 						EMeshCopyFlags::WithVertexBuffers |
 						EMeshCopyFlags::WithIndexBuffers  |
-						EMeshCopyFlags::WithLayouts		  |
-						EMeshCopyFlags::WithClothData;
+						EMeshCopyFlags::WithLayouts       |
+						EMeshCopyFlags::WithClothData     |
+						EMeshCopyFlags::WithMorphData     ;
 
 				TManagedPtr<FMesh> MeshGeometryData = MeshData->Clone(GeometryDataCopyFlags);
 
@@ -227,10 +228,11 @@ namespace UE::Mutable::Private
 			// MetadataMesh Mesh
 			{
 				const EMeshCopyFlags MetadataDataCopyFlags = 
-						EMeshCopyFlags::WithSurfaces         |
-						EMeshCopyFlags::WithTags             |
-						EMeshCopyFlags::WithSkeletonIDs      |
-						EMeshCopyFlags::WithStreamedResources;
+						EMeshCopyFlags::WithSurfaces          |
+						EMeshCopyFlags::WithTags              |
+						EMeshCopyFlags::WithSkeletonIDs       |
+						EMeshCopyFlags::WithStreamedResources |
+						EMeshCopyFlags::WithMorphMetaData     ;
 
 				TManagedPtr<FMesh> MeshMetadataData = MeshData->Clone(MetadataDataCopyFlags);
 			
diff --git a/Engine/Plugins/Mutable/Source/MutableTools/Private/MuT/ASTOpMeshAddMetadata.cpp b/Engine/Plugins/Mutable/Source/MutableTools/Private/MuT/ASTOpMeshAddMetadata.cpp
--- a/Engine/Plugins/Mutable/Source/MutableTools/Private/MuT/ASTOpMeshAddMetadata.cpp
+++ b/Engine/Plugins/Mutable/Source/MutableTools/Private/MuT/ASTOpMeshAddMetadata.cpp
@@ -37,7 +37,8 @@ namespace UE::Mutable::Private
 				SkeletonId == Other->SkeletonId &&
 				BonePosePriority == Other->BonePosePriority &&
 				PhysicsAssetId == Other->PhysicsAssetId &&
-				AdditionalPhysicsAssetIds == Other->AdditionalPhysicsAssetIds;
+				AdditionalPhysicsAssetIds == Other->AdditionalPhysicsAssetIds &&
+				RealTimeMorphNames == Other->RealTimeMorphNames;
 		}
 
 		return false;
@@ -54,6 +55,7 @@ namespace UE::Mutable::Private
 		NewOp->BonePosePriority = BonePosePriority;
 		NewOp->PhysicsAssetId = PhysicsAssetId;
 		NewOp->AdditionalPhysicsAssetIds = AdditionalPhysicsAssetIds;
+		NewOp->RealTimeMorphNames = RealTimeMorphNames;
 		return NewOp;
 	}
 
@@ -101,7 +103,7 @@ namespace UE::Mutable::Private
 				}
 				else
 				{
-					Args.Tags.TagAddress = Program.AddConstant(Tags[0]);
+					Args.Tags.StringAddress = Program.AddConstant(Tags[0]);
 				}
 			}
 
@@ -120,6 +122,21 @@ namespace UE::Mutable::Private
 				}
 			}
 
+			if (RealTimeMorphNames.Num())
+			{
+				EnumAddFlags(Args.Flags, OpEnumFlags::HasRealTimeMorphNames);
+
+				if (RealTimeMorphNames.Num() > 1)
+				{
+					EnumAddFlags(Args.Flags, OpEnumFlags::IsRealTimeMorphNamesList);
+					Args.RealTimeMorphNames.ListAddress = Program.AddConstant(RealTimeMorphNames);
+				}
+				else
+				{
+					Args.RealTimeMorphNames.StringAddress = Program.AddConstant(RealTimeMorphNames[0]);
+				}
+			}
+
 			Args.SkeletonId = SkeletonId;
 			Args.BonePosePriority = BonePosePriority;
 			
diff --git a/Engine/Plugins/Mutable/Source/MutableTools/Private/MuT/ASTOpMeshAddMetadata.h b/Engine/Plugins/Mutable/Source/MutableTools/Private/MuT/ASTOpMeshAddMetadata.h
--- a/Engine/Plugins/Mutable/Source/MutableTools/Private/MuT/ASTOpMeshAddMetadata.h
+++ b/Engine/Plugins/Mutable/Source/MutableTools/Private/MuT/ASTOpMeshAddMetadata.h
@@ -30,6 +30,8 @@ namespace UE::Mutable::Private
 		/** SkeletonId to add. */
 		PASSTHROUGH_ID SkeletonId = PASSTHROUGH_ID_INVALID;
 
+		TArray<FString> RealTimeMorphNames;
+
 		/** Bone Pose priority when merging two meshes. */
 		int16 BonePosePriority = 0;
 
diff --git a/Engine/Plugins/Mutable/Source/MutableTools/Private/MuT/CodeGenerator_Mesh.cpp b/Engine/Plugins/Mutable/Source/MutableTools/Private/MuT/CodeGenerator_Mesh.cpp
--- a/Engine/Plugins/Mutable/Source/MutableTools/Private/MuT/CodeGenerator_Mesh.cpp
+++ b/Engine/Plugins/Mutable/Source/MutableTools/Private/MuT/CodeGenerator_Mesh.cpp
@@ -470,29 +470,30 @@ namespace UE::Mutable::Private
 	{
 		Ptr<NodeMeshSwitch> SwitchNode;
 		GenerateTableSwitchNode<NodeMeshTable, ETableColumnType::Mesh, NodeMesh, NodeMeshSwitch>(*TableNode, SwitchNode,
-			[this] (const NodeMeshTable& node, int32 colIndex, int32 row, FErrorLog* pErrorLog) -> Ptr<NodeMesh>
+			[this] (const NodeMeshTable& Node, int32 ColIndex, int32 Row, FErrorLog* pErrorLog) -> Ptr<NodeMesh>
 			{
-				const FTableValue& Cell = node.Table->GetPrivate()->Rows[row].Values[colIndex];
-				TManagedPtr<UE::Mutable::Private::FMesh> pMesh = Cell.Mesh;
+				const FTableValue& Cell = Node.Table->GetPrivate()->Rows[Row].Values[ColIndex];
+				TManagedPtr<UE::Mutable::Private::FMesh> MeshValue = Cell.Mesh;
 				
-				if (!pMesh)
+				if (!MeshValue)
 				{
 					return nullptr;
 				}
 
 				Ptr<NodeMeshConstant> CellNode = new NodeMeshConstant();
-				CellNode->Value = pMesh;
+				CellNode->Value = MeshValue;
 
 				// TODO Take into account layout strategy
-				CellNode->Layouts = node.Layouts;
+				CellNode->Layouts = Node.Layouts;
+				CellNode->RealTimeMorphNames = Node.RealTimeMorphNames;
 
 				CellNode->SetMessageContext(Cell.ErrorContext);
-				CellNode->SourceDataDescriptor = node.SourceDataDescriptor;
+				CellNode->SourceDataDescriptor = Node.SourceDataDescriptor;
 
 				// Combine the SourceId of the node with the RowId to generate one shared between all resources from this row.
 				// Hash collisions are allowed, since it is used to group resources, not to differentiate them.
-				const uint32 RowId = node.Table->GetPrivate()->Rows[row].Id;
-				CellNode->SourceDataDescriptor.SourceId = HashCombine(node.SourceDataDescriptor.SourceId, RowId);
+				const uint32 RowId = Node.Table->GetPrivate()->Rows[Row].Id;
+				CellNode->SourceDataDescriptor.SourceId = HashCombine(Node.SourceDataDescriptor.SourceId, RowId);
 
 				return CellNode;
 			});
@@ -932,8 +933,16 @@ namespace UE::Mutable::Private
 
 				Result.BaseMeshOp = LastMeshOp;
 
+	
+				TArray<FString> RealTimeMorphNames;
+				if (NodeCopy->RealTimeMorphNames.Num())
+				{
+					RealTimeMorphNames.Reserve(NodeCopy->RealTimeMorphNames.Num());
+					Algo::Transform(NodeCopy->RealTimeMorphNames, RealTimeMorphNames, [](const FName& Name) {return Name.ToString(); });
+				}
+
 				// Add the metadata operation
-				if (bMeshHasMetadata)
+				if (bMeshHasMetadata || RealTimeMorphNames.Num() > 0)
 				{
 					Ptr<ASTOpMeshAddMetadata> AddMetadataOp = new ASTOpMeshAddMetadata;
 					AddMetadataOp->Source = LastMeshOp;
@@ -941,11 +950,14 @@ namespace UE::Mutable::Private
 					AddMetadataOp->Tags        = MoveTemp(Tags);
 					AddMetadataOp->ResourceIds = MoveTemp(ResourceIds);
 					AddMetadataOp->SkeletonId = SkeletonID;
+
 					AddMetadataOp->BonePosePriority = BonePosePriority;
 
 					AddMetadataOp->PhysicsAssetId = PhysicsAssetId;
 					AddMetadataOp->AdditionalPhysicsAssetIds = MoveTemp(AdditionalPhysicsAssetId);
 
+					AddMetadataOp->RealTimeMorphNames = MoveTemp(RealTimeMorphNames);
+
 					LastMeshOp = AddMetadataOp;
 				}
 
@@ -1461,7 +1473,21 @@ namespace UE::Mutable::Private
 				int32 LODIndex = StaticOptions.LODIndex;
 
 				FMeshGenerationDynamicOptions DynamicOptions = InOptions.GetResult();
-					
+				
+				TArray<FString> RealTimeMorphNames;
+				RealTimeMorphNames.Reserve(Node->RealTimeMorphNames.Num());
+				Algo::Transform(Node->RealTimeMorphNames, RealTimeMorphNames, [](const FName& Name) { return Name.ToString(); });
+
+				if (RealTimeMorphNames.Num() > 0)
+				{
+					Ptr<ASTOpMeshAddMetadata> AddMetadataOp = new ASTOpMeshAddMetadata;
+
+					AddMetadataOp->RealTimeMorphNames = MoveTemp(RealTimeMorphNames);
+
+					AddMetadataOp->Source = Result.MeshOp;
+					Result.MeshOp = AddMetadataOp;
+				}
+
 				return Result;
 			},
 			Requisites
