[Mutable] Fix signed/unsigned mismatch compare.

#rnx
#Synced-CL 50436190

diff --git a/Engine/Plugins/Mutable/Source/CustomizableObject/Private/MuCO/CustomizableObjectSystem.cpp b/Engine/Plugins/Mutable/Source/CustomizableObject/Private/MuCO/CustomizableObjectSystem.cpp
--- a/Engine/Plugins/Mutable/Source/CustomizableObject/Private/MuCO/CustomizableObjectSystem.cpp
+++ b/Engine/Plugins/Mutable/Source/CustomizableObject/Private/MuCO/CustomizableObjectSystem.cpp
@@ -4050,7 +4050,7 @@ void ReconstructMorphTargetsFromMeshCompressedData(const UE::Mutable::Private::F
 				// Get the Offset value for the morph traget first batch.
 				MorphDeltaDataBeginInDwords = Mesh.MorphDataBuffer[BatchStartOffset * NumBatchHeaderDwords] / sizeof(uint32);
 
-				for (int32 BatchHeaderIndex = 0; BatchHeaderIndex < MorphNumBatches; ++BatchHeaderIndex)
+				for (uint32 BatchHeaderIndex = 0; BatchHeaderIndex < MorphNumBatches; ++BatchHeaderIndex)
 				{
 					TConstArrayView<uint32> BatchHeaderData(
 							&Mesh.MorphDataBuffer[(BatchStartOffset + BatchHeaderIndex) * NumBatchHeaderDwords], 
@@ -4065,7 +4065,7 @@ void ReconstructMorphTargetsFromMeshCompressedData(const UE::Mutable::Private::F
 						CalculateBatchDwords(OutMorphTargets[MorphIndex].PackedDeltaHeaders.Last()); 
 
 				// Fixup header offset.
-				for (int32 BatchHeaderIndex = 0; BatchHeaderIndex < MorphNumBatches; ++BatchHeaderIndex)
+				for (uint32 BatchHeaderIndex = 0; BatchHeaderIndex < MorphNumBatches; ++BatchHeaderIndex)
 				{
 					FDeltaBatchHeader& Header = OutMorphTargets[MorphIndex].PackedDeltaHeaders[BatchHeaderIndex];
 					Header.DataOffset -= MorphDeltaDataBeginInDwords * sizeof(uint32);
