[Mutable] Fix crash when removing all vertices of a cloth section.

#Synced-CL 50403824

diff --git a/Engine/Plugins/Mutable/Source/MutableRuntime/Private/MuR/OpMeshRemove.cpp b/Engine/Plugins/Mutable/Source/MutableRuntime/Private/MuR/OpMeshRemove.cpp
--- a/Engine/Plugins/Mutable/Source/MutableRuntime/Private/MuR/OpMeshRemove.cpp
+++ b/Engine/Plugins/Mutable/Source/MutableRuntime/Private/MuR/OpMeshRemove.cpp
@@ -354,6 +354,11 @@ namespace UE::Mutable::Private
 	
 	void CopyMoveBuffer(uint8* BufferData, uint32 ElemSize, int32 NumVertices, const TBitArray<>& UsedVertices)
 	{
+    	if (ElemSize == 0)
+    	{
+    		return;
+    	}
+    	
 		int32 LastFreeVertexIndex = 0;
 
 		for (int32 VertexIndex = UsedVertices.FindFrom(true, 0); VertexIndex >= 0;)
@@ -577,7 +582,7 @@ namespace UE::Mutable::Private
     	if (Result->ClothSections[0].IsValid() && NumClothVertices)
     	{
     		uint8* BufferData = reinterpret_cast<uint8*>(Result->ClothSections[0].Data.GetData());
-    		const int32 NumInfluences = NumClothVertices / NumVerticesRemaining;
+    		const int32 NumInfluences = NumVerticesRemaining ? NumClothVertices / NumVerticesRemaining : 0;
     		const uint32 ElemSize = sizeof(FMeshToMeshVertData) * NumInfluences;
 
     		CopyMoveBuffer(BufferData, ElemSize, NumVertices, UsedVertices);
