[mutable] Fixed scaling issue with the Layout areas caused by the unrequired clamping of the layout areas when not needed (strategy != Fixed)
- Other small fixes

#presubmit 697cb927f04a4cd39ff423a1
#rnx
#Synced-CL 50435902

diff --git a/Engine/Plugins/Mutable/Source/CustomizableObjectEditor/Private/MuCOE/SCustomizableObjectLayoutEditor.cpp b/Engine/Plugins/Mutable/Source/CustomizableObjectEditor/Private/MuCOE/SCustomizableObjectLayoutEditor.cpp
--- a/Engine/Plugins/Mutable/Source/CustomizableObjectEditor/Private/MuCOE/SCustomizableObjectLayoutEditor.cpp
+++ b/Engine/Plugins/Mutable/Source/CustomizableObjectEditor/Private/MuCOE/SCustomizableObjectLayoutEditor.cpp
@@ -352,7 +352,7 @@ void FCustomizableObjectLayoutEditorDetailsBuilder::CustomizeDetails(IDetailLayo
 			.ValueContent()
 			[
 				SNew(SCheckBox)
-					.IsChecked(false)
+					.IsChecked(ECheckBoxState::Unchecked)
 					.OnCheckStateChanged(LayoutEditor.Get(), &SCustomizableObjectLayoutEditor::OnIgnoreErrorsCheckStateChanged)
 			];
 
@@ -415,7 +415,6 @@ void SCustomizableObjectLayoutEditor::Construct(const FArguments& InArgs)
 	CurrentLayout = nullptr;
 	Node = InArgs._Node;
 	check(Node.Get());
-	
 
 	BindCommands();
 
@@ -508,11 +507,10 @@ void SCustomizableObjectLayoutEditor::Construct(const FArguments& InArgs)
 
 
 	// Array of available grid size options
-	const int32 MaxGridSize = 128;
 	LayoutGridSizes.Empty();
 	MaxLayoutGridSizes.Empty();
 
-	for (int32 Size = 1; Size <= MaxGridSize; Size *= 2)
+	for (int32 Size = 1; Size <= MaxPossibleGridSize; Size *= 2)
 	{
 		LayoutGridSizes.Add(MakeShareable(new FString(FString::FromInt(Size))));
 		MaxLayoutGridSizes.Add(MakeShareable(new FString(FString::Printf(TEXT("%d x %d"), Size, Size))));
@@ -775,7 +773,7 @@ void SCustomizableObjectLayoutEditor::OnOpenGridSizeComboBox()
 		return;
 	}
 
-	const int32 MaxGridSize = FixedStrategyOptionsVisibility() == EVisibility::Visible ? CurrentLayout->GetMaxGridSize().X : 128;
+	const int32 MaxGridSize = FixedStrategyOptionsVisibility() == EVisibility::Visible ? CurrentLayout->GetMaxGridSize().X : MaxPossibleGridSize;
 	const int32 NumOptions = FMath::CountTrailingZeros(MaxGridSize) + 1;
 
 	LayoutGridSizes.SetNum(NumOptions);
@@ -1491,6 +1489,7 @@ void SCustomizableObjectLayoutEditor::FillLayoutComboBoxOptions()
 	if (CurrentLayout)
 	{
 		// Update LayoutGridSizes
+		if (CurrentLayout->PackingStrategy == ECustomizableObjectTextureLayoutPackingStrategy::Fixed)
 		{
 			const int32 MaxGridSize = CurrentLayout->GetMaxGridSize().GetMax();
 			LayoutGridSizes.Empty();
diff --git a/Engine/Plugins/Mutable/Source/CustomizableObjectEditor/Private/MuCOE/SCustomizableObjectLayoutEditor.h b/Engine/Plugins/Mutable/Source/CustomizableObjectEditor/Private/MuCOE/SCustomizableObjectLayoutEditor.h
--- a/Engine/Plugins/Mutable/Source/CustomizableObjectEditor/Private/MuCOE/SCustomizableObjectLayoutEditor.h
+++ b/Engine/Plugins/Mutable/Source/CustomizableObjectEditor/Private/MuCOE/SCustomizableObjectLayoutEditor.h
@@ -250,4 +250,6 @@ private:
 	EVisibility AutoBlocksMergeStrategyVisibility() const;
 	EVisibility FixedStrategyOptionsVisibility() const;
 	EVisibility WarningOptionsVisibility() const;
+	
+	const static int32 MaxPossibleGridSize = 128;
 };
