[Mutable] Fix crash when not being able to get an Image rom.

This can happen due to i/o.

#Synced-CL 50403824

diff --git a/Engine/Plugins/Mutable/Source/MutableRuntime/Internal/MuR/Model.h b/Engine/Plugins/Mutable/Source/MutableRuntime/Internal/MuR/Model.h
--- a/Engine/Plugins/Mutable/Source/MutableRuntime/Internal/MuR/Model.h
+++ b/Engine/Plugins/Mutable/Source/MutableRuntime/Internal/MuR/Model.h
@@ -606,6 +606,12 @@ namespace UE::Mutable::Private
 					MUTABLE_CPUPROFILER_SCOPE(NonTailLODs);
 					TManagedPtr<const FImage> Image = GetImageLOD(Program.ConstantImageLODIndices[FirstLODIndexIndex + LOD]);
 
+					if (!ensureAlwaysMsgf(Image, TEXT("Missing hi-res image rom.")))
+					{
+						Result = nullptr;
+						return;
+					}
+					
 					if (bFormatNeedsDataResize)
 					{
 						ResultImage->DataStorage.ResizeLOD(DestLOD, Image->DataStorage.GetLOD(0).Num());
@@ -625,10 +631,9 @@ namespace UE::Mutable::Private
 					MUTABLE_CPUPROFILER_SCOPE(TailLODs);
 					TManagedPtr<const FImage> TailLODsImage = GetImageLOD(Program.ConstantImageLODIndices[FirstLODIndexIndex + FirstTailLOD]);
 
-					if (!ensureAlwaysMsgf(TailLODsImage, TEXT("Missing tail image rom, setting image to black.")))
+					if (!ensureAlwaysMsgf(TailLODsImage, TEXT("Missing tail image rom.")))
 					{
-						ResultImage->InitToBlack();
-						Result = ResultImage;
+						Result = nullptr;
 						return;
 					}
 
