[Mutable] Fix possible use after free when setting null images and meshes to the cache.

#rnx

#Synced-CL 50488581

diff --git a/Engine/Plugins/Mutable/Source/MutableRuntime/Internal/MuR/System.h b/Engine/Plugins/Mutable/Source/MutableRuntime/Internal/MuR/System.h
--- a/Engine/Plugins/Mutable/Source/MutableRuntime/Internal/MuR/System.h
+++ b/Engine/Plugins/Mutable/Source/MutableRuntime/Internal/MuR/System.h
@@ -448,13 +448,13 @@ namespace UE::Mutable::Private
 		TMemoryTrackedArray<TManagedPtr<const FImage>> StrongImageResults;
 		TMemoryTrackedArray<TManagedPtr<const FMesh>> StrongMeshResults;
 
-		//TMemoryTrackedArray<TManagedPtr<const FResource>> ResourceResults;
-		//TMemoryTrackedArray<TManagedPtr<const FResource>> WeakResourceResults;
-
 		TMemoryTrackedArray<FVector4f> ColorResults;
 		TMemoryTrackedArray<FProjector> ProjectorResults;
 		TMemoryTrackedArray<FMatrix44f> MatrixResults;
 
+		TManagedPtr<const FImage> NullImage = MakeManagedNull<const FImage>();
+		TManagedPtr<const FMesh> NullMesh = MakeManagedNull<const FMesh>();
+
 		/** */
 		inline const ExecutionIndex& GetRangeIndex(uint32_t i)
 		{
@@ -543,10 +543,6 @@ namespace UE::Mutable::Private
 			{
 				StrongMeshPtr = nullptr;
 			}
-
-			StrongImageResults[0] = MakeManagedNull<const FImage>();
-			StrongMeshResults[0] = MakeManagedNull<const FMesh>();
-
 		}
 
 		bool IsValid(const FCacheAddress& Address) const
@@ -1287,16 +1283,7 @@ namespace UE::Mutable::Private
 			Data->DataType = EDataType::Image;
 			Data->IsValueValid = true;
 
-			const TManagedPtr<const FImage>* ValueToStorePtr;
-			if (Value.IsValid())
-			{
-				ValueToStorePtr = &Value;
-			}
-			else
-			{
-				// Set to special null reference that is not valid but can be checked if has been set.
-				ValueToStorePtr = &StrongImageResults[0];  
-			}
+			const TManagedPtr<const FImage>* ValueToStorePtr = Value.IsValid() ? &Value : &NullImage; 
 
 			check(ValueToStorePtr->IsSet());
 
@@ -1327,16 +1314,7 @@ namespace UE::Mutable::Private
 			Data->DataType = EDataType::Mesh;
 			Data->IsValueValid = true;
 
-			const TManagedPtr<const FMesh>* ValueToStorePtr;
-			if (Value.IsValid())
-			{
-				ValueToStorePtr = &Value;
-			}
-			else
-			{
-				// Set to special null reference that is not valid but can be check if is Set.
-				ValueToStorePtr = &StrongMeshResults[0];   
-			}
+			const TManagedPtr<const FMesh>* ValueToStorePtr = Value.IsValid() ? &Value : &NullMesh; 
 
 			check(ValueToStorePtr->IsSet());
 
