[Mutable] Fix checks used to enable unlimited bone influences.
Update MeshFormat code to use the source mesh as base and apply the format on top to avoid copying most variables by hand.

#rnx

#Synced-CL 50489647

diff --git a/Engine/Plugins/Mutable/Source/CustomizableObjectEditor/Private/MuCOE/GenerateMutableSource/GenerateMutableSourceMesh.cpp b/Engine/Plugins/Mutable/Source/CustomizableObjectEditor/Private/MuCOE/GenerateMutableSource/GenerateMutableSourceMesh.cpp
--- a/Engine/Plugins/Mutable/Source/CustomizableObjectEditor/Private/MuCOE/GenerateMutableSource/GenerateMutableSourceMesh.cpp
+++ b/Engine/Plugins/Mutable/Source/CustomizableObjectEditor/Private/MuCOE/GenerateMutableSource/GenerateMutableSourceMesh.cpp
@@ -115,9 +115,9 @@ void SetBuffersFormat(FMutableCompilationContext& Context,
 	// Limit skinning weights if necessary
 	// \todo: make it more flexible to support 3 or 5 or 1 weight, since there is support for this in 4.25
 	int32 MutableBonesPerVertex = 0;
-	if (FGPUBaseSkinVertexFactory::GetUnlimitedBoneInfluences(Context.Options.TargetPlatform))
+	if (!FGPUBaseSkinVertexFactory::GetUnlimitedBoneInfluences(Context.Options.TargetPlatform))
 	{
-		MutableBonesPerVertex = FMath::Min(MaxNumBonesPerVertex, EXTRA_BONE_INFLUENCES);
+		MutableBonesPerVertex = FMath::Min(MaxNumBonesPerVertex, (int32)Context.Options.CustomizableObjectNumBoneInfluences, EXTRA_BONE_INFLUENCES);
 	}
 	else
 	{
diff --git a/Engine/Plugins/Mutable/Source/MutableRuntime/Private/MuR/OpMeshFormat.cpp b/Engine/Plugins/Mutable/Source/MutableRuntime/Private/MuR/OpMeshFormat.cpp
--- a/Engine/Plugins/Mutable/Source/MutableRuntime/Private/MuR/OpMeshFormat.cpp
+++ b/Engine/Plugins/Mutable/Source/MutableRuntime/Private/MuR/OpMeshFormat.cpp
@@ -667,11 +667,15 @@ namespace UE::Mutable::Private
 			return;
 		}
 
-		Result->CopyFrom(*Format);
-		Result->MeshIDPrefix = Source->MeshIDPrefix;
+		EMeshCopyFlags CopyFlags = EMeshCopyFlags::AllFlags;
+		EnumRemoveFlags(CopyFlags, EMeshCopyFlags::WithVertexBuffers | EMeshCopyFlags::WithIndexBuffers | EMeshCopyFlags::WithLayouts);
+
+		Result->CopyFrom(*Source, CopyFlags);
 
 		if (bFormatVertices)
 		{
+			Result->VertexBuffers = Format->VertexBuffers;
+
 			UpdateBoneIndexFormatIfNeeded(Source, Result->GetVertexBuffers());
 
 			FormatBufferSet(Source->GetVertexBuffers(), Result->GetVertexBuffers(),
@@ -684,6 +688,8 @@ namespace UE::Mutable::Private
 
 		if (bFormatIndices)
 		{
+			Result->IndexBuffers = Format->IndexBuffers;
+
 			// \todo Make sure that the vertex indices will fit in this format, or extend it.
 			FormatBufferSet(Source->GetIndexBuffers(), Result->GetIndexBuffers(), bKeepSystemBuffers,
 				bIgnoreMissingChannels, false);
@@ -693,37 +699,12 @@ namespace UE::Mutable::Private
 			Result->IndexBuffers = Source->GetIndexBuffers();
 		}
 
-		// Copy the rest of the data
-		Result->SetSkeleton(Source->GetSkeleton());
-
-		Result->SetPhysicsBody(Source->GetPhysicsBody());
-		Result->PhysicsAssets = Source->PhysicsAssets;
-
-		Result->Layouts.Empty();
+		Result->Layouts.Reset(Source->Layouts.Num());
 		for (const TManagedPtr<const FLayout>& Layout : Source->Layouts)
 		{
 			Result->Layouts.Add(Layout->Clone());
 		}
 
-		Result->Tags = Source->Tags;
-		Result->StreamedResources = Source->StreamedResources;
-		Result->SkeletalMeshes = Source->SkeletalMeshes;
-
-		Result->AdditionalBuffers = Source->AdditionalBuffers;
-
-		Result->BonePoses = Source->BonePoses;
-		Result->BonePosePriority = Source->BonePosePriority;
-		Result->BoneMap = Source->BoneMap;
-
-		Result->SkeletonObjects = Source->SkeletonObjects;
-
-		// A shallow copy is done here, it should not be a problem.
-		Result->AdditionalPhysicsBodies = Source->AdditionalPhysicsBodies;
-		Result->AdditionalPhysicsAssets = Source->AdditionalPhysicsAssets;
-
-		Result->Surfaces = Source->Surfaces;
-		Result->ClothSections = Source->ClothSections;
-
 		Result->ResetStaticFormatFlags();
 		Result->EnsureSurfaceData();
 	}
