[mutable] Passthrough mesh component overlay/override pins.
- Added the option to change the use of the input material pins of the node. Now they can override the material slot (override) or be rendered as an overlay on top of them (overlay).
- Added the option to overlay a material on top of the whole component.
- Added the option to provide to our CustomCreateNode method a pin parameters object, which allows us to set the index where the pin would be created.
- Added a new method to be able to, per pin, determine what subcategory the pin can get set as. It works quite similarly to the already existing method for the main pin category.
- Added controls to be able to set the subtype of the pins that can have multiple subtypes set. This includes a way of changing the subtype using the left mouse button context and also the PinViewer.
- Fixed bug that made the left mouse button menu be displayed while the node was being reconstructed.
- Made the operation type "IN_ADDOVERLAYMATERIAL" be able to handle cases where no slot name is provided (whole component overlay) and cases where a slot name is provided (material slot override).
- Added the "ASTOpAddOverlayMaterial" ASTOp to handle the overlay of materials over a full component or over a single material slot.
- Other small tweaks and changes.

#rnx

#Synced-CL 50530077

diff --git a/Engine/Plugins/Mutable/Source/CustomizableObject/Internal/MuCO/CustomizableObjectCustomVersion.h b/Engine/Plugins/Mutable/Source/CustomizableObject/Internal/MuCO/CustomizableObjectCustomVersion.h
--- a/Engine/Plugins/Mutable/Source/CustomizableObject/Internal/MuCO/CustomizableObjectCustomVersion.h
+++ b/Engine/Plugins/Mutable/Source/CustomizableObject/Internal/MuCO/CustomizableObjectCustomVersion.h
@@ -217,6 +217,10 @@ struct FCustomizableObjectCustomVersion
 		MeshSectionNodeInputPinsAsReferences,
 		
 		WrongVariationNodeInputOrderAfterSerializationFix,
+		
+		PassthroughSkeletalMeshComponentOverrideSubtypeFix,
+		
+		PassthroughSkeletalMeshComponentComponentMaterialOverride,
 
 		// -----<new versions can be added above this line>-------------------------------------------------
 		VersionPlusOne,
diff --git a/Engine/Plugins/Mutable/Source/CustomizableObject/Internal/MuCO/CustomizableObjectInstancePrivate.h b/Engine/Plugins/Mutable/Source/CustomizableObject/Internal/MuCO/CustomizableObjectInstancePrivate.h
--- a/Engine/Plugins/Mutable/Source/CustomizableObject/Internal/MuCO/CustomizableObjectInstancePrivate.h
+++ b/Engine/Plugins/Mutable/Source/CustomizableObject/Internal/MuCO/CustomizableObjectInstancePrivate.h
@@ -107,6 +107,9 @@ struct FCustomizableInstanceComponentData
 
 	UPROPERTY(Transient)
 	TObjectPtr<UMaterialInterface> OverlayMaterial;
+	
+	UPROPERTY(Transient)
+	TArray<TObjectPtr<UMaterialInterface>> OverlayMaterials;
 };
 
 USTRUCT()
diff --git a/Engine/Plugins/Mutable/Source/CustomizableObject/Internal/MuCO/CustomizableObjectSystemPrivate.h b/Engine/Plugins/Mutable/Source/CustomizableObject/Internal/MuCO/CustomizableObjectSystemPrivate.h
--- a/Engine/Plugins/Mutable/Source/CustomizableObject/Internal/MuCO/CustomizableObjectSystemPrivate.h
+++ b/Engine/Plugins/Mutable/Source/CustomizableObject/Internal/MuCO/CustomizableObjectSystemPrivate.h
@@ -322,6 +322,13 @@ struct FInstanceUpdateData
 		
 		uint32 MaterialIndex = INDEX_NONE;
 	};
+	
+	struct FOverlayMaterial
+	{
+		FName SlotName;
+		
+		uint32 MaterialIndex = INDEX_NONE;
+	};
 
 	struct FBone
 	{
@@ -342,6 +349,9 @@ struct FInstanceUpdateData
 		uint16 FirstOverrideMaterial = 0;
 		uint16 OverrideMaterialCount = 0;
 		
+		uint16 FirstOverlayMaterial = 0;
+		uint16 OverlayMaterialCount = 0;
+
 		uint32 OverlayMaterialIndex = INDEX_NONE;
 		
 		UE::Mutable::Private::FSkeletalMeshId SkeletalMeshId;
@@ -353,6 +363,7 @@ struct FInstanceUpdateData
 	TArray<TSharedRef<FComponent>> Components;
 	TArray<TSharedRef<FLOD>> LODs;
 	TArray<FOverrideMaterial> OverrideMaterials;
+	TArray<FOverlayMaterial> OverlayMaterials;
 	TArray<TSharedRef<FMaterial>> Materials;
 	TArray<TSharedRef<FSurface>> Surfaces;
 	TArray<TSharedRef<FImage>> Images;
@@ -380,6 +391,7 @@ struct FInstanceUpdateData
 		Components.Empty();
 		Surfaces.Empty();
 		OverrideMaterials.Empty();
+		OverlayMaterials.Empty();
 		Materials.Empty();
 		Images.Empty();
 		Scalars.Empty();
diff --git a/Engine/Plugins/Mutable/Source/CustomizableObject/Private/MuCO/CustomizableObjectInstance.cpp b/Engine/Plugins/Mutable/Source/CustomizableObject/Private/MuCO/CustomizableObjectInstance.cpp
--- a/Engine/Plugins/Mutable/Source/CustomizableObject/Private/MuCO/CustomizableObjectInstance.cpp
+++ b/Engine/Plugins/Mutable/Source/CustomizableObject/Private/MuCO/CustomizableObjectInstance.cpp
@@ -5309,22 +5309,40 @@ void UCustomizableInstancePrivate::BuildResources(const TSharedRef<FUpdateContex
 
 		// It is not safe to replace the materials of a SkeletalMesh whose resources are initialized. Use overrides instead.
 		const bool bUseOverrideMaterialsOnly = !bIsTransientMesh || (Context->bUseMeshCache && SkeletalMesh->GetResourceForRendering()->IsInitialized());
-
-		UMaterialInterface* OverlayMaterial = nullptr;
-
+		
 		FCustomizableInstanceComponentData& ComponentData = GetComponentData(Component->Id);
 		ComponentData.OverrideMaterials.Reset();
+		ComponentData.OverlayMaterials.Reset();
 		ComponentData.OverlayMaterial = nullptr;
 			
 		if (Component->OverlayMaterialIndex != INDEX_NONE)
 		{
-			FBuildMaterialOut Out = BuildMaterial(Context, Component->OverlayMaterialIndex, !bUVChanged);
-			OverlayMaterial = Out.Material;
+			const FBuildMaterialOut Out = BuildMaterial(Context, Component->OverlayMaterialIndex, !bUVChanged);
 			ComponentData.OverlayMaterial = Out.Material;
 				
 			NotifyUpdateOnInstanceComponent[ComponentIndex] = Out.bNotifyPrimitiveUpdated;
 		}
 			
+		for (int32 OverlayMaterialIndex = 0; OverlayMaterialIndex < Component->OverlayMaterialCount; ++OverlayMaterialIndex)
+		{
+			const FInstanceUpdateData::FOverlayMaterial& OverlayMaterial = Context->InstanceUpdateData.OverlayMaterials[Component->FirstOverlayMaterial + OverlayMaterialIndex];
+				
+			const int32 SlotIndex = SkeletalMesh->GetMaterials().IndexOfByPredicate([&](const FSkeletalMaterial& Material)
+			{
+				return Material.MaterialSlotName == OverlayMaterial.SlotName;
+			});
+				
+			if (SlotIndex != INDEX_NONE)
+			{
+				const FBuildMaterialOut NewMaterial = BuildMaterial(Context, OverlayMaterial.MaterialIndex, !bUVChanged);
+					
+				ComponentData.OverlayMaterials.SetNum(SlotIndex + 1);
+				ComponentData.OverlayMaterials[SlotIndex] = NewMaterial.Material;
+					
+				NotifyUpdateOnInstanceComponent[ComponentIndex] = NewMaterial.bNotifyPrimitiveUpdated;
+			}
+		}
+		
 		for (int32 MaterialIndex = 0; MaterialIndex < Component->OverrideMaterialCount; ++MaterialIndex)
 		{
 			const FInstanceUpdateData::FOverrideMaterial& OverrideMaterial = Context->InstanceUpdateData.OverrideMaterials[Component->FirstOverrideMaterial + MaterialIndex];
@@ -5336,7 +5354,7 @@ void UCustomizableInstancePrivate::BuildResources(const TSharedRef<FUpdateContex
 				
 			if (SlotIndex != INDEX_NONE)
 			{
-				FBuildMaterialOut Out = BuildMaterial(Context, OverrideMaterial.MaterialIndex, !bUVChanged);
+				const FBuildMaterialOut Out = BuildMaterial(Context, OverrideMaterial.MaterialIndex, !bUVChanged);
 					
 				ComponentData.OverrideMaterials.SetNum(SlotIndex + 1);
 				ComponentData.OverrideMaterials[SlotIndex] = Out.Material;
@@ -5347,8 +5365,8 @@ void UCustomizableInstancePrivate::BuildResources(const TSharedRef<FUpdateContex
 
 		if (!bUseOverrideMaterialsOnly)
 		{
-			RecreateRenderStateOnInstanceComponent[ComponentIndex] |= SkeletalMesh->GetOverlayMaterial() != OverlayMaterial;
-			SkeletalMesh->SetOverlayMaterial(OverlayMaterial);
+			RecreateRenderStateOnInstanceComponent[ComponentIndex] |= SkeletalMesh->GetOverlayMaterial() != ComponentData.OverlayMaterial;
+			SkeletalMesh->SetOverlayMaterial(ComponentData.OverlayMaterial);
 		}
 
 		TArray<FSkeletalMaterial> SkeletalMaterials;
diff --git a/Engine/Plugins/Mutable/Source/CustomizableObject/Private/MuCO/CustomizableObjectSystem.cpp b/Engine/Plugins/Mutable/Source/CustomizableObject/Private/MuCO/CustomizableObjectSystem.cpp
--- a/Engine/Plugins/Mutable/Source/CustomizableObject/Private/MuCO/CustomizableObjectSystem.cpp
+++ b/Engine/Plugins/Mutable/Source/CustomizableObject/Private/MuCO/CustomizableObjectSystem.cpp
@@ -1350,14 +1350,27 @@ void UpdateSkeletalMesh(const TSharedRef<FUpdateContextPrivate>& Context)
 		}
 
 
-		UMaterialInterface* OverlayMaterial = nullptr;
-		if (bUseOverrideMaterials && ComponentData)
 		{
-			OverlayMaterial = ComponentData->OverlayMaterial;
+			UMaterialInterface* OverlayMaterial = nullptr;
+			if (bUseOverrideMaterials && ComponentData)
+			{
+				OverlayMaterial = ComponentData->OverlayMaterial;
+			}
+
+			bAttachedParentUpdated |= Parent->GetOverlayMaterial() != OverlayMaterial;
+			Parent->SetOverlayMaterial(OverlayMaterial);
 		}
 
-		bAttachedParentUpdated |= Parent->GetOverlayMaterial() != OverlayMaterial;
-		Parent->SetOverlayMaterial(OverlayMaterial);
+		if (bUseOverrideMaterials && ComponentData)
+		{
+			for (int32 OverlayMaterialIndex = 0; OverlayMaterialIndex < ComponentData->OverlayMaterials.Num(); ++OverlayMaterialIndex)
+			{
+				UMaterialInterface* OverlayMaterial = ComponentData->OverlayMaterials[OverlayMaterialIndex]; 
+					
+				bAttachedParentUpdated |= Parent->GetOverlayMaterial(true, OverlayMaterialIndex) != OverlayMaterial;
+				Parent->SetOverlayMaterial(OverlayMaterial, true, OverlayMaterialIndex);
+			}
+		}
 		
 		if (CustomizableObjectInstancePrivateData->HasCOInstanceFlags(ReplacePhysicsAssets) &&
 			SkeletalMesh &&
@@ -2206,6 +2219,24 @@ namespace impl
 				Material->MaterialId = MaterialId;
 			}
 			
+			// OVERLAY MATERIALS
+			const int32 OverlayMaterialCount = Instance->GetOverlayMaterialCount(ComponentIndex);
+			Component->FirstOverlayMaterial = OperationData->InstanceUpdateData.OverlayMaterials.Num();
+			for (int32 MaterialIndex = 0; MaterialIndex < OverlayMaterialCount; ++MaterialIndex)
+			{
+				if (UE::Mutable::Private::FMaterialId SlotMaterialId = Instance->GetOverlayMaterialId(ComponentIndex, MaterialIndex))
+				{
+					FInstanceUpdateData::FOverlayMaterial& OverlayMaterial = OperationData->InstanceUpdateData.OverlayMaterials.AddDefaulted_GetRef();
+					++Component->OverlayMaterialCount;
+					OverlayMaterial.SlotName = Instance->GetOverlayMaterialSlotSlotName(ComponentIndex, MaterialIndex);
+					
+					OverlayMaterial.MaterialIndex = OperationData->InstanceUpdateData.Materials.Num();
+					TSharedRef<FInstanceUpdateData::FMaterial> Material = OperationData->InstanceUpdateData.Materials.Add_GetRef(MakeShared<FInstanceUpdateData::FMaterial>());
+					Material->MaterialId = SlotMaterialId;
+				}	
+			}
+			
+			// OVERRIDE MATERIALS
 			const int32 OverrideMaterialCount = Instance->GetOverrideMaterialCount(ComponentIndex);
 			Component->FirstOverrideMaterial = OperationData->InstanceUpdateData.OverrideMaterials.Num();
 			for (int32 MaterialIndex = 0; MaterialIndex < OverrideMaterialCount; ++MaterialIndex)
diff --git a/Engine/Plugins/Mutable/Source/CustomizableObjectEditor/Private/MuCOE/EdGraphSchema_CustomizableObject.cpp b/Engine/Plugins/Mutable/Source/CustomizableObjectEditor/Private/MuCOE/EdGraphSchema_CustomizableObject.cpp
--- a/Engine/Plugins/Mutable/Source/CustomizableObjectEditor/Private/MuCOE/EdGraphSchema_CustomizableObject.cpp
+++ b/Engine/Plugins/Mutable/Source/CustomizableObjectEditor/Private/MuCOE/EdGraphSchema_CustomizableObject.cpp
@@ -317,6 +317,11 @@ const FName UEdGraphSchema_CustomizableObject::PC_PoseAsset("poseAsset");
 const FName UEdGraphSchema_CustomizableObject::PC_Transform("transform");
 const FName UEdGraphSchema_CustomizableObject::PC_String("string");
 
+// DO NOT change the values because it will break the external pin nodes!
+const FName UEdGraphSchema_CustomizableObject::PSC_Material_Override("overrideMaterial");
+const FName UEdGraphSchema_CustomizableObject::PSC_Material_Overlay("overlayMaterial");
+
+
 // Add more pin types to this array if needed
 const TArray<FName> UEdGraphSchema_CustomizableObject::SupportedMacroPinTypes = { PC_Object, PC_Component, PC_MeshSection, PC_Modifier, PC_Mesh, PC_Texture,
 	PC_Texture_Passthrough, PC_Projector, PC_GroupProjector, PC_Color, PC_Float, PC_Enum, PC_Stack, PC_Material, PC_PoseAsset, PC_Transform, PC_String };
@@ -898,8 +903,7 @@ FLinearColor UEdGraphSchema_CustomizableObject::GetPinTypeColor(const FName& Typ
 			const FLinearColor Color = AssetDefinitionRegistry->GetAssetDefinitionForClass(USkeletalMesh::StaticClass())->GetAssetColor();
 			PinCategoryColors.Add(PC_SkeletalMesh, Color);
 			PinCategoryColors.Add(PC_SkeletalMesh_Passthrough, Color); 	
-			PinCategoryColors.Add(PC_SkeletalMesh_PassthroughParameter, Color); 	
-			
+			PinCategoryColors.Add(PC_SkeletalMesh_PassthroughParameter, Color);
 		}
 
 		{
@@ -1147,17 +1151,44 @@ void UEdGraphSchema_CustomizableObject::GetContextMenuActions(UToolMenu* Menu, U
 				}
 				
 				const TArray<FName> AllowedTypes = Node->GetPinAllowedTypes(*Pin);
-				if (AllowedTypes.Num())
+				if (AllowedTypes.Num() > 1)
 				{
 					FEdGraphPinReference PinReference = FEdGraphPinReference(const_cast<UEdGraphPin*>(Pin));
 					TSharedRef<SMutablePinTypeSelector> TypeSelectionSlate = SNew(SMutablePinTypeSelector).PinReference(PinReference);
 					
+					// In order to not access to invalid node data close the widget containing the SMutablePinTypeSelector
+					Node->PostReconstructNodeDelegate.AddLambda([WeakMenu = TypeSelectionSlate->AsWeak()](){
+						if (TSharedPtr<SWidget> Menu = WeakMenu.Pin())
+						{
+							FSlateApplication::Get().DismissMenuByWidget(Menu.ToSharedRef());
+						}
+					});
+					
 					FToolMenuSection& Section = Menu->FindOrAddSection("EdGraphSchemaPinActions");
 					Section.AddEntry(FToolMenuEntry::InitWidget("Pin Type Selector", TypeSelectionSlate,  LOCTEXT("SetPinType_Label", "Pin Type"), true));
 				}
 				
+				const TArray<FName> AllowedSubTypes = Node->GetPinAllowedSubTypes(*Pin);
+				if (AllowedSubTypes.Num() > 1)
+				{
+					FEdGraphPinReference PinReference = FEdGraphPinReference(const_cast<UEdGraphPin*>(Pin));
+					TSharedRef<SMutablePinSubTypeSelector> SubTypeSelectionSlate = SNew(SMutablePinSubTypeSelector).PinReference(PinReference);
+					
+					// In order to not access to invalid node data close the widget containing the SubTypeSelectionSlate
+					Node->PostReconstructNodeDelegate.AddLambda([WeakMenu = SubTypeSelectionSlate->AsWeak()](){
+						if (TSharedPtr<SWidget> Menu = WeakMenu.Pin())
+						{
+							FSlateApplication::Get().DismissMenuByWidget(Menu.ToSharedRef());
+						}
+					});
+					
+					FToolMenuSection& Section = Menu->FindOrAddSection("EdGraphSchemaPinActions");
+					Section.AddEntry(FToolMenuEntry::InitWidget("Pin Subtype Selector", SubTypeSelectionSlate,  LOCTEXT("SetPinSubtype_Label", "Pin Subtype"), true));
+				}
+				
 				if (TSharedPtr<IDetailsView> Widget = Node->CustomizePinDetails(*Pin))
 				{
+					// In order to not access to invalid node data close the widget containing the custom widget
 					Node->PostReconstructNodeDelegate.AddLambda([WeakMenu = Widget->AsWeak()](){
 						if (TSharedPtr<SWidget> Menu = WeakMenu.Pin())
 						{
@@ -1653,6 +1684,20 @@ FName UEdGraphSchema_CustomizableObject::GetPinCategoryName(const UScriptStruct*
 	}
 }
 
+FName UEdGraphSchema_CustomizableObject::GetPinSubCategoryName(const FName& PinSubCategory)
+{
+	if (PinSubCategory == PSC_Material_Overlay)
+	{
+		return TEXT("Overlay");
+	}
+	else if (PinSubCategory == PSC_Material_Override)
+	{
+		return TEXT("Override");
+	}
+	
+	return {};
+}
+
 
 FText UEdGraphSchema_CustomizableObject::GetPinCategoryFriendlyName(const FName& PinCategory)
 {
@@ -1776,6 +1821,21 @@ FText UEdGraphSchema_CustomizableObject::GetPinCategoryFriendlyName(const UScrip
 }
 
 
+FText UEdGraphSchema_CustomizableObject::GetPinSubCategoryFriendlyName(const FName& PinSubcategory)
+{
+	if (PinSubcategory == PSC_Material_Overlay)
+	{
+		return LOCTEXT("Material_Overlay_Pin_Subcategory", "Overlay");
+	}
+	else if (PinSubcategory == PSC_Material_Override)
+	{
+		return LOCTEXT("Material_Override_Pin_Subcategory", "Override");
+	}
+	
+	return {};
+}
+
+
 bool UEdGraphSchema_CustomizableObject::IsPassthrough(const FName& PinCategory)
 {
 	return PinCategory == PC_Texture_Passthrough || 
@@ -1812,4 +1872,19 @@ bool UEdGraphSchema_CustomizableObject::IsTypeNative(const UScriptStruct* Type)
 }
 
 
+FText UEdGraphSchema_CustomizableObject::GetSubCategoryTooltipText(FName PinSubType)
+{
+	if (PinSubType == PSC_Material_Override)
+	{
+		return LOCTEXT("PSC_Material_Override_Tooltip", "Replaces one material with another one.");
+	}
+	else if (PinSubType == PSC_Material_Overlay)
+	{
+		return LOCTEXT("PSC_Material_Overlay_Tooltip", "Renders a secondary material on top of the already set one.");
+	}
+
+	return {};
+}
+
+
 #undef LOCTEXT_NAMESPACE
diff --git a/Engine/Plugins/Mutable/Source/CustomizableObjectEditor/Private/MuCOE/GenerateMutableSource/GenerateMutableSourceComponent.cpp b/Engine/Plugins/Mutable/Source/CustomizableObjectEditor/Private/MuCOE/GenerateMutableSource/GenerateMutableSourceComponent.cpp
--- a/Engine/Plugins/Mutable/Source/CustomizableObjectEditor/Private/MuCOE/GenerateMutableSource/GenerateMutableSourceComponent.cpp
+++ b/Engine/Plugins/Mutable/Source/CustomizableObjectEditor/Private/MuCOE/GenerateMutableSource/GenerateMutableSourceComponent.cpp
@@ -31,6 +31,7 @@
 
 #include "Rendering/SkeletalMeshLODModel.h"
 #include "Interfaces/ITargetPlatform.h"
+#include "MuCOE/EdGraphSchema_CustomizableObject.h"
 
 
 #define LOCTEXT_NAMESPACE "CustomizableObjectEditor"
@@ -155,32 +156,69 @@ UE::Mutable::Private::Ptr<UE::Mutable::Private::NodeComponent> GenerateMutableSo
 		ComponentNode->Id = GenerationContext.ComponentNames.Add(ComponentName);
 		ComponentNode->SkeletalMesh = UE::Mutable::Private::GenerateMutableSourceSkeletalMesh(Pin, GenerationContext);
 
-		// Override Material slots
-		for (const FEdGraphPinReference& OverrideMaterialPin : TypedComponentPassthroughMesh->MaterialSlotOverridePins)
+		// Overlay material (applied to the whole component)
+		if (UEdGraphPin* OverlayMaterialPin = TypedComponentPassthroughMesh->OverlayMaterialPin.Get())
+		{
+			if (const UEdGraphPin* ConnectedPin = FollowInputPin(*OverlayMaterialPin))
+			{
+				UE::Mutable::Private::Ptr<UE::Mutable::Private::NodeMaterial> MaterialNode = 
+					GenerateMutableSourceMaterial(ConnectedPin,GenerationContext, FGenerationMaterialOptions());
+				
+				ComponentNode->OverlayMaterial = MaterialNode;
+			}
+		}
+		
+		
+		// Override/Overlay Material slots
+		TArray<FName> ProcessedSlotNames;
+		for (const FEdGraphPinReference& InputMaterialPinReference : TypedComponentPassthroughMesh->MaterialSlotPins)
 		{
 			// Pin Data
 			FName TargetMaterialSlotName;
-			if (const UEdGraphPin* MaterialOverridePinPointer = OverrideMaterialPin.Get())
+			if (const UEdGraphPin* InputMaterialPin = InputMaterialPinReference.Get())
 			{
-				const UCONodeComponentPassthroughMeshMaterialOverridePinData& OverrideMaterialPinData =
-					TypedComponentPassthroughMesh->GetPinData<UCONodeComponentPassthroughMeshMaterialOverridePinData>(*MaterialOverridePinPointer);
+				const UCONodeComponentPassthroughMeshMaterialPinData& OverrideMaterialPinData =
+					TypedComponentPassthroughMesh->GetPinData<UCONodeComponentPassthroughMeshMaterialPinData>(*InputMaterialPin);
 
 				TargetMaterialSlotName = OverrideMaterialPinData.TargetMaterialSlotName;
 			}
-
+			
 			// Material node
 			UE::Mutable::Private::Ptr<UE::Mutable::Private::NodeMaterial> MaterialNode = nullptr;
-			if (const UEdGraphPin* MaterialOverridePin = OverrideMaterialPin.Get())
+			if (const UEdGraphPin* MaterialOverridePin = InputMaterialPinReference.Get())
 			{
 				if (const UEdGraphPin* ConnectedPin = FollowInputPin(*MaterialOverridePin))
 				{
+					// Skip invalid slot names (None)
+					if (TargetMaterialSlotName.IsNone())
+					{
+						GenerationContext.Log(LOCTEXT("NoneMaterialSlotName", "'None' is not a valid Material slot name."), Node, EMessageSeverity::Warning);
+						continue;
+					}
+					
+					// Skip repeated slot names
+					if (ProcessedSlotNames.Contains(TargetMaterialSlotName))
+					{
+						GenerationContext.Log(FText::Format(LOCTEXT("RepeatedMaterialSlotName", "The slot name '{0}' has been defined multiple times. Only the first entry will be used."), FText::FromName(TargetMaterialSlotName)), Node, EMessageSeverity::Warning);
+						continue;
+					}
+				
 					MaterialNode = GenerateMutableSourceMaterial(ConnectedPin, GenerationContext, FGenerationMaterialOptions());
 				}
 			}
 
 			if (MaterialNode)
 			{
-				ComponentNode->OverrideMaterials.Add(TargetMaterialSlotName, MaterialNode);
+				if (InputMaterialPinReference.Get()->PinType.PinSubCategory == UEdGraphSchema_CustomizableObject::PSC_Material_Override)
+				{
+					ComponentNode->OverrideMaterials.Add(TargetMaterialSlotName, MaterialNode);
+				}
+				else if (InputMaterialPinReference.Get()->PinType.PinSubCategory == UEdGraphSchema_CustomizableObject::PSC_Material_Overlay)
+				{
+					ComponentNode->OverlayMaterials.Add(TargetMaterialSlotName, MaterialNode);
+				}
+				
+				ProcessedSlotNames.Add(TargetMaterialSlotName);
 			}
 		}
 		
diff --git a/Engine/Plugins/Mutable/Source/CustomizableObjectEditor/Private/MuCOE/Nodes/CustomizableObjectNode.cpp b/Engine/Plugins/Mutable/Source/CustomizableObjectEditor/Private/MuCOE/Nodes/CustomizableObjectNode.cpp
--- a/Engine/Plugins/Mutable/Source/CustomizableObjectEditor/Private/MuCOE/Nodes/CustomizableObjectNode.cpp
+++ b/Engine/Plugins/Mutable/Source/CustomizableObjectEditor/Private/MuCOE/Nodes/CustomizableObjectNode.cpp
@@ -52,6 +52,7 @@ UEdGraphPin* UCustomizableObjectNode::CustomCreatePin(EEdGraphPinDirection Direc
 UEdGraphPin* UCustomizableObjectNode::CustomCreatePin(EEdGraphPinDirection Direction, const FName& Category, const FName& Name, const FText& FriendlyName, UCustomizableObjectNodePinData* PinData)
 {
 	UEdGraphPin* Pin = CreatePin(Direction, Category, Name);
+	
 	Pin->PinFriendlyName = FriendlyName;
 	
 	if (PinData)
@@ -814,6 +815,11 @@ TArray<FName> UCustomizableObjectNode::GetPinAllowedTypes(const UEdGraphPin& Pin
 	return {};
 }
 
+TArray<FName> UCustomizableObjectNode::GetPinAllowedSubTypes(const UEdGraphPin& Pin) const
+{
+	return {};
+}
+
 
 bool UCustomizableObjectNode::IsCompatibleWithGraph(const UEdGraph* TargetGraph) const
 {
diff --git a/Engine/Plugins/Mutable/Source/CustomizableObjectEditor/Private/MuCOE/Nodes/CustomizableObjectNodeComponentPassthroughMesh.cpp b/Engine/Plugins/Mutable/Source/CustomizableObjectEditor/Private/MuCOE/Nodes/CustomizableObjectNodeComponentPassthroughMesh.cpp
--- a/Engine/Plugins/Mutable/Source/CustomizableObjectEditor/Private/MuCOE/Nodes/CustomizableObjectNodeComponentPassthroughMesh.cpp
+++ b/Engine/Plugins/Mutable/Source/CustomizableObjectEditor/Private/MuCOE/Nodes/CustomizableObjectNodeComponentPassthroughMesh.cpp
@@ -3,6 +3,7 @@
 #include "MuCOE/Nodes/CustomizableObjectNodeComponentPassthroughMesh.h"
 
 #include "CONodeMaterialBreak.h"
+// #include "PropertyEditorModule.h"
 #include "MuCO/CustomizableObjectCustomVersion.h"
 #include "MuCO/LoadUtils.h"
 #include "MuCOE/EdGraphSchema_CustomizableObject.h"
@@ -13,7 +14,6 @@
 
 #include UE_INLINE_GENERATED_CPP_BY_NAME(CustomizableObjectNodeComponentPassthroughMesh)
 
-class UCustomizableObjectNodeRemapPinsByName;
 class UObject;
 
 
@@ -40,6 +40,38 @@ void UCustomizableObjectNodeComponentPassthroughMesh::BackwardsCompatibleFixup(i
 			SkeletalMeshPin = CustomCreatePin(EGPD_Input, UEdGraphSchema_CustomizableObject::PC_SkeletalMesh_Passthrough);
 		}
 	}
+	else if (CustomizableObjectCustomVersion == FCustomizableObjectCustomVersion::PassthroughSkeletalMeshComponentOverrideSubtypeFix)
+	{
+		// Set the subcategory of the pins that do require one defined. We use the "Override" one as that is the behaviour the pins had before the option to make them able to be set as "overlay" was added.
+		for (FEdGraphPinReference MaterialPinReference : MaterialSlotPins)
+		{
+			UEdGraphPin* MaterialPin = MaterialPinReference.Get();
+			if (!MaterialPin)
+			{
+				checkNoEntry();
+				continue;
+			}
+			
+			MaterialPin->PinType.PinSubCategory = UEdGraphSchema_CustomizableObject::PSC_Material_Override;
+			MaterialPin->PinFriendlyName = FText::FromString("Override Material");
+			MaterialPin->PinName = FName("Override Material");
+		}
+	}
+	else if (CustomizableObjectCustomVersion == FCustomizableObjectCustomVersion::PassthroughSkeletalMeshComponentComponentMaterialOverride)
+	{
+		// Add an overlay material pin after the PassthroughSkeletalMeshPin so it gets to be on top of the Material Override/Overlay pins.
+		
+		const UEdGraphPin* PassthroughSkeletalMeshPin = SkeletalMeshPin.Get();
+		const int32 PassthroughMeshPinIndex = Pins.IndexOfByKey(PassthroughSkeletalMeshPin);
+		
+		check(PassthroughMeshPinIndex != INDEX_NONE);
+		
+		UEdGraphPin* NewMaterialOverlayPin = CustomCreatePin(EGPD_Input, UEdGraphSchema_CustomizableObject::PC_Material, "Overlay Material" , FText::FromString("Overlay Material"));
+		Pins.Remove(NewMaterialOverlayPin);
+		Pins.Insert(NewMaterialOverlayPin, PassthroughMeshPinIndex + 1);
+		
+		OverlayMaterialPin = NewMaterialOverlayPin;
+	}
 }
 
 
@@ -48,24 +80,34 @@ void UCustomizableObjectNodeComponentPassthroughMesh::AllocateDefaultPins(UCusto
 	Super::AllocateDefaultPins(RemapPins);
 
 	SkeletalMeshPin = CustomCreatePin(EGPD_Input, UEdGraphSchema_CustomizableObject::PC_SkeletalMesh_Passthrough);
-
+	
+	OverlayMaterialPin = CustomCreatePin(EGPD_Input, UEdGraphSchema_CustomizableObject::PC_Material, "Overlay Material" , FText::FromString("Overlay Material"), nullptr);
+	
 	// Get old pins
 	const TArray<UEdGraphPin*> AllPins = GetAllNonOrphanPins();
 	
 	// Regenerate the new pins from the old pins
-	MaterialSlotOverridePins.Reset();
+	MaterialSlotPins.Reset();
 	for (const UEdGraphPin* OldPin : AllPins)
 	{
+		check(OldPin)
+		
 		if (OldPin->Direction != EEdGraphPinDirection::EGPD_Input)
 		{
 			continue;
 		}
 		
-		if (UCONodeComponentPassthroughMeshMaterialOverridePinData* PinData = Cast<UCONodeComponentPassthroughMeshMaterialOverridePinData>(GetPinData(*OldPin)))
+		if (UCONodeComponentPassthroughMeshMaterialPinData* OldPinData = Cast<UCONodeComponentPassthroughMeshMaterialPinData>(GetPinData(*OldPin)))
 		{
+			UCONodeComponentPassthroughMeshMaterialPinData* PinData = NewObject<UCONodeComponentPassthroughMeshMaterialPinData>(this);
+			PinData->TargetMaterialSlotName = OldPinData->TargetMaterialSlotName;
+			
+			// The SubCategory could have changed after a node refresh
+			const FName OldPinSubCategory = OldPin->PinType.PinSubCategory;
+			
 			// Use the same pin data as the old pin since the new one is a copy
-			const UEdGraphPin* NewPin = CreateMaterialSlotOverridePin(PinData);
-			MaterialSlotOverridePins.Add(NewPin);
+			const UEdGraphPin* NewPin = CreateMaterialSlotPin(OldPinSubCategory, PinData);
+			MaterialSlotPins.Add(NewPin);
 		}
 	}
 }
@@ -96,11 +138,12 @@ void UCustomizableObjectNodeComponentPassthroughMesh::AddPinFromUI()
 	FScopedTransaction Transaction(LOCTEXT("UCustomizableObjectNodeComponentPassthroughMesh_AddPinFromNodeUITransaction", "Add Pin from Node UI"));
 	Modify();
 
-	UCONodeComponentPassthroughMeshMaterialOverridePinData* PinData = NewObject<UCONodeComponentPassthroughMeshMaterialOverridePinData>(this);
+	UCONodeComponentPassthroughMeshMaterialPinData* PinData = NewObject<UCONodeComponentPassthroughMeshMaterialPinData>(this);
 	
-	if (const UEdGraphPin* NewPin = CreateMaterialSlotOverridePin(PinData))
+	// By default create Override subtype pins
+	if (const UEdGraphPin* NewPin = CreateMaterialSlotPin(UEdGraphSchema_CustomizableObject::PSC_Material_Override, PinData))
 	{		
-		MaterialSlotOverridePins.Add(NewPin);
+		MaterialSlotPins.Add(NewPin);
 		GetGraph()->NotifyNodeChanged(this);
 	}
 }
@@ -114,20 +157,20 @@ bool UCustomizableObjectNodeComponentPassthroughMesh::CanPinBeRemoved(const UEdG
 		return false;
 	}
 
-	return IsMaterialSlotOverridePin(Pin);
+	return IsMaterialSlotPin(Pin);
 }
 
 
 
 bool UCustomizableObjectNodeComponentPassthroughMesh::CanRenamePin(const UEdGraphPin& Pin) const
 {
-	return IsMaterialSlotOverridePin(Pin) || Super::CanRenamePin(Pin);
+	return IsMaterialSlotPin(Pin) || Super::CanRenamePin(Pin);
 }
 
 
 EEditablePinNameBoxVisibilityPolicy UCustomizableObjectNodeComponentPassthroughMesh::GetEditablePinNameVisibilityPolicy(const UEdGraphPin& Pin) const
 {
-	if (IsMaterialSlotOverridePin(Pin))
+	if (IsMaterialSlotPin(Pin))
 	{
 		return EEditablePinNameBoxVisibilityPolicy::ALWAYS_VISIBLE;
 	}
@@ -138,7 +181,7 @@ EEditablePinNameBoxVisibilityPolicy UCustomizableObjectNodeComponentPassthroughM
 
 bool UCustomizableObjectNodeComponentPassthroughMesh::ShouldPinViewerShowPinEditableName(const UEdGraphPin& Pin) const
 {
-	if (IsMaterialSlotOverridePin(Pin))
+	if (IsMaterialSlotPin(Pin))
 	{
 		return true;
 	}
@@ -149,14 +192,14 @@ bool UCustomizableObjectNodeComponentPassthroughMesh::ShouldPinViewerShowPinEdit
 
 void UCustomizableObjectNodeComponentPassthroughMesh::SetPinEditableName(const UEdGraphPin& Pin, const FText& Value)
 {
-	if (IsMaterialSlotOverridePin(Pin))
+	if (IsMaterialSlotPin(Pin))
 	{
-		if (UCONodeComponentPassthroughMeshMaterialOverridePinData* PinData = Cast<UCONodeComponentPassthroughMeshMaterialOverridePinData>(GetPinData(Pin)))
+		if (UCONodeComponentPassthroughMeshMaterialPinData* PinData = CastChecked<UCONodeComponentPassthroughMeshMaterialPinData>(GetPinData(Pin)))
 		{
 			const FName NewName = FName{Value.ToString()};
 			if (NewName != PinData->TargetMaterialSlotName)		
 			{
-				FScopedTransaction LocalTransaction(LOCTEXT("UCustomizableObjectNodeComponentPassthroughMesh_SetOverrideMaterialPinEditableNameTransaction", "Change OverrideMaterial PinName"));
+				FScopedTransaction LocalTransaction(LOCTEXT("UCustomizableObjectNodeComponentPassthroughMesh_SetMaterialPinEditableNameTransaction", "Change Material PinName"));
 				PinData->Modify();
 				
 				PinData->TargetMaterialSlotName = NewName;
@@ -175,9 +218,9 @@ void UCustomizableObjectNodeComponentPassthroughMesh::SetPinEditableName(const U
 
 FText UCustomizableObjectNodeComponentPassthroughMesh::GetPinEditableName(const UEdGraphPin& Pin) const
 {
-	if (IsMaterialSlotOverridePin(Pin))
+	if (IsMaterialSlotPin(Pin))
 	{
-		if (const UCONodeComponentPassthroughMeshMaterialOverridePinData* PinData = Cast<UCONodeComponentPassthroughMeshMaterialOverridePinData>(GetPinData(Pin)))
+		if (const UCONodeComponentPassthroughMeshMaterialPinData* PinData = Cast<UCONodeComponentPassthroughMeshMaterialPinData>(GetPinData(Pin)))
 		{
 			return FText::FromName(PinData->TargetMaterialSlotName);
 		}
@@ -187,28 +230,48 @@ FText UCustomizableObjectNodeComponentPassthroughMesh::GetPinEditableName(const
 }
 
 
-bool UCustomizableObjectNodeComponentPassthroughMesh::IsMaterialSlotOverridePin(const UEdGraphPin& Pin) const
+bool UCustomizableObjectNodeComponentPassthroughMesh::IsMaterialSlotPin(const UEdGraphPin& Pin) const
 {
-	return Pin.Direction == EEdGraphPinDirection::EGPD_Input && Pin.PinType.PinCategory == UEdGraphSchema_CustomizableObject::PC_Material;
+	return MaterialSlotPins.Contains(&Pin);
 }
 
 
 UCustomizableObjectNodeRemapPins* UCustomizableObjectNodeComponentPassthroughMesh::CreateRemapPinsDefault() const
 {
-	return Super::CreateRemapPinsByPosition();
+	return NewObject<UCustomizableObjectNodeRemapPinsByPosition>();
 }
 
 
 bool UCustomizableObjectNodeComponentPassthroughMesh::CustomRemovePin(UEdGraphPin& Pin)
 {
-	MaterialSlotOverridePins.Remove(&Pin);
+	MaterialSlotPins.Remove(&Pin);
 	return Super::CustomRemovePin(Pin);
 }
 
 
-UEdGraphPin* UCustomizableObjectNodeComponentPassthroughMesh::CreateMaterialSlotOverridePin(UCONodeComponentPassthroughMeshMaterialOverridePinData* InPinData)
+TArray<FName> UCustomizableObjectNodeComponentPassthroughMesh::GetPinAllowedSubTypes(const UEdGraphPin& Pin) const
+{
+	if (IsMaterialSlotPin(Pin))
+	{
+		const TArray<FName> PinAllowedSubTypes =  {
+			UEdGraphSchema_CustomizableObject::PSC_Material_Overlay, 
+			UEdGraphSchema_CustomizableObject::PSC_Material_Override
+		};
+		return PinAllowedSubTypes;
+	}
+	
+	return Super::GetPinAllowedSubTypes(Pin);
+}
+
+
+UEdGraphPin* UCustomizableObjectNodeComponentPassthroughMesh::CreateMaterialSlotPin(const FName& PinSubCategory, UCONodeComponentPassthroughMeshMaterialPinData* InPinData)
 {
-	return CustomCreatePin(EGPD_Input, UEdGraphSchema_CustomizableObject::PC_Material, TEXT("Override Material"), LOCTEXT("OverrideMaterial", "Override Material"), InPinData);
+	const FName NewPinName = UEdGraphSchema_CustomizableObject::GetPinSubCategoryName(PinSubCategory);
+	const FText NewPinFriendlyName =  FText::FromString( UEdGraphSchema_CustomizableObject::GetPinSubCategoryFriendlyName(PinSubCategory).ToString() + " Material");
+	
+	UEdGraphPin* NewMaterialPin = CustomCreatePin(EGPD_Input, UEdGraphSchema_CustomizableObject::PC_Material, NewPinName, NewPinFriendlyName, InPinData);
+	NewMaterialPin->PinType.PinSubCategory = PinSubCategory;
+	return NewMaterialPin;
 }
 
 
diff --git a/Engine/Plugins/Mutable/Source/CustomizableObjectEditor/Private/MuCOE/Nodes/CustomizableObjectNodeComponentPassthroughMesh.h b/Engine/Plugins/Mutable/Source/CustomizableObjectEditor/Private/MuCOE/Nodes/CustomizableObjectNodeComponentPassthroughMesh.h
--- a/Engine/Plugins/Mutable/Source/CustomizableObjectEditor/Private/MuCOE/Nodes/CustomizableObjectNodeComponentPassthroughMesh.h
+++ b/Engine/Plugins/Mutable/Source/CustomizableObjectEditor/Private/MuCOE/Nodes/CustomizableObjectNodeComponentPassthroughMesh.h
@@ -27,7 +27,7 @@ class UEdGraphPin;
 
 
 UCLASS(MinimalAPI)
-class UCONodeComponentPassthroughMeshMaterialOverridePinData : public UCustomizableObjectNodePinData
+class UCONodeComponentPassthroughMeshMaterialPinData : public UCustomizableObjectNodePinData
 {
 	GENERATED_BODY()
 
@@ -68,16 +68,22 @@ public:
 	virtual FText GetPinEditableName(const UEdGraphPin& Pin) const override;
 	virtual bool CustomRemovePin(UEdGraphPin& Pin) override;
 	
+	// Override / Overlay pin configuration
+	virtual TArray<FName> GetPinAllowedSubTypes(const UEdGraphPin& Pin) const override;
+	
 	UPROPERTY()
 	FEdGraphPinReference SkeletalMeshPin;
+	
+	UPROPERTY()
+	FEdGraphPinReference OverlayMaterialPin;
 
 	UPROPERTY()
-	TArray<FEdGraphPinReference> MaterialSlotOverridePins;
+	TArray<FEdGraphPinReference> MaterialSlotPins;
 
 private:
-	UEdGraphPin* CreateMaterialSlotOverridePin(UCONodeComponentPassthroughMeshMaterialOverridePinData* InPinData);
+	UEdGraphPin* CreateMaterialSlotPin(const FName& PinSubCategory, UCONodeComponentPassthroughMeshMaterialPinData* InPinData);
 
-	bool IsMaterialSlotOverridePin(const UEdGraphPin& Pin) const;
+	bool IsMaterialSlotPin(const UEdGraphPin& Pin) const;
 
 protected:
 	virtual UCustomizableObjectNodeRemapPins* CreateRemapPinsDefault() const override;
diff --git a/Engine/Plugins/Mutable/Source/CustomizableObjectEditor/Private/MuCOE/PinViewer/SPinViewer.cpp b/Engine/Plugins/Mutable/Source/CustomizableObjectEditor/Private/MuCOE/PinViewer/SPinViewer.cpp
--- a/Engine/Plugins/Mutable/Source/CustomizableObjectEditor/Private/MuCOE/PinViewer/SPinViewer.cpp
+++ b/Engine/Plugins/Mutable/Source/CustomizableObjectEditor/Private/MuCOE/PinViewer/SPinViewer.cpp
@@ -22,6 +22,8 @@ const FName SPinViewer::COLUMN_NAME(TEXT("Name"));
 
 const FName SPinViewer::COLUMN_TYPE(TEXT("Type"));
 
+const FName SPinViewer::COLUMN_SUBTYPE(TEXT("Subtype"));
+
 const FName SPinViewer::COLUMN_VISIBILITY(TEXT("Visible"));
 
 
@@ -113,6 +115,45 @@ void Sort(const FName& Method, const EColumnSortMode::Type& SortMode, TArray<TSh
 				});
 		}
 	}
+	else if (Method == SPinViewer::COLUMN_SUBTYPE)
+	{
+		if (SortMode == EColumnSortMode::Descending)
+		{
+			Pins.Sort([](const TSharedPtr<FEdGraphPinReference>& PinReferenceA, const TSharedPtr<FEdGraphPinReference>& PinReferenceB)
+				{
+					if (PinReferenceA->Get()->PinType.PinSubCategory == PinReferenceB->Get()->PinType.PinSubCategory)
+					{
+						return SortPinsByNameDesc(PinReferenceA, PinReferenceB);
+					}
+					else if (PinReferenceA->Get()->PinType.PinSubCategory.FastLess(PinReferenceB->Get()->PinType.PinSubCategory))
+					{
+						return false;
+					}
+					else
+					{
+						return true;
+					}
+				});
+		}
+		else
+		{
+			Pins.Sort([](const TSharedPtr<FEdGraphPinReference>& PinReferenceA, const TSharedPtr<FEdGraphPinReference>& PinReferenceB)
+				{
+					if (PinReferenceA->Get()->PinType.PinSubCategory == PinReferenceB->Get()->PinType.PinSubCategory)
+					{
+						return SortPinsByNameAsc(PinReferenceA, PinReferenceB);
+					}
+					else if (PinReferenceA->Get()->PinType.PinSubCategory.FastLess(PinReferenceB->Get()->PinType.PinSubCategory))
+					{
+						return true;
+					}
+					else
+					{
+						return false;
+					}
+				});
+		}
+	}
 	else if (Method == SPinViewer::COLUMN_VISIBILITY)
 	{
 		if (SortMode == EColumnSortMode::Descending)
@@ -260,13 +301,20 @@ void SPinViewer::Construct(const FArguments& InArgs)
 				.DefaultLabel(LOCTEXT("TypeLabel", "Type"))
 				.SortMode(this, &SPinViewer::GetColumnSortMode, COLUMN_TYPE)
 				.OnSort(this, &SPinViewer::SortListView)
-				.FillWidth(0.25)
+				.FillWidth(0.20)
+				
+				// todo: since it is not the norm having a pin with a SubCategory set, maybe hide this column is no pin has a subcategory set.
+				+ SHeaderRow::Column(COLUMN_SUBTYPE)
+				.DefaultLabel(LOCTEXT("SubTypeLabel", "Subtype"))
+				.SortMode(this, &SPinViewer::GetColumnSortMode, COLUMN_SUBTYPE)
+				.OnSort(this, &SPinViewer::SortListView)
+				.FillWidth(0.20)
 
 				+ SHeaderRow::Column(COLUMN_VISIBILITY)
 				.DefaultLabel(LOCTEXT("VisibilityLabel", "Visibility"))
 				.SortMode(this, &SPinViewer::GetColumnSortMode, COLUMN_VISIBILITY)
 				.OnSort(this, &SPinViewer::SortListView)
-				.FillWidth(0.25)
+				.FillWidth(0.10)
 			)
 		]
 	];
diff --git a/Engine/Plugins/Mutable/Source/CustomizableObjectEditor/Private/MuCOE/PinViewer/SPinViewer.h b/Engine/Plugins/Mutable/Source/CustomizableObjectEditor/Private/MuCOE/PinViewer/SPinViewer.h
--- a/Engine/Plugins/Mutable/Source/CustomizableObjectEditor/Private/MuCOE/PinViewer/SPinViewer.h
+++ b/Engine/Plugins/Mutable/Source/CustomizableObjectEditor/Private/MuCOE/PinViewer/SPinViewer.h
@@ -33,6 +33,7 @@ public:
 
 	static const FName COLUMN_NAME;
 	static const FName COLUMN_TYPE;
+	static const FName COLUMN_SUBTYPE;
 	static const FName COLUMN_VISIBILITY;
 	
 	void Construct(const FArguments& InArgs);
diff --git a/Engine/Plugins/Mutable/Source/CustomizableObjectEditor/Private/MuCOE/PinViewer/SPinViewerListRow.cpp b/Engine/Plugins/Mutable/Source/CustomizableObjectEditor/Private/MuCOE/PinViewer/SPinViewerListRow.cpp
--- a/Engine/Plugins/Mutable/Source/CustomizableObjectEditor/Private/MuCOE/PinViewer/SPinViewerListRow.cpp
+++ b/Engine/Plugins/Mutable/Source/CustomizableObjectEditor/Private/MuCOE/PinViewer/SPinViewerListRow.cpp
@@ -101,34 +101,34 @@ TSharedRef<SWidget> SPinViewerListRow::GenerateWidgetForColumn(const FName& Colu
 	{
 		if (TStrongObjectPtr<UCustomizableObjectNode> PinnedPinOwner = PinOwner.Pin())
 		{
-			if (PinnedPinOwner->GetPinAllowedTypes(*PinReference.Get()).Num() > 0)
-			{
-				return SNew(SVerticalBox)
+			return SNew(SVerticalBox)
 				
-				+ SVerticalBox::Slot()
-				.AutoHeight()
-				.Padding(4)
-				.VAlign(VAlign_Center)
-				.HAlign(HAlign_Left)
-				[
-					SNew(SMutablePinTypeSelector)
-					.PinReference(PinReference)
-				];
-			}
-			// If three is no more than one possible pin type for this pin just expose the type as a non-editable text
-			else
-			{
-				return SNew(SVerticalBox)
-					
-				+ SVerticalBox::Slot()
-				.AutoHeight().Padding(4)
-				.VAlign(VAlign_Center)
-				.HAlign(HAlign_Left)
-				[
-					SNew(STextBlock)
-					.Text(UEdGraphSchema_CustomizableObject::GetPinCategoryFriendlyName(GetInitiallySelectedPinType()))
-				];
-			}
+			+ SVerticalBox::Slot()
+			.AutoHeight()
+			.Padding(4)
+			.VAlign(VAlign_Center)
+			.HAlign(HAlign_Left)
+			[
+				SNew(SMutablePinTypeSelector)
+				.PinReference(PinReference)
+			];
+		}
+	}
+	else if (ColumnName == SPinViewer::COLUMN_SUBTYPE)
+	{
+		if (TStrongObjectPtr<UCustomizableObjectNode> PinnedPinOwner = PinOwner.Pin())
+		{
+			return SNew(SVerticalBox)
+				
+			+ SVerticalBox::Slot()
+			.AutoHeight()
+			.Padding(4)
+			.VAlign(VAlign_Center)
+			.HAlign(HAlign_Left)
+			[
+				SNew(SMutablePinSubTypeSelector)
+				.PinReference(PinReference)
+			];
 		}
 	}
 	else if (ColumnName == SPinViewer::COLUMN_VISIBILITY)
@@ -226,4 +226,17 @@ FName SPinViewerListRow::GetInitiallySelectedPinType() const
 }
 
 
+FName SPinViewerListRow::GetInitiallySelectedPinSubType() const
+{
+	if (UEdGraphPin* PinPointer = PinReference.Get())
+	{
+		return PinPointer->PinType.PinSubCategory;
+	}
+	else 
+	{
+		return NAME_Name;
+	}
+}
+
+
 #undef LOCTEXT_NAMESPACE
diff --git a/Engine/Plugins/Mutable/Source/CustomizableObjectEditor/Private/MuCOE/PinViewer/SPinViewerListRow.h b/Engine/Plugins/Mutable/Source/CustomizableObjectEditor/Private/MuCOE/PinViewer/SPinViewerListRow.h
--- a/Engine/Plugins/Mutable/Source/CustomizableObjectEditor/Private/MuCOE/PinViewer/SPinViewerListRow.h
+++ b/Engine/Plugins/Mutable/Source/CustomizableObjectEditor/Private/MuCOE/PinViewer/SPinViewerListRow.h
@@ -43,7 +43,8 @@ public:
 private:
 	
 	FName GetInitiallySelectedPinType() const;
-	
+	FName GetInitiallySelectedPinSubType() const;
+
 	SPinViewer* PinViewer = nullptr;
 	
 	/** Weak pointer to the Customizable Object node that contains this property */
diff --git a/Engine/Plugins/Mutable/Source/CustomizableObjectEditor/Private/MuCOE/SMutableCodeViewer.cpp b/Engine/Plugins/Mutable/Source/CustomizableObjectEditor/Private/MuCOE/SMutableCodeViewer.cpp
--- a/Engine/Plugins/Mutable/Source/CustomizableObjectEditor/Private/MuCOE/SMutableCodeViewer.cpp
+++ b/Engine/Plugins/Mutable/Source/CustomizableObjectEditor/Private/MuCOE/SMutableCodeViewer.cpp
@@ -1721,7 +1721,7 @@ void SMutableCodeViewer::GenerateElementRecursive(const int32& InStateIndex, UE:
 		
 	case UE::Mutable::Private::EOpType::IN_ADDOVERLAYMATERIAL:
 	{
-		UE::Mutable::Private::OP::InstanceAddMaterialArgs Args = InProgram.GetOpArgs<UE::Mutable::Private::OP::InstanceAddMaterialArgs>(InParentAddress);
+		UE::Mutable::Private::OP::InstanceAddOverlayMaterialArgs Args = InProgram.GetOpArgs<UE::Mutable::Private::OP::InstanceAddOverlayMaterialArgs>(InParentAddress);
 		AddOpFunc(Args.Instance, TEXT("instance "));
 		AddOpFunc(Args.Material, TEXT("value "));
 		break;
diff --git a/Engine/Plugins/Mutable/Source/CustomizableObjectEditor/Private/MuCOE/Widgets/SMutablePinTypeSelector.cpp b/Engine/Plugins/Mutable/Source/CustomizableObjectEditor/Private/MuCOE/Widgets/SMutablePinTypeSelector.cpp
--- a/Engine/Plugins/Mutable/Source/CustomizableObjectEditor/Private/MuCOE/Widgets/SMutablePinTypeSelector.cpp
+++ b/Engine/Plugins/Mutable/Source/CustomizableObjectEditor/Private/MuCOE/Widgets/SMutablePinTypeSelector.cpp
@@ -14,7 +14,7 @@ void SMutablePinTypeSelector::Construct(const FArguments& InArgs)
 {
 	PinReference = InArgs._PinReference;
 
-	UEdGraphPin* Pin = PinReference.Get();
+	const UEdGraphPin* Pin = PinReference.Get();
 	check(Pin);
 	
 	if (UCustomizableObjectNode* OwningNode = Cast<UCustomizableObjectNode>(Pin->GetOwningNode()))
@@ -27,19 +27,32 @@ void SMutablePinTypeSelector::Construct(const FArguments& InArgs)
 		checkNoEntry();
 	}
 	
-	ChildSlot
-	[
-		SNew(SComboBox<FName>)
-		.OptionsSource(&AllowedPinTypes)
-		.InitiallySelectedItem(GetInitiallySelectedPinType())
-		.IsEnabled(this, &SMutablePinTypeSelector::ShouldVariableTypePropertyBeEnabled)
-		.OnGenerateWidget(this, &SMutablePinTypeSelector::OnGenerateVariableTypeRow)
-		.OnSelectionChanged(this, &SMutablePinTypeSelector::OnTypeDropdownSelectionChange)
-		.Content()
+	if (AllowedPinTypes.Num() > 1)
+	{
+		ChildSlot
 		[
-			GenerateCurrentSelectedTypeWidget()
-		]
-	];
+			SNew(SComboBox<FName>)
+			.OptionsSource(&AllowedPinTypes)
+			.InitiallySelectedItem(GetInitiallySelectedPinType())
+			.IsEnabled(this, &SMutablePinTypeSelector::ShouldVariableTypePropertyBeEnabled)
+			.OnGenerateWidget(this, &SMutablePinTypeSelector::OnGenerateVariableTypeRow)
+			.OnSelectionChanged(this, &SMutablePinTypeSelector::OnTypeDropdownSelectionChange)
+			.Content()
+			[
+				GenerateCurrentSelectedTypeWidget()
+			]
+		];
+	}
+	else
+	{
+		ChildSlot
+		[
+			SNew(STextBlock)
+			.Text(UEdGraphSchema_CustomizableObject::GetPinCategoryFriendlyName(GetInitiallySelectedPinType()))
+		];
+	}
+	
+
 }
 
 
@@ -58,7 +71,12 @@ FName SMutablePinTypeSelector::GetInitiallySelectedPinType() const
 
 bool SMutablePinTypeSelector::ShouldVariableTypePropertyBeEnabled() const
 {
-	return PinReference.Get()->LinkedTo.Num() == 0;
+	if (PinReference.Get())
+	{
+		return PinReference.Get()->LinkedTo.Num() == 0;
+	}
+	
+	return false;
 }
 
 
@@ -109,4 +127,98 @@ TSharedRef<SWidget> SMutablePinTypeSelector::GenerateCurrentSelectedTypeWidget()
 	return OnGenerateVariableTypeRow(InitiallySelectedPinType);
 }
 
+
+
+void SMutablePinSubTypeSelector::Construct(const FArguments& InArgs)
+{
+	PinReference = InArgs._PinReference;
+
+	const UEdGraphPin* Pin = PinReference.Get();
+	check(Pin);
+	
+	if (UCustomizableObjectNode* OwningNode = Cast<UCustomizableObjectNode>(Pin->GetOwningNode()))
+	{
+		AllowedPinSubtypes = OwningNode->GetPinAllowedSubTypes(*Pin);
+		PinOwner = MakeWeakObjectPtr(OwningNode);
+	}
+	else
+	{
+		checkNoEntry();
+	}
+	
+	if (AllowedPinSubtypes.Num() > 1)
+	{
+		ChildSlot
+		[
+			SNew(SComboBox<FName>)
+			.OptionsSource(&AllowedPinSubtypes)
+			.InitiallySelectedItem(GetInitiallySelectedPinSubType())
+			.IsEnabled(true)
+			.OnGenerateWidget(this, &SMutablePinSubTypeSelector::OnGenerateVariableSubtypeRow)
+			.OnSelectionChanged(this, &SMutablePinSubTypeSelector::OnTypeDropdownSelectionChange)
+			.Content()
+			[
+				GenerateCurrentSelectedTypeWidget()
+			]
+		];
+	}
+	else
+	{
+		ChildSlot
+		[
+			SNew(STextBlock)
+			.Text(UEdGraphSchema_CustomizableObject::GetPinSubCategoryFriendlyName(GetInitiallySelectedPinSubType()))
+		];
+	}
+	
+}
+
+
+FName SMutablePinSubTypeSelector::GetInitiallySelectedPinSubType() const
+{
+	if (UEdGraphPin* PinPointer = PinReference.Get())
+	{
+		return PinPointer->PinType.PinSubCategory;
+	}
+	else 
+	{
+		return NAME_Name;
+	}
+}
+
+
+TSharedRef<SWidget> SMutablePinSubTypeSelector::OnGenerateVariableSubtypeRow(FName Subtype)
+{
+	TSharedRef<STextBlock> RowWidget = SNew(STextBlock)
+				.Text(Subtype.IsValid() ? UEdGraphSchema_CustomizableObject::GetPinSubCategoryFriendlyName(Subtype) : FText::FromString("Invalid"))
+				.Font(IDetailLayoutBuilder::GetDetailFont())
+				.ToolTipText(UEdGraphSchema_CustomizableObject::GetSubCategoryTooltipText(Subtype));
+
+	return RowWidget;
+}
+
+
+void SMutablePinSubTypeSelector::OnTypeDropdownSelectionChange(FName NewSubTypeName, ESelectInfo::Type Arg) const
+{
+	if (UEdGraphPin* Pin = PinReference.Get())
+	{
+		Pin->PinType.PinSubCategory = NewSubTypeName;
+		
+		if (TStrongObjectPtr<UCustomizableObjectNode> Node = PinOwner.Pin())
+		{
+			Node->ReconstructNode();
+		}
+	}
+}
+
+
+TSharedRef<SWidget> SMutablePinSubTypeSelector::GenerateCurrentSelectedTypeWidget()
+{
+	const FName InitiallySelectedPinSubType = GetInitiallySelectedPinSubType();
+	check(AllowedPinSubtypes.Contains(InitiallySelectedPinSubType))
+	
+	return OnGenerateVariableSubtypeRow(InitiallySelectedPinSubType);
+}
+
+
 #undef LOCTEXT_NAMESPACE
diff --git a/Engine/Plugins/Mutable/Source/CustomizableObjectEditor/Private/MuCOE/Widgets/SMutablePinTypeSelector.h b/Engine/Plugins/Mutable/Source/CustomizableObjectEditor/Private/MuCOE/Widgets/SMutablePinTypeSelector.h
--- a/Engine/Plugins/Mutable/Source/CustomizableObjectEditor/Private/MuCOE/Widgets/SMutablePinTypeSelector.h
+++ b/Engine/Plugins/Mutable/Source/CustomizableObjectEditor/Private/MuCOE/Widgets/SMutablePinTypeSelector.h
@@ -12,7 +12,10 @@ class UCustomizableObjectNode;
 struct FEdGraphPinReference;
 struct FGuid;
 
-
+/**
+ * Slate that exposes a dropdown that takes care of exposing the pin category and allowing the end user to change it.
+ * It will call the refresh of the node to apply that change.
+ */
 class SMutablePinTypeSelector : public SCompoundWidget
 {
 public:
@@ -29,7 +32,6 @@ private:
 	TSharedRef<SWidget> OnGenerateVariableTypeRow(FName Type);
 	void OnTypeDropdownSelectionChange(FName NewTypeName, ESelectInfo::Type Arg) const;
 	TSharedRef<SWidget> GenerateCurrentSelectedTypeWidget();
-
 	
 	FEdGraphPinReference PinReference;
 
@@ -37,4 +39,34 @@ private:
 	
 	/** Weak pointer to the Customizable Object node that contains the pin */
 	TWeakObjectPtr<UCustomizableObjectNode> PinOwner;
-};
\ No newline at end of file
+};
+
+
+/**
+ * Slate that exposes a dropdown that takes care of exposing the pin subcategory and allowing the end user to change it.
+ * It will call the refresh of the node to apply that change.
+ */
+class SMutablePinSubTypeSelector : public SCompoundWidget
+{
+public:
+	SLATE_BEGIN_ARGS(SMutablePinSubTypeSelector) {}
+		SLATE_ARGUMENT(FEdGraphPinReference, PinReference)
+	SLATE_END_ARGS()
+
+	void Construct(const FArguments& InArgs);
+	
+private:
+	
+	FName GetInitiallySelectedPinSubType() const;
+	TSharedRef<SWidget> OnGenerateVariableSubtypeRow(FName Subtype);
+	void OnTypeDropdownSelectionChange(FName NewSubTypeName, ESelectInfo::Type Arg) const;
+	TSharedRef<SWidget> GenerateCurrentSelectedTypeWidget();
+	
+	FEdGraphPinReference PinReference;
+
+	TArray<FName> AllowedPinSubtypes;
+	
+	/** Weak pointer to the Customizable Object node that contains the pin */
+	TWeakObjectPtr<UCustomizableObjectNode> PinOwner;
+};
+
diff --git a/Engine/Plugins/Mutable/Source/CustomizableObjectEditor/Public/MuCOE/EdGraphSchema_CustomizableObject.h b/Engine/Plugins/Mutable/Source/CustomizableObjectEditor/Public/MuCOE/EdGraphSchema_CustomizableObject.h
--- a/Engine/Plugins/Mutable/Source/CustomizableObjectEditor/Public/MuCOE/EdGraphSchema_CustomizableObject.h
+++ b/Engine/Plugins/Mutable/Source/CustomizableObjectEditor/Public/MuCOE/EdGraphSchema_CustomizableObject.h
@@ -51,6 +51,9 @@ public:
 	static UE_API const FName PC_Transform;
 	static UE_API const FName PC_String;
 
+	static UE_API const FName PSC_Material_Override;
+	static UE_API const FName PSC_Material_Overlay;
+
 	static UE_API const TArray<FName> SupportedMacroPinTypes;
 
 	// Node categories
@@ -77,12 +80,14 @@ public:
 	static UE_API FLinearColor GetPinTypeColor(const FName& PinType);
 	static UE_API FName GetPinCategoryName(const FName& PinCategory);
 	static UE_API FName GetPinCategoryName(const UScriptStruct* Type);
+	static UE_API FName GetPinSubCategoryName(const FName& PinSubCategory);
 	static UE_API FText GetPinCategoryFriendlyName(const FName& PinCategory);
 	static UE_API FText GetPinCategoryFriendlyName(const UScriptStruct* Type);
+	static UE_API FText GetPinSubCategoryFriendlyName(const FName& PinSubcategory);
 	static UE_API bool IsPassthrough(const FName& PinCategory);
 	static UE_API bool IsTypeNative(const UScriptStruct* Type);
-
-
+	static UE_API FText GetSubCategoryTooltipText(FName PinSubType);
+	
 private:
 	static TSharedPtr<IAssetReferenceFilter> MakeAssetReferenceFilter(const UEdGraph* Graph);
 };
diff --git a/Engine/Plugins/Mutable/Source/CustomizableObjectEditor/Public/MuCOE/Nodes/CustomizableObjectNode.h b/Engine/Plugins/Mutable/Source/CustomizableObjectEditor/Public/MuCOE/Nodes/CustomizableObjectNode.h
--- a/Engine/Plugins/Mutable/Source/CustomizableObjectEditor/Public/MuCOE/Nodes/CustomizableObjectNode.h
+++ b/Engine/Plugins/Mutable/Source/CustomizableObjectEditor/Public/MuCOE/Nodes/CustomizableObjectNode.h
@@ -344,6 +344,8 @@ public:
 	 */
 	UE_API virtual TArray<FName> GetPinAllowedTypes(const UEdGraphPin& Pin) const;
 
+	UE_API virtual TArray<FName> GetPinAllowedSubTypes(const UEdGraphPin& Pin) const;
+
 	UE_API virtual bool IsCompatibleWithGraph(UEdGraph const* Graph) const override;
 	
 	FPostEditChangePropertyDelegate PostEditChangePropertyDelegate;
diff --git a/Engine/Plugins/Mutable/Source/MutableRuntime/Internal/MuR/Instance.h b/Engine/Plugins/Mutable/Source/MutableRuntime/Internal/MuR/Instance.h
--- a/Engine/Plugins/Mutable/Source/MutableRuntime/Internal/MuR/Instance.h
+++ b/Engine/Plugins/Mutable/Source/MutableRuntime/Internal/MuR/Instance.h
@@ -51,6 +51,11 @@ namespace UE::Mutable::Private
 		FMaterialId MaterialId;
 	};
 	
+	struct FOverlayMaterial
+	{
+		FName SlotName;
+		FMaterialId MaterialId;
+	};
 
 	struct FInstanceComponent
 	{
@@ -60,6 +65,8 @@ namespace UE::Mutable::Private
 
 		TArray<FOverrideMaterial> OverrideMaterials;
 
+		TArray<FOverlayMaterial> OverlayMaterials;
+		
 		FSkeletalMeshId SkeletalMeshId;
 
 		TArray<FInstanceLOD, TInlineAllocator<1>> LODs;
@@ -169,6 +176,13 @@ namespace UE::Mutable::Private
 
 		UE_API FMaterialId GetOverrideMaterialId(int32 ComponentIndex, int32 MaterialIndex) const;
 		
+		
+		UE_API int32 GetOverlayMaterialCount(int32 ComponentIndex) const;
+
+		UE_API FName GetOverlayMaterialSlotSlotName(int32 ComponentIndex, int32 MaterialIndex) const;
+
+		UE_API FMaterialId GetOverlayMaterialId(int32 ComponentIndex, int32 MaterialIndex) const;
+		
 		UE_API FMaterialId GetMaterialId(int32 ComponentIndex, int32 LODIndex, int32 SurfaceIndex) const;
 
 		//! Get the number of ExtensionData values in a component
@@ -200,6 +214,8 @@ namespace UE::Mutable::Private
 		
 		void AddOverrideMaterial(int32 ComponentIndex, const FName& SlotName, const FMaterialId& MaterialId);
 
+		void AddOverlayMaterial(int32 ComponentIndex, const FName& SlotName, const FMaterialId& MaterialId);
+		
 		void SetExtensionData(const TManagedRef<const FExtensionData>& Data, FName Name);
 		
 		TArray<FInstanceComponent, TInlineAllocator<1>> Components;
diff --git a/Engine/Plugins/Mutable/Source/MutableRuntime/Internal/MuR/Operations.h b/Engine/Plugins/Mutable/Source/MutableRuntime/Internal/MuR/Operations.h
--- a/Engine/Plugins/Mutable/Source/MutableRuntime/Internal/MuR/Operations.h
+++ b/Engine/Plugins/Mutable/Source/MutableRuntime/Internal/MuR/Operations.h
@@ -1176,8 +1176,16 @@ namespace UE::Mutable::Private
 		{
 			ADDRESS Instance;
 			ADDRESS Material;
-			CONSTANT_STRING_INDEX SlotName;
+			CONSTANT_STRING_INDEX SlotName = INDEX_NONE;
 		};
+    	
+    	struct InstanceAddOverlayMaterialArgs
+    	{
+    		ADDRESS Instance;
+    		ADDRESS Material;
+    		
+    		CONSTANT_STRING_INDEX SlotName = INDEX_NONE;
+    	};
 
         struct LayoutPackArgs
         {
diff --git a/Engine/Plugins/Mutable/Source/MutableRuntime/Private/MuR/CodeRunner.cpp b/Engine/Plugins/Mutable/Source/MutableRuntime/Private/MuR/CodeRunner.cpp
--- a/Engine/Plugins/Mutable/Source/MutableRuntime/Private/MuR/CodeRunner.cpp
+++ b/Engine/Plugins/Mutable/Source/MutableRuntime/Private/MuR/CodeRunner.cpp
@@ -879,7 +879,7 @@ namespace UE::Mutable::Private
         	{
         		MUTABLE_CPUPROFILER_SCOPE(IN_ADDOVERLAYMATERIAL);
         		
-        		OP::InstanceAddMaterialArgs Args = Program.GetOpArgs<OP::InstanceAddMaterialArgs>(Item.At);
+        		OP::InstanceAddOverlayMaterialArgs Args = Program.GetOpArgs<OP::InstanceAddOverlayMaterialArgs>(Item.At);
         		switch (Item.Stage)
         		{
         		case 0:
@@ -906,8 +906,17 @@ namespace UE::Mutable::Private
 
         				if (Args.Material)
         				{
-        					FMaterialId MaterialId = FSystem::GetMaterialId(LiveInstance->MaterialIdRegistry, Model, Parameters, Args.Material);
-        					Result->SetOverlayMaterialId(0, MaterialId);
+        					const FMaterialId MaterialId = FSystem::GetMaterialId(LiveInstance->MaterialIdRegistry, Model, Parameters, Args.Material);
+        					
+        					const FName& SlotName = Args.SlotName == INDEX_NONE ? FName() : FName(Program.ConstantStrings[Args.SlotName]);
+        					if (SlotName.IsNone())
+        					{
+        						Result->SetOverlayMaterialId(0, MaterialId);
+        					}
+					        else
+					        {
+					        	Result->AddOverlayMaterial(0, SlotName, MaterialId);
+					        }
         				}
 
         				StoreInstance(Item, Result);
diff --git a/Engine/Plugins/Mutable/Source/MutableRuntime/Private/MuR/CodeRunnerBegin.cpp b/Engine/Plugins/Mutable/Source/MutableRuntime/Private/MuR/CodeRunnerBegin.cpp
--- a/Engine/Plugins/Mutable/Source/MutableRuntime/Private/MuR/CodeRunnerBegin.cpp
+++ b/Engine/Plugins/Mutable/Source/MutableRuntime/Private/MuR/CodeRunnerBegin.cpp
@@ -928,7 +928,7 @@ namespace UE::Mutable::Private
 				{
 					check(Item.bEvaluate);
 
-					OP::InstanceAddMaterialArgs Args = Program.GetOpArgs<OP::InstanceAddMaterialArgs>(Item.At);
+					OP::InstanceAddOverlayMaterialArgs Args = Program.GetOpArgs<OP::InstanceAddOverlayMaterialArgs>(Item.At);
 					switch (Item.Stage)
 					{
 					case 0:
@@ -960,9 +960,18 @@ namespace UE::Mutable::Private
 
 							if (Args.Material)
 							{
-        						FMaterialId MaterialId = FSystem::GetMaterialId(MaterialIdRegistry, *Model.Get(), *Params.Get(), Args.Material);
-								Instance->SetOverlayMaterialId(0, MaterialId);
-
+								const FMaterialId MaterialId = FSystem::GetMaterialId(MaterialIdRegistry,  *Model.Get(),  *Params.Get(), Args.Material);
+        					
+								const FName& SlotName = Args.SlotName == INDEX_NONE ? FName() : FName(Program.ConstantStrings[Args.SlotName]);
+								if (SlotName.IsNone())
+								{
+									Instance->SetOverlayMaterialId(0, MaterialId);
+								}
+								else
+								{
+									Instance->AddOverlayMaterial(0, SlotName, MaterialId);
+								}
+								
 								// Pop Material
 								LoadMaterial(Args.Material);
 							}
diff --git a/Engine/Plugins/Mutable/Source/MutableRuntime/Private/MuR/Instance.cpp b/Engine/Plugins/Mutable/Source/MutableRuntime/Private/MuR/Instance.cpp
--- a/Engine/Plugins/Mutable/Source/MutableRuntime/Private/MuR/Instance.cpp
+++ b/Engine/Plugins/Mutable/Source/MutableRuntime/Private/MuR/Instance.cpp
@@ -194,6 +194,21 @@ int32 FInstance::GetStringCount(int32 ComponentIndex, int32 LODIndex, int32 Surf
     	return Components[ComponentIndex].OverrideMaterials[MaterialIndex].MaterialId;
 	}
 
+	
+	int32 FInstance::GetOverlayMaterialCount(int32 ComponentIndex) const
+	{
+    	return Components[ComponentIndex].OverlayMaterials.Num();
+	}
+
+	FName FInstance::GetOverlayMaterialSlotSlotName(int32 ComponentIndex, int32 MaterialIndex) const
+	{
+    	return Components[ComponentIndex].OverlayMaterials[MaterialIndex].SlotName;
+	}
+
+	FMaterialId FInstance::GetOverlayMaterialId(int32 ComponentIndex, int32 MaterialIndex) const
+	{
+    	return Components[ComponentIndex].OverlayMaterials[MaterialIndex].MaterialId;
+	}
 
 	FMaterialId FInstance::GetMaterialId(int32 ComponentIndex, int32 LODIndex, int32 SurfaceIndex) const
 	{
@@ -364,6 +379,16 @@ int32 FInstance::GetStringCount(int32 ComponentIndex, int32 LODIndex, int32 Surf
     	Components[ComponentIndex].OverrideMaterials.Add({ SlotName, MaterialId });
 	}
 
+	void FInstance::AddOverlayMaterial(int32 ComponentIndex, const FName& SlotName, const FMaterialId& MaterialId)
+	{
+    	while (ComponentIndex >= Components.Num())
+    	{
+    		AddComponent();
+    	}
+
+    	Components[ComponentIndex].OverlayMaterials.Add({ SlotName, MaterialId });
+	}
+
 
 	void FInstance::SetExtensionData(const TManagedRef<const FExtensionData>& Data, FName Name)
 	{
diff --git a/Engine/Plugins/Mutable/Source/MutableRuntime/Private/MuR/Operations.cpp b/Engine/Plugins/Mutable/Source/MutableRuntime/Private/MuR/Operations.cpp
--- a/Engine/Plugins/Mutable/Source/MutableRuntime/Private/MuR/Operations.cpp
+++ b/Engine/Plugins/Mutable/Source/MutableRuntime/Private/MuR/Operations.cpp
@@ -947,6 +947,15 @@ namespace UE::Mutable::Private
 		}
 
 		case EOpType::IN_ADDOVERLAYMATERIAL:
+        {
+        	const OP::InstanceAddOverlayMaterialArgs Args = program.GetOpArgs<OP::InstanceAddOverlayMaterialArgs>(at);
+
+        	f(Args.Instance);
+        	f(Args.Material);
+        		
+        	break;
+        }
+        	
 		case EOpType::IN_ADDMATERIAL:
 		{
 			const OP::InstanceAddMaterialArgs Args = program.GetOpArgs<OP::InstanceAddMaterialArgs>(at);
diff --git a/Engine/Plugins/Mutable/Source/MutableTools/Internal/MuT/NodeComponentNew.h b/Engine/Plugins/Mutable/Source/MutableTools/Internal/MuT/NodeComponentNew.h
--- a/Engine/Plugins/Mutable/Source/MutableTools/Internal/MuT/NodeComponentNew.h
+++ b/Engine/Plugins/Mutable/Source/MutableTools/Internal/MuT/NodeComponentNew.h
@@ -28,6 +28,9 @@ namespace UE::Mutable::Private
 		/** Map where the key is the target Material slot and the value the replacement material to be used in the target material slot*/
 		TMap<const FName, Ptr<NodeMaterial>> OverrideMaterials;
 		
+		/** Map where the key is the target Material slot and the value the material to be rendered on top in the target material slot*/
+		TMap<const FName, Ptr<NodeMaterial>> OverlayMaterials;
+		
 		//-----------------------------------------------------------------------------------------
         // Node interface
 		//-----------------------------------------------------------------------------------------		
diff --git a/Engine/Plugins/Mutable/Source/MutableTools/Private/MuT/ASTOpAddOverlayMaterial.cpp b/Engine/Plugins/Mutable/Source/MutableTools/Private/MuT/ASTOpAddOverlayMaterial.cpp
new file mode 100644
--- /dev/null
+++ b/Engine/Plugins/Mutable/Source/MutableTools/Private/MuT/ASTOpAddOverlayMaterial.cpp
@@ -0,0 +1,137 @@
+// Copyright Epic Games, Inc. All Rights Reserved.
+
+#include "MuT/ASTOpAddOverlayMaterial.h"
+
+#include "MuR/Model.h"
+#include "MuT/CodeOptimiser.h"
+
+
+namespace UE::Mutable::Private
+{
+	ASTOpAddOverlayMaterial::ASTOpAddOverlayMaterial() :
+		Instance(this),
+		Material(this)
+	{
+	}
+
+
+	ASTOpAddOverlayMaterial::~ASTOpAddOverlayMaterial()
+	{
+		// Explicit call needed to avoid recursive destruction
+		RemoveChildren();
+	}
+
+
+	EOpType ASTOpAddOverlayMaterial::GetOpType() const
+	{
+		return EOpType::IN_ADDOVERLAYMATERIAL;
+	}
+
+
+	bool ASTOpAddOverlayMaterial::IsEqual(const ASTOp& OtherUntyped) const
+	{
+		if (OtherUntyped.GetOpType() == GetOpType())
+		{
+			const ASTOpAddOverlayMaterial* Other = static_cast<const ASTOpAddOverlayMaterial*>(&OtherUntyped);
+			return Instance == Other->Instance &&
+				Material == Other->Material &&
+				SlotName == Other->SlotName;
+		}
+
+		return false;
+	}
+
+	
+	uint32 ASTOpAddOverlayMaterial::Hash() const
+	{
+		uint32 Result = GetTypeHash(GetOpType());
+		
+		Result = HashCombineFast(Result, GetTypeHash(Instance));
+		Result = HashCombineFast(Result, GetTypeHash(Material));
+		Result = HashCombineFast(Result, GetTypeHash(SlotName));
+		
+		return Result;
+	}
+	
+
+	Ptr<ASTOp> ASTOpAddOverlayMaterial::Clone(MapChildFuncRef MapChild) const
+	{
+		Ptr<ASTOpAddOverlayMaterial> NewInstance = new ASTOpAddOverlayMaterial();
+
+		NewInstance->Instance = MapChild(Instance.child());
+		NewInstance->Material = MapChild(Material.child());
+		NewInstance->SlotName = SlotName;
+
+		return NewInstance;
+	}
+	
+
+	void ASTOpAddOverlayMaterial::ForEachChild(const TFunctionRef<void(ASTChild&)> Func)
+	{
+		Func(Instance);
+		Func(Material);
+	}
+	
+
+	void ASTOpAddOverlayMaterial::Link(FProgram& Program, FLinkerOptions*)
+	{
+		// Already linked?
+		if (linkedAddress)
+		{
+			return;
+		}
+
+		OP::InstanceAddOverlayMaterialArgs Args;
+		FMemory::Memzero(Args);
+
+		if (Instance)
+		{
+			Args.Instance = Instance->linkedAddress;
+		}
+
+		if (Material)
+		{
+			Args.Material = Material->linkedAddress;
+			
+			// Find out relevant parameters. \todo: this may be optimised by reusing partial
+			// values in a LINK_CONTEXT or similar
+			SubtreeRelevantParametersVisitorAST Visitor;
+			Visitor.Run(Material.child());
+			
+			TArray<uint16> Params;
+			for (const FString& ParameterName : Visitor.Parameters)
+			{
+				for (int32 Index = 0; Index < Program.Parameters.Num(); ++Index)
+				{
+					const FParameterDesc& Param = Program.Parameters[Index];
+					if (Param.Name == ParameterName)
+					{
+						Params.Add(uint16(Index));
+						break;
+					}
+				}
+			}
+
+			Params.Sort();
+
+			int32 It = Program.ParameterLists.Find(Params);
+
+			if (It != INDEX_NONE)
+			{
+				Program.RelevantParameterList.Add(Material->linkedAddress, It);
+			}
+			else
+			{
+				Program.RelevantParameterList.Add(Material->linkedAddress, Program.ParameterLists.Num());
+				Program.ParameterLists.Add(Params);
+			}
+		}
+		
+		Args.SlotName = Program.AddConstant(SlotName.ToString());
+
+		linkedAddress = (OP::ADDRESS)Program.OpAddress.Num();
+		Program.OpAddress.Add((uint32_t)Program.ByteCode.Num());
+		AppendCode(Program.ByteCode, GetOpType());
+		AppendCode(Program.ByteCode, Args);
+	}
+}
\ No newline at end of file
diff --git a/Engine/Plugins/Mutable/Source/MutableTools/Private/MuT/ASTOpAddOverlayMaterial.h b/Engine/Plugins/Mutable/Source/MutableTools/Private/MuT/ASTOpAddOverlayMaterial.h
new file mode 100644
--- /dev/null
+++ b/Engine/Plugins/Mutable/Source/MutableTools/Private/MuT/ASTOpAddOverlayMaterial.h
@@ -0,0 +1,35 @@
+// Copyright Epic Games, Inc. All Rights Reserved.
+
+#pragma once
+
+#include "HAL/Platform.h"
+#include "MuR/Operations.h"
+#include "MuR/Ptr.h"
+#include "MuT/AST.h"
+
+
+namespace UE::Mutable::Private
+{
+	struct FProgram;
+	
+	class ASTOpAddOverlayMaterial final : public ASTOp
+	{
+	public:
+		ASTChild Instance;
+		
+		ASTChild Material;
+		
+		FName SlotName;
+
+		ASTOpAddOverlayMaterial();
+		ASTOpAddOverlayMaterial(const ASTOpAddOverlayMaterial&) = delete;
+		virtual ~ASTOpAddOverlayMaterial() override;
+
+		virtual EOpType GetOpType() const override;
+		virtual uint32 Hash() const override;
+		virtual void ForEachChild(const TFunctionRef<void(ASTChild&)> Func) override;
+		virtual bool IsEqual(const ASTOp& OtherUntyped) const override;
+		virtual Ptr<ASTOp> Clone(MapChildFuncRef MapChild) const override;
+		virtual void Link(FProgram& Program, FLinkerOptions* Options) override;
+	};
+}
\ No newline at end of file
diff --git a/Engine/Plugins/Mutable/Source/MutableTools/Private/MuT/CodeGenerator.cpp b/Engine/Plugins/Mutable/Source/MutableTools/Private/MuT/CodeGenerator.cpp
--- a/Engine/Plugins/Mutable/Source/MutableTools/Private/MuT/CodeGenerator.cpp
+++ b/Engine/Plugins/Mutable/Source/MutableTools/Private/MuT/CodeGenerator.cpp
@@ -2,6 +2,7 @@
 
 #include "MuT/CodeGenerator.h"
 
+#include "ASTOpAddOverlayMaterial.h"
 #include "ASTOpAddOverrideMaterial.h"
 #include "ASTOpAddSkeletalMesh.h"
 #include "ASTOpEntry.h"
@@ -585,9 +586,9 @@ namespace UE::Mutable::Private
 					
 				if (Ptr<NodeMaterial> MaterialNode = InNode->OverlayMaterial)
 				{
-					Ptr<ASTOpAddMaterial> AddOverlayMaterialOp = new ASTOpAddMaterial();
-					AddOverlayMaterialOp->Type = EOpType::IN_ADDOVERLAYMATERIAL;
+					Ptr<ASTOpAddOverlayMaterial> AddOverlayMaterialOp = new ASTOpAddOverlayMaterial();
 					AddOverlayMaterialOp->Instance = LastInstOp;
+					AddOverlayMaterialOp->SlotName = NAME_None;
 
 					// Material
 					FMaterialGenerationOptions MaterialOptions;
@@ -598,19 +599,33 @@ namespace UE::Mutable::Private
 					LastInstOp = AddOverlayMaterialOp;
 				}
 				
-				for (TPair<const FName, Ptr<NodeMaterial>> Pair : InNode->OverrideMaterials)
+				for (const TPair<const FName, Ptr<NodeMaterial>>& NameAndMaterial : InNode->OverrideMaterials)
 				{
 					Ptr<ASTOpAddOverrideMaterial> AddOverrideMaterialOp = new ASTOpAddOverrideMaterial();
-					AddOverrideMaterialOp->SlotName = Pair.Key;
+					AddOverrideMaterialOp->SlotName = NameAndMaterial.Key;
 					AddOverrideMaterialOp->Instance = LastInstOp;
 
 					FMaterialGenerationOptions MaterialOptions;
 					FMaterialGenerationResult MaterialResult;
-					GenerateMaterial(MaterialOptions, MaterialResult, Pair.Value);
+					GenerateMaterial(MaterialOptions, MaterialResult, NameAndMaterial.Value);
 					AddOverrideMaterialOp->Material = MaterialResult.op;
 
 					LastInstOp = AddOverrideMaterialOp;
 				}
+					
+				for (const TPair<const FName, Ptr<NodeMaterial>>& NameAndMaterial : InNode->OverlayMaterials)
+				{ 
+					Ptr<ASTOpAddOverlayMaterial> AddOverlayMaterialOp = new ASTOpAddOverlayMaterial();
+					AddOverlayMaterialOp->SlotName = NameAndMaterial.Key;
+					AddOverlayMaterialOp->Instance = LastInstOp;
+
+					FMaterialGenerationOptions MaterialOptions;
+					FMaterialGenerationResult MaterialResult;
+					GenerateMaterial(MaterialOptions, MaterialResult, NameAndMaterial.Value);
+					AddOverlayMaterialOp->Material = MaterialResult.op;
+
+					LastInstOp = AddOverlayMaterialOp;
+				}
 
 				Ptr<ASTOpComponentAdd> InstanceOp = new ASTOpComponentAdd();
 				InstanceOp->Instance = Options.BaseInstance;
