[Mutable] Fix UVs changing precision  when streaming.

#Synced-CL 50493652

diff --git a/Engine/Plugins/Mutable/Source/CustomizableObject/Private/MuCO/UnrealConversionUtils.cpp b/Engine/Plugins/Mutable/Source/CustomizableObject/Private/MuCO/UnrealConversionUtils.cpp
--- a/Engine/Plugins/Mutable/Source/CustomizableObject/Private/MuCO/UnrealConversionUtils.cpp
+++ b/Engine/Plugins/Mutable/Source/CustomizableObject/Private/MuCO/UnrealConversionUtils.cpp
@@ -37,6 +37,7 @@ namespace UnrealConversionUtils
 		 * @param NumVertices - The amount of vertices on the buffer
 		 * @param NumTexCoords - The amount of texture coordinates
 		 * @param bUseFullPrecisionUVs - Determines if we want to use or not full precision UVs
+		 * @param bUseFullPrecisionTangentBasis - Determines if we want to use or not full precision Tangent Basis
 		 * @param bNeedCPUAccess - Determines if CPU access is required
 		 * @param InMutablePositionData - Mutable position data buffer
 		 * @param InMutableTangentData - Mutable tangent data buffer
@@ -47,6 +48,7 @@ namespace UnrealConversionUtils
 			const int32 NumVertices,
 			const int32 NumTexCoords,
 			const bool bUseFullPrecisionUVs,
+			const bool bUseFullPrecisionTangentBasis,
 			const bool bNeedCPUAccess,
 			const void* InMutablePositionData,
 			const void* InMutableTangentData,
@@ -57,8 +59,13 @@ namespace UnrealConversionUtils
 			FMemory::Memcpy(OutVertexBuffers.PositionVertexBuffer.GetVertexData(), InMutablePositionData, NumVertices * OutVertexBuffers.PositionVertexBuffer.GetStride());
 
 			// tangent and texture coords
+			
+			check(!OutVertexBuffers.StaticMeshVertexBuffer.IsInitialized() || OutVertexBuffers.StaticMeshVertexBuffer.GetUseFullPrecisionUVs() == bUseFullPrecisionUVs); // If already initialized to full precision, it must not change.
 			OutVertexBuffers.StaticMeshVertexBuffer.SetUseFullPrecisionUVs(bUseFullPrecisionUVs);
-			OutVertexBuffers.StaticMeshVertexBuffer.SetUseHighPrecisionTangentBasis(false);
+			
+			check(!OutVertexBuffers.StaticMeshVertexBuffer.IsInitialized() || OutVertexBuffers.StaticMeshVertexBuffer.GetUseHighPrecisionTangentBasis() == bUseFullPrecisionTangentBasis); // If already initialized to full precision, it must not change.
+			OutVertexBuffers.StaticMeshVertexBuffer.SetUseHighPrecisionTangentBasis(bUseFullPrecisionTangentBasis);
+			
 			OutVertexBuffers.StaticMeshVertexBuffer.Init(NumVertices, NumTexCoords, bNeedCPUAccess);
 			FMemory::Memcpy(OutVertexBuffers.StaticMeshVertexBuffer.GetTangentData(), InMutableTangentData, OutVertexBuffers.StaticMeshVertexBuffer.GetTangentSize());
 			FMemory::Memcpy(OutVertexBuffers.StaticMeshVertexBuffer.GetTexCoordData(), InMutableTextureData, OutVertexBuffers.StaticMeshVertexBuffer.GetTexCoordSize());
@@ -302,6 +309,34 @@ namespace UnrealConversionUtils
 	}
 
 
+	bool IsUseFullPrecisionUVs(const UE::Mutable::Private::FMesh& Mesh)
+	{
+		int32 TexCoordsBufferIndex = -1;
+		int32 TexCoordsChannelIndex = -1;
+		Mesh.VertexBuffers.FindChannel(UE::Mutable::Private::EMeshBufferSemantic::TexCoords, 0, &TexCoordsBufferIndex, &TexCoordsChannelIndex);
+		if (TexCoordsBufferIndex >= 0 && TexCoordsChannelIndex >= 0)
+		{
+			return Mesh.VertexBuffers.Buffers[TexCoordsBufferIndex].Channels[TexCoordsChannelIndex].Format == UE::Mutable::Private::EMeshBufferFormat::Float32;
+		}
+
+		return true;
+	}
+	
+	
+	bool IsUseFullPrecisionTangentBasis(const UE::Mutable::Private::FMesh& Mesh)
+	{
+		int32 BufferIndex = -1;
+		int32 ChannelIndex = -1;
+		Mesh.VertexBuffers.FindChannel(UE::Mutable::Private::EMeshBufferSemantic::Tangent, 0, &BufferIndex, &ChannelIndex);
+		if (BufferIndex >= 0 && ChannelIndex >= 0)
+		{
+			return Mesh.VertexBuffers.Buffers[BufferIndex].Channels[ChannelIndex].Format == UE::Mutable::Private::EMeshBufferFormat::Float32;
+		}
+		
+		return false;
+	}
+	
+	
 	void InitVertexBuffersWithDummyData(
 		FSkeletalMeshLODRenderData& LODResource,
 		const UE::Mutable::Private::FMesh* InMutableMesh,
@@ -312,8 +347,8 @@ namespace UnrealConversionUtils
 		const UE::Mutable::Private::FMeshBufferSet& MutableMeshVertexBuffers = InMutableMesh->GetVertexBuffers();
 		check(MutableMeshVertexBuffers.GetElementCount() > 0);
 
-		const bool bUseFullPrecisionUVs = true;
-		const bool bUseFullPrecisionTangentBasis = false;
+		const bool bUseFullPrecisionUVs = IsUseFullPrecisionUVs(*InMutableMesh);
+		const bool bUseFullPrecisionTangentBasis = IsUseFullPrecisionTangentBasis(*InMutableMesh);
 
 		const uint32 NumVertices = MutableMeshVertexBuffers.GetElementCount();
 		const uint32 NumTexCoords = MutableMeshVertexBuffers.GetBufferChannelCount(MUTABLE_VERTEXBUFFER_TEXCOORDS);
@@ -415,27 +450,18 @@ namespace UnrealConversionUtils
 		const int32 NumVertices = MutableMeshVertexBuffers.IsDescriptor()
 				? 0 
 				: MutableMeshVertexBuffers.GetElementCount();
-
-		//const ESkeletalMeshVertexFlags BuildFlags = SkeletalMesh->GetVertexBufferFlags();
-		//const bool bUseFullPrecisionUVs = EnumHasAllFlags(BuildFlags, ESkeletalMeshVertexFlags::UseFullPrecisionUVs);
-		bool bUseFullPrecisionUVs = true;
-		{
-			int32 TexCoordsBufferIndex = -1;
-			int32 TexCoordsChannelIndex = -1;
-			MutableMesh->VertexBuffers.FindChannel(UE::Mutable::Private::EMeshBufferSemantic::TexCoords, 0, &TexCoordsBufferIndex, &TexCoordsChannelIndex);
-			if (TexCoordsBufferIndex >= 0 && TexCoordsChannelIndex >= 0)
-			{
-				bUseFullPrecisionUVs = MutableMesh->VertexBuffers.Buffers[TexCoordsBufferIndex].Channels[TexCoordsChannelIndex].Format == UE::Mutable::Private::EMeshBufferFormat::Float32;
-			}
-		}
-
+		
+		const bool bUseFullPrecisionUVs = IsUseFullPrecisionUVs(*MutableMesh);
+		const bool bUseFullPrecisionTangentBasis = IsUseFullPrecisionTangentBasis(*MutableMesh);
+		
 		const int NumTexCoords = MutableMeshVertexBuffers.GetBufferChannelCount(MUTABLE_VERTEXBUFFER_TEXCOORDS);
-
+		
 		FStaticMeshVertexBuffers_InitWithMutableData(
 			LODResource.StaticVertexBuffers,
 			NumVertices,
 			NumTexCoords,
 			bUseFullPrecisionUVs,
+			bUseFullPrecisionTangentBasis,
 			bAllowCPUAccess,
 			MutableMeshVertexBuffers.GetBufferData(MUTABLE_VERTEXBUFFER_POSITION),
 			MutableMeshVertexBuffers.GetBufferData(MUTABLE_VERTEXBUFFER_TANGENT),
