[mutable] Improved reporting of cyclical CO references.

#rnx

#Synced-CL 50486562

diff --git a/Engine/Plugins/Mutable/Source/CustomizableObjectEditor/Private/MuCOE/CustomizableObjectCompiler.cpp b/Engine/Plugins/Mutable/Source/CustomizableObjectEditor/Private/MuCOE/CustomizableObjectCompiler.cpp
--- a/Engine/Plugins/Mutable/Source/CustomizableObjectEditor/Private/MuCOE/CustomizableObjectCompiler.cpp
+++ b/Engine/Plugins/Mutable/Source/CustomizableObjectEditor/Private/MuCOE/CustomizableObjectCompiler.cpp
@@ -7,7 +7,6 @@
 #include "ClothConfig.h"
 #include "HAL/RunnableThread.h"
 #include "Interfaces/ITargetPlatformManagerModule.h"
-
 #include "MessageLogModule.h"
 #include "GenerateMutableSource/GenerateMutableSourceComponent.h"
 #include "MuCO/CustomizableObject.h"
@@ -298,13 +297,27 @@ void FCustomizableObjectCompiler::Compile(const TSharedRef<FCompilationRequest>&
 		return;
 	}
 
-	const UCustomizableObject* RootObject = GraphTraversal::GetRootObject(CurrentObject);
+	TArray<const UObject*> VisitedCustomizableObjects;
+	const UCustomizableObject* RootObject = GraphTraversal::GetRootObject(CurrentObject, &VisitedCustomizableObjects); 
+	
+	// No root object means a cycle was found when looking for it
 	if (!RootObject)
 	{
-		const UCustomizableObject* OtherObject = GraphTraversal::GetCylindricallyReferencedObject(CurrentObject);
-		check(OtherObject);
+		CompilerLog(LOCTEXT("CompilationFailure_COCircularReferenceFound", "Failed to compile Customizable Object. Customizable Object circular reference found."), CurrentObject, EMessageSeverity::Warning);
+
+		// Report objects part of the cycle
+		check(VisitedCustomizableObjects.Num() >= 2);
+		if (VisitedCustomizableObjects.Num() >= 2)
+		{
+			const FString InvalidObjectName = "Unknown";
+				
+			const FText CycleStartObjectName =  FText::FromString( ensure(VisitedCustomizableObjects[0]) ? VisitedCustomizableObjects[0]->GetName() : InvalidObjectName );
+			const FText CycleClosureObjectName = FText::FromString( ensure(VisitedCustomizableObjects[VisitedCustomizableObjects.Num() - 1]) ? VisitedCustomizableObjects[VisitedCustomizableObjects.Num() - 1]->GetName() : InvalidObjectName );
+			const FText WarningText = FText::Format(LOCTEXT("CompilationFailure_COCircularReferenceFoundDetails", "Circular reference found starting at '{0}' CO and finishing at '{1}' CO. COs part of the cycle :"),  CycleStartObjectName,  CycleClosureObjectName);
+		
+			CompilerLog( WarningText , VisitedCustomizableObjects, EMessageSeverity::Warning);
+		}
 
-		CompilerLog( LOCTEXT("CompilationFailure_COCircularReferenceFound", "Failed to compile Customizable Object. A circular reference between this and another CO was found."), {CurrentObject, OtherObject}, EMessageSeverity::Warning);
 		CompleteRequest(ECompilationStatePrivate::Completed, ECompilationResultPrivate::Errors);
 		return;
 	}
diff --git a/Engine/Plugins/Mutable/Source/CustomizableObjectEditor/Private/MuCOE/GraphTraversal.cpp b/Engine/Plugins/Mutable/Source/CustomizableObjectEditor/Private/MuCOE/GraphTraversal.cpp
--- a/Engine/Plugins/Mutable/Source/CustomizableObjectEditor/Private/MuCOE/GraphTraversal.cpp
+++ b/Engine/Plugins/Mutable/Source/CustomizableObjectEditor/Private/MuCOE/GraphTraversal.cpp
@@ -339,7 +339,7 @@ bool HasCandidateAsParent(UCustomizableObjectNodeObject* Node, UCustomizableObje
 }
 
 
-UCustomizableObject* GetFullGraphRootObject(const UCustomizableObjectNodeObject* Node, TArray<UCustomizableObject*>& VisitedObjects)
+UCustomizableObject* GetFullGraphRootObject(const UCustomizableObjectNodeObject* Node, TArray<const UObject*>& VisitedObjects)
 {
 	if (Node->ParentObject != nullptr)
 	{
@@ -369,28 +369,6 @@ UCustomizableObject* GetFullGraphRootObject(const UCustomizableObjectNodeObject*
 }
 
 
-UCustomizableObject* GetCylindricallyReferencedObject(const UCustomizableObjectNodeObject* Node, TArray<UCustomizableObject*>& VisitedObjects)
-{
-	if (Node->ParentObject != nullptr)
-	{
-		VisitedObjects.Add(Node->ParentObject);
-
-		UCustomizableObjectNodeObject* Root = GetRootNode(Node->ParentObject);
-		check(Root);
-		
-		if (VisitedObjects.Contains(Root->ParentObject))
-		{
-			//There is a cycle so report about what node is the cause
-			return Root->ParentObject;
-		}
-		else
-		{
-			return GetCylindricallyReferencedObject(Root, VisitedObjects);
-		}
-	}
-
-	return nullptr;
-}
 
 
 UCustomizableObject* GraphTraversal::GetObject(const UCustomizableObjectNode& Node)
@@ -404,22 +382,29 @@ UCustomizableObject* GraphTraversal::GetObject(const UCustomizableObjectNode& No
 }
 
 
-UCustomizableObject* GraphTraversal::GetRootObject(UCustomizableObject* ChildObject)
+UCustomizableObject* GraphTraversal::GetRootObject(UCustomizableObject* ChildObject, TArray<const UObject*>* OutVisitedObjects)
 {
 	const UCustomizableObject* ConstChildObject = ChildObject;
-	return const_cast<UCustomizableObject*>(GetRootObject(ConstChildObject));
+	return const_cast<UCustomizableObject*>(GetRootObject(ConstChildObject, OutVisitedObjects));
 }
 
 
-const UCustomizableObject* GraphTraversal::GetRootObject(const UCustomizableObject* ChildObject)
+const UCustomizableObject* GraphTraversal::GetRootObject(const UCustomizableObject* ChildObject, TArray<const UObject*>* OutVisitedObjects)
 {
 	// Grab a node to start the search -> Get the root since it should be always present
 	UCustomizableObjectNodeObject* ObjectRootNode = GetRootNode(ChildObject);
 
 	if (ObjectRootNode && ObjectRootNode->ParentObject)
 	{
-		TArray<UCustomizableObject*> VisitedNodes;
-		return GetFullGraphRootObject(ObjectRootNode, VisitedNodes);
+		if (!OutVisitedObjects)	
+		{
+			TArray<const UObject*> VisitedObjects;
+			return GetFullGraphRootObject(ObjectRootNode, VisitedObjects);
+		}
+		else
+		{
+			return GetFullGraphRootObject(ObjectRootNode, *OutVisitedObjects);
+		}
 	}
 
 	// No parent object found, return input as the parent of the graph
@@ -428,22 +413,6 @@ const UCustomizableObject* GraphTraversal::GetRootObject(const UCustomizableObje
 }
 
 
-const UCustomizableObject* GraphTraversal::GetCylindricallyReferencedObject(const UCustomizableObject* ChildObject)
-{
-	// Grab a node to start the search -> Get the root since it should be always present
-	UCustomizableObjectNodeObject* ObjectRootNode = GetRootNode(ChildObject);
-
-	if (ObjectRootNode && ObjectRootNode->ParentObject)
-	{
-		TArray<UCustomizableObject*> VisitedNodes;
-		return GetCylindricallyReferencedObject(ObjectRootNode, VisitedNodes);
-	}
-
-	// No parent object found so there can be no cycle
-	return nullptr;
-}
-
-
 void RecursiveVisitNodes(UCustomizableObjectNode& CurrentNode, const TFunction<void(UCustomizableObjectNode&)>& VisitFunction,
 	TArray<const UCustomizableObjectNodeMacroInstance*>& MacroContext, TSet<TPair<UCustomizableObjectNode*, TArray<const UCustomizableObjectNodeMacroInstance*>>>& VisitedNodes, const TMultiMap<FGuid, UCustomizableObjectNodeObject*>& ObjectGroupMap)
 {
@@ -613,7 +582,7 @@ UCustomizableObjectNodeObject* GraphTraversal::GetFullGraphRootNode(const UCusto
 			}
 			else
 			{
-				return GetFullGraphRootNode(Root->ParentObject, VisitedObjects);
+				return GraphTraversal::GetFullGraphRootNode(Root->ParentObject, VisitedObjects);
 			}
 		}
 	}
diff --git a/Engine/Plugins/Mutable/Source/CustomizableObjectEditor/Private/MuCOE/GraphTraversal.h b/Engine/Plugins/Mutable/Source/CustomizableObjectEditor/Private/MuCOE/GraphTraversal.h
--- a/Engine/Plugins/Mutable/Source/CustomizableObjectEditor/Private/MuCOE/GraphTraversal.h
+++ b/Engine/Plugins/Mutable/Source/CustomizableObjectEditor/Private/MuCOE/GraphTraversal.h
@@ -70,7 +70,7 @@ bool HasCandidateAsParent(UCustomizableObjectNodeObject* Node, UCustomizableObje
  * @param VisitedObjects List of objects used to record the search path already taken. It is used to determine if a cycle has been found.
  * return The root CO. It can be null if a cycle has been found.
  */
-UCustomizableObject* GetFullGraphRootObject(const UCustomizableObjectNodeObject* Node, TArray<UCustomizableObject*>& VisitedObjects);
+UCustomizableObject* GetFullGraphRootObject(const UCustomizableObjectNodeObject* Node, TArray<const UObject*>& VisitedObjects);
 
 namespace GraphTraversal
 {
@@ -84,16 +84,8 @@ namespace GraphTraversal
 		 * @param InObject Customizable object whose root CO we are asking for.
 		 * @return The CO that is the root of the provided Customizable Object. It can be equal to InObject if the provided
 		 * object does not have any parent.	*/
-	CUSTOMIZABLEOBJECTEDITOR_API UCustomizableObject* GetRootObject(UCustomizableObject* ChildObject);
-	CUSTOMIZABLEOBJECTEDITOR_API const UCustomizableObject* GetRootObject(const UCustomizableObject* ChildObject);
-
-	/**
-	 * Provided a CO it looks for cylindrical references while going up in search of the root CO. It is quite similar with GetRootObject but, instead
-	 * of returning the root CO will instead return you the CO that was found to be referencing the provided node
-	 * @param ChildObject The object where to start the search for cylindrical references.
-	 * @return The CO that is referencing the provided one.
-	 */
-	CUSTOMIZABLEOBJECTEDITOR_API const UCustomizableObject* GetCylindricallyReferencedObject(const UCustomizableObject* ChildObject);
+	CUSTOMIZABLEOBJECTEDITOR_API UCustomizableObject* GetRootObject(UCustomizableObject* ChildObject, TArray<const UObject*>* OutVisitedObjects = nullptr);
+	CUSTOMIZABLEOBJECTEDITOR_API const UCustomizableObject* GetRootObject(const UCustomizableObject* ChildObject, TArray<const UObject*>* OutVisitedObjects = nullptr);
 	
 	/** Return the full graph Customizable Object Node root of the node given as parameter */
 	UCustomizableObjectNodeObject* GetFullGraphRootNode(const UCustomizableObject* Object, TArray<const UCustomizableObject*>& VisitedObjects);
