[mutable] Added missing core redirects for the updates in the Passthrough Mesh Component node

#rnx
#Synced-CL 0

diff --git a/Engine/Plugins/Mutable/Config/DefaultMutable.ini b/Engine/Plugins/Mutable/Config/DefaultMutable.ini
--- a/Engine/Plugins/Mutable/Config/DefaultMutable.ini
+++ b/Engine/Plugins/Mutable/Config/DefaultMutable.ini
@@ -28,6 +28,7 @@
 +PropertyRedirects=(OldName="/Script/CustomizableObject.CustomizableObject.bEnableAnimBpPhysicsAssetsManipualtion",NewName="/Script/CustomizableObject.CustomizableObject.bEnableAnimBpPhysicsAssetsManipulation")
 +PropertyRedirects=(OldName="/Script/CustomizableObject.CompileParams.bGatherReferneces",NewName="/Script/CustomizableObject.CompileParams.bGatherReferences")
 +PropertyRedirects=(OldName="/Script/CustomizableObject.CompileParams.bSkipIfOutOfDate",NewName="/Script/CustomizableObject.CompileParams.bSkipIfNotOutOfDate")
++PropertyRedirects=(OldName="/Script/CustomizableObjectEditor.CustomizableObjectNodeComponentPassthroughMesh.MaterialSlotOverridePins",NewName="/Script/CustomizableObjectEditor.CustomizableObjectNodeComponentPassthroughMesh.MaterialSlotPins")
 
 ; Pin Data renames
 +ClassRedirects=(OldName="/Script/CustomizableObjectEditor.CONodeMaterialBreakPinData",NewName="/Script/CustomizableObjectEditor.CONodeMaterialBreakParameterPinData")
@@ -45,6 +46,7 @@
 +ClassRedirects=(OldName="/Script/CustomizableObjectEditor.CustomizableObjectNodeMaterialPinDataImage",NewName="/Script/CustomizableObjectEditor.CONodeSkeletalMeshSectionPinDataImage")
 +ClassRedirects=(OldName="/Script/CustomizableObjectEditor.CustomizableObjectNodeMaterialPinDataVector",NewName="/Script/CustomizableObjectEditor.CONodeSkeletalMeshSectionPinDataVector")
 +ClassRedirects=(OldName="/Script/CustomizableObjectEditor.CustomizableObjectNodeMaterialPinDataScalar",NewName="/Script/CustomizableObjectEditor.CONodeSkeletalMeshSectionPinDataScalar")
++ClassRedirects=(OldName="/Script/CustomizableObjectEditor.CONodeComponentPassthroughMeshMaterialOverridePinData",NewName="/Script/CustomizableObjectEditor.CONodeComponentPassthroughMeshMaterialPinData")
 
 ; Settings for the Mutable plugin
 [Features]
