[Mutable] SkeletalMesh Parameter: convert mesh sockets. Moved mesh sockets from ModelResources to Program.
Add conversion flags to runtime mesh conversion.
Remove Model Resources' unused properties.

#rnx

#Synced-CL 50587801

diff --git a/Engine/Plugins/Mutable/Source/CustomizableObject/Internal/MuCO/CustomizableObjectInstancePrivate.h b/Engine/Plugins/Mutable/Source/CustomizableObject/Internal/MuCO/CustomizableObjectInstancePrivate.h
--- a/Engine/Plugins/Mutable/Source/CustomizableObject/Internal/MuCO/CustomizableObjectInstancePrivate.h
+++ b/Engine/Plugins/Mutable/Source/CustomizableObject/Internal/MuCO/CustomizableObjectInstancePrivate.h
@@ -383,7 +383,7 @@ UE_API void BuildOrCopyClothingData(const TSharedRef<FUpdateContextPrivate>& Con
 
 UE_API bool BuildSkeletonData(const TSharedRef<FUpdateContextPrivate>& Context, USkeletalMesh& SkeletalMesh, int32 ComponentIndex);
 
-UE_API void BuildMeshSockets(USkeletalMesh* SkeletalMesh, const UModelResources& ModelResources, const FMutableRefSkeletalMeshData& RefSkeletalMeshData, UE::Mutable::Private::TManagedPtr<const UE::Mutable::Private::FMesh> MutableMesh);
+UE_API void BuildMeshSockets(USkeletalMesh* SkeletalMesh, const UModelResources& ModelResources, const FMutableRefSkeletalMeshData& RefSkeletalMeshData, UE::Mutable::Private::TManagedPtr<const UE::Mutable::Private::FMesh>& MutableMesh);
 
 UE_API bool BuildOrCopyRenderData(const TSharedRef<FUpdateContextPrivate>& Context, USkeletalMesh* SkeletalMesh, UCustomizableObjectInstance* CustomizableObjectInstance, int32 ComponentIndex);
 
diff --git a/Engine/Plugins/Mutable/Source/CustomizableObject/Internal/MuCO/CustomizableObjectPrivate.h b/Engine/Plugins/Mutable/Source/CustomizableObject/Internal/MuCO/CustomizableObjectPrivate.h
--- a/Engine/Plugins/Mutable/Source/CustomizableObject/Internal/MuCO/CustomizableObjectPrivate.h
+++ b/Engine/Plugins/Mutable/Source/CustomizableObject/Internal/MuCO/CustomizableObjectPrivate.h
@@ -90,7 +90,6 @@ struct FCustomizableObjectStreameableResourceId
 	{
 		None                  = 0,
 		AssetUserData         = 1,
-		Socket	              = 4,
 	};
 
 	uint64 Id   : 64 - 8;
@@ -330,10 +329,6 @@ struct FMutableRefSocket
 	UPROPERTY()
 	bool bForceAlwaysAnimated = false;
 
-	// When two sockets have the same name, the one with higher priority will be picked and the other discarded
-	UPROPERTY()
-	int32 Priority = -1;
-
 	CUSTOMIZABLEOBJECT_API bool operator ==(const FMutableRefSocket& Other) const;
 
 #if WITH_EDITORONLY_DATA
@@ -1028,10 +1023,6 @@ public:
 	TArray<TSoftObjectPtr<UTexture2D>> RuntimeReferencedTextures;
 #endif
 	
-	/** Physics assets gathered from the SkeletalMeshes, to be used in mesh generation in-game */
-	UPROPERTY()
-	TArray<TSoftObjectPtr<UPhysicsAsset>> PhysicsAssets;
-
 	/** UAnimBlueprint assets gathered from the SkeletalMesh, to be used in mesh generation in-game */
 	UPROPERTY()
 	TArray<TSoftClassPtr<UAnimInstance>> AnimBPs;
@@ -1040,14 +1031,6 @@ public:
 	UPROPERTY()
 	TArray<FAnimBpOverridePhysicsAssetsInfo> AnimBpOverridePhysiscAssetsInfo;
 
-	/** Material slot names for the materials referenced by the surfaces. */
-	UPROPERTY()
-	TArray<FName> MaterialSlotNames;
-
-	/** Mesh sockets provided by the part skeletal meshes, to be merged in the generated meshes */
-	UPROPERTY()
-	TMap<uint32, FMutableRefSocket> Sockets;
-
 	UPROPERTY()
 	TArray<FMutableSkinWeightProfileInfo> SkinWeightProfilesInfo;
 
@@ -1500,6 +1483,8 @@ enum class ECustomizableObjectVersions : int32
 	MutableSkeletonToPassthroughObject,
 	
 	UseCompressedMorphsForAllMorphData,
+
+	TestMoveSocketsToCore,
 	
 	// -----<new versions can be added above this line>--------
 	LastCustomizableObjectVersion
diff --git a/Engine/Plugins/Mutable/Source/CustomizableObject/Internal/MuCO/UnrealConversionUtils.h b/Engine/Plugins/Mutable/Source/CustomizableObject/Internal/MuCO/UnrealConversionUtils.h
--- a/Engine/Plugins/Mutable/Source/CustomizableObject/Internal/MuCO/UnrealConversionUtils.h
+++ b/Engine/Plugins/Mutable/Source/CustomizableObject/Internal/MuCO/UnrealConversionUtils.h
@@ -169,7 +169,7 @@ namespace UnrealConversionUtils
 			uint32 SourceSectionBaseVertexIndex, uint32 SourceSectionNumVertices, uint32 SourceMeshNumVertices, 
 			FMorphTargetVertexInfoBuffers* OptionalMorphTargetVertexBuffer);
 
-	CUSTOMIZABLEOBJECT_API UE::Tasks::FTask ConvertSkeletalMeshFromRuntimeData(USkeletalMesh* SkeletalMesh, int32 LODIndex, int32 SectionIndex, UE::Mutable::Private::TManagedPtr<UE::Mutable::Private::FMesh> Result);
+	CUSTOMIZABLEOBJECT_API UE::Tasks::FTask ConvertSkeletalMeshFromRuntimeData(USkeletalMesh* SkeletalMesh, int32 LODIndex, int32 SectionIndex, uint8 ConversionFlags, UE::Mutable::Private::TManagedPtr<UE::Mutable::Private::FMesh> Result);
 
 	TSharedPtr<UE::Mutable::Private::FPhysicsBody> MakePhysicsBodyFromAsset(UPhysicsAsset* Asset, const TArray<uint8>& BodySetupRelevancyMap);
 	
diff --git a/Engine/Plugins/Mutable/Source/CustomizableObject/Internal/MuCO/UnrealMutableImageProvider.h b/Engine/Plugins/Mutable/Source/CustomizableObject/Internal/MuCO/UnrealMutableImageProvider.h
--- a/Engine/Plugins/Mutable/Source/CustomizableObject/Internal/MuCO/UnrealMutableImageProvider.h
+++ b/Engine/Plugins/Mutable/Source/CustomizableObject/Internal/MuCO/UnrealMutableImageProvider.h
@@ -20,7 +20,7 @@ public:
 	CUSTOMIZABLEOBJECT_API virtual TTuple<UE::Tasks::FTask, TFunction<void()>> GetImageAsync(UTexture* Texture, uint8 MipmapsToSkip, bool bLoadMipTail, TFunction<void(UE::Mutable::Private::TManagedPtr<UE::Mutable::Private::FImage>)>& ResultCallback) override;
 	CUSTOMIZABLEOBJECT_API virtual TTuple<UE::Tasks::FTask, TFunction<void()>> GetReferencedImageAsync( int32 Id, uint8 MipmapsToSkip, TFunction<void(UE::Mutable::Private::TManagedPtr<UE::Mutable::Private::FImage>)>& ResultCallback) override;
 	CUSTOMIZABLEOBJECT_API virtual UE::Mutable::Private::FExtendedImageDesc GetImageDesc(UTexture* Texture) override;
-	CUSTOMIZABLEOBJECT_API virtual TTuple<UE::Tasks::FTask, TFunction<void()>> GetMeshAsync(USkeletalMesh* SkeletalMesh, int32 LODIndex, int32 SectionIndex, TFunction<void(UE::Mutable::Private::TManagedPtr<UE::Mutable::Private::FMesh>)>& ResultCallback) override;
+	CUSTOMIZABLEOBJECT_API virtual TTuple<UE::Tasks::FTask, TFunction<void()>> GetMeshAsync(USkeletalMesh* SkeletalMesh, int32 LODIndex, int32 SectionIndex, uint8 ConversionFlags, TFunction<void(UE::Mutable::Private::TManagedPtr<UE::Mutable::Private::FMesh>)>& ResultCallback) override;
 
 #if WITH_EDITOR
 	void CacheRuntimeReferencedImages(const TArray<TSoftObjectPtr<UTexture2D>>& RuntimeReferencedTextures);
diff --git a/Engine/Plugins/Mutable/Source/CustomizableObject/Internal/MuCO/UnrealToMutableTextureConversionUtils.h b/Engine/Plugins/Mutable/Source/CustomizableObject/Internal/MuCO/UnrealToMutableTextureConversionUtils.h
--- a/Engine/Plugins/Mutable/Source/CustomizableObject/Internal/MuCO/UnrealToMutableTextureConversionUtils.h
+++ b/Engine/Plugins/Mutable/Source/CustomizableObject/Internal/MuCO/UnrealToMutableTextureConversionUtils.h
@@ -9,6 +9,31 @@
 
 #define UE_API CUSTOMIZABLEOBJECT_API
 
+// Flags that can influence the mesh conversion
+enum class EMutableMeshConversionFlags : uint8
+{
+	// 
+	None = 0,
+	// Ignore the skeleton and skinning
+	IgnoreSkinning = 1 << 0,
+
+	// Ignore Physics assets
+	IgnorePhysics = 1 << 1,
+
+	// Ignore Morphs
+	IgnoreMorphs = 1 << 2,
+
+	// Ignore Texture coordinates
+	IgnoreTexCoords = 1 << 3,
+
+	// Ignore AssetUserData
+	IgnoreAUD = 1 << 4,
+
+	// Prevent this mesh generation from adding per mesh metadata. 
+	DoNotCreateMeshMetadata = 1 << 5
+};
+
+
 #if WITH_EDITOR
 
 // Forward declarations
@@ -48,31 +73,6 @@ private:
 };
 
 
-
-// Flags that can influence the mesh conversion
-enum class EMutableMeshConversionFlags : uint32
-{
-	// 
-	None = 0,
-	// Ignore the skeleton and skinning
-	IgnoreSkinning = 1 << 0,
-
-	// Ignore Physics assets
-	IgnorePhysics = 1 << 1,
-
-	// Ignore Morphs
-	IgnoreMorphs = 1 << 2,
-
-	// Ignore Texture coordinates
-	IgnoreTexCoords = 1 << 3,
-
-	// Ignore AssetUserData
-	IgnoreAUD = 1 << 4,
-
-	// Prevent this mesh generation from adding per mesh metadata. 
-	DoNotCreateMeshMetadata = 1 << 5
-};
-
 ENUM_CLASS_FLAGS(EMutableMeshConversionFlags)
 
 struct FMutableSourceSurfaceMetadata
@@ -111,10 +111,6 @@ struct FMutableSourceMeshData
 	bool bUseAllRealTimeMorphs = false;
 	TArray<FString> UsedRealTimeMorphTargetNames;
 
-	/** */
-	int32 SocketPriority = 0;
-	int16 BonePosePriority = 0;
-
 	/** */
 	bool bOnlyConnectedLOD = false;
 
@@ -133,8 +129,6 @@ inline uint32 GetTypeHash(const FMutableSourceMeshData& Key)
 	GuidHash = HashCombineFast(GuidHash, GetTypeHash(Key.BaseLODIndex));
 	GuidHash = HashCombineFast(GuidHash, GetTypeHash(Key.BaseSectionIndex));
 	GuidHash = HashCombineFast(GuidHash, GetTypeHash(Key.LODOffset));
-	GuidHash = HashCombineFast(GuidHash, GetTypeHash(Key.SocketPriority));
-	GuidHash = HashCombineFast(GuidHash, GetTypeHash(Key.BonePosePriority));
 	GuidHash = HashCombineFast(GuidHash, GetTypeHash(Key.bOnlyConnectedLOD));
 	GuidHash = HashCombineFast(GuidHash, GetTypeHash(Key.ComponentId));
 	GuidHash = HashCombineFast(GuidHash, GetTypeHash(Key.Flags));
diff --git a/Engine/Plugins/Mutable/Source/CustomizableObject/Private/MuCO/CustomizableObject.cpp b/Engine/Plugins/Mutable/Source/CustomizableObject/Private/MuCO/CustomizableObject.cpp
--- a/Engine/Plugins/Mutable/Source/CustomizableObject/Private/MuCO/CustomizableObject.cpp
+++ b/Engine/Plugins/Mutable/Source/CustomizableObject/Private/MuCO/CustomizableObject.cpp
@@ -3245,6 +3245,8 @@ void UCustomizableObjectPrivate::LogMemory()
 			SkeletonsTotal += sizeof(UE::Mutable::Private::FSkeleton) + Entry->BoneNames.GetAllocatedSize() + Entry->BoneIds.GetAllocatedSize() + Entry->BoneParents.GetAllocatedSize();
 		}
 
+		int32 SocketsTotal = Program.ConstantSockets.GetAllocatedSize();
+
 		int32 PhysicsTotal = Program.ConstantPhysicsBodies.GetAllocatedSize();
 		for (const UE::Mutable::Private::TManagedPtr<const UE::Mutable::Private::FPhysicsBody>& Entry : Program.ConstantPhysicsBodies)
 		{
@@ -3269,6 +3271,7 @@ void UCustomizableObjectPrivate::LogMemory()
 			+ ImageTotal + MeshesTotal 
 			+ StringsTotal + ConstantNames
 			+ SkeletonsTotal + PhysicsTotal
+			+ SocketsTotal
 			+ ParametersTotal;
 
 
@@ -3289,11 +3292,8 @@ void UCustomizableObjectPrivate::LogMemory()
 			{
 				Arch << ModelResources->ReferenceSkeletalMeshesData;
 				Arch << ModelResources->PassthroughObjects;
-				Arch << ModelResources->PhysicsAssets;
 				Arch << ModelResources->AnimBPs;
 				Arch << ModelResources->AnimBpOverridePhysiscAssetsInfo;
-				Arch << ModelResources->MaterialSlotNames;
-				Arch << ModelResources->Sockets;
 				Arch << ModelResources->SkinWeightProfilesInfo;
 				Arch << ModelResources->MeshMetadata;
 				Arch << ModelResources->SurfaceMetadata;
@@ -3321,6 +3321,7 @@ void UCustomizableObjectPrivate::LogMemory()
 		UE_LOG(LogMutable, Log, TEXT("    OpAddress        : %8d"), OpAddress);
 		UE_LOG(LogMutable, Log, TEXT("    Roms             : %8d"), Roms);
 		UE_LOG(LogMutable, Log, TEXT("    Strings          : %8d"), StringsTotal);
+		UE_LOG(LogMutable, Log, TEXT("    Names			   : %8d"), ConstantNames);
 		UE_LOG(LogMutable, Log, TEXT("    Images           : %8d"), ImageTotal);
 		UE_LOG(LogMutable, Log, TEXT("      Buffer         : %8d"), ConstantImages);
 		UE_LOG(LogMutable, Log, TEXT("      LODIndices     : %8d"), ConstantImageLODIndices);
@@ -3332,6 +3333,7 @@ void UCustomizableObjectPrivate::LogMemory()
 		UE_LOG(LogMutable, Log, TEXT("      PermanentBuf   : %8d"), ConstantMeshesPermanent);
 		UE_LOG(LogMutable, Log, TEXT("      Permanent      : %8d"), PermanentMeshes);
 		UE_LOG(LogMutable, Log, TEXT("    Skeletons        : %8d"), SkeletonsTotal);
+		UE_LOG(LogMutable, Log, TEXT("    Sockets	       : %8d"), SocketsTotal);
 		UE_LOG(LogMutable, Log, TEXT("    PhysicBodies     : %8d"), PhysicsTotal);
 		UE_LOG(LogMutable, Log, TEXT("    Parameters       : %8d"), ParametersTotal);
 	}
@@ -3375,8 +3377,7 @@ bool FMutableRefSocket::operator==(const FMutableRefSocket& Other) const
 		RelativeLocation == Other.RelativeLocation &&
 		RelativeRotation == Other.RelativeRotation &&
 		RelativeScale == Other.RelativeScale &&
-		bForceAlwaysAnimated == Other.bForceAlwaysAnimated &&
-		Priority == Other.Priority;
+		bForceAlwaysAnimated == Other.bForceAlwaysAnimated;
 }
 
 
@@ -3447,7 +3448,6 @@ FArchive& operator<<(FArchive& Ar, FMutableRefSocket& Data)
 	Ar << Data.RelativeRotation;
 	Ar << Data.RelativeScale;
 	Ar << Data.bForceAlwaysAnimated;
-	Ar << Data.Priority;
 
 	return Ar;
 }
diff --git a/Engine/Plugins/Mutable/Source/CustomizableObject/Private/MuCO/CustomizableObjectInstance.cpp b/Engine/Plugins/Mutable/Source/CustomizableObject/Private/MuCO/CustomizableObjectInstance.cpp
--- a/Engine/Plugins/Mutable/Source/CustomizableObject/Private/MuCO/CustomizableObjectInstance.cpp
+++ b/Engine/Plugins/Mutable/Source/CustomizableObject/Private/MuCO/CustomizableObjectInstance.cpp
@@ -3814,97 +3814,60 @@ bool BuildSkeletonData(const TSharedRef<FUpdateContextPrivate>& Context, USkelet
 }
 
 
-void BuildMeshSockets(USkeletalMesh* SkeletalMesh, const UModelResources& ModelResources, const FMutableRefSkeletalMeshData& RefSkeletalMeshData, UE::Mutable::Private::TManagedPtr<const UE::Mutable::Private::FMesh> MutableMesh)
+void BuildMeshSockets(USkeletalMesh* SkeletalMesh, const UModelResources& ModelResources, const FMutableRefSkeletalMeshData& RefSkeletalMeshData, UE::Mutable::Private::TManagedPtr<const UE::Mutable::Private::FMesh>& MutableMesh)
 {
 	// Build mesh sockets.
 	MUTABLE_CPUPROFILER_SCOPE(UCustomizableInstancePrivate::BuildMeshSockets);
-	check(SkeletalMesh);
 
-	const uint32 SocketCount = RefSkeletalMeshData.Sockets.Num();
+	check(SkeletalMesh);
+	check(MutableMesh);
 
 	TArray<TObjectPtr<USkeletalMeshSocket>>& Sockets = SkeletalMesh->GetMeshOnlySocketList();
-	Sockets.Empty(SocketCount);
-	TMap<FName, TTuple<int32, int32>> SocketMap; // Maps Socket name to Sockets Array index and priority
-	
-	// Add sockets used by the SkeletalMesh of reference.
-	{
-		MUTABLE_CPUPROFILER_SCOPE(BuildMeshSockets_RefMeshSockets);
-	
-		for (uint32 SocketIndex = 0; SocketIndex < SocketCount; ++SocketIndex)
-		{
-			const FMutableRefSocket& RefSocket = RefSkeletalMeshData.Sockets[SocketIndex];
+	TArray<FName> SocketNames;
 
-			USkeletalMeshSocket* Socket = NewObject<USkeletalMeshSocket>(SkeletalMesh, RefSocket.SocketName);
+	const int32 NumSockets = FMath::Max(RefSkeletalMeshData.Sockets.Num(), MutableMesh->Sockets.Num());
+	Sockets.Reserve(NumSockets);
+	SocketNames.Reserve(NumSockets);
 
-			Socket->SocketName = RefSocket.SocketName;
-			Socket->BoneName = RefSocket.BoneName;
+	// Add Mutable mesh sockets.
+	for (const UE::Mutable::Private::FMeshSocket& MutableSocket : MutableMesh->Sockets)
+	{
+		USkeletalMeshSocket* Socket = NewObject<USkeletalMeshSocket>(SkeletalMesh, MutableSocket.SocketName);
+		Socket->SocketName = MutableSocket.SocketName;
+		Socket->BoneName = MutableSocket.BoneName;
 
-			Socket->RelativeLocation = RefSocket.RelativeLocation;
-			Socket->RelativeRotation = RefSocket.RelativeRotation;
-			Socket->RelativeScale = RefSocket.RelativeScale;
+		Socket->RelativeLocation = MutableSocket.RelativeLocation;
+		Socket->RelativeRotation = MutableSocket.RelativeRotation;
+		Socket->RelativeScale = MutableSocket.RelativeScale;
 
-			Socket->bForceAlwaysAnimated = RefSocket.bForceAlwaysAnimated;
-			const int32 LastIndex = Sockets.Add(Socket);
+		Socket->bForceAlwaysAnimated = MutableSocket.bForceAlwaysAnimated;
 
-			SocketMap.Add(Socket->SocketName, TTuple<int32, int32>(LastIndex, RefSocket.Priority));
-		}
+		Sockets.Add(Socket);
+		SocketNames.Add(MutableSocket.SocketName);
 	}
 
-	// Add or update sockets modified by Mutable.
-	if (MutableMesh)
+	// Add sockets used by the SkeletalMesh of reference.
+	for (const FMutableRefSocket& RefSocket : RefSkeletalMeshData.Sockets)
 	{
-		MUTABLE_CPUPROFILER_SCOPE(BuildMeshSockets_MutableSockets);
-
-		const TArray<uint64>& StreamedResources = MutableMesh->GetStreamedResources();
-
-		for (uint64 ResourceId : StreamedResources)
+		if (SocketNames.Find(RefSocket.SocketName) != INDEX_NONE)
 		{
-			FCustomizableObjectStreameableResourceId TypedResourceId = BitCast<FCustomizableObjectStreameableResourceId>(ResourceId);
-
-			if (TypedResourceId.Type == (uint8)FCustomizableObjectStreameableResourceId::EType::Socket)
-			{
-				check(TypedResourceId.Id != 0 && TypedResourceId.Id <= TNumericLimits<uint32>::Max());
-				if (const FMutableRefSocket* MutableSocket = ModelResources.Sockets.Find(TypedResourceId.Id))
-				{
-					int32 IndexToWriteSocket = -1;
-
-					if (TTuple<int32, int32>* FoundSocket = SocketMap.Find(MutableSocket->SocketName))
-					{
-						if (FoundSocket->Value < MutableSocket->Priority)
-						{
-							// Overwrite the existing socket because the new mesh part one is higher priority
-							IndexToWriteSocket = FoundSocket->Key;
-							FoundSocket->Value = MutableSocket->Priority;
-						}
-					}
-					else
-					{
-						// New Socket
-						USkeletalMeshSocket* Socket = NewObject<USkeletalMeshSocket>(SkeletalMesh, MutableSocket->SocketName);
-						IndexToWriteSocket = Sockets.Add(Socket);
-						SocketMap.Add(MutableSocket->SocketName, TTuple<int32, int32>(IndexToWriteSocket, MutableSocket->Priority));
-					}
-
-					if (IndexToWriteSocket >= 0)
-					{
-						check(Sockets.IsValidIndex(IndexToWriteSocket));
+			continue;
+		}
 
-						USkeletalMeshSocket* SocketToWrite = Sockets[IndexToWriteSocket];
+		USkeletalMeshSocket* Socket = NewObject<USkeletalMeshSocket>(SkeletalMesh, RefSocket.SocketName);
+		Socket->SocketName = RefSocket.SocketName;
+		Socket->BoneName = RefSocket.BoneName;
 
-						SocketToWrite->SocketName = MutableSocket->SocketName;
-						SocketToWrite->BoneName = MutableSocket->BoneName;
+		Socket->RelativeLocation = RefSocket.RelativeLocation;
+		Socket->RelativeRotation = RefSocket.RelativeRotation;
+		Socket->RelativeScale = RefSocket.RelativeScale;
 
-						SocketToWrite->RelativeLocation = MutableSocket->RelativeLocation;
-						SocketToWrite->RelativeRotation = MutableSocket->RelativeRotation;
-						SocketToWrite->RelativeScale = MutableSocket->RelativeScale;
+		Socket->bForceAlwaysAnimated = RefSocket.bForceAlwaysAnimated;
 
-						SocketToWrite->bForceAlwaysAnimated = MutableSocket->bForceAlwaysAnimated;
-					}
-				}
-			}
-		}
+		Sockets.Add(Socket);
+		SocketNames.Add(RefSocket.SocketName);
 	}
-
+	
 #if !WITH_EDITOR
 	SkeletalMesh->RebuildSocketMap();
 #endif // !WITH_EDITOR
diff --git a/Engine/Plugins/Mutable/Source/CustomizableObject/Private/MuCO/UnrealConversionUtils.cpp b/Engine/Plugins/Mutable/Source/CustomizableObject/Private/MuCO/UnrealConversionUtils.cpp
--- a/Engine/Plugins/Mutable/Source/CustomizableObject/Private/MuCO/UnrealConversionUtils.cpp
+++ b/Engine/Plugins/Mutable/Source/CustomizableObject/Private/MuCO/UnrealConversionUtils.cpp
@@ -8,6 +8,7 @@
 #include "MuCO/CustomizableObjectPrivate.h"
 #include "MuCO/CustomizableObjectSystemPrivate.h"
 #include "MuCO/MutableMeshBufferUtils.h"
+#include "MuCO/UnrealToMutableTextureConversionUtils.h"
 #include "MuR/Skeleton.h"
 #include "MuR/Mesh.h"
 #include "MuR/MeshBufferSet.h"
@@ -15,6 +16,7 @@
 #include "Animation/Skeleton.h"
 #include "Containers/StaticArray.h"
 #include "Engine/SkeletalMesh.h"
+#include "Engine/SkeletalMeshSocket.h"
 #include "UObject/StrongObjectPtr.h"
 #include "GPUSkinVertexFactory.h"
 #include "SkinnedAssetCompiler.h"
@@ -1411,7 +1413,7 @@ namespace UnrealConversionUtils
 	}
 
 	
-	UE::Tasks::FTask ConvertSkeletalMeshFromRuntimeData(USkeletalMesh* SkeletalMesh, int32 LODIndex, int32 SectionIndex, UE::Mutable::Private::TManagedPtr<UE::Mutable::Private::FMesh> Result)
+	UE::Tasks::FTask ConvertSkeletalMeshFromRuntimeData(USkeletalMesh* SkeletalMesh, int32 LODIndex, int32 SectionIndex, uint8 ConversionFlags, UE::Mutable::Private::TManagedPtr<UE::Mutable::Private::FMesh> Result)
 	{
 		MUTABLE_CPUPROFILER_SCOPE(ConvertSkeletalMeshFromRuntimeData);
 
@@ -1529,7 +1531,7 @@ namespace UnrealConversionUtils
 
 		
 		UE::Tasks::FTask FinalConversionTask = UE::Tasks::Launch( TEXT("MutableRuntimeMeshConversionFinal"),
-			[Context, SkeletalMesh = TStrongObjectPtr(SkeletalMesh), LODIndex, SectionIndex, Result]()
+			[Context, SkeletalMesh = TStrongObjectPtr(SkeletalMesh), LODIndex, SectionIndex, ConversionFlags = (EMutableMeshConversionFlags)ConversionFlags, Result]()
 			{
 				MUTABLE_CPUPROFILER_SCOPE(MeshConversion);
 
@@ -1554,9 +1556,13 @@ namespace UnrealConversionUtils
 				Result->SkeletalMeshes.AddUnique(SkeletalMesh);
 
 				const FSkelMeshRenderSection& Section = Context->OriginalLOD->RenderSections[SectionIndex];
-					
-				bool bIgnoreSkeleton = false;
-				if (!bIgnoreSkeleton)
+				
+				const bool bIgnoreSkinning = EnumHasAnyFlags(ConversionFlags, EMutableMeshConversionFlags::IgnoreSkinning);
+				const bool bIgnorePhysics = bIgnoreSkinning || EnumHasAnyFlags(ConversionFlags, EMutableMeshConversionFlags::IgnorePhysics);
+				const bool bIgnoreMorphs = EnumHasAnyFlags(ConversionFlags, EMutableMeshConversionFlags::IgnoreMorphs);
+				const bool bIgnoreTexCoords = EnumHasAnyFlags(ConversionFlags, EMutableMeshConversionFlags::IgnoreTexCoords);
+
+				if (!bIgnoreSkinning)
 				{
 					MUTABLE_CPUPROFILER_SCOPE(MeshConversion_Skeleton)
 					USkeleton* Skeleton = SkeletalMesh->GetSkeleton();
@@ -1630,9 +1636,28 @@ namespace UnrealConversionUtils
 					Result->SetBoneMap(MutableBoneMap);
 				}
 
+				if (!bIgnoreSkinning)
+				{
+					for (int32 SocketIndex = 0; SocketIndex < SkeletalMesh->NumSockets(); ++SocketIndex)
+					{
+						const USkeletalMeshSocket* Socket = SkeletalMesh->GetSocketByIndex(SocketIndex);
+						check(Socket);
+
+						UE::Mutable::Private::FMeshSocket MutableSocket;
+						MutableSocket.SocketName = Socket->SocketName;
+						MutableSocket.BoneName = Socket->BoneName;
+						MutableSocket.RelativeLocation = Socket->RelativeLocation;
+						MutableSocket.RelativeRotation = Socket->RelativeRotation;
+						MutableSocket.RelativeScale = Socket->RelativeScale;
+						MutableSocket.bForceAlwaysAnimated = Socket->bForceAlwaysAnimated;
+
+						Result->Sockets.Add(MutableSocket);
+					}
+				}
+
 				// Physics
 				TObjectPtr<UPhysicsAsset> PhysicsAsset = SkeletalMesh->GetPhysicsAsset();
-				if (!bIgnoreSkeleton && PhysicsAsset && Result->GetSkeleton()/* && Context.Options.bPhysicsAssetMergeEnabled*/)
+				if (!bIgnorePhysics && PhysicsAsset && Result->GetSkeleton())
 				{
 					UE::Mutable::Private::TManagedPtr<UE::Mutable::Private::FPhysicsBody> PhysicsBody = UnrealConversionUtils::CreatePhysicsBodyForMesh(*PhysicsAsset, *Result);
 
@@ -1709,6 +1734,7 @@ namespace UnrealConversionUtils
 				}
 
 				// Texture coords buffer
+				if(!bIgnoreTexCoords)
 				{
 					MUTABLE_CPUPROFILER_SCOPE(MeshConversion_TexCoords)
 
@@ -1743,24 +1769,27 @@ namespace UnrealConversionUtils
 				}
 
 				// Skin buffer
-				if (Section.MaxBoneInfluences > 0)
+				if (!bIgnoreSkinning)
 				{
-					MUTABLE_CPUPROFILER_SCOPE(MeshConversion_Skin)
+					if (Section.MaxBoneInfluences > 0)
+					{
+						MUTABLE_CPUPROFILER_SCOPE(MeshConversion_Skin)
 
-					const FSkinWeightVertexBuffer* SkinBuffer = Context->DataLOD->GetSkinWeightVertexBuffer();
-					const int32 MaxBoneIndexTypeSizeBytes = SkinBuffer->GetBoneIndexByteSize();
-					const int32 MaxBoneWeightTypeSizeBytes = SkinBuffer->GetBoneWeightByteSize();
-					const int32 MaxBonesPerVertex = SkinBuffer->GetMaxBoneInfluences();
-					MutableMeshBufferUtils::SetupSkinBuffer(CurrentVertexBuffer, MaxBoneIndexTypeSizeBytes, MaxBoneWeightTypeSizeBytes, MaxBonesPerVertex, Vertices);
+							const FSkinWeightVertexBuffer* SkinBuffer = Context->DataLOD->GetSkinWeightVertexBuffer();
+						const int32 MaxBoneIndexTypeSizeBytes = SkinBuffer->GetBoneIndexByteSize();
+						const int32 MaxBoneWeightTypeSizeBytes = SkinBuffer->GetBoneWeightByteSize();
+						const int32 MaxBonesPerVertex = SkinBuffer->GetMaxBoneInfluences();
+						MutableMeshBufferUtils::SetupSkinBuffer(CurrentVertexBuffer, MaxBoneIndexTypeSizeBytes, MaxBoneWeightTypeSizeBytes, MaxBonesPerVertex, Vertices);
 
-					int32 ElementSize = Vertices.Buffers[CurrentVertexBuffer].ElementSize;
-					const uint8* SourceVertexData = ((const uint8*)SkinBuffer->GetDataVertexBuffer()->GetWeightData()) + FirstVertexIndex * ElementSize;
+						int32 ElementSize = Vertices.Buffers[CurrentVertexBuffer].ElementSize;
+						const uint8* SourceVertexData = ((const uint8*)SkinBuffer->GetDataVertexBuffer()->GetWeightData()) + FirstVertexIndex * ElementSize;
 
-					check(ElementSize * (FirstVertexIndex + VertexCount) <= (int32)SkinBuffer->GetVertexDataSize());
+						check(ElementSize * (FirstVertexIndex + VertexCount) <= (int32)SkinBuffer->GetVertexDataSize());
 
-					FMemory::Memcpy(Vertices.GetBufferData(CurrentVertexBuffer), SourceVertexData, ElementSize * VertexCount);
+						FMemory::Memcpy(Vertices.GetBufferData(CurrentVertexBuffer), SourceVertexData, ElementSize * VertexCount);
 
-					++CurrentVertexBuffer;
+						++CurrentVertexBuffer;
+					}
 				}
 
 				// Colour buffer
@@ -1840,15 +1869,19 @@ namespace UnrealConversionUtils
 					}
 				}
 
-				SetMorphData(
-						Result.Get(), SkeletalMesh.Get(), 
-						LODIndex, SectionIndex, 
-						(int32)Section.BaseVertexIndex, (int32)Section.NumVertices, (int32)Context->OriginalLOD->GetNumVertices(), 
+				if (!bIgnoreMorphs)
+				{
+					SetMorphData(
+						Result.Get(), SkeletalMesh.Get(),
+						LODIndex, SectionIndex,
+						(int32)Section.BaseVertexIndex, (int32)Section.NumVertices, (int32)Context->OriginalLOD->GetNumVertices(),
 						nullptr);
-					
+				}
+
 				Result->EnsureSurfaceData();
 
 				// Cloth
+				if(!bIgnoreSkinning)
 				{
 					FSkeletalMeshVertexClothBuffer& ClothVertexBuffer = Context->OriginalLOD->ClothVertexBuffer;
 
@@ -1890,11 +1923,6 @@ namespace UnrealConversionUtils
 						FMemory::Memcpy(Result->ClothSections[0].Data.GetData(), &ClothVertexBuffer.MappingData(SectionClothIndexMapping.BaseVertexIndex), NumVertices * sizeof(FMeshToMeshVertData));
 					}	
 				}
-
-				// Physics Asset
-				{
-					// TODO
-				}
 			},
 			ConversionTask);
 
diff --git a/Engine/Plugins/Mutable/Source/CustomizableObject/Private/MuCO/UnrealMutableImageProvider.cpp b/Engine/Plugins/Mutable/Source/CustomizableObject/Private/MuCO/UnrealMutableImageProvider.cpp
--- a/Engine/Plugins/Mutable/Source/CustomizableObject/Private/MuCO/UnrealMutableImageProvider.cpp
+++ b/Engine/Plugins/Mutable/Source/CustomizableObject/Private/MuCO/UnrealMutableImageProvider.cpp
@@ -417,7 +417,7 @@ UE::Mutable::Private::FExtendedImageDesc FUnrealMutableResourceProvider::GetImag
 }
 
 
-TTuple<UE::Tasks::FTask, TFunction<void()>> FUnrealMutableResourceProvider::GetMeshAsync(USkeletalMesh* SkeletalMesh, int32 LODIndex, int32 SectionIndex, TFunction<void(UE::Mutable::Private::TManagedPtr<UE::Mutable::Private::FMesh>)>& ResultCallback)
+TTuple<UE::Tasks::FTask, TFunction<void()>> FUnrealMutableResourceProvider::GetMeshAsync(USkeletalMesh* SkeletalMesh, int32 LODIndex, int32 SectionIndex, uint8 ConversionFlags, TFunction<void(UE::Mutable::Private::TManagedPtr<UE::Mutable::Private::FMesh>)>& ResultCallback)
 {
 	// Thread: worker
 	MUTABLE_CPUPROFILER_SCOPE(FUnrealMutableImageProvider::GetMeshAsync);
@@ -435,7 +435,7 @@ TTuple<UE::Tasks::FTask, TFunction<void()>> FUnrealMutableResourceProvider::GetM
 		return Invoke(TrivialReturn);
 	}
 	
-	UE::Tasks::FTask ConversionTask = UnrealConversionUtils::ConvertSkeletalMeshFromRuntimeData(SkeletalMesh, LODIndex, SectionIndex, Result);
+	UE::Tasks::FTask ConversionTask = UnrealConversionUtils::ConvertSkeletalMeshFromRuntimeData(SkeletalMesh, LODIndex, SectionIndex, ConversionFlags, Result);
 
 	return MakeTuple(
 		// Some post-game conversion stuff can happen here in a worker thread
diff --git a/Engine/Plugins/Mutable/Source/CustomizableObjectEditor/Private/MuCOE/CustomizableObjectCompiler.cpp b/Engine/Plugins/Mutable/Source/CustomizableObjectEditor/Private/MuCOE/CustomizableObjectCompiler.cpp
--- a/Engine/Plugins/Mutable/Source/CustomizableObjectEditor/Private/MuCOE/CustomizableObjectCompiler.cpp
+++ b/Engine/Plugins/Mutable/Source/CustomizableObjectEditor/Private/MuCOE/CustomizableObjectCompiler.cpp
@@ -973,8 +973,6 @@ void FCustomizableObjectCompiler::CompileInternal(bool bAsync)
 
 		ModelResources->AnimBPs = MoveTemp(GenerationContext.AnimBPAssets);
 
-		ModelResources->MaterialSlotNames = MoveTemp(GenerationContext.ReferencedMaterialSlotNames);
-
 		TArray<FGeneratedImageProperties> ImageProperties;
 		GenerationContext.ImageProperties.GenerateValueArray(ImageProperties);
 		
@@ -1674,11 +1672,8 @@ void FCustomizableObjectCompiler::FinishCompilationTask()
 			ModelResources->ExternalOperationVersions.Add(ExternalOperation.Value.GetScriptStruct(), Version);
 		}
 		
-		ModelResources->PhysicsAssets = MoveTemp(CompilationContext->PhysicsAssets);
 		ModelResources->AnimBpOverridePhysiscAssetsInfo = MoveTemp(CompilationContext->AnimBpOverridePhysicsAssetsInfo);
 		
-		ModelResources->Sockets = MoveTemp(CompilationContext->Sockets);
-
 		ModelResources->SkinWeightProfilesInfo = MoveTemp(CompilationContext->SkinWeightProfilesInfo);
 
 		ModelResources->StreamedResourceDataEditor = MoveTemp(CompilationContext->StreamedResourceData);
diff --git a/Engine/Plugins/Mutable/Source/CustomizableObjectEditor/Private/MuCOE/GenerateMutableSource/GenerateMutableSource.h b/Engine/Plugins/Mutable/Source/CustomizableObjectEditor/Private/MuCOE/GenerateMutableSource/GenerateMutableSource.h
--- a/Engine/Plugins/Mutable/Source/CustomizableObjectEditor/Private/MuCOE/GenerateMutableSource/GenerateMutableSource.h
+++ b/Engine/Plugins/Mutable/Source/CustomizableObjectEditor/Private/MuCOE/GenerateMutableSource/GenerateMutableSource.h
@@ -68,7 +68,6 @@ struct FAnimBpOverridePhysicsAssetsInfo;
 struct FMutableGraphGenerationContext;
 struct FMutableParameterData;
 struct FMutableRefSkeletalMeshData;
-struct FMutableRefSocket;
 struct FMutableSkinWeightProfileInfo;
 struct FMorphTargetVertexData;
 struct FMutableSourceSurfaceMetadata;
@@ -524,9 +523,6 @@ struct FMutableCompilationContext
 	TMap<uint32, int32> StreamedResourceIndices;
 	TArray<FCustomizableObjectResourceData> StreamedResourceData;
 
-	/** Stores the sockets provided by the part skeletal meshes, to be merged in the generated meshes */
-	TMap<uint32, FMutableRefSocket> Sockets;
-
 	/** Data used for SkinWeightProfiles reconstruction. */
 	TArray<FMutableSkinWeightProfileInfo> SkinWeightProfilesInfo;
 
@@ -826,7 +822,6 @@ struct FMutableGraphGenerationContext
 	
 	TArray<FMutableRefSkeletalMeshData> ReferenceSkeletalMeshesData;
 	
-	TArray<FName> ReferencedMaterialSlotNames;
 	TMap<FGeneratedImagePropertiesKey, FGeneratedImageProperties> ImageProperties;
 	TArray<const UCustomizableObjectNode*> NoNameNodeObjectArray;
 	TMap<FString, FCustomizableObjectIdPair> GroupNodeMap;
@@ -878,12 +873,12 @@ struct FMutableGraphGenerationContext
 	/** Used to propagate the socket priority defined in group nodes to their child skeletal mesh nodes
 	* It's a stack because group nodes are recursive
 	*/
-	TArray<int32> SocketPriorityStack;
+	TArray<int8> SocketPriorityStack;
 	
 	/** Used to propagate the BonePose priority defined in group nodes to their child skeletal mesh nodes
 	* It's a stack because group nodes are recursive
 	*/
-	TArray<int16> BonePosePriorityStack;
+	TArray<int8> BonePosePriorityStack;
 
 	// Stores what param names use a certain table as a table can be used from multiple table nodes, useful for partial compilations to restrict params
 	TMap<FString, FMutableParamNameSet> TableToParamNames;
diff --git a/Engine/Plugins/Mutable/Source/CustomizableObjectEditor/Private/MuCOE/GenerateMutableSource/GenerateMutableSourceMesh.cpp b/Engine/Plugins/Mutable/Source/CustomizableObjectEditor/Private/MuCOE/GenerateMutableSource/GenerateMutableSourceMesh.cpp
--- a/Engine/Plugins/Mutable/Source/CustomizableObjectEditor/Private/MuCOE/GenerateMutableSource/GenerateMutableSourceMesh.cpp
+++ b/Engine/Plugins/Mutable/Source/CustomizableObjectEditor/Private/MuCOE/GenerateMutableSource/GenerateMutableSourceMesh.cpp
@@ -1275,78 +1275,24 @@ UE::Tasks::FTask ConvertSkeletalMeshToMutable(TSharedRef<FMeshConversionContext>
 		}
 
 		MutableMesh->SetBoneMap(MutableBoneMap);
-
-		MutableMesh->BonePosePriority = Source.BonePosePriority;
 	}
 
 	if (!bIgnoreSkinning)
 	{
-		constexpr uint32 InvalidResourceId = 0;
-
-		TArray<uint8> SocketData;
-
 		for (int32 SocketIndex = 0; SocketIndex < SkeletalMesh->NumSockets(); ++SocketIndex)
 		{
-			USkeletalMeshSocket* Socket = SkeletalMesh->GetSocketByIndex(SocketIndex);
+			const USkeletalMeshSocket* Socket = SkeletalMesh->GetSocketByIndex(SocketIndex);
+			check(Socket);
 
-			FMutableRefSocket MutableSocket;
+			UE::Mutable::Private::FMeshSocket MutableSocket;
 			MutableSocket.SocketName = Socket->SocketName;
 			MutableSocket.BoneName = Socket->BoneName;
 			MutableSocket.RelativeLocation = Socket->RelativeLocation;
 			MutableSocket.RelativeRotation = Socket->RelativeRotation;
 			MutableSocket.RelativeScale = Socket->RelativeScale;
 			MutableSocket.bForceAlwaysAnimated = Socket->bForceAlwaysAnimated;
-			MutableSocket.Priority = Source.SocketPriority;
-	
-
-			SocketData.SetNum(0, EAllowShrinking::No);
-			FMemoryWriter MemoryWriter(SocketData);
-			MemoryWriter << MutableSocket;
-
-			uint32 SocketDataHash = CityHash32(reinterpret_cast<const char*>(SocketData.GetData()), SocketData.Num());
-			const FMutableRefSocket* FoundData = Context.Sockets.Find(SocketDataHash);
-
-			// NOTE: This way of unique hash generation guarantees all valid values can be used but given its 
-			// sequential nature a cascade of changes can occur if new meshes are added. Not many hash collisions 
-			// are expected so it should not be problematic.
-			if (FoundData && *FoundData != MutableSocket)
-			{
-				uint32 NumTries = 0;
-				for (; NumTries < TNumericLimits<uint32>::Max(); ++NumTries)
-				{
-					SocketDataHash++;
-
-					if (SocketDataHash == InvalidResourceId)
-					{
-						continue;
-					}
-
-					FoundData = Context.Sockets.Find(SocketDataHash);
 
-					if (!FoundData || *FoundData == MutableSocket)
-					{
-						break;
-					}
-				}
-
-				if (NumTries == TNumericLimits<uint32>::Max())
-				{
-					UE_LOG(LogMutable, Warning, TEXT("Maximum number of mesh sockets reached, some meshes may not work as expected."));
-					SocketDataHash = InvalidResourceId;
-				}
-				
-			}
-
-			if (SocketDataHash != InvalidResourceId)
-			{
-				Context.Sockets.Add(SocketDataHash, MutableSocket);
-
-				FCustomizableObjectStreameableResourceId StreamedMorphResource;
-				StreamedMorphResource.Id = SocketDataHash;
-				StreamedMorphResource.Type = static_cast<uint8>(FCustomizableObjectStreameableResourceId::EType::Socket);
-
-				MutableMesh->AddStreamedResource(BitCast<uint64>(StreamedMorphResource));
-			}
+			MutableMesh->Sockets.Add(MutableSocket);
 		}
 	}
 
@@ -1673,10 +1619,8 @@ UE::Mutable::Private::TManagedPtr<UE::Mutable::Private::FMesh> ConvertSkeletalMe
 	Source.Flags = GenerationContext.MeshGenerationFlags.Last();
 	Source.MessageContext = CurrentNode;
 	Source.ComponentId = GenerationContext.CurrentSkeletalMeshComponent;
-	Source.SocketPriority = GenerationContext.SocketPriorityStack.Last();
-	Source.BonePosePriority = GenerationContext.BonePosePriorityStack.Last();
 
-	// Prepare source data for realtime morph targets
+	// Prepare source data for real-time morph targets
 	const UCustomizableObjectNodeSkeletalMesh* NodeTypedSkMesh = Cast<UCustomizableObjectNodeSkeletalMesh>(CurrentNode);
 	if (NodeTypedSkMesh)
 	{
@@ -2821,6 +2765,9 @@ UE::Mutable::Private::Ptr<UE::Mutable::Private::NodeMesh> GenerateMutableSourceM
 
 			MeshNode->RealTimeMorphNames = GetUsedRealTimeMorphsNames(Source, GenerationContext.CompilationContext->RealTimeMorphTargetsOverrides);
 
+			MeshNode->BonePosePriority = GenerationContext.BonePosePriorityStack.Last();
+			MeshNode->SocketPriority = GenerationContext.SocketPriorityStack.Last();
+
 			const FString MeshName = TypedNodeSkel->SkeletalMesh.GetLongPackageName().ToLower();
 			MeshNode->SourceDataDescriptor.SourceId = CityHash32(reinterpret_cast<const char*>(*MeshName), MeshName.Len() * sizeof(FString::ElementType));
 			MeshNode->SourceDataDescriptor.OptionalMaxLODSize = 0;
@@ -3507,6 +3454,11 @@ UE::Mutable::Private::Ptr<UE::Mutable::Private::NodeMesh> GenerateMutableSourceM
 			BreakNode->bAutomaticLODs = GenerationContext.CurrentAutoLODStrategy == ECustomizableObjectAutomaticLODStrategy::AutomaticFromMesh;
 		}
 		
+		BreakNode->ConversionFlags = (uint8)GenerationContext.MeshGenerationFlags.Last();
+		BreakNode->BonePosePriority = GenerationContext.BonePosePriorityStack.Last();
+		BreakNode->SocketPriority = GenerationContext.SocketPriorityStack.Last();
+
+
 		TSoftObjectPtr<UStreamableRenderAsset> NodeMesh = TypedNodeParam->GetMesh();
 		FMutableSourceMeshData Source;
 		Source.Mesh = NodeMesh;
diff --git a/Engine/Plugins/Mutable/Source/CustomizableObjectEditor/Private/MuCOE/Nodes/CustomizableObjectNodeObjectGroup.h b/Engine/Plugins/Mutable/Source/CustomizableObjectEditor/Private/MuCOE/Nodes/CustomizableObjectNodeObjectGroup.h
--- a/Engine/Plugins/Mutable/Source/CustomizableObjectEditor/Private/MuCOE/Nodes/CustomizableObjectNodeObjectGroup.h
+++ b/Engine/Plugins/Mutable/Source/CustomizableObjectEditor/Private/MuCOE/Nodes/CustomizableObjectNodeObjectGroup.h
@@ -38,12 +38,12 @@ public:
 	// The pose defined in meshes deriving from this node will inherit this bone pose priority. Decides the bone pose to keep when two bones
 	// that have the same bone usage flags (E.g. EBoneUsageFlags::Skinning). The one with higher priority will be kept.
 	UPROPERTY(EditAnywhere, Category = MeshPose)
-	int16 BonePosePriority = 0;
+	int8 BonePosePriority = 0;
 
 	// The sockets defined in meshes deriving from this node will inherit this socket priority. When in the generated merged mesh there
 	// are clashes with socket names, the one with higher priority will be kept and the other discarded.
 	UPROPERTY(EditAnywhere, Category = MeshSockets)
-	int32 SocketPriority = 0;
+	int8 SocketPriority = 0;
 	
 	UPROPERTY()
 	FEdGraphPinReference NamePin;
diff --git a/Engine/Plugins/Mutable/Source/CustomizableObjectEditor/Private/MuCOE/SMutableCodeViewer.cpp b/Engine/Plugins/Mutable/Source/CustomizableObjectEditor/Private/MuCOE/SMutableCodeViewer.cpp
--- a/Engine/Plugins/Mutable/Source/CustomizableObjectEditor/Private/MuCOE/SMutableCodeViewer.cpp
+++ b/Engine/Plugins/Mutable/Source/CustomizableObjectEditor/Private/MuCOE/SMutableCodeViewer.cpp
@@ -3250,7 +3250,7 @@ void FMutableCodeTreeElement::GenerateLabelText()
 	case UE::Mutable::Private::EOpType::ME_SKELETALMESH_BREAK:
 	{
 		UE::Mutable::Private::OP::MeshSkeletalMeshBreakArgs Args = Program.GetOpArgs<UE::Mutable::Private::OP::MeshSkeletalMeshBreakArgs>(MutableOperation);
-		OpName += FString::Printf( TEXT(" LOD %d Section %d"), Args.LOD, Args.Section );
+		OpName += FString::Printf( TEXT(" LOD %d Section %d Flags %d"), Args.LOD, Args.Section, Args.Flags );
 		break;
 	}
 
diff --git a/Engine/Plugins/Mutable/Source/MutableRuntime/Internal/MuR/Mesh.h b/Engine/Plugins/Mutable/Source/MutableRuntime/Internal/MuR/Mesh.h
--- a/Engine/Plugins/Mutable/Source/MutableRuntime/Internal/MuR/Mesh.h
+++ b/Engine/Plugins/Mutable/Source/MutableRuntime/Internal/MuR/Mesh.h
@@ -7,6 +7,7 @@
 #include "HAL/PlatformMath.h"
 #include "Math/Transform.h"
 #include "Math/VectorRegister.h"
+#include "Math/Rotator.h"
 #include "Misc/AssertionMacros.h"
 #include "Misc/EnumClassFlags.h"
 #include "MuR/Layout.h"
@@ -138,6 +139,24 @@ namespace UE::Mutable::Private
 		void Unserialise(FInputArchive&);
 	};
 
+
+	struct FMeshSocket
+	{
+		FName SocketName;
+		FName BoneName;
+		FVector RelativeLocation = FVector::ZeroVector;
+		FRotator RelativeRotation = FRotator::ZeroRotator;
+		FVector RelativeScale = FVector::ZeroVector;
+
+		bool bForceAlwaysAnimated = false;
+
+		int8 Priority = 0; // Not serialized. AddMetadata may modify it
+
+		bool operator==(const FMeshSocket&) const = default;
+
+		void Serialise(FOutputArchive&) const;
+		void Unserialise(FInputArchive&);
+	};
 	
 	/** Fill an array with the indices of all triangles belonging to the same UV island as InFirstTriangle.
 	*/
@@ -211,7 +230,7 @@ namespace UE::Mutable::Private
 		WithTags = 1 << 6,
 		WithVertexBuffers = 1 << 7,
 		WithIndexBuffers = 1 << 8,
-		// deprecated WithFaceBuffers = 1 << 9,
+		WithSockets = 1 << 9,
 		WithAdditionalBuffers = 1 << 10,
 		WithLayouts = 1 << 11,
 		WithPoses = 1 << 12,
@@ -419,7 +438,7 @@ namespace UE::Mutable::Private
 		UE_API int32 GetBonePoseCount() const;
 
 		//!
-		UE_API void SetBonePose(int32 Index, FBoneIdOrIndex BoneName, FTransform3f Transform, EBoneUsageFlags BoneUsageFlags);
+		UE_API void SetBonePose(int32 Index, FBoneIdOrIndex BoneName, FTransform3f Transform, EBoneUsageFlags BoneUsageFlags, int8 Priority = 0);
 
 		//! @return - Bone identifier of the pose at 'Index'.
 		UE_API FBoneIdOrIndex GetBonePoseId(int32 BoneIndex) const;
@@ -527,7 +546,7 @@ namespace UE::Mutable::Private
 			EBoneUsageFlags BoneUsageFlags = EBoneUsageFlags::None;
 			
 			// Not serialized. Add Metadata may overwrite it. 
-			int16 BonePosePriority = 0;
+			int8 BonePosePriority = 0;
 
 			FTransform3f BoneTransform;
 
@@ -549,12 +568,12 @@ namespace UE::Mutable::Private
 		 */
 		TMemoryTrackedArray<FBonePose> BonePoses;
 
-		/** Priority used to decide which bone pose to keep when merging two meshes with same BoneUsageFlags. Serialized in the Metadata args */
-		int16 BonePosePriority = 0;
-
 		/** Array containing the bonemaps of all surfaces in the mesh. */
 		TArray<FBoneIdOrIndex> BoneMap;
     	
+		/** Array containing the sockets of this mesh. Not serialized. Added with Metadata */
+		TArray<FMeshSocket> Sockets;
+
 		inline void Serialise(FOutputArchive& arch) const;
 
 		inline void Unserialise(FInputArchive& arch);
@@ -609,13 +628,14 @@ namespace UE::Mutable::Private
 				bEqual &= AdditionalBuffers[i] == Other.AdditionalBuffers[i];
 			}
 
-			bEqual &= BonePosePriority == Other.BonePosePriority;
 			bEqual &= BonePoses.Num() == Other.BonePoses.Num();
 			for (int32 i = 0; bEqual && i < BonePoses.Num(); ++i)
 			{
 				bEqual &= BonePoses[i] == Other.BonePoses[i];
 			}
 
+			bEqual &= Sockets == Other.Sockets;
+
 			if (bEqual) bEqual = BoneMap == Other.BoneMap;
 
 			bEqual &= AdditionalPhysicsBodies.Num() == Other.AdditionalPhysicsBodies.Num();
diff --git a/Engine/Plugins/Mutable/Source/MutableRuntime/Internal/MuR/Model.h b/Engine/Plugins/Mutable/Source/MutableRuntime/Internal/MuR/Model.h
--- a/Engine/Plugins/Mutable/Source/MutableRuntime/Internal/MuR/Model.h
+++ b/Engine/Plugins/Mutable/Source/MutableRuntime/Internal/MuR/Model.h
@@ -328,6 +328,10 @@ namespace UE::Mutable::Private
 		/** Constant FNames */
 		TMap<uint32, FName> ConstantNames;
 
+		/** Constant Mesh Sockets */
+		TMap<uint32, FMeshSocket> ConstantSockets;
+
+
     	TSharedPtr<const IExternalOperationProvider> ExternalOperationProvider;
     	
 #if WITH_EDITOR
@@ -364,7 +368,8 @@ namespace UE::Mutable::Private
 		 * returns -1.
 		 */
     	UE_API int32 GetConstantImageRomId(int32 ConstantImageIndex, int32 LODIndex) const;
-    	
+
+#if WITH_EDITOR
         UE_API OP::CONSTANT_EXTENSIONDATA_INDEX AddConstant(const TManagedPtr<const FExtensionData>& Data);
     	
         UE_API OP::ADDRESS AddConstant(const TManagedPtr<const FLayout>& Layout);
@@ -372,6 +377,10 @@ namespace UE::Mutable::Private
         UE_API OP::ADDRESS AddConstant(const TManagedPtr<const FSkeleton>& Skeleton);
     	
         UE_API OP::ADDRESS AddConstant(const TManagedPtr<const FPhysicsBody>& PhysicsBody);
+
+        UE_API uint32 AddConstant(const FMeshSocket& Socket);
+
+        UE_API OP::LIST_UINT32_ADDRESS AddConstant(const TArray<FMeshSocket>& Sockets);
     	
         UE_API OP::CONSTANT_STRING_INDEX AddConstant(const FString& Str);
     	
@@ -394,6 +403,7 @@ namespace UE::Mutable::Private
         UE_API OP::ADDRESS AddConstant(const FRichCurve& Constant);
     	
         UE_API OP::ADDRESS AddConstant(const TManagedPtr<const FMaterial>& Material);
+#endif //WITH_EDITOR
 
 		void GetExtensionDataConstant(int32 ConstantIndex, TManagedPtr<const FExtensionData>& Result) const;
 
diff --git a/Engine/Plugins/Mutable/Source/MutableRuntime/Internal/MuR/Operations.h b/Engine/Plugins/Mutable/Source/MutableRuntime/Internal/MuR/Operations.h
--- a/Engine/Plugins/Mutable/Source/MutableRuntime/Internal/MuR/Operations.h
+++ b/Engine/Plugins/Mutable/Source/MutableRuntime/Internal/MuR/Operations.h
@@ -525,6 +525,7 @@ namespace UE::Mutable::Private
     		PARAMETER SkeletalMeshParameter;
     		uint8 LOD = 0;
 			uint8 Section = 0;
+			uint8 Flags = 0;
 			uint32 MeshID = 0;
 		};
 
@@ -1101,6 +1102,8 @@ namespace UE::Mutable::Private
 				IsAdditionalPhysicsAssetList = 1 << 7,
 				HasRealTimeMorphNames = 1 << 8,
 				IsRealTimeMorphNamesList = 1 << 9,	
+				HasSockets = 1 << 10,	
+				IsSocketsList = 1 << 11,	
             };
 
             union StringAddressOrUInt32ListAddress
@@ -1121,8 +1124,15 @@ namespace UE::Mutable::Private
 				LIST_UINT32_ADDRESS ListAddress;
 			};
 
+			union SocketIdOrUInt32ListAddress
+			{
+				uint32 SocketId;
+				LIST_UINT32_ADDRESS ListAddress;
+			};
+
             EnumFlags Flags; 
-			int16 BonePosePriority;
+			int8 BonePosePriority;
+			int8 SocketPriority;
             ADDRESS Source;
             PASSTHROUGH_ID SkeletonId;
             StringAddressOrUInt32ListAddress Tags;
@@ -1130,6 +1140,7 @@ namespace UE::Mutable::Private
 			PASSTHROUGH_ID PhysicsAssetId; // Main Physics Asset of the mesh
 			PassthroughIDOrUInt32ListAddress AdditionalPhysicsAssetsIds; // Physics Assets used by AnimBPs
             StringAddressOrUInt32ListAddress RealTimeMorphNames;
+			SocketIdOrUInt32ListAddress Sockets;
         };
         static_assert(sizeof(MeshAddMetadataArgs) == sizeof(uint64)*5, "MeshAddMetadataArgs has an unexpected size.");
 
@@ -1200,7 +1211,7 @@ namespace UE::Mutable::Private
 
 		struct LayoutRemoveBlocksArgs
 		{
-			/** Layout to be processedand modified. */
+			/** Layout to be processed and modified. */
 			ADDRESS Source;
 
 			/** Source layout to scan for active blocks. */
diff --git a/Engine/Plugins/Mutable/Source/MutableRuntime/Internal/MuR/Serialisation.h b/Engine/Plugins/Mutable/Source/MutableRuntime/Internal/MuR/Serialisation.h
--- a/Engine/Plugins/Mutable/Source/MutableRuntime/Internal/MuR/Serialisation.h
+++ b/Engine/Plugins/Mutable/Source/MutableRuntime/Internal/MuR/Serialisation.h
@@ -13,6 +13,7 @@
 #include "Math/Quat.h"
 #include "Math/IntVector.h"
 #include "Math/Vector4.h"
+#include "Math/Rotator.h"
 #include "HAL/Platform.h"
 #include "Curves/RichCurve.h"
 #include "Containers/StaticArray.h"
@@ -748,6 +749,24 @@ namespace UE::Mutable::Private
 	}
 
 
+	//---------------------------------------------------------------------------------------------
+	template< typename T >
+	inline void operator<< (FOutputArchive& arch, const UE::Math::TRotator<T>& R)
+	{
+		arch << R.Pitch;
+		arch << R.Yaw;
+		arch << R.Roll;
+	}
+
+	template< typename T >
+	inline void operator>> (FInputArchive& arch, UE::Math::TRotator<T>& R)
+	{
+		arch >> R.Pitch;
+		arch >> R.Yaw;
+		arch >> R.Roll;
+	}
+
+
 	//---------------------------------------------------------------------------------------------
 	template< typename T >
 	inline void operator<< (FOutputArchive& arch, const UE::Math::TTransform<T>& v)
diff --git a/Engine/Plugins/Mutable/Source/MutableRuntime/Internal/MuR/System.h b/Engine/Plugins/Mutable/Source/MutableRuntime/Internal/MuR/System.h
--- a/Engine/Plugins/Mutable/Source/MutableRuntime/Internal/MuR/System.h
+++ b/Engine/Plugins/Mutable/Source/MutableRuntime/Internal/MuR/System.h
@@ -102,7 +102,7 @@ namespace UE::Mutable::Private
 		virtual FExtendedImageDesc GetImageDesc(UTexture* Texture) = 0;
 
 		/** Returns the completion event and a cleanup function that must be called once event is completed. */
-		virtual TTuple<UE::Tasks::FTask, TFunction<void()>> GetMeshAsync(USkeletalMesh* SkeleltalMesh, int32 InLODIndex, int32 InSectionIndex, TFunction<void(UE::Mutable::Private::TManagedPtr<FMesh>)>& ResultCallback) = 0;
+		virtual TTuple<UE::Tasks::FTask, TFunction<void()>> GetMeshAsync(USkeletalMesh* SkeleltalMesh, int32 InLODIndex, int32 InSectionIndex, uint8 InConversionFlags, TFunction<void(UE::Mutable::Private::TManagedPtr<FMesh>)>& ResultCallback) = 0;
     };
 
 	
diff --git a/Engine/Plugins/Mutable/Source/MutableRuntime/Private/MuR/CodeRunner.cpp b/Engine/Plugins/Mutable/Source/MutableRuntime/Private/MuR/CodeRunner.cpp
--- a/Engine/Plugins/Mutable/Source/MutableRuntime/Private/MuR/CodeRunner.cpp
+++ b/Engine/Plugins/Mutable/Source/MutableRuntime/Private/MuR/CodeRunner.cpp
@@ -247,7 +247,7 @@ namespace UE::Mutable::Private
 	}
 
 
-	TTuple<UE::Tasks::FTask, TFunction<void()>> CodeRunner::LoadExternalMeshAsync(FExternalResourceId Id, int32 LODIndex, int32 SectionIndex, TFunction<void(TManagedPtr<FMesh>)>& ResultCallback)
+	TTuple<UE::Tasks::FTask, TFunction<void()>> CodeRunner::LoadExternalMeshAsync(FExternalResourceId Id, int32 LODIndex, int32 SectionIndex, uint8 ConversionFlags, TFunction<void(TManagedPtr<FMesh>)>& ResultCallback)
 	{
 		MUTABLE_CPUPROFILER_SCOPE(LoadExternalImageAsync);
 
@@ -256,7 +256,7 @@ namespace UE::Mutable::Private
 			if (Id.ReferenceResourceId == PASSTHROUGH_ID_INVALID)
 			{
 				// It's a parameter mesh
-				return ExternalResourceProvider->GetMeshAsync(Id.MeshParameter, LODIndex, SectionIndex, ResultCallback);
+				return ExternalResourceProvider->GetMeshAsync(Id.MeshParameter, LODIndex, SectionIndex, ConversionFlags, ResultCallback);
 			}
 			else
 			{
@@ -2894,16 +2894,42 @@ namespace UE::Mutable::Private
 					}
 					
 					
-					if (Result->BonePosePriority != OpArgs.BonePosePriority)
+					if (OpArgs.BonePosePriority != 0)
 					{
-						Result->BonePosePriority = OpArgs.BonePosePriority;
-
 						for (FMesh::FBonePose& BonePose : Result->BonePoses)
 						{
 							BonePose.BonePosePriority = OpArgs.BonePosePriority;
 						}
 					}
 
+					if (EnumHasAnyFlags(OpArgs.Flags, OpEnumFlags::IsSocketsList) &&
+						Program.ConstantUInt32Lists.IsValidIndex(OpArgs.Sockets.ListAddress))
+					{
+						const TArray<uint32>& ConstantSocketsList = Program.ConstantUInt32Lists[OpArgs.Sockets.ListAddress];
+						Result->Sockets.Reserve(ConstantSocketsList.Num());
+
+						for (uint32 SocketId : ConstantSocketsList)
+						{
+							Result->Sockets.Emplace(Program.ConstantSockets.FindChecked(SocketId));
+						}
+					}
+					else if(EnumHasAnyFlags(OpArgs.Flags, OpEnumFlags::HasSockets))
+					{
+						Result->Sockets.Emplace(Program.ConstantSockets.FindChecked(OpArgs.Sockets.SocketId));
+					}
+					else
+					{	
+						check(!EnumHasAnyFlags(OpArgs.Flags, OpEnumFlags::IsSocketsList));
+					}
+
+					if (OpArgs.SocketPriority != 0)
+					{
+						for (FMeshSocket& Socket : Result->Sockets)
+						{
+							Socket.Priority = OpArgs.SocketPriority;
+						}
+					}
+
 					if (OpArgs.PhysicsAssetId != PASSTHROUGH_ID_INVALID)
 					{
 						Result->PhysicsAssets.Emplace(FindPassthrough<UPhysicsAsset>(OpArgs.PhysicsAssetId));
diff --git a/Engine/Plugins/Mutable/Source/MutableRuntime/Private/MuR/CodeRunner.h b/Engine/Plugins/Mutable/Source/MutableRuntime/Private/MuR/CodeRunner.h
--- a/Engine/Plugins/Mutable/Source/MutableRuntime/Private/MuR/CodeRunner.h
+++ b/Engine/Plugins/Mutable/Source/MutableRuntime/Private/MuR/CodeRunner.h
@@ -156,7 +156,7 @@ namespace  UE::Mutable::Private
  	    UE::Mutable::Private::FExtendedImageDesc GetExternalImageDesc(UTexture* Id);
 
 		/** Load an external mesh asynchronously, returns an event to wait for complition and a cleanup function that must be called once the event has completed. */
-		TTuple<UE::Tasks::FTask, TFunction<void()>> LoadExternalMeshAsync(FExternalResourceId Id, int32 LODIndex, int32 SectionIndex, TFunction<void(TManagedPtr<FMesh>)>& ResultCallback);
+		TTuple<UE::Tasks::FTask, TFunction<void()>> LoadExternalMeshAsync(FExternalResourceId Id, int32 LODIndex, int32 SectionIndex, uint8 ConversionFlags, TFunction<void(TManagedPtr<FMesh>)>& ResultCallback);
 
 		/** Settings that may affect the execution of some operations, like image conversion quality. */
 		FSettings Settings;
diff --git a/Engine/Plugins/Mutable/Source/MutableRuntime/Private/MuR/CodeRunnerTasked.cpp b/Engine/Plugins/Mutable/Source/MutableRuntime/Private/MuR/CodeRunnerTasked.cpp
--- a/Engine/Plugins/Mutable/Source/MutableRuntime/Private/MuR/CodeRunnerTasked.cpp
+++ b/Engine/Plugins/Mutable/Source/MutableRuntime/Private/MuR/CodeRunnerTasked.cpp
@@ -2227,7 +2227,7 @@ namespace UE::Mutable::Private
 	{
 	public:
 
-		FMeshExternalLoadTask(const FScheduledOp&, CodeRunner::FExternalResourceId, int32 InLODIndex, int32 InSectionIndex, uint32 MeshID);
+		FMeshExternalLoadTask(const FScheduledOp&, CodeRunner::FExternalResourceId, int32 InLODIndex, int32 InSectionIndex, uint8 ConversionFlags, uint32 MeshID);
 
 		// FIssuedTask interface
 		virtual bool Prepare(CodeRunner*, bool& bOutFailed) override;
@@ -2238,6 +2238,7 @@ namespace UE::Mutable::Private
 		CodeRunner::FExternalResourceId Id;
 		int32 LODIndex = 0;
 		int32 SectionIndex = 0;
+		uint8 ConversionFlags = 0;
 		uint32 MeshID = 0;
 
 		TManagedPtr<FMesh> Result;
@@ -2246,11 +2247,12 @@ namespace UE::Mutable::Private
 	};
 
 
-	FMeshExternalLoadTask::FMeshExternalLoadTask(const FScheduledOp& InOp, CodeRunner::FExternalResourceId InId, int32 InLODIndex, int32 InSectionIndex, uint32 InMeshID)
+	FMeshExternalLoadTask::FMeshExternalLoadTask(const FScheduledOp& InOp, CodeRunner::FExternalResourceId InId, int32 InLODIndex, int32 InSectionIndex, uint8 InConversionFlags, uint32 InMeshID)
 		: FIssuedTask(InOp)
 		, Id(InId)
 		, LODIndex(InLODIndex)
 		, SectionIndex(InSectionIndex)
+		, ConversionFlags(InConversionFlags)
 		, MeshID(InMeshID)
 	{
 	}
@@ -2273,7 +2275,7 @@ namespace UE::Mutable::Private
 				Result = InResult;
 			};
 
-		Tie(Event, ExternalCleanUpFunc) = Runner->LoadExternalMeshAsync(Id, LODIndex, SectionIndex, ResultCallback);
+		Tie(Event, ExternalCleanUpFunc) = Runner->LoadExternalMeshAsync(Id, LODIndex, SectionIndex, ConversionFlags, ResultCallback);
 
 		// return false indicating there is no work to do so Event is not overriden by a DoWork task.
 		return false;
@@ -2451,7 +2453,7 @@ namespace UE::Mutable::Private
 
 				CodeRunner::FExternalResourceId FullId;
 				FullId.MeshParameter = SkeletalMesh->PassthroughObject.Get();
-				Issued = MakeShared<FMeshExternalLoadTask>(Item, FullId, LODIndex, SectionIndex, Args.MeshID);
+				Issued = MakeShared<FMeshExternalLoadTask>(Item, FullId, LODIndex, SectionIndex, Args.Flags, Args.MeshID);
 			}
 
 			break;
diff --git a/Engine/Plugins/Mutable/Source/MutableRuntime/Private/MuR/Mesh.cpp b/Engine/Plugins/Mutable/Source/MutableRuntime/Private/MuR/Mesh.cpp
--- a/Engine/Plugins/Mutable/Source/MutableRuntime/Private/MuR/Mesh.cpp
+++ b/Engine/Plugins/Mutable/Source/MutableRuntime/Private/MuR/Mesh.cpp
@@ -139,6 +139,12 @@ void FMesh::CopyFrom(const FMesh& From, EMeshCopyFlags InFlags)
 		IndexBuffers = From.IndexBuffers;
 	}
 
+	// Copy additional buffers
+	if (EnumHasAnyFlags(InFlags, EMeshCopyFlags::WithSockets))
+	{
+		Sockets = From.Sockets;
+	}
+
 	// Copy additional buffers
 	if (EnumHasAnyFlags(InFlags, EMeshCopyFlags::WithAdditionalBuffers))
 	{
@@ -150,16 +156,11 @@ void FMesh::CopyFrom(const FMesh& From, EMeshCopyFlags InFlags)
 	{
 		Layouts = From.Layouts;
 	}
-    // The skeleton is not copied because it is not owned by this mesh and it is always assumed
-    // to be shared.
-
-	// physics body doen't need to be deep copied either as they are also assumed to be shared.
 	
 	// Copy bone poses
 	if (EnumHasAnyFlags(InFlags, EMeshCopyFlags::WithPoses))
 	{
 		BonePoses = From.BonePoses;
-		BonePosePriority = From.BonePosePriority;
 	}
 
 	// Copy BoneMap
@@ -626,12 +627,12 @@ int32 UE::Mutable::Private::FMesh::GetBonePoseCount() const
 }
 
 
-void UE::Mutable::Private::FMesh::SetBonePose(int32 Index, FBoneIdOrIndex BoneId, FTransform3f Transform, EBoneUsageFlags BoneUsageFlags)
+void UE::Mutable::Private::FMesh::SetBonePose(int32 Index, FBoneIdOrIndex BoneId, FTransform3f Transform, EBoneUsageFlags BoneUsageFlags, int8 Priority)
 {
 	check(BonePoses.IsValidIndex(Index));
 	if (BonePoses.IsValidIndex(Index))
 	{
-		BonePoses[Index] = FBonePose{ BoneId, BoneUsageFlags, 0, Transform };
+		BonePoses[Index] = FBonePose{ BoneId, BoneUsageFlags, Priority, Transform };
 	}
 }
 
@@ -969,6 +970,28 @@ void FMeshMorph::Unserialise(FInputArchive& Arch)
 }
 
 
+void FMeshSocket::Serialise(FOutputArchive& Arch) const
+{
+	Arch << SocketName;
+	Arch << BoneName;
+	Arch << RelativeLocation;
+	Arch << RelativeRotation;
+	Arch << RelativeScale;
+	Arch << bForceAlwaysAnimated;
+}
+
+
+void FMeshSocket::Unserialise(FInputArchive& Arch)
+{
+	Arch >> SocketName;
+	Arch >> BoneName;
+	Arch >> RelativeLocation;
+	Arch >> RelativeRotation;
+	Arch >> RelativeScale;
+	Arch >> bForceAlwaysAnimated;
+}
+
+
 void FMesh::FBonePose::Serialise(FOutputArchive& Arch) const
 {
 	Arch << BoneId;
diff --git a/Engine/Plugins/Mutable/Source/MutableRuntime/Private/MuR/Model.cpp b/Engine/Plugins/Mutable/Source/MutableRuntime/Private/MuR/Model.cpp
--- a/Engine/Plugins/Mutable/Source/MutableRuntime/Private/MuR/Model.cpp
+++ b/Engine/Plugins/Mutable/Source/MutableRuntime/Private/MuR/Model.cpp
@@ -10,13 +10,15 @@
 #include "Misc/AssertionMacros.h"
 #include "MuR/Image.h"
 #include "MuR/Mesh.h"
-#include "MuR/Model.h"
 #include "MuR/Operations.h"
 #include "MuR/Parameters.h"
 #include "MuR/Serialisation.h"
-#include "MuR/Serialisation.h"
 #include "MuR/Types.h"
-#include "Templates/Tuple.h"
+
+#if WITH_EDITOR
+#include "Hash/xxhash.h"
+#include "Serialization/MemoryWriter.h"
+#endif
 
 class UTexture;
 class USkeletalMesh;
@@ -132,6 +134,7 @@ namespace UE::Mutable::Private
 		Arch << RelevantParameterList;
 		Arch << ConstantMaterials;
 		Arch << ConstantNames;
+		Arch << ConstantSockets;
 	}
 
 
@@ -164,6 +167,7 @@ namespace UE::Mutable::Private
 		Arch >> RelevantParameterList;
 		Arch >> ConstantMaterials;
 		Arch >> ConstantNames;
+		Arch >> ConstantSockets;
 
 		PostLoad();
 	}
@@ -221,7 +225,7 @@ namespace UE::Mutable::Private
 #endif
     }
 	
-
+#if WITH_EDITOR
     OP::CONSTANT_EXTENSIONDATA_INDEX FProgram::AddConstant(const TManagedPtr<const FExtensionData>& Data)
     {
 	    // Ensure unique
@@ -307,6 +311,60 @@ namespace UE::Mutable::Private
     }
 
 
+	uint32 AddSocketToProgram(TMap<uint32, FMeshSocket>& ConstantSockets, FMeshSocket Socket, TArray<uint8>& AuxData)
+	{
+		AuxData.SetNum(0, EAllowShrinking::No);
+		FMemoryWriter Writer(AuxData);
+
+		FString SocketName = Socket.SocketName.ToString().ToLower();
+		Writer << SocketName;
+
+		FString BoneName = Socket.BoneName.ToString().ToLower();
+		Writer << BoneName;
+
+		Writer << Socket.RelativeLocation;
+		Writer << Socket.RelativeRotation;
+		Writer << Socket.RelativeScale;
+		Writer << Socket.bForceAlwaysAnimated;
+
+		uint32 SocketDataHash = static_cast<uint32>(FXxHash64::HashBuffer(AuxData.GetData(), AuxData.Num()).Hash);
+		
+		for (uint32 NumTries = 0; NumTries < TNumericLimits<uint32>::Max(); ++NumTries)
+		{
+			FMeshSocket& FoundData = ConstantSockets.FindOrAdd(SocketDataHash, Socket);
+			if (FoundData == Socket)
+			{
+				break;
+			}
+
+			SocketDataHash++;
+		}
+
+		return SocketDataHash;
+	}
+
+
+	uint32 FProgram::AddConstant(const FMeshSocket& Socket)
+	{
+		TArray<uint8> SocketData;
+		return AddSocketToProgram(ConstantSockets, Socket, SocketData);
+	}
+
+
+	OP::LIST_UINT32_ADDRESS FProgram::AddConstant(const TArray<FMeshSocket>& Sockets)
+	{
+		TArray<uint32> SocketIds;
+		SocketIds.Reset(Sockets.Num());
+
+		TArray<uint8> SocketData;
+		for (const FMeshSocket& Socket : Sockets)
+		{
+			SocketIds.Add(AddSocketToProgram(ConstantSockets, Socket, SocketData));
+		}
+
+		return AddConstant(SocketIds);
+	}
+
     OP::CONSTANT_STRING_INDEX FProgram::AddConstant(const FString& Str)
     {     
 	    return (OP::CONSTANT_STRING_INDEX)ConstantStrings.AddUnique(Str);
@@ -446,7 +504,7 @@ namespace UE::Mutable::Private
 	    ConstantMaterials.Add(Material);
 	    return Index;
     }
-	
+#endif // WITH_EDITOR
 
     void FProgram::GetExtensionDataConstant(int32 ConstantIndex, TManagedPtr<const FExtensionData>& Result) const
     {
diff --git a/Engine/Plugins/Mutable/Source/MutableRuntime/Private/MuR/OpMeshApplyPose.h b/Engine/Plugins/Mutable/Source/MutableRuntime/Private/MuR/OpMeshApplyPose.h
--- a/Engine/Plugins/Mutable/Source/MutableRuntime/Private/MuR/OpMeshApplyPose.h
+++ b/Engine/Plugins/Mutable/Source/MutableRuntime/Private/MuR/OpMeshApplyPose.h
@@ -19,28 +19,23 @@ namespace UE::Mutable::Private
 		// to multiply PoseXform * ModelRefXform. 
 		// (Note, transform multiplication application is from right to left for this comment, but TTransform::operator* is reversed).
 		const int32 ResultNumBones = InOutResult->GetBonePoseCount();
-		for (int32 PoseIndex = 0; PoseIndex < ResultNumBones; ++PoseIndex)
+		for (FMesh::FBonePose& BonePose : InOutResult->BonePoses)
 		{
-			const FBoneIdOrIndex BoneId = InOutResult->BonePoses[PoseIndex].BoneId;
-			const int32 FoundPoseIndex = BoneToPoseIndex[BoneId.Index];
+			const int32 FoundPoseIndex = BoneToPoseIndex[BonePose.BoneId.Index];
 			if (FoundPoseIndex == INDEX_NONE)
 			{
 				continue;
 			}
 
-			FTransform3f ModelRefTransform;
-			InOutResult->GetBonePoseTransform(PoseIndex, ModelRefTransform);
-
 			FTransform3f PoseTransform;
 			PoseMesh->GetBonePoseTransform(FoundPoseIndex, PoseTransform); 
 
+			BonePose.BoneTransform = BonePose.BoneTransform * PoseTransform;
+
 			// Adding the Reshaped flag will prioritize this bone over other bones without the flag
 			// when merging poses with the same bone.
 			// TODO: Add another flag to indicate this case or generalize the Reshaped flag. 
-			EBoneUsageFlags UsageFlags = InOutResult->GetBoneUsageFlags(PoseIndex);
-			EnumAddFlags(UsageFlags, EBoneUsageFlags::Reshaped);
-			
-			InOutResult->SetBonePose(PoseIndex, BoneId, ModelRefTransform * PoseTransform, UsageFlags);
+			EnumAddFlags(BonePose.BoneUsageFlags, EBoneUsageFlags::Reshaped);
 		}
 	}
 
diff --git a/Engine/Plugins/Mutable/Source/MutableRuntime/Private/MuR/OpMeshMerge.cpp b/Engine/Plugins/Mutable/Source/MutableRuntime/Private/MuR/OpMeshMerge.cpp
--- a/Engine/Plugins/Mutable/Source/MutableRuntime/Private/MuR/OpMeshMerge.cpp
+++ b/Engine/Plugins/Mutable/Source/MutableRuntime/Private/MuR/OpMeshMerge.cpp
@@ -1134,6 +1134,29 @@ namespace OpMeshMergeClothInternal
 		}
 
 
+		// Sockets
+		//---------------------------
+		{
+			// Copy sockets from the first mesh
+			Result->Sockets = pFirst->Sockets;
+
+			for (const FMeshSocket& Socket : pSecond->Sockets)
+			{
+				FMeshSocket* Found = Result->Sockets.FindByPredicate(
+					[SocketName = Socket.SocketName](const FMeshSocket& OtherSocket) { return OtherSocket.SocketName == SocketName; });
+
+				if (!Found)
+				{
+					Result->Sockets.Add(Socket);
+				}
+				else if (Found->Priority < Socket.Priority)
+				{
+					*Found = Socket;
+				}
+			}
+		}
+
+
 		// PhysicsBodies
 		//---------------------------
 		{
diff --git a/Engine/Plugins/Mutable/Source/MutableTools/Internal/MuT/NodeMeshConstant.h b/Engine/Plugins/Mutable/Source/MutableTools/Internal/MuT/NodeMeshConstant.h
--- a/Engine/Plugins/Mutable/Source/MutableTools/Internal/MuT/NodeMeshConstant.h
+++ b/Engine/Plugins/Mutable/Source/MutableTools/Internal/MuT/NodeMeshConstant.h
@@ -29,6 +29,9 @@ namespace UE::Mutable::Private
 		
 		TArray<FName> RealTimeMorphNames;
 
+		int8 BonePosePriority = 0;
+		int8 SocketPriority = 0;
+
 	public:
 
 		// Node Interface
diff --git a/Engine/Plugins/Mutable/Source/MutableTools/Internal/MuT/NodeMeshSkeletalMeshBreak.h b/Engine/Plugins/Mutable/Source/MutableTools/Internal/MuT/NodeMeshSkeletalMeshBreak.h
--- a/Engine/Plugins/Mutable/Source/MutableTools/Internal/MuT/NodeMeshSkeletalMeshBreak.h
+++ b/Engine/Plugins/Mutable/Source/MutableTools/Internal/MuT/NodeMeshSkeletalMeshBreak.h
@@ -25,6 +25,11 @@ namespace UE::Mutable::Private
 
 		uint8 SectionIndex = 0;
 
+		uint8 ConversionFlags = 0;
+
+		int8 BonePosePriority = 0;
+		int8 SocketPriority = 0;
+
 		bool bAutomaticLODs = false;
 
 		Ptr<NodeSkeletalMeshParameter> SkeletalMeshParameter;
diff --git a/Engine/Plugins/Mutable/Source/MutableTools/Private/MuT/ASTOpConstantResource.cpp b/Engine/Plugins/Mutable/Source/MutableTools/Private/MuT/ASTOpConstantResource.cpp
--- a/Engine/Plugins/Mutable/Source/MutableTools/Private/MuT/ASTOpConstantResource.cpp
+++ b/Engine/Plugins/Mutable/Source/MutableTools/Private/MuT/ASTOpConstantResource.cpp
@@ -229,9 +229,6 @@ namespace UE::Mutable::Private
 			{
 				const EMeshCopyFlags MetadataDataCopyFlags = 
 						EMeshCopyFlags::WithSurfaces          |
-						EMeshCopyFlags::WithTags              |
-						EMeshCopyFlags::WithSkeletonIDs       |
-						EMeshCopyFlags::WithStreamedResources |
 						EMeshCopyFlags::WithMorphMetaData     ;
 
 				TManagedPtr<FMesh> MeshMetadataData = MeshData->Clone(MetadataDataCopyFlags);
diff --git a/Engine/Plugins/Mutable/Source/MutableTools/Private/MuT/ASTOpMeshAddMetadata.cpp b/Engine/Plugins/Mutable/Source/MutableTools/Private/MuT/ASTOpMeshAddMetadata.cpp
--- a/Engine/Plugins/Mutable/Source/MutableTools/Private/MuT/ASTOpMeshAddMetadata.cpp
+++ b/Engine/Plugins/Mutable/Source/MutableTools/Private/MuT/ASTOpMeshAddMetadata.cpp
@@ -36,6 +36,8 @@ namespace UE::Mutable::Private
 				ResourceIds == Other->ResourceIds &&
 				SkeletonId == Other->SkeletonId &&
 				BonePosePriority == Other->BonePosePriority &&
+				SocketPriority == Other->SocketPriority &&
+				Sockets == Other->Sockets && 
 				PhysicsAssetId == Other->PhysicsAssetId &&
 				AdditionalPhysicsAssetIds == Other->AdditionalPhysicsAssetIds &&
 				RealTimeMorphNames == Other->RealTimeMorphNames;
@@ -53,6 +55,8 @@ namespace UE::Mutable::Private
 		NewOp->Tags = Tags;
 		NewOp->ResourceIds = ResourceIds;
 		NewOp->BonePosePriority = BonePosePriority;
+		NewOp->SocketPriority = SocketPriority;
+		NewOp->Sockets = Sockets;
 		NewOp->PhysicsAssetId = PhysicsAssetId;
 		NewOp->AdditionalPhysicsAssetIds = AdditionalPhysicsAssetIds;
 		NewOp->RealTimeMorphNames = RealTimeMorphNames;
@@ -139,7 +143,23 @@ namespace UE::Mutable::Private
 
 			Args.SkeletonId = SkeletonId;
 			Args.BonePosePriority = BonePosePriority;
-			
+			Args.SocketPriority = SocketPriority;
+
+			if (Sockets.Num())
+			{
+				EnumAddFlags(Args.Flags, OpEnumFlags::HasSockets);
+
+				if (Sockets.Num() > 1)
+				{
+					EnumAddFlags(Args.Flags, OpEnumFlags::IsSocketsList);
+					Args.Sockets.ListAddress = Program.AddConstant(Sockets);
+				}
+				else
+				{
+					Args.Sockets.SocketId = Program.AddConstant(Sockets[0]);
+				}
+			}
+
 			Args.PhysicsAssetId = PhysicsAssetId;
 
 			if (AdditionalPhysicsAssetIds.Num())
diff --git a/Engine/Plugins/Mutable/Source/MutableTools/Private/MuT/ASTOpMeshAddMetadata.h b/Engine/Plugins/Mutable/Source/MutableTools/Private/MuT/ASTOpMeshAddMetadata.h
--- a/Engine/Plugins/Mutable/Source/MutableTools/Private/MuT/ASTOpMeshAddMetadata.h
+++ b/Engine/Plugins/Mutable/Source/MutableTools/Private/MuT/ASTOpMeshAddMetadata.h
@@ -33,7 +33,13 @@ namespace UE::Mutable::Private
 		TArray<FString> RealTimeMorphNames;
 
 		/** Bone Pose priority when merging two meshes. */
-		int16 BonePosePriority = 0;
+		int8 BonePosePriority = 0;
+
+		/** Socket priority when merging two meshes. */
+		int8 SocketPriority = 0;
+
+		/** Array of sockets. */
+		TArray<FMeshSocket> Sockets;
 
 		/** PhysicsAsset to add */
 		PASSTHROUGH_ID PhysicsAssetId = PASSTHROUGH_ID_INVALID;
diff --git a/Engine/Plugins/Mutable/Source/MutableTools/Private/MuT/ASTOpMeshSkeletalMeshBreak.cpp b/Engine/Plugins/Mutable/Source/MutableTools/Private/MuT/ASTOpMeshSkeletalMeshBreak.cpp
--- a/Engine/Plugins/Mutable/Source/MutableTools/Private/MuT/ASTOpMeshSkeletalMeshBreak.cpp
+++ b/Engine/Plugins/Mutable/Source/MutableTools/Private/MuT/ASTOpMeshSkeletalMeshBreak.cpp
@@ -34,6 +34,7 @@ namespace UE::Mutable::Private
 			return Type == Other->Type &&
 				LODIndex == Other->LODIndex &&
 				SectionIndex == Other->SectionIndex &&
+				ConversionFlags == Other->ConversionFlags &&
 				MeshID == Other->MeshID &&
 				SkeletalMeshParameter == Other->SkeletalMeshParameter;
 		}
@@ -48,6 +49,7 @@ namespace UE::Mutable::Private
 		New->Type = Type;
 		New->LODIndex = LODIndex;
 		New->SectionIndex = SectionIndex;
+		New->ConversionFlags = ConversionFlags;
 		New->MeshID = MeshID;
 		New->SkeletalMeshParameter = ASTChild(New, MapChild(SkeletalMeshParameter.child()));
 
@@ -68,6 +70,7 @@ namespace UE::Mutable::Private
 		Result = HashCombineFast(Result, GetTypeHash(SkeletalMeshParameter));
 		Result = HashCombineFast(Result, GetTypeHash(LODIndex));
 		Result = HashCombineFast(Result, GetTypeHash(SectionIndex));
+		Result = HashCombineFast(Result, GetTypeHash(ConversionFlags));
 		Result = HashCombineFast(Result, GetTypeHash(MeshID));
 
 		return Result;
@@ -96,6 +99,7 @@ namespace UE::Mutable::Private
 
 			Args.LOD = LODIndex;
 			Args.Section = SectionIndex;
+			Args.Flags = ConversionFlags;
 			Args.MeshID = MeshID;
 
 			AppendCode(Program.ByteCode, Args);
diff --git a/Engine/Plugins/Mutable/Source/MutableTools/Private/MuT/ASTOpMeshSkeletalMeshBreak.h b/Engine/Plugins/Mutable/Source/MutableTools/Private/MuT/ASTOpMeshSkeletalMeshBreak.h
--- a/Engine/Plugins/Mutable/Source/MutableTools/Private/MuT/ASTOpMeshSkeletalMeshBreak.h
+++ b/Engine/Plugins/Mutable/Source/MutableTools/Private/MuT/ASTOpMeshSkeletalMeshBreak.h
@@ -20,8 +20,9 @@ namespace UE::Mutable::Private
 		EOpType Type = EOpType::ME_SKELETALMESH_BREAK;
 
 		/** Used by some parameter types (Mesh) to specify which subdata from the actual parameter value should be used in the operation. */
-		int32 LODIndex = 0;
-		int32 SectionIndex = 0;
+		uint8 LODIndex = 0;
+		uint8 SectionIndex = 0;
+		uint8 ConversionFlags = 0;
 		uint32 MeshID = 0;
 
 		ASTChild SkeletalMeshParameter;
diff --git a/Engine/Plugins/Mutable/Source/MutableTools/Private/MuT/CodeGenerator_Mesh.cpp b/Engine/Plugins/Mutable/Source/MutableTools/Private/MuT/CodeGenerator_Mesh.cpp
--- a/Engine/Plugins/Mutable/Source/MutableTools/Private/MuT/CodeGenerator_Mesh.cpp
+++ b/Engine/Plugins/Mutable/Source/MutableTools/Private/MuT/CodeGenerator_Mesh.cpp
@@ -816,7 +816,9 @@ namespace UE::Mutable::Private
 
 				TArray<uint64> ResourceIds;
 				PASSTHROUGH_ID SkeletonID = PASSTHROUGH_ID_INVALID;
-				int16 BonePosePriority = 0;
+				const int8 BonePosePriority = NodeCopy->BonePosePriority;
+				TArray<FMeshSocket> Sockets;
+				const int8 SocketPriority = NodeCopy->SocketPriority;
 				PASSTHROUGH_ID PhysicsAssetId = PASSTHROUGH_ID_INVALID;
 				TArray<PASSTHROUGH_ID> AdditionalPhysicsAssetId;
 
@@ -824,7 +826,9 @@ namespace UE::Mutable::Private
 					Mesh->SkeletonObjects.Num() > 0 ||
 					Mesh->Tags.Num() ||
 					Mesh->StreamedResources.Num() ||
-					Mesh->BonePosePriority != 0 ||
+					BonePosePriority != 0 ||
+					Mesh->Sockets.Num() != 0 ||
+					(Mesh->Sockets.Num() && SocketPriority != 0) ||
 					Mesh->PhysicsAssets.Num() ||
 					Mesh->AdditionalPhysicsAssets.Num();
 
@@ -842,8 +846,6 @@ namespace UE::Mutable::Private
 
 					SkeletonID = Mesh->SkeletonObjects.IsEmpty() ? PASSTHROUGH_ID_INVALID : Mesh->SkeletonObjects[0].GetId();
 
-					BonePosePriority = Mesh->BonePosePriority;
-
 					if (Mesh->PhysicsAssets.Num())
 					{
 						check(Mesh->PhysicsAssets.Num() == 1);
@@ -860,10 +862,11 @@ namespace UE::Mutable::Private
 					MetadatalessMesh->Tags.SetNum(0, EAllowShrinking::No);
 					MetadatalessMesh->StreamedResources.SetNum(0, EAllowShrinking::No);
 					MetadatalessMesh->SkeletonObjects.SetNum(0, EAllowShrinking::No);
-					MetadatalessMesh->BonePosePriority = 0;
 					MetadatalessMesh->PhysicsAssets.SetNum(0, EAllowShrinking::No);
 					MetadatalessMesh->AdditionalPhysicsAssets.SetNum(0, EAllowShrinking::No);
 
+					Exchange(Sockets, MetadatalessMesh->Sockets);
+
 					Mesh = MetadatalessMesh;
 				}
 
@@ -953,6 +956,9 @@ namespace UE::Mutable::Private
 
 					AddMetadataOp->BonePosePriority = BonePosePriority;
 
+					AddMetadataOp->Sockets = MoveTemp(Sockets);
+					AddMetadataOp->SocketPriority = SocketPriority;
+
 					AddMetadataOp->PhysicsAssetId = PhysicsAssetId;
 					AddMetadataOp->AdditionalPhysicsAssetIds = MoveTemp(AdditionalPhysicsAssetId);
 
@@ -1449,6 +1455,7 @@ namespace UE::Mutable::Private
 				Op->SkeletalMeshParameter = SkeletalMeshParameter;
 				Op->LODIndex = Node->LODIndex + (Node->bAutomaticLODs ? StaticOptions.LODIndex : 0);
 				Op->SectionIndex = Node->SectionIndex;
+				Op->ConversionFlags = Node->ConversionFlags;
 
 				// Assign an ID prefix to the mesh. 
 				uint32 MeshID = Op->Hash();
@@ -1478,11 +1485,14 @@ namespace UE::Mutable::Private
 				RealTimeMorphNames.Reserve(Node->RealTimeMorphNames.Num());
 				Algo::Transform(Node->RealTimeMorphNames, RealTimeMorphNames, [](const FName& Name) { return Name.ToString(); });
 
-				if (RealTimeMorphNames.Num() > 0)
+				const bool bAddMetadata = RealTimeMorphNames.Num() || Node->BonePosePriority != 0 || Node->SocketPriority != 0;
+				if (bAddMetadata)
 				{
 					Ptr<ASTOpMeshAddMetadata> AddMetadataOp = new ASTOpMeshAddMetadata;
 
 					AddMetadataOp->RealTimeMorphNames = MoveTemp(RealTimeMorphNames);
+					AddMetadataOp->BonePosePriority = Node->BonePosePriority;
+					AddMetadataOp->SocketPriority = Node->SocketPriority;
 
 					AddMetadataOp->Source = Result.MeshOp;
 					Result.MeshOp = AddMetadataOp;
