[Mutable] Removed warning.

#rnx

#Synced-CL 50588143

diff --git a/Engine/Plugins/Mutable/Source/MutableTools/Private/MuT/CodeGenerator.cpp b/Engine/Plugins/Mutable/Source/MutableTools/Private/MuT/CodeGenerator.cpp
--- a/Engine/Plugins/Mutable/Source/MutableTools/Private/MuT/CodeGenerator.cpp
+++ b/Engine/Plugins/Mutable/Source/MutableTools/Private/MuT/CodeGenerator.cpp
@@ -890,7 +890,7 @@ namespace UE::Mutable::Private
 		GetModifiersFor(ComponentId, SurfaceNode->Tags, bModifiersForBeforeOperations, Modifiers);
 
 		// This pass on the modifiers is only to detect errors that cannot be detected at the point they are applied.
-		CheckModifiersForSurface(*SurfaceNode, Modifiers, Options.LODIndex);
+		//CheckModifiersForSurface(*SurfaceNode, Modifiers, Options.LODIndex);
 		
         // Generate the mesh
 		FMeshGenerationStaticOptions MeshStaticOptions(ComponentId, Options.LODIndex);
