[Mutable] Added a new warning indicating a reference texture couldn't be found when trying to break a material.

#rnx
#Synced-CL 50588143

diff --git a/Engine/Plugins/Mutable/Source/CustomizableObjectEditor/Private/MuCOE/GenerateMutableSource/GenerateMutableSourceMaterial.cpp b/Engine/Plugins/Mutable/Source/CustomizableObjectEditor/Private/MuCOE/GenerateMutableSource/GenerateMutableSourceMaterial.cpp
--- a/Engine/Plugins/Mutable/Source/CustomizableObjectEditor/Private/MuCOE/GenerateMutableSource/GenerateMutableSourceMaterial.cpp
+++ b/Engine/Plugins/Mutable/Source/CustomizableObjectEditor/Private/MuCOE/GenerateMutableSource/GenerateMutableSourceMaterial.cpp
@@ -439,6 +439,13 @@ UE::Mutable::Private::Ptr<UE::Mutable::Private::NodeMaterial> GenerateMutableSou
 								MaterialNode->ImageValues.Add(CurrentBreakParameter.ParameterKey, FinalNodeImag);
 							}
 						}
+						else
+						{
+							FText Message = FText::Format(LOCTEXT("NoReferenceTextureToBreak", "Could not find a Reference Texture for break parameter {0}."), FText::FromName(CurrentBreakParameter.ParameterKey.ParameterName));
+							GenerationContext.Log(Message, TypedNodeMaterialConstant);
+
+							MaterialNode = nullptr;
+						}
 					}
 					
 					else if (CurrentBreakParameter.ParameterType == UEdGraphSchema_CustomizableObject::PC_Texture_Passthrough)
@@ -626,14 +633,21 @@ UE::Mutable::Private::Ptr<UE::Mutable::Private::NodeMaterial> GenerateMutableSou
 					const uint32 LODBias = ComputeLODBiasForTexture(GenerationContext, *ReferenceTexture, nullptr, MaterialOptions.ReferenceTextureSize);
 					TextureSize.X = FMath::Max(ReferenceTexture->Source.GetSizeX() >> LODBias, 1);
 					TextureSize.Y = FMath::Max(ReferenceTexture->Source.GetSizeY() >> LODBias, 1);
-				}
 
-				ResizeNode->SizeX = TextureSize.X;
-				ResizeNode->SizeY = TextureSize.Y;
+					ResizeNode->SizeX = TextureSize.X;
+					ResizeNode->SizeY = TextureSize.Y;
 
-				// Force the same format that the default texture if any.
-				UE::Mutable::Private::Ptr<UE::Mutable::Private::NodeImageFormat> FormatNode = GenerateMutableImageFormat(ResizeNode, ReferenceTexture, GenerationContext, Node);
-				MaterialNode->ImageParametersToCompile.Add(CurrentBreakParameter.ParameterKey, FormatNode);
+					// Force the same format that the default texture if any.
+					UE::Mutable::Private::Ptr<UE::Mutable::Private::NodeImageFormat> FormatNode = GenerateMutableImageFormat(ResizeNode, ReferenceTexture, GenerationContext, Node);
+					MaterialNode->ImageParametersToCompile.Add(CurrentBreakParameter.ParameterKey, FormatNode);
+				}
+				else
+				{
+					FText Message = FText::Format(LOCTEXT("NoReferenceTextureToBreak", "Could not find a Reference Texture for break parameter {0}."), FText::FromName(CurrentBreakParameter.ParameterKey.ParameterName));
+					GenerationContext.Log(Message, TypedNodeMaterialConstant);
+
+					MaterialNode = nullptr;
+				}
 			}
 		}
 
