[Mutable] SkeletalMesh Parameter: add support for meshes with unlimited bone influences.

#rnx

#Synced-CL 50587801

diff --git a/Engine/Plugins/Mutable/Source/CustomizableObject/Private/MuCO/UnrealConversionUtils.cpp b/Engine/Plugins/Mutable/Source/CustomizableObject/Private/MuCO/UnrealConversionUtils.cpp
--- a/Engine/Plugins/Mutable/Source/CustomizableObject/Private/MuCO/UnrealConversionUtils.cpp
+++ b/Engine/Plugins/Mutable/Source/CustomizableObject/Private/MuCO/UnrealConversionUtils.cpp
@@ -1781,12 +1781,53 @@ namespace UnrealConversionUtils
 						const int32 MaxBonesPerVertex = SkinBuffer->GetMaxBoneInfluences();
 						MutableMeshBufferUtils::SetupSkinBuffer(CurrentVertexBuffer, MaxBoneIndexTypeSizeBytes, MaxBoneWeightTypeSizeBytes, MaxBonesPerVertex, Vertices);
 
-						int32 ElementSize = Vertices.Buffers[CurrentVertexBuffer].ElementSize;
-						const uint8* SourceVertexData = ((const uint8*)SkinBuffer->GetDataVertexBuffer()->GetWeightData()) + FirstVertexIndex * ElementSize;
+						switch (SkinBuffer->GetBoneInfluenceType())
+						{
+						case GPUSkinBoneInfluenceType::DefaultBoneInfluence:
+						{
+							int32 ElementSize = Vertices.Buffers[CurrentVertexBuffer].ElementSize;
+							const uint8* SourceVertexData = ((const uint8*)SkinBuffer->GetDataVertexBuffer()->GetWeightData()) + FirstVertexIndex * ElementSize;
+
+							check(ElementSize * (FirstVertexIndex + VertexCount) <= (int32)SkinBuffer->GetVertexDataSize());
+
+							FMemory::Memcpy(Vertices.GetBufferData(CurrentVertexBuffer), SourceVertexData, ElementSize * VertexCount);
+							break;
+						}
+						case GPUSkinBoneInfluenceType::UnlimitedBoneInfluence:
+						{
+							check(SkinBuffer->GetVariableBonesPerVertex());
+
+							int32 ElementSize = Vertices.Buffers[CurrentVertexBuffer].ElementSize;
+							uint8* DestVertexData = Vertices.GetBufferData(CurrentVertexBuffer);
+							FMemory::Memzero(DestVertexData, VertexCount * ElementSize);
+
+							const uint8* SourceVertexData = (const uint8*)SkinBuffer->GetDataVertexBuffer()->GetWeightData();
+							const FSkinWeightLookupVertexBuffer* LookUpVertexBuffer = SkinBuffer->GetLookupVertexBuffer();
+							for (int32 VertexIndex = 0; VertexIndex < VertexCount; ++VertexIndex)
+							{
+								uint32 WeightOffset = 0;
+								uint32 InfluenceCount = 0;
+								SkinBuffer->GetVertexInfluenceOffsetCount(FirstVertexIndex + VertexIndex, WeightOffset, InfluenceCount);
+
+								// Bone indices
+								int32 SourceIndicesDataSize = InfluenceCount * MaxBoneIndexTypeSizeBytes;
+								FMemory::Memcpy(DestVertexData, SourceVertexData + WeightOffset, SourceIndicesDataSize);
+								DestVertexData += MaxBonesPerVertex * MaxBoneIndexTypeSizeBytes;
 
-						check(ElementSize * (FirstVertexIndex + VertexCount) <= (int32)SkinBuffer->GetVertexDataSize());
 
-						FMemory::Memcpy(Vertices.GetBufferData(CurrentVertexBuffer), SourceVertexData, ElementSize * VertexCount);
+								// Weights
+								int32 SourceWeightDataSize = InfluenceCount * MaxBoneWeightTypeSizeBytes;
+								FMemory::Memcpy(DestVertexData, SourceVertexData + WeightOffset + SourceIndicesDataSize, SourceWeightDataSize);
+								DestVertexData += MaxBonesPerVertex * MaxBoneWeightTypeSizeBytes;
+							}
+
+							check((int32)(DestVertexData - Vertices.GetBufferData(CurrentVertexBuffer)) == (int32)(VertexCount * ElementSize));
+							break;
+						}
+						default:
+							unimplemented();
+							break;
+						}
 
 						++CurrentVertexBuffer;
 					}
