[Mutable] Cache DoesBlockExist.

#rnx

#Synced-CL 50587801

diff --git a/Engine/Plugins/Mutable/Source/CustomizableObject/Internal/MuCO/CustomizableObjectPrivate.h b/Engine/Plugins/Mutable/Source/CustomizableObject/Internal/MuCO/CustomizableObjectPrivate.h
--- a/Engine/Plugins/Mutable/Source/CustomizableObject/Internal/MuCO/CustomizableObjectPrivate.h
+++ b/Engine/Plugins/Mutable/Source/CustomizableObject/Internal/MuCO/CustomizableObjectPrivate.h
@@ -704,7 +704,9 @@ struct FMutableStateData
 enum class EMutableFileFlags : uint8
 {
 	None	= 0,
-	HighRes = 1 << 0
+	HighRes = 1 << 0,
+	DoesExist = 1 << 1,
+	DoesExistCached = 1 << 2
 };
 
 
@@ -939,6 +941,7 @@ struct FModelStreamableBulkData
 	/** Map of Hash to Streaming blocks, used to stream a block of data representing a resource from the BulkData */
 	TMap<uint32, FMutableStreamableBlock> ModelStreamables;
 
+	TArray<EMutableFileFlags> BulkDataFlags;
 	TArray<FByteBulkData> StreamableBulkData;
 
 #if WITH_EDITORONLY_DATA
diff --git a/Engine/Plugins/Mutable/Source/CustomizableObject/Private/MuCO/UnrealMutableModelDiskStreamer.cpp b/Engine/Plugins/Mutable/Source/CustomizableObject/Private/MuCO/UnrealMutableModelDiskStreamer.cpp
--- a/Engine/Plugins/Mutable/Source/CustomizableObject/Private/MuCO/UnrealMutableModelDiskStreamer.cpp
+++ b/Engine/Plugins/Mutable/Source/CustomizableObject/Private/MuCO/UnrealMutableModelDiskStreamer.cpp
@@ -91,6 +91,8 @@ bool FUnrealMutableModelBulkReader::PrepareStreamingForObject(UCustomizableObjec
 		return true;
 	}
 
+	ModelStreamableBulkData->BulkDataFlags.SetNumZeroed(ModelStreamableBulkData->StreamableBulkData.Num());
+
 	const TSharedPtr<const UE::Mutable::Private::FModel>& Model = CustomizableObject->GetPrivate()->GetModel();
 
 	// Is the object already prepared for streaming?
@@ -248,12 +250,23 @@ bool FUnrealMutableModelBulkReader::DoesBlockExist(const UE::Mutable::Private::F
 	}
 
 #if !WITH_EDITOR
-	if (ensure(PinnedModelStreamableBulkData->StreamableBulkData.IsValidIndex(Block->FileId)))
+	if (ensure(PinnedModelStreamableBulkData->BulkDataFlags.IsValidIndex(Block->FileId)))
 	{	
-		MUTABLE_CPUPROFILER_SCOPE(FUnrealMutableModelBulkReader::DoesBlockExist_Bulk)
+		MUTABLE_CPUPROFILER_SCOPE(FUnrealMutableModelBulkReader::DoesBlockExist_Bulk);
 
-		FByteBulkData& BulkData = PinnedModelStreamableBulkData->StreamableBulkData[Block->FileId];
-		return BulkData.DoesExist();
+		EMutableFileFlags BulkDataFlags = PinnedModelStreamableBulkData->BulkDataFlags[Block->FileId];
+		if (EnumHasAnyFlags(BulkDataFlags, EMutableFileFlags::DoesExistCached))
+		{
+			return EnumHasAnyFlags(BulkDataFlags, EMutableFileFlags::DoesExist);
+		}
+		else
+		{
+			const bool bDoesExist = PinnedModelStreamableBulkData->StreamableBulkData[Block->FileId].DoesExist();
+			EnumAddFlags(BulkDataFlags, bDoesExist ? EMutableFileFlags::DoesExist : EMutableFileFlags::None);
+			EnumAddFlags(BulkDataFlags, EMutableFileFlags::DoesExistCached);
+			PinnedModelStreamableBulkData->BulkDataFlags[Block->FileId] = BulkDataFlags;
+			return bDoesExist;
+		}
 	}
 	else
 	{
