[Mutable] Remove EMutableMeshConversionFlags::IgnoreTexCoords when generating meshes for ClipDeform operations.

#rnx
#Synced-CL 50587801

diff --git a/Engine/Plugins/Mutable/Source/CustomizableObjectEditor/Private/MuCOE/GenerateMutableSource/GenerateMutableSourceModifier.cpp b/Engine/Plugins/Mutable/Source/CustomizableObjectEditor/Private/MuCOE/GenerateMutableSource/GenerateMutableSourceModifier.cpp
--- a/Engine/Plugins/Mutable/Source/CustomizableObjectEditor/Private/MuCOE/GenerateMutableSource/GenerateMutableSourceModifier.cpp
+++ b/Engine/Plugins/Mutable/Source/CustomizableObjectEditor/Private/MuCOE/GenerateMutableSource/GenerateMutableSourceModifier.cpp
@@ -113,7 +113,6 @@ UE::Mutable::Private::Ptr<UE::Mutable::Private::NodeModifier> GenerateMutableSou
 			EMutableMeshConversionFlags::IgnoreSkinning |
 			EMutableMeshConversionFlags::IgnorePhysics |
 			EMutableMeshConversionFlags::IgnoreMorphs |
-			EMutableMeshConversionFlags::IgnoreTexCoords |
 			EMutableMeshConversionFlags::IgnoreAUD |
 			EMutableMeshConversionFlags::DoNotCreateMeshMetadata;
 		GenerationContext.MeshGenerationFlags.Push(ModifiersMeshFlags);
