[Mutable] Updated States tab to recieve messages when updating a Customizable Object Instance.
+ Added a new Compilation Notification to indicate graph compilation.
+ Removed unnecessary method from COEditor
+ Fixed small coment error

#rnx
#Synced-CL 50588143

diff --git a/Engine/Plugins/Mutable/Source/CustomizableObject/Internal/MuCO/CustomizableObjectInstancePrivate.h b/Engine/Plugins/Mutable/Source/CustomizableObject/Internal/MuCO/CustomizableObjectInstancePrivate.h
--- a/Engine/Plugins/Mutable/Source/CustomizableObject/Internal/MuCO/CustomizableObjectInstancePrivate.h
+++ b/Engine/Plugins/Mutable/Source/CustomizableObject/Internal/MuCO/CustomizableObjectInstancePrivate.h
@@ -8,6 +8,7 @@
 #include "MuCO/CustomizableObjectPrivate.h"
 #include "MuCO/CustomizableObjectSystemPrivate.h"
 #include "MuR/Instance.h"
+#include "MuR/MutableEditorLogger.h"
 #include "GameplayTagContainer.h"
 #include "MuCO/DescriptorHash.h"
 #include "UObject/Package.h"
@@ -370,6 +371,9 @@ public:
 
 	TMap<FTextureReuseCacheKey, TStrongObjectPtr<UTexture2D>> LastTextureReuseCache;
 	TMap<FTextureReuseCacheKey, TStrongObjectPtr<UTexture2D>> CurrentTextureReuseCache;
+
+	/** Pointer to a instance of class that logs update messages to an editor widget. */
+	TSharedPtr<IMutableEditorLogger> UpdateLogger;
 };
 
 
diff --git a/Engine/Plugins/Mutable/Source/CustomizableObject/Internal/MuCO/CustomizableObjectPrivate.h b/Engine/Plugins/Mutable/Source/CustomizableObject/Internal/MuCO/CustomizableObjectPrivate.h
--- a/Engine/Plugins/Mutable/Source/CustomizableObject/Internal/MuCO/CustomizableObjectPrivate.h
+++ b/Engine/Plugins/Mutable/Source/CustomizableObject/Internal/MuCO/CustomizableObjectPrivate.h
@@ -1685,7 +1685,7 @@ public:
 	CUSTOMIZABLEOBJECT_API int32 GetStateParameterIndex(int32 StateIndex, int32 ParameterIndex) const;
 
 	CUSTOMIZABLEOBJECT_API bool IsParameterMultidimensional(const int32& InParamIndex) const;
-	
+
 	/** Cache of generated Skeletal Meshes. Passthrough Skeletal Meshes are not contemplated. */
 	FMeshCache MeshCache;
 
@@ -1759,8 +1759,10 @@ public:
 	  * This Id is part of the DDC key of the cooked data distribution. */
 	UPROPERTY()
 	FGuid CookedDataDistributionId;
-
 #endif
+
+	/** Editor Log list to print all the messages related to this CO. */
+	TSharedPtr<class IMutableEditorLogger> CompilationLogger;
 	
 	static constexpr int32 DerivedDataVersion = 0x5854e531;
 };
diff --git a/Engine/Plugins/Mutable/Source/CustomizableObject/Internal/MuCO/CustomizableObjectSystemPrivate.h b/Engine/Plugins/Mutable/Source/CustomizableObject/Internal/MuCO/CustomizableObjectSystemPrivate.h
--- a/Engine/Plugins/Mutable/Source/CustomizableObject/Internal/MuCO/CustomizableObjectSystemPrivate.h
+++ b/Engine/Plugins/Mutable/Source/CustomizableObject/Internal/MuCO/CustomizableObjectSystemPrivate.h
@@ -852,6 +852,8 @@ public:
 	 */
 	CUSTOMIZABLEOBJECT_API static void SetUsageOfBenchmarkingSettings(bool bUseBenchmarkingOptimizedSettings);
 
+	void LogUpdateMessage(const FString& Message, ELogVerbosity::Type Verbosity, UCustomizableObjectInstance* Instance, bool bClearMessageList = false);
+
 private:
 	/** Flag that controls some of the settings used for the generation of instances. */
 	inline static bool bUseBenchmarkingSettings = false;
diff --git a/Engine/Plugins/Mutable/Source/CustomizableObject/Private/MuCO/CustomizableObject.cpp b/Engine/Plugins/Mutable/Source/CustomizableObject/Private/MuCO/CustomizableObject.cpp
--- a/Engine/Plugins/Mutable/Source/CustomizableObject/Private/MuCO/CustomizableObject.cpp
+++ b/Engine/Plugins/Mutable/Source/CustomizableObject/Private/MuCO/CustomizableObject.cpp
@@ -27,6 +27,7 @@
 #include "MuCO/MutableProjectorTypeUtils.h"
 #include "MuCO/UnrealMutableModelDiskStreamer.h"
 #include "MuR/Model.h"
+#include "MuR/MutableEditorLogger.h"
 #include "MuR/Operations.h"
 #include "PhysicsEngine/PhysicsAsset.h"
 #include "Serialization/MemoryReader.h"
@@ -3339,8 +3340,6 @@ void UCustomizableObjectPrivate::LogMemory()
 	}
 
 }
-
-
 #endif // WITH_EDITOR
 
 
diff --git a/Engine/Plugins/Mutable/Source/CustomizableObject/Private/MuCO/CustomizableObjectInstance.cpp b/Engine/Plugins/Mutable/Source/CustomizableObject/Private/MuCO/CustomizableObjectInstance.cpp
--- a/Engine/Plugins/Mutable/Source/CustomizableObject/Private/MuCO/CustomizableObjectInstance.cpp
+++ b/Engine/Plugins/Mutable/Source/CustomizableObject/Private/MuCO/CustomizableObjectInstance.cpp
@@ -53,6 +53,7 @@
 #include "MuCO/Plugins/IMutableClothingModule.h"
 #include "Rendering/SkeletalMeshModel.h"
 
+#include "MuR/MutableEditorLogger.h"
 #include "MuR/PassthroughObject.h"
 
 #include UE_INLINE_GENERATED_CPP_BY_NAME(CustomizableObjectInstance)
diff --git a/Engine/Plugins/Mutable/Source/CustomizableObject/Private/MuCO/CustomizableObjectSystem.cpp b/Engine/Plugins/Mutable/Source/CustomizableObject/Private/MuCO/CustomizableObjectSystem.cpp
--- a/Engine/Plugins/Mutable/Source/CustomizableObject/Private/MuCO/CustomizableObjectSystem.cpp
+++ b/Engine/Plugins/Mutable/Source/CustomizableObject/Private/MuCO/CustomizableObjectSystem.cpp
@@ -39,6 +39,7 @@
 #include "Interfaces/ITargetPlatform.h"
 #include "MuT/UnrealPixelFormatOverride.h"
 #include "MuR/Parameters.h"
+#include "IMessageLogListing.h"
 #else
 #include "Engine/Engine.h"
 #endif
@@ -741,9 +742,9 @@ UCustomizableObjectSystem* UCustomizableObjectSystem::GetInstance()
 
 		UCustomizableObjectSystemPrivate::SSystem = NewObject<UCustomizableObjectSystem>(UCustomizableObjectSystem::StaticClass());
 		check(UCustomizableObjectSystemPrivate::SSystem != nullptr);
-		checkf(!GUObjectArray.IsDisregardForGC(UCustomizableObjectSystemPrivate::SSystem), TEXT("Mutable was initialized too early in the UE4 init process, for instance, in the constructor of a default UObject."));
+		checkf(!GUObjectArray.IsDisregardForGC(UCustomizableObjectSystemPrivate::SSystem), TEXT("Mutable was initialized too early in the UE init process, for instance, in the constructor of a default UObject."));
 		UCustomizableObjectSystemPrivate::SSystem->AddToRoot();
-		checkf(!GUObjectArray.IsDisregardForGC(UCustomizableObjectSystemPrivate::SSystem), TEXT("Mutable was initialized too early in the UE4 init process, for instance, in the constructor of a default UObject."));
+		checkf(!GUObjectArray.IsDisregardForGC(UCustomizableObjectSystemPrivate::SSystem), TEXT("Mutable was initialized too early in the UE init process, for instance, in the constructor of a default UObject."));
 		UCustomizableObjectSystemPrivate::SSystem->InitSystem();
 	}
 
@@ -1128,13 +1129,18 @@ void FinishUpdateGlobal(const TSharedRef<FUpdateContextPrivate>& Context)
 	
 	if (Context->StartUpdateTime != 0.0) // Update started.
 	{
-		Context->UpdateTime = FPlatformTime::Seconds() - Context->StartUpdateTime;		
+		Context->UpdateTime = FPlatformTime::Seconds() - Context->StartUpdateTime;
 	}
 	
 	const FName ObjectName = GetFNameSafe(Object);
 	const FName InstanceName = GetFNameSafe(Instance);
-	UE_LOG(LogMutable, Verbose, TEXT("Finished Update Skeletal Mesh Async. CustomizableObject=%s Instance=%s, Frame=%d  QueueTime=%f, UpdateTime=%f"), *ObjectName.ToString(), *InstanceName.ToString(), GFrameNumber, Context->QueueTime, Context->UpdateTime);
-	
+
+	if (SystemPrivate)
+	{
+		FString Message = FString::Printf(TEXT("Finished Update Skeletal Mesh Async. CustomizableObject=%s Instance=%s, Frame=%d  QueueTime=%f, UpdateTime=%f"), *ObjectName.ToString(), *InstanceName.ToString(), GFrameNumber, Context->QueueTime, Context->UpdateTime);
+		SystemPrivate->LogUpdateMessage(Message, ELogVerbosity::Verbose, Instance);
+	}
+
 	if (SystemPrivate && FLogBenchmarkUtil::IsBenchmarkingReportingEnabled())	
 	{
 		FFunctionGraphTask::CreateAndDispatchWhenReady( // Calling Benchmark in a task so we make sure we exited all scopes.
@@ -3075,7 +3081,10 @@ namespace impl
 			Operation->bLiveUpdateMode = SystemPrivate->EnableMutableLiveUpdate && (StateData ? StateData->bLiveUpdateMode : false);
 		
 			Operation->bStreamMeshLODs = Operation->bLiveUpdateMode ? false : IsStreamingEnabled(*Object, State);
-			
+
+			TSharedPtr<UE::Mutable::Private::FLiveInstanceLogger> UpdateInstanceLogger = MakeShared<UE::Mutable::Private::FLiveInstanceLogger>();
+			UpdateInstanceLogger->SetLogger(InstancePrivate->UpdateLogger);
+
 			// Ensure that in the case of running with the live update enabled we have mesh and mip streaming disabled
 			check(!Operation->bLiveUpdateMode || Operation->bNeverStreamMips);
 			check(!Operation->bLiveUpdateMode || Operation->MipsToSkip == 0);
@@ -3088,7 +3097,8 @@ namespace impl
 				ObjectPrivate->MaterialIdRegistry,
 				ObjectPrivate->SkeletalMeshIdRegistry,
 				Operation->ExternalResourceProvider,
-				Operation->PassthroughObjectLoader);
+				Operation->PassthroughObjectLoader,
+				UpdateInstanceLogger);
 			
 			Operation->LiveInstance->Parameters = Operation->Parameters;
 			Operation->LiveInstance->State = Operation->GetCapturedDescriptor().GetState();
@@ -3224,6 +3234,7 @@ namespace impl
 			Task_Mutable_GetMeshes(Operation);
 		},
 		{ Prerequisites});
+
 		UE::Tasks::FTask Mutable_GetSurfacesTask = SystemPrivate->MutableTaskGraph.AddMutableThreadTask(
 		TEXT("Task_Mutable_GetSurfaces"),
 		[Operation]()
@@ -3951,7 +3962,8 @@ void UCustomizableObjectSystemPrivate::StartUpdateSkeletalMesh(const TSharedRef<
 	{
 		const FName ObjectName = Object ? Object->GetFName() : NAME_None;
 		const FName InstanceName = Context->Instance->GetFName();
-		UE_LOG(LogMutable, Log, TEXT("Started Update Skeletal Mesh Async. CustomizableObject=%s Instance=%s, Frame=%d"), *ObjectName.ToString(), *InstanceName.ToString(), GFrameNumber);
+		FString Message = FString::Printf(TEXT("Started Update Skeletal Mesh Async. CustomizableObject=%s Instance=%s, Frame=%d"), *ObjectName.ToString(), *InstanceName.ToString(), GFrameNumber);
+		LogUpdateMessage(Message, ELogVerbosity::Log, Context->Instance.Get());
 
 		const float CurrentTime = FPlatformTime::Seconds();
 
@@ -4048,6 +4060,61 @@ int32 UCustomizableObjectSystemPrivate::BlockTillAllRequestsFinished(float TimeL
 }
 
 
+void UCustomizableObjectSystemPrivate::LogUpdateMessage(const FString& Message, ELogVerbosity::Type Verbosity, UCustomizableObjectInstance* Instance, bool bClearMessageList)
+{
+	switch (Verbosity)
+	{
+	case ELogVerbosity::Error:
+		UE_LOG(LogMutable, Error, TEXT("%s"), *Message);
+		break;
+	case ELogVerbosity::Warning:
+		UE_LOG(LogMutable, Warning, TEXT("%s"), *Message);
+		break;
+	case ELogVerbosity::Verbose:
+		UE_LOG(LogMutable, Verbose, TEXT("%s"), *Message);
+		break;
+	case ELogVerbosity::Display:
+		UE_LOG(LogMutable, Verbose, TEXT("%s"), *Message);
+		break;
+	case ELogVerbosity::Log:
+	default:
+		UE_LOG(LogMutable, Log, TEXT("%s"), *Message);
+		break;
+	}
+
+	if (Instance && Instance->GetPrivate()->UpdateLogger)
+	{
+		if (bClearMessageList)
+		{
+			Instance->GetPrivate()->UpdateLogger->ClearLogMessageList();
+		}
+
+		// Translation from ELogVerbosity to EMessageSeverity
+		EMessageSeverity::Type MessageSeverity;
+
+		switch(Verbosity)
+		{
+		case ELogVerbosity::Error:
+			MessageSeverity = EMessageSeverity::Error;
+			break;
+		case ELogVerbosity::Warning:
+			MessageSeverity = EMessageSeverity::Warning;
+			break;
+		case ELogVerbosity::Verbose:
+		case ELogVerbosity::Log:
+		case ELogVerbosity::Display:
+		default:
+			MessageSeverity = EMessageSeverity::Info;
+			break;
+		}
+
+		TSharedRef<FTokenizedMessage> LogMessage = FTokenizedMessage::Create(MessageSeverity, FText::FromString(Message));
+		LogMessage->SetIndentationLevel(MessageSeverity == EMessageSeverity::Info ? 0 : 1);
+		Instance->GetPrivate()->UpdateLogger->LogMessage(LogMessage);
+	}
+}
+
+
 #if !WITH_EDITOR
 void ReconstructMorphTargetsFromMeshCompressedData(const UE::Mutable::Private::FMesh& Mesh, TArray<FMorphTargetCompressedLODModel>& OutMorphTargets, UE::Mutable::Private::EMorphUsageFlags UsageFilter)
 {
diff --git a/Engine/Plugins/Mutable/Source/CustomizableObjectEditor/Internal/MuCOE/CompileRequest.h b/Engine/Plugins/Mutable/Source/CustomizableObjectEditor/Internal/MuCOE/CompileRequest.h
--- a/Engine/Plugins/Mutable/Source/CustomizableObjectEditor/Internal/MuCOE/CompileRequest.h
+++ b/Engine/Plugins/Mutable/Source/CustomizableObjectEditor/Internal/MuCOE/CompileRequest.h
@@ -50,8 +50,6 @@ public:
 	bool bSkipIfCompiled = false;
 	bool bSkipIfNotOutOfDate = false;
 	bool bSilentCompilation = true;
-
-	TSharedPtr<FCustomizableObjectCustomLogList> EditorCustomLogListPtr;
 };
 
 #undef UE_API
diff --git a/Engine/Plugins/Mutable/Source/CustomizableObjectEditor/Private/MuCOE/CustomizableInstanceDetails.cpp b/Engine/Plugins/Mutable/Source/CustomizableObjectEditor/Private/MuCOE/CustomizableInstanceDetails.cpp
--- a/Engine/Plugins/Mutable/Source/CustomizableObjectEditor/Private/MuCOE/CustomizableInstanceDetails.cpp
+++ b/Engine/Plugins/Mutable/Source/CustomizableObjectEditor/Private/MuCOE/CustomizableInstanceDetails.cpp
@@ -314,6 +314,12 @@ void FCustomizableInstanceDetails::Refresh() const
 
 void FCustomizableInstanceDetails::UpdateInstance()
 {
+	if (CustomInstance->GetPrivate()->UpdateLogger)
+	{
+		// Clear the States tab before loging all the messages
+		CustomInstance->GetPrivate()->UpdateLogger->ClearLogMessageList();
+	}
+
 	CustomInstance->UpdateSkeletalMeshAsync(true, true);
 	CustomInstance->PostEditChange();
 }
diff --git a/Engine/Plugins/Mutable/Source/CustomizableObjectEditor/Private/MuCOE/CustomizableObjectCompiler.cpp b/Engine/Plugins/Mutable/Source/CustomizableObjectEditor/Private/MuCOE/CustomizableObjectCompiler.cpp
--- a/Engine/Plugins/Mutable/Source/CustomizableObjectEditor/Private/MuCOE/CustomizableObjectCompiler.cpp
+++ b/Engine/Plugins/Mutable/Source/CustomizableObjectEditor/Private/MuCOE/CustomizableObjectCompiler.cpp
@@ -35,8 +35,7 @@
 #include "DerivedDataCacheKey.h"
 #include "DerivedDataRequestOwner.h"
 #include "MuR/External/Operation.h"
-#include "IMessageLogListing.h"
-
+#include "MuR/MutableEditorLogger.h"
 
 class UTexture2D;
 
@@ -284,10 +283,11 @@ void FCustomizableObjectCompiler::Compile(const TSharedRef<FCompilationRequest>&
 	CurrentObject = Object;
 	CurrentOptions = InCompileRequest->Options;
 
-	// If it's compiling inside an editor, clear any previous message of the terminal
-	if (CurrentRequest->EditorCustomLogListPtr)
+	// Add visual separation in States tab between Compilation Messages and Update Messages
+	if(CurrentObject->GetPrivate()->CompilationLogger)
 	{
-		CurrentRequest->EditorCustomLogListPtr->MessageLogListing->ClearMessages();
+		TSharedRef<FTokenizedMessage> Message = FTokenizedMessage::Create(EMessageSeverity::Info, FText::FromString("=================== Object Compilation ==================="));
+		CurrentObject->GetPrivate()->CompilationLogger->LogMessage(Message);
 	}
 	
 	if (!UCustomizableObjectSystem::IsActive())
@@ -410,7 +410,6 @@ void FCustomizableObjectCompiler::Compile(const TSharedRef<FCompilationRequest>&
 		
 	}
 
-	
 	if (CurrentOptions.bForceLargeLODBias)
 	{
 		UE_LOG(LogMutable, Display, TEXT("Compiling Customizable Object with %d LODBias."), CurrentOptions.DebugBias);
@@ -919,6 +918,8 @@ void FCustomizableObjectCompiler::CompileInternal(bool bAsync)
 	CompilationLogsContainer.ClearMessageCounters();
 	CompilationLogsContainer.ClearMessagesArray();
 
+	AddCompileNotification(LOCTEXT("GraphCompilationMessage", "Compiling Graph"));
+
 	// Generate the mutable node expression
 	UE::Mutable::Private::Ptr<UE::Mutable::Private::NodeObject> MutableRoot = GenerateMutableRoot(CurrentObject, GenerationContext);
 	if (!MutableRoot)
@@ -1088,7 +1089,7 @@ void FCustomizableObjectCompiler::CompileInternal(bool bAsync)
 		}
 		else
 		{
-			AddCompileNotification(LOCTEXT("CustomizableObjectCompileInProgress", "Compiling"));
+			AddCompileNotification(LOCTEXT("CustomizableObjectCompileInProgress", "Compiling Object"));
 
 			// Even for async build, we spawn a thread, so that we can set a large stack. 
 			// Thread names need to be unique, apparently.
@@ -1166,6 +1167,12 @@ void FCustomizableObjectCompiler::CompleteRequest(ECompilationStatePrivate State
 	RemoveCompileNotification();
 	NotifyCompilationErrors();
 
+	if (CurrentObject && CurrentObject->GetPrivate()->CompilationLogger)
+	{
+		TSharedRef<FTokenizedMessage> Message = FTokenizedMessage::Create(EMessageSeverity::Info, FText::FromString("=========================================================="));
+		CurrentObject->GetPrivate()->CompilationLogger->LogMessage(Message);
+	}
+
 	// Update compilation progress notification
 	if (CompileNotificationHandle.IsValid())
 	{
@@ -1788,7 +1795,8 @@ void FCustomizableObjectCompiler::AddCompileNotification(const FText& Compilatio
 	.Notification(!CurrentRequest->bSilentCompilation)
 	.CustomNotification()
 	.FixNotification()
-	.Log(CurrentRequest->EditorCustomLogListPtr);
+	.CustomizableObject(CurrentObject)
+	.Log();
 }
 
 
@@ -1830,18 +1838,13 @@ void FCustomizableObjectCompiler::NotifyCompilationErrors() const
 		:
 		FText::Format(LOCTEXT("CompilationFinished_WithoutIgnoreds", "{0} finished compiling with {1} {1}|plural(one=warning,other=warnings) and {2} {2}|plural(one=error,other=errors)."), Prefix, NumWarnings, NumErrors);
 	
-	// Open the Stats tab if it was previously closed an there are errors/warnings
-	if (CurrentRequest->EditorCustomLogListPtr && !NoWarningsOrErrors)
-	{
-		CurrentRequest->EditorCustomLogListPtr->DisplayStatesTab();
-	}
-
 	FCustomizableObjectEditorLogger::CreateLog(Message)
 	.Category(ELoggerCategory::Compilation)
 	.Severity(Severity)
 	.Notification(!CurrentRequest->bSilentCompilation || !NoWarningsOrErrors)
 	.CustomNotification()
-	.Log(CurrentRequest->EditorCustomLogListPtr);
+	.CustomizableObject(CurrentObject)
+	.Log();
 }
 
 
@@ -1854,7 +1857,8 @@ void FCustomizableObjectCompiler::CompilerLog(const FText& Message, const TArray
 			.Context(Context)
 			.BaseObject(bAddBaseObjectInfo)
 			.SpamBin(SpamBin)
-			.Log(CurrentRequest ? CurrentRequest->EditorCustomLogListPtr : nullptr); // Compilations from a debugger do not have a compilation request
+			.CustomizableObject(CurrentRequest ? CurrentRequest->GetCustomizableObject() : nullptr)
+			.Log(); 
 	}
 }
 
diff --git a/Engine/Plugins/Mutable/Source/CustomizableObjectEditor/Private/MuCOE/CustomizableObjectEditor.cpp b/Engine/Plugins/Mutable/Source/CustomizableObjectEditor/Private/MuCOE/CustomizableObjectEditor.cpp
--- a/Engine/Plugins/Mutable/Source/CustomizableObjectEditor/Private/MuCOE/CustomizableObjectEditor.cpp
+++ b/Engine/Plugins/Mutable/Source/CustomizableObjectEditor/Private/MuCOE/CustomizableObjectEditor.cpp
@@ -19,6 +19,7 @@
 #include "GraphEditorActions.h"
 #include "HAL/PlatformApplicationMisc.h"
 #include "IDetailsView.h"
+#include "IMessageLogListing.h"
 #include "MaterialGraph/MaterialGraphNode.h"
 #include "MaterialGraph/MaterialGraphNode_Comment.h"
 #include "MuCO/CustomizableObjectInstance.h"
@@ -166,6 +167,12 @@ FCustomizableObjectEditor::~FCustomizableObjectEditor()
 		{
 			PreviewInstance->GetPrivate()->SaveParametersToProfile(PreviewInstance->GetPrivate()->SelectedProfileIndex);
 		}
+
+		// Remove the message log list when closing the editor
+		if (PreviewInstance->GetPrivate()->UpdateLogger)
+		{
+			PreviewInstance->GetPrivate()->UpdateLogger.Reset();
+		}
 	}
 	
 	ObjectDetailsView.Reset();
@@ -174,6 +181,12 @@ FCustomizableObjectEditor::~FCustomizableObjectEditor()
 
 	CustomizableObject->GetPrivate()->Status.GetOnStateChangedDelegate().RemoveAll(this);
 
+	// Remove the message log list when closing the editor
+	if (CustomizableObject->GetPrivate()->CompilationLogger)
+	{
+		CustomizableObject->GetPrivate()->CompilationLogger.Reset();
+	}
+
 	FAssetRegistryModule& AssetRegistryModule = FModuleManager::LoadModuleChecked<FAssetRegistryModule>(TEXT("AssetRegistry"));
 	AssetRegistryModule.Get().OnFilesLoaded().RemoveAll(this);
 
@@ -243,7 +256,8 @@ void FCustomizableObjectEditor::InitCustomizableObjectEditor(const EToolkitMode:
 	CustomizableObjectEditorAdvancedPreviewSettings->LoadProfileEnvironment();
 	AdvancedPreviewSettingsWidget = CustomizableObjectEditorAdvancedPreviewSettings;
 
-	COEditorLogList = MakeShareable(new FCustomizableObjectCustomLogList(SharedThis(this), StatsWidget));
+	// Initialize the Log List for the current Customizable Object
+	InitLogList();
 
 	const TSharedRef<FTabManager::FLayout> StandaloneDefaultLayout = FTabManager::NewLayout( "Standalone_CustomizableObjectEditor_Layout_v1.5" )
 	->AddArea
@@ -328,6 +342,8 @@ void FCustomizableObjectEditor::InitCustomizableObjectEditor(const EToolkitMode:
 	PreviewInstance = CustomizableObject->CreateInstance();	
 	PreviewInstance->UpdatedNativeDelegate.AddSP(SharedThis(this), &FCustomizableObjectEditor::OnUpdatePreviewInstance);
 	PreviewInstance->SetBuildParameterRelevancy(true);
+	// Set the message log listing to get update messages
+	PreviewInstance->GetPrivate()->UpdateLogger = CustomizableObjectPrivate->CompilationLogger;
 
 	CustomizableInstanceDetailsView->SetObject(PreviewInstance, true);
 	
@@ -1332,12 +1348,19 @@ void FCustomizableObjectEditor::CompileObject(bool bOnlySelectedParameters, bool
 	TSharedRef<FCompilationRequest> CompileRequest = MakeShared<FCompilationRequest>(*CustomizableObject);
 	CompileRequest->bSilentCompilation = false;
 	CompileRequest->Options.bGatherReferences = bGatherReferences;
-	CompileRequest->EditorCustomLogListPtr = COEditorLogList;
 
 	if (bOnlySelectedParameters)
 	{
 		CompileRequest->Options.ParamNamesToSelectedOptions = GetCompileOnlySelectedParameters(*GetPreviewInstance());
 	}
+
+	//Ensure the Stats tab is open and empty
+	GetTabManager()->TryInvokeTab(FCustomizableObjectEditor::StatsTabId);
+
+	if (CustomizableObject->GetPrivate()->CompilationLogger)
+	{
+		CustomizableObject->GetPrivate()->CompilationLogger->ClearLogMessageList();
+	}
 	
 	ICustomizableObjectEditorModulePrivate::GetChecked().EnqueueCompileRequest(CompileRequest);
 }
@@ -1682,103 +1705,6 @@ const UObject* FCustomizableObjectEditor::GetObjectBeingEdited()
 }
 
 
-void FCustomizableObjectEditor::FindProperty(const FProperty* Property, const void* InContainer, const FString& FindString, const UObject& Context, bool& bFound)
-{
-	if (!Property || !InContainer)
-	{
-		return;
-	}
-
-	const FString PropertyName = Property->GetDisplayNameText().ToString();
-	if (PropertyName.Contains(FindString))
-	{
-		LogSearchResult(Context, "Property Name", bFound, *PropertyName);
-		bFound = true;
-	}
-	
-	for (int32 Index = 0; Index < Property->ArrayDim; ++Index)
-	{
-		const uint8* ValuePtr = Property->ContainerPtrToValuePtr<uint8>(InContainer, Index);
-
-		if (const FStrProperty* StringProperty = CastField<FStrProperty>(Property))
-		{
-			const FString* StringResult = StringProperty->GetPropertyValuePtr(ValuePtr);
-			if (StringResult->Contains(FindString))
-			{
-				LogSearchResult(Context, "Property Value", bFound, *StringResult);
-				bFound = true;
-			}
-		}
-		else if (const FEnumProperty* EnumProperty = CastField<FEnumProperty>(Property))
-		{
-			const UEnum* EnumResult = EnumProperty->GetEnum();
-
-			const FString StringResult = EnumResult->GetDisplayNameTextByIndex(*ValuePtr).ToString();
-			if (StringResult.Contains(FindString))
-			{
-				LogSearchResult(Context, "Property Value", bFound, StringResult);
-				bFound = true;
-			}
-		}
-		else if (const FSoftObjectProperty* SoftObjectProperty = CastField<FSoftObjectProperty>(Property))
-		{
-			const FString ObjectPath = SoftObjectProperty->GetPropertyValuePtr(ValuePtr)->ToString();
-			if (ObjectPath.Contains(FindString))
-			{
-				LogSearchResult(Context, "Property Value", bFound, ObjectPath);
-				bFound = true;
-			}
-		}
-		else if (const FObjectPropertyBase* ObjectProperty = CastField<FObjectPropertyBase>(Property))
-		{
-			if (const UObject* ObjectValue = ObjectProperty->GetObjectPropertyValue(ValuePtr))
-			{
-				const FString Name = ObjectValue->GetName();
-			
-				if (ObjectValue->GetName().Contains(FindString))
-				{
-					LogSearchResult(Context, "Property Value", bFound, Name);
-					bFound = true;
-				}				
-			}
-		}
-		else if (const FStructProperty* StructProperty = CastField<FStructProperty>(Property))
-		{
-			for (TFieldIterator<FProperty> It(StructProperty->Struct); It; ++It)
-			{
-				FindProperty(*It, ValuePtr, FindString, Context, bFound);
-			}
-		}
-		else if (const FArrayProperty* ArrayProperty = CastField<FArrayProperty>(Property))
-		{
-			FScriptArrayHelper ArrayHelper(ArrayProperty, ValuePtr);
-			for (int32 ValueIdx = 0; ValueIdx < ArrayHelper.Num(); ++ValueIdx)
-			{
-				FindProperty(ArrayProperty->Inner, ArrayHelper.GetRawPtr(ValueIdx), FindString, Context, bFound);
-			}
-		}
-		else if (const FSetProperty* SetProperty = CastField<FSetProperty>(Property))
-		{
-			FScriptSetHelper SetHelper(SetProperty, ValuePtr);
-			for (FScriptSetHelper::FIterator SetIt = SetHelper.CreateIterator(); SetIt; ++SetIt)
-			{
-				FindProperty(SetProperty->ElementProp, SetHelper.GetElementPtr(SetIt), FindString, Context, bFound);
-			}
-		}
-		else if (const FMapProperty* MapProperty = CastField<FMapProperty>(Property))
-		{
-			FScriptMapHelper MapHelper(MapProperty, ValuePtr);
-			for (FScriptMapHelper::FIterator MapIt = MapHelper.CreateIterator(); MapIt; ++MapIt)
-			{
-				const uint8* MapValuePtr = MapHelper.GetPairPtr(MapIt);
-				FindProperty(MapProperty->KeyProp, MapValuePtr, FindString, Context, bFound);
-				FindProperty(MapProperty->ValueProp, MapValuePtr, FindString, Context, bFound);
-			}
-		}
-	}
-}
-
-
 void FCustomizableObjectEditor::OnPostCompile()
 {
 	Viewport->CreatePreviewActor(PreviewInstance);
@@ -1786,23 +1712,6 @@ void FCustomizableObjectEditor::OnPostCompile()
 }
 
 
-void FCustomizableObjectEditor::LogSearchResult(const UObject& Context, const FString& Type, bool bIsFirst, const FString& Result) const
-{
-	if (!bIsFirst)
-	{
-		FCustomizableObjectEditorLogger::CreateLog(LOCTEXT("SearchResults", "Search Results:"))
-		.Notification(false)
-		.Log();
-	}
-	
-	FCustomizableObjectEditorLogger::CreateLog(FText::FromString(Type + ": " + Result))
-	.Context(Context)
-	.BaseObject()
-	.Notification(false)
-	.Log();
-}
-
-
 void FCustomizableObjectEditor::OnUpdatePreviewInstance(UCustomizableObjectInstance* Instance)
 {
 	if (TextureAnalyzer.IsValid())
@@ -1939,7 +1848,6 @@ void FCustomizableObjectEditor::OnCustomizableObjectStatusChanged(FCustomizableO
 			CompileRequest->bSkipIfNotOutOfDate = true;
 			CompileRequest->bSilentCompilation = false;
 			CompileRequest->SetDerivedDataCachePolicy(GetDerivedDataCachePolicyForEditor());
-			CompileRequest->EditorCustomLogListPtr = COEditorLogList;
 
 			ICustomizableObjectEditorModulePrivate::GetChecked().EnqueueCompileRequest(CompileRequest);
 		}
@@ -1953,38 +1861,60 @@ void FCustomizableObjectEditor::OnChangeDebugPlatform(TSharedPtr<FString> NewSel
 }
 
 
-//---------------------------------------------------------------
-
-FCustomizableObjectCustomLogList::FCustomizableObjectCustomLogList(TSharedPtr<FAssetEditorToolkit> InEditorPtr, TSharedPtr<SWidget>& StatsWidget)
+void FCustomizableObjectEditor::InitLogList()
 {
-	EditorPtr = InEditorPtr;
-
 	FMessageLogModule& MessageLogModule = FModuleManager::LoadModuleChecked<FMessageLogModule>("MessageLog");
-	FMessageLogInitializationOptions LogOptions;
-	// Show Pages so that user is never allowed to clear log messages
-	LogOptions.bShowPages = false;
-	LogOptions.bShowFilters = false;
-	LogOptions.bAllowClear = false;
-	LogOptions.MaxPageCount = 1;
+	TSharedPtr<FEditorLogger> EditorLogger;
+
+	if (!CustomizableObject)
+	{
+		return;
+	}
 
-	MessageLogListing = MessageLogModule.CreateLogListing("MutableEditorLog", LogOptions);
-	check(MessageLogListing);
+	if (!CustomizableObject->GetPrivate()->CompilationLogger)
+	{
+		FMessageLogInitializationOptions LogOptions;
+		LogOptions.bShowPages = false;
+		LogOptions.bShowFilters = false;
+		LogOptions.bAllowClear = false;
+		LogOptions.MaxPageCount = 1;
+
+		EditorLogger = MakeShared<FEditorLogger>();
+		EditorLogger->MessageLogListing = MessageLogModule.CreateLogListing("MutableEditor", LogOptions);
+
+		CustomizableObject->GetPrivate()->CompilationLogger = EditorLogger;
+	}
+	else
+	{
+		EditorLogger = StaticCastSharedPtr<FEditorLogger>(CustomizableObject->GetPrivate()->CompilationLogger);
+	}
 
-	StatsWidget = MessageLogModule.CreateLogListingWidget(MessageLogListing.ToSharedRef());
+	// Set the CO's LogListing as the point to get the messages of the States Tab
+	StatsWidget = MessageLogModule.CreateLogListingWidget(EditorLogger->MessageLogListing.ToSharedRef());
 }
 
 
-void FCustomizableObjectCustomLogList::DisplayStatesTab()
-{
-	TSharedPtr<FAssetEditorToolkit> Editor = EditorPtr.Pin();
+// ---- FEditorLogger class ---
 
-	if (!Editor || !Editor->GetTabManager())
+bool FEditorLogger::LogMessage(TSharedRef<FTokenizedMessage> Message)
+{
+	if (MessageLogListing)
 	{
-		return;
+		MessageLogListing->AddMessage(Message);
+		
+		return true;
 	}
 
-	Editor->GetTabManager()->TryInvokeTab(FCustomizableObjectEditor::StatsTabId);
+	return false;
 }
 
 
+void FEditorLogger::ClearLogMessageList()
+{
+	if (MessageLogListing)
+	{
+		MessageLogListing->ClearMessages();
+	}
+}
+
 #undef LOCTEXT_NAMESPACE
diff --git a/Engine/Plugins/Mutable/Source/CustomizableObjectEditor/Private/MuCOE/CustomizableObjectEditor.h b/Engine/Plugins/Mutable/Source/CustomizableObjectEditor/Private/MuCOE/CustomizableObjectEditor.h
--- a/Engine/Plugins/Mutable/Source/CustomizableObjectEditor/Private/MuCOE/CustomizableObjectEditor.h
+++ b/Engine/Plugins/Mutable/Source/CustomizableObjectEditor/Private/MuCOE/CustomizableObjectEditor.h
@@ -7,11 +7,10 @@
 #include "GraphEditor.h"
 #include "Misc/NotifyHook.h"
 #include "MuCO/CustomizableObjectPrivate.h"
+#include "MuR/MutableEditorLogger.h"
 #include "MuT/Node.h"
 #include "Widgets/Input/SNumericDropDown.h"
 
-#include "IMessageLogListing.h"
-
 #include "CustomizableObjectEditor.generated.h"
 
 struct FMutableSourceTextureData;
@@ -24,6 +23,7 @@ class FCustomizableObjectEditorViewportClient;
 class FProperty;
 class FSpawnTabArgs;
 class FUICommandList;
+class IMessageLogListing;
 class IToolkitHost;
 class SCustomizableObjectEditorViewportTabBody;
 class SDockTab;
@@ -52,6 +52,22 @@ DECLARE_DELEGATE(FCreatePreviewInstanceFlagDelegate);
 extern void RemoveRestrictedChars(FString& String);
 
 
+class FEditorLogger : public IMutableEditorLogger
+{
+public:
+
+	/** Logs message into the Message Log Listing */
+	bool LogMessage(TSharedRef<FTokenizedMessage> Message) override;
+
+	/** Claer all the messages of the Message Log Listing */
+	void ClearLogMessageList() override;
+
+public:
+
+	TSharedPtr<IMessageLogListing> MessageLogListing;
+};
+
+
 enum class EGizmoType : uint8
 {
 	Hidden,
@@ -80,30 +96,6 @@ public:
 	FCreatePreviewInstanceFlagDelegate Delegate;
 };
 
-/** 
-* Class that adds a stats tab to the CustomizableObjectEditor
-* Also allows the managment of this tab
-*/
-class FCustomizableObjectCustomLogList
-{
-public:
-
-	FCustomizableObjectCustomLogList(TSharedPtr<FAssetEditorToolkit> InEditorPtr, TSharedPtr<SWidget>& StatsWidget);
-
-	/** Tries to reopen the States tab if it's been closed (and needs to be reopened) to show a compilation error or warning */
-	void DisplayStatesTab();
-
-public:
-
-	TSharedPtr<IMessageLogListing> MessageLogListing;
-
-private:
-
-	/** Pointer to the editor that contains this Stats */
-	TWeakPtr<FAssetEditorToolkit> EditorPtr;
-
-};
-
 
 /**
  * CustomizableObject Editor class
@@ -209,7 +201,9 @@ public:
 
 	/** Clear game asset references saved in the Customizable Object. */
 	void ClearGatheredReferences();
-	
+
+	void InitLogList();
+
 private:
 	TSharedRef<SDockTab> SpawnTab_Viewport(const FSpawnTabArgs& Args);
 	TSharedRef<SDockTab> SpawnTab_Details(const FSpawnTabArgs& Args);
@@ -257,28 +251,12 @@ private:
 	/** Callback for the object modified event */
 	void OnObjectModified(UObject* Object);
 
-	/** Logs the search results of the search
-	 * @param Context The UObject we have found to be related with the searched string.
-	 * @param Type The type of relation with the searched word. It is a node, a value or maybe a variable?
-	 * @param bIsFirst Is this the first time we encountered something during our search?
-	 * @param Result The string containing the search word we are looking for in Node
-	 */
-	void LogSearchResult(const UObject& Context, const FString& Type, bool bIsFirst, const FString& Result) const;
-
 	/** Open the Texture Analyzer tab */
 	void OpenTextureAnalyzerTab();
 
 	/** Open the Performance Analyzer tab */
 	void OpenPerformanceAnalyzerTab();
 
-	/** Recursively find any property that its name or value contains the given string.
-	  * @param Property Root property.
-	  * @param Container Root property container (address of the property value).
-	  * @param FindString String to find for.
-	  * @param Context UObject Context where this string has been found.
-	  * @param bFound Mark as true if any property has been found. */
-	void FindProperty(const FProperty* Property, const void* Container, const FString& FindString, const UObject& Context, bool& bFound);
-
 	void OnPostCompile();
 
 	void OnChangeDebugPlatform(TSharedPtr<FString> NewSelection, ESelectInfo::Type SelectInfo);
@@ -363,8 +341,6 @@ private:
 	TSharedPtr<FString> DebugTargetPlatformName;
 	
 	bool bDebugForceLargeLODBias = false;
-
-	TSharedPtr<FCustomizableObjectCustomLogList> COEditorLogList;
 	
 	TSharedPtr<SWidget> StatsWidget;
 
diff --git a/Engine/Plugins/Mutable/Source/CustomizableObjectEditor/Private/MuCOE/CustomizableObjectEditorLogger.cpp b/Engine/Plugins/Mutable/Source/CustomizableObjectEditor/Private/MuCOE/CustomizableObjectEditorLogger.cpp
--- a/Engine/Plugins/Mutable/Source/CustomizableObjectEditor/Private/MuCOE/CustomizableObjectEditorLogger.cpp
+++ b/Engine/Plugins/Mutable/Source/CustomizableObjectEditor/Private/MuCOE/CustomizableObjectEditorLogger.cpp
@@ -12,12 +12,10 @@
 #include "MuCOE/CustomizableObjectMacroLibrary/CustomizableObjectMacroLibraryEditor.h"
 #include "MuCOE/CustomizableObjectMacroLibrary/CustomizableObjectMacroLibrary.h"
 #include "MuCOE/Nodes/CustomizableObjectNode.h"
+#include "MuR/MutableEditorLogger.h"
 #include "Toolkits/ToolkitManager.h"
 #include "Widgets/Notifications/SNotificationList.h"
 
-#include "IMessageLogListing.h"
-
-
 #define LOCTEXT_NAMESPACE "CustomizableObject"
 
 
@@ -230,9 +228,16 @@ FLogParameters& FLogParameters::SpamBin(const ELoggerSpamBin SpamBin)
 }
 
 
-void FLogParameters::Log(TSharedPtr<FCustomizableObjectCustomLogList> MessageLogListing)
+FLogParameters& FLogParameters::CustomizableObject(UCustomizableObject* Object)
 {
-	Logger.Log(*this, MessageLogListing);
+	ParamCustomizableObject = Object;
+	return *this;
+}
+
+
+void FLogParameters::Log()
+{
+	Logger.Log(*this);
 }
 
 
@@ -249,7 +254,7 @@ FLogParameters FCustomizableObjectEditorLogger::CreateLog(FText&& Text)
 }
 
 
-void FCustomizableObjectEditorLogger::Log(FLogParameters& LogParameters, TSharedPtr<FCustomizableObjectCustomLogList> MessageLogListing)
+void FCustomizableObjectEditorLogger::Log(FLogParameters& LogParameters)
 {
 	FMessageLog MessageLog(LOG_NAME);
 	const TSharedRef<FTokenizedMessage> Message = MessageLog.Message(LogParameters.ParamSeverity);
@@ -282,12 +287,15 @@ void FCustomizableObjectEditorLogger::Log(FLogParameters& LogParameters, TShared
 		}
 	}
 
-	if (MessageLogListing)
-	{
-		MessageLogListing->MessageLogListing->AddMessage(Message);
+	Message->SetIndentationLevel(LogParameters.ParamSeverity == EMessageSeverity::Info ? 0 : 1);
 
-		// Since we are compiling from an editor, we do not need to send any popup notification 
-		return;
+	if (UCustomizableObject* CustomizableObject = LogParameters.ParamCustomizableObject.Get())
+	{
+		if (CustomizableObject->GetPrivate()->CompilationLogger && CustomizableObject->GetPrivate()->CompilationLogger->LogMessage(Message))
+		{
+			// Since we are compiling from an editor, we do not need to send any popup notification 
+			return;
+		}
 	}
 
 	if (LogParameters.bParamNotification)
diff --git a/Engine/Plugins/Mutable/Source/CustomizableObjectEditor/Private/MuCOE/CustomizableObjectEditorLogger.h b/Engine/Plugins/Mutable/Source/CustomizableObjectEditor/Private/MuCOE/CustomizableObjectEditorLogger.h
--- a/Engine/Plugins/Mutable/Source/CustomizableObjectEditor/Private/MuCOE/CustomizableObjectEditorLogger.h
+++ b/Engine/Plugins/Mutable/Source/CustomizableObjectEditor/Private/MuCOE/CustomizableObjectEditorLogger.h
@@ -6,6 +6,7 @@
 #include "MuCOE/CustomizableObjectEditor.h"
 
 class UCustomizableObjectNode;
+class UCustomizableObject;
 
 class FCustomizableObjectEditorLogger;
 class FName;
@@ -56,9 +57,11 @@ public:
 	FLogParameters& FixNotification(bool FixNotification = true);
 
 	FLogParameters& SpamBin(ELoggerSpamBin SpamBin);
+
+	FLogParameters& CustomizableObject(UCustomizableObject* Object = nullptr);
 	
 	/** Actually display the log. */
-	void Log(TSharedPtr<FCustomizableObjectCustomLogList> MessageLogListing = nullptr);
+	void Log();
 	
 private:
 	/** Constructor. Contains all required parameters. */
@@ -88,6 +91,8 @@ private:
 	bool bParamNotification = true;
 
 	bool bParamFixNotification = false;
+
+	TWeakObjectPtr<UCustomizableObject> ParamCustomizableObject;
 };
 
 
@@ -103,7 +108,7 @@ public:
 	static FLogParameters CreateLog(FText&& Text);
 
 	/** Display the log. Equivalent to FLogParameters::Log(). */
-	void Log(FLogParameters& LogParameters, TSharedPtr<FCustomizableObjectCustomLogList> MessageLogListing = nullptr);
+	void Log(FLogParameters& LogParameters);
 
 	/** Dismiss a fixed notification. */
 	static void DismissNotification(ELoggerCategory Category);
diff --git a/Engine/Plugins/Mutable/Source/CustomizableObjectEditor/Private/MuCOE/CustomizableObjectMacroLibrary/CustomizableObjectGraphEditorToolkit.cpp b/Engine/Plugins/Mutable/Source/CustomizableObjectEditor/Private/MuCOE/CustomizableObjectMacroLibrary/CustomizableObjectGraphEditorToolkit.cpp
--- a/Engine/Plugins/Mutable/Source/CustomizableObjectEditor/Private/MuCOE/CustomizableObjectMacroLibrary/CustomizableObjectGraphEditorToolkit.cpp
+++ b/Engine/Plugins/Mutable/Source/CustomizableObjectEditor/Private/MuCOE/CustomizableObjectMacroLibrary/CustomizableObjectGraphEditorToolkit.cpp
@@ -394,8 +394,8 @@ void FCustomizableObjectGraphEditorToolkit::OnEnterText(const FText& NewText, ET
 	}
 
 	bool bFound = false;
-
 	const FString FindString = NewText.ToString();
+	UCustomizableObject* CustomizableObject = Cast<UCustomizableObject>(Graph->GetOutermostObject()); // Null if this is a Macro
 
 	for (TObjectPtr<UEdGraphNode> Node : Graph->Nodes)
 	{
@@ -408,8 +408,7 @@ void FCustomizableObjectGraphEditorToolkit::OnEnterText(const FText& NewText, ET
 		const FString NodeName = Node->GetNodeTitle(ENodeTitleType::FullTitle).ToString().Replace(TEXT("\n"), TEXT(" "));
 		if (NodeName.Contains(NewText.ToString(), ESearchCase::IgnoreCase))
 		{
-			LogSearchResult(*Node, "Node", bFound, NodeName);
-			bFound = true;
+			LogSearchResult(*Node, "Node", bFound, NodeName, CustomizableObject);
 		}
 
 		// Pins are not in the reflection system
@@ -418,15 +417,14 @@ void FCustomizableObjectGraphEditorToolkit::OnEnterText(const FText& NewText, ET
 			const FString PinFriendlyName = Pin->PinFriendlyName.ToString();
 			if (PinFriendlyName.Contains(FindString))
 			{
-				LogSearchResult(*Node, "Pin", bFound, PinFriendlyName);
-				bFound = true;
+				LogSearchResult(*Node, "Pin", bFound, PinFriendlyName, CustomizableObject);
 			}
 		}
 
 		// Find anything marked as a UPROPERTY
 		for (TFieldIterator<FProperty> It(Node->GetClass()); It; ++It)
 		{
-			FindProperty(*It, Node, FindString, *Node, bFound);
+			FindProperty(*It, Node, FindString, *Node, bFound, CustomizableObject);
 		}
 	}
 
@@ -437,11 +435,12 @@ void FCustomizableObjectGraphEditorToolkit::OnEnterText(const FText& NewText, ET
 	FCustomizableObjectEditorLogger::CreateLog(Text)
 		.Category(ELoggerCategory::GraphSearch)
 		.CustomNotification()
+		.CustomizableObject(CustomizableObject)
 		.Log();
 }
 
 
-void FCustomizableObjectGraphEditorToolkit::FindProperty(const FProperty* Property, const void* InContainer, const FString& FindString, const UObject& Context, bool& bFound)
+void FCustomizableObjectGraphEditorToolkit::FindProperty(const FProperty* Property, const void* InContainer, const FString& FindString, const UObject& Context, bool& bFound, UCustomizableObject* Object)
 {
 	if (!Property || !InContainer)
 	{
@@ -451,8 +450,7 @@ void FCustomizableObjectGraphEditorToolkit::FindProperty(const FProperty* Proper
 	const FString PropertyName = Property->GetDisplayNameText().ToString();
 	if (PropertyName.Contains(FindString))
 	{
-		LogSearchResult(Context, "Property Name", bFound, *PropertyName);
-		bFound = true;
+		LogSearchResult(Context, "Property Name", bFound, *PropertyName, Object);
 	}
 
 	for (int32 Index = 0; Index < Property->ArrayDim; ++Index)
@@ -464,8 +462,7 @@ void FCustomizableObjectGraphEditorToolkit::FindProperty(const FProperty* Proper
 			const FString* StringResult = StringProperty->GetPropertyValuePtr(ValuePtr);
 			if (StringResult->Contains(FindString))
 			{
-				LogSearchResult(Context, "Property Value", bFound, *StringResult);
-				bFound = true;
+				LogSearchResult(Context, "Property Value", bFound, *StringResult, Object);
 			}
 		}
 		else if (const FEnumProperty* EnumProperty = CastField<FEnumProperty>(Property))
@@ -475,8 +472,7 @@ void FCustomizableObjectGraphEditorToolkit::FindProperty(const FProperty* Proper
 			const FString StringResult = EnumResult->GetDisplayNameTextByIndex(*ValuePtr).ToString();
 			if (StringResult.Contains(FindString))
 			{
-				LogSearchResult(Context, "Property Value", bFound, StringResult);
-				bFound = true;
+				LogSearchResult(Context, "Property Value", bFound, StringResult, Object);
 			}
 		}
 		else if (const FSoftObjectProperty* SoftObjectProperty = CastField<FSoftObjectProperty>(Property))
@@ -484,8 +480,7 @@ void FCustomizableObjectGraphEditorToolkit::FindProperty(const FProperty* Proper
 			const FString ObjectPath = SoftObjectProperty->GetPropertyValuePtr(ValuePtr)->ToString();
 			if (ObjectPath.Contains(FindString))
 			{
-				LogSearchResult(Context, "Property Value", bFound, ObjectPath);
-				bFound = true;
+				LogSearchResult(Context, "Property Value", bFound, ObjectPath, Object);
 			}
 		}
 		else if (const FObjectPropertyBase* ObjectProperty = CastField<FObjectPropertyBase>(Property))
@@ -496,8 +491,7 @@ void FCustomizableObjectGraphEditorToolkit::FindProperty(const FProperty* Proper
 
 				if (ObjectValue->GetName().Contains(FindString))
 				{
-					LogSearchResult(Context, "Property Value", bFound, Name);
-					bFound = true;
+					LogSearchResult(Context, "Property Value", bFound, Name, Object);
 				}
 			}
 		}
@@ -505,7 +499,7 @@ void FCustomizableObjectGraphEditorToolkit::FindProperty(const FProperty* Proper
 		{
 			for (TFieldIterator<FProperty> It(StructProperty->Struct); It; ++It)
 			{
-				FindProperty(*It, ValuePtr, FindString, Context, bFound);
+				FindProperty(*It, ValuePtr, FindString, Context, bFound, Object);
 			}
 		}
 		else if (const FArrayProperty* ArrayProperty = CastField<FArrayProperty>(Property))
@@ -513,7 +507,7 @@ void FCustomizableObjectGraphEditorToolkit::FindProperty(const FProperty* Proper
 			FScriptArrayHelper ArrayHelper(ArrayProperty, ValuePtr);
 			for (int32 ValueIdx = 0; ValueIdx < ArrayHelper.Num(); ++ValueIdx)
 			{
-				FindProperty(ArrayProperty->Inner, ArrayHelper.GetRawPtr(ValueIdx), FindString, Context, bFound);
+				FindProperty(ArrayProperty->Inner, ArrayHelper.GetRawPtr(ValueIdx), FindString, Context, bFound, Object);
 			}
 		}
 		else if (const FSetProperty* SetProperty = CastField<FSetProperty>(Property))
@@ -521,7 +515,7 @@ void FCustomizableObjectGraphEditorToolkit::FindProperty(const FProperty* Proper
 			FScriptSetHelper SetHelper(SetProperty, ValuePtr);
 			for (FScriptSetHelper::FIterator SetIt = SetHelper.CreateIterator(); SetIt; ++SetIt)
 			{
-				FindProperty(SetProperty->ElementProp, SetHelper.GetElementPtr(SetIt), FindString, Context, bFound);
+				FindProperty(SetProperty->ElementProp, SetHelper.GetElementPtr(SetIt), FindString, Context, bFound, Object);
 			}
 		}
 		else if (const FMapProperty* MapProperty = CastField<FMapProperty>(Property))
@@ -530,20 +524,28 @@ void FCustomizableObjectGraphEditorToolkit::FindProperty(const FProperty* Proper
 			for (FScriptMapHelper::FIterator MapIt = MapHelper.CreateIterator(); MapIt; ++MapIt)
 			{
 				const uint8* MapValuePtr = MapHelper.GetPairPtr(MapIt);
-				FindProperty(MapProperty->KeyProp, MapValuePtr, FindString, Context, bFound);
-				FindProperty(MapProperty->ValueProp, MapValuePtr, FindString, Context, bFound);
+				FindProperty(MapProperty->KeyProp, MapValuePtr, FindString, Context, bFound, Object);
+				FindProperty(MapProperty->ValueProp, MapValuePtr, FindString, Context, bFound, Object);
 			}
 		}
 	}
 }
 
 
-void FCustomizableObjectGraphEditorToolkit::LogSearchResult(const UObject& Context, const FString& Type, bool bIsFirst, const FString& Result) const
+void FCustomizableObjectGraphEditorToolkit::LogSearchResult(const UObject& Context, const FString& Type, bool& bIsFirst, const FString& Result, UCustomizableObject* Object) const
 {
 	if (!bIsFirst)
 	{
+		bIsFirst = true;
+
+		if (Object && Object->GetPrivate()->CompilationLogger)
+		{
+			Object->GetPrivate()->CompilationLogger->ClearLogMessageList();
+		}
+
 		FCustomizableObjectEditorLogger::CreateLog(LOCTEXT("SearchResults", "Search Results:"))
 			.Notification(false)
+			.CustomizableObject(Object)
 			.Log();
 	}
 
@@ -551,6 +553,7 @@ void FCustomizableObjectGraphEditorToolkit::LogSearchResult(const UObject& Conte
 		.Context(Context)
 		.BaseObject()
 		.Notification(false)
+		.CustomizableObject(Object)
 		.Log();
 }
 
diff --git a/Engine/Plugins/Mutable/Source/CustomizableObjectEditor/Private/MuCOE/CustomizableObjectMacroLibrary/CustomizableObjectGraphEditorToolkit.h b/Engine/Plugins/Mutable/Source/CustomizableObjectEditor/Private/MuCOE/CustomizableObjectMacroLibrary/CustomizableObjectGraphEditorToolkit.h
--- a/Engine/Plugins/Mutable/Source/CustomizableObjectEditor/Private/MuCOE/CustomizableObjectMacroLibrary/CustomizableObjectGraphEditorToolkit.h
+++ b/Engine/Plugins/Mutable/Source/CustomizableObjectEditor/Private/MuCOE/CustomizableObjectMacroLibrary/CustomizableObjectGraphEditorToolkit.h
@@ -14,6 +14,7 @@
 #include "Toolkits/AssetEditorToolkit.h"
 
 class UClass;
+class UCustomizableObject;
 class UCustomizableObjectNode;
 
 // Public interface to Customizable Object Editor
@@ -29,8 +30,20 @@ public:
 
 	/** Searches a node that contains the inserted word */
 	void OnEnterText(const FText& NewText, ETextCommit::Type TextType);
-	void FindProperty(const FProperty* Property, const void* InContainer, const FString& FindString, const UObject& Context, bool& bFound);
-	void LogSearchResult(const UObject& Context, const FString& Type, bool bIsFirst, const FString& Result) const;
+	/** Recursively find any property that its name or value contains the given string.
+	  * @param Property Root property.
+	  * @param Container Root property container (address of the property value).
+	  * @param FindString String to find for.
+	  * @param Context UObject Context where this string has been found.
+	  * @param bFound Mark as true if any property has been found. */
+	void FindProperty(const FProperty* Property, const void* InContainer, const FString& FindString, const UObject& Context, bool& bFound, UCustomizableObject* Object);
+	/** Logs the search results of the search
+	 * @param Context The UObject we have found to be related with the searched string.
+	 * @param Type The type of relation with the searched word. It is a node, a value or maybe a variable?
+	 * @param bIsFirst Is this the first time we encountered something during our search?
+	 * @param Result The string containing the search word we are looking for in Node
+	 */
+	void LogSearchResult(const UObject& Context, const FString& Type, bool& bIsFirst, const FString& Result, UCustomizableObject* Object) const;
 
 	/** Bind common graph commands */
 	void BindGraphCommands();
diff --git a/Engine/Plugins/Mutable/Source/MutableRuntime/Internal/MuR/MutableEditorLogger.h b/Engine/Plugins/Mutable/Source/MutableRuntime/Internal/MuR/MutableEditorLogger.h
new file mode 100644
--- /dev/null
+++ b/Engine/Plugins/Mutable/Source/MutableRuntime/Internal/MuR/MutableEditorLogger.h
@@ -0,0 +1,20 @@
+// Copyright Epic Games, Inc. All Rights Reserved.
+
+#pragma once
+
+#include "Logging/TokenizedMessage.h"
+
+/** Class used to log messages from CO compilations and COI updates to the open editor. */
+class IMutableEditorLogger
+{
+public:
+
+	virtual ~IMutableEditorLogger() = default;
+
+	/** Logs the Message */
+	virtual bool LogMessage(TSharedRef<FTokenizedMessage> Message) = 0;
+
+	/** Clears all the visible Logs */
+	virtual void ClearLogMessageList() = 0;
+};
+
diff --git a/Engine/Plugins/Mutable/Source/MutableRuntime/Internal/MuR/System.h b/Engine/Plugins/Mutable/Source/MutableRuntime/Internal/MuR/System.h
--- a/Engine/Plugins/Mutable/Source/MutableRuntime/Internal/MuR/System.h
+++ b/Engine/Plugins/Mutable/Source/MutableRuntime/Internal/MuR/System.h
@@ -16,6 +16,7 @@
 #include "MuR/Parameters.h"
 #include "MuR/MutableTrace.h"
 #include "MuR/ManagedPointer.h"
+#include "MuR/MutableEditorLogger.h"
 
 #include "SkeletalMesh.h"
 #include "Tasks/Task.h"
@@ -1599,6 +1600,25 @@ namespace UE::Mutable::Private
 		return a.Type < b.Type;
 	}
 
+	/** Class to log update messages */
+	class FLiveInstanceLogger
+	{
+	public:
+		
+		/** Set the logger class to print the messages to an editor widget */
+		UE_API void SetLogger(TSharedPtr<IMutableEditorLogger> InInstanceLogger);
+
+		/** Prints the message into the Logs console/file and also prints it into an editor widget.
+			Use this function instead of UE_LOG to print important messages.
+		*/
+		void LogUpdateMessage(const FString& Message, ELogVerbosity::Type Verbosity = ELogVerbosity::Error);
+
+	private:
+
+		/** Pointer to the editor log list to print all the messages related to the instance update. */
+		TSharedPtr<IMutableEditorLogger> InstanceLogger;
+	};
+
 	/** Data for an instance that is currently being processed in the mutable system. This means it is
 	* between a BeginUpdate and EndUpdate, or during an "atomic" operation (like generate a single resource).
 	*/
@@ -1621,6 +1641,8 @@ namespace UE::Mutable::Private
 		
 		TSharedPtr<FExternalResourceProvider> ExternalResourceProvider;
 		TSharedPtr<FPassthroughObjectLoader> PassthroughObjectLoader;
+		
+		TSharedPtr<FLiveInstanceLogger> UpdateLogger;
 	};
 
 
@@ -1765,7 +1787,8 @@ namespace UE::Mutable::Private
 		const TSharedPtr<UE::Mutable::Private::FMaterialIdRegistry>& InMaterialIdRegistry,
 		const TSharedPtr<UE::Mutable::Private::FSkeletalMeshIdRegistry>& InSkeletalMeshIdRegistry,
         	const TSharedPtr<FExternalResourceProvider>& ExternalResourceProvider,
-        	const TSharedPtr<FPassthroughObjectLoader>& PassthroughObjectLoader);
+        	const TSharedPtr<FPassthroughObjectLoader>& PassthroughObjectLoader,
+			const TSharedPtr<FLiveInstanceLogger>& LiveInstanceLogger);
 
 		/** 
 		 * Create a ready to use new FLiveInstance to be used during the mutable build operations. It will not be taken into consideration for any update operation
diff --git a/Engine/Plugins/Mutable/Source/MutableRuntime/Private/MuR/CodeRunner.cpp b/Engine/Plugins/Mutable/Source/MutableRuntime/Private/MuR/CodeRunner.cpp
--- a/Engine/Plugins/Mutable/Source/MutableRuntime/Private/MuR/CodeRunner.cpp
+++ b/Engine/Plugins/Mutable/Source/MutableRuntime/Private/MuR/CodeRunner.cpp
@@ -2562,7 +2562,7 @@ namespace UE::Mutable::Private
 				TManagedPtr<FMesh> ReshapedMeshResult = CreateMesh(BaseMesh->GetDataSize());
 
 				bool bOutSuccess = false;
-				MeshApplyShape(ReshapedMeshResult.Get(), BaseMesh.Get(), Shape.Get(), ReshapeFlags, bOutSuccess);
+				MeshApplyShape(ReshapedMeshResult.Get(), BaseMesh.Get(), Shape.Get(), ReshapeFlags, bOutSuccess, LiveInstance->UpdateLogger.ToSharedRef());
 
 				Shape = nullptr;
 
@@ -5742,7 +5742,8 @@ namespace UE::Mutable::Private
 				
 				if(!bParameterFound)
 				{
-					UE_LOG(LogMutableCore, Error, TEXT("Expected Texture parameter [%s] not found in Break Material node."), *ParameterName.ToString());
+					LiveInstance->UpdateLogger->LogUpdateMessage(FString::Printf(TEXT("Expected Texture parameter [%s] not found in Break Material node."), *ParameterName.ToString()));
+
 					AbortRun();
 				}
 
@@ -6257,7 +6258,7 @@ namespace UE::Mutable::Private
 
 				if(!bParameterFound)
 				{
-					UE_LOG(LogMutableCore, Error, TEXT("Expected Scalar parameter [%s] not found in Break Material node."), *ParameterName.ToString());
+					LiveInstance->UpdateLogger->LogUpdateMessage(FString::Printf(TEXT("Expected Scalar parameter [%s] not found in Break Material node."), *ParameterName.ToString()));
 					AbortRun();
 				}
 
@@ -6595,7 +6596,7 @@ namespace UE::Mutable::Private
 
 				if(!bParameterFound)
 				{
-					UE_LOG(LogMutableCore, Error, TEXT("Expected Color parameter [%s] not found in Break Material node."), *ParameterName.ToString());
+					LiveInstance->UpdateLogger->LogUpdateMessage(FString::Printf(TEXT("Expected Color parameter [%s] not found in Break Material node."), *ParameterName.ToString()));
 					AbortRun();
 				}
 
@@ -7232,7 +7233,7 @@ namespace UE::Mutable::Private
 				const FInstancedStruct& OperationInstancedStruct = Program.ExternalOperationProvider->Get(Item.At);
 				if (!OperationInstancedStruct.IsValid()) // Plugin not loaded
 				{
-                    UE_LOG(LogMutableCore, Error, TEXT("Trying to evaluate an External Operation which has not been loaded."));
+					LiveInstance->UpdateLogger->LogUpdateMessage(FString::Printf(TEXT("Trying to evaluate an External Operation which has not been loaded.")));
 					AbortRun();
 
 					break;
@@ -7243,7 +7244,7 @@ namespace UE::Mutable::Private
 				const TArray<TPair<FName, const UScriptStruct*>>& Inputs = Operation->GetInputs();
 				if (Inputs.Num() != Args.NumOperants)
 				{
-					UE_LOG(LogMutableCore, Error, TEXT("External Operation inputs do not match compiled data: %s"), *OperationInstancedStruct.GetScriptStruct()->GetName());
+					LiveInstance->UpdateLogger->LogUpdateMessage(FString::Printf(TEXT("External Operation inputs do not match compiled data: %s"), *OperationInstancedStruct.GetScriptStruct()->GetName()));
 					AbortRun();
 
 					return;
@@ -7368,7 +7369,7 @@ namespace UE::Mutable::Private
 
 						if (!Context.Output.Ptr)
 						{
-							UE_LOG(LogMutableCore, Error, TEXT("External Operation output not set: %s"), *OperationInstancedStruct.GetScriptStruct()->GetName());
+							LiveInstance->UpdateLogger->LogUpdateMessage(FString::Printf(TEXT("External Operation output not set: %s"), *OperationInstancedStruct.GetScriptStruct()->GetName()));
 							AbortRun();
 
 							return;
diff --git a/Engine/Plugins/Mutable/Source/MutableRuntime/Private/MuR/OpMeshReshape.h b/Engine/Plugins/Mutable/Source/MutableRuntime/Private/MuR/OpMeshReshape.h
--- a/Engine/Plugins/Mutable/Source/MutableRuntime/Private/MuR/OpMeshReshape.h
+++ b/Engine/Plugins/Mutable/Source/MutableRuntime/Private/MuR/OpMeshReshape.h
@@ -6,6 +6,7 @@
 #include "MuR/MeshPrivate.h"
 #include "MuR/Operations.h"
 #include "MuR/OpMeshBind.h"
+#include "MuR/System.h"
 
 #include "Math/Plane.h"
 #include "Math/Ray.h"
@@ -377,7 +378,7 @@ namespace UE::Mutable::Private
 		OutL = FMath::Max(0.0f, (Centroids[0] - Centroids[1]).Length() - (OutR0 + OutR1));
 	}
 
-	inline void ApplyToVertices(FMesh* Mesh, TArrayView<const FReshapeVertexBindingData> BindingData, const FShapeMeshDescriptorApply& Shape, bool bSkipNormalReshape)
+	inline void ApplyToVertices(FMesh* Mesh, TArrayView<const FReshapeVertexBindingData> BindingData, const FShapeMeshDescriptorApply& Shape, bool bSkipNormalReshape, const TSharedRef<FLiveInstanceLogger> MessageLogger)
 	{
 		check(Mesh);
 		check(Mesh->GetVertexCount() == BindingData.Num());
@@ -407,7 +408,7 @@ namespace UE::Mutable::Private
 
 		if (bTriangleOutOfScopeFound)
 		{
-			UE_LOG(LogMutableCore, Warning, TEXT("Performing a Mesh Reshape where base shape and target shape do not have the same number of triangles."));
+			MessageLogger->LogUpdateMessage("Performing a Mesh Reshape where base shape and target shape do not have the same number of triangles.", ELogVerbosity::Warning);
 		}	
 #endif
 
@@ -569,7 +570,8 @@ namespace UE::Mutable::Private
 			const FSkeleton& BaseSkeleton, 
 			TArrayView<const FReshapePointBindingData> BindingData, 
 			TArrayView<const int32> UsedIndices, 
-			const FShapeMeshDescriptorApply& Shape)
+			const FShapeMeshDescriptorApply& Shape,
+			const TSharedRef<FLiveInstanceLogger> MessageLogger)
 	{
 #if DO_CHECK
 		// checking if the Base shape has more triangles than the target shape
@@ -578,7 +580,7 @@ namespace UE::Mutable::Private
 
 		if (bTriangleOutOfScopeFound)
 		{
-			UE_LOG(LogMutableCore, Warning, TEXT("Performing a Mesh Reshape where base shape and target shape do not have the same number of triangles."));
+			MessageLogger->LogUpdateMessage("Performing a Mesh Reshape where base shape and target shape do not have the same number of triangles.", ELogVerbosity::Warning);
 		}
 #endif
 		
@@ -712,7 +714,8 @@ namespace UE::Mutable::Private
 		TArrayView<const FReshapePointBindingData> BindingData,
 		TArrayView<const int32> UsedIndices,
 		TArrayView<const int32> BodyOffsets,
-		const FShapeMeshDescriptorApply& Shape)
+		const FShapeMeshDescriptorApply& Shape,
+		const TSharedRef<FLiveInstanceLogger> MessageLogger)
 	{
 		if (!BodyOffsets.Num())
 		{
@@ -725,7 +728,7 @@ namespace UE::Mutable::Private
 		{
 			TArrayView<const int32> BodyUsedIndices(UsedIndices.GetData() + IndicesBegin, IndicesEnd - IndicesBegin);
 
-			ApplyToPhysicsBodies(OutBody, NumProcessedBindPoints, BaseMesh, BaseSkeleton, BindingData, BodyUsedIndices, Shape);
+			ApplyToPhysicsBodies(OutBody, NumProcessedBindPoints, BaseMesh, BaseSkeleton, BindingData, BodyUsedIndices, Shape, MessageLogger);
 		};
 
 
@@ -757,7 +760,7 @@ namespace UE::Mutable::Private
 	//! Rebuild the (previously bound) mesh data for a new shape.
 	//! Proof-of-concept implementation.
 	//---------------------------------------------------------------------------------------------
-	inline void MeshApplyShape(FMesh* Result, const FMesh* BaseMesh, const FMesh* ShapeMesh, EMeshBindShapeFlags BindFlags, bool& bOutSuccess)
+	inline void MeshApplyShape(FMesh* Result, const FMesh* BaseMesh, const FMesh* ShapeMesh, EMeshBindShapeFlags BindFlags, bool& bOutSuccess, const TSharedRef<FLiveInstanceLogger> MessageLogger)
 	{
 		MUTABLE_CPUPROFILER_SCOPE(MeshApplyReshape);
 
@@ -918,7 +921,7 @@ namespace UE::Mutable::Private
 						(const FReshapeVertexBindingData*)VB.GetBufferData(BarycentricDataBuffer),
 						VB.GetElementCount());
 
-				ApplyToVertices(VerticesReshapeMesh, VerticesBindingData, ShapeDescriptor, bRecomputeNormals);
+				ApplyToVertices(VerticesReshapeMesh, VerticesBindingData, ShapeDescriptor, bRecomputeNormals, MessageLogger);
 			}
 
 			if (bApplyLaplacian)
@@ -1026,7 +1029,7 @@ namespace UE::Mutable::Private
 						(const int32*)PhyiscsBindOffsetsBuffer.GetBufferData(0), 
 					    PhyiscsBindOffsetsBuffer.GetElementCount());
 
-				ApplyToAllPhysicsBodies(*Result, *BaseMesh, *BaseMesh->GetSkeleton(), BindingData, UsedIndices, Offsets, ShapeDescriptor);
+				ApplyToAllPhysicsBodies(*Result, *BaseMesh, *BaseMesh->GetSkeleton(), BindingData, UsedIndices, Offsets, ShapeDescriptor, MessageLogger);
 			}
 		}
 	}
diff --git a/Engine/Plugins/Mutable/Source/MutableRuntime/Private/MuR/System.cpp b/Engine/Plugins/Mutable/Source/MutableRuntime/Private/MuR/System.cpp
--- a/Engine/Plugins/Mutable/Source/MutableRuntime/Private/MuR/System.cpp
+++ b/Engine/Plugins/Mutable/Source/MutableRuntime/Private/MuR/System.cpp
@@ -128,12 +128,11 @@ namespace UE::Mutable::Private
 		const TSharedPtr<UE::Mutable::Private::FMaterialIdRegistry>& InMaterialIdRegistry,
 		const TSharedPtr<UE::Mutable::Private::FSkeletalMeshIdRegistry>& InSkeletalMeshIdRegistry,
     	const TSharedPtr<FExternalResourceProvider>& ExternalResourceProvider,
-    	const TSharedPtr<FPassthroughObjectLoader>& PassthroughObjectLoader)
+    	const TSharedPtr<FPassthroughObjectLoader>& PassthroughObjectLoader,
+		const TSharedPtr<FLiveInstanceLogger>& LiveInstanceLogger)
     {
 		LLM_SCOPE_BYNAME(TEXT("MutableRuntime"));
 		MUTABLE_CPUPROFILER_SCOPE(NewInstance);
-
-    	UE_LOG(LogMutableCore, Verbose, TEXT("Creating Mutable instance"));
     	check(InModel)
 
 		TSharedRef<FLiveInstance> InstanceData = MakeShared<FLiveInstance>();
@@ -146,11 +145,14 @@ namespace UE::Mutable::Private
     	InstanceData->ImageIdRegistry = InImageIdRegistry;
     	InstanceData->MaterialIdRegistry = InMaterialIdRegistry;
     	InstanceData->SkeletalMeshIdRegistry = InSkeletalMeshIdRegistry;
+    	InstanceData->UpdateLogger = LiveInstanceLogger;
     	
     	check(InstanceData->MeshIdRegistry.IsValid())
 		check(InstanceData->ImageIdRegistry.IsValid())
 		check(InstanceData->MaterialIdRegistry.IsValid())
 		check(InstanceData->SkeletalMeshIdRegistry.IsValid())
+
+		InstanceData->UpdateLogger->LogUpdateMessage(TEXT("Creating Mutable instance"), ELogVerbosity::Verbose);
     	
     	WorkingMemoryManager.LiveInstances.Add(InstanceData.ToWeakPtr());
     	WorkingMemoryManager.LiveInstances.RemoveAllSwap([](const TWeakPtr<FLiveInstance>& Instance)
@@ -182,6 +184,7 @@ namespace UE::Mutable::Private
     	InstanceData->ImageIdRegistry = MakeShared<FImageIdRegistry>();
     	InstanceData->MaterialIdRegistry = MakeShared<FMaterialIdRegistry>();
     	InstanceData->SkeletalMeshIdRegistry = MakeShared<FSkeletalMeshIdRegistry>();
+    	InstanceData->UpdateLogger = MakeShared<FLiveInstanceLogger>();
 
     	InstanceData->Cache->Init(InModel->GetProgram().OpAddress.Num());
     	
@@ -221,7 +224,7 @@ namespace UE::Mutable::Private
 
 		if (!LiveInstance->Parameters)
 		{
-			UE_LOG(LogMutableCore, Error, TEXT("Invalid parameters in mutable update."));
+			LiveInstance->UpdateLogger->LogUpdateMessage(TEXT("Invalid parameters in mutable update."));
 			return nullptr;
 		}
 
@@ -232,7 +235,7 @@ namespace UE::Mutable::Private
 		const bool bValidState = LiveInstanceState >= 0 && LiveInstanceState < (int32)Program.States.Num();
 		if (!bValidState)
 		{
-			UE_LOG(LogMutableCore, Error, TEXT("Invalid state in mutable update."));
+			LiveInstance->UpdateLogger->LogUpdateMessage(TEXT("Invalid state in mutable update."));
 			return nullptr;
 		}
 
@@ -1782,4 +1785,68 @@ namespace UE::Mutable::Private
 
     	return Registry->Add(NewKey, {});
 	}
+
+
+	// FLiveInstanceLogger ---------------------------------------
+
+	void FLiveInstanceLogger::SetLogger(TSharedPtr<IMutableEditorLogger> InInstanceLogger)
+	{
+		InstanceLogger = InInstanceLogger;
+	}
+
+
+	void FLiveInstanceLogger::LogUpdateMessage(const FString& Message, ELogVerbosity::Type Verbosity)
+	{
+		switch (Verbosity)
+		{
+		case ELogVerbosity::Error:
+			UE_LOG(LogMutableCore, Error, TEXT("%s"), *Message);
+			break;
+		case ELogVerbosity::Warning:
+			UE_LOG(LogMutableCore, Warning, TEXT("%s"), *Message);
+			break;
+		case ELogVerbosity::Verbose:
+			UE_LOG(LogMutableCore, Verbose, TEXT("%s"), *Message);
+			break;
+		case ELogVerbosity::Display:
+			UE_LOG(LogMutableCore, Verbose, TEXT("%s"), *Message);
+			break;
+		case ELogVerbosity::Log:
+		default:
+			UE_LOG(LogMutableCore, Log, TEXT("%s"), *Message);
+			break;
+		}
+
+		if (InstanceLogger)
+		{
+			EMessageSeverity::Type MessageSeverity;
+
+			switch (Verbosity)
+			{
+			case ELogVerbosity::Error:
+				MessageSeverity = EMessageSeverity::Error;
+				break;
+			case ELogVerbosity::Warning:
+				MessageSeverity = EMessageSeverity::Warning;
+				break;
+			case ELogVerbosity::Verbose:
+			case ELogVerbosity::Log:
+			case ELogVerbosity::Display:
+			default:
+				MessageSeverity = EMessageSeverity::Info;
+				break;
+			}
+
+			TSharedRef<FTokenizedMessage> LogMessage = FTokenizedMessage::Create(MessageSeverity, FText::FromString(Message));
+			LogMessage->SetIndentationLevel(1);
+			TSharedPtr<IMutableEditorLogger> CopyMessageLogListing = InstanceLogger;
+
+			// Required since slate widgets need to run in the gamethread
+			ExecuteOnGameThread(UE_SOURCE_LOCATION, [CopyMessageLogListing, LogMessage]() mutable
+			{
+				CopyMessageLogListing->LogMessage(LogMessage);
+			});
+		}
+	}
+
 }
