[Mutable] Fix debugger crash when selecting any operation requiring rom loads.

#Synced-CL 50648887

diff --git a/Engine/Plugins/Mutable/Source/CustomizableObjectEditor/Private/MuCOE/SMutableCodeViewer.cpp b/Engine/Plugins/Mutable/Source/CustomizableObjectEditor/Private/MuCOE/SMutableCodeViewer.cpp
--- a/Engine/Plugins/Mutable/Source/CustomizableObjectEditor/Private/MuCOE/SMutableCodeViewer.cpp
+++ b/Engine/Plugins/Mutable/Source/CustomizableObjectEditor/Private/MuCOE/SMutableCodeViewer.cpp
@@ -2859,7 +2859,9 @@ void SMutableCodeViewer::Tick(const FGeometry& AllottedGeometry, const double In
 
 	UE::Mutable::Private::FSettings Settings;
 	const TSharedPtr<UE::Mutable::Private::FSystem> System = MakeShared<UE::Mutable::Private::FSystem>(Settings);
+	const TSharedRef<UE::Mutable::Private::FModelReader> ModelReader = UCustomizableObjectSystem::GetInstance()->GetPrivate()->MutableSystem->StreamInterface.ToSharedRef();
 
+	System->SetStreamingInterface(ModelReader);
 	TSharedRef<UE::Mutable::Private::FLiveInstance> TempLiveInstance = System->NewBuildInstance(MutableModel, PreviewParameters, ExternalResourceProvider);
 	
 	switch (OperationDataType)
