[mutable] Baked and removed not in use CVar handling the cancellation of mutable mesh streaming tasks.

#rnx
#Synced-CL 50769877

diff --git a/Engine/Plugins/Mutable/Source/CustomizableObject/Private/MuCO/CustomizableObjectMeshUpdate.cpp b/Engine/Plugins/Mutable/Source/CustomizableObject/Private/MuCO/CustomizableObjectMeshUpdate.cpp
--- a/Engine/Plugins/Mutable/Source/CustomizableObject/Private/MuCO/CustomizableObjectMeshUpdate.cpp
+++ b/Engine/Plugins/Mutable/Source/CustomizableObject/Private/MuCO/CustomizableObjectMeshUpdate.cpp
@@ -26,13 +26,6 @@ template class TRenderAssetUpdate<FSkelMeshUpdateContext>;
 
 #define UE_MUTABLE_UPDATE_MESH_REGION		TEXT("Task_Mutable_UpdateMesh")
 
-static bool bEnableGCHangFix = true;
-FAutoConsoleVariableRef CVarMutableEnableGCHangFix(
-	TEXT("mutable.EnableGCHangFix"),
-	bEnableGCHangFix,
-	TEXT("Fix hang when FCustomizableObjectMeshStreamIn is canceled and TaskSynchronization is higher than 0.")
-	TEXT("If true, TaskSynchronization decrement will happen in the Abort method instead of DoCancelMeshUpdate.")
-);
 
 static bool bEnableCPUMorphTargetStreaming = true;
 FAutoConsoleVariableRef CVarMutableEnableCPUMorphTargetStreaming(
@@ -59,7 +52,7 @@ FCustomizableObjectMeshStreamIn::FCustomizableObjectMeshStreamIn(
 
 void FCustomizableObjectMeshStreamIn::OnUpdateMeshFinished()
 {
-	if (!IsCancelled() || !bEnableGCHangFix)
+	if (!IsCancelled())
 	{
 		check(TaskSynchronization.GetValue() > 0)
 
@@ -75,7 +68,7 @@ void FCustomizableObjectMeshStreamIn::OnUpdateMeshFinished()
 
 void FCustomizableObjectMeshStreamIn::Abort()
 {
-	if (!IsCancelled() && !IsCompleted() && bEnableGCHangFix)
+	if (!IsCancelled() && !IsCompleted())
 	{
 		FSkeletalMeshStreamIn::Abort();
 
@@ -107,14 +100,7 @@ void FCustomizableObjectMeshStreamIn::DoInitiate(const FContext& Context)
 	// Launch MutableTask
 	RequestMeshUpdate(Context);
 
-	if (bEnableGCHangFix)
-	{
-		PushTask(Context, TT_Async, SRA_UPDATE_CALLBACK(DoConvertResources), TT_Async, SRA_UPDATE_CALLBACK(DoCancel));
-	}
-	else
-	{
-		PushTask(Context, TT_Async, SRA_UPDATE_CALLBACK(DoConvertResources), TT_Async, SRA_UPDATE_CALLBACK(DoCancelMeshUpdate));
-	}
+	PushTask(Context, TT_Async, SRA_UPDATE_CALLBACK(DoConvertResources), TT_Async, SRA_UPDATE_CALLBACK(DoCancel));
 }
 
 
@@ -154,15 +140,6 @@ void FCustomizableObjectMeshStreamIn::DoCreateBuffers(const FContext& Context)
 }
 
 
-void FCustomizableObjectMeshStreamIn::DoCancelMeshUpdate(const FContext& Context)
-{
-	MUTABLE_CPUPROFILER_SCOPE(FCustomizableObjectMeshStreamIn::DoCancelMeshUpdate)
-
-	CancelMeshUpdate(Context);
-	PushTask(Context, TT_None, nullptr, (EThreadType)Context.CurrentThread, SRA_UPDATE_CALLBACK(DoCancel));
-}
-
-
 namespace impl
 {
 	void Task_Mutable_UpdateMesh_End(const TSharedPtr<FMutableUpdateMeshContext>& OperationData, TRefCountPtr<FCustomizableObjectMeshStreamIn>& StreamingTask)
@@ -236,7 +213,7 @@ namespace impl
 	{
 		MUTABLE_CPUPROFILER_SCOPE(Task_Mutable_UpdateMesh);
 
-		if (Task->IsCancelled() && bEnableGCHangFix)
+		if (Task->IsCancelled())
 		{
 			return;
 		}
@@ -334,43 +311,13 @@ void FCustomizableObjectMeshStreamIn::RequestMeshUpdate(const FContext& Context)
 	}
 
 	
-	if (IsCancelled() && TaskSynchronization.GetValue() > 0 && bEnableGCHangFix)
+	if (IsCancelled() && TaskSynchronization.GetValue() > 0)
 	{
 		TaskSynchronization.Set(0);
 	}
 }
 
 
-void FCustomizableObjectMeshStreamIn::CancelMeshUpdate(const FContext& Context)
-{
-	MUTABLE_CPUPROFILER_SCOPE(FCustomizableObjectMeshStreamIn::CancelMeshUpdate)
-	
-	UCustomizableObjectSystemPrivate* CustomizableObjectSystem = UCustomizableObjectSystem::GetInstance()->GetPrivate();
-	if (CustomizableObjectSystem && MutableTaskId != FMutableTaskGraph::INVALID_ID)
-	{
-		// Cancel task if not launched yet.
-		const bool bMutableTaskCancelledBeforeRun = CustomizableObjectSystem->MutableTaskGraph.CancelMutableThreadTaskLowPriority(MutableTaskId);
-		if (bMutableTaskCancelledBeforeRun)
-		{
-			// Clear MeshUpdate data
-			OperationData = nullptr;
-
-			// At this point task synchronization would hold the number of pending requests.
-			TaskSynchronization.Decrement();
-			check(TaskSynchronization.GetValue() == 0);
-		}
-	}
-	else
-	{
-		check(TaskSynchronization.GetValue() == 0);
-	}
-
-	// The tick here is intended to schedule the success or cancel callback.
-	// Using TT_None ensure gets which could create a dead lock.
-	Tick(FSkeletalMeshUpdate::TT_None);
-}
-
-
 void FCustomizableObjectMeshStreamIn::ConvertMesh(const FContext& Context, bool& bOutMarkRenderStateDirty)
 {
 	MUTABLE_CPUPROFILER_SCOPE(FCustomizableObjectMeshStreamIn::ConvertMesh);
diff --git a/Engine/Plugins/Mutable/Source/CustomizableObject/Private/MuCO/CustomizableObjectMeshUpdate.h b/Engine/Plugins/Mutable/Source/CustomizableObject/Private/MuCO/CustomizableObjectMeshUpdate.h
--- a/Engine/Plugins/Mutable/Source/CustomizableObject/Private/MuCO/CustomizableObjectMeshUpdate.h
+++ b/Engine/Plugins/Mutable/Source/CustomizableObject/Private/MuCO/CustomizableObjectMeshUpdate.h
@@ -78,14 +78,9 @@ private:
 
 	void DoCreateBuffers(const FContext& Context);
 
-	void DoCancelMeshUpdate(const FContext& Context);
-
 	/** Creates a MeshUpdate task to generate the meshes for the LODs to stream in. */
 	void RequestMeshUpdate(const FContext& Context);
 
-	/** Cancels the MeshUpdate task. */
-	void CancelMeshUpdate(const FContext& Context);
-
 	/** Converts from UE::Mutable::Private::FMesh to FSkeletalMeshLODRenderData. */
 	void ConvertMesh(const FContext& Context, bool& bOutMarkRenderStateDirty);
 
