[Mutable] Fix crash when re-compiling a CO with texture parameters.
Remove invalid warning.

#rnx

#Synced-CL 50860972

diff --git a/Engine/Plugins/Mutable/Source/CustomizableObject/Private/MuCO/CustomizableObjectInstance.cpp b/Engine/Plugins/Mutable/Source/CustomizableObject/Private/MuCO/CustomizableObjectInstance.cpp
--- a/Engine/Plugins/Mutable/Source/CustomizableObject/Private/MuCO/CustomizableObjectInstance.cpp
+++ b/Engine/Plugins/Mutable/Source/CustomizableObject/Private/MuCO/CustomizableObjectInstance.cpp
@@ -4851,8 +4851,11 @@ UTexture* BuildTexture(const TSharedRef<FUpdateContextPrivate>& Context, const i
 	{
 		UTexture* Texture = Image->Image->PassthroughObject.Get(); // TODO GMT Why not also cache Passthrough?
 
-		Context->Instance->GetPrivate()->Textures.AddUnique(Texture);
-		
+		if (Texture)
+		{
+			Context->Instance->GetPrivate()->Textures.AddUnique(Texture);
+		}
+
 		return Texture;
 	}
 
diff --git a/Engine/Plugins/Mutable/Source/CustomizableObject/Private/MuCO/CustomizableObjectSystem.cpp b/Engine/Plugins/Mutable/Source/CustomizableObject/Private/MuCO/CustomizableObjectSystem.cpp
--- a/Engine/Plugins/Mutable/Source/CustomizableObject/Private/MuCO/CustomizableObjectSystem.cpp
+++ b/Engine/Plugins/Mutable/Source/CustomizableObject/Private/MuCO/CustomizableObjectSystem.cpp
@@ -1767,7 +1767,7 @@ bool UCustomizableObjectSystem::CheckIfDiskOrMipUpdateOperationsPending(const UC
 		{
 			for (const TObjectPtr<UTexture>& GeneratedTexture : CustomizableObjectInstance->GetPrivate()->Textures) // TODO GMT This approach is incorrect. An update may clear the generated Textures array but these may live longer.
 			{
-				if (GeneratedTexture->HasPendingInitOrStreaming())
+				if (GeneratedTexture && GeneratedTexture->HasPendingInitOrStreaming())
 				{
 					return true;
 				}
diff --git a/Engine/Plugins/Mutable/Source/CustomizableObject/Private/MuCO/UnrealConversionUtils.cpp b/Engine/Plugins/Mutable/Source/CustomizableObject/Private/MuCO/UnrealConversionUtils.cpp
--- a/Engine/Plugins/Mutable/Source/CustomizableObject/Private/MuCO/UnrealConversionUtils.cpp
+++ b/Engine/Plugins/Mutable/Source/CustomizableObject/Private/MuCO/UnrealConversionUtils.cpp
@@ -1470,10 +1470,6 @@ namespace UnrealConversionUtils
 
 				if (!LOD->RenderSections.IsValidIndex(SectionIndex))
 				{
-					// This may happen in the provided mesh in the parameter value doesn't have the expected number of sections
-					UE_LOG(LogMutable, Warning, TEXT("Trying to convert mesh [%s] section %d in LOD %d but it only has %d sections"),
-						*SkeletalMesh->GetName(), SectionIndex, LODIndex, LOD->RenderSections.Num());
-
 					return;
 				}
 					
