Optimized Mutable_UCustomizableObjectSystem::Tick by removing some usages of TObjectIterator.

Saves 50us per frame on Switch.

#tests PIE, win64 test client/server, and switch replayrun

#Synced-CL 50857149

#Synced-CL 50857149

diff --git a/Engine/Plugins/Mutable/Source/CustomizableObject/Internal/MuCO/CustomizableObjectSystemPrivate.h b/Engine/Plugins/Mutable/Source/CustomizableObject/Internal/MuCO/CustomizableObjectSystemPrivate.h
--- a/Engine/Plugins/Mutable/Source/CustomizableObject/Internal/MuCO/CustomizableObjectSystemPrivate.h
+++ b/Engine/Plugins/Mutable/Source/CustomizableObject/Internal/MuCO/CustomizableObjectSystemPrivate.h
@@ -840,6 +840,20 @@ public:
 	float LogStartedUpdateLast = 0;
 
 public:
+	const TSet<UCustomizableObjectInstanceUsage*>& GetCustomizableObjectInstanceUsages() const
+	{
+		return CustomizableObjectInstanceUsages;
+	}
+	static void RegisterCustomizableObjectInstanceUsage(UCustomizableObjectInstanceUsage* Obj);
+	static void UnregisterCustomizableObjectInstanceUsage(UCustomizableObjectInstanceUsage* Obj);
+
+	const TSet<UCustomizableObjectInstance*>& GetCustomizableObjectInstances() const
+	{
+		return CustomizableObjectInstances;
+	}
+	static void RegisterCustomizableObjectInstance(UCustomizableObjectInstance* Obj);
+	static void UnregisterCustomizableObjectInstance(UCustomizableObjectInstance* Obj);
+
 	/**
 	 * Get to know if the settings used by the mutable syustem are optimzed for benchmarking operations or not
 	 * @return true if using benchmarking optimized settings, false otherwise.
@@ -857,6 +871,10 @@ public:
 private:
 	/** Flag that controls some of the settings used for the generation of instances. */
 	inline static bool bUseBenchmarkingSettings = false;
+
+	// Added to in PostInitProperties, removed from in BeginDestroy. No CDOs
+	TSet<UCustomizableObjectInstanceUsage*> CustomizableObjectInstanceUsages;
+	TSet<UCustomizableObjectInstance*> CustomizableObjectInstances;
 };
 
 
diff --git a/Engine/Plugins/Mutable/Source/CustomizableObject/Private/MuCO/CustomizableInstanceLODManagement.cpp b/Engine/Plugins/Mutable/Source/CustomizableObject/Private/MuCO/CustomizableInstanceLODManagement.cpp
--- a/Engine/Plugins/Mutable/Source/CustomizableObject/Private/MuCO/CustomizableInstanceLODManagement.cpp
+++ b/Engine/Plugins/Mutable/Source/CustomizableObject/Private/MuCO/CustomizableInstanceLODManagement.cpp
@@ -38,6 +38,11 @@ void UCustomizableInstanceLODManagementBase::EditorUpdateComponent(UCustomizable
 }
 #endif
 
+const TSet<UCustomizableObjectInstanceUsage*>& UCustomizableInstanceLODManagementBase::GetCustomizableObjectInstanceUsages() const
+{
+	UCustomizableObjectSystem* CustomizableObjectSystem = UCustomizableObjectSystem::GetInstance();
+	return CustomizableObjectSystem->GetPrivate()->GetCustomizableObjectInstanceUsages();
+}
 
 UCustomizableInstanceLODManagement::UCustomizableInstanceLODManagement() : UCustomizableInstanceLODManagementBase()
 {
@@ -52,39 +57,50 @@ UCustomizableInstanceLODManagement::~UCustomizableInstanceLODManagement()
 
 // Used to manually update distances used in "OnlyUpdateCloseCustomizableObjects" system. If OnlyForInstance is null, all instances have their distance updated
 // ViewCenter is the origin where the distances will be measured from.
-void UpdatePawnToInstancesDistances(const class UCustomizableObjectInstance* OnlyForInstance, const TWeakObjectPtr<const AActor> ViewCenter)
+static void UpdatePawnToInstancesDistances(
+	const TSet<UCustomizableObjectInstanceUsage*>& CustomizableObjectInstanceUsages,
+	const UCustomizableObjectInstance* OnlyForInstance,
+	const TWeakObjectPtr<const AActor> ViewCenter)
 {
-	for (TObjectIterator<UCustomizableObjectInstanceUsage> CustomizableObjectInstanceUsage; CustomizableObjectInstanceUsage; ++CustomizableObjectInstanceUsage)
+	for (UCustomizableObjectInstanceUsage* CustomizableObjectInstanceUsage : CustomizableObjectInstanceUsages)
 	{
-#if WITH_EDITOR
-		if (IsValid(*CustomizableObjectInstanceUsage) && CustomizableObjectInstanceUsage->GetPrivate()->IsNetMode(NM_DedicatedServer))
+		if (IsValid(CustomizableObjectInstanceUsage))
 		{
-			continue;
-		}
+#if WITH_EDITOR
+			if (CustomizableObjectInstanceUsage->GetPrivate()->IsNetMode(NM_DedicatedServer))
+			{
+				continue;
+			}
 #endif
 
-		if (IsValid(*CustomizableObjectInstanceUsage) && (OnlyForInstance == nullptr || CustomizableObjectInstanceUsage->GetCustomizableObjectInstance() == OnlyForInstance))
-		{
-			CustomizableObjectInstanceUsage->GetPrivate()->UpdateDistFromComponentToPlayer(ViewCenter.IsValid() ? ViewCenter.Get() : nullptr, OnlyForInstance != nullptr);
+			if (OnlyForInstance == nullptr || CustomizableObjectInstanceUsage->GetCustomizableObjectInstance() == OnlyForInstance)
+			{
+				CustomizableObjectInstanceUsage->GetPrivate()->UpdateDistFromComponentToPlayer(ViewCenter.IsValid() ? ViewCenter.Get() : nullptr, OnlyForInstance != nullptr);
+			}
 		}
 	}
 }
 
 #if WITH_EDITOR
 // Used to manually update instances in the level editor
-void UpdateCameraToInstancesDistance(const FVector CameraPosition)
+static void UpdateCameraToInstancesDistance(
+	const TSet<UCustomizableObjectInstanceUsage*>& CustomizableObjectInstanceUsages,
+	const FVector& CameraPosition)
 {
-	for (TObjectIterator<UCustomizableObjectInstanceUsage> CustomizableObjectInstanceUsage; CustomizableObjectInstanceUsage; ++CustomizableObjectInstanceUsage)
+	for (UCustomizableObjectInstanceUsage* CustomizableObjectInstanceUsage : CustomizableObjectInstanceUsages)
 	{
-#if WITH_EDITOR
-		if (IsValid(*CustomizableObjectInstanceUsage) && CustomizableObjectInstanceUsage->GetPrivate()->IsNetMode(NM_DedicatedServer))
+		if (IsValid(CustomizableObjectInstanceUsage))
 		{
-			continue;
-		}
+#if WITH_EDITOR
+			if (CustomizableObjectInstanceUsage->GetPrivate()->IsNetMode(NM_DedicatedServer))
+			{
+				continue;
+			}
 #endif
-		if (IsValid(*CustomizableObjectInstanceUsage) && !CustomizableObjectInstanceUsage->IsTemplate())
-		{
-			CustomizableObjectInstanceUsage->GetPrivate()->UpdateDistFromComponentToLevelEditorCamera(CameraPosition);
+			if (!CustomizableObjectInstanceUsage->IsTemplate())
+			{
+				CustomizableObjectInstanceUsage->GetPrivate()->UpdateDistFromComponentToLevelEditorCamera(CameraPosition);
+			}
 		}
 	}
 }
@@ -93,10 +109,14 @@ void UpdateCameraToInstancesDistance(const FVector CameraPosition)
 
 void UCustomizableInstanceLODManagement::UpdateInstanceDistsAndLODs(FMutableInstanceUpdateMap& InOutRequestedUpdates)
 {
-	if (IsOnlyUpdateCloseCustomizableObjectsEnabled() && IsOnlyGenerateRequestedLODLevelsEnabled())
+	const TSet<UCustomizableObjectInstanceUsage*>& LocalCustomizableObjectInstanceUsages = GetCustomizableObjectInstanceUsages();
+	if (LocalCustomizableObjectInstanceUsages.IsEmpty())
 	{
-		UCustomizableObjectSystem* CustomizableObjectSystem = UCustomizableObjectSystem::GetInstance();
+		return;
+	}
 
+	if (IsOnlyUpdateCloseCustomizableObjectsEnabled() && IsOnlyGenerateRequestedLODLevelsEnabled())
+	{
 		if (ViewCenters.Num() == 0) // Just use the first pawn
 		{
 			UCustomizableObjectInstanceUsage* FirstCustomizableObjectInstanceUsage = nullptr;
@@ -105,15 +125,15 @@ void UCustomizableInstanceLODManagement::UpdateInstanceDistsAndLODs(FMutableInst
 			bool bLevelEditorInstancesUpdated = false;
 	#endif
 
-			for (TObjectIterator<UCustomizableObjectInstanceUsage> CustomizableObjectInstanceUsage; CustomizableObjectInstanceUsage; ++CustomizableObjectInstanceUsage)
+			for (const UCustomizableObjectInstanceUsage* CustomizableObjectInstanceUsage : LocalCustomizableObjectInstanceUsages)
 			{
 #if WITH_EDITOR
-				if (IsValid(*CustomizableObjectInstanceUsage) && CustomizableObjectInstanceUsage->GetPrivate()->IsNetMode(NM_DedicatedServer))
+				if (IsValid(CustomizableObjectInstanceUsage) && CustomizableObjectInstanceUsage->GetPrivate()->IsNetMode(NM_DedicatedServer))
 				{
 					continue;
 				}
 #endif
-				if (!IsValid(*CustomizableObjectInstanceUsage))
+				if (!IsValid(CustomizableObjectInstanceUsage))
 				{
 					continue;
 				}
@@ -128,7 +148,7 @@ void UCustomizableInstanceLODManagement::UpdateInstanceDistsAndLODs(FMutableInst
 
 					if (ViewCenter.IsValid())
 					{
-						UpdatePawnToInstancesDistances(nullptr, ViewCenter);
+						UpdatePawnToInstancesDistances(LocalCustomizableObjectInstanceUsages, nullptr, ViewCenter);
 						break;
 					}
 
@@ -144,7 +164,7 @@ void UCustomizableInstanceLODManagement::UpdateInstanceDistsAndLODs(FMutableInst
 							{
 								if (LevelVC && LevelVC->IsPerspective())
 								{
-									UpdateCameraToInstancesDistance(LevelVC->GetViewLocation());
+									UpdateCameraToInstancesDistance(LocalCustomizableObjectInstanceUsages, LevelVC->GetViewLocation());
 									bLevelEditorInstancesUpdated = true;
 									break;
 								}
@@ -162,7 +182,7 @@ void UCustomizableInstanceLODManagement::UpdateInstanceDistsAndLODs(FMutableInst
 			{
 				if (ViewCenter.IsValid())
 				{
-					UpdatePawnToInstancesDistances(nullptr, ViewCenter);
+					UpdatePawnToInstancesDistances(LocalCustomizableObjectInstanceUsages, nullptr, ViewCenter);
 				}
 			}
 		}
@@ -181,16 +201,16 @@ void UCustomizableInstanceLODManagement::UpdateInstanceDistsAndLODs(FMutableInst
 		TMap<TObjectPtr<UCustomizableObjectInstance>, FLODTracker> InstancesMinLOD;
 		InstancesMinLOD.Reserve(100);
 
-		for (TObjectIterator<UCustomizableObjectInstanceUsage> CustomizableObjectInstanceUsage; CustomizableObjectInstanceUsage; ++CustomizableObjectInstanceUsage)
+		for (UCustomizableObjectInstanceUsage* CustomizableObjectInstanceUsage : LocalCustomizableObjectInstanceUsages)
 		{
 #if WITH_EDITOR
-			if (IsValid(*CustomizableObjectInstanceUsage) && CustomizableObjectInstanceUsage->GetPrivate()->IsNetMode(NM_DedicatedServer))
+			if (IsValid(CustomizableObjectInstanceUsage) && CustomizableObjectInstanceUsage->GetPrivate()->IsNetMode(NM_DedicatedServer))
 			{
 				continue;
 			}
 #endif
 
-			if (IsValid(*CustomizableObjectInstanceUsage) && !CustomizableObjectInstanceUsage->IsTemplate())
+			if (IsValid(CustomizableObjectInstanceUsage) && !CustomizableObjectInstanceUsage->IsTemplate())
 			{
 				UCustomizableObjectInstance* COI = CustomizableObjectInstanceUsage->GetCustomizableObjectInstance();
 				if (!COI || !COI->GetCustomizableObject())
diff --git a/Engine/Plugins/Mutable/Source/CustomizableObject/Private/MuCO/CustomizableObjectInstance.cpp b/Engine/Plugins/Mutable/Source/CustomizableObject/Private/MuCO/CustomizableObjectInstance.cpp
--- a/Engine/Plugins/Mutable/Source/CustomizableObject/Private/MuCO/CustomizableObjectInstance.cpp
+++ b/Engine/Plugins/Mutable/Source/CustomizableObject/Private/MuCO/CustomizableObjectInstance.cpp
@@ -323,10 +323,12 @@ bool UCustomizableObjectInstance::IsEditorOnly() const
 
 void UCustomizableObjectInstance::PostInitProperties()
 {
-	UObject::PostInitProperties();
+	Super::PostInitProperties();
 
-	if (!HasAllFlags(RF_ClassDefaultObject))
+	if (!HasAnyFlags(RF_ClassDefaultObject | RF_DefaultSubObject))
 	{
+		UCustomizableObjectSystemPrivate::RegisterCustomizableObjectInstance(this);
+
 		if (!PrivateData)
 		{
 			PrivateData = NewObject<UCustomizableInstancePrivate>(this, FName("Private"));
@@ -348,6 +350,11 @@ void UCustomizableObjectInstance::BeginDestroy()
 		PrivateData->BindObjectDelegates(GetCustomizableObject(), nullptr);
 #endif
 	}
+
+	if (!HasAnyFlags(RF_ClassDefaultObject | RF_DefaultSubObject))
+	{
+		UCustomizableObjectSystemPrivate::UnregisterCustomizableObjectInstance(this);
+	}
 	
 	Super::BeginDestroy();
 }
diff --git a/Engine/Plugins/Mutable/Source/CustomizableObject/Private/MuCO/CustomizableObjectInstanceUsage.cpp b/Engine/Plugins/Mutable/Source/CustomizableObject/Private/MuCO/CustomizableObjectInstanceUsage.cpp
--- a/Engine/Plugins/Mutable/Source/CustomizableObject/Private/MuCO/CustomizableObjectInstanceUsage.cpp
+++ b/Engine/Plugins/Mutable/Source/CustomizableObject/Private/MuCO/CustomizableObjectInstanceUsage.cpp
@@ -51,6 +51,23 @@ UCustomizableObjectInstanceUsage::UCustomizableObjectInstanceUsage()
 	Private = CreateDefaultSubobject<UCustomizableObjectInstanceUsagePrivate>(FName("Private"));
 }
 
+void UCustomizableObjectInstanceUsage::PostInitProperties()
+{
+	Super::PostInitProperties();
+	if (!HasAnyFlags(RF_ClassDefaultObject | RF_DefaultSubObject))
+	{
+		UCustomizableObjectSystemPrivate::RegisterCustomizableObjectInstanceUsage(this);
+	}
+}
+
+void UCustomizableObjectInstanceUsage::BeginDestroy()
+{
+	if (!HasAnyFlags(RF_ClassDefaultObject | RF_DefaultSubObject))
+	{
+		UCustomizableObjectSystemPrivate::UnregisterCustomizableObjectInstanceUsage(this);
+	}
+	Super::BeginDestroy();
+}
 
 void UCustomizableObjectInstanceUsage::SetCustomizableObjectInstance(UCustomizableObjectInstance* CustomizableObjectInstance)
 {
diff --git a/Engine/Plugins/Mutable/Source/CustomizableObject/Private/MuCO/CustomizableObjectSystem.cpp b/Engine/Plugins/Mutable/Source/CustomizableObject/Private/MuCO/CustomizableObjectSystem.cpp
--- a/Engine/Plugins/Mutable/Source/CustomizableObject/Private/MuCO/CustomizableObjectSystem.cpp
+++ b/Engine/Plugins/Mutable/Source/CustomizableObject/Private/MuCO/CustomizableObjectSystem.cpp
@@ -3283,10 +3283,8 @@ void UCustomizableObjectSystemPrivate::UpdateTick(const bool bBlocking)
 {
 	FMutableUpdateCandidate* LODUpdateCandidateFound = nullptr;
 
-	for (TObjectIterator<UCustomizableObjectInstanceUsage> CustomizableObjectInstanceUsageItr; CustomizableObjectInstanceUsageItr; ++CustomizableObjectInstanceUsageItr)
+	for (UCustomizableObjectInstanceUsage* CustomizableObjectInstanceUsage : CustomizableObjectInstanceUsages)
 	{
-		TObjectPtr<UCustomizableObjectInstanceUsage> CustomizableObjectInstanceUsage = *CustomizableObjectInstanceUsageItr;
-
 		if (!IsValid(CustomizableObjectInstanceUsage))
 		{
 			continue;
@@ -3352,31 +3350,31 @@ void UCustomizableObjectSystemPrivate::UpdateTick(const bool bBlocking)
 		
 		CurrentInstanceLODManagement->UpdateInstanceDistsAndLODs(RequestedLODUpdates);
 
-		for (TObjectIterator<UCustomizableObjectInstance> CustomizableObjectInstance; CustomizableObjectInstance; ++CustomizableObjectInstance)
+		for (UCustomizableObjectInstance* CustomizableObjectInstance : CustomizableObjectInstances)
 		{
-			if (IsValid(*CustomizableObjectInstance) && CustomizableObjectInstance->GetPrivate())
+			if (IsValid(CustomizableObjectInstance) && CustomizableObjectInstance->GetPrivate())
 			{
 				UCustomizableInstancePrivate* ObjectInstancePrivateData = CustomizableObjectInstance->GetPrivate();
 
 				if (ObjectInstancePrivateData->HasCOInstanceFlags(UsedByComponentInPlay))
 				{
-					ObjectInstancePrivateData->TickUpdateCloseCustomizableObjects(**CustomizableObjectInstance, RequestedLODUpdates);
+					ObjectInstancePrivateData->TickUpdateCloseCustomizableObjects(*CustomizableObjectInstance, RequestedLODUpdates);
 				}
 				else if (ObjectInstancePrivateData->HasCOInstanceFlags(UsedByComponent))
 				{
-					ensure(!RequestedLODUpdates.Contains(*CustomizableObjectInstance));
-					ObjectInstancePrivateData->UpdateInstanceIfNotGenerated(**CustomizableObjectInstance, RequestedLODUpdates);
+					ensure(!RequestedLODUpdates.Contains(CustomizableObjectInstance));
+					ObjectInstancePrivateData->UpdateInstanceIfNotGenerated(*CustomizableObjectInstance, RequestedLODUpdates);
 				}
 				else
 				{
-					ensure(!RequestedLODUpdates.Contains(*CustomizableObjectInstance));
+					ensure(!RequestedLODUpdates.Contains(CustomizableObjectInstance));
 				}
 
 				ObjectInstancePrivateData->ClearCOInstanceFlags((ECOInstanceFlags)(UsedByComponent | UsedByComponentInPlay | PendingLODsUpdate)); // TODO MTBL-391: Makes no sense to clear it here, what if an update is requested before we set it back to true
 			}
 			else
 			{
-				ensure(!RequestedLODUpdates.Contains(*CustomizableObjectInstance));
+				ensure(!RequestedLODUpdates.Contains(CustomizableObjectInstance));
 			}
 		}
 
@@ -3496,12 +3494,15 @@ void UCustomizableObjectSystemPrivate::UpdateTick(const bool bBlocking)
 		}
 
 		{
-			for (TObjectIterator<UCustomizableObjectInstance> CustomizableObjectInstance; CustomizableObjectInstance; ++CustomizableObjectInstance)
+			for (UCustomizableObjectInstance* CustomizableObjectInstance : CustomizableObjectInstances)
 			{
-				if (IsValid(*CustomizableObjectInstance) && CustomizableObjectInstance->GetPrivate())
+				if (IsValid(CustomizableObjectInstance))
 				{
-					CustomizableObjectInstance->GetPrivate()->LastMinSquareDistFromComponentToPlayer = CustomizableObjectInstance->GetPrivate()->MinSquareDistFromComponentToPlayer;
-					CustomizableObjectInstance->GetPrivate()->MinSquareDistFromComponentToPlayer = FLT_MAX;
+					if (UCustomizableInstancePrivate* InstancePrivate = CustomizableObjectInstance->GetPrivate())
+					{
+						InstancePrivate->LastMinSquareDistFromComponentToPlayer = InstancePrivate->MinSquareDistFromComponentToPlayer;
+						InstancePrivate->MinSquareDistFromComponentToPlayer = FLT_MAX;
+					}
 				}
 			}
 		}
@@ -3646,9 +3647,8 @@ void UCustomizableObjectSystem::DiscardInstances()
 				{
 					if (const UModelResources* ModelResources = CustomizableObject->GetPrivate()->GetModelResources())
 					{
-						for (TObjectIterator<UCustomizableObjectInstanceUsage> It; It; ++It)
+						for (UCustomizableObjectInstanceUsage* InstanceUsage : GetPrivate()->GetCustomizableObjectInstanceUsages())
 						{
-							UCustomizableObjectInstanceUsage* InstanceUsage = *It;
 							if (!IsValid(InstanceUsage) || InstanceUsage->GetCustomizableObjectInstance() != COI)
 							{
 								continue;
@@ -3699,6 +3699,35 @@ bool UCustomizableObjectSystem::IsUpdating(const UCustomizableObjectInstance* In
 	return GetPrivate()->IsUpdating(*Instance);
 }
 
+void UCustomizableObjectSystemPrivate::RegisterCustomizableObjectInstanceUsage(UCustomizableObjectInstanceUsage* Obj)
+{
+	UCustomizableObjectSystem* System = UCustomizableObjectSystem::GetInstance();
+	System->GetPrivate()->CustomizableObjectInstanceUsages.Add(Obj);
+}
+
+void UCustomizableObjectSystemPrivate::UnregisterCustomizableObjectInstanceUsage(UCustomizableObjectInstanceUsage* Obj)
+{
+	if (UCustomizableObjectSystem::IsCreated())
+	{
+		UCustomizableObjectSystem* System = UCustomizableObjectSystem::GetInstance();
+		System->GetPrivate()->CustomizableObjectInstanceUsages.Remove(Obj);
+	}
+}
+
+void UCustomizableObjectSystemPrivate::RegisterCustomizableObjectInstance(UCustomizableObjectInstance* Obj)
+{
+	UCustomizableObjectSystem* System = UCustomizableObjectSystem::GetInstance();
+	System->GetPrivate()->CustomizableObjectInstances.Add(Obj);
+}
+
+void UCustomizableObjectSystemPrivate::UnregisterCustomizableObjectInstance(UCustomizableObjectInstance* Obj)
+{
+	if (UCustomizableObjectSystem::IsCreated())
+	{
+		UCustomizableObjectSystem* System = UCustomizableObjectSystem::GetInstance();
+		System->GetPrivate()->CustomizableObjectInstances.Remove(Obj);
+	}
+}
 
 bool UCustomizableObjectSystemPrivate::IsUsingBenchmarkingSettings()
 {
diff --git a/Engine/Plugins/Mutable/Source/CustomizableObject/Public/MuCO/CustomizableInstanceLODManagement.h b/Engine/Plugins/Mutable/Source/CustomizableObject/Public/MuCO/CustomizableInstanceLODManagement.h
--- a/Engine/Plugins/Mutable/Source/CustomizableObject/Public/MuCO/CustomizableInstanceLODManagement.h
+++ b/Engine/Plugins/Mutable/Source/CustomizableObject/Public/MuCO/CustomizableInstanceLODManagement.h
@@ -28,6 +28,9 @@ public:
 #if WITH_EDITOR
 	UE_API void EditorUpdateComponent(UCustomizableObjectInstanceUsage* InstanceUsage);
 #endif
+
+protected:
+	UE_API const TSet<UCustomizableObjectInstanceUsage*>& GetCustomizableObjectInstanceUsages() const;
 };
 
 
diff --git a/Engine/Plugins/Mutable/Source/CustomizableObject/Public/MuCO/CustomizableObjectInstanceUsage.h b/Engine/Plugins/Mutable/Source/CustomizableObject/Public/MuCO/CustomizableObjectInstanceUsage.h
--- a/Engine/Plugins/Mutable/Source/CustomizableObject/Public/MuCO/CustomizableObjectInstanceUsage.h
+++ b/Engine/Plugins/Mutable/Source/CustomizableObject/Public/MuCO/CustomizableObjectInstanceUsage.h
@@ -35,6 +35,9 @@ public:
 
 	// Own interface
 	UE_API UCustomizableObjectInstanceUsage();
+
+	virtual void PostInitProperties() override;
+	virtual void BeginDestroy() override;
 	
 	UFUNCTION(BlueprintCallable, Category = CustomizableObjectInstanceUsage)
 	UE_API void SetCustomizableObjectInstance(UCustomizableObjectInstance* CustomizableObjectInstance);
