[mutable] Fixed morph-related error messages not having the node context set.
- This happened because we do generate the morph nodes as a string of nodes where only the first one to be found is returned, and then only that node gets its context set.
- By setting the context each time we encounter a new morph in the stack we ensure each one of them does have the context appropriately set.

#rnx

#Synced-CL 50903042

diff --git a/Engine/Plugins/Mutable/Source/CustomizableObjectEditor/Private/MuCOE/GenerateMutableSource/GenerateMutableSourceMesh.cpp b/Engine/Plugins/Mutable/Source/CustomizableObjectEditor/Private/MuCOE/GenerateMutableSource/GenerateMutableSourceMesh.cpp
--- a/Engine/Plugins/Mutable/Source/CustomizableObjectEditor/Private/MuCOE/GenerateMutableSource/GenerateMutableSourceMesh.cpp
+++ b/Engine/Plugins/Mutable/Source/CustomizableObjectEditor/Private/MuCOE/GenerateMutableSource/GenerateMutableSourceMesh.cpp
@@ -2496,9 +2496,6 @@ UE::Mutable::Private::NodeMeshPtr GenerateMorphMesh(const UEdGraphPin* Pin,
 		UE::Mutable::Private::NodeMeshPtr NextMorph = GenerateMorphMesh(Pin, TypedNodeMorphs, MorphIndex + 1, SourceNode, GenerationContext, bOnlyConnectedLOD, TableColumnName); // TODO FutureGMT change to a for. This recursion can be problematic with the production cache
 		Result->Base = NextMorph;
 	}
-	
-	// Target
-	UE::Mutable::Private::NodeMeshPtr BaseSourceMesh = SourceNode;
 
 	if (UCustomizableObjectNodeMeshMorph* TypedMorphNode = Cast<UCustomizableObjectNodeMeshMorph>(TypedNodeMorphs[MorphIndex].OwningNode))
 	{
@@ -2536,6 +2533,8 @@ UE::Mutable::Private::NodeMeshPtr GenerateMorphMesh(const UEdGraphPin* Pin,
 		}
 	}
 
+	Result->SetMessageContext(MorphNode);
+	
 	return Result;
 }
 
