[Mutable] Fix ForEachReference function for ME_MORPH Ops.

- Now ME_MERGE Ops use structured Args.
- Remove unused code paths using old morph data.
- Remove code generation for morph target mesh.

#rnx

#Synced-CL 50905358

diff --git a/Engine/Plugins/Mutable/Source/CustomizableObject/Internal/MuCO/CustomizableObjectPrivate.h b/Engine/Plugins/Mutable/Source/CustomizableObject/Internal/MuCO/CustomizableObjectPrivate.h
--- a/Engine/Plugins/Mutable/Source/CustomizableObject/Internal/MuCO/CustomizableObjectPrivate.h
+++ b/Engine/Plugins/Mutable/Source/CustomizableObject/Internal/MuCO/CustomizableObjectPrivate.h
@@ -1488,6 +1488,8 @@ enum class ECustomizableObjectVersions : int32
 	UseCompressedMorphsForAllMorphData,
 
 	TestMoveSocketsToCore,
+
+	RemoveMorphTarget,
 	
 	// -----<new versions can be added above this line>--------
 	LastCustomizableObjectVersion
diff --git a/Engine/Plugins/Mutable/Source/CustomizableObjectEditor/Private/MuCOE/SMutableCodeViewer.cpp b/Engine/Plugins/Mutable/Source/CustomizableObjectEditor/Private/MuCOE/SMutableCodeViewer.cpp
--- a/Engine/Plugins/Mutable/Source/CustomizableObjectEditor/Private/MuCOE/SMutableCodeViewer.cpp
+++ b/Engine/Plugins/Mutable/Source/CustomizableObjectEditor/Private/MuCOE/SMutableCodeViewer.cpp
@@ -1643,28 +1643,15 @@ void SMutableCodeViewer::GenerateElementRecursive(const int32& InStateIndex, UE:
 		
 	case UE::Mutable::Private::EOpType::ME_MORPH:
 	{
-		const uint8* Data = InProgram.GetOpArgsPointer(InParentAddress);
+		const UE::Mutable::Private::OP::MeshMorphArgs Args = InProgram.GetOpArgs<UE::Mutable::Private::OP::MeshMorphArgs>(InParentAddress);
 
-		UE::Mutable::Private::OP::CONSTANT_STRING_INDEX NameAt = 0;
-		FMemory::Memcpy(&NameAt, Data, sizeof(UE::Mutable::Private::OP::CONSTANT_STRING_INDEX));
-		AddOpFunc(NameAt, TEXT("name "));
-		Data += sizeof(UE::Mutable::Private::OP::CONSTANT_STRING_INDEX);
-			
-		UE::Mutable::Private::OP::ADDRESS FactorAt = 0;
-		FMemory::Memcpy(&FactorAt, Data, sizeof(UE::Mutable::Private::OP::ADDRESS));
-		AddOpFunc(FactorAt, TEXT("factor "));
-		Data += sizeof(UE::Mutable::Private::OP::ADDRESS);
-		
-		UE::Mutable::Private::OP::ADDRESS BaseAt = 0;
-		FMemory::Memcpy(&BaseAt, Data, sizeof(UE::Mutable::Private::OP::ADDRESS));
-		AddOpFunc(BaseAt, TEXT("base "));
-		Data += sizeof(UE::Mutable::Private::OP::ADDRESS);
-
-		UE::Mutable::Private::OP::ADDRESS TargetAt = 0;
-		FMemory::Memcpy(&TargetAt, Data, sizeof(UE::Mutable::Private::OP::ADDRESS));
-		AddOpFunc(TargetAt, TEXT("target "));
-		Data += sizeof(UE::Mutable::Private::OP::ADDRESS);
+		FString MorphName = InProgram.ConstantStrings.IsValidIndex(Args.Name) 
+			? InProgram.ConstantStrings[Args.Name]
+			: FString(TEXT("[Name Not Found]"));
 		
+		AddOpFunc(Args.Factor, TEXT("Factor "));
+		AddOpFunc(Args.Base, FString::Printf(TEXT("%s : Base "), *MorphName));
+	
 		break;
 	}
 		
diff --git a/Engine/Plugins/Mutable/Source/CustomizableObjectEditor/Private/MuCOE/SMutableGraphViewer.cpp b/Engine/Plugins/Mutable/Source/CustomizableObjectEditor/Private/MuCOE/SMutableGraphViewer.cpp
--- a/Engine/Plugins/Mutable/Source/CustomizableObjectEditor/Private/MuCOE/SMutableGraphViewer.cpp
+++ b/Engine/Plugins/Mutable/Source/CustomizableObjectEditor/Private/MuCOE/SMutableGraphViewer.cpp
@@ -685,7 +685,6 @@ void SMutableGraphViewer::GetChildrenForInfo(TSharedPtr<FMutableGraphTreeElement
 	{
 		UE::Mutable::Private::NodeMeshMorph* MeshMorphVar = StaticCast<UE::Mutable::Private::NodeMeshMorph*>(ParentNode);
 		AddChildFunc(MeshMorphVar->Base.get(), FString::Printf(TEXT("BASE")));
-		AddChildFunc(MeshMorphVar->Morph.get(), FString::Printf(TEXT("MORPH")));
 		AddChildFunc(MeshMorphVar->Factor.get(), FString::Printf(TEXT("FACTOR")));
 	}
 
diff --git a/Engine/Plugins/Mutable/Source/MutableRuntime/Internal/MuR/Operations.h b/Engine/Plugins/Mutable/Source/MutableRuntime/Internal/MuR/Operations.h
--- a/Engine/Plugins/Mutable/Source/MutableRuntime/Internal/MuR/Operations.h
+++ b/Engine/Plugins/Mutable/Source/MutableRuntime/Internal/MuR/Operations.h
@@ -1037,6 +1037,13 @@ namespace UE::Mutable::Private
 			ADDRESS scalarB;
 		};
 
+		struct MeshMorphArgs
+		{
+			CONSTANT_STRING_INDEX Name;
+			ADDRESS Base;
+			ADDRESS Factor;
+		};
+
 		struct MeshBindShapeArgs
 		{
 			ADDRESS mesh;
diff --git a/Engine/Plugins/Mutable/Source/MutableRuntime/Private/MuR/CodeRunner.cpp b/Engine/Plugins/Mutable/Source/MutableRuntime/Private/MuR/CodeRunner.cpp
--- a/Engine/Plugins/Mutable/Source/MutableRuntime/Private/MuR/CodeRunner.cpp
+++ b/Engine/Plugins/Mutable/Source/MutableRuntime/Private/MuR/CodeRunner.cpp
@@ -1354,39 +1354,24 @@ namespace UE::Mutable::Private
 
         case EOpType::ME_MORPH:
         {
-			const uint8* Data = Program.GetOpArgsPointer(Item.At);
-
-			OP::CONSTANT_STRING_INDEX NameAt = 0;
-			FMemory::Memcpy(&NameAt, Data, sizeof(OP::CONSTANT_STRING_INDEX)); 
-			Data += sizeof(OP::CONSTANT_STRING_INDEX);
-				
-			OP::ADDRESS FactorAt = 0;
-			FMemory::Memcpy(&FactorAt, Data, sizeof(OP::ADDRESS)); 
-			Data += sizeof(OP::ADDRESS);
-			
-			OP::ADDRESS BaseAt = 0;
-			FMemory::Memcpy(&BaseAt, Data, sizeof(OP::ADDRESS)); 
-			Data += sizeof(OP::ADDRESS);
-
-			OP::ADDRESS TargetAt = 0;
-			FMemory::Memcpy(&TargetAt, Data, sizeof(OP::ADDRESS)); 
-			Data += sizeof(OP::ADDRESS);
+			const OP::MeshMorphArgs Args = Program.GetOpArgs<OP::MeshMorphArgs>(Item.At);
 
 			switch (Item.Stage)
             {
             case 0:
 			{
-				if (BaseAt)
+				if (Args.Base)
 				{
 					if (!EnumHasAnyFlags(EMeshContentFlags(Item.ExecutionOptions), EMeshContentFlags::GeometryData))
 					{
 						AddOp(FScheduledOp(Item.At, Item, 1),
-							FScheduledOp(BaseAt, Item));
+							FScheduledOp(Args.Base, Item));
 					}
 					else
 					{
 						AddOp(FScheduledOp(Item.At, Item, 1),
-							FScheduledOp(FactorAt, Item));
+							FScheduledOp(Args.Base, Item),
+							FScheduledOp(Args.Factor, Item));
 					}
 				}
 				else
@@ -1395,85 +1380,38 @@ namespace UE::Mutable::Private
 				}
                 break;
 			}
+
             case 1:
             {
-                MUTABLE_CPUPROFILER_SCOPE(ME_MORPH_1)
-
+       		    MUTABLE_CPUPROFILER_SCOPE(ME_MORPH_1)
+				
 				if (!EnumHasAnyFlags(EMeshContentFlags(Item.ExecutionOptions), EMeshContentFlags::GeometryData))
 				{
-                	TManagedPtr<const FMesh> Base = LoadMesh(FCacheAddress(BaseAt, Item));
-					StoreMesh(Item, Base);
+                	StoreMesh(Item, LoadMesh(FCacheAddress(Args.Base, Item)));
 				}
-				else	
+				else
 				{
-					float Factor = LoadScalar(FCacheAddress(FactorAt, Item));
-
-					// Factor goes from -1 to 1 across all targets. [0 - 1] represents positive morphs, while [-1, 0) represent negative morphs.
-					Factor = FMath::Clamp(Factor, -1.0f, 1.0f); // Is the factor not in range [-1, 1], it will index a non existing morph.
+					TManagedPtr<const FMesh> BaseMesh = LoadMesh(FCacheAddress(Args.Base, Item));
+					float Factor = LoadScalar(FCacheAddress(Args.Factor, Item));
 
-					FScheduledOpData OpHeapData;
-					OpHeapData.Interpolate.Bifactor = Factor;
-					uint32 DataAddress = uint32(HeapData.Add(OpHeapData));
-
-					// No morph
-					if (FMath::IsNearlyZero(Factor))
-					{                        
-						AddOp(FScheduledOp(Item.At, Item, 2, DataAddress),
-							FScheduledOp(BaseAt, Item));
-					}
-					// The Morph, partial or full
-					else
-					{
-						AddOp(FScheduledOp(Item.At, Item, 2, DataAddress),
-							FScheduledOp(BaseAt, Item),
-							FScheduledOp(TargetAt, Item));
-					}
-				}
-                break;
-            }
-
-            case 2:
-            {
-       		    MUTABLE_CPUPROFILER_SCOPE(ME_MORPH_2)
-				check(EnumHasAnyFlags(EMeshContentFlags(Item.ExecutionOptions), EMeshContentFlags::GeometryData))
-				
-				TManagedPtr<const FMesh> BaseMesh = LoadMesh(FCacheAddress(BaseAt, Item));
-
-				// Factor from 0 to 1 between the two targets
-				const FScheduledOpData& OpHeapData = HeapData[Item.CustomState];
-				float Factor = OpHeapData.Interpolate.Bifactor;
-
-				if (BaseMesh)
-				{
 					if (FMath::IsNearlyZero(Factor))
 					{
 						StoreMesh(Item, BaseMesh);
 					}
 					else 
 					{
-						TManagedPtr<const FMesh> MorphMesh = LoadMesh(FCacheAddress(TargetAt, Item));
-
 						TManagedPtr<FMesh> Result = CloneOrTakeOver(MoveTemp(BaseMesh));
-						if (Result->HasMorphs())
+						if (Result && Result->HasMorphs() && Program.ConstantStrings.IsValidIndex(Args.Name))
 						{
 							// Morph data from the Base Mesh
-							FName Name(Program.ConstantStrings[NameAt]);
+							FName Name(Program.ConstantStrings[Args.Name]);
 							MeshMorph(Result.Get(), Name, Factor);
 						
 						}
-						else if (MorphMesh)
-						{
-							MeshMorph(Result.Get(), MorphMesh.Get(), Factor);
-						}
 
 						StoreMesh(Item, Result);
 					}
 				}
-				else
-				{
-					StoreMesh(Item, nullptr);
-				}
-
 
                 break;
             }
diff --git a/Engine/Plugins/Mutable/Source/MutableRuntime/Private/MuR/CodeRunner.h b/Engine/Plugins/Mutable/Source/MutableRuntime/Private/MuR/CodeRunner.h
--- a/Engine/Plugins/Mutable/Source/MutableRuntime/Private/MuR/CodeRunner.h
+++ b/Engine/Plugins/Mutable/Source/MutableRuntime/Private/MuR/CodeRunner.h
@@ -530,9 +530,12 @@ namespace  UE::Mutable::Private
 			TManagedPtr<FImage> Result;
 			if (!Resource.IsUniqueReference())
 			{
-				uint32 DataSize = Resource->GetDataSize();
-				EnsureBudgetBelow(DataSize);
-				Result = Resource->Clone();
+				if (Resource)
+				{
+					uint32 DataSize = Resource->GetDataSize();
+					EnsureBudgetBelow(DataSize);
+					Result = Resource->Clone();
+				}
 			}
 			else
 			{
@@ -556,9 +559,12 @@ namespace  UE::Mutable::Private
 			TManagedPtr<FMesh> Result;
 			if (!Resource.IsUniqueReference())
 			{
-				const int32 ResourceDataSize = Resource->GetDataSize();
-				Result = CreateMesh(ResourceDataSize); 
-				Result->CopyFrom(*Resource);
+				if (Resource)
+				{
+					const int32 ResourceDataSize = Resource->GetDataSize();
+					Result = CreateMesh(ResourceDataSize); 
+					Result->CopyFrom(*Resource);
+				}
 			}
 			else
 			{
diff --git a/Engine/Plugins/Mutable/Source/MutableRuntime/Private/MuR/CodeRunnerBegin.cpp b/Engine/Plugins/Mutable/Source/MutableRuntime/Private/MuR/CodeRunnerBegin.cpp
--- a/Engine/Plugins/Mutable/Source/MutableRuntime/Private/MuR/CodeRunnerBegin.cpp
+++ b/Engine/Plugins/Mutable/Source/MutableRuntime/Private/MuR/CodeRunnerBegin.cpp
@@ -1956,27 +1956,10 @@ namespace UE::Mutable::Private
 				{
 					check(!Item.bEvaluate);
 
-					const uint8* Data = Program.GetOpArgsPointer(Item.At);
-
-					OP::CONSTANT_STRING_INDEX NameAt = 0;
-					FMemory::Memcpy(&NameAt, Data, sizeof(OP::CONSTANT_STRING_INDEX)); 
-					Data += sizeof(OP::CONSTANT_STRING_INDEX);
-					
-					OP::ADDRESS FactorAt = 0;
-					FMemory::Memcpy(&FactorAt, Data, sizeof(OP::ADDRESS)); 
-					Data += sizeof(OP::ADDRESS);
-			
-					OP::ADDRESS BaseAt = 0;
-					FMemory::Memcpy(&BaseAt, Data, sizeof(OP::ADDRESS)); 
-					Data += sizeof(OP::ADDRESS);
+					OP::MeshMorphArgs Args = Program.GetOpArgs<OP::MeshMorphArgs>(Item.At);
 
-					OP::ADDRESS TargetAt = 0;
-					FMemory::Memcpy(&TargetAt, Data, sizeof(OP::ADDRESS)); 
-					Data += sizeof(OP::ADDRESS);
-
-					//PushOp(FScheduledOpInline(FactorAt, Item));
-					PushOp(FScheduledOpInline(BaseAt, Item));
-					PushOp(FScheduledOpInline(TargetAt, Item));
+					//PushOp(FScheduledOpInline(Args.Factor, Item));
+					PushOp(FScheduledOpInline(Args.Base, Item));
 
 					StoreNone(Item);
 					break;
diff --git a/Engine/Plugins/Mutable/Source/MutableRuntime/Private/MuR/Operations.cpp b/Engine/Plugins/Mutable/Source/MutableRuntime/Private/MuR/Operations.cpp
--- a/Engine/Plugins/Mutable/Source/MutableRuntime/Private/MuR/Operations.cpp
+++ b/Engine/Plugins/Mutable/Source/MutableRuntime/Private/MuR/Operations.cpp
@@ -686,19 +686,11 @@ namespace UE::Mutable::Private
 
         case EOpType::ME_MORPH:
         {
-			const uint8_t* data = program.GetOpArgsPointer(at);
-
-			OP::ADDRESS FactorAt = 0;
-			FMemory::Memcpy(&FactorAt, data, sizeof(OP::ADDRESS)); data += sizeof(OP::ADDRESS);
-			f(FactorAt);
+			OP::MeshMorphArgs Args = program.GetOpArgs<OP::MeshMorphArgs>(at);
 
-			OP::ADDRESS BaseAt = 0;
-			FMemory::Memcpy(&BaseAt, data, sizeof(OP::ADDRESS)); data += sizeof(OP::ADDRESS);
-			f(BaseAt);
+			f(Args.Factor);
+			f(Args.Base);
 
-			OP::ADDRESS TargetAt = 0;
-			FMemory::Memcpy(&TargetAt, data, sizeof(OP::ADDRESS)); data += sizeof(OP::ADDRESS);
-			f(TargetAt);
 			break;
         }
 
diff --git a/Engine/Plugins/Mutable/Source/MutableTools/Internal/MuT/NodeMeshMorph.h b/Engine/Plugins/Mutable/Source/MutableTools/Internal/MuT/NodeMeshMorph.h
--- a/Engine/Plugins/Mutable/Source/MutableTools/Internal/MuT/NodeMeshMorph.h
+++ b/Engine/Plugins/Mutable/Source/MutableTools/Internal/MuT/NodeMeshMorph.h
@@ -22,7 +22,6 @@ namespace UE::Mutable::Private
 		FName Name;
 		Ptr<NodeScalar> Factor;
 		Ptr<NodeMesh> Base;
-		Ptr<NodeMesh> Morph;
 
 		bool bReshapeSkeleton = false;
 		bool bReshapePhysicsVolumes = false;
diff --git a/Engine/Plugins/Mutable/Source/MutableTools/Private/MuT/ASTOpMeshExtractLayoutBlocks.cpp b/Engine/Plugins/Mutable/Source/MutableTools/Private/MuT/ASTOpMeshExtractLayoutBlocks.cpp
--- a/Engine/Plugins/Mutable/Source/MutableTools/Private/MuT/ASTOpMeshExtractLayoutBlocks.cpp
+++ b/Engine/Plugins/Mutable/Source/MutableTools/Private/MuT/ASTOpMeshExtractLayoutBlocks.cpp
@@ -227,7 +227,6 @@ namespace UE::Mutable::Private
 		{
 			Ptr<ASTOpMeshMorph> NewOp = UE::Mutable::Private::Clone<ASTOpMeshMorph>(at);
 			NewOp->Base = Visit(NewOp->Base.child(), currentSinkOp);
-			NewOp->Target = Visit(NewOp->Target.child(), currentSinkOp);
 			newAt = NewOp;
 			break;
 		}
diff --git a/Engine/Plugins/Mutable/Source/MutableTools/Private/MuT/ASTOpMeshFormat.cpp b/Engine/Plugins/Mutable/Source/MutableTools/Private/MuT/ASTOpMeshFormat.cpp
--- a/Engine/Plugins/Mutable/Source/MutableTools/Private/MuT/ASTOpMeshFormat.cpp
+++ b/Engine/Plugins/Mutable/Source/MutableTools/Private/MuT/ASTOpMeshFormat.cpp
@@ -172,62 +172,6 @@ namespace
 		return Result;
 	}
 
-
-	// Make a mesh format suitable to morph a particular other format.
-	TManagedPtr<FMesh> MakeMorphTargetFormat(TManagedPtr<const FMesh> pTargetFormat)
-	{
-		MUTABLE_CPUPROFILER_SCOPE(MakeMorphTargetFormat);
-
-		// Make a morph format by adding all the vertex channels from the base into a single
-		// vertex buffer
-
-		int32 offset = 0;
-		int32 numChannels = 0;
-		TArray<EMeshBufferSemantic> semantics;
-		TArray<int32> semanticIndices;
-		TArray<EMeshBufferFormat> formats;
-		TArray<int32> components;
-		TArray<int32> offsets;
-
-		// Add the vertex channels from the new format
-		for (int32 vb = 0; vb < pTargetFormat->GetVertexBuffers().GetBufferCount(); ++vb)
-		{
-			for (int32 c = 0; c < pTargetFormat->GetVertexBuffers().GetBufferChannelCount(vb); ++c)
-			{
-				// Channel info
-				EMeshBufferSemantic semantic;
-				int semanticIndex;
-				EMeshBufferFormat format;
-				int component;
-				pTargetFormat->GetVertexBuffers().GetChannel(vb, c, &semantic, &semanticIndex, &format, &component, nullptr);
-
-				// TODO: Filter useless semantics for morphing.
-				// Maybe some formats like the ones with a packed tangent sign need to be tweaked here, to make sense of the whole buffer.
-				semantics.Add(semantic);
-				semanticIndices.Add(semanticIndex);
-				formats.Add(format);
-				components.Add(component);
-				offsets.Add(offset);
-				offset += components[numChannels] * GetMeshFormatData(formats[numChannels]).SizeInBytes;
-				numChannels++;
-			}
-		}
-
-
-		TManagedPtr<FMesh> pTargetMorphFormat = MakeManaged<FMesh>();
-		pTargetMorphFormat->GetVertexBuffers().SetBufferCount(1);
-
-		pTargetMorphFormat->GetVertexBuffers().SetBuffer(0, offset,
-			numChannels,
-			semantics.GetData(),
-			semanticIndices.GetData(),
-			formats.GetData(),
-			components.GetData(),
-			offsets.GetData());
-
-		return pTargetMorphFormat;
-	}
-
     TManagedPtr<const FMesh> EnsureFormatHasSkinningBuffers(TManagedPtr<const FMesh>& FormatMesh)
     {
         const FMeshBufferSet& FormatMeshVertexBuffers = FormatMesh->GetVertexBuffers();
@@ -343,24 +287,6 @@ UE::Mutable::Private::Ptr<ASTOp> Sink_MeshFormatAST::Visit(const UE::Mutable::Pr
 		Ptr<ASTOpMeshMorph> NewOp = UE::Mutable::Private::Clone<ASTOpMeshMorph>(at);
 		NewOp->Base = Visit(NewOp->Base.child(), currentFormatOp);
 
-		// Reformat the morph targets to match the new format.
-		TManagedPtr<const FMesh> pTargetFormat = FindBaseMeshConstant(currentFormatOp->Format.child());
-		TManagedPtr<const FMesh> pTargetMorphFormat = MakeMorphTargetFormat(pTargetFormat);
-
-		UE::Mutable::Private::Ptr<ASTOpConstantResource> NewFormatConstant = new ASTOpConstantResource();
-		NewFormatConstant->Type = EOpType::ME_CONSTANT;
-		NewFormatConstant->SetValue(pTargetMorphFormat);
-		NewFormatConstant->SourceDataDescriptor = at->GetSourceDataDescriptor();
-
-		if (NewOp->Target)
-		{
-			UE::Mutable::Private::Ptr<ASTOpMeshFormat> newFormat = UE::Mutable::Private::Clone<ASTOpMeshFormat>(currentFormatOp);
-			newFormat->Flags = OP::MeshFormatArgs::Vertex | OP::MeshFormatArgs::IgnoreMissing;
-			newFormat->Format = NewFormatConstant;
-
-			NewOp->Target = Visit(NewOp->Target.child(), newFormat.get());
-		}
-
 		newAt = NewOp;
 		break;
 	}
diff --git a/Engine/Plugins/Mutable/Source/MutableTools/Private/MuT/ASTOpMeshMorph.cpp b/Engine/Plugins/Mutable/Source/MutableTools/Private/MuT/ASTOpMeshMorph.cpp
--- a/Engine/Plugins/Mutable/Source/MutableTools/Private/MuT/ASTOpMeshMorph.cpp
+++ b/Engine/Plugins/Mutable/Source/MutableTools/Private/MuT/ASTOpMeshMorph.cpp
@@ -14,7 +14,7 @@ namespace UE::Mutable::Private
 
 	//---------------------------------------------------------------------------------------------
 	ASTOpMeshMorph::ASTOpMeshMorph()
-		: Factor(this), Base(this), Target(this)
+		: Factor(this), Base(this)
 	{
 	}
 
@@ -28,35 +28,33 @@ namespace UE::Mutable::Private
 
 
 	//---------------------------------------------------------------------------------------------
-	bool ASTOpMeshMorph::IsEqual(const ASTOp& otherUntyped) const
+	bool ASTOpMeshMorph::IsEqual(const ASTOp& OtherUntyped) const
 	{
-		if (otherUntyped.GetOpType() == GetOpType())
+		if (OtherUntyped.GetOpType() == GetOpType())
 		{
-			const ASTOpMeshMorph* other = static_cast<const ASTOpMeshMorph*>(&otherUntyped);
-			return Factor == other->Factor && Base == other->Base && Target == other->Target;
+			const ASTOpMeshMorph* Other = static_cast<const ASTOpMeshMorph*>(&OtherUntyped);
+			return Name == Other->Name && Factor == Other->Factor && Base == Other->Base;
 		}
 		return false;
 	}
 
 
 	//---------------------------------------------------------------------------------------------
-	UE::Mutable::Private::Ptr<ASTOp> ASTOpMeshMorph::Clone(MapChildFuncRef mapChild) const
+	UE::Mutable::Private::Ptr<ASTOp> ASTOpMeshMorph::Clone(MapChildFuncRef MapChild) const
 	{
-		Ptr<ASTOpMeshMorph> n = new ASTOpMeshMorph();
-		n->Name = Name;
-		n->Factor = mapChild(Factor.child());
-		n->Base = mapChild(Base.child());
-		n->Target = mapChild(Target.child());
-		return n;
+		Ptr<ASTOpMeshMorph> New = new ASTOpMeshMorph();
+		New->Name = Name;
+		New->Factor = MapChild(Factor.child());
+		New->Base = MapChild(Base.child());
+		return New;
 	}
 
 
 	//---------------------------------------------------------------------------------------------
-	void ASTOpMeshMorph::ForEachChild(const TFunctionRef<void(ASTChild&)> f)
+	void ASTOpMeshMorph::ForEachChild(const TFunctionRef<void(ASTChild&)> Func)
 	{
-		f(Factor);
-		f(Base);
-		f(Target);
+		Func(Factor);
+		Func(Base);
 	}
 
 
@@ -68,34 +66,30 @@ namespace UE::Mutable::Private
 		Result = HashCombineFast(Result, GetTypeHash(Name));
 		Result = HashCombineFast(Result, GetTypeHash(Factor));
 		Result = HashCombineFast(Result, GetTypeHash(Base));
-		Result = HashCombineFast(Result, GetTypeHash(Target));
 		
 		return Result;
 	}
 
 
 	//---------------------------------------------------------------------------------------------
-	void ASTOpMeshMorph::Link(FProgram& program, FLinkerOptions*)
+	void ASTOpMeshMorph::Link(FProgram& Program, FLinkerOptions*)
 	{
 		// Already linked?
 		if (!linkedAddress)
 		{
-			linkedAddress = (OP::ADDRESS)program.OpAddress.Num();
+			OP::MeshMorphArgs Args;
+			FMemory::Memzero(Args);
 
-			program.OpAddress.Add((uint32_t)program.ByteCode.Num());
-			AppendCode(program.ByteCode, EOpType::ME_MORPH);
+			linkedAddress = (OP::ADDRESS)Program.OpAddress.Num();
 
-			OP::CONSTANT_STRING_INDEX NameAt = program.AddConstant(Name.ToString());
-			AppendCode(program.ByteCode, NameAt);
-			
-			OP::ADDRESS FactorAt = Factor ? Factor->linkedAddress : 0;
-			AppendCode(program.ByteCode, FactorAt);
-
-			OP::ADDRESS BaseAt = Base ? Base->linkedAddress : 0;
-			AppendCode(program.ByteCode, BaseAt);
+			Args.Name = Program.AddConstant(Name.ToString());
+			Args.Base = Base ? Base->linkedAddress : 0;
+			Args.Factor = Factor ? Factor->linkedAddress : 0;
 
-			OP::ADDRESS TargetAt = Target ? Target->linkedAddress : 0;
-			AppendCode(program.ByteCode, TargetAt);
+			Program.OpAddress.Add((uint32)Program.ByteCode.Num());
+			AppendCode(Program.ByteCode, GetOpType());
+			AppendCode(Program.ByteCode, Args);
+			
 		}
 	}
 
@@ -136,32 +130,6 @@ namespace UE::Mutable::Private
 
 		}
 
-		// If not optimized yet
-		if (!NewOp && Target)
-		{
-			// Target optimizations
-			EOpType MorphType = Target.child()->GetOpType();
-			switch (MorphType)
-			{
-
-			case EOpType::ME_ADDMETADATA:
-			{
-				// Ignore the morph target tags
-				const ASTOpMeshAddMetadata* AddMetadata = static_cast<const ASTOpMeshAddMetadata*>(Target.child().get());
-
-				Ptr<ASTOpMeshMorph> New = UE::Mutable::Private::Clone<ASTOpMeshMorph>(this);
-				New->Target = AddMetadata->Source.child();
-				NewOp = New;
-
-				break;
-			}
-
-			default:
-				break;
-
-			}
-		}
-
 		return NewOp;
 	}
 
diff --git a/Engine/Plugins/Mutable/Source/MutableTools/Private/MuT/ASTOpMeshMorph.h b/Engine/Plugins/Mutable/Source/MutableTools/Private/MuT/ASTOpMeshMorph.h
--- a/Engine/Plugins/Mutable/Source/MutableTools/Private/MuT/ASTOpMeshMorph.h
+++ b/Engine/Plugins/Mutable/Source/MutableTools/Private/MuT/ASTOpMeshMorph.h
@@ -26,9 +26,6 @@ namespace UE::Mutable::Private
 		//! Base mesh to morph.
 		ASTChild Base;
 
-		//! Target to apply on the base depending on the factor
-		ASTChild Target;
-
 	public:
 
 		ASTOpMeshMorph();
diff --git a/Engine/Plugins/Mutable/Source/MutableTools/Private/MuT/CodeGenerator.cpp b/Engine/Plugins/Mutable/Source/MutableTools/Private/MuT/CodeGenerator.cpp
--- a/Engine/Plugins/Mutable/Source/MutableTools/Private/MuT/CodeGenerator.cpp
+++ b/Engine/Plugins/Mutable/Source/MutableTools/Private/MuT/CodeGenerator.cpp
@@ -2017,38 +2017,12 @@ namespace UE::Mutable::Private
 							Ptr<ASTOpMeshMorph> MorphOp; 
 							if (OriginalMeshCopy->GetType() == NodeMeshConstant::GetStaticType())
 							{
-								// Target mesh
-								TManagedPtr<FMesh> TargetMesh = *ResolveMorphedMesh;
-								if (!TargetMesh)
-								{
-									return LastMeshResults;
-								}
-
-								// Morph generation through mesh diff
-								Ptr<ASTOpConstantResource> TargetMeshOp = new ASTOpConstantResource;
-								TargetMeshOp->Type = EOpType::ME_CONSTANT;
-								TargetMeshOp->SetValue(TargetMesh->Clone());
-								TargetMeshOp->SourceDataDescriptor = static_cast<const NodeMeshConstant*>(OriginalMeshCopy.get())->SourceDataDescriptor;
-				
-								Ptr<ASTOpMeshDifference> DiffOp;
-								{
-									DiffOp = new ASTOpMeshDifference();
-									DiffOp->Base = BaseResults.BaseMeshOp;
-									DiffOp->Target = TargetMeshOp;
-
-									// Morphing tex coords here is not supported:
-									// Generating the homogoneous UVs is difficult since we don't have the base
-									// layout yet.                       
-									DiffOp->bIgnoreTextureCoords = true;
-								}
-								
 								MorphOp = new ASTOpMeshMorph();
 								// Morph operation
 								{
 									MorphOp->Base = LastMeshResults.MeshOp;
-									MorphOp->Target = DiffOp;
 									MorphOp->Factor = FactorOp;
-									//MorphOp->Name = FName(*Edit->MeshMorph);
+									MorphOp->Name = FName(*Edit->MeshMorph);
 								}
 							}
 							else
@@ -2058,7 +2032,6 @@ namespace UE::Mutable::Private
 								MorphOp = new ASTOpMeshMorph();
 								
 								MorphOp->Base = LastMeshResults.MeshOp;
-								MorphOp->Target = nullptr;
 								MorphOp->Factor = FactorOp;
 								MorphOp->Name = FName(*Edit->MeshMorph);
 							}
diff --git a/Engine/Plugins/Mutable/Source/MutableTools/Private/MuT/CodeGenerator_Mesh.cpp b/Engine/Plugins/Mutable/Source/MutableTools/Private/MuT/CodeGenerator_Mesh.cpp
--- a/Engine/Plugins/Mutable/Source/MutableTools/Private/MuT/CodeGenerator_Mesh.cpp
+++ b/Engine/Plugins/Mutable/Source/MutableTools/Private/MuT/CodeGenerator_Mesh.cpp
@@ -140,45 +140,20 @@ namespace UE::Mutable::Private
 			return UE::Tasks::MakeCompletedTask<FMeshGenerationResult>();
         }        
 
-		if (Node->Name.IsNone() && !Node->Morph)
+		if (Node->Name.IsNone())
 		{
 			return BaseTask;
 		}
 
-		FMeshTask TargetTask;
-        {
-			FMeshGenerationStaticOptions TargetStaticOptions = StaticOptions;
-			TargetStaticOptions.ActiveTags.Empty();
-			FMeshOptionsTask TargetOptionsTask = UE::Tasks::Launch( TEXT("MutableMorphOptions"),
-				[BaseTask, InOptions]() mutable
-				{
-					FMeshGenerationDynamicOptions TargetOptions = InOptions.GetResult();
-					TargetOptions.bLayouts = false;
-
-					// We need to override the layouts with the layouts that were generated for the base to make
-					// sure that we get the correct mesh when generating the target
-					FMeshGenerationResult BaseResult = BaseTask.GetResult();
-					TargetOptions.OverrideLayouts = BaseResult.GeneratedLayouts;
-
-					return TargetOptions;
-				},
-				UE::Tasks::Prerequisites( BaseTask, InOptions )
-				);
-
-			TargetTask = GenerateMesh(TargetStaticOptions, TargetOptionsTask, Node->Morph);
-        }
-
 		return UE::Tasks::Launch(TEXT("MutableMorph"),
-			[BaseTask, TargetTask, Node, FactorOp]() mutable
+			[BaseTask, Node, FactorOp]() mutable
 			{
 				FMeshGenerationResult BaseResult = BaseTask.GetResult();
-				FMeshGenerationResult TargetResult = TargetTask.GetResult();
 
 				Ptr<ASTOpMeshMorph> OpMorph = new ASTOpMeshMorph();
 				OpMorph->Name = Node->Name;
 				OpMorph->Factor = FactorOp;
 				OpMorph->Base = BaseTask.GetResult().MeshOp;
-				OpMorph->Target = TargetTask.GetResult().MeshOp;
 
 				const bool bReshapeEnabled = Node->bReshapeSkeleton || Node->bReshapePhysicsVolumes;
 
@@ -230,7 +205,7 @@ namespace UE::Mutable::Private
 
 				return Result;
 			},
-			UE::Tasks::Prerequisites(BaseTask, TargetTask)
+			UE::Tasks::Prerequisites(BaseTask)
 		);
 	}
 
diff --git a/Engine/Plugins/Mutable/Source/MutableTools/Private/MuT/DataPacker.cpp b/Engine/Plugins/Mutable/Source/MutableTools/Private/MuT/DataPacker.cpp
--- a/Engine/Plugins/Mutable/Source/MutableTools/Private/MuT/DataPacker.cpp
+++ b/Engine/Plugins/Mutable/Source/MutableTools/Private/MuT/DataPacker.cpp
@@ -353,7 +353,6 @@ namespace UE::Mutable::Private
                 uint64 NewState = CurrentSemantics;
                 NewState |= (UINT64_C(1)<< uint32(EMeshBufferSemantic::VertexIndex));
                 RecurseWithState( op->Base.child(), NewState );
-                RecurseWithState( op->Target.child(), NewState );
                 break;
              }
 
