[Mutable] Fix possible crash in MeshApplyShape operation for streaming initial geneartion when physics or pose reshaping is enabled.

- In the case the skeleton or the physics voluems need reshaping, the shape geometry needs to be loaded but the flags passed to the sheduled op could discard geometry.
- Fortified code around shape data transformations for cases where must have data is not present.

#rnx

#Synced-CL 50931775

diff --git a/Engine/Plugins/Mutable/Source/MutableRuntime/Private/MuR/CodeRunner.cpp b/Engine/Plugins/Mutable/Source/MutableRuntime/Private/MuR/CodeRunner.cpp
--- a/Engine/Plugins/Mutable/Source/MutableRuntime/Private/MuR/CodeRunner.cpp
+++ b/Engine/Plugins/Mutable/Source/MutableRuntime/Private/MuR/CodeRunner.cpp
@@ -2446,10 +2446,20 @@ namespace UE::Mutable::Private
 			break;
 		}
 
-
 		case EOpType::ME_APPLYSHAPE:
 		{
 			OP::MeshApplyShapeArgs Args = Program.GetOpArgs<OP::MeshApplyShapeArgs>(Item.At);
+
+			EMeshBindShapeFlags ReshapeFlags = static_cast<EMeshBindShapeFlags>(Args.flags);
+
+			bool bNeedsShapeGeometry = EnumHasAnyFlags(ReshapeFlags, EMeshBindShapeFlags::ReshapeSkeleton | EMeshBindShapeFlags::ReshapePhysicsVolumes);
+
+			EMeshContentFlags ShapeContentFilter = EMeshContentFlags(Item.ExecutionOptions);
+			if (bNeedsShapeGeometry)
+			{
+				EnumAddFlags(ShapeContentFilter, EMeshContentFlags::GeometryData); 
+			}
+
 			switch (Item.Stage)
 			{
 			case 0:
@@ -2458,7 +2468,7 @@ namespace UE::Mutable::Private
 				{
 					AddOp(FScheduledOp(Item.At, Item, 1),
 						FScheduledOp(Args.mesh, Item),
-						FScheduledOp(Args.shape, Item));
+						FScheduledOp::FromOpAndOptions(Args.shape, Item, (uint8)ShapeContentFilter));
 				}
 				else
 				{
@@ -2471,9 +2481,8 @@ namespace UE::Mutable::Private
 				MUTABLE_CPUPROFILER_SCOPE(ME_APPLYSHAPE_1)
 					
 				TManagedPtr<const FMesh> BaseMesh = LoadMesh(FCacheAddress(Args.mesh, Item));
-				TManagedPtr<const FMesh> Shape = LoadMesh(FCacheAddress(Args.shape, Item));
+				TManagedPtr<const FMesh> Shape = LoadMesh(FScheduledOp::FromOpAndOptions(Args.shape, Item, (uint8)ShapeContentFilter));
 
-				EMeshBindShapeFlags ReshapeFlags = static_cast<EMeshBindShapeFlags>(Args.flags);
 				const EMeshContentFlags MeshContentFilter = (EMeshContentFlags)Item.ExecutionOptions;
 
 				if (!EnumHasAnyFlags(MeshContentFilter, EMeshContentFlags::GeometryData))
diff --git a/Engine/Plugins/Mutable/Source/MutableRuntime/Private/MuR/OpMeshReshape.h b/Engine/Plugins/Mutable/Source/MutableRuntime/Private/MuR/OpMeshReshape.h
--- a/Engine/Plugins/Mutable/Source/MutableRuntime/Private/MuR/OpMeshReshape.h
+++ b/Engine/Plugins/Mutable/Source/MutableRuntime/Private/MuR/OpMeshReshape.h
@@ -787,7 +787,7 @@ namespace UE::Mutable::Private
 			bOutSuccess = false;
 			return;
 		}
-	
+
 		// \TODO: Multiple binding data support
 		int32 BindingDataIndex = 0;
 
@@ -857,27 +857,34 @@ namespace UE::Mutable::Private
 		// Generate the temp vertex query data for the shape
 		// TODO: Vertex data copy may be saved in most cases.
 		FShapeMeshDescriptorApply ShapeDescriptor;
+		if (!ShapeMesh->GetVertexBuffers().IsDescriptor())
 		{
 			const UntypedMeshBufferIteratorConst ItPosition(ShapeMesh->GetVertexBuffers(), EMeshBufferSemantic::Position);
-			ShapeDescriptor.Positions = TArrayView<const FVector3f>(reinterpret_cast<const FVector3f*>(ItPosition.ptr()), ShapeVertexCount);
 
+			if (ItPosition.ptr())
+			{
+				ShapeDescriptor.Positions = TArrayView<const FVector3f>(reinterpret_cast<const FVector3f*>(ItPosition.ptr()), ShapeVertexCount);
+			}
 
 			const UntypedMeshBufferIteratorConst ItNormal(ShapeMesh->GetVertexBuffers(), EMeshBufferSemantic::Normal);
-			ShapeDescriptor.Normals.SetNumUninitialized(ShapeVertexCount);
-			
-			if (ItNormal.GetFormat() == EMeshBufferFormat::PackedDirS8_W_TangentSign)
+			if (ItNormal.ptr())
 			{
-				for (int32 ShapeVertexIndex = 0; ShapeVertexIndex < ShapeVertexCount; ++ShapeVertexIndex)
+				ShapeDescriptor.Normals.SetNumUninitialized(ShapeVertexCount);
+		
+				if (ItNormal.GetFormat() == EMeshBufferFormat::PackedDirS8_W_TangentSign)
 				{
-					const FPackedNormal* PackedNormal = reinterpret_cast<const FPackedNormal*>((ItNormal + ShapeVertexIndex).ptr());
-					ShapeDescriptor.Normals[ShapeVertexIndex] = PackedNormal->ToFVector3f();
+					for (int32 ShapeVertexIndex = 0; ShapeVertexIndex < ShapeVertexCount; ++ShapeVertexIndex)
+					{
+						const FPackedNormal* PackedNormal = reinterpret_cast<const FPackedNormal*>((ItNormal + ShapeVertexIndex).ptr());
+						ShapeDescriptor.Normals[ShapeVertexIndex] = PackedNormal->ToFVector3f();
+					}
 				}
-			}
-			else
-			{
-				for (int32 ShapeVertexIndex = 0; ShapeVertexIndex < ShapeVertexCount; ++ShapeVertexIndex)
+				else
 				{
-					ShapeDescriptor.Normals[ShapeVertexIndex] = (ItNormal + ShapeVertexIndex).GetAsVec3f();
+					for (int32 ShapeVertexIndex = 0; ShapeVertexIndex < ShapeVertexCount; ++ShapeVertexIndex)
+					{
+						ShapeDescriptor.Normals[ShapeVertexIndex] = (ItNormal + ShapeVertexIndex).GetAsVec3f();
+					}
 				}
 			}
 		}
@@ -886,29 +893,40 @@ namespace UE::Mutable::Private
 		// TODO: Index data copy may be saved in most cases.
 
 		TArray<UE::Geometry::FIndex3i> FormattedTriangles;
-
+		if (!ShapeMesh->GetIndexBuffers().IsDescriptor())
 		{
 			const UntypedMeshBufferIteratorConst ItIndices(ShapeMesh->GetIndexBuffers(), EMeshBufferSemantic::VertexIndex);
-			if (ItIndices.GetFormat() == EMeshBufferFormat::UInt32)
-			{
-				ShapeDescriptor.Triangles = TArrayView<const UE::Geometry::FIndex3i>(reinterpret_cast<const UE::Geometry::FIndex3i*>(ItIndices.ptr()), ShapeTriangleCount);
-			}
-			else
+			if (ItIndices.ptr())
 			{
-				FormattedTriangles.SetNumUninitialized(ShapeTriangleCount);
-				for (int32 TriangleIndex = 0; TriangleIndex < ShapeTriangleCount; ++TriangleIndex)
+				if (ItIndices.GetFormat() == EMeshBufferFormat::UInt32)
 				{
-					UE::Geometry::FIndex3i Triangle;
-					Triangle.A = int32((ItIndices + TriangleIndex * 3 + 0).GetAsUINT32());
-					Triangle.B = int32((ItIndices + TriangleIndex * 3 + 1).GetAsUINT32());
-					Triangle.C = int32((ItIndices + TriangleIndex * 3 + 2).GetAsUINT32());
+					ShapeDescriptor.Triangles = TArrayView<const UE::Geometry::FIndex3i>(reinterpret_cast<const UE::Geometry::FIndex3i*>(ItIndices.ptr()), ShapeTriangleCount);
+				}
+				else
+				{
+					FormattedTriangles.SetNumUninitialized(ShapeTriangleCount);
+					for (int32 TriangleIndex = 0; TriangleIndex < ShapeTriangleCount; ++TriangleIndex)
+					{
+						UE::Geometry::FIndex3i Triangle;
+						Triangle.A = int32((ItIndices + TriangleIndex * 3 + 0).GetAsUINT32());
+						Triangle.B = int32((ItIndices + TriangleIndex * 3 + 1).GetAsUINT32());
+						Triangle.C = int32((ItIndices + TriangleIndex * 3 + 2).GetAsUINT32());
 
-					FormattedTriangles[TriangleIndex] = Triangle;
+						FormattedTriangles[TriangleIndex] = Triangle;
+					}
+					ShapeDescriptor.Triangles = TArrayView<const UE::Geometry::FIndex3i>(FormattedTriangles);
 				}
-				ShapeDescriptor.Triangles = TArrayView<const UE::Geometry::FIndex3i>(FormattedTriangles);
 			}
 		}
 
+		// Abort the operations as we don't have a shape.
+		if (!(ShapeDescriptor.Triangles.Num() && ShapeDescriptor.Normals.Num() && ShapeDescriptor.Positions.Num()))
+		{
+			bOutSuccess = false;
+			UE_LOG(LogMutableCore, Error, TEXT("The shape mesh provided does not have proper data. Skipping reshape operation."));
+			return;
+		}
+
 		if (bReshapeVertices && BarycentricDataBuffer >= 0)
 		{
 			{
