[Mutable] Change CountBits64 for CountBits for counting mesh content flags set.

#rnx

#Synced-CL 50931775

diff --git a/Engine/Plugins/Mutable/Source/MutableRuntime/Internal/MuR/Model.h b/Engine/Plugins/Mutable/Source/MutableRuntime/Internal/MuR/Model.h
--- a/Engine/Plugins/Mutable/Source/MutableRuntime/Internal/MuR/Model.h
+++ b/Engine/Plugins/Mutable/Source/MutableRuntime/Internal/MuR/Model.h
@@ -791,7 +791,7 @@ namespace UE::Mutable::Private
 			}
 			MeshRomCurrentIndex += (int32)EnumHasAnyFlags(MeshContentRange.GetContentFlags(), EMeshContentFlags::MetaData);
 
-			check(MeshRomCurrentIndex - MeshContentRange.GetFirstIndex() == FMath::CountBits64((uint64)MeshContentRange.GetContentFlags()));
+			check(MeshRomCurrentIndex - MeshContentRange.GetFirstIndex() == FMath::CountBits((uint32)MeshContentRange.GetContentFlags()));
 
 			uint32 MeshBudgetReserve = 0;
 			MeshBudgetReserve += GeometryMesh ? GeometryMesh->GetDataSize() : 0;
diff --git a/Engine/Plugins/Mutable/Source/MutableRuntime/Private/MuR/CodeRunnerTasked.cpp b/Engine/Plugins/Mutable/Source/MutableRuntime/Private/MuR/CodeRunnerTasked.cpp
--- a/Engine/Plugins/Mutable/Source/MutableRuntime/Private/MuR/CodeRunnerTasked.cpp
+++ b/Engine/Plugins/Mutable/Source/MutableRuntime/Private/MuR/CodeRunnerTasked.cpp
@@ -2003,7 +2003,7 @@ namespace UE::Mutable::Private
 		}
 		RomContentIndex += EnumHasAnyFlags(RomContentFlags, EMeshContentFlags::MetaData);
 
-		check(RomContentIndex == FMath::CountBits64((uint64)RomContentFlags));
+		check(RomContentIndex == FMath::CountBits((uint32)RomContentFlags));
 
 		for (const int32 MeshContentRomIndex : RomsToLoad)
 		{
diff --git a/Engine/Plugins/Mutable/Source/MutableTools/Private/MuT/Compiler.cpp b/Engine/Plugins/Mutable/Source/MutableTools/Private/MuT/Compiler.cpp
--- a/Engine/Plugins/Mutable/Source/MutableTools/Private/MuT/Compiler.cpp
+++ b/Engine/Plugins/Mutable/Source/MutableTools/Private/MuT/Compiler.cpp
@@ -899,8 +899,7 @@ namespace UE::Mutable::Private
 
 				const uint32 SourceId = AdditionalData.SourceMeshPerConstant[MeshIndex].SourceId;
 
-				const int32 NumMeshContentElements = 
-						FMath::CountBits64(static_cast<uint64>(MeshContentRange.GetContentFlags()));
+				const int32 NumMeshContentElements = FMath::CountBits((uint32)MeshContentRange.GetContentFlags());
 
 				for (int32 MeshContentIndex = 0; MeshContentIndex < NumMeshContentElements; ++MeshContentIndex)
 				{
