[mutable] Fixed mutable compilation hang caused by the SKM param node + Texture Project node.
- The texture project requested the evaluation of the SKM Param node, which granted the execution of one branch that, during the compilation, wrongfully set as the next task to run in the pipe the task that was required for the current pipe to be completed first. That made the current pipe operation wait for the execution of the next pipe operation, which would never happen, as for the second task to be executed, the first one should first complete.
- Removed unnecessary use of a lock in a case where it made no sense (a lock was already in place guarding "this")
- Removed some unnecessary includes.

#prefligh 698d8b1aec78de4433a30c59

#rnx

#Synced-CL 50966636

diff --git a/Engine/Plugins/Mutable/Source/MutableTools/Private/MuT/CodeGenerator_Mesh.cpp b/Engine/Plugins/Mutable/Source/MutableTools/Private/MuT/CodeGenerator_Mesh.cpp
--- a/Engine/Plugins/Mutable/Source/MutableTools/Private/MuT/CodeGenerator_Mesh.cpp
+++ b/Engine/Plugins/Mutable/Source/MutableTools/Private/MuT/CodeGenerator_Mesh.cpp
@@ -3,18 +3,14 @@
 #include "ASTOpExternal.h"
 #include "ASTOpMeshSkeletalMeshBreak.h"
 #include "Containers/Array.h"
-#include "Math/IntPoint.h"
-#include "Math/UnrealMathSSE.h"
 #include "Misc/AssertionMacros.h"
 #include "MuR/Layout.h"
 #include "MuR/Mesh.h"
 #include "MuR/MeshPrivate.h"
 #include "MuR/MeshBufferSet.h"
 #include "MuR/OpMeshPrepareLayout.h"
-#include "MuR/MutableMath.h"
 #include "MuR/Operations.h"
 #include "MuR/Parameters.h"
-#include "MuR/Platform.h"
 #include "MuR/Ptr.h"
 #include "MuR/RefCounted.h"
 #include "MuR/System.h"
@@ -36,13 +32,10 @@
 #include "MuT/ASTOpMeshPrepareLayout.h"
 #include "MuT/ASTOpMeshAddMetadata.h"
 #include "MuT/ASTOpSwitch.h"
-#include "MuT/ASTOpReferenceResource.h"
 #include "MuT/CodeGenerator.h"
 #include "MuT/CodeGenerator_FirstPass.h"
 #include "MuT/CompilerPrivate.h"
 #include "MuT/ErrorLog.h"
-#include "MuT/NodeExternal.h"
-#include "MuT/NodeImageProject.h"
 #include "MuT/NodeLayout.h"
 #include "MuT/NodeMesh.h"
 #include "MuT/NodeMeshApplyPose.h"
@@ -61,7 +54,6 @@
 #include "MuT/NodeMeshTransform.h"
 #include "MuT/NodeMeshVariation.h"
 #include "MuT/NodeMeshSkeletalMeshBreak.h"
-#include "MuT/NodeScalar.h"
 #include "MuT/Table.h"
 #include "MuT/TablePrivate.h"
 
@@ -672,7 +664,7 @@ namespace UE::Mutable::Private
 			for (int32 LayoutIndex = 0; LayoutIndex < OriginalLayouts.Num(); ++LayoutIndex)
 			{
 				FGeneratedLayout GeneratedData;
-				Ptr<NodeLayout> LayoutNode = OriginalLayouts[LayoutIndex];
+				Ptr<const NodeLayout> LayoutNode = OriginalLayouts[LayoutIndex];
 				if (LayoutNode)
 				{
 					GeneratedData.Source = LayoutNode;
@@ -1417,8 +1409,7 @@ namespace UE::Mutable::Private
 
 		Ptr<ASTOp> SkeletalMeshParameter = SkeletalMeshResult.SkeletalMeshOp;
 		
-		// Local pipe because we call GenerateRange in there.
-		FMeshTask MeshParamTask = this->LocalPipe.Launch(TEXT("MutableMeshParameter"),
+		FMeshTask MeshParamTask = UE::Tasks::Launch(TEXT("MutableMeshParameter"),
 			[InOptions, Node, ReferenceMesh, StaticOptions, SkeletalMeshParameter, this]() mutable
 			{
 				FMeshGenerationDynamicOptions Options = InOptions.GetResult();
@@ -1443,19 +1434,12 @@ namespace UE::Mutable::Private
 					GenerateLayoutsForMesh(*ReferenceMesh, Node->Layouts);
 				}
 
-				{
-					UE::TUniqueLock Lock(FirstPass.ParameterNodes.Mutex);
-					
-					// Fill the result structure. Some ops will be completed outside the lock.
-					Ptr<ASTOp> LastMeshOp = GenerateLayoutOpsAndResult(Options, Op, Node->Layouts, Op->MeshID, Result);
-					Result.MeshOp = LastMeshOp;
-					Result.BaseMeshOp = LastMeshOp;
-				}
-
-				int32 LODIndex = StaticOptions.LODIndex;
-
-				FMeshGenerationDynamicOptions DynamicOptions = InOptions.GetResult();
 				
+				// Fill the result structure. Some ops will be completed outside the lock.
+				Ptr<ASTOp> LastMeshOp = GenerateLayoutOpsAndResult(Options, Op, Node->Layouts, Op->MeshID, Result);
+				Result.MeshOp = LastMeshOp;
+				Result.BaseMeshOp = LastMeshOp;
+			
 				TArray<FString> RealTimeMorphNames;
 				RealTimeMorphNames.Reserve(Node->RealTimeMorphNames.Num());
 				Algo::Transform(Node->RealTimeMorphNames, RealTimeMorphNames, [](const FName& Name) { return Name.ToString(); });
@@ -1478,6 +1462,12 @@ namespace UE::Mutable::Private
 			Requisites
 		);
 
+		
+		if (UE::Tasks::Private::GetCurrentTask())
+		{
+			UE::Tasks::AddNested(MeshParamTask);
+		}
+		
 		// Apply the modifier for the pre-normal operations stage.
 		TArray<FirstPassGenerator::FModifier> Modifiers;
 		constexpr bool bModifiersForBeforeOperations = true;
