[mutable] Removed the Material Overlay property from the Skeletal Mesh Component node.
- Added a method to be able to get the material output pin of the Material Constant node. Backwards compatibility code has been added for completeness.
- Made the Overlay Material Pin of the Skeletal Mesh Component node be stored as a reference and added backwards compatibility code to handle the setup of the variable holding said reference.
- Made the Skeletal Mesh Component nodes that had an Overlay Material property filled and no overlay material pin connected add a new material constant node and use it as the host of the material set in the node Overlay Material now deprecated property.

#rnx

#Synced-CL 50966636

diff --git a/Engine/Plugins/Mutable/Source/CustomizableObject/Internal/MuCO/CustomizableObjectCustomVersion.h b/Engine/Plugins/Mutable/Source/CustomizableObject/Internal/MuCO/CustomizableObjectCustomVersion.h
--- a/Engine/Plugins/Mutable/Source/CustomizableObject/Internal/MuCO/CustomizableObjectCustomVersion.h
+++ b/Engine/Plugins/Mutable/Source/CustomizableObject/Internal/MuCO/CustomizableObjectCustomVersion.h
@@ -221,7 +221,13 @@ struct FCustomizableObjectCustomVersion
 		PassthroughSkeletalMeshComponentOverrideSubtypeFix,
 		
 		PassthroughSkeletalMeshComponentComponentMaterialOverride,
+		
+		StoreMaterialPinReferenceFromConstantMaterialNode,
+		
+		StoreOverlayMaterialPinAsReference,
 
+		RemovedOverlayMaterialUpropertyFromSkeletalMeshComponentNode,
+		
 		// -----<new versions can be added above this line>-------------------------------------------------
 		VersionPlusOne,
 		LatestVersion = VersionPlusOne - 1
diff --git a/Engine/Plugins/Mutable/Source/CustomizableObjectEditor/Private/MuCOE/GenerateMutableSource/GenerateMutableSourceComponent.cpp b/Engine/Plugins/Mutable/Source/CustomizableObjectEditor/Private/MuCOE/GenerateMutableSource/GenerateMutableSourceComponent.cpp
--- a/Engine/Plugins/Mutable/Source/CustomizableObjectEditor/Private/MuCOE/GenerateMutableSource/GenerateMutableSourceComponent.cpp
+++ b/Engine/Plugins/Mutable/Source/CustomizableObjectEditor/Private/MuCOE/GenerateMutableSource/GenerateMutableSourceComponent.cpp
@@ -71,20 +71,14 @@ UE::Mutable::Private::Ptr<UE::Mutable::Private::NodeComponent> GenerateMutableSo
 			UE::Mutable::Private::Ptr<UE::Mutable::Private::NodeComponentNew> NodeComponentNew = new UE::Mutable::Private::NodeComponentNew();
 			NodeComponentNew->Id = ComponentInfo->ComponentId;
 			NodeComponentNew->SetMessageContext(Node);
-		
-			UEdGraphPin* MaterialAssetPin = TypedComponentMesh->GetOverlayMaterialAssetPin();
-			if (const UEdGraphPin* ConnectedPin = MaterialAssetPin ? FollowInputPin(*MaterialAssetPin) : nullptr)
+
+			const FEdGraphPinReference MaterialAssetPin = TypedComponentMesh->GetOverlayMaterialAssetPin();
+			if (const UEdGraphPin* ConnectedPin = MaterialAssetPin.Get() ? FollowInputPin(*MaterialAssetPin.Get()) : nullptr)
 			{
 				GenerationContext.CurrentMaterialParameterId = ConnectedPin->PinId.ToString();
 				UE::Mutable::Private::Ptr<UE::Mutable::Private::NodeMaterial> MaterialNode = GenerateMutableSourceMaterial(ConnectedPin, GenerationContext, FGenerationMaterialOptions());
 				NodeComponentNew->OverlayMaterial = MaterialNode;
 			}
-			else if(UMaterialInterface* OverlayMaterial = GenerationContext.LoadObject(TypedComponentMesh->GetOverlayMaterial()))
-			{
-				UE::Mutable::Private::Ptr<UE::Mutable::Private::NodeMaterialConstant> MaterialNode = new UE::Mutable::Private::NodeMaterialConstant();
-				MaterialNode->MaterialId = GenerationContext.CompilationContext->PassthroughObjectFactory.Add(*OverlayMaterial);
-				NodeComponentNew->OverlayMaterial = MaterialNode;
-			}
 
 			Result = NodeComponentNew;
 
diff --git a/Engine/Plugins/Mutable/Source/CustomizableObjectEditor/Private/MuCOE/Nodes/CONodeMaterialConstant.cpp b/Engine/Plugins/Mutable/Source/CustomizableObjectEditor/Private/MuCOE/Nodes/CONodeMaterialConstant.cpp
--- a/Engine/Plugins/Mutable/Source/CustomizableObjectEditor/Private/MuCOE/Nodes/CONodeMaterialConstant.cpp
+++ b/Engine/Plugins/Mutable/Source/CustomizableObjectEditor/Private/MuCOE/Nodes/CONodeMaterialConstant.cpp
@@ -1,6 +1,8 @@
 // Copyright Epic Games, Inc. All Rights Reserved.
 
 #include "MuCOE/Nodes/CONodeMaterialConstant.h"
+
+#include "MuCO/CustomizableObjectCustomVersion.h"
 #include "MuCOE/CustomizableObjectGraph.h"
 #include "MuCOE/EdGraphSchema_CustomizableObject.h"
 
@@ -13,8 +15,26 @@ struct FSlateBrush;
 #define LOCTEXT_NAMESPACE "CustomizableObjectEditor"
 
 void UCONodeMaterialConstant::AllocateDefaultPins(UCustomizableObjectNodeRemapPins* RemapPins)
+{ 
+	MaterialPin = CustomCreatePin(EGPD_Output, UEdGraphSchema_CustomizableObject::PC_Material);
+}
+
+
+void UCONodeMaterialConstant::BackwardsCompatibleFixup(int32 CustomizableObjectCustomVersion)
 {
-	CustomCreatePin(EGPD_Output, UEdGraphSchema_CustomizableObject::PC_Material);
+	Super::BackwardsCompatibleFixup(CustomizableObjectCustomVersion);
+	
+	if (CustomizableObjectCustomVersion == FCustomizableObjectCustomVersion::StoreMaterialPinReferenceFromConstantMaterialNode)
+	{
+		for (const UEdGraphPin* NodePin : Pins)
+		{
+			if (NodePin->PinType.PinCategory == UEdGraphSchema_CustomizableObject::PC_Material)
+			{
+				MaterialPin = NodePin;
+				break;
+			}
+		}
+	}
 }
 
 
diff --git a/Engine/Plugins/Mutable/Source/CustomizableObjectEditor/Private/MuCOE/Nodes/CONodeMaterialConstant.h b/Engine/Plugins/Mutable/Source/CustomizableObjectEditor/Private/MuCOE/Nodes/CONodeMaterialConstant.h
--- a/Engine/Plugins/Mutable/Source/CustomizableObjectEditor/Private/MuCOE/Nodes/CONodeMaterialConstant.h
+++ b/Engine/Plugins/Mutable/Source/CustomizableObjectEditor/Private/MuCOE/Nodes/CONodeMaterialConstant.h
@@ -17,6 +17,9 @@ public:
 
 	GENERATED_BODY()
 
+	// UCustomizableObjectNode interface
+	virtual void BackwardsCompatibleFixup(int32 CustomizableObjectCustomVersion) override;
+	
 	/** Value of the Node. */
 	UPROPERTY(EditAnywhere, Category = CustomizableObject)
 	TSoftObjectPtr<UMaterialInterface> Material;
@@ -29,6 +32,7 @@ public:
 	// UCustomizableObjectNode interface
 	UE_API virtual void AllocateDefaultPins(UCustomizableObjectNodeRemapPins* RemapPins) override;
 
+	FEdGraphPinReference MaterialPin;
 };
 
 #undef UE_API
diff --git a/Engine/Plugins/Mutable/Source/CustomizableObjectEditor/Private/MuCOE/Nodes/CustomizableObjectNodeComponentMesh.cpp b/Engine/Plugins/Mutable/Source/CustomizableObjectEditor/Private/MuCOE/Nodes/CustomizableObjectNodeComponentMesh.cpp
--- a/Engine/Plugins/Mutable/Source/CustomizableObjectEditor/Private/MuCOE/Nodes/CustomizableObjectNodeComponentMesh.cpp
+++ b/Engine/Plugins/Mutable/Source/CustomizableObjectEditor/Private/MuCOE/Nodes/CustomizableObjectNodeComponentMesh.cpp
@@ -2,7 +2,9 @@
 
 #include "MuCOE/Nodes/CustomizableObjectNodeComponentMesh.h"
 
+#include "CONodeMaterialConstant.h"
 #include "MuCO/CustomizableObjectCustomVersion.h"
+#include "MuCOE/CustomizableObjectSchemaActions.h"
 #include "MuCOE/EdGraphSchema_CustomizableObject.h"
 #include "MuCOE/GraphTraversal.h"
 
@@ -55,6 +57,50 @@ void UCustomizableObjectNodeComponentMesh::BackwardsCompatibleFixup(int32 Custom
 			ReferenceSkeletalMesh = Result->ReferenceSkeletalMesh;
 		}
 	}
+	else if (CustomizableObjectCustomVersion == FCustomizableObjectCustomVersion::StoreOverlayMaterialPinAsReference)
+	{
+		const FString PinFriendlyName = TEXT("Overlay Material");
+		const FString PinName = PinFriendlyName + FString(TEXT("_Input_Pin"));
+
+		OverlayMaterialPin = FindPin(PinName);
+		if (!OverlayMaterialPin.Get())
+		{
+			OverlayMaterialPin = FindPin(PinFriendlyName);
+		}
+	}
+	else if (CustomizableObjectCustomVersion == FCustomizableObjectCustomVersion::RemovedOverlayMaterialUpropertyFromSkeletalMeshComponentNode)
+	{
+		UEdGraphPin* OverlayMaterialPinPtr = OverlayMaterialPin.Get();
+		// Only add a new material constant node if a node is not already been used to determine the overlay material
+		if (OverlayMaterialPinPtr && !OverlayMaterialPinPtr->LinkedTo.IsEmpty())
+		{
+			return;
+		}
+		
+		PRAGMA_DISABLE_DEPRECATION_WARNINGS
+		if (!OverlayMaterial_DEPRECATED.IsNull())
+		PRAGMA_ENABLE_DEPRECATION_WARNINGS
+		{
+			UEdGraph* Graph = GetGraph();
+			check(Graph)
+			
+			// Set node to the left of this node
+			const FVector2D MaterialConstantNodePos = FVector2D(GetNodePosX() - 200, GetNodePosY());
+			UCONodeMaterialConstant* MaterialConstantNode = Cast<UCONodeMaterialConstant>(FCustomizableObjectSchemaAction_NewNode::CreateNode(Graph, nullptr, MaterialConstantNodePos, Cast<UEdGraphNode>(NewObject<UCONodeMaterialConstant>())));
+			
+			PRAGMA_DISABLE_DEPRECATION_WARNINGS
+			MaterialConstantNode->Material = OverlayMaterial_DEPRECATED;
+			PRAGMA_ENABLE_DEPRECATION_WARNINGS
+
+			// Connect the node with this one
+			UEdGraphPin* MaterialConstantNodeMaterialPin = MaterialConstantNode->MaterialPin.Get();
+			check(MaterialConstantNodeMaterialPin);
+			
+			MaterialConstantNodeMaterialPin->MakeLinkTo(OverlayMaterialPinPtr);
+			Graph->NotifyNodeChanged(MaterialConstantNode);
+			Graph->NotifyNodeChanged(this);
+		}
+	}
 }
 
 
@@ -63,8 +109,8 @@ void UCustomizableObjectNodeComponentMesh::AllocateDefaultPins(UCustomizableObje
 	Super::AllocateDefaultPins(RemapPins);
 
 	// Mesh Component Pins
-	UEdGraphPin* OverlayMaterialPin = CustomCreatePin(EGPD_Input, UEdGraphSchema_CustomizableObject::PC_Material, TEXT("Overlay Material_Input_Pin"), LOCTEXT("OverlayMaterial", "Overlay Material"));
-	OverlayMaterialPin->PinToolTip = "Pin for an Overlay Material from a Table Node";
+	OverlayMaterialPin = CustomCreatePin(EGPD_Input, UEdGraphSchema_CustomizableObject::PC_Material, TEXT("Overlay Material_Input_Pin"), LOCTEXT("OverlayMaterial", "Overlay Material"));
+	OverlayMaterialPin.Get()->PinToolTip = "Pin for an Overlay Material from a Table Node"; 
 	
 	// Base Mesh Interface Pins
 	LODPins.Empty(NumLODs);
@@ -86,25 +132,11 @@ FText UCustomizableObjectNodeComponentMesh::GetNodeTitle(ENodeTitleType::Type Ti
 }
 
 
-TSoftObjectPtr<UMaterialInterface> UCustomizableObjectNodeComponentMesh::GetOverlayMaterial() const
-{
-	return OverlayMaterial;
-}
 
 
-UEdGraphPin* UCustomizableObjectNodeComponentMesh::GetOverlayMaterialAssetPin() const
+FEdGraphPinReference UCustomizableObjectNodeComponentMesh::GetOverlayMaterialAssetPin() const
 {
-	const FString PinFriendlyName = TEXT("Overlay Material");
-	const FString PinName = PinFriendlyName + FString(TEXT("_Input_Pin"));
-
-	if (UEdGraphPin* Pin = FindPin(PinName))
-	{
-		return Pin;
-	}
-	else
-	{
-		return FindPin(PinFriendlyName);
-	}
+	return OverlayMaterialPin.Get();
 }
 
 
diff --git a/Engine/Plugins/Mutable/Source/CustomizableObjectEditor/Private/MuCOE/Nodes/CustomizableObjectNodeComponentMesh.h b/Engine/Plugins/Mutable/Source/CustomizableObjectEditor/Private/MuCOE/Nodes/CustomizableObjectNodeComponentMesh.h
--- a/Engine/Plugins/Mutable/Source/CustomizableObjectEditor/Private/MuCOE/Nodes/CustomizableObjectNodeComponentMesh.h
+++ b/Engine/Plugins/Mutable/Source/CustomizableObjectEditor/Private/MuCOE/Nodes/CustomizableObjectNodeComponentMesh.h
@@ -68,8 +68,7 @@ public:
 	UE_API bool IsSingleOutputNode() const;
 
 	// Own interface
-	UE_API TSoftObjectPtr<UMaterialInterface> GetOverlayMaterial() const;
-	UE_API UEdGraphPin* GetOverlayMaterialAssetPin() const;
+	UE_API FEdGraphPinReference GetOverlayMaterialAssetPin() const;
 
 	// ComponentMesh Interface
 	UE_API virtual int32 GetNumLODs() override;
@@ -97,8 +96,8 @@ public:
 	UPROPERTY(EditAnywhere, Category = ComponentMesh)
 	TArray<FLODReductionSettings> LODReductionSettings;
 
-	UPROPERTY(EditAnywhere, Category = ComponentMesh)
-	TSoftObjectPtr<UMaterialInterface> OverlayMaterial;
+	UPROPERTY()
+	TSoftObjectPtr<UMaterialInterface> OverlayMaterial_DEPRECATED;
 
 	UPROPERTY(EditAnywhere, Category = ComponentMesh)
 	FMutableLODSettings LODSettings;
@@ -115,6 +114,8 @@ public:
 	UPROPERTY()
 	TArray<FEdGraphPinReference> LODPins;
 
+	UPROPERTY()
+	FEdGraphPinReference OverlayMaterialPin;
 };
 
 #undef UE_API
