[Mutable] Fix ME_EXTRACTLAYOUTBLOCK not removing morph vertices.

- Now ME_EXTRACTLAYOUTBLOCK only selects what vertices needs to be removed and the actual remove work is done by the MeshRemoveVerticesWithCullSet.
- Now triangles with vertices in and out of the layout block are not culled.
- Moved OpMeshExtractLayoutBlock implementation to its own compilation unit.

#rnx

#Synced-CL 50939624

diff --git a/Engine/Plugins/Mutable/Source/MutableRuntime/Private/MuR/CodeRunner.cpp b/Engine/Plugins/Mutable/Source/MutableRuntime/Private/MuR/CodeRunner.cpp
--- a/Engine/Plugins/Mutable/Source/MutableRuntime/Private/MuR/CodeRunner.cpp
+++ b/Engine/Plugins/Mutable/Source/MutableRuntime/Private/MuR/CodeRunner.cpp
@@ -1804,28 +1804,28 @@ namespace UE::Mutable::Private
 
         case EOpType::ME_EXTRACTLAYOUTBLOCK:
         {
-            const uint8* data = Program.GetOpArgsPointer(Item.At);
+            const uint8* Data = Program.GetOpArgsPointer(Item.At);
 
-            OP::ADDRESS source;
-            FMemory::Memcpy( &source, data, sizeof(OP::ADDRESS) );
-            data += sizeof(OP::ADDRESS);
+            OP::ADDRESS SourceAddress;
+            FMemory::Memcpy(&SourceAddress, Data, sizeof(OP::ADDRESS));
+            Data += sizeof(OP::ADDRESS);
 
             uint16 LayoutIndex;
-			FMemory::Memcpy( &LayoutIndex, data, sizeof(uint16) );
-            data += sizeof(uint16);
+			FMemory::Memcpy(&LayoutIndex, Data, sizeof(uint16));
+            Data += sizeof(uint16);
 
-            uint16 blockCount;
-			FMemory::Memcpy( &blockCount, data, sizeof(uint16) );
-            data += sizeof(uint16);
+            uint16 BlockCount;
+			FMemory::Memcpy(&BlockCount, Data, sizeof(uint16));
+            Data += sizeof(uint16);
 
             switch (Item.Stage)
             {
             case 0:
 			{
-				if (source)
+				if (SourceAddress)
 				{
 					AddOp(FScheduledOp(Item.At, Item, 1),
-							FScheduledOp(source, Item));
+							FScheduledOp(SourceAddress, Item));
 				}
 				else
 				{
@@ -1839,26 +1839,29 @@ namespace UE::Mutable::Private
 
 				if (!EnumHasAnyFlags(EMeshContentFlags(Item.ExecutionOptions), EMeshContentFlags::GeometryData))
 				{
-					TManagedPtr<const FMesh> SourceMesh = LoadMesh(FCacheAddress(source, Item));
+					TManagedPtr<const FMesh> SourceMesh = LoadMesh(FCacheAddress(SourceAddress, Item));
 					StoreMesh(Item, SourceMesh);
 				}
 				else
 				{
-					TManagedPtr<const FMesh> Source = LoadMesh(FCacheAddress(source, Item));
+					TManagedPtr<const FMesh> Source = LoadMesh(FCacheAddress(SourceAddress, Item));
 
 					// Access with memcpy necessary for unaligned memory access issues.
-					uint64 blocks[512];
-					check(blockCount< 512);
-					FMemory::Memcpy(blocks, data, sizeof(uint64)*FMath::Min(512,int32(blockCount)));
+					uint64 Blocks[512];
+					check(BlockCount < 512);
+
+					BlockCount = FMath::Min(512, int32(BlockCount));
+					FMemory::Memcpy(Blocks, Data, sizeof(uint64)*BlockCount);
 
 					if (Source)
 					{
 						TManagedPtr<FMesh> Result = CreateMesh();
 						bool bOutSuccess;
 
-						if (blockCount > 0)
+						if (BlockCount > 0)
 						{
-							MeshExtractLayoutBlock(Result.Get(), Source.Get(), LayoutIndex, blockCount, blocks, bOutSuccess);
+							TArrayView<uint64> BlockIds(Blocks, BlockCount);
+							MeshExtractLayoutBlock(Result.Get(), Source.Get(), LayoutIndex, BlockIds, bOutSuccess);
 						}
 						else
 						{
diff --git a/Engine/Plugins/Mutable/Source/MutableRuntime/Private/MuR/OpMeshExtractLayoutBlock.cpp b/Engine/Plugins/Mutable/Source/MutableRuntime/Private/MuR/OpMeshExtractLayoutBlock.cpp
new file mode 100644
--- /dev/null
+++ b/Engine/Plugins/Mutable/Source/MutableRuntime/Private/MuR/OpMeshExtractLayoutBlock.cpp
@@ -0,0 +1,100 @@
+// Copyright Epic Games, Inc. All Rights Reserved.
+
+#include "MuR/OpMeshExtractLayoutBlock.h"
+
+#include "MuR/MeshPrivate.h"
+#include "MuR/OpMeshRemove.h"
+
+namespace UE::Mutable::Private
+{
+	void MeshExtractLayoutBlock(FMesh* Result, const FMesh* Source, uint32 LayoutIndex, TArrayView<uint64> BlockIds, bool& bOutSuccess)
+	{
+		check(Source);
+		bOutSuccess = true;
+
+		// TODO: Optimise
+		Result->CopyFrom(*Source);
+
+		UntypedMeshBufferIteratorConst ItBlocks(Source->GetVertexBuffers(), EMeshBufferSemantic::LayoutBlock, LayoutIndex);
+
+		if (ItBlocks.GetFormat() != EMeshBufferFormat::None)
+		{
+			int32 ResultVertices = 0;
+
+			TBitArray<> VerticesToCull;
+			VerticesToCull.SetNumUninitialized(Source->GetVertexCount());
+
+			if (ItBlocks.GetFormat() == EMeshBufferFormat::UInt16)
+			{
+				TArrayView<const uint16> VertexBlockIdsView(reinterpret_cast<const uint16*>(ItBlocks.ptr()), Source->GetVertexCount());
+
+				const uint64 AbsoluteVertexBlockIdHighBits = uint64(Source->MeshIDPrefix) << 32;
+				for (int32 VertexIndex = VertexBlockIdsView.Num() - 1; VertexIndex >= 0; --VertexIndex)
+				{
+					const uint64 AbsoluteVertexBlockId = AbsoluteVertexBlockIdHighBits | VertexBlockIdsView[VertexIndex];
+					VerticesToCull[VertexIndex] = BlockIds.Find(AbsoluteVertexBlockId) == INDEX_NONE;
+				}
+			}
+			else if (ItBlocks.GetFormat() == EMeshBufferFormat::UInt64)
+			{
+				TArrayView<const uint64> VertexBlockIdsView(reinterpret_cast<const uint64*>(ItBlocks.ptr()), Source->GetVertexCount());
+				for (int32 VertexIndex = VertexBlockIdsView.Num() - 1; VertexIndex >= 0; --VertexIndex)
+				{
+					VerticesToCull[VertexIndex] = BlockIds.Find(VertexBlockIdsView[VertexIndex]) == INDEX_NONE;
+				}
+			}
+			else
+			{
+				check(false);
+			}
+
+			constexpr bool bRemoveIfAllVerticesCulled = true;
+			MeshRemoveVerticesWithCullSet(Result, VerticesToCull, bRemoveIfAllVerticesCulled); 
+		}
+	}
+
+	void MeshExtractLayoutBlock(FMesh* Result, const FMesh* Source, uint32 LayoutIndex, bool& bOutSuccess)
+	{
+		check(Source);
+		bOutSuccess = true;
+
+		// TODO: Optimise
+		Result->CopyFrom(*Source);
+
+		UntypedMeshBufferIteratorConst ItBlocks(Source->GetVertexBuffers(), EMeshBufferSemantic::LayoutBlock, LayoutIndex);
+
+		if (ItBlocks.GetFormat() != EMeshBufferFormat::None)
+		{
+			int32 ResultVertices = 0;
+
+			TBitArray<> VerticesToCull;
+			VerticesToCull.SetNumUninitialized(Source->GetVertexCount());
+
+			if (ItBlocks.GetFormat() == EMeshBufferFormat::UInt16)
+			{
+				TArrayView<const uint16> VertexBlockIdsView(reinterpret_cast<const uint16*>(ItBlocks.ptr()), Source->GetVertexCount());
+
+				for (int32 VertexIndex = VertexBlockIdsView.Num() - 1; VertexIndex >= 0; --VertexIndex)
+				{
+					VerticesToCull[VertexIndex] = VertexBlockIdsView[VertexIndex] == std::numeric_limits<uint16>::max();
+				}
+			}
+			else if (ItBlocks.GetFormat() == EMeshBufferFormat::UInt64)
+			{
+				TArrayView<const uint64> VertexBlockIdsView(reinterpret_cast<const uint64*>(ItBlocks.ptr()), Source->GetVertexCount());
+				for (int32 VertexIndex = VertexBlockIdsView.Num() - 1; VertexIndex >= 0; --VertexIndex)
+				{
+					VerticesToCull[VertexIndex] = VertexBlockIdsView[VertexIndex] == std::numeric_limits<uint64>::max();
+				}
+			}
+			else
+			{
+				check(false);
+			}
+
+			constexpr bool bRemoveIfAllVerticesCulled = true;
+			MeshRemoveVerticesWithCullSet(Result, VerticesToCull, bRemoveIfAllVerticesCulled); 
+		}
+	}
+
+}
diff --git a/Engine/Plugins/Mutable/Source/MutableRuntime/Private/MuR/OpMeshExtractLayoutBlock.h b/Engine/Plugins/Mutable/Source/MutableRuntime/Private/MuR/OpMeshExtractLayoutBlock.h
--- a/Engine/Plugins/Mutable/Source/MutableRuntime/Private/MuR/OpMeshExtractLayoutBlock.h
+++ b/Engine/Plugins/Mutable/Source/MutableRuntime/Private/MuR/OpMeshExtractLayoutBlock.h
@@ -2,262 +2,16 @@
 
 #pragma once
 
-#include "MuR/MeshPrivate.h"
-#include "MuR/Platform.h"
-
-#include "MuR/OpMeshRemove.h"
+#include "HAL/Platform.h"
+#include "Containers/ArrayView.h"
 
 namespace UE::Mutable::Private
 {
-    inline void MeshExtractFromVertices( const FMesh* Source,
-                                         FMesh* Result,
-                                         const TArray<int32>& OldToNew,
-                                         const TArray<int32>& NewToOld )
-    {
-        int32 ResultVertices = NewToOld.Num();
-
-        // Assemble the new vertex buffer
-        Result->GetVertexBuffers().SetElementCount(ResultVertices);
-        for (int32 BufferIndex = 0; BufferIndex < Result->GetVertexBuffers().GetBufferCount(); ++BufferIndex)
-        {
-            const uint8* pSourceData = Source->GetVertexBuffers().GetBufferData(BufferIndex);
-            uint8* pDest = Result->GetVertexBuffers().GetBufferData(BufferIndex);
-            int32 Size = Result->GetVertexBuffers().GetElementSize(BufferIndex);
-
-            for (int32 NewIndex = 0; NewIndex < ResultVertices; ++NewIndex)
-            {
-                int32 OldIndex = NewToOld[NewIndex];
-                FMemory::Memcpy(pDest, pSourceData + Size*OldIndex, Size);
-
-                pDest += Size;
-            }
-        }
-
-		// If the vertices are explicit or relative, the above operation will already handle them correctly
-		// Otherwise, create a relative vertex ID buffer if necessary
-		Result->MeshIDPrefix = Source->MeshIDPrefix;
-		if (Source->AreVertexIdsImplicit() 
-			&& 
-			// If we extract everything, we can keep the ids implicit.
-			ResultVertices!=Source->GetVertexCount())
-		{
-			// Add a new buffer
-			int32 NewBuffer = Result->VertexBuffers.GetBufferCount();
-			Result->VertexBuffers.SetBufferCount(NewBuffer + 1);
-			EMeshBufferSemantic Semantic = EMeshBufferSemantic::VertexIndex;
-			int32 SemanticIndex = 0;
-			EMeshBufferFormat Format = EMeshBufferFormat::UInt32;
-			int32 Components = 1;
-			int32 Offset = 0;
-			Result->VertexBuffers.SetBuffer(NewBuffer, sizeof(uint32), 1, &Semantic, &SemanticIndex, &Format, &Components, &Offset);
-			uint32* pIdData = reinterpret_cast<uint32*>(Result->VertexBuffers.GetBufferData(NewBuffer));
-
-			for (int32 NewIndex = 0; NewIndex < ResultVertices; ++NewIndex)
-			{
-				uint32 OldIndex = NewToOld[NewIndex];
-				(*pIdData++) = OldIndex;
-			}
-		}
-
-        // Assemble the new index buffers
-		TBitArray<> UsedSourceFaces;
-
-		UsedSourceFaces.SetNum(Source->GetFaceCount(), false);
-        UntypedMeshBufferIteratorConst itIndex(Source->GetIndexBuffers(), EMeshBufferSemantic::VertexIndex);
-        UntypedMeshBufferIterator itResultIndex(Result->GetIndexBuffers(), EMeshBufferSemantic::VertexIndex);
-        
-        int32 IndexCount = 0;
-        if (itIndex.GetFormat() == EMeshBufferFormat::UInt32)
-        {
-            const uint32* pIndices = reinterpret_cast<const uint32*>(itIndex.ptr());
-            uint32* pDestIndices = reinterpret_cast<uint32*>(itResultIndex.ptr());
-            for (int32 i = 0; i < Source->GetIndexCount()/3; ++i)
-            {
-                if (OldToNew[pIndices[i*3 + 0]] >= 0 && OldToNew[pIndices[i*3 + 1]] >= 0 && OldToNew[pIndices[i*3 + 2]] >= 0)
-                {
-                    UsedSourceFaces[i] = true;
-
-                    // Clamp in case triangles go across blocks
-                    pDestIndices[IndexCount++] = FMath::Max(0, OldToNew[pIndices[i*3 + 0]]);
-                    pDestIndices[IndexCount++] = FMath::Max(0, OldToNew[pIndices[i*3 + 1]]);
-                    pDestIndices[IndexCount++] = FMath::Max(0, OldToNew[pIndices[i*3 + 2]]);
-                }
-            }
-        }
-        else if (itIndex.GetFormat() == EMeshBufferFormat::UInt16)
-        {
-            const uint16* pIndices = reinterpret_cast<const uint16*>(itIndex.ptr());
-            uint16* pDestIndices = reinterpret_cast<uint16*>(itResultIndex.ptr());
-            for (int32 i = 0; i < Source->GetIndexCount()/3; ++i)
-            {
-                if (OldToNew[pIndices[i*3 + 0]] >= 0 && OldToNew[pIndices[i*3 + 1]] >= 0 && OldToNew[pIndices[i*3 + 2]] >= 0)
-                {
-                    UsedSourceFaces[i] = true;
-
-                    // Clamp in case triangles go across blocks
-                    pDestIndices[IndexCount++] = (uint16)FMath::Max(0, OldToNew[pIndices[i*3 + 0]]);
-                    pDestIndices[IndexCount++] = (uint16)FMath::Max(0, OldToNew[pIndices[i*3 + 1]]);
-                    pDestIndices[IndexCount++] = (uint16)FMath::Max(0, OldToNew[pIndices[i*3 + 2]]);
-                }
-            }
-        }
-        else
-        {
-            check(false);
-        }
-     
-        Result->GetIndexBuffers().SetElementCount(IndexCount);
-
-        const int32 NumOriginalVerts = OldToNew.Num();
-
-        TBitArray<> UsedVertices;
-        UsedVertices.SetNum(NumOriginalVerts, false);
-
-        for (int32 I = 0; I < NumOriginalVerts; ++I)
-        {
-            UsedVertices[I] = OldToNew[I] >= 0;
-        }
-
-        MeshRemoveRecreateSurface(Result, UsedVertices, UsedSourceFaces);
-    }
-
+    class FMesh;
 
 	/** Extract all vertices that have a layout block from the passed list on the given channel. */
-    inline void MeshExtractLayoutBlock(FMesh* Result, const FMesh* Source,
-                                           uint32 LayoutIndex,
-                                           uint16 BlockCount,
-                                           const uint64* BlockIds, bool& bOutSuccess)
-	{
-		check(Source);
-		bOutSuccess = true;
-
-		// TODO: Optimise
-		Result->CopyFrom(*Source);
-
-		UntypedMeshBufferIteratorConst itBlocks(Source->GetVertexBuffers(), EMeshBufferSemantic::LayoutBlock, LayoutIndex);
-
-        if (itBlocks.GetFormat() != EMeshBufferFormat::None)
-        {
-            int32 ResultVertices = 0;
-			TArray<int32> OldToNew;
-			OldToNew.Init(-1, Source->GetVertexCount());
-			TArray<int32> NewToOld;
-            NewToOld.Reserve(Source->GetVertexCount());
-
-            if (itBlocks.GetFormat() == EMeshBufferFormat::UInt16)
-            {
-                const uint16* pBlocks = reinterpret_cast<const uint16*>(itBlocks.ptr());
-                for ( int32 i=0; i<Source->GetVertexCount(); ++i )
-                {
-                    uint64 VertexBlockRelative = pBlocks[i];
-					uint64 VertexBlockId = (uint64(Source->MeshIDPrefix) << 32) | VertexBlockRelative;
-
-                    bool bFound = false;
-                    for (int32 j = 0; j < BlockCount; ++j)
-                    {
-                        if (VertexBlockId == BlockIds[j])
-                        {
-                            bFound = true;
-                            break;
-                        }
-                    }
-
-                    if (bFound)
-                    {
-                        OldToNew[i] = ResultVertices++;
-                        NewToOld.Add(i);
-                    }
-                }
-            }
-            else if (itBlocks.GetFormat() == EMeshBufferFormat::UInt64)
-			{
-				const uint64* pBlocks = reinterpret_cast<const uint64*>(itBlocks.ptr());
-				for (int32 i = 0; i < Source->GetVertexCount(); ++i)
-				{
-					uint64 VertexBlockId = pBlocks[i];
-
-					bool bFound = false;
-					for (int32 j = 0; j < BlockCount; ++j)
-					{
-						if (VertexBlockId == BlockIds[j])
-						{
-							bFound = true;
-							break;
-						}
-					}
-
-					if (bFound)
-					{
-						OldToNew[i] = ResultVertices++;
-						NewToOld.Add(i);
-					}
-				}
-			}
-			else
-            {
-                check( false );
-            }
-
-            MeshExtractFromVertices(Source, Result, OldToNew, NewToOld);
-        }
-	}
+    void MeshExtractLayoutBlock(FMesh* Result, const FMesh* Source, uint32 LayoutIndex, TArrayView<uint64> BlockIds, bool& bOutSuccess);
 
 	/** Extract all vertices that have a valid layout block on the given channel. */
-	inline void MeshExtractLayoutBlock(FMesh* Result, const FMesh* Source, uint32 LayoutIndex, bool& bOutSuccess)
-	{
-		check(Source);
-		bOutSuccess = true;
-
-		// TODO: Optimise
-		Result->CopyFrom(*Source);
-
-		UntypedMeshBufferIteratorConst itBlocks(Source->GetVertexBuffers(), EMeshBufferSemantic::LayoutBlock, LayoutIndex);
-
-		if (itBlocks.GetFormat() != EMeshBufferFormat::None)
-		{
-			int32 ResultVertices = 0;
-			TArray<int32> OldToNew;
-			OldToNew.Init(-1, Source->GetVertexCount());
-			TArray<int32> NewToOld;
-			NewToOld.Reserve(Source->GetVertexCount());
-
-			if (itBlocks.GetFormat() == EMeshBufferFormat::UInt16)
-			{
-				const uint16* pBlocks = reinterpret_cast<const uint16*>(itBlocks.ptr());
-				for (int32 i = 0; i < Source->GetVertexCount(); ++i)
-				{
-					uint16 VertexBlockRelative = pBlocks[i];
-
-					bool bFound = VertexBlockRelative != std::numeric_limits<uint16>::max();
-					if (bFound)
-					{
-						OldToNew[i] = ResultVertices++;
-						NewToOld.Add(i);
-					}
-				}
-			}
-			else if (itBlocks.GetFormat() == EMeshBufferFormat::UInt64)
-			{
-				const uint64* pBlocks = reinterpret_cast<const uint64*>(itBlocks.ptr());
-				for (int32 i = 0; i < Source->GetVertexCount(); ++i)
-				{
-					uint64 VertexBlockId = pBlocks[i];
-
-					bool bFound = VertexBlockId != std::numeric_limits<uint64>::max();
-					if (bFound)
-					{
-						OldToNew[i] = ResultVertices++;
-						NewToOld.Add(i);
-					}
-				}
-			}
-			else
-			{
-				check(false);
-			}
-
-			MeshExtractFromVertices(Source, Result, OldToNew, NewToOld);
-		}
-	}
-
+	void MeshExtractLayoutBlock(FMesh* Result, const FMesh* Source, uint32 LayoutIndex, bool& bOutSuccess);
 }
