[Mutable] Fix missing Override and Overlay materials.

#rnx

#Synced-CL 51046105

diff --git a/Engine/Plugins/Mutable/Source/CustomizableObject/Private/MuCO/CustomizableObjectInstance.cpp b/Engine/Plugins/Mutable/Source/CustomizableObject/Private/MuCO/CustomizableObjectInstance.cpp
--- a/Engine/Plugins/Mutable/Source/CustomizableObject/Private/MuCO/CustomizableObjectInstance.cpp
+++ b/Engine/Plugins/Mutable/Source/CustomizableObject/Private/MuCO/CustomizableObjectInstance.cpp
@@ -5309,8 +5309,12 @@ void UCustomizableInstancePrivate::BuildResources(const TSharedRef<FUpdateContex
 			if (SlotIndex != INDEX_NONE)
 			{
 				const FBuildMaterialOut NewMaterial = BuildMaterial(Context, OverlayMaterial.MaterialIndex, !bUVChanged);
-					
-				ComponentData.OverlayMaterials.SetNum(SlotIndex + 1);
+				
+				if (!ComponentData.OverlayMaterials.IsValidIndex(SlotIndex))
+				{
+					ComponentData.OverlayMaterials.SetNum(SlotIndex + 1);
+				}
+				
 				ComponentData.OverlayMaterials[SlotIndex] = NewMaterial.Material;
 					
 				NotifyUpdateOnInstanceComponent[ComponentIndex] = NewMaterial.bNotifyPrimitiveUpdated;
@@ -5329,8 +5333,12 @@ void UCustomizableInstancePrivate::BuildResources(const TSharedRef<FUpdateContex
 			if (SlotIndex != INDEX_NONE)
 			{
 				const FBuildMaterialOut Out = BuildMaterial(Context, OverrideMaterial.MaterialIndex, !bUVChanged);
-					
-				ComponentData.OverrideMaterials.SetNum(SlotIndex + 1);
+				
+				if (!ComponentData.OverrideMaterials.IsValidIndex(SlotIndex))
+				{
+					ComponentData.OverrideMaterials.SetNum(SlotIndex + 1);
+				}
+
 				ComponentData.OverrideMaterials[SlotIndex] = Out.Material;
 					
 				NotifyUpdateOnInstanceComponent[ComponentIndex] = Out.bNotifyPrimitiveUpdated;
