[mutable] Made the ASTOpMaterialOverlay operation not store a "None" constant name if no slot name is provided.
- The behaviour of the Op will be the same but now the "None" string will not be stored as a constant.

#rnx

#Synced-CL 51048384

diff --git a/Engine/Plugins/Mutable/Source/MutableRuntime/Internal/MuR/Operations.h b/Engine/Plugins/Mutable/Source/MutableRuntime/Internal/MuR/Operations.h
--- a/Engine/Plugins/Mutable/Source/MutableRuntime/Internal/MuR/Operations.h
+++ b/Engine/Plugins/Mutable/Source/MutableRuntime/Internal/MuR/Operations.h
@@ -1194,7 +1194,7 @@ namespace UE::Mutable::Private
 		{
 			ADDRESS Instance;
 			ADDRESS Material;
-			CONSTANT_STRING_INDEX SlotName = INDEX_NONE;
+			CONSTANT_STRING_INDEX SlotName;
 		};
     	
     	struct InstanceAddOverlayMaterialArgs
@@ -1202,7 +1202,7 @@ namespace UE::Mutable::Private
     		ADDRESS Instance;
     		ADDRESS Material;
     		
-    		CONSTANT_STRING_INDEX SlotName = INDEX_NONE;
+    		CONSTANT_STRING_INDEX SlotName;
     	};
 
         struct LayoutPackArgs
diff --git a/Engine/Plugins/Mutable/Source/MutableRuntime/Private/MuR/CodeRunner.cpp b/Engine/Plugins/Mutable/Source/MutableRuntime/Private/MuR/CodeRunner.cpp
--- a/Engine/Plugins/Mutable/Source/MutableRuntime/Private/MuR/CodeRunner.cpp
+++ b/Engine/Plugins/Mutable/Source/MutableRuntime/Private/MuR/CodeRunner.cpp
@@ -908,7 +908,7 @@ namespace UE::Mutable::Private
         				{
         					const FMaterialId MaterialId = FSystem::GetMaterialId(LiveInstance->MaterialIdRegistry, Model, Parameters, Args.Material);
         					
-        					const FName& SlotName = Args.SlotName == INDEX_NONE ? FName() : FName(Program.ConstantStrings[Args.SlotName]);
+        					const FName SlotName = Args.SlotName == TNumericLimits<uint32>::Max() ? FName() : FName(Program.ConstantNames[Args.SlotName]);
         					if (SlotName.IsNone())
         					{
         						Result->SetOverlayMaterialId(0, MaterialId);
diff --git a/Engine/Plugins/Mutable/Source/MutableRuntime/Private/MuR/CodeRunnerBegin.cpp b/Engine/Plugins/Mutable/Source/MutableRuntime/Private/MuR/CodeRunnerBegin.cpp
--- a/Engine/Plugins/Mutable/Source/MutableRuntime/Private/MuR/CodeRunnerBegin.cpp
+++ b/Engine/Plugins/Mutable/Source/MutableRuntime/Private/MuR/CodeRunnerBegin.cpp
@@ -962,7 +962,7 @@ namespace UE::Mutable::Private
 							{
 								const FMaterialId MaterialId = FSystem::GetMaterialId(MaterialIdRegistry,  *Model.Get(),  *Params.Get(), Args.Material);
         					
-								const FName& SlotName = Args.SlotName == INDEX_NONE ? FName() : FName(Program.ConstantStrings[Args.SlotName]);
+								const FName SlotName = Args.SlotName == TNumericLimits<uint32>::Max() ? FName() : FName(Program.ConstantNames[Args.SlotName]);
 								if (SlotName.IsNone())
 								{
 									Instance->SetOverlayMaterialId(0, MaterialId);
diff --git a/Engine/Plugins/Mutable/Source/MutableTools/Private/MuT/ASTOpAddOverlayMaterial.cpp b/Engine/Plugins/Mutable/Source/MutableTools/Private/MuT/ASTOpAddOverlayMaterial.cpp
--- a/Engine/Plugins/Mutable/Source/MutableTools/Private/MuT/ASTOpAddOverlayMaterial.cpp
+++ b/Engine/Plugins/Mutable/Source/MutableTools/Private/MuT/ASTOpAddOverlayMaterial.cpp
@@ -127,8 +127,17 @@ namespace UE::Mutable::Private
 			}
 		}
 		
-		Args.SlotName = Program.AddConstant(SlotName.ToString());
-
+		if (SlotName.IsNone())
+		{
+			// Mark the SlotName as not relevant so the operation result does not get limited to only one slot
+			Args.SlotName = TNumericLimits<uint32>::Max();
+		}
+		else
+		{
+			// Having no slotName here is valid but that would make the whole component get an overlay material and not the material slot
+			Args.SlotName = Program.AddConstant(SlotName);
+		}
+		
 		linkedAddress = (OP::ADDRESS)Program.OpAddress.Num();
 		Program.OpAddress.Add((uint32_t)Program.ByteCode.Num());
 		AppendCode(Program.ByteCode, GetOpType());
