[mutable] Added "Clip Mesh With Mesh" node.
- This node is based on the modifier node with the same name, but, for this case, it applies the changes within the execution branch and not at the end of it.
- Made the "ASTOpMeshClipWithMesh" able to base its clipping operation on a "FaceCullStrategy".
- Made the core node "ClipMeshWithMesh" evaluate not to an "ASTOpClipMeshWithMesh" but to a chain of nodes as the ones used for the modifier version of this op.
- Fixed bug that, by crashing the editor, prevented the regeneration of the imported SKM model.
- Removed the Modifier node Clip Mesh With Mesh details object as it was doing nothing.
- Added CVar to be able to change the approach the "GenerateMesh_ClipWithMesh" takes.
- Removed unnecessary includes.

#rnx

#Synced-CL 51049204

diff --git a/Engine/Plugins/Mutable/Source/CustomizableObjectEditor/Private/MuCOE/CustomizableObjectEditor.cpp b/Engine/Plugins/Mutable/Source/CustomizableObjectEditor/Private/MuCOE/CustomizableObjectEditor.cpp
--- a/Engine/Plugins/Mutable/Source/CustomizableObjectEditor/Private/MuCOE/CustomizableObjectEditor.cpp
+++ b/Engine/Plugins/Mutable/Source/CustomizableObjectEditor/Private/MuCOE/CustomizableObjectEditor.cpp
@@ -6,18 +6,15 @@
 #include "CustomizableObjectCompileRunnable.h"
 #include "CustomizableObjectEditorModule.h"
 #include "CustomizableObjectEditorPerformanceAnalyzer.h"
-#include "AssetRegistry/ARFilter.h"
 #include "AssetRegistry/AssetRegistryModule.h"
 #include "DetailsViewArgs.h"
 #include "EdGraphUtilities.h"
 #include "EditorModeManager.h"
-#include "Engine/StaticMesh.h"
 #include "FileHelpers.h"
 #include "Framework/Commands/GenericCommands.h"
 #include "Framework/MultiBox/MultiBoxBuilder.h"
 #include "Framework/Notifications/NotificationManager.h"
 #include "GraphEditorActions.h"
-#include "HAL/PlatformApplicationMisc.h"
 #include "IDetailsView.h"
 #include "IMessageLogListing.h"
 #include "MaterialGraph/MaterialGraphNode.h"
@@ -27,7 +24,6 @@
 #include "MuCO/CustomizableObjectCompilerTypes.h"
 #include "MuCOE/CustomizableObjectCustomSettings.h"
 #include "MuCOE/CustomizableObjectEditorActions.h"
-#include "MuCOE/CustomizableObjectEditorLogger.h"
 #include "MuCO/ICustomizableObjectEditorModule.h"
 #include "MuCOE/CustomizableObjectEditorNodeContextCommands.h"
 #include "MuCOE/CustomizableObjectEditorStyle.h"
@@ -35,7 +31,6 @@
 #include "MuCOE/CustomizableObjectGraph.h"
 #include "MuCOE/CustomizableObjectPreviewScene.h"
 #include "MuCOE/GraphTraversal.h"
-#include "MuCOE/Nodes/CustomizableObjectNodeEnumParameter.h"
 #include "MuCOE/Nodes/CustomizableObjectNodeGroupProjectorParameter.h"
 #include "MuCOE/Nodes/CustomizableObjectNodeModifierClipMorph.h"
 #include "MuCOE/Nodes/CustomizableObjectNodeModifierClipWithMesh.h"
@@ -47,6 +42,7 @@
 #include "MuCOE/Nodes/CustomizableObjectNodeStaticMesh.h"
 #include "MuCOE/Nodes/CustomizableObjectNodeTable.h"
 #include "MuCOE/Nodes/CONodeModifierTransformWithBone.h"
+#include "MuCOE/Nodes/CONodeClipMeshWithMesh.h"
 #include "MuCOE/SCustomizableObjectEditorAdvancedPreviewSettings.h"
 #include "MuCOE/SCustomizableObjectEditorTagExplorer.h"
 #include "MuCOE/SCustomizableObjectEditorTextureAnalyzer.h"
@@ -57,7 +53,6 @@
 #include "SMutableGraphViewer.h"
 #include "MuCO/CustomizableObjectInstancePrivate.h"
 #include "MuCO/CustomizableObjectPrivate.h"
-#include "UObject/EnumProperty.h"
 #include "Widgets/Docking/SDockTab.h"
 #include "Widgets/Input/SSearchBox.h"
 #include "Widgets/Input/STextComboBox.h"
@@ -969,7 +964,8 @@ void FCustomizableObjectEditor::HideGizmoClipMesh()
 		const UObject* Node = *NodeIt;
 		if (Node->IsA<UCustomizableObjectNodeModifierClipWithMesh>() ||
 			Node->IsA<UCustomizableObjectNodeModifierTransformInMesh>() ||
-			Node->IsA<UCONodeModifierTransformWithBone>())
+			Node->IsA<UCONodeModifierTransformWithBone>() || 
+			Node->IsA<UCONodeClipMeshWithMesh>())
 		{
 			GraphEditor->ClearSelectionSet();
 			break;
@@ -1634,6 +1630,13 @@ void FCustomizableObjectEditor::OnSelectedGraphNodesChanged(const FGraphPanelSel
 				ShowGizmoClipMesh(*NodeModifierClipWithMesh, &NodeModifierClipWithMesh->Transform, *Pin);
 			}
 		}
+		else if (UCONodeClipMeshWithMesh* NodeClipMeshWithMesh = Cast<UCONodeClipMeshWithMesh>(Objects[0]))
+		{
+			if (const UEdGraphPin* CLipMeshPin = NodeClipMeshWithMesh->ClipMeshPin.Get())
+			{
+				ShowGizmoClipMesh(*NodeClipMeshWithMesh, &NodeClipMeshWithMesh->Transform, *CLipMeshPin);
+			}
+		}
 		else if (UCustomizableObjectNodeModifierTransformInMesh* NodeModifierTransformInMesh = Cast<UCustomizableObjectNodeModifierTransformInMesh>(Objects[0]))
 		{
 			if (const UEdGraphPin* Pin = NodeModifierTransformInMesh->GetBoundingMeshPin())
diff --git a/Engine/Plugins/Mutable/Source/CustomizableObjectEditor/Private/MuCOE/CustomizableObjectEditorModule.cpp b/Engine/Plugins/Mutable/Source/CustomizableObjectEditor/Private/MuCOE/CustomizableObjectEditorModule.cpp
--- a/Engine/Plugins/Mutable/Source/CustomizableObjectEditor/Private/MuCOE/CustomizableObjectEditorModule.cpp
+++ b/Engine/Plugins/Mutable/Source/CustomizableObjectEditor/Private/MuCOE/CustomizableObjectEditorModule.cpp
@@ -19,7 +19,6 @@
 #include "MuCOE/CustomizableInstanceDetails.h"
 #include "MuCOE/CustomizableObjectCustomSettingsDetails.h"
 #include "MuCOE/CustomizableObjectDetails.h"
-#include "MuCOE/CustomizableObjectEditor.h"
 #include "MuCOE/CustomizableObjectEditorSettings.h"
 #include "MuCOE/CustomizableObjectEditorStyle.h"
 #include "MuCOE/CustomizableObjectIdentifierCustomization.h"
@@ -49,8 +48,6 @@
 #include "MuCOE/Nodes/CustomizableObjectNodeModifierRemoveMeshDetails.h"
 #include "MuCOE/Nodes/CustomizableObjectNodeModifierClipMorph.h"
 #include "MuCOE/Nodes/CustomizableObjectNodeModifierClipMorphDetails.h"
-#include "MuCOE/Nodes/CustomizableObjectNodeModifierClipWithMesh.h"
-#include "MuCOE/Nodes/CustomizableObjectNodeModifierClipWithMeshDetails.h"
 #include "MuCOE/Nodes/CustomizableObjectNodeModifierClipWithUVMask.h"
 #include "MuCOE/Nodes/CustomizableObjectNodeModifierClipDeform.h"
 #include "MuCOE/Nodes/CustomizableObjectNodeModifierTransformInMesh.h"
@@ -253,7 +250,6 @@ void FCustomizableObjectEditorModule::StartupModule()
 	RegisterCustomDetails(PropertyModule, UCustomizableObjectNodeModifierRemoveMeshBlocks::StaticClass(), FOnGetDetailCustomizationInstance::CreateStatic(&FCustomizableObjectNodeModifierRemoveMeshBlocksDetails::MakeInstance));
 	RegisterCustomDetails(PropertyModule, UCustomizableObjectNodeModifierMorphMeshSection::StaticClass(), FOnGetDetailCustomizationInstance::CreateStatic(&FCustomizableObjectNodeModifierMorphMeshSectionDetails::MakeInstance));
 	RegisterCustomDetails(PropertyModule, UCustomizableObjectNodeModifierClipMorph::StaticClass(), FOnGetDetailCustomizationInstance::CreateStatic(&FCustomizableObjectNodeModifierClipMorphDetails::MakeInstance));
-	RegisterCustomDetails(PropertyModule, UCustomizableObjectNodeModifierClipWithMesh::StaticClass(), FOnGetDetailCustomizationInstance::CreateStatic(&FCustomizableObjectNodeModifierClipWithMeshDetails::MakeInstance));
 	RegisterCustomDetails(PropertyModule, UCustomizableObjectNodeModifierClipWithUVMask::StaticClass(), FOnGetDetailCustomizationInstance::CreateStatic(&FCustomizableObjectNodeModifierBaseDetails::MakeInstance));
 	RegisterCustomDetails(PropertyModule, UCustomizableObjectNodeModifierClipDeform::StaticClass(), FOnGetDetailCustomizationInstance::CreateStatic(&FCustomizableObjectNodeModifierBaseDetails::MakeInstance));
 	RegisterCustomDetails(PropertyModule, UCustomizableObjectNodeModifierTransformInMesh::StaticClass(), FOnGetDetailCustomizationInstance::CreateStatic(&FCustomizableObjectNodeModifierTransformInMeshDetails::MakeInstance));
diff --git a/Engine/Plugins/Mutable/Source/CustomizableObjectEditor/Private/MuCOE/CustomizableObjectEditorViewportClient.cpp b/Engine/Plugins/Mutable/Source/CustomizableObjectEditor/Private/MuCOE/CustomizableObjectEditorViewportClient.cpp
--- a/Engine/Plugins/Mutable/Source/CustomizableObjectEditor/Private/MuCOE/CustomizableObjectEditorViewportClient.cpp
+++ b/Engine/Plugins/Mutable/Source/CustomizableObjectEditor/Private/MuCOE/CustomizableObjectEditorViewportClient.cpp
@@ -45,6 +45,7 @@
 #include "MuCOE/UnrealBakeHelpers.h"
 #include "MuCO/CustomizableObjectPrivate.h"
 #include "MuCO/LoadUtils.h"
+#include "Nodes/CONodeClipMeshWithMesh.h"
 #include "Widgets/Input/SButton.h"
 #include "Widgets/Input/SEditableTextBox.h"
 #include "Widgets/Layout/SExpandableArea.h"
@@ -647,6 +648,10 @@ void FCustomizableObjectEditorViewportClient::ShowGizmoClipMesh(UCustomizableObj
 	{
 		TransformExternallyChangedDelegateHandle = NodeModifierTransformWithBone->TransformChangedDelegate.AddSP(this, &FCustomizableObjectEditorViewportClient::UpdateGizmoClipMeshTransform);
 	}
+	else if (UCONodeClipMeshWithMesh* NodeClipMeshWithMesh = Cast<UCONodeClipMeshWithMesh>(ClipMeshNode.Get()))
+	{
+		TransformExternallyChangedDelegateHandle = NodeClipMeshWithMesh->TransformChangedDelegate.AddSP(this, &FCustomizableObjectEditorViewportClient::UpdateGizmoClipMeshTransform);
+	}
 	else
 	{
 		unimplemented();
@@ -713,6 +718,10 @@ void FCustomizableObjectEditorViewportClient::HideGizmoClipMesh()
 			{
 				NodeModifierTransformWithBone->TransformChangedDelegate.Remove(TransformExternallyChangedDelegateHandle);
 			}
+			else if (UCONodeClipMeshWithMesh* NodeClipMeshWithMesh = Cast<UCONodeClipMeshWithMesh>(ClipMeshNode.Get()))
+			{
+				NodeClipMeshWithMesh->TransformChangedDelegate.Remove(TransformExternallyChangedDelegateHandle);
+			}
 			else
 			{
 				unimplemented();
@@ -1714,14 +1723,12 @@ void FCustomizableObjectEditorViewportClient::ShowInstanceGeometryInformation(FC
 		if (SkeletalMeshComponent && SkeletalMeshComponent->GetSkinnedAsset())
 		{
 			const FSkeletalMeshRenderData* MeshRes = SkeletalMeshComponent->GetSkinnedAsset()->GetResourceForRendering();
-			int32 NumTriangles;
-			int32 NumVertices;
 			int32 NumLODLevel = MeshRes->LODRenderData.Num();
 
 			for (int32 i = 0; i < NumLODLevel; ++i)
 			{
-				NumTriangles = 0;
-				NumVertices = 0;
+				int32 NumTriangles = 0;
+				int32 NumVertices = 0;
 				const FSkeletalMeshLODRenderData& lodModel = MeshRes->LODRenderData[i];
 				for (int32 j = 0; j < lodModel.RenderSections.Num(); ++j)
 				{
@@ -2105,17 +2112,13 @@ FText FCustomizableObjectEditorViewportClient::GetMeshInfoText() const
 			continue;
 		}
 
-		// Draw stats about the mesh
-		int32 NumBonesInUse;
-		int32 NumBonesMappedToVerts;
-		int32 NumSectionsInUse;
-
 		const int32 LODIndex = FMath::Clamp(PreviewMeshComponent->GetPredictedLODLevel(), 0, SkelMeshResource->LODRenderData.Num() - 1);
 		FSkeletalMeshLODRenderData& LODData = SkelMeshResource->LODRenderData[LODIndex];
 
-		NumBonesInUse = LODData.RequiredBones.Num();
-		NumBonesMappedToVerts = LODData.ActiveBoneIndices.Num();
-		NumSectionsInUse = LODData.RenderSections.Num();
+		// Draw stats about the mesh
+		const int32 NumBonesInUse = LODData.RequiredBones.Num();
+		const int32 NumBonesMappedToVerts = LODData.ActiveBoneIndices.Num();
+		const int32 NumSectionsInUse = LODData.RenderSections.Num();
 
 		// Calculate polys based on non clothing sections so we don't duplicate the counts.
 		uint32 NumTotalTriangles = 0;
diff --git a/Engine/Plugins/Mutable/Source/CustomizableObjectEditor/Private/MuCOE/EdGraphSchema_CustomizableObject.cpp b/Engine/Plugins/Mutable/Source/CustomizableObjectEditor/Private/MuCOE/EdGraphSchema_CustomizableObject.cpp
--- a/Engine/Plugins/Mutable/Source/CustomizableObjectEditor/Private/MuCOE/EdGraphSchema_CustomizableObject.cpp
+++ b/Engine/Plugins/Mutable/Source/CustomizableObjectEditor/Private/MuCOE/EdGraphSchema_CustomizableObject.cpp
@@ -118,6 +118,7 @@
 #include "MuR/External/Operation.h"
 #include "MuR/External/TextureAdapter.h"
 #include "MuR/External/VectorAdapter.h"
+#include "Nodes/CONodeClipMeshWithMesh.h"
 #include "Nodes/CONodeExternalOperation.h"
 #include "Nodes/CONodeExternalTypeParameter.h"
 #include "Nodes/CONodeSkeletalMeshBreak.h"
@@ -530,6 +531,7 @@ void UEdGraphSchema_CustomizableObject::GetGraphContextActions(FGraphContextMenu
 			ContextMenuBuilder.CreateTemplateNode<UCustomizableObjectNodeMeshMorphStackDefinition>(),
 			ContextMenuBuilder.CreateTemplateNode<UCustomizableObjectNodeMeshMorphStackApplication>(),
 			ContextMenuBuilder.CreateTemplateNode<UCustomizableObjectNodeAnimationPose>(),
+			ContextMenuBuilder.CreateTemplateNode<UCONodeClipMeshWithMesh>(),
 		};
 
 		AddNewNodeCategoryActionsFiltered(MeshTemplateNodes, ContextMenuBuilder, TEXT("Mesh"), GeneralGrouping, Filter);
diff --git a/Engine/Plugins/Mutable/Source/CustomizableObjectEditor/Private/MuCOE/GenerateMutableSource/GenerateMutableSourceMesh.cpp b/Engine/Plugins/Mutable/Source/CustomizableObjectEditor/Private/MuCOE/GenerateMutableSource/GenerateMutableSourceMesh.cpp
--- a/Engine/Plugins/Mutable/Source/CustomizableObjectEditor/Private/MuCOE/GenerateMutableSource/GenerateMutableSourceMesh.cpp
+++ b/Engine/Plugins/Mutable/Source/CustomizableObjectEditor/Private/MuCOE/GenerateMutableSource/GenerateMutableSourceMesh.cpp
@@ -59,8 +59,11 @@
 #include "AssetRegistry/AssetRegistryModule.h"
 #include "Hash/CityHash.h"
 #include "PackedNormal.h"
+#include "MuCOE/Nodes/CONodeClipMeshWithMesh.h"
 #include "MuCOE/Nodes/CONodeExternalOperation.h"
+#include "MuT/NodeMeshClipWithMesh.h"
 #include "MuT/NodeMeshExternal.h"
+#include "MuT/NodeMeshTransform.h"
 
 #define LOCTEXT_NAMESPACE "CustomizableObjectEditor"
 
@@ -3195,6 +3198,63 @@ UE::Mutable::Private::Ptr<UE::Mutable::Private::NodeMesh> GenerateMutableSourceM
 
 		GenerationContext.MeshGenerationFlags.Pop();
 	}
+	
+	else if (const UCONodeClipMeshWithMesh* ClipMeshWithMeshNode = Cast<UCONodeClipMeshWithMesh>(Node))
+	{
+		UE::Mutable::Private::Ptr<UE::Mutable::Private::NodeMeshClipWithMesh> ClipNode = new UE::Mutable::Private::NodeMeshClipWithMesh();
+		ClipNode->FaceCullStrategy = ClipMeshWithMeshNode->FaceCullStrategy;
+		
+		// Base Mesh
+		const UEdGraphPin* BaseMeshPin = ClipMeshWithMeshNode->BaseMeshPin.Get();
+		check(BaseMeshPin);
+		if (const UEdGraphPin* ConnectedPin = FollowInputPin(*BaseMeshPin))
+		{
+			UE::Mutable::Private::NodeMeshPtr BaseMesh = GenerateMutableSourceMesh(ConnectedPin, GenerationContext, FMutableSourceMeshData(), false, true);
+			ClipNode->Source = BaseMesh;
+		}
+		else
+		{
+			const FText ErrorMsg = LOCTEXT("Base mesh missing", "The clip mesh with mesh node requires an input base mesh.");
+			GenerationContext.Log(ErrorMsg, ClipMeshWithMeshNode, EMessageSeverity::Error);
+		}
+		
+		// Clip Mesh
+		const UEdGraphPin* ClipMeshPin = ClipMeshWithMeshNode->ClipMeshPin.Get();
+		check(ClipMeshPin);
+		if (const UEdGraphPin* ConnectedPin = FollowInputPin(*ClipMeshPin))
+		{
+			constexpr EMutableMeshConversionFlags ClipMeshFlags =
+				EMutableMeshConversionFlags::IgnoreSkinning |
+				EMutableMeshConversionFlags::IgnorePhysics |
+				EMutableMeshConversionFlags::IgnoreMorphs |
+				EMutableMeshConversionFlags::IgnoreTexCoords |
+				EMutableMeshConversionFlags::IgnoreAUD |
+				EMutableMeshConversionFlags::DoNotCreateMeshMetadata;
+			GenerationContext.MeshGenerationFlags.Push(ClipMeshFlags);
+			
+			UE::Mutable::Private::NodeMeshPtr ClipMesh = GenerateMutableSourceMesh(ConnectedPin, GenerationContext, FMutableSourceMeshData(), false, true);
+
+			if (const FMatrix Matrix = ClipMeshWithMeshNode->Transform.ToMatrixWithScale(); Matrix != FMatrix::Identity)
+			{
+				UE::Mutable::Private::Ptr<UE::Mutable::Private::NodeMeshTransform> TransformMesh = new UE::Mutable::Private::NodeMeshTransform();
+				TransformMesh->Source = ClipMesh;
+
+				TransformMesh->Transform = FMatrix44f(Matrix);
+				ClipMesh = TransformMesh;
+			}
+
+			ClipNode->ClipMesh = ClipMesh;
+			
+			GenerationContext.MeshGenerationFlags.Pop();
+		}
+		else
+		{
+			const FText ErrorMsg = LOCTEXT("Clipping mesh missing", "The clip mesh with mesh node requires an input clip mesh.");
+			GenerationContext.Log(ErrorMsg, ClipMeshWithMeshNode, EMessageSeverity::Error);
+		}
+		
+		Result = ClipNode;
+	}
 
 	else if (const UCustomizableObjectNodeAnimationPose* TypedNode = Cast<UCustomizableObjectNodeAnimationPose>(Node))
 	{
@@ -3235,11 +3295,11 @@ UE::Mutable::Private::Ptr<UE::Mutable::Private::NodeMesh> GenerateMutableSourceM
 				}
 				else
 				{
-						if (!TypedNode->PoseAsset) // Check if the slot has a selected pose. Could be left empty by the user
-						{
-							FString msg = FString::Printf(TEXT("Found pose mesh node without a pose asset assigned."));
-							GenerationContext.Log(FText::FromString(msg), TypedNode);
-						}
+					if (!TypedNode->PoseAsset) // Check if the slot has a selected pose. Could be left empty by the user
+					{
+						const FString Message = FString::Printf(TEXT("Found pose mesh node without a pose asset assigned."));
+						GenerationContext.Log(FText::FromString(Message), TypedNode);
+					}
 
 					Result = InputMeshNode;
 				}
diff --git a/Engine/Plugins/Mutable/Source/CustomizableObjectEditor/Private/MuCOE/GraphTraversal.cpp b/Engine/Plugins/Mutable/Source/CustomizableObjectEditor/Private/MuCOE/GraphTraversal.cpp
--- a/Engine/Plugins/Mutable/Source/CustomizableObjectEditor/Private/MuCOE/GraphTraversal.cpp
+++ b/Engine/Plugins/Mutable/Source/CustomizableObjectEditor/Private/MuCOE/GraphTraversal.cpp
@@ -29,6 +29,7 @@
 #include "MuCOE/Nodes/CustomizableObjectNodeObjectGroup.h"
 #include "MuCOE/Nodes/CustomizableObjectNodeReroute.h"
 #include "Nodes/CONodeExternalOperation.h"
+#include "Nodes/CONodeClipMeshWithMesh.h"
 #include "UObject/UObjectIterator.h"
 
 
@@ -762,6 +763,14 @@ const UEdGraphPin* FindMeshBaseSource(const UEdGraphPin& Pin, const bool bOnlyLo
 			}
 		}
 	}
+	
+	else if (const UCONodeClipMeshWithMesh* NodeClipMeshWithMesh = Cast<UCONodeClipMeshWithMesh>(Node))
+	{
+		if (const UEdGraphPin* ConnectedPin = FollowInputPin(*NodeClipMeshWithMesh->BaseMeshPin.Get()))
+		{
+			return FindMeshBaseSource(*ConnectedPin, bOnlyLookForStaticMesh, MacroContext);
+		}
+	}
 
 	else if (Cast<UCONodeExternalOperation>(Node))
 	{
diff --git a/Engine/Plugins/Mutable/Source/CustomizableObjectEditor/Private/MuCOE/Nodes/CONodeClipMeshWithMesh.cpp b/Engine/Plugins/Mutable/Source/CustomizableObjectEditor/Private/MuCOE/Nodes/CONodeClipMeshWithMesh.cpp
new file mode 100644
--- /dev/null
+++ b/Engine/Plugins/Mutable/Source/CustomizableObjectEditor/Private/MuCOE/Nodes/CONodeClipMeshWithMesh.cpp
@@ -0,0 +1,64 @@
+// Copyright Epic Games, Inc. All Rights Reserved.
+
+#include "MuCOE/Nodes/CONodeClipMeshWithMesh.h"
+
+#include "MuCOE/Nodes/CustomizableObjectNode.h"
+#include "MuCOE/CustomizableObjectGraph.h"
+#include "MuCOE/EdGraphSchema_CustomizableObject.h"
+
+
+#define LOCTEXT_NAMESPACE "CustomizableObjectEditor"
+
+
+void UCONodeClipMeshWithMesh::AllocateDefaultPins(UCustomizableObjectNodeRemapPins* RemapPins)
+{
+	BaseMeshPin = CustomCreatePin(EGPD_Input, UEdGraphSchema_CustomizableObject::PC_Mesh, TEXT("Base Mesh"), LOCTEXT("BaseMesh", "Base Mesh"));
+	ClipMeshPin =  CustomCreatePin(EGPD_Input, UEdGraphSchema_CustomizableObject::PC_Mesh, TEXT("Clip Mesh"), LOCTEXT("ClipMesh", "Clip Mesh"));
+
+	OutputMeshPin = CustomCreatePin(EGPD_Output, UEdGraphSchema_CustomizableObject::PC_Mesh);
+}
+
+
+void UCONodeClipMeshWithMesh::PostEditChangeProperty(FPropertyChangedEvent& PropertyChangedEvent)
+{
+	Super::PostEditChangeProperty(PropertyChangedEvent);
+
+	if (PropertyChangedEvent.Property == nullptr)
+	{
+		return;
+	}
+	
+	if (PropertyChangedEvent.Property->GetFName() == GET_MEMBER_NAME_CHECKED(UCONodeClipMeshWithMesh, Transform))
+	{
+		if (TransformChangedDelegate.IsBound())
+		{
+			TransformChangedDelegate.Broadcast(Transform);
+		}
+	}
+}
+
+
+bool UCONodeClipMeshWithMesh::IsExperimental() const
+{
+	return true;
+}
+
+
+FText UCONodeClipMeshWithMesh::GetNodeTitle(ENodeTitleType::Type TitleType) const
+{
+	return LOCTEXT("Clip_Mesh_With_Mesh", "Clip Mesh With Mesh");
+}
+
+
+FLinearColor UCONodeClipMeshWithMesh::GetNodeTitleColor() const
+{
+	return UEdGraphSchema_CustomizableObject::GetPinTypeColor(UEdGraphSchema_CustomizableObject::PC_Mesh);
+}
+
+
+FText UCONodeClipMeshWithMesh::GetTooltipText() const
+{
+	return LOCTEXT("Clip_Mesh_Mesh_Tooltip", "Removes the part of a mesh section that is completely enclosed in a mesh volume.");
+}
+
+#undef LOCTEXT_NAMESPACE
diff --git a/Engine/Plugins/Mutable/Source/CustomizableObjectEditor/Private/MuCOE/Nodes/CONodeClipMeshWithMesh.h b/Engine/Plugins/Mutable/Source/CustomizableObjectEditor/Private/MuCOE/Nodes/CONodeClipMeshWithMesh.h
new file mode 100644
--- /dev/null
+++ b/Engine/Plugins/Mutable/Source/CustomizableObjectEditor/Private/MuCOE/Nodes/CONodeClipMeshWithMesh.h
@@ -0,0 +1,56 @@
+// Copyright Epic Games, Inc. All Rights Reserved.
+
+#pragma once
+
+#include "MuCOE/Nodes/CustomizableObjectNode.h"
+
+#include "CONodeClipMeshWithMesh.generated.h"
+
+namespace ENodeTitleType { enum Type : int; }
+
+class UCustomizableObject;
+class UCustomizableObjectNodeRemapPins;
+class UEdGraphPin;
+class UObject;
+struct FGuid;
+
+DECLARE_MULTICAST_DELEGATE_OneParam(OnTransformChanged, const FTransform&);
+
+UCLASS()
+class UCONodeClipMeshWithMesh : public UCustomizableObjectNode
+{
+	GENERATED_BODY()
+	
+public:
+	
+	// EdGraphNode interface
+	virtual FText GetNodeTitle(ENodeTitleType::Type TitleType) const override;
+	virtual FLinearColor GetNodeTitleColor() const override;
+	virtual FText GetTooltipText() const override;
+	virtual bool IsExperimental() const override;
+	
+	// UCustomizableObjectNode interface
+	virtual void AllocateDefaultPins(UCustomizableObjectNodeRemapPins* RemapPins) override;
+	virtual void PostEditChangeProperty( struct FPropertyChangedEvent& PropertyChangedEvent) override;
+	
+	/** Transform to apply to the clip mesh before clipping. */
+	UPROPERTY(EditAnywhere, Category = ClipMesh)
+	FTransform Transform = FTransform::Identity;
+
+	UPROPERTY(EditAnywhere, Category = ClipMesh)
+	EFaceCullStrategy FaceCullStrategy = EFaceCullStrategy::AllVerticesCulled;
+	
+	UPROPERTY()
+	FEdGraphPinReference BaseMeshPin;
+	
+	UPROPERTY()
+	FEdGraphPinReference ClipMeshPin;
+	
+	UPROPERTY()
+	FEdGraphPinReference OutputMeshPin;
+	
+	/**
+	 * Delegate invoked each time the BoundingMeshTransform value gets modified
+	 */
+	OnTransformChanged TransformChangedDelegate;
+};
diff --git a/Engine/Plugins/Mutable/Source/CustomizableObjectEditor/Private/MuCOE/Nodes/CustomizableObjectNodeModifierClipWithMeshDetails.cpp b/Engine/Plugins/Mutable/Source/CustomizableObjectEditor/Private/MuCOE/Nodes/CustomizableObjectNodeModifierClipWithMeshDetails.cpp
deleted file mode 100644
diff --git a/Engine/Plugins/Mutable/Source/CustomizableObjectEditor/Private/MuCOE/Nodes/CustomizableObjectNodeModifierClipWithMeshDetails.h b/Engine/Plugins/Mutable/Source/CustomizableObjectEditor/Private/MuCOE/Nodes/CustomizableObjectNodeModifierClipWithMeshDetails.h
deleted file mode 100644
diff --git a/Engine/Plugins/Mutable/Source/CustomizableObjectEditor/Private/MuCOE/SMutableGraphViewer.cpp b/Engine/Plugins/Mutable/Source/CustomizableObjectEditor/Private/MuCOE/SMutableGraphViewer.cpp
--- a/Engine/Plugins/Mutable/Source/CustomizableObjectEditor/Private/MuCOE/SMutableGraphViewer.cpp
+++ b/Engine/Plugins/Mutable/Source/CustomizableObjectEditor/Private/MuCOE/SMutableGraphViewer.cpp
@@ -44,6 +44,7 @@
 #include "MuT/NodeSurfaceVariation.h"
 #include "MuT/NodeLOD.h"
 #include "MuT/NodeMaterialModify.h"
+#include "MuT/NodeMeshClipWithMesh.h"
 #include "MuT/NodeMeshConstant.h"
 #include "MuT/NodeMeshFormat.h"
 #include "MuT/NodeMeshFragment.h"
@@ -615,6 +616,13 @@ void SMutableGraphViewer::GetChildrenForInfo(TSharedPtr<FMutableGraphTreeElement
 		UE::Mutable::Private::NodeModifierMeshClipWithMesh* ModifierMeshClipWithMeshVar = StaticCast<UE::Mutable::Private::NodeModifierMeshClipWithMesh*>(ParentNode);
 		AddChildFunc(ModifierMeshClipWithMeshVar->ClipMesh.get(), FString::Printf(TEXT("CLIP MESH")));
 	}
+	
+	else if (ParentNode->GetType() == UE::Mutable::Private::NodeMeshClipWithMesh::GetStaticType())
+	{
+		UE::Mutable::Private::NodeMeshClipWithMesh* ClipMeshWithMeshVar = StaticCast<UE::Mutable::Private::NodeMeshClipWithMesh*>(ParentNode);
+		AddChildFunc(ClipMeshWithMeshVar->Source.get(), FString::Printf(TEXT("SOURCE MESH")));
+		AddChildFunc(ClipMeshWithMeshVar->ClipMesh.get(), FString::Printf(TEXT("CLIP MESH")));
+	}
 
 	else if (ParentNode->GetType() == UE::Mutable::Private::NodeModifierMeshClipDeform::GetStaticType())
 	{
diff --git a/Engine/Plugins/Mutable/Source/MutableRuntime/Internal/MuR/Operations.h b/Engine/Plugins/Mutable/Source/MutableRuntime/Internal/MuR/Operations.h
--- a/Engine/Plugins/Mutable/Source/MutableRuntime/Internal/MuR/Operations.h
+++ b/Engine/Plugins/Mutable/Source/MutableRuntime/Internal/MuR/Operations.h
@@ -1015,6 +1015,8 @@ namespace UE::Mutable::Private
         {
             ADDRESS Source;
             ADDRESS ClipMesh;
+        	
+        	EFaceCullStrategy FaceCullStrategy;
         };
 
         struct MeshProjectArgs
diff --git a/Engine/Plugins/Mutable/Source/MutableRuntime/Private/MuR/CodeRunner.cpp b/Engine/Plugins/Mutable/Source/MutableRuntime/Private/MuR/CodeRunner.cpp
--- a/Engine/Plugins/Mutable/Source/MutableRuntime/Private/MuR/CodeRunner.cpp
+++ b/Engine/Plugins/Mutable/Source/MutableRuntime/Private/MuR/CodeRunner.cpp
@@ -2,10 +2,8 @@
 
 #include "MuR/CodeRunner.h"
 
-#include "SkinnedAssetCompiler.h"
 #include "GenericPlatform/GenericPlatformMath.h"
 #include "HAL/UnrealMemory.h"
-#include "Logging/LogCategory.h"
 #include "Logging/LogMacros.h"
 #include "Math/IntPoint.h"
 #include "Math/UnrealMathUtility.h"
@@ -44,7 +42,6 @@
 #include "MuR/OpMeshDifference.h"
 #include "MuR/OpMeshExtractLayoutBlock.h"
 #include "MuR/OpMeshFormat.h"
-#include "MuR/OpMeshGeometryOperation.h"
 #include "MuR/OpMeshMerge.h"
 #include "MuR/OpMeshMorph.h"
 #include "MuR/OpMeshRemove.h"
@@ -55,9 +52,6 @@
 #include "MuR/OpMeshOptimizeSkinning.h"
 #include "MuR/Operations.h"
 #include "MuR/Parameters.h"
-#include "MuR/Parameters.h"
-#include "MuR/PhysicsBody.h"
-#include "MuR/Skeleton.h"
 #include "MuR/System.h"
 #include "MuR/External/FloatAdapter.h"
 #include "MuR/External/MaterialAdapter.h"
@@ -2112,8 +2106,9 @@ namespace UE::Mutable::Private
 					{
 						TManagedPtr<FMesh> Result = CreateMesh(Source->GetDataSize());
 
+						const bool bRemoveAllVerticesCulled = Args.FaceCullStrategy == EFaceCullStrategy::AllVerticesCulled;
 						bool bOutSuccess = false;
-						MeshClipWithMesh(Result.Get(), Source.Get(), Clip.Get(), bOutSuccess);
+						MeshClipWithMesh(Result.Get(), Source.Get(), Clip.Get(), bRemoveAllVerticesCulled, bOutSuccess);
 
 						if (!bOutSuccess)
 						{
diff --git a/Engine/Plugins/Mutable/Source/MutableRuntime/Private/MuR/OpMeshClipWithMesh.cpp b/Engine/Plugins/Mutable/Source/MutableRuntime/Private/MuR/OpMeshClipWithMesh.cpp
--- a/Engine/Plugins/Mutable/Source/MutableRuntime/Private/MuR/OpMeshClipWithMesh.cpp
+++ b/Engine/Plugins/Mutable/Source/MutableRuntime/Private/MuR/OpMeshClipWithMesh.cpp
@@ -15,6 +15,7 @@
 #include "Intersection/IntrRay3Triangle3.h"
 #include "MathUtil.h"
 #include "IntVectorTypes.h"
+#include "OpMeshRemove.h"
 #include "Math/UnrealMathUtility.h"
 #include "Spatial/PointHashGrid3.h"
 
@@ -699,7 +700,7 @@ namespace UE::Mutable::Private
         }
     }	
 
-    void MeshClipWithMesh(FMesh* Result, const FMesh* pBase, const FMesh* ClipMesh, bool& bOutSuccess)
+    void MeshClipWithMesh(FMesh* Result, const FMesh* pBase, const FMesh* ClipMesh, bool bRemoveAllVerticesCulled, bool& bOutSuccess)
     {
         MUTABLE_CPUPROFILER_SCOPE(MeshClipWithMesh);
 		bOutSuccess = true;
@@ -715,42 +716,8 @@ namespace UE::Mutable::Private
 
         TBitArray<> VertexInClipMesh;  // Stores whether each vertex in the original mesh is in the clip mesh volume
         MeshClipMeshClassifyVertices(VertexInClipMesh, pBase, ClipMesh);
-
-        // Now remove all the faces from the result mesh that have all the vertices outside the clip volume
-        UntypedMeshBufferIteratorConst ItBase(Result->GetIndexBuffers(), EMeshBufferSemantic::VertexIndex);
-        UntypedMeshBufferIterator ItDest(Result->GetIndexBuffers(), EMeshBufferSemantic::VertexIndex);
-        int32 AFaceCount = Result->GetFaceCount();
-
-        UntypedMeshBufferIteratorConst Ito(Result->GetIndexBuffers(), EMeshBufferSemantic::VertexIndex);
-        for (int32 F = 0; F < AFaceCount; ++F)
-        {
-			uint32 OV[3] = {0, 0, 0};
-
-            OV[0] = Ito.GetAsUINT32(); ++Ito;
-            OV[1] = Ito.GetAsUINT32(); ++Ito;
-            OV[2] = Ito.GetAsUINT32(); ++Ito;
-
-            bool AllVertsIn = VertexInClipMesh[OV[0]] && VertexInClipMesh[OV[1]] && VertexInClipMesh[OV[2]];
-
-            if (!AllVertsIn)
-            {
-                if (ItDest.ptr() != ItBase.ptr())
-                {
-					FMemory::Memcpy(ItDest.ptr(), ItBase.ptr(), ItBase.GetElementSize() * 3);
-                }
-
-                ItDest += 3;
-            }
-
-            ItBase += 3;
-        }
-
-        SIZE_T RemovedIndices = ItBase - ItDest;
-        check(RemovedIndices % 3 == 0);
-
-        Result->GetIndexBuffers().SetElementCount(AFaceCount * 3 - (int32)RemovedIndices);
-
-        MeshRemoveUnusedVertices(Result);
+		
+		MeshRemoveVerticesWithCullSet(Result, VertexInClipMesh, bRemoveAllVerticesCulled);
     }
 
 
diff --git a/Engine/Plugins/Mutable/Source/MutableRuntime/Private/MuR/OpMeshClipWithMesh.h b/Engine/Plugins/Mutable/Source/MutableRuntime/Private/MuR/OpMeshClipWithMesh.h
--- a/Engine/Plugins/Mutable/Source/MutableRuntime/Private/MuR/OpMeshClipWithMesh.h
+++ b/Engine/Plugins/Mutable/Source/MutableRuntime/Private/MuR/OpMeshClipWithMesh.h
@@ -14,7 +14,7 @@ namespace UE::Mutable::Private
 	extern void MeshClipMeshClassifyVertices(TBitArray<>& VertexInClipMesh, const FMesh* Base, const FMesh* ClipMesh);
 
     /**  */
-	extern void MeshClipWithMesh(FMesh* Result, const FMesh* Base, const FMesh* ClipMesh, bool& bOutSuccess);
+	extern void MeshClipWithMesh(FMesh* Result, const FMesh* Base, const FMesh* ClipMesh, bool bRemoveAllVerticesCulled, bool& bOutSuccess);
 
     /** Generate a mask mesh with the faces of the base mesh inside the clip mesh. */
 	extern void MeshMaskClipMesh(FMesh* Result, const FMesh* Base, const FMesh* ClipMesh, bool& bOutSuccess);
diff --git a/Engine/Plugins/Mutable/Source/MutableTools/Internal/MuT/NodeMeshClipWithMesh.h b/Engine/Plugins/Mutable/Source/MutableTools/Internal/MuT/NodeMeshClipWithMesh.h
--- a/Engine/Plugins/Mutable/Source/MutableTools/Internal/MuT/NodeMeshClipWithMesh.h
+++ b/Engine/Plugins/Mutable/Source/MutableTools/Internal/MuT/NodeMeshClipWithMesh.h
@@ -3,7 +3,7 @@
 #pragma once
 
 #include "MuR/Ptr.h"
-#include "MuR/RefCounted.h"
+#include "MuR/Types.h"
 #include "MuT/Node.h"
 #include "MuT/NodeMesh.h"
 
@@ -20,10 +20,11 @@ namespace UE::Mutable::Private
 	public:
 
 		Ptr<NodeMesh> Source;
+    	
 		Ptr<NodeMesh> ClipMesh;
 
-	public:
-
+    	EFaceCullStrategy FaceCullStrategy = EFaceCullStrategy::AllVerticesCulled;
+    	
 		// Node interface
 		virtual const FNodeType* GetType() const override { return &StaticType; }
 		static const FNodeType* GetStaticType() { return &StaticType; }
diff --git a/Engine/Plugins/Mutable/Source/MutableTools/Private/MuT/ASTOpMeshClipWithMesh.cpp b/Engine/Plugins/Mutable/Source/MutableTools/Private/MuT/ASTOpMeshClipWithMesh.cpp
--- a/Engine/Plugins/Mutable/Source/MutableTools/Private/MuT/ASTOpMeshClipWithMesh.cpp
+++ b/Engine/Plugins/Mutable/Source/MutableTools/Private/MuT/ASTOpMeshClipWithMesh.cpp
@@ -31,7 +31,8 @@ namespace UE::Mutable::Private
 		{
 			const ASTOpMeshClipWithMesh* Other = static_cast<const ASTOpMeshClipWithMesh*>(&OtherUntyped);
 			return Source == Other->Source &&
-				ClipMesh == Other->ClipMesh;
+				ClipMesh == Other->ClipMesh &&
+				FaceCullStrategy == Other->FaceCullStrategy;
 		}
 		return false;
 	}
@@ -43,6 +44,7 @@ namespace UE::Mutable::Private
 
 		Result = HashCombineFast(Result, GetTypeHash(Source));
 		Result = HashCombineFast(Result, GetTypeHash(ClipMesh));
+		Result = HashCombineFast( Result, GetTypeHash(FaceCullStrategy));
 		
 		return Result;
 	}
@@ -53,6 +55,7 @@ namespace UE::Mutable::Private
 		Ptr<ASTOpMeshClipWithMesh> New = new ASTOpMeshClipWithMesh();
 		New->Source = MapChild(Source.child());
 		New->ClipMesh = MapChild(ClipMesh.child());
+		New->FaceCullStrategy = FaceCullStrategy;
 		return New;
 	}
 
@@ -69,12 +72,20 @@ namespace UE::Mutable::Private
 		// Already linked?
 		if (!linkedAddress)
 		{
-			OP::MeshApplyPoseArgs Args;
+			OP::MeshClipWithMeshArgs Args;
 			FMemory::Memzero(Args);
 
-			if (Source) Args.base = Source->linkedAddress;
-			if (ClipMesh) Args.pose = ClipMesh->linkedAddress;
+			if (Source)
+			{
+				Args.Source = Source->linkedAddress;
+			}
+			if (ClipMesh)
+			{
+				Args.ClipMesh = ClipMesh->linkedAddress;
+			}
 
+			Args.FaceCullStrategy = FaceCullStrategy;
+			
 			linkedAddress = (OP::ADDRESS)Program.OpAddress.Num();
 			Program.OpAddress.Add(Program.ByteCode.Num());
 			AppendCode(Program.ByteCode, GetOpType());
diff --git a/Engine/Plugins/Mutable/Source/MutableTools/Private/MuT/ASTOpMeshClipWithMesh.h b/Engine/Plugins/Mutable/Source/MutableTools/Private/MuT/ASTOpMeshClipWithMesh.h
--- a/Engine/Plugins/Mutable/Source/MutableTools/Private/MuT/ASTOpMeshClipWithMesh.h
+++ b/Engine/Plugins/Mutable/Source/MutableTools/Private/MuT/ASTOpMeshClipWithMesh.h
@@ -4,6 +4,7 @@
 
 #include "HAL/Platform.h"
 #include "MuR/Operations.h"
+#include "MuR/Types.h"
 #include "MuR/Ptr.h"
 #include "MuT/AST.h"
 
@@ -19,7 +20,8 @@ namespace UE::Mutable::Private
 
 		ASTChild Source;
 		ASTChild ClipMesh;
-
+		EFaceCullStrategy FaceCullStrategy = EFaceCullStrategy::AllVerticesCulled;
+		
 	public:
 
 		ASTOpMeshClipWithMesh();
diff --git a/Engine/Plugins/Mutable/Source/MutableTools/Private/MuT/CodeGenerator_Mesh.cpp b/Engine/Plugins/Mutable/Source/MutableTools/Private/MuT/CodeGenerator_Mesh.cpp
--- a/Engine/Plugins/Mutable/Source/MutableTools/Private/MuT/CodeGenerator_Mesh.cpp
+++ b/Engine/Plugins/Mutable/Source/MutableTools/Private/MuT/CodeGenerator_Mesh.cpp
@@ -1,6 +1,9 @@
 // Copyright Epic Games, Inc. All Rights Reserved.
 
+#include "ASTOpConstantBool.h"
 #include "ASTOpExternal.h"
+#include "ASTOpMeshMaskClipMesh.h"
+#include "ASTOpMeshRemoveMask.h"
 #include "ASTOpMeshSkeletalMeshBreak.h"
 #include "Containers/Array.h"
 #include "Misc/AssertionMacros.h"
@@ -62,6 +65,14 @@ namespace UE::Mutable::Private
 {
 	class Node;
 
+	static TAutoConsoleVariable<bool> CVarUseASTOpMaskCliMeshForClipMeshWithMesh(
+		TEXT("mutable.UseASTOpMaskToPerformClipMeshWithMesh"),
+		true,
+		TEXT("Defines what ASTOps does the 'GenerateMesh_ClipWithMesh' use to perform the clipping of the mesh.\n")
+		TEXT("  false : Use the 'ASTOpMeshClipWithMesh' to perform the mesh clipping (not tested in depth so it may have bugs).\n")
+		TEXT("  true :  Use instead an 'ASTOpMeshMaskClipMesh' op alongside an 'ASTOpMeshRemoveMask (used by some modifier nodes already).'\n"),
+		ECVF_Scalability | ECVF_RenderThreadSafe
+	);
 
 	FMeshTask CodeGenerator::GenerateMesh(const FMeshGenerationStaticOptions& StaticOptions, FMeshOptionsTask InOptions, const Ptr<const NodeMesh>& InUntypedNode)
     {
@@ -1151,16 +1162,38 @@ namespace UE::Mutable::Private
 			{
 				FMeshGenerationResult SourceResult = SourceTask.GetResult();
 				FMeshGenerationResult ClipResult = ClipTask.GetResult();
+					
+				FMeshGenerationResult Result;
+					
+				if (CVarUseASTOpMaskCliMeshForClipMeshWithMesh.GetValueOnAnyThread())
+				{
+					Ptr<UE::Mutable::Private::ASTOpMeshMaskClipMesh> MaskOp = new UE::Mutable::Private::ASTOpMeshMaskClipMesh();
+					MaskOp->source = SourceResult.MeshOp;
+					MaskOp->clip = ClipResult.MeshOp;
+
+					Ptr<UE::Mutable::Private::ASTOpConstantBool> ConditionOp = new UE::Mutable::Private::ASTOpConstantBool();
+					ConditionOp->bValue = true;
+										
+					Ptr<UE::Mutable::Private::ASTOpMeshRemoveMask> RemoveOp = new UE::Mutable::Private::ASTOpMeshRemoveMask();
+					RemoveOp->source = SourceResult.MeshOp;
+					RemoveOp->FaceCullStrategy = ClipNode->FaceCullStrategy;
+					RemoveOp->AddRemove(ConditionOp, MaskOp);
+
+					Result = SourceResult;
+					Result.MeshOp = RemoveOp;
+				}
+				else
+				{
+					Ptr<ASTOpMeshClipWithMesh> op = new ASTOpMeshClipWithMesh();
 
-				Ptr<ASTOpMeshClipWithMesh> op = new ASTOpMeshClipWithMesh;
-
-				// Base
-				op->Source = SourceResult.MeshOp;
-				op->ClipMesh = ClipResult.MeshOp;
-
-				FMeshGenerationResult Result = SourceResult;
-				Result.MeshOp = op;
+					// Base
+					op->Source = SourceResult.MeshOp;
+					op->ClipMesh = ClipResult.MeshOp;
 
+					Result = SourceResult;
+					Result.MeshOp = op;
+				}
+					
 				return Result;
 			},
 			UE::Tasks::Prerequisites( SourceTask, ClipTask )
