[Mutable] Fix assert when multiple shared cloth configs of the same type are present in a generated skeletalmesh.

- This could happen when multiple clothing assets coming from different skeletal meshes where used in the same generated mesh.

#rnx

#Synced-CL 51096194

diff --git a/Engine/Plugins/Mutable/Source/CustomizableObject/Private/MuCO/CustomizableObjectInstance.cpp b/Engine/Plugins/Mutable/Source/CustomizableObject/Private/MuCO/CustomizableObjectInstance.cpp
--- a/Engine/Plugins/Mutable/Source/CustomizableObject/Private/MuCO/CustomizableObjectInstance.cpp
+++ b/Engine/Plugins/Mutable/Source/CustomizableObject/Private/MuCO/CustomizableObjectInstance.cpp
@@ -4121,6 +4121,8 @@ void BuildOrCopyClothingData(
 			ClothingData.AssetLodIndex = Asset->LodMap[LODIndex];
 		}
 	}
+
+	TArray<TNotNull<UClothSharedConfigCommon*>, TInlineAllocator<8>> SkeletalMeshSharedConfigs;
 	
 	for (int32 Index = 0; Index < ClothingAssets.Num(); ++Index)
 	{
@@ -4129,7 +4131,34 @@ void BuildOrCopyClothingData(
 		Asset = DuplicateObject<UClothingAssetCommon>(Asset, SkeletalMesh);
 		Asset->RefreshBoneMapping(SkeletalMesh);
 		Asset->CalculateReferenceBoneIndex();
-		
+	
+		// Only one cloth shared config of the same type is allowed for the same SkeletalMesh. Keep the
+		// first found and replace the others of the same type with the first found.
+		for (TPair<FName, TObjectPtr<UClothConfigBase>>& Entry : Asset->ClothConfigs)
+		{
+			UClothSharedConfigCommon* SharedConfig = Cast<UClothSharedConfigCommon>(Entry.Value);
+			if (!SharedConfig)
+			{
+				continue;
+			}
+
+			TNotNull<UClothSharedConfigCommon*>* Found = SkeletalMeshSharedConfigs.FindByPredicate(
+			[Class = SharedConfig->GetClass()](TNotNull<UClothSharedConfigCommon*> Elem) 
+			{ 
+				return Elem->GetClass() == Class;
+			});
+
+			if (!Found)
+			{
+				SkeletalMeshSharedConfigs.Add(SharedConfig);
+			}
+			else
+			{
+				// Should we LOG we are not using the original shared config?
+				Entry.Value = *Found;
+			}
+		}
+
 		ClothingAssets[Index] = Asset;
 	}
 	
