[mutable] RemoveMesh node
- Added the "Remove Mesh" node.
- Added a new core compilation node to handle this new UE node.
- Added code for the generation of the ASTOp tree for this new node. It uses the same tree of Ops that the modifier does with some minor differences.
- Added support for the new Core node into the Mutable Graph Viewer.
- Minor typo fixing and include removal.

#rnx

#Synced-CL 51204827

diff --git a/Engine/Plugins/Mutable/Source/CustomizableObjectEditor/Private/MuCOE/EdGraphSchema_CustomizableObject.cpp b/Engine/Plugins/Mutable/Source/CustomizableObjectEditor/Private/MuCOE/EdGraphSchema_CustomizableObject.cpp
--- a/Engine/Plugins/Mutable/Source/CustomizableObjectEditor/Private/MuCOE/EdGraphSchema_CustomizableObject.cpp
+++ b/Engine/Plugins/Mutable/Source/CustomizableObjectEditor/Private/MuCOE/EdGraphSchema_CustomizableObject.cpp
@@ -121,6 +121,7 @@
 #include "Nodes/CONodeClipMeshWithMesh.h"
 #include "Nodes/CONodeExternalOperation.h"
 #include "Nodes/CONodeExternalTypeParameter.h"
+#include "Nodes/CONodeRemoveMesh.h"
 #include "Nodes/CONodeSkeletalMeshBreak.h"
 #include "UObject/Class.h"
 #include "Widgets/SMutablePinTypeSelector.h"
@@ -532,6 +533,7 @@ void UEdGraphSchema_CustomizableObject::GetGraphContextActions(FGraphContextMenu
 			ContextMenuBuilder.CreateTemplateNode<UCustomizableObjectNodeMeshMorphStackApplication>(),
 			ContextMenuBuilder.CreateTemplateNode<UCustomizableObjectNodeAnimationPose>(),
 			ContextMenuBuilder.CreateTemplateNode<UCONodeClipMeshWithMesh>(),
+			ContextMenuBuilder.CreateTemplateNode<UCONodeRemoveMesh>(),
 		};
 
 		AddNewNodeCategoryActionsFiltered(MeshTemplateNodes, ContextMenuBuilder, TEXT("Mesh"), GeneralGrouping, Filter);
diff --git a/Engine/Plugins/Mutable/Source/CustomizableObjectEditor/Private/MuCOE/GenerateMutableSource/GenerateMutableSourceMesh.cpp b/Engine/Plugins/Mutable/Source/CustomizableObjectEditor/Private/MuCOE/GenerateMutableSource/GenerateMutableSourceMesh.cpp
--- a/Engine/Plugins/Mutable/Source/CustomizableObjectEditor/Private/MuCOE/GenerateMutableSource/GenerateMutableSourceMesh.cpp
+++ b/Engine/Plugins/Mutable/Source/CustomizableObjectEditor/Private/MuCOE/GenerateMutableSource/GenerateMutableSourceMesh.cpp
@@ -61,8 +61,10 @@
 #include "PackedNormal.h"
 #include "MuCOE/Nodes/CONodeClipMeshWithMesh.h"
 #include "MuCOE/Nodes/CONodeExternalOperation.h"
+#include "MuCOE/Nodes/CONodeRemoveMesh.h"
 #include "MuT/NodeMeshClipWithMesh.h"
 #include "MuT/NodeMeshExternal.h"
+#include "MuT/NodeMeshRemoveMesh.h"
 #include "MuT/NodeMeshTransform.h"
 
 #define LOCTEXT_NAMESPACE "CustomizableObjectEditor"
@@ -3256,6 +3258,43 @@ UE::Mutable::Private::Ptr<UE::Mutable::Private::NodeMesh> GenerateMutableSourceM
 		Result = ClipNode;
 	}
 
+	else if (const UCONodeRemoveMesh* RemoveMeshNode = Cast<UCONodeRemoveMesh>(Node))
+	{
+		UE::Mutable::Private::Ptr<UE::Mutable::Private::NodeMeshRemoveMesh> RemoveNode = new UE::Mutable::Private::NodeMeshRemoveMesh();
+		RemoveNode->FaceCullStrategy = RemoveMeshNode->FaceCullStrategy;
+
+		// Base Mesh
+		UEdGraphPin* BaseMeshPin = RemoveMeshNode->BaseMeshPin.Get();
+		check(BaseMeshPin);
+		if (const UEdGraphPin* ConnectedPin = FollowInputPin(*BaseMeshPin))
+		{
+			UE::Mutable::Private::Ptr<UE::Mutable::Private::NodeMesh> BaseMeshCoreNode = GenerateMutableSourceMesh(ConnectedPin, GenerationContext, FMutableSourceMeshData(), false, true);
+			RemoveNode->Source = BaseMeshCoreNode;
+		}
+		
+		// Remove Mesh
+		UEdGraphPin* RemoveMeshPin = RemoveMeshNode->RemoveMeshPin.Get();
+		check(RemoveMeshPin);
+		if (const UEdGraphPin* ConnectedPin = FollowInputPin(*RemoveMeshPin))
+		{
+			constexpr EMutableMeshConversionFlags RemoveMeshFlags = 
+				EMutableMeshConversionFlags::IgnoreSkinning |
+				EMutableMeshConversionFlags::IgnorePhysics |
+				EMutableMeshConversionFlags::IgnoreMorphs |
+				EMutableMeshConversionFlags::IgnoreTexCoords |
+				EMutableMeshConversionFlags::IgnoreAUD |
+				EMutableMeshConversionFlags::DoNotCreateMeshMetadata;
+			GenerationContext.MeshGenerationFlags.Push(RemoveMeshFlags);
+			
+			UE::Mutable::Private::Ptr<UE::Mutable::Private::NodeMesh> RemoveMeshCoreNode = GenerateMutableSourceMesh(ConnectedPin, GenerationContext, FMutableSourceMeshData(), false, true);
+			RemoveNode->RemoveMesh = RemoveMeshCoreNode;
+			
+			GenerationContext.MeshGenerationFlags.Pop();
+		}
+		
+		Result = RemoveNode;
+	}
+	
 	else if (const UCustomizableObjectNodeAnimationPose* TypedNode = Cast<UCustomizableObjectNodeAnimationPose>(Node))
 	{
 		if (const UEdGraphPin* InputMeshPin = FollowInputPin(*TypedNode->GetInputMeshPin()))
diff --git a/Engine/Plugins/Mutable/Source/CustomizableObjectEditor/Private/MuCOE/GraphTraversal.cpp b/Engine/Plugins/Mutable/Source/CustomizableObjectEditor/Private/MuCOE/GraphTraversal.cpp
--- a/Engine/Plugins/Mutable/Source/CustomizableObjectEditor/Private/MuCOE/GraphTraversal.cpp
+++ b/Engine/Plugins/Mutable/Source/CustomizableObjectEditor/Private/MuCOE/GraphTraversal.cpp
@@ -30,6 +30,7 @@
 #include "MuCOE/Nodes/CustomizableObjectNodeReroute.h"
 #include "Nodes/CONodeExternalOperation.h"
 #include "Nodes/CONodeClipMeshWithMesh.h"
+#include "Nodes/CONodeRemoveMesh.h"
 #include "UObject/UObjectIterator.h"
 
 
@@ -771,6 +772,14 @@ const UEdGraphPin* FindMeshBaseSource(const UEdGraphPin& Pin, const bool bOnlyLo
 			return FindMeshBaseSource(*ConnectedPin, bOnlyLookForStaticMesh, MacroContext);
 		}
 	}
+	
+	else if (const UCONodeRemoveMesh* NodeMeshRemoveMesh = Cast<UCONodeRemoveMesh>(Node))
+	{
+		if (const UEdGraphPin* ConnectedPin = FollowInputPin(*NodeMeshRemoveMesh->BaseMeshPin.Get()))
+		{
+			return FindMeshBaseSource(*ConnectedPin, bOnlyLookForStaticMesh, MacroContext);
+		}
+	}
 
 	else if (Cast<UCONodeExternalOperation>(Node))
 	{
diff --git a/Engine/Plugins/Mutable/Source/CustomizableObjectEditor/Private/MuCOE/Nodes/CONodeRemoveMesh.cpp b/Engine/Plugins/Mutable/Source/CustomizableObjectEditor/Private/MuCOE/Nodes/CONodeRemoveMesh.cpp
new file mode 100644
--- /dev/null
+++ b/Engine/Plugins/Mutable/Source/CustomizableObjectEditor/Private/MuCOE/Nodes/CONodeRemoveMesh.cpp
@@ -0,0 +1,56 @@
+// Copyright Epic Games, Inc. All Rights Reserved.
+
+#include "MuCOE/Nodes/CONodeRemoveMesh.h"
+
+#include "MuCOE/EdGraphSchema_CustomizableObject.h"
+
+#include UE_INLINE_GENERATED_CPP_BY_NAME(CONodeRemoveMesh)
+
+class UCustomizableObjectNodeRemapPins;
+
+#define LOCTEXT_NAMESPACE "CustomizableObjectEditor"
+
+
+void UCONodeRemoveMesh::AllocateDefaultPins(UCustomizableObjectNodeRemapPins* RemapPins)
+{
+	BaseMeshPin = CustomCreatePin(EGPD_Input, UEdGraphSchema_CustomizableObject::PC_Mesh, TEXT("Base Mesh"), LOCTEXT("BaseMesh", "Base Mesh"));
+	RemoveMeshPin =  CustomCreatePin(EGPD_Input, UEdGraphSchema_CustomizableObject::PC_Mesh, TEXT("Remove Mesh"), LOCTEXT("RemoveMesh", "Remove Mesh"));
+	
+	OutputMeshPin =  CustomCreatePin(EGPD_Output, UEdGraphSchema_CustomizableObject::PC_Mesh);
+}
+
+
+FText UCONodeRemoveMesh::GetNodeTitle(ENodeTitleType::Type TitleType) const
+{
+	return LOCTEXT("Remove_Mesh", "Remove Mesh");
+}
+
+
+FText UCONodeRemoveMesh::GetTooltipText() const
+{
+	return LOCTEXT("Remove_Mesh_Tooltip",
+	"Removes the faces of a mesh section that are defined only by the vertexes shared by said mesh section and the input mesh. \n"
+	"It also removes any vertex and edge that only define deleted faces. \n"
+	"If the removed mesh covers all the faces included in one or more layout blocks those blocks get removed, freeing layout space in the final texture.");
+}
+
+
+FLinearColor UCONodeRemoveMesh::GetNodeTitleColor() const
+{
+	return UEdGraphSchema_CustomizableObject::GetPinTypeColor(UEdGraphSchema_CustomizableObject::PC_Mesh);
+}
+
+
+bool UCONodeRemoveMesh::IsSingleOutputNode() const
+{
+	return true;
+}
+
+
+bool UCONodeRemoveMesh::IsExperimental() const
+{
+	return true;
+}
+
+
+#undef LOCTEXT_NAMESPACE
diff --git a/Engine/Plugins/Mutable/Source/CustomizableObjectEditor/Private/MuCOE/Nodes/CONodeRemoveMesh.h b/Engine/Plugins/Mutable/Source/CustomizableObjectEditor/Private/MuCOE/Nodes/CONodeRemoveMesh.h
new file mode 100644
--- /dev/null
+++ b/Engine/Plugins/Mutable/Source/CustomizableObjectEditor/Private/MuCOE/Nodes/CONodeRemoveMesh.h
@@ -0,0 +1,45 @@
+// Copyright Epic Games, Inc. All Rights Reserved.
+
+#pragma once
+
+#include "MuCOE/Nodes/CustomizableObjectNode.h"
+
+#include "CONodeRemoveMesh.generated.h"
+
+namespace ENodeTitleType { enum Type : int; }
+
+class UCustomizableObjectNodeRemapPins;
+class UEdGraphPin;
+
+UCLASS()
+class UCONodeRemoveMesh : public UCustomizableObjectNode
+{
+	GENERATED_BODY()
+
+public:
+	
+	// Begin EdGraphNode interface
+	virtual FText GetNodeTitle(ENodeTitleType::Type TitleType) const override;
+	virtual FText GetTooltipText() const override;
+	virtual FLinearColor GetNodeTitleColor() const override;
+
+	// UCustomizableObjectNode interface
+	virtual void AllocateDefaultPins(UCustomizableObjectNodeRemapPins* RemapPins) override;
+	virtual bool IsSingleOutputNode() const override;
+	virtual bool IsExperimental() const override;
+
+	UPROPERTY(EditAnywhere, Category = RemoveOptions)
+	EFaceCullStrategy FaceCullStrategy = EFaceCullStrategy::AllVerticesCulled;
+	
+	UPROPERTY()
+	FEdGraphPinReference BaseMeshPin;
+	
+	UPROPERTY()
+	FEdGraphPinReference RemoveMeshPin;
+	
+	UPROPERTY()
+	FEdGraphPinReference OutputMeshPin;
+};
+
+
+#undef UE_API
diff --git a/Engine/Plugins/Mutable/Source/CustomizableObjectEditor/Private/MuCOE/SMutableGraphViewer.cpp b/Engine/Plugins/Mutable/Source/CustomizableObjectEditor/Private/MuCOE/SMutableGraphViewer.cpp
--- a/Engine/Plugins/Mutable/Source/CustomizableObjectEditor/Private/MuCOE/SMutableGraphViewer.cpp
+++ b/Engine/Plugins/Mutable/Source/CustomizableObjectEditor/Private/MuCOE/SMutableGraphViewer.cpp
@@ -50,6 +50,7 @@
 #include "MuT/NodeMeshFragment.h"
 #include "MuT/NodeMeshMakeMorph.h"
 #include "MuT/NodeMeshMorph.h"
+#include "MuT/NodeMeshRemoveMesh.h"
 #include "MuT/NodeMeshTable.h"
 #include "MuT/NodeModifierMeshClipDeform.h"
 #include "MuT/NodeModifierMeshClipMorphPlane.h"
@@ -101,6 +102,7 @@ static const char* MutableNodeNames[] =
 	"MeshTransform",
 	"MeshClipMorphPlane",
 	"MeshClipWithMesh",
+	"MeshRemoveMesh",
 	"MeshApplyPose",
 	"MeshVariation",
 	"MeshReshape",
@@ -624,6 +626,13 @@ void SMutableGraphViewer::GetChildrenForInfo(TSharedPtr<FMutableGraphTreeElement
 		AddChildFunc(ClipMeshWithMeshVar->ClipMesh.get(), FString::Printf(TEXT("CLIP MESH")));
 	}
 
+	else if (ParentNode->GetType() == UE::Mutable::Private::NodeMeshRemoveMesh::GetStaticType())
+	{
+		UE::Mutable::Private::NodeMeshRemoveMesh* RemoveMeshVar = StaticCast<UE::Mutable::Private::NodeMeshRemoveMesh*>(ParentNode);
+		AddChildFunc(RemoveMeshVar->Source.get(), FString::Printf(TEXT("SOURCE MESH")));
+		AddChildFunc(RemoveMeshVar->RemoveMesh.get(), FString::Printf(TEXT("REMOVE MESH")));
+	}
+	
 	else if (ParentNode->GetType() == UE::Mutable::Private::NodeModifierMeshClipDeform::GetStaticType())
 	{
 		UE::Mutable::Private::NodeModifierMeshClipDeform* ModifierMeshClipDeformVar = StaticCast<UE::Mutable::Private::NodeModifierMeshClipDeform*>(ParentNode);
diff --git a/Engine/Plugins/Mutable/Source/MutableTools/Internal/MuT/Node.h b/Engine/Plugins/Mutable/Source/MutableTools/Internal/MuT/Node.h
--- a/Engine/Plugins/Mutable/Source/MutableTools/Internal/MuT/Node.h
+++ b/Engine/Plugins/Mutable/Source/MutableTools/Internal/MuT/Node.h
@@ -44,6 +44,7 @@ namespace UE::Mutable::Private
 			MeshTransform,
 			MeshClipMorphPlane,
 			MeshClipWithMesh,
+			MeshRemoveMesh,
 			MeshApplyPose,
 			MeshVariation,
 			MeshReshape,
diff --git a/Engine/Plugins/Mutable/Source/MutableTools/Internal/MuT/NodeMeshRemoveMesh.h b/Engine/Plugins/Mutable/Source/MutableTools/Internal/MuT/NodeMeshRemoveMesh.h
new file mode 100644
--- /dev/null
+++ b/Engine/Plugins/Mutable/Source/MutableTools/Internal/MuT/NodeMeshRemoveMesh.h
@@ -0,0 +1,44 @@
+// Copyright Epic Games, Inc. All Rights Reserved.
+
+#pragma once
+
+#include "MuR/Ptr.h"
+#include "MuR/Types.h"
+#include "MuT/Node.h"
+#include "MuT/NodeMesh.h"
+
+#define UE_API MUTABLETOOLS_API
+
+
+namespace UE::Mutable::Private
+{
+	
+	class NodeMeshRemoveMesh : public NodeMesh
+	{
+	public:
+
+		Ptr<NodeMesh> Source;
+    	
+		Ptr<NodeMesh> RemoveMesh;
+
+		EFaceCullStrategy FaceCullStrategy = EFaceCullStrategy::AllVerticesCulled;
+    	
+		// Node interface
+		virtual const FNodeType* GetType() const override { return &StaticType; }
+		static const FNodeType* GetStaticType() { return &StaticType; }
+
+	protected:
+
+		/** Forbidden. Manage with the Ptr<> template. */
+		inline ~NodeMeshRemoveMesh() = default;
+
+	private:
+
+		static UE_API FNodeType StaticType;
+
+	};
+
+
+}
+
+#undef UE_API
diff --git a/Engine/Plugins/Mutable/Source/MutableTools/Private/MuT/CodeGenerator.h b/Engine/Plugins/Mutable/Source/MutableTools/Private/MuT/CodeGenerator.h
--- a/Engine/Plugins/Mutable/Source/MutableTools/Private/MuT/CodeGenerator.h
+++ b/Engine/Plugins/Mutable/Source/MutableTools/Private/MuT/CodeGenerator.h
@@ -106,6 +106,7 @@ namespace UE::Mutable::Private
 	class NodeMeshClipDeform;
 	class NodeMeshClipMorphPlane;
 	class NodeMeshClipWithMesh;
+	class NodeMeshRemoveMesh;
 	class NodeMeshConstant;
 	class NodeMeshFormat;
 	class NodeMeshFragment;
@@ -856,7 +857,8 @@ namespace UE::Mutable::Private
         FMeshTask GenerateMesh_Transform(const FMeshGenerationStaticOptions&, FMeshOptionsTask, const NodeMeshTransform*);
         FMeshTask GenerateMesh_ClipMorphPlane(const FMeshGenerationStaticOptions&, FMeshOptionsTask, const NodeMeshClipMorphPlane*);
         FMeshTask GenerateMesh_ClipWithMesh(const FMeshGenerationStaticOptions&, FMeshOptionsTask, const NodeMeshClipWithMesh*);
-        FMeshTask GenerateMesh_ApplyPose(const FMeshGenerationStaticOptions&, FMeshOptionsTask, const NodeMeshApplyPose*);
+        FMeshTask GenerateMesh_RemoveMesh(const FMeshGenerationStaticOptions&, FMeshOptionsTask, const NodeMeshRemoveMesh*);
+    	FMeshTask GenerateMesh_ApplyPose(const FMeshGenerationStaticOptions&, FMeshOptionsTask, const NodeMeshApplyPose*);
         FMeshTask GenerateMesh_Variation(const FMeshGenerationStaticOptions&, FMeshOptionsTask, const NodeMeshVariation*);
 		FMeshTask GenerateMesh_Table(const FMeshGenerationStaticOptions&, FMeshOptionsTask, const NodeMeshTable*);
 		FMeshTask GenerateMesh_Reshape(const FMeshGenerationStaticOptions&, FMeshOptionsTask, const NodeMeshReshape*);
diff --git a/Engine/Plugins/Mutable/Source/MutableTools/Private/MuT/CodeGenerator_Mesh.cpp b/Engine/Plugins/Mutable/Source/MutableTools/Private/MuT/CodeGenerator_Mesh.cpp
--- a/Engine/Plugins/Mutable/Source/MutableTools/Private/MuT/CodeGenerator_Mesh.cpp
+++ b/Engine/Plugins/Mutable/Source/MutableTools/Private/MuT/CodeGenerator_Mesh.cpp
@@ -3,6 +3,7 @@
 #include "ASTOpConstantBool.h"
 #include "ASTOpExternal.h"
 #include "ASTOpMeshMaskClipMesh.h"
+#include "ASTOpMeshMaskDiff.h"
 #include "ASTOpMeshRemoveMask.h"
 #include "ASTOpMeshSkeletalMeshBreak.h"
 #include "Containers/Array.h"
@@ -51,6 +52,7 @@
 #include "MuT/NodeMeshFragment.h"
 #include "MuT/NodeMeshMakeMorph.h"
 #include "MuT/NodeMeshMorph.h"
+#include "MuT/NodeMeshRemoveMesh.h"
 #include "MuT/NodeMeshReshape.h"
 #include "MuT/NodeMeshSwitch.h"
 #include "MuT/NodeMeshTable.h"
@@ -96,6 +98,7 @@ namespace UE::Mutable::Private
 		case Node::EType::MeshTransform: Result = GenerateMesh_Transform(StaticOptions, InOptions, static_cast<const NodeMeshTransform*>(Node)); break;
 		case Node::EType::MeshClipMorphPlane: Result = GenerateMesh_ClipMorphPlane(StaticOptions, InOptions, static_cast<const NodeMeshClipMorphPlane*>(Node)); break;
 		case Node::EType::MeshClipWithMesh: Result = GenerateMesh_ClipWithMesh(StaticOptions, InOptions, static_cast<const NodeMeshClipWithMesh*>(Node)); break;
+		case Node::EType::MeshRemoveMesh: Result = GenerateMesh_RemoveMesh(StaticOptions, InOptions, static_cast<const NodeMeshRemoveMesh*>(Node)); break;
 		case Node::EType::MeshApplyPose: Result = GenerateMesh_ApplyPose(StaticOptions, InOptions, static_cast<const NodeMeshApplyPose*>(Node)); break;
 		case Node::EType::MeshVariation: Result = GenerateMesh_Variation(StaticOptions, InOptions, static_cast<const NodeMeshVariation*>(Node)); break;
 		case Node::EType::MeshTable: Result = GenerateMesh_Table(StaticOptions, InOptions, static_cast<const NodeMeshTable*>(Node)); break;
@@ -1171,8 +1174,7 @@ namespace UE::Mutable::Private
 					MaskOp->source = SourceResult.MeshOp;
 					MaskOp->clip = ClipResult.MeshOp;
 
-					Ptr<UE::Mutable::Private::ASTOpConstantBool> ConditionOp = new UE::Mutable::Private::ASTOpConstantBool();
-					ConditionOp->bValue = true;
+					Ptr<UE::Mutable::Private::ASTOpConstantBool> ConditionOp = new UE::Mutable::Private::ASTOpConstantBool(true);
 										
 					Ptr<UE::Mutable::Private::ASTOpMeshRemoveMask> RemoveOp = new UE::Mutable::Private::ASTOpMeshRemoveMask();
 					RemoveOp->source = SourceResult.MeshOp;
@@ -1184,14 +1186,14 @@ namespace UE::Mutable::Private
 				}
 				else
 				{
-					Ptr<ASTOpMeshClipWithMesh> op = new ASTOpMeshClipWithMesh();
+					Ptr<ASTOpMeshClipWithMesh> ClipOp = new ASTOpMeshClipWithMesh();
 
 					// Base
-					op->Source = SourceResult.MeshOp;
-					op->ClipMesh = ClipResult.MeshOp;
+					ClipOp->Source = SourceResult.MeshOp;
+					ClipOp->ClipMesh = ClipResult.MeshOp;
 
 					Result = SourceResult;
-					Result.MeshOp = op;
+					Result.MeshOp = ClipOp;
 				}
 					
 				return Result;
@@ -1201,6 +1203,66 @@ namespace UE::Mutable::Private
     }
 
 
+	FMeshTask CodeGenerator::GenerateMesh_RemoveMesh(const FMeshGenerationStaticOptions& StaticOptions, FMeshOptionsTask InOptions, const NodeMeshRemoveMesh* RemoveMeshNode)
+	{
+		if (!RemoveMeshNode->Source)
+		{
+			// This argument is required
+			ErrorLog->Add("RemoveMesh base mesh node is not set.", ELMT_ERROR, RemoveMeshNode->GetMessageContext());
+			return UE::Tasks::MakeCompletedTask<FMeshGenerationResult>();
+		}
+		
+		if (!RemoveMeshNode->RemoveMesh)
+		{
+			// This argument is required
+			ErrorLog->Add("RemoveMesh mesh to remove node is not set.", ELMT_ERROR, RemoveMeshNode->GetMessageContext());
+			return UE::Tasks::MakeCompletedTask<FMeshGenerationResult>();
+		}
+		
+		FMeshTask BaseMeshTask = GenerateMesh(StaticOptions, InOptions, RemoveMeshNode->Source);
+		
+		FMeshGenerationStaticOptions StaticRemoveOptions = StaticOptions;
+		StaticRemoveOptions.ActiveTags.Empty();
+		FMeshOptionsTask RemoveMeshOptionsTask = UE::Tasks::Launch(TEXT("MutableMeshRemoveMeshOptions"),
+			[InOptions]() mutable
+			{
+				FMeshGenerationDynamicOptions Result = InOptions.GetResult();
+				Result.bLayouts = false;
+				Result.OverrideLayouts.Empty();
+				return Result;
+			},
+			InOptions
+		);
+		
+		FMeshTask RemoveMeshTask = GenerateMesh(StaticRemoveOptions, RemoveMeshOptionsTask, RemoveMeshNode->RemoveMesh);
+		
+		return UE::Tasks::Launch(TEXT("MutableRemoveMesh"),
+			[RemoveMeshTask, BaseMeshTask, RemoveMeshNode]() mutable
+			{
+				FMeshGenerationResult RemoveResults = RemoveMeshTask.GetResult();
+				FMeshGenerationResult BaseResults = BaseMeshTask.GetResult();
+					
+				Ptr<ASTOpMeshMaskDiff> MaskOp = new ASTOpMeshMaskDiff;
+				MaskOp->Source = BaseResults.BaseMeshOp;
+				MaskOp->Fragment = RemoveResults.MeshOp;
+					
+				Ptr<UE::Mutable::Private::ASTOpConstantBool> ConditionOp = new UE::Mutable::Private::ASTOpConstantBool(true);
+					
+				Ptr<ASTOpMeshRemoveMask> RemoveOp = new ASTOpMeshRemoveMask();
+				RemoveOp->source = BaseResults.MeshOp;
+				RemoveOp->FaceCullStrategy = RemoveMeshNode->FaceCullStrategy;
+				RemoveOp->AddRemove(ConditionOp, MaskOp);
+
+				FMeshGenerationResult Result = BaseResults;
+				Result.MeshOp = RemoveOp;
+					
+				return Result;
+			},
+			UE::Tasks::Prerequisites( BaseMeshTask, RemoveMeshTask)
+		);
+	}
+
+
 	FMeshTask CodeGenerator::GenerateMesh_ClipDeform(const FMeshGenerationStaticOptions& StaticOptions, FMeshOptionsTask InOptions, const NodeMeshClipDeform* ClipDeform)
 	{
 		if (!ClipDeform->BaseMesh)
diff --git a/Engine/Plugins/Mutable/Source/MutableTools/Private/MuT/Node.cpp b/Engine/Plugins/Mutable/Source/MutableTools/Private/MuT/Node.cpp
--- a/Engine/Plugins/Mutable/Source/MutableTools/Private/MuT/Node.cpp
+++ b/Engine/Plugins/Mutable/Source/MutableTools/Private/MuT/Node.cpp
@@ -2,7 +2,6 @@
 
 
 #include "MuT/Node.h"
-#include "Misc/AssertionMacros.h"
 
 #include "MuT/NodeObject.h"
 #include "MuT/NodeObjectNew.h"
@@ -93,7 +92,6 @@
 #include "MuT/NodeString.h"
 #include "MuT/NodeStringConstant.h"
 #include "MuT/NodeStringParameter.h"
-#include "MuT/NodeBool.h"
 #include "MuT/NodeColourExternal.h"
 #include "MuT/NodeExternal.h"
 #include "MuT/NodeExternalOperation.h"
@@ -118,12 +116,13 @@
 #include "MuT/NodeSkeletalMeshNew.h"
 #include "MuT/NodeSkeletalMeshParameter.h"
 #include "MuT/NodeMaterialModify.h"
+#include "MuT/NodeMeshRemoveMesh.h"
 
 
 namespace UE::Mutable::Private
 {
 
-	// Static initialisation
+	// Static initialization
 	FNodeType Node::StaticType = FNodeType(Node::EType::Node, nullptr);
 
 	FNodeType NodeObject::StaticType = FNodeType(Node::EType::Object, Node::GetStaticType());
@@ -217,6 +216,7 @@ namespace UE::Mutable::Private
 	FNodeType NodeMeshClipMorphPlane::StaticType = FNodeType(Node::EType::MeshClipMorphPlane, NodeMesh::GetStaticType());
 	FNodeType NodeMeshClipDeform::StaticType = FNodeType(Node::EType::MeshClipDeform, NodeMesh::GetStaticType());
 	FNodeType NodeMeshClipWithMesh::StaticType = FNodeType(Node::EType::MeshClipWithMesh, NodeMesh::GetStaticType());
+	FNodeType NodeMeshRemoveMesh::StaticType = FNodeType(Node::EType::MeshRemoveMesh, NodeMesh::GetStaticType());
 	FNodeType NodeMeshSkeletalMeshBreak::StaticType = FNodeType(Node::EType::MeshSkeletalMeshBreak, NodeMesh::GetStaticType());
 	FNodeType NodeMeshMakeMorph::StaticType = FNodeType(Node::EType::MeshMakeMorph, NodeMesh::GetStaticType());
 	FNodeType NodeMeshApplyPose::StaticType = FNodeType(Node::EType::MeshApplyPose, NodeMesh::GetStaticType());
