[mutable] Made the real-time morphs list respond to the state of the SelectAllMorphs bool UI control.
- Now, if the "Select All Morphs" option is set to true, then the list of morphs will be disabled so it can not be interacted with.

#rnx

#Synced-CL 51205810

diff --git a/Engine/Plugins/Mutable/Source/CustomizableObjectEditor/Private/MuCOE/SCustomizableObjectNodeSkeletalMeshRTMorphSelector.cpp b/Engine/Plugins/Mutable/Source/CustomizableObjectEditor/Private/MuCOE/SCustomizableObjectNodeSkeletalMeshRTMorphSelector.cpp
--- a/Engine/Plugins/Mutable/Source/CustomizableObjectEditor/Private/MuCOE/SCustomizableObjectNodeSkeletalMeshRTMorphSelector.cpp
+++ b/Engine/Plugins/Mutable/Source/CustomizableObjectEditor/Private/MuCOE/SCustomizableObjectNodeSkeletalMeshRTMorphSelector.cpp
@@ -115,6 +115,7 @@ void SCustomizableObjectNodeSkeletalMeshRTMorphSelector::Construct(const FArgume
     UpdateWidget();
 }
 
+
 void SCustomizableObjectNodeSkeletalMeshRTMorphSelector::AddReferencedObjects(
         FReferenceCollector& Collector)
 {
@@ -166,6 +167,7 @@ void SCustomizableObjectNodeSkeletalMeshRTMorphSelector::UpdateWidget()
         +SVerticalBox::Slot()
         [
             SAssignNew(ListWidget, ListViewType)
+	        .IsEnabled(this, &ThisType::IsRuntimeMorphsListEnabled)
             .ListItemsSource(&Items)
             .OnGenerateRow(this, &ThisType::GenerateRow)
             .SelectionMode(ESelectionMode::Single)
@@ -185,6 +187,7 @@ void SCustomizableObjectNodeSkeletalMeshRTMorphSelector::UpdateWidget()
     ListWidget->RequestListRefresh();
 }
 
+
 void SCustomizableObjectNodeSkeletalMeshRTMorphSelector::GenerateMorphsInfoList()
 {
     Items.Reset();
@@ -216,6 +219,7 @@ void SCustomizableObjectNodeSkeletalMeshRTMorphSelector::GenerateMorphsInfoList(
     });
 }
 
+
 TSharedRef<ITableRow> SCustomizableObjectNodeSkeletalMeshRTMorphSelector::GenerateRow(
     TSharedPtr<SCustomizableObjectNodeSkeletalMeshRTMorphSelector::ItemType> InItem, 
     const TSharedRef<STableViewBase>& OwnerTable)
@@ -230,6 +234,13 @@ TSharedRef<ITableRow> SCustomizableObjectNodeSkeletalMeshRTMorphSelector::Genera
         .OwnerWidget(SharedThis(this));
 }
 
+
+bool SCustomizableObjectNodeSkeletalMeshRTMorphSelector::IsRuntimeMorphsListEnabled() const
+{
+	return Node ? !Node->bUseAllRealTimeMorphs : true;
+}
+
+
 ECheckBoxState SCustomizableObjectNodeSkeletalMeshRTMorphSelector::IsSelectAllChecked() const
 {
     if (!Node)
@@ -242,6 +253,7 @@ ECheckBoxState SCustomizableObjectNodeSkeletalMeshRTMorphSelector::IsSelectAllCh
             : ECheckBoxState::Unchecked;
 }
 
+
 void SCustomizableObjectNodeSkeletalMeshRTMorphSelector::OnSelectAllCheckStateChanged(
         ECheckBoxState CheckBoxState)
 {
diff --git a/Engine/Plugins/Mutable/Source/CustomizableObjectEditor/Private/MuCOE/SCustomizableObjectNodeSkeletalMeshRTMorphSelector.h b/Engine/Plugins/Mutable/Source/CustomizableObjectEditor/Private/MuCOE/SCustomizableObjectNodeSkeletalMeshRTMorphSelector.h
--- a/Engine/Plugins/Mutable/Source/CustomizableObjectEditor/Private/MuCOE/SCustomizableObjectNodeSkeletalMeshRTMorphSelector.h
+++ b/Engine/Plugins/Mutable/Source/CustomizableObjectEditor/Private/MuCOE/SCustomizableObjectNodeSkeletalMeshRTMorphSelector.h
@@ -77,6 +77,7 @@ public:
     TSharedRef<ITableRow> GenerateRow(
             TSharedPtr<ItemType> InItem, 
             const TSharedRef<STableViewBase>& OwnerTable);
+    bool IsRuntimeMorphsListEnabled() const;
 
     ECheckBoxState IsSelectAllChecked() const;
     void OnSelectAllCheckStateChanged(ECheckBoxState CheckBoxState);
