[mutable] Mesh Transform with Bone node.

- New Node "CONodeTransformWithBone"
-- This new node is an inlined version of the modifier node with the same name.
-- Added a core compilation node to support this new type of operation.
-- Added code generation code for this case.
-- This node does get interpreted quite similarly to the modifier with the same name (ASTOp-wise).

- Improved Transform Mesh and Transform mesh with Bone modifiers.
-- Removed the details for the withMesh and WithBone transform modifier nodes as neither of the two detail objects was doing any work.
-- Updated both nodes to make use of pin references.
-- Removed transform controls for the TransformWithBone node as those got only used in the preview and never in the actual compilation. The controls also did not work due to an incomplete implementation of the feature.

- Updated debugger with some more types that are related to the nodes modified in this CL.
- Other small changes.

#rnx

#Synced-CL 51205810

diff --git a/Engine/Plugins/Mutable/Source/CustomizableObject/Internal/MuCO/CustomizableObjectCustomVersion.h b/Engine/Plugins/Mutable/Source/CustomizableObject/Internal/MuCO/CustomizableObjectCustomVersion.h
--- a/Engine/Plugins/Mutable/Source/CustomizableObject/Internal/MuCO/CustomizableObjectCustomVersion.h
+++ b/Engine/Plugins/Mutable/Source/CustomizableObject/Internal/MuCO/CustomizableObjectCustomVersion.h
@@ -228,6 +228,8 @@ struct FCustomizableObjectCustomVersion
 
 		RemovedOverlayMaterialUpropertyFromSkeletalMeshComponentNode,
 		
+		NodeModifierTransformInMeshStorePinsAsReferences,
+		
 		// -----<new versions can be added above this line>-------------------------------------------------
 		VersionPlusOne,
 		LatestVersion = VersionPlusOne - 1
diff --git a/Engine/Plugins/Mutable/Source/CustomizableObjectEditor/Private/MuCOE/CustomizableObjectEditor.cpp b/Engine/Plugins/Mutable/Source/CustomizableObjectEditor/Private/MuCOE/CustomizableObjectEditor.cpp
--- a/Engine/Plugins/Mutable/Source/CustomizableObjectEditor/Private/MuCOE/CustomizableObjectEditor.cpp
+++ b/Engine/Plugins/Mutable/Source/CustomizableObjectEditor/Private/MuCOE/CustomizableObjectEditor.cpp
@@ -41,7 +41,6 @@
 #include "MuCOE/Nodes/CustomizableObjectNodeSkeletalMesh.h"
 #include "MuCOE/Nodes/CustomizableObjectNodeStaticMesh.h"
 #include "MuCOE/Nodes/CustomizableObjectNodeTable.h"
-#include "MuCOE/Nodes/CONodeModifierTransformWithBone.h"
 #include "MuCOE/Nodes/CONodeClipMeshWithMesh.h"
 #include "MuCOE/SCustomizableObjectEditorAdvancedPreviewSettings.h"
 #include "MuCOE/SCustomizableObjectEditorTagExplorer.h"
@@ -964,7 +963,6 @@ void FCustomizableObjectEditor::HideGizmoClipMesh()
 		const UObject* Node = *NodeIt;
 		if (Node->IsA<UCustomizableObjectNodeModifierClipWithMesh>() ||
 			Node->IsA<UCustomizableObjectNodeModifierTransformInMesh>() ||
-			Node->IsA<UCONodeModifierTransformWithBone>() || 
 			Node->IsA<UCONodeClipMeshWithMesh>())
 		{
 			GraphEditor->ClearSelectionSet();
@@ -1639,7 +1637,7 @@ void FCustomizableObjectEditor::OnSelectedGraphNodesChanged(const FGraphPanelSel
 		}
 		else if (UCustomizableObjectNodeModifierTransformInMesh* NodeModifierTransformInMesh = Cast<UCustomizableObjectNodeModifierTransformInMesh>(Objects[0]))
 		{
-			if (const UEdGraphPin* Pin = NodeModifierTransformInMesh->GetBoundingMeshPin())
+			if (const UEdGraphPin* Pin = NodeModifierTransformInMesh->BoundingMeshPin.Get())
 			{
 				ShowGizmoClipMesh(*NodeModifierTransformInMesh, &NodeModifierTransformInMesh->BoundingMeshTransform, *Pin);
 			}
diff --git a/Engine/Plugins/Mutable/Source/CustomizableObjectEditor/Private/MuCOE/CustomizableObjectEditorModule.cpp b/Engine/Plugins/Mutable/Source/CustomizableObjectEditor/Private/MuCOE/CustomizableObjectEditorModule.cpp
--- a/Engine/Plugins/Mutable/Source/CustomizableObjectEditor/Private/MuCOE/CustomizableObjectEditorModule.cpp
+++ b/Engine/Plugins/Mutable/Source/CustomizableObjectEditor/Private/MuCOE/CustomizableObjectEditorModule.cpp
@@ -50,10 +50,6 @@
 #include "MuCOE/Nodes/CustomizableObjectNodeModifierClipMorphDetails.h"
 #include "MuCOE/Nodes/CustomizableObjectNodeModifierClipWithUVMask.h"
 #include "MuCOE/Nodes/CustomizableObjectNodeModifierClipDeform.h"
-#include "MuCOE/Nodes/CustomizableObjectNodeModifierTransformInMesh.h"
-#include "MuCOE/Nodes/CustomizableObjectNodeModifierTransformInMeshDetails.h"
-#include "MuCOE/Nodes/CONodeModifierTransformWithBone.h"
-#include "MuCOE/Nodes/CONodeModifierTransformWithBoneDetails.h"
 #include "MuCOE/Nodes/CustomizableObjectNodeSkeletalMeshParameter.h"
 #include "MuCOE/Nodes/CustomizableObjectNodeMeshMorph.h"
 #include "MuCOE/Nodes/CustomizableObjectNodeMeshMorphDetails.h"
@@ -252,8 +248,6 @@ void FCustomizableObjectEditorModule::StartupModule()
 	RegisterCustomDetails(PropertyModule, UCustomizableObjectNodeModifierClipMorph::StaticClass(), FOnGetDetailCustomizationInstance::CreateStatic(&FCustomizableObjectNodeModifierClipMorphDetails::MakeInstance));
 	RegisterCustomDetails(PropertyModule, UCustomizableObjectNodeModifierClipWithUVMask::StaticClass(), FOnGetDetailCustomizationInstance::CreateStatic(&FCustomizableObjectNodeModifierBaseDetails::MakeInstance));
 	RegisterCustomDetails(PropertyModule, UCustomizableObjectNodeModifierClipDeform::StaticClass(), FOnGetDetailCustomizationInstance::CreateStatic(&FCustomizableObjectNodeModifierBaseDetails::MakeInstance));
-	RegisterCustomDetails(PropertyModule, UCustomizableObjectNodeModifierTransformInMesh::StaticClass(), FOnGetDetailCustomizationInstance::CreateStatic(&FCustomizableObjectNodeModifierTransformInMeshDetails::MakeInstance));
-	RegisterCustomDetails(PropertyModule, UCONodeModifierTransformWithBone::StaticClass(), FOnGetDetailCustomizationInstance::CreateStatic(&FCONodeModifierTransformWithBoneDetails::MakeInstance));
 	RegisterCustomDetails(PropertyModule, UCustomizableObjectNodeObject::StaticClass(), FOnGetDetailCustomizationInstance::CreateStatic(&FCustomizableObjectNodeObjectDetails::MakeInstance));
 	RegisterCustomDetails(PropertyModule, UCustomizableObjectNodeObjectGroup::StaticClass(), FOnGetDetailCustomizationInstance::CreateStatic(&FCustomizableObjectNodeObjectGroupDetails::MakeInstance));
 	RegisterCustomDetails(PropertyModule, UCustomizableObjectNodeProjectorParameter::StaticClass(), FOnGetDetailCustomizationInstance::CreateStatic(&FCustomizableObjectNodeProjectorParameterDetails::MakeInstance));
diff --git a/Engine/Plugins/Mutable/Source/CustomizableObjectEditor/Private/MuCOE/CustomizableObjectEditorViewportClient.cpp b/Engine/Plugins/Mutable/Source/CustomizableObjectEditor/Private/MuCOE/CustomizableObjectEditorViewportClient.cpp
--- a/Engine/Plugins/Mutable/Source/CustomizableObjectEditor/Private/MuCOE/CustomizableObjectEditorViewportClient.cpp
+++ b/Engine/Plugins/Mutable/Source/CustomizableObjectEditor/Private/MuCOE/CustomizableObjectEditorViewportClient.cpp
@@ -35,7 +35,6 @@
 #include "MuCOE/Nodes/CustomizableObjectNodeModifierClipWithMesh.h"
 #include "MuCOE/Nodes/CustomizableObjectNodeModifierTransformInMesh.h"
 #include "MuCOE/Nodes/CustomizableObjectNodeProjectorConstant.h"
-#include "MuCOE/Nodes/CONodeModifierTransformWithBone.h"
 #include "MuCOE/UnrealEditorPortabilityHelpers.h"
 #include "Preferences/PersonaOptions.h"
 #include "SceneView.h"
@@ -644,10 +643,6 @@ void FCustomizableObjectEditorViewportClient::ShowGizmoClipMesh(UCustomizableObj
 	{
 		TransformExternallyChangedDelegateHandle = NodeModifierTransformInMesh->TransformChangedDelegate.AddSP(this, &FCustomizableObjectEditorViewportClient::UpdateGizmoClipMeshTransform);
 	}
-	else if (UCONodeModifierTransformWithBone* NodeModifierTransformWithBone = Cast<UCONodeModifierTransformWithBone>(ClipMeshNode.Get()))
-	{
-		TransformExternallyChangedDelegateHandle = NodeModifierTransformWithBone->TransformChangedDelegate.AddSP(this, &FCustomizableObjectEditorViewportClient::UpdateGizmoClipMeshTransform);
-	}
 	else if (UCONodeClipMeshWithMesh* NodeClipMeshWithMesh = Cast<UCONodeClipMeshWithMesh>(ClipMeshNode.Get()))
 	{
 		TransformExternallyChangedDelegateHandle = NodeClipMeshWithMesh->TransformChangedDelegate.AddSP(this, &FCustomizableObjectEditorViewportClient::UpdateGizmoClipMeshTransform);
@@ -714,10 +709,6 @@ void FCustomizableObjectEditorViewportClient::HideGizmoClipMesh()
 			{
 				NodeModifierTransformInMesh->TransformChangedDelegate.Remove(TransformExternallyChangedDelegateHandle);
 			}
-			else if (UCONodeModifierTransformWithBone* NodeModifierTransformWithBone = Cast<UCONodeModifierTransformWithBone>(ClipMeshNode.Get()))
-			{
-				NodeModifierTransformWithBone->TransformChangedDelegate.Remove(TransformExternallyChangedDelegateHandle);
-			}
 			else if (UCONodeClipMeshWithMesh* NodeClipMeshWithMesh = Cast<UCONodeClipMeshWithMesh>(ClipMeshNode.Get()))
 			{
 				NodeClipMeshWithMesh->TransformChangedDelegate.Remove(TransformExternallyChangedDelegateHandle);
diff --git a/Engine/Plugins/Mutable/Source/CustomizableObjectEditor/Private/MuCOE/EdGraphSchema_CustomizableObject.cpp b/Engine/Plugins/Mutable/Source/CustomizableObjectEditor/Private/MuCOE/EdGraphSchema_CustomizableObject.cpp
--- a/Engine/Plugins/Mutable/Source/CustomizableObjectEditor/Private/MuCOE/EdGraphSchema_CustomizableObject.cpp
+++ b/Engine/Plugins/Mutable/Source/CustomizableObjectEditor/Private/MuCOE/EdGraphSchema_CustomizableObject.cpp
@@ -123,6 +123,7 @@
 #include "Nodes/CONodeExternalTypeParameter.h"
 #include "Nodes/CONodeRemoveMesh.h"
 #include "Nodes/CONodeSkeletalMeshBreak.h"
+#include "Nodes/CONodeTransformWithBone.h"
 #include "UObject/Class.h"
 #include "Widgets/SMutablePinTypeSelector.h"
 
@@ -504,6 +505,7 @@ void UEdGraphSchema_CustomizableObject::GetGraphContextActions(FGraphContextMenu
 			ContextMenuBuilder.CreateTemplateNode<UCustomizableObjectNodeModifierClipMorph>(),
 			ContextMenuBuilder.CreateTemplateNode<UCustomizableObjectNodeModifierClipWithMesh>(),
 			ContextMenuBuilder.CreateTemplateNode<UCustomizableObjectNodeModifierClipWithUVMask>(),
+			ContextMenuBuilder.CreateTemplateNode<UCONodeModifierTransformWithBone>(),
 		};
 
 		AddNewNodeCategoryActionsFiltered(ModifierTemplateNodes, ContextMenuBuilder, TEXT("Modifier"), GeneralGrouping, Filter);
@@ -534,6 +536,7 @@ void UEdGraphSchema_CustomizableObject::GetGraphContextActions(FGraphContextMenu
 			ContextMenuBuilder.CreateTemplateNode<UCustomizableObjectNodeAnimationPose>(),
 			ContextMenuBuilder.CreateTemplateNode<UCONodeClipMeshWithMesh>(),
 			ContextMenuBuilder.CreateTemplateNode<UCONodeRemoveMesh>(),
+			ContextMenuBuilder.CreateTemplateNode<UCONodeTransformWithBone>(),
 		};
 
 		AddNewNodeCategoryActionsFiltered(MeshTemplateNodes, ContextMenuBuilder, TEXT("Mesh"), GeneralGrouping, Filter);
@@ -650,7 +653,6 @@ void UEdGraphSchema_CustomizableObject::GetGraphContextActions(FGraphContextMenu
 			ContextMenuBuilder.CreateTemplateNode<UCustomizableObjectNodeMeshReshape>(),
 			ContextMenuBuilder.CreateTemplateNode<UCustomizableObjectNodeModifierClipDeform>(),
 			ContextMenuBuilder.CreateTemplateNode<UCustomizableObjectNodeModifierTransformInMesh>(),
-			ContextMenuBuilder.CreateTemplateNode<UCONodeModifierTransformWithBone>(),
 			ContextMenuBuilder.CreateTemplateNode<UCustomizableObjectNodeTextureParameter>(),
 			ContextMenuBuilder.CreateTemplateNode<UCustomizableObjectNodeSkeletalMeshParameter>(),
 			ContextMenuBuilder.CreateTemplateNode<UCustomizableObjectNodeMaterialParameter>(),
diff --git a/Engine/Plugins/Mutable/Source/CustomizableObjectEditor/Private/MuCOE/GenerateMutableSource/GenerateMutableSourceMesh.cpp b/Engine/Plugins/Mutable/Source/CustomizableObjectEditor/Private/MuCOE/GenerateMutableSource/GenerateMutableSourceMesh.cpp
--- a/Engine/Plugins/Mutable/Source/CustomizableObjectEditor/Private/MuCOE/GenerateMutableSource/GenerateMutableSourceMesh.cpp
+++ b/Engine/Plugins/Mutable/Source/CustomizableObjectEditor/Private/MuCOE/GenerateMutableSource/GenerateMutableSourceMesh.cpp
@@ -12,6 +12,7 @@
 #include "GenerateMutableSourceExternal.h"
 #include "GenerateMutableSourceImage.h"
 #include "GenerateMutableSourceSkeletalMesh.h"
+#include "GenerateMutableSourceTransform.h"
 #include "MuCO/CustomizableObjectPrivate.h"
 #include "MuCO/CustomizableObjectSystem.h"
 #include "MuCO/UnrealConversionUtils.h"
@@ -60,12 +61,14 @@
 #include "Hash/CityHash.h"
 #include "PackedNormal.h"
 #include "MuCOE/Nodes/CONodeClipMeshWithMesh.h"
+#include "MuCOE/Nodes/CONodeTransformWithBone.h"
 #include "MuCOE/Nodes/CONodeExternalOperation.h"
 #include "MuCOE/Nodes/CONodeRemoveMesh.h"
 #include "MuT/NodeMeshClipWithMesh.h"
 #include "MuT/NodeMeshExternal.h"
 #include "MuT/NodeMeshRemoveMesh.h"
 #include "MuT/NodeMeshTransform.h"
+#include "MuT/NodeMeshTransformWithBone.h"
 
 #define LOCTEXT_NAMESPACE "CustomizableObjectEditor"
 
@@ -3295,6 +3298,29 @@ UE::Mutable::Private::Ptr<UE::Mutable::Private::NodeMesh> GenerateMutableSourceM
 		Result = RemoveNode;
 	}
 	
+	else if (const UCONodeTransformWithBone* TransformWithBoneNode = Cast<UCONodeTransformWithBone>(Node))
+	{
+		UE::Mutable::Private::Ptr<UE::Mutable::Private::NodeMeshTransformWithBone> TransformNode = new UE::Mutable::Private::NodeMeshTransformWithBone();
+		Result = TransformNode;
+
+		TransformNode->BoneName = *TransformWithBoneNode->BoneName;
+		TransformNode->ThresholdFactor = TransformWithBoneNode->ThresholdFactor;
+
+		const UEdGraphPin* SourceMeshPin = TransformWithBoneNode->BaseMeshPin.Get();
+		check(SourceMeshPin);
+		if (const UEdGraphPin* ConnectedPin = FollowInputPin(*SourceMeshPin))
+		{
+			TransformNode->Source = GenerateMutableSourceMesh(ConnectedPin, GenerationContext, BaseMeshData, false, bOnlyConnectedLOD);
+		}
+		
+		const UEdGraphPin* TransformPin = TransformWithBoneNode->TransformPin.Get();
+		check(TransformPin);
+		if (const UEdGraphPin* ConnectedPin = FollowInputPin(*TransformPin))
+		{
+			TransformNode->MatrixNode = GenerateMutableSourceTransform(ConnectedPin, GenerationContext);
+		}
+	}
+	
 	else if (const UCustomizableObjectNodeAnimationPose* TypedNode = Cast<UCustomizableObjectNodeAnimationPose>(Node))
 	{
 		if (const UEdGraphPin* InputMeshPin = FollowInputPin(*TypedNode->GetInputMeshPin()))
diff --git a/Engine/Plugins/Mutable/Source/CustomizableObjectEditor/Private/MuCOE/GenerateMutableSource/GenerateMutableSourceModifier.cpp b/Engine/Plugins/Mutable/Source/CustomizableObjectEditor/Private/MuCOE/GenerateMutableSource/GenerateMutableSourceModifier.cpp
--- a/Engine/Plugins/Mutable/Source/CustomizableObjectEditor/Private/MuCOE/GenerateMutableSource/GenerateMutableSourceModifier.cpp
+++ b/Engine/Plugins/Mutable/Source/CustomizableObjectEditor/Private/MuCOE/GenerateMutableSource/GenerateMutableSourceModifier.cpp
@@ -623,13 +623,17 @@ UE::Mutable::Private::Ptr<UE::Mutable::Private::NodeModifier> GenerateMutableSou
 		UE::Mutable::Private::Ptr<UE::Mutable::Private::NodeModifierMeshTransformInMesh> TransformNode = new UE::Mutable::Private::NodeModifierMeshTransformInMesh();
 		Result = TransformNode;
 
-		if (const UEdGraphPin* ConnectedPin = FollowInputPin(*TypedNodeTransformMesh->GetTransformPin()))
+		const UEdGraphPin* TransformPin = TypedNodeTransformMesh->TransformMeshPin.Get();
+		check(TransformPin);
+		if (const UEdGraphPin* ConnectedPin = FollowInputPin(*TransformPin))
 		{
 			TransformNode->MatrixNode = GenerateMutableSourceTransform(ConnectedPin, GenerationContext);
 		}
 
 		// If no bounding mesh is provided, we transform the entire mesh.
-		if (const UEdGraphPin* ConnectedPin = FollowInputPin(*TypedNodeTransformMesh->GetBoundingMeshPin()))
+		const UEdGraphPin* BoundingMeshPin = TypedNodeTransformMesh->BoundingMeshPin.Get();
+		check(BoundingMeshPin);
+		if (const UEdGraphPin* ConnectedPin = FollowInputPin(*BoundingMeshPin))
 		{
 			UE::Mutable::Private::Ptr<UE::Mutable::Private::NodeMesh> BoundingMesh = GenerateMutableSourceMesh(ConnectedPin, GenerationContext, FMutableSourceMeshData(), false, true);
 
@@ -661,7 +665,9 @@ UE::Mutable::Private::Ptr<UE::Mutable::Private::NodeModifier> GenerateMutableSou
 		TransformNode->BoneName = *TypedNodeTransformWithBone->BoneName;
 		TransformNode->ThresholdFactor = TypedNodeTransformWithBone->ThresholdFactor;
 
-		if (const UEdGraphPin* ConnectedPin = FollowInputPin(*TypedNodeTransformWithBone->GetTransformPin()))
+		const UEdGraphPin* TransformPin = TypedNodeTransformWithBone->TransformPin.Get();
+		check(TransformPin);
+		if (const UEdGraphPin* ConnectedPin = FollowInputPin(*TransformPin))
 		{
 			TransformNode->MatrixNode = GenerateMutableSourceTransform(ConnectedPin, GenerationContext);
 		}
diff --git a/Engine/Plugins/Mutable/Source/CustomizableObjectEditor/Private/MuCOE/GraphTraversal.cpp b/Engine/Plugins/Mutable/Source/CustomizableObjectEditor/Private/MuCOE/GraphTraversal.cpp
--- a/Engine/Plugins/Mutable/Source/CustomizableObjectEditor/Private/MuCOE/GraphTraversal.cpp
+++ b/Engine/Plugins/Mutable/Source/CustomizableObjectEditor/Private/MuCOE/GraphTraversal.cpp
@@ -31,6 +31,7 @@
 #include "Nodes/CONodeExternalOperation.h"
 #include "Nodes/CONodeClipMeshWithMesh.h"
 #include "Nodes/CONodeRemoveMesh.h"
+#include "Nodes/CONodeTransformWithBone.h"
 #include "UObject/UObjectIterator.h"
 
 
@@ -781,6 +782,14 @@ const UEdGraphPin* FindMeshBaseSource(const UEdGraphPin& Pin, const bool bOnlyLo
 		}
 	}
 
+	else if (const UCONodeTransformWithBone* NodeTransformWithBone = Cast<UCONodeTransformWithBone>(Node))
+	{
+		if (const UEdGraphPin* ConnectedPin = FollowInputPin(*NodeTransformWithBone->BaseMeshPin.Get()))
+		{
+			return FindMeshBaseSource(*ConnectedPin, bOnlyLookForStaticMesh, MacroContext);
+		}
+	}
+
 	else if (Cast<UCONodeExternalOperation>(Node))
 	{
 		if (!bOnlyLookForStaticMesh)
diff --git a/Engine/Plugins/Mutable/Source/CustomizableObjectEditor/Private/MuCOE/Nodes/CONodeModifierTransformWithBone.cpp b/Engine/Plugins/Mutable/Source/CustomizableObjectEditor/Private/MuCOE/Nodes/CONodeModifierTransformWithBone.cpp
--- a/Engine/Plugins/Mutable/Source/CustomizableObjectEditor/Private/MuCOE/Nodes/CONodeModifierTransformWithBone.cpp
+++ b/Engine/Plugins/Mutable/Source/CustomizableObjectEditor/Private/MuCOE/Nodes/CONodeModifierTransformWithBone.cpp
@@ -8,16 +8,19 @@
 
 #define LOCTEXT_NAMESPACE "CustomizableObjectEditor"
 
+
 FText UCONodeModifierTransformWithBone::GetNodeTitle(ENodeTitleType::Type TitleType) const
 {
 	return LOCTEXT("Transform_Mesh_In_Bone_Hierarchy", "Transform Mesh In Bone Hierarchy");
 }
 
+
 FText UCONodeModifierTransformWithBone::GetTooltipText() const
 {
-	return LOCTEXT("Transform_Mesh_In_Bone_Hierarchy_Tooltip", "Applies a transform to the vertices of a mesh that is skinned to the target bone or its child bones");
+	return LOCTEXT("Transform_Mesh_In_Bone_Hierarchy_Tooltip", "Applies a transform to the vertices of a mesh that are skinned to the target bone or its child bones");
 }
 
+
 void UCONodeModifierTransformWithBone::AllocateDefaultPins(UCustomizableObjectNodeRemapPins* RemapPins)
 {
 	TransformPin = CustomCreatePin(EGPD_Input, UEdGraphSchema_CustomizableObject::PC_Transform);
@@ -29,35 +32,10 @@ void UCONodeModifierTransformWithBone::AllocateDefaultPins(UCustomizableObjectNo
 }
 
 
-void UCONodeModifierTransformWithBone::PostEditChangeProperty(struct FPropertyChangedEvent& PropertyChangedEvent)
-{
-	Super::PostEditChangeProperty(PropertyChangedEvent);
-
-	if (PropertyChangedEvent.Property == nullptr)
-	{
-		return;
-	}
-	
-	if (PropertyChangedEvent.Property->GetFName() == GET_MEMBER_NAME_CHECKED(UCONodeModifierTransformWithBone, BoundingMeshTransform)
-		|| PropertyChangedEvent.Property->GetFName() == GET_MEMBER_NAME_CHECKED(UCONodeModifierTransformWithBone, BoneName))
-	{
-		if (TransformChangedDelegate.IsBound())
-		{
-			TransformChangedDelegate.Broadcast(BoundingMeshTransform);
-		}
-	}
-}
-
-
 UEdGraphPin* UCONodeModifierTransformWithBone::GetOutputPin() const
 {
 	return OutputPin.Get();
 }
 
 
-UEdGraphPin* UCONodeModifierTransformWithBone::GetTransformPin() const
-{
-	return TransformPin.Get();
-}
-
 #undef LOCTEXT_NAMESPACE
diff --git a/Engine/Plugins/Mutable/Source/CustomizableObjectEditor/Private/MuCOE/Nodes/CONodeModifierTransformWithBone.h b/Engine/Plugins/Mutable/Source/CustomizableObjectEditor/Private/MuCOE/Nodes/CONodeModifierTransformWithBone.h
--- a/Engine/Plugins/Mutable/Source/CustomizableObjectEditor/Private/MuCOE/Nodes/CONodeModifierTransformWithBone.h
+++ b/Engine/Plugins/Mutable/Source/CustomizableObjectEditor/Private/MuCOE/Nodes/CONodeModifierTransformWithBone.h
@@ -7,64 +7,43 @@
 
 #include "CONodeModifierTransformWithBone.generated.h"
 
-#define UE_API CUSTOMIZABLEOBJECTEDITOR_API
-
 namespace ENodeTitleType { enum Type : int; }
 
 class UCustomizableObject;
 class UCustomizableObjectNodeRemapPins;
 class UEdGraphPin;
+struct FPropertyChangedEvent;
 
-DECLARE_MULTICAST_DELEGATE_OneParam(OnTransformChanged, const FTransform&);
-
-UCLASS(MinimalAPI)
+UCLASS()
 class UCONodeModifierTransformWithBone : public UCustomizableObjectNodeModifierBase
 {
 	GENERATED_BODY()
-
+	
 public:
-	/** Transform to apply to the vertices. */
-	UPROPERTY(EditAnywhere, Category = Transform)
-	FTransform BoundingMeshTransform = FTransform::Identity;
+	
+	// EdGraphNode interface
+	virtual FText GetNodeTitle(ENodeTitleType::Type TitleType) const override;
+	virtual FText GetTooltipText() const override;
+
+	// UCustomizableObjectNode interface
+	virtual void AllocateDefaultPins(UCustomizableObjectNodeRemapPins* RemapPins) override;
+
+	// UCustomizableObjectNodeModifierBase interface
+	virtual UEdGraphPin* GetOutputPin() const override;
 
 	/** Root bone used to filter vertices based on skinning weights. Vertices skinned to this
-	  * bone or below will me transformed. */
+	  * bone or below will be transformed. */
 	UPROPERTY(EditAnywhere, Category = Transform)
 	FString BoneName;
 
 	/** Vertex influences under this threshold will not be considered when filtering vertices.
-	  * Percentage of total skin weight [0.0 .. 1.0]. */
+	  * Percentage of total skin weight [0.0 , 1.0]. */
 	UPROPERTY(EditAnywhere, Category = Transform, meta = (ClampMin = 0.0f, ClampMax = 1.0f))
 	float ThresholdFactor = 0.05f;
-
-	/**
-	 * Delegate invoked each time the BoundingMeshTransform value gets modified
-	 */
-	OnTransformChanged TransformChangedDelegate;
 	
-public:
-	// EdGraphNode interface
-	UE_API virtual FText GetNodeTitle(ENodeTitleType::Type TitleType) const override;
-	UE_API virtual FText GetTooltipText() const override;
-
-	// UCustomizableObjectNode interface
-	UE_API virtual void AllocateDefaultPins(UCustomizableObjectNodeRemapPins* RemapPins) override;
-	virtual bool IsExperimental() const override { return true; }
-	UE_API virtual void PostEditChangeProperty( struct FPropertyChangedEvent& PropertyChangedEvent) override;
-
-	// UCustomizableObjectNodeModifierBase interface
-	UE_API virtual UEdGraphPin* GetOutputPin() const override;
-
-	// Own interface
-	UE_API UEdGraphPin* GetTransformPin() const;
-	
-private:
-
 	UPROPERTY()
 	FEdGraphPinReference TransformPin;
 
 	UPROPERTY()
 	FEdGraphPinReference OutputPin;
 };
-
-#undef UE_API
diff --git a/Engine/Plugins/Mutable/Source/CustomizableObjectEditor/Private/MuCOE/Nodes/CONodeModifierTransformWithBoneDetails.cpp b/Engine/Plugins/Mutable/Source/CustomizableObjectEditor/Private/MuCOE/Nodes/CONodeModifierTransformWithBoneDetails.cpp
deleted file mode 100644
diff --git a/Engine/Plugins/Mutable/Source/CustomizableObjectEditor/Private/MuCOE/Nodes/CONodeModifierTransformWithBoneDetails.h b/Engine/Plugins/Mutable/Source/CustomizableObjectEditor/Private/MuCOE/Nodes/CONodeModifierTransformWithBoneDetails.h
deleted file mode 100644
diff --git a/Engine/Plugins/Mutable/Source/CustomizableObjectEditor/Private/MuCOE/Nodes/CONodeTransformWithBone.cpp b/Engine/Plugins/Mutable/Source/CustomizableObjectEditor/Private/MuCOE/Nodes/CONodeTransformWithBone.cpp
new file mode 100644
--- /dev/null
+++ b/Engine/Plugins/Mutable/Source/CustomizableObjectEditor/Private/MuCOE/Nodes/CONodeTransformWithBone.cpp
@@ -0,0 +1,45 @@
+// Copyright Epic Games, Inc. All Rights Reserved.
+
+#include "MuCOE/Nodes/CONodeTransformWithBone.h"
+
+#include "MuCOE/EdGraphSchema_CustomizableObject.h"
+
+#include UE_INLINE_GENERATED_CPP_BY_NAME(CONodeTransformWithBone)
+
+#define LOCTEXT_NAMESPACE "CustomizableObjectEditor"
+
+
+FText UCONodeTransformWithBone::GetNodeTitle(ENodeTitleType::Type TitleType) const
+{
+	return LOCTEXT("Transform_Mesh_In_Bone_Hierarchy", "Transform Mesh In Bone Hierarchy");
+}
+
+
+FText UCONodeTransformWithBone::GetTooltipText() const
+{
+	return LOCTEXT("Transform_Mesh_In_Bone_Hierarchy_Tooltip", "Applies a transform to the vertices of a mesh that are skinned to the target bone or its child bones");
+}
+
+
+void UCONodeTransformWithBone::AllocateDefaultPins(UCustomizableObjectNodeRemapPins* RemapPins)
+{
+	BaseMeshPin = CustomCreatePin(EGPD_Input, UEdGraphSchema_CustomizableObject::PC_Mesh);
+	TransformPin = CustomCreatePin(EGPD_Input, UEdGraphSchema_CustomizableObject::PC_Transform);
+
+	OutputPin = CustomCreatePin(EGPD_Output, UEdGraphSchema_CustomizableObject::PC_Mesh);
+}
+
+
+FLinearColor UCONodeTransformWithBone::GetNodeTitleColor() const
+{
+	return UEdGraphSchema_CustomizableObject::GetPinTypeColor(UEdGraphSchema_CustomizableObject::PC_Mesh);
+}
+
+
+bool UCONodeTransformWithBone::IsExperimental() const
+{
+	return true;
+}
+
+
+#undef LOCTEXT_NAMESPACE
diff --git a/Engine/Plugins/Mutable/Source/CustomizableObjectEditor/Private/MuCOE/Nodes/CONodeTransformWithBone.h b/Engine/Plugins/Mutable/Source/CustomizableObjectEditor/Private/MuCOE/Nodes/CONodeTransformWithBone.h
new file mode 100644
--- /dev/null
+++ b/Engine/Plugins/Mutable/Source/CustomizableObjectEditor/Private/MuCOE/Nodes/CONodeTransformWithBone.h
@@ -0,0 +1,58 @@
+// Copyright Epic Games, Inc. All Rights Reserved.
+
+#pragma once
+
+#include "MuCOE/Nodes/CustomizableObjectNode.h"
+#include "MuT/NodeModifier.h"
+
+#include "CONodeTransformWithBone.generated.h"
+
+namespace ENodeTitleType { enum Type : int; }
+
+class UCustomizableObject;
+class UCustomizableObjectNodeRemapPins;
+class UEdGraphPin;
+struct FPropertyChangedEvent;
+
+DECLARE_MULTICAST_DELEGATE_OneParam(OnTransformChanged, const FTransform&);
+
+UCLASS()
+class UCONodeTransformWithBone : public UCustomizableObjectNode
+{
+	GENERATED_BODY()
+	
+public:
+	
+	// EdGraphNode interface
+	virtual FText GetNodeTitle(ENodeTitleType::Type TitleType) const override;
+	virtual FText GetTooltipText() const override;
+
+	// UCustomizableObjectNode interface
+	virtual void AllocateDefaultPins(UCustomizableObjectNodeRemapPins* RemapPins) override;
+	virtual FLinearColor GetNodeTitleColor() const override;
+	virtual bool IsExperimental() const override;
+
+	/** Root bone used to filter vertices based on skinning weights. Vertices skinned to this
+	  * bone or below will be transformed. */
+	UPROPERTY(EditAnywhere, Category = Transform)
+	FString BoneName;
+
+	/** Vertex influences under this threshold will not be considered when filtering vertices.
+	  * Percentage of total skin weight [0.0 .. 1.0]. */
+	UPROPERTY(EditAnywhere, Category = Transform, meta = (ClampMin = 0.0f, ClampMax = 1.0f))
+	float ThresholdFactor = 0.05f;
+	
+	UPROPERTY()
+	FEdGraphPinReference BaseMeshPin;
+	
+	UPROPERTY()
+	FEdGraphPinReference TransformPin;
+
+	UPROPERTY()
+	FEdGraphPinReference OutputPin;
+	
+	/**
+	* Delegate invoked each time the BoundingMeshTransform value gets modified
+	*/
+	OnTransformChanged TransformChangedDelegate;
+};
\ No newline at end of file
diff --git a/Engine/Plugins/Mutable/Source/CustomizableObjectEditor/Private/MuCOE/Nodes/CustomizableObjectNodeModifierTransformInMesh.cpp b/Engine/Plugins/Mutable/Source/CustomizableObjectEditor/Private/MuCOE/Nodes/CustomizableObjectNodeModifierTransformInMesh.cpp
--- a/Engine/Plugins/Mutable/Source/CustomizableObjectEditor/Private/MuCOE/Nodes/CustomizableObjectNodeModifierTransformInMesh.cpp
+++ b/Engine/Plugins/Mutable/Source/CustomizableObjectEditor/Private/MuCOE/Nodes/CustomizableObjectNodeModifierTransformInMesh.cpp
@@ -2,6 +2,7 @@
 
 #include "MuCOE/Nodes/CustomizableObjectNodeModifierTransformInMesh.h"
 
+#include "MuCO/CustomizableObjectCustomVersion.h"
 #include "MuCOE/EdGraphSchema_CustomizableObject.h"
 
 #include UE_INLINE_GENERATED_CPP_BY_NAME(CustomizableObjectNodeModifierTransformInMesh)
@@ -9,28 +10,24 @@
 #define LOCTEXT_NAMESPACE "CustomizableObjectEditor"
 
 
-const TCHAR* UCustomizableObjectNodeModifierTransformInMesh::OutputPinName = TEXT("Modifier");
-const TCHAR* UCustomizableObjectNodeModifierTransformInMesh::BoundingMeshPinName = TEXT("Bounding Mesh");
-const TCHAR* UCustomizableObjectNodeModifierTransformInMesh::TransformPinName = TEXT("Transform");
-
-
 FText UCustomizableObjectNodeModifierTransformInMesh::GetNodeTitle(ENodeTitleType::Type TitleType) const
 {
 	return LOCTEXT("Transform_Mesh_In_Mesh", "Transform Mesh In Mesh");
 }
 
+
 FText UCustomizableObjectNodeModifierTransformInMesh::GetTooltipText() const
 {
 	return LOCTEXT("Transform_Mesh_In_Mesh_Tooltip", "Applies a transform to the vertices of a mesh that is contained within the given bounding mesh");
 }
 
+
 void UCustomizableObjectNodeModifierTransformInMesh::AllocateDefaultPins(UCustomizableObjectNodeRemapPins* RemapPins)
 {
-	const UEdGraphSchema_CustomizableObject* Schema = GetDefault<UEdGraphSchema_CustomizableObject>();
-
-	CustomCreatePin(EGPD_Input, Schema->PC_Mesh, FName(BoundingMeshPinName), FText::FromString(BoundingMeshPinName));
-	CustomCreatePin(EGPD_Input, Schema->PC_Transform, FName(TransformPinName), FText::FromString(TransformPinName));
-	CustomCreatePin(EGPD_Output, Schema->PC_Modifier, FName(OutputPinName), FText::FromString(TransformPinName));
+	BoundingMeshPin = CustomCreatePin(EGPD_Input, UEdGraphSchema_CustomizableObject::PC_Mesh, TEXT("Bounding Mesh"), LOCTEXT("BoundingMesh", "Bounding Mesh"));
+	TransformMeshPin = CustomCreatePin(EGPD_Input, UEdGraphSchema_CustomizableObject::PC_Transform, TEXT("Transform"), LOCTEXT("Transform", "Transform"));
+	
+	OutputPin =  CustomCreatePin(EGPD_Output, UEdGraphSchema_CustomizableObject::PC_Modifier, TEXT("Modifier"), LOCTEXT("Modifier", "Modifier"));
 
 	//Create Node Modifier Common Pins
 	Super::AllocateDefaultPins(RemapPins);
@@ -58,17 +55,49 @@ void UCustomizableObjectNodeModifierTransformInMesh::PostEditChangeProperty(stru
 
 UEdGraphPin* UCustomizableObjectNodeModifierTransformInMesh::GetOutputPin() const
 {
-	return FindPin(OutputPinName);
+	return OutputPin.Get();
 }
 
-UEdGraphPin* UCustomizableObjectNodeModifierTransformInMesh::GetBoundingMeshPin() const
-{
-	return FindPin(BoundingMeshPinName);
-}
 
-UEdGraphPin* UCustomizableObjectNodeModifierTransformInMesh::GetTransformPin() const
+void UCustomizableObjectNodeModifierTransformInMesh::BackwardsCompatibleFixup(int32 CustomizableObjectCustomVersion)
 {
-	return FindPin(TransformPinName);
+	Super::BackwardsCompatibleFixup(CustomizableObjectCustomVersion);
+	
+	if (CustomizableObjectCustomVersion == FCustomizableObjectCustomVersion::NodeModifierTransformInMeshStorePinsAsReferences)
+	{
+		const FString OutputPinName = TEXT("Modifier");
+		const FString BoundingMeshPinName = TEXT("Bounding Mesh");
+		const FString TransformPinName = TEXT("Transform");
+
+		for (UEdGraphPin* Pin : Pins)
+		{
+			const FString PinName = Pin->GetName();
+			const EEdGraphPinDirection PinDirection = Pin->Direction;
+			
+			if (PinDirection == EEdGraphPinDirection::EGPD_Input)
+			{
+				if (PinName == BoundingMeshPinName)
+				{
+					BoundingMeshPin = Pin;
+				}
+				else if (PinName == TransformPinName)
+				{
+					TransformMeshPin = Pin;
+				}
+			}
+			else 
+			{
+				if (PinName == OutputPinName)
+				{
+					OutputPin = Pin;
+				}
+			}
+		}
+		
+		check(BoundingMeshPin.Get());
+		check(TransformMeshPin.Get());
+		check(OutputPin.Get());
+	}
 }
 
 #undef LOCTEXT_NAMESPACE
diff --git a/Engine/Plugins/Mutable/Source/CustomizableObjectEditor/Private/MuCOE/Nodes/CustomizableObjectNodeModifierTransformInMesh.h b/Engine/Plugins/Mutable/Source/CustomizableObjectEditor/Private/MuCOE/Nodes/CustomizableObjectNodeModifierTransformInMesh.h
--- a/Engine/Plugins/Mutable/Source/CustomizableObjectEditor/Private/MuCOE/Nodes/CustomizableObjectNodeModifierTransformInMesh.h
+++ b/Engine/Plugins/Mutable/Source/CustomizableObjectEditor/Private/MuCOE/Nodes/CustomizableObjectNodeModifierTransformInMesh.h
@@ -7,8 +7,6 @@
 
 #include "CustomizableObjectNodeModifierTransformInMesh.generated.h"
 
-#define UE_API CUSTOMIZABLEOBJECTEDITOR_API
-
 namespace ENodeTitleType { enum Type : int; }
 
 class UCustomizableObject;
@@ -17,43 +15,41 @@ class UEdGraphPin;
 
 DECLARE_MULTICAST_DELEGATE_OneParam(OnTransformChanged, const FTransform&);
 
-UCLASS(MinimalAPI)
+UCLASS()
 class UCustomizableObjectNodeModifierTransformInMesh : public UCustomizableObjectNodeModifierBase
 {
 	GENERATED_BODY()
 
 public:
-	/** Transform to apply to the bounding mesh before selecting for vertices to transform. */
-	UPROPERTY(EditAnywhere, Category = BoundingMesh)
-	FTransform BoundingMeshTransform = FTransform::Identity;
-
-	/**
-	 * Delegate invoked each time the BoundingMeshTransform value gets modified
-	 */
-	OnTransformChanged TransformChangedDelegate;
 	
-public:
 	// EdGraphNode interface
-	UE_API virtual FText GetNodeTitle(ENodeTitleType::Type TitleType) const override;
-	UE_API virtual FText GetTooltipText() const override;
+	virtual FText GetNodeTitle(ENodeTitleType::Type TitleType) const override;
+	virtual FText GetTooltipText() const override;
 
 	// UCustomizableObjectNode interface
-	UE_API virtual void AllocateDefaultPins(UCustomizableObjectNodeRemapPins* RemapPins) override;
-	virtual bool IsExperimental() const override { return true; }
-	UE_API virtual void PostEditChangeProperty( struct FPropertyChangedEvent& PropertyChangedEvent) override;
+	virtual void AllocateDefaultPins(UCustomizableObjectNodeRemapPins* RemapPins) override;
+	virtual void PostEditChangeProperty( struct FPropertyChangedEvent& PropertyChangedEvent) override;
 
 	// UCustomizableObjectNodeModifierBase interface
-	UE_API virtual UEdGraphPin* GetOutputPin() const override;
-
-	// Own interface
-	UE_API UEdGraphPin* GetBoundingMeshPin() const;
-	UE_API UEdGraphPin* GetTransformPin() const;
+	virtual UEdGraphPin* GetOutputPin() const override;
+	virtual bool IsExperimental() const override { return true; }
+	virtual void BackwardsCompatibleFixup(int32 CustomizableObjectCustomVersion) override;
 
-private:
-	static UE_API const TCHAR* OutputPinName;
-	static UE_API const TCHAR* BoundingMeshPinName;
-	static UE_API const TCHAR* TransformPinName;
+	/** Transform to apply to the bounding mesh before selecting for vertices to transform. */
+	UPROPERTY(EditAnywhere, Category = BoundingMesh)
+	FTransform BoundingMeshTransform = FTransform::Identity;
+	
+	UPROPERTY()
+	FEdGraphPinReference BoundingMeshPin;
+	
+	UPROPERTY()
+	FEdGraphPinReference TransformMeshPin;
 	
+	UPROPERTY()
+	FEdGraphPinReference OutputPin;
+	
+	/**
+	 * Delegate invoked each time the BoundingMeshTransform value gets modified
+	 */
+	OnTransformChanged TransformChangedDelegate;
 };
-
-#undef UE_API
diff --git a/Engine/Plugins/Mutable/Source/CustomizableObjectEditor/Private/MuCOE/Nodes/CustomizableObjectNodeModifierTransformInMeshDetails.cpp b/Engine/Plugins/Mutable/Source/CustomizableObjectEditor/Private/MuCOE/Nodes/CustomizableObjectNodeModifierTransformInMeshDetails.cpp
deleted file mode 100644
diff --git a/Engine/Plugins/Mutable/Source/CustomizableObjectEditor/Private/MuCOE/Nodes/CustomizableObjectNodeModifierTransformInMeshDetails.h b/Engine/Plugins/Mutable/Source/CustomizableObjectEditor/Private/MuCOE/Nodes/CustomizableObjectNodeModifierTransformInMeshDetails.h
deleted file mode 100644
diff --git a/Engine/Plugins/Mutable/Source/CustomizableObjectEditor/Private/MuCOE/SMutableGraphViewer.cpp b/Engine/Plugins/Mutable/Source/CustomizableObjectEditor/Private/MuCOE/SMutableGraphViewer.cpp
--- a/Engine/Plugins/Mutable/Source/CustomizableObjectEditor/Private/MuCOE/SMutableGraphViewer.cpp
+++ b/Engine/Plugins/Mutable/Source/CustomizableObjectEditor/Private/MuCOE/SMutableGraphViewer.cpp
@@ -44,6 +44,8 @@
 #include "MuT/NodeSurfaceVariation.h"
 #include "MuT/NodeLOD.h"
 #include "MuT/NodeMaterialModify.h"
+#include "MuT/NodeMatrixConstant.h"
+#include "MuT/NodeMatrixParameter.h"
 #include "MuT/NodeMeshClipWithMesh.h"
 #include "MuT/NodeMeshConstant.h"
 #include "MuT/NodeMeshFormat.h"
@@ -52,6 +54,8 @@
 #include "MuT/NodeMeshMorph.h"
 #include "MuT/NodeMeshRemoveMesh.h"
 #include "MuT/NodeMeshTable.h"
+#include "MuT/NodeMeshTransform.h"
+#include "MuT/NodeMeshTransformWithBone.h"
 #include "MuT/NodeModifierMeshClipDeform.h"
 #include "MuT/NodeModifierMeshClipMorphPlane.h"
 #include "MuT/NodeModifierMeshClipWithUVMask.h"
@@ -103,6 +107,7 @@ static const char* MutableNodeNames[] =
 	"MeshClipMorphPlane",
 	"MeshClipWithMesh",
 	"MeshRemoveMesh",
+	"MeshTransformWithBone",
 	"MeshApplyPose",
 	"MeshVariation",
 	"MeshReshape",
@@ -438,7 +443,9 @@ void SMutableGraphViewer::GetChildrenForInfo(TSharedPtr<FMutableGraphTreeElement
 		++InputIndex;
 	};
 
-	if (ParentNode->GetType() == UE::Mutable::Private::NodeObjectNew::GetStaticType())
+	const UE::Mutable::Private::FNodeType* ParentNodeType = ParentNode->GetType();
+	
+	if (ParentNodeType == UE::Mutable::Private::NodeObjectNew::GetStaticType())
 	{
 		UE::Mutable::Private::NodeObjectNew* ObjectNew = StaticCast<UE::Mutable::Private::NodeObjectNew*>(ParentNode);
 		for (int32 l = 0; l < ObjectNew->Components.Num(); ++l)
@@ -457,7 +464,7 @@ void SMutableGraphViewer::GetChildrenForInfo(TSharedPtr<FMutableGraphTreeElement
 		}
 	}
 
-	else if (ParentNode->GetType() == UE::Mutable::Private::NodeObjectGroup::GetStaticType())
+	else if (ParentNodeType == UE::Mutable::Private::NodeObjectGroup::GetStaticType())
 	{
 		UE::Mutable::Private::NodeObjectGroup* ObjectGroup = StaticCast<UE::Mutable::Private::NodeObjectGroup*>(ParentNode);
 		for (int32 l = 0; l < ObjectGroup->Children.Num(); ++l)
@@ -466,7 +473,7 @@ void SMutableGraphViewer::GetChildrenForInfo(TSharedPtr<FMutableGraphTreeElement
 		}
 	}
 
-	else if (ParentNode->GetType() == UE::Mutable::Private::NodeSurfaceNew::GetStaticType())
+	else if (ParentNodeType == UE::Mutable::Private::NodeSurfaceNew::GetStaticType())
 	{
 		UE::Mutable::Private::NodeSurfaceNew* SurfaceNew = StaticCast<UE::Mutable::Private::NodeSurfaceNew*>(ParentNode);
 		AddChildFunc(SurfaceNew->Mesh.get(), TEXT("MESH"));
@@ -479,7 +486,7 @@ void SMutableGraphViewer::GetChildrenForInfo(TSharedPtr<FMutableGraphTreeElement
 		}
 	}
 
-	else if (ParentNode->GetType() == UE::Mutable::Private::NodeModifierSurfaceEdit::GetStaticType())
+	else if (ParentNodeType == UE::Mutable::Private::NodeModifierSurfaceEdit::GetStaticType())
 	{
 		UE::Mutable::Private::NodeModifierSurfaceEdit* SurfaceEdit = StaticCast<UE::Mutable::Private::NodeModifierSurfaceEdit*>(ParentNode);
 
@@ -500,7 +507,7 @@ void SMutableGraphViewer::GetChildrenForInfo(TSharedPtr<FMutableGraphTreeElement
 		}
 	}
 
-	else if (ParentNode->GetType() == UE::Mutable::Private::NodeSurfaceSwitch::GetStaticType())
+	else if (ParentNodeType == UE::Mutable::Private::NodeSurfaceSwitch::GetStaticType())
 	{
 		UE::Mutable::Private::NodeSurfaceSwitch* SurfaceSwitch = StaticCast<UE::Mutable::Private::NodeSurfaceSwitch*>(ParentNode);
 		AddChildFunc(SurfaceSwitch->Parameter.get(), TEXT("PARAM"));
@@ -510,7 +517,7 @@ void SMutableGraphViewer::GetChildrenForInfo(TSharedPtr<FMutableGraphTreeElement
 		}
 	}
 
-	else if (ParentNode->GetType() == UE::Mutable::Private::NodeSurfaceVariation::GetStaticType())
+	else if (ParentNodeType == UE::Mutable::Private::NodeSurfaceVariation::GetStaticType())
 	{
 		UE::Mutable::Private::NodeSurfaceVariation* SurfaceVar = StaticCast<UE::Mutable::Private::NodeSurfaceVariation*>(ParentNode);
 		for (int32 l = 0; l < SurfaceVar->DefaultSurfaces.Num(); ++l)
@@ -536,7 +543,7 @@ void SMutableGraphViewer::GetChildrenForInfo(TSharedPtr<FMutableGraphTreeElement
 		}
 	}
 	
-	else if (ParentNode->GetType() == UE::Mutable::Private::NodeLOD::GetStaticType())
+	else if (ParentNodeType == UE::Mutable::Private::NodeLOD::GetStaticType())
 	{
 		UE::Mutable::Private::NodeLOD* LodVar = StaticCast<UE::Mutable::Private::NodeLOD*>(ParentNode);
 
@@ -546,20 +553,20 @@ void SMutableGraphViewer::GetChildrenForInfo(TSharedPtr<FMutableGraphTreeElement
 		}
 	}
 	
-	else if (ParentNode->GetType() == UE::Mutable::Private::NodeComponentNew::GetStaticType())
+	else if (ParentNodeType == UE::Mutable::Private::NodeComponentNew::GetStaticType())
 	{
 		UE::Mutable::Private::NodeComponentNew* ComponentVar = StaticCast<UE::Mutable::Private::NodeComponentNew*>(ParentNode);
 		AddChildFunc(ComponentVar->SkeletalMesh.get(), TEXT("SKELETAL MESH"));
 		AddChildFunc(ComponentVar->OverlayMaterial.get(), TEXT("OVERLAY MATERIAL"));
 	}
 
-	else if (ParentNode->GetType() == UE::Mutable::Private::NodeComponentEdit::GetStaticType())
+	else if (ParentNodeType == UE::Mutable::Private::NodeComponentEdit::GetStaticType())
 	{
 		UE::Mutable::Private::NodeComponentEdit* ComponentVar = StaticCast<UE::Mutable::Private::NodeComponentEdit*>(ParentNode);
 		AddChildFunc(ComponentVar->SkeletalMesh.get(), TEXT("SKELETAL MESH"));
 	}
 
-	else if (ParentNode->GetType() == UE::Mutable::Private::NodeComponentSwitch::GetStaticType())
+	else if (ParentNodeType == UE::Mutable::Private::NodeComponentSwitch::GetStaticType())
 	{
 		UE::Mutable::Private::NodeComponentSwitch* ComponentSwitch = StaticCast<UE::Mutable::Private::NodeComponentSwitch*>(ParentNode);
 		AddChildFunc(ComponentSwitch->Parameter.get(), TEXT("PARAM"));
@@ -569,7 +576,7 @@ void SMutableGraphViewer::GetChildrenForInfo(TSharedPtr<FMutableGraphTreeElement
 		}
 	}
 
-	else if (ParentNode->GetType() == UE::Mutable::Private::NodeSkeletalMeshParameter::GetStaticType())
+	else if (ParentNodeType == UE::Mutable::Private::NodeSkeletalMeshParameter::GetStaticType())
 	{
 		UE::Mutable::Private::NodeSkeletalMeshParameter* SkeletalMeshVar = StaticCast<UE::Mutable::Private::NodeSkeletalMeshParameter*>(ParentNode);
 		for (int32 RangeIndex = 0; RangeIndex < SkeletalMeshVar->Ranges.Num(); RangeIndex++)
@@ -578,7 +585,7 @@ void SMutableGraphViewer::GetChildrenForInfo(TSharedPtr<FMutableGraphTreeElement
 		}
 	}
 	
-	else if (ParentNode->GetType() == UE::Mutable::Private::NodeSkeletalMeshNew::GetStaticType())
+	else if (ParentNodeType == UE::Mutable::Private::NodeSkeletalMeshNew::GetStaticType())
 	{
 		UE::Mutable::Private::NodeSkeletalMeshNew* SkeletalMeshVar = StaticCast<UE::Mutable::Private::NodeSkeletalMeshNew*>(ParentNode);
 		for (int32 LODIndex = 0; LODIndex < SkeletalMeshVar->LODs.Num(); LODIndex++)
@@ -587,7 +594,7 @@ void SMutableGraphViewer::GetChildrenForInfo(TSharedPtr<FMutableGraphTreeElement
 		}
 	}
 	
-	else if (ParentNode->GetType() == UE::Mutable::Private::NodeMeshConstant::GetStaticType())
+	else if (ParentNodeType == UE::Mutable::Private::NodeMeshConstant::GetStaticType())
 	{
 		UE::Mutable::Private::NodeMeshConstant* MeshConstantVar = StaticCast<UE::Mutable::Private::NodeMeshConstant*>(ParentNode);
 		for (int32 LayoutIndex = 0; LayoutIndex < MeshConstantVar->Layouts.Num(); LayoutIndex++)
@@ -596,30 +603,30 @@ void SMutableGraphViewer::GetChildrenForInfo(TSharedPtr<FMutableGraphTreeElement
 		}
 	}
 
-	else if (ParentNode->GetType() == UE::Mutable::Private::NodeImageFormat::GetStaticType())
+	else if (ParentNodeType == UE::Mutable::Private::NodeImageFormat::GetStaticType())
 	{
 		UE::Mutable::Private::NodeImageFormat* ImageFormatVar = StaticCast<UE::Mutable::Private::NodeImageFormat*>(ParentNode);
 		AddChildFunc(ImageFormatVar->Source.get(), FString::Printf(TEXT("SOURCE IMAGE")));
 	}
 
-	else if (ParentNode->GetType() == UE::Mutable::Private::NodeMeshFormat::GetStaticType())
+	else if (ParentNodeType == UE::Mutable::Private::NodeMeshFormat::GetStaticType())
 	{
 		UE::Mutable::Private::NodeMeshFormat* MeshFormatVar = StaticCast<UE::Mutable::Private::NodeMeshFormat*>(ParentNode);
 		AddChildFunc(MeshFormatVar->Source.get(), FString::Printf(TEXT("SOURCE MESH")));
 	}
 
-	else if (ParentNode->GetType() == UE::Mutable::Private::NodeModifierMeshClipMorphPlane::GetStaticType())
+	else if (ParentNodeType == UE::Mutable::Private::NodeModifierMeshClipMorphPlane::GetStaticType())
 	{
 		// Nothing to show
 	}
 
-	else if (ParentNode->GetType() == UE::Mutable::Private::NodeModifierMeshClipWithMesh::GetStaticType())
+	else if (ParentNodeType == UE::Mutable::Private::NodeModifierMeshClipWithMesh::GetStaticType())
 	{
 		UE::Mutable::Private::NodeModifierMeshClipWithMesh* ModifierMeshClipWithMeshVar = StaticCast<UE::Mutable::Private::NodeModifierMeshClipWithMesh*>(ParentNode);
 		AddChildFunc(ModifierMeshClipWithMeshVar->ClipMesh.get(), FString::Printf(TEXT("CLIP MESH")));
 	}
 	
-	else if (ParentNode->GetType() == UE::Mutable::Private::NodeMeshClipWithMesh::GetStaticType())
+	else if (ParentNodeType == UE::Mutable::Private::NodeMeshClipWithMesh::GetStaticType())
 	{
 		UE::Mutable::Private::NodeMeshClipWithMesh* ClipMeshWithMeshVar = StaticCast<UE::Mutable::Private::NodeMeshClipWithMesh*>(ParentNode);
 		AddChildFunc(ClipMeshWithMeshVar->Source.get(), FString::Printf(TEXT("SOURCE MESH")));
@@ -633,33 +640,33 @@ void SMutableGraphViewer::GetChildrenForInfo(TSharedPtr<FMutableGraphTreeElement
 		AddChildFunc(RemoveMeshVar->RemoveMesh.get(), FString::Printf(TEXT("REMOVE MESH")));
 	}
 	
-	else if (ParentNode->GetType() == UE::Mutable::Private::NodeModifierMeshClipDeform::GetStaticType())
+	else if (ParentNodeType == UE::Mutable::Private::NodeModifierMeshClipDeform::GetStaticType())
 	{
 		UE::Mutable::Private::NodeModifierMeshClipDeform* ModifierMeshClipDeformVar = StaticCast<UE::Mutable::Private::NodeModifierMeshClipDeform*>(ParentNode);
 		AddChildFunc(ModifierMeshClipDeformVar->ClipMesh.get(), FString::Printf(TEXT("CLIP MESH")));
 	}
 
-	else if (ParentNode->GetType() == UE::Mutable::Private::NodeModifierMeshClipWithUVMask::GetStaticType())
+	else if (ParentNodeType == UE::Mutable::Private::NodeModifierMeshClipWithUVMask::GetStaticType())
 	{
 		UE::Mutable::Private::NodeModifierMeshClipWithUVMask* ModifierMeshClipWithUVMaskVar = StaticCast<UE::Mutable::Private::NodeModifierMeshClipWithUVMask*>(ParentNode);
 		AddChildFunc(ModifierMeshClipWithUVMaskVar->ClipMask.get(), FString::Printf(TEXT("CLIP MASK")));
 		AddChildFunc(ModifierMeshClipWithUVMaskVar->ClipLayout.get(), FString::Printf(TEXT("CLIP LAYOUT")));
 	}
 
-	else if (ParentNode->GetType() == UE::Mutable::Private::NodeModifierMeshTransformInMesh::GetStaticType())
+	else if (ParentNodeType == UE::Mutable::Private::NodeModifierMeshTransformInMesh::GetStaticType())
 	{
 		UE::Mutable::Private::NodeModifierMeshTransformInMesh* ModifierMeshTransformInMeshVar = StaticCast<UE::Mutable::Private::NodeModifierMeshTransformInMesh*>(ParentNode);
 		AddChildFunc(ModifierMeshTransformInMeshVar->BoundingMesh.get(), FString::Printf(TEXT("BOUNDING MESH")));
 		AddChildFunc(ModifierMeshTransformInMeshVar->MatrixNode.get(), FString::Printf(TEXT("MESH TRANSFORM")));
 	}
 
-	else if (ParentNode->GetType() == UE::Mutable::Private::NodeModifierMeshTransformWithBone::GetStaticType())
+	else if (ParentNodeType == UE::Mutable::Private::NodeModifierMeshTransformWithBone::GetStaticType())
 	{
 		UE::Mutable::Private::NodeModifierMeshTransformWithBone* ModifierMeshTransformWithBoneVar = StaticCast<UE::Mutable::Private::NodeModifierMeshTransformWithBone*>(ParentNode);
 		AddChildFunc(ModifierMeshTransformWithBoneVar->MatrixNode.get(), FString::Printf(TEXT("MESH TRANSFORM")));
 	}
 
-	else if (ParentNode->GetType() == UE::Mutable::Private::NodeImageSwitch::GetStaticType())
+	else if (ParentNodeType == UE::Mutable::Private::NodeImageSwitch::GetStaticType())
 	{
 		UE::Mutable::Private::NodeImageSwitch* ImageSwitchVar = StaticCast<UE::Mutable::Private::NodeImageSwitch*>(ParentNode);
 		AddChildFunc(ImageSwitchVar->Parameter.get(), FString::Printf(TEXT("PARAM")));
@@ -669,14 +676,14 @@ void SMutableGraphViewer::GetChildrenForInfo(TSharedPtr<FMutableGraphTreeElement
 		}
 	}
 
-	else if (ParentNode->GetType() == UE::Mutable::Private::NodeImageMipmap::GetStaticType())
+	else if (ParentNodeType == UE::Mutable::Private::NodeImageMipmap::GetStaticType())
 	{
 		UE::Mutable::Private::NodeImageMipmap* ImageMipMapVar = StaticCast<UE::Mutable::Private::NodeImageMipmap*>(ParentNode);
 		AddChildFunc(ImageMipMapVar->Source.get(), FString::Printf(TEXT("SOURCE")));
 		AddChildFunc(ImageMipMapVar->Factor.get(), FString::Printf(TEXT("FACTOR")));
 	}
 
-	else if (ParentNode->GetType() == UE::Mutable::Private::NodeImageLayer::GetStaticType())
+	else if (ParentNodeType == UE::Mutable::Private::NodeImageLayer::GetStaticType())
 	{
 		UE::Mutable::Private::NodeImageLayer* ImageLayerVar = StaticCast<UE::Mutable::Private::NodeImageLayer*>(ParentNode);
 		AddChildFunc(ImageLayerVar->Base.get(), FString::Printf(TEXT("BASE")));
@@ -684,7 +691,7 @@ void SMutableGraphViewer::GetChildrenForInfo(TSharedPtr<FMutableGraphTreeElement
 		AddChildFunc(ImageLayerVar->Blended.get(), FString::Printf(TEXT("BLEND")));
 	}
 	
-	else if (ParentNode->GetType() == UE::Mutable::Private::NodeImageLayerColour::GetStaticType())
+	else if (ParentNodeType == UE::Mutable::Private::NodeImageLayerColour::GetStaticType())
 	{
 		UE::Mutable::Private::NodeImageLayerColour* ImageLayerColourVar = StaticCast<UE::Mutable::Private::NodeImageLayerColour*>(ParentNode);
 		AddChildFunc(ImageLayerColourVar->Base.get(), FString::Printf(TEXT("BASE")));
@@ -692,20 +699,20 @@ void SMutableGraphViewer::GetChildrenForInfo(TSharedPtr<FMutableGraphTreeElement
 		AddChildFunc(ImageLayerColourVar->Colour.get(), FString::Printf(TEXT("COLOR")));
 	}
 
-	else if (ParentNode->GetType() == UE::Mutable::Private::NodeImageResize::GetStaticType())
+	else if (ParentNodeType == UE::Mutable::Private::NodeImageResize::GetStaticType())
 	{
 		UE::Mutable::Private::NodeImageResize* ImageResizeVar = StaticCast<UE::Mutable::Private::NodeImageResize*>(ParentNode);
 		AddChildFunc(ImageResizeVar->Base.get(), FString::Printf(TEXT("BASE")));
 	}
 
-	else if (ParentNode->GetType() == UE::Mutable::Private::NodeMeshMorph::GetStaticType())
+	else if (ParentNodeType == UE::Mutable::Private::NodeMeshMorph::GetStaticType())
 	{
 		UE::Mutable::Private::NodeMeshMorph* MeshMorphVar = StaticCast<UE::Mutable::Private::NodeMeshMorph*>(ParentNode);
 		AddChildFunc(MeshMorphVar->Base.get(), FString::Printf(TEXT("BASE")));
 		AddChildFunc(MeshMorphVar->Factor.get(), FString::Printf(TEXT("FACTOR")));
 	}
 
-	else if (ParentNode->GetType() == UE::Mutable::Private::NodeImageProject::GetStaticType())
+	else if (ParentNodeType == UE::Mutable::Private::NodeImageProject::GetStaticType())
 	{
 		UE::Mutable::Private::NodeImageProject* ImageProjectVar = StaticCast<UE::Mutable::Private::NodeImageProject*>(ParentNode);
 		AddChildFunc(ImageProjectVar->Projector.get(), FString::Printf(TEXT("PROJECTOR")));
@@ -716,18 +723,18 @@ void SMutableGraphViewer::GetChildrenForInfo(TSharedPtr<FMutableGraphTreeElement
 		AddChildFunc(ImageProjectVar->AngleFadeEnd.get(), FString::Printf(TEXT("FADE END ANGLE")));
 	}
 
-	else if (ParentNode->GetType() == UE::Mutable::Private::NodeImagePlainColour::GetStaticType())
+	else if (ParentNodeType == UE::Mutable::Private::NodeImagePlainColour::GetStaticType())
 	{
 		UE::Mutable::Private::NodeImagePlainColour* ImagePlainColourVar = StaticCast<UE::Mutable::Private::NodeImagePlainColour*>(ParentNode);
 		AddChildFunc(ImagePlainColourVar->Colour.get(), FString::Printf(TEXT("COLOR")));
 	}
 
-	else if (ParentNode->GetType() == UE::Mutable::Private::NodeLayout::GetStaticType())
+	else if (ParentNodeType == UE::Mutable::Private::NodeLayout::GetStaticType())
 	{
 		// Nothing to show
 	}
 
-	else if (ParentNode->GetType() == UE::Mutable::Private::NodeScalarEnumParameter::GetStaticType())
+	else if (ParentNodeType == UE::Mutable::Private::NodeScalarEnumParameter::GetStaticType())
 	{
 		UE::Mutable::Private::NodeScalarEnumParameter* ScalarEnumParameterVar = StaticCast<UE::Mutable::Private::NodeScalarEnumParameter*>(ParentNode);
 		for (int32 RangeIndex = 0; RangeIndex < ScalarEnumParameterVar->Ranges.Num(); RangeIndex++)
@@ -736,13 +743,13 @@ void SMutableGraphViewer::GetChildrenForInfo(TSharedPtr<FMutableGraphTreeElement
 		}
 	}
 
-	else if (ParentNode->GetType() == UE::Mutable::Private::NodeMeshFragment::GetStaticType())
+	else if (ParentNodeType == UE::Mutable::Private::NodeMeshFragment::GetStaticType())
 	{
 		UE::Mutable::Private::NodeMeshFragment* MeshFragmentVar = StaticCast<UE::Mutable::Private::NodeMeshFragment*>(ParentNode);
 		AddChildFunc(MeshFragmentVar->SourceMesh.get(), FString::Printf(TEXT("MESH")));
 	}
 
-	else if (ParentNode->GetType() == UE::Mutable::Private::NodeColourSampleImage::GetStaticType())
+	else if (ParentNodeType == UE::Mutable::Private::NodeColourSampleImage::GetStaticType())
 	{
 		UE::Mutable::Private::NodeColourSampleImage* ColorSampleImageVar = StaticCast<UE::Mutable::Private::NodeColourSampleImage*>(ParentNode);
 		AddChildFunc(ColorSampleImageVar->Image.get(), FString::Printf(TEXT("IMAGE")));
@@ -750,7 +757,7 @@ void SMutableGraphViewer::GetChildrenForInfo(TSharedPtr<FMutableGraphTreeElement
 		AddChildFunc(ColorSampleImageVar->Y.get(), FString::Printf(TEXT("Y")));
 	}
 
-	else if (ParentNode->GetType() == UE::Mutable::Private::NodeImageInterpolate::GetStaticType())
+	else if (ParentNodeType == UE::Mutable::Private::NodeImageInterpolate::GetStaticType())
 	{
 		UE::Mutable::Private::NodeImageInterpolate* ImageInterpolateVar = StaticCast<UE::Mutable::Private::NodeImageInterpolate*>(ParentNode);
 		AddChildFunc(ImageInterpolateVar->Factor.get(), FString::Printf(TEXT("FACTOR")));
@@ -760,12 +767,12 @@ void SMutableGraphViewer::GetChildrenForInfo(TSharedPtr<FMutableGraphTreeElement
 		}
 	}
 	
-	else if (ParentNode->GetType() == UE::Mutable::Private::NodeScalarConstant::GetStaticType())
+	else if (ParentNodeType == UE::Mutable::Private::NodeScalarConstant::GetStaticType())
 	{
 		// Nothing to show
 	}
 
-	else if (ParentNode->GetType() == UE::Mutable::Private::NodeScalarParameter::GetStaticType())
+	else if (ParentNodeType == UE::Mutable::Private::NodeScalarParameter::GetStaticType())
 	{
 		UE::Mutable::Private::NodeScalarParameter* ScalarParameterVar = StaticCast<UE::Mutable::Private::NodeScalarParameter*>(ParentNode);
 		for (int32 RangeIndex = 0; RangeIndex < ScalarParameterVar->Ranges.Num(); RangeIndex++)
@@ -774,7 +781,7 @@ void SMutableGraphViewer::GetChildrenForInfo(TSharedPtr<FMutableGraphTreeElement
 		}
 	}
 
-	else if (ParentNode->GetType() == UE::Mutable::Private::NodeColourParameter::GetStaticType())
+	else if (ParentNodeType == UE::Mutable::Private::NodeColourParameter::GetStaticType())
 	{
 		UE::Mutable::Private::NodeColourParameter* ColorParameterVar = StaticCast<UE::Mutable::Private::NodeColourParameter*>(ParentNode);
 		for (int32 RangeIndex = 0; RangeIndex < ColorParameterVar->Ranges.Num(); RangeIndex++)
@@ -783,32 +790,47 @@ void SMutableGraphViewer::GetChildrenForInfo(TSharedPtr<FMutableGraphTreeElement
 		}
 	}
 
-	else if (ParentNode->GetType() == UE::Mutable::Private::NodeColourConstant::GetStaticType())
+	else if (ParentNodeType == UE::Mutable::Private::NodeColourConstant::GetStaticType())
 	{
 		// Nothing to show
 	}
 
 
-	else if (ParentNode->GetType() == UE::Mutable::Private::NodeImageConstant::GetStaticType())
+	else if (ParentNodeType == UE::Mutable::Private::NodeImageConstant::GetStaticType())
 	{
 		// Nothing to show
 	}
 
+	else if (ParentNodeType == UE::Mutable::Private::NodeMatrixParameter::GetStaticType())
+	{
+		const UE::Mutable::Private::NodeMatrixParameter* NodeMatrix = StaticCast<UE::Mutable::Private::NodeMatrixParameter*>(ParentNode);
+			
+		for (int32 RangeIndex = 0; RangeIndex < NodeMatrix->Ranges.Num(); RangeIndex++)
+		{
+			AddChildFunc(NodeMatrix->Ranges[RangeIndex].get(), FString::Printf(TEXT("RANGE [%d]"), RangeIndex));
+		}
+	}
+	
+	else if (ParentNodeType == UE::Mutable::Private::NodeMatrixConstant::GetStaticType())
+	{
+		// Nothing to show
+	}
+	
 	
-	else if (ParentNode->GetType() == UE::Mutable::Private::NodeScalarCurve::GetStaticType())
+	else if (ParentNodeType == UE::Mutable::Private::NodeScalarCurve::GetStaticType())
 	{
 		UE::Mutable::Private::NodeScalarCurve* ScalarCurveVar = StaticCast<UE::Mutable::Private::NodeScalarCurve*>(ParentNode);
 		AddChildFunc(ScalarCurveVar->CurveSampleValue.get(), FString::Printf(TEXT("INPUT")));
 	}
 
-	else if (ParentNode->GetType() == UE::Mutable::Private::NodeMeshMakeMorph::GetStaticType())
+	else if (ParentNodeType == UE::Mutable::Private::NodeMeshMakeMorph::GetStaticType())
 	{
 		UE::Mutable::Private::NodeMeshMakeMorph* MeshMakeMorphVar = StaticCast<UE::Mutable::Private::NodeMeshMakeMorph*>(ParentNode);
 		AddChildFunc(MeshMakeMorphVar->Base.get(), FString::Printf(TEXT("BASE")));
 		AddChildFunc(MeshMakeMorphVar->Target.get(), FString::Printf(TEXT("TARGET")));
 	}
 	
-	else if (ParentNode->GetType() == UE::Mutable::Private::NodeProjectorParameter::GetStaticType())
+	else if (ParentNodeType == UE::Mutable::Private::NodeProjectorParameter::GetStaticType())
 	{
 		UE::Mutable::Private::NodeProjectorParameter* ProjectorParameterVar = StaticCast<UE::Mutable::Private::NodeProjectorParameter*>(ParentNode);
 		for (int32 RangeIndex = 0; RangeIndex < ProjectorParameterVar->Ranges.Num(); RangeIndex++)
@@ -817,12 +839,12 @@ void SMutableGraphViewer::GetChildrenForInfo(TSharedPtr<FMutableGraphTreeElement
 		}
 	}
 
-	else if (ParentNode->GetType() == UE::Mutable::Private::NodeProjectorConstant::GetStaticType())
+	else if (ParentNodeType == UE::Mutable::Private::NodeProjectorConstant::GetStaticType())
 	{
 		// Nothing to show
 	}
 
-	else if (ParentNode->GetType() == UE::Mutable::Private::NodeColourSwitch::GetStaticType())
+	else if (ParentNodeType == UE::Mutable::Private::NodeColourSwitch::GetStaticType())
 	{
 		UE::Mutable::Private::NodeColourSwitch* ColorSwitchVar = StaticCast<UE::Mutable::Private::NodeColourSwitch*>(ParentNode);
 		AddChildFunc(ColorSwitchVar->Parameter.get(), TEXT("PARAM"));
@@ -832,7 +854,7 @@ void SMutableGraphViewer::GetChildrenForInfo(TSharedPtr<FMutableGraphTreeElement
 		}
 	}
 
-	else if (ParentNode->GetType() == UE::Mutable::Private::NodeImageSwizzle::GetStaticType())
+	else if (ParentNodeType == UE::Mutable::Private::NodeImageSwizzle::GetStaticType())
 	{
 		UE::Mutable::Private::NodeImageSwizzle* ImageSwizzleVar = StaticCast<UE::Mutable::Private::NodeImageSwizzle*>(ParentNode);
 		for (int32 SourceIndex = 0; SourceIndex < ImageSwizzleVar->Sources.Num(); ++SourceIndex)
@@ -841,13 +863,13 @@ void SMutableGraphViewer::GetChildrenForInfo(TSharedPtr<FMutableGraphTreeElement
 		}
 	}
 	
-	else if (ParentNode->GetType() == UE::Mutable::Private::NodeImageInvert::GetStaticType())
+	else if (ParentNodeType == UE::Mutable::Private::NodeImageInvert::GetStaticType())
 	{
 		UE::Mutable::Private::NodeImageInvert* ImageInvertVar = StaticCast<UE::Mutable::Private::NodeImageInvert*>(ParentNode);
 		AddChildFunc(ImageInvertVar->Base.get(), TEXT("BASE"));
 	}
 
-	else if (ParentNode->GetType() == UE::Mutable::Private::NodeImageMultiLayer::GetStaticType())
+	else if (ParentNodeType == UE::Mutable::Private::NodeImageMultiLayer::GetStaticType())
 	{
 		UE::Mutable::Private::NodeImageMultiLayer* ImageMultilayerVar = StaticCast<UE::Mutable::Private::NodeImageMultiLayer*>(ParentNode);
 		AddChildFunc(ImageMultilayerVar->Base.get(), FString::Printf(TEXT("BASE")));
@@ -856,12 +878,12 @@ void SMutableGraphViewer::GetChildrenForInfo(TSharedPtr<FMutableGraphTreeElement
 		AddChildFunc(ImageMultilayerVar->Range.get(), FString::Printf(TEXT("RANGE")));
 	}
 	
-	else if (ParentNode->GetType() == UE::Mutable::Private::NodeImageTable::GetStaticType())
+	else if (ParentNodeType == UE::Mutable::Private::NodeImageTable::GetStaticType())
 	{
 		// No nodes to show
 	}
 
-	else if (ParentNode->GetType() == UE::Mutable::Private::NodeMeshTable::GetStaticType())
+	else if (ParentNodeType == UE::Mutable::Private::NodeMeshTable::GetStaticType())
 	{
 		UE::Mutable::Private::NodeMeshTable* MeshTableVar = StaticCast<UE::Mutable::Private::NodeMeshTable*>(ParentNode);
 		for (int32 LayoutIndex = 0; LayoutIndex < MeshTableVar->Layouts.Num(); ++LayoutIndex)
@@ -870,12 +892,12 @@ void SMutableGraphViewer::GetChildrenForInfo(TSharedPtr<FMutableGraphTreeElement
 		}
 	}
 
-	else if (ParentNode->GetType() == UE::Mutable::Private::NodeScalarTable::GetStaticType())
+	else if (ParentNodeType == UE::Mutable::Private::NodeScalarTable::GetStaticType())
 	{
 		// Nothing to show
 	}
 
-	else if (ParentNode->GetType() == UE::Mutable::Private::NodeScalarSwitch::GetStaticType())
+	else if (ParentNodeType == UE::Mutable::Private::NodeScalarSwitch::GetStaticType())
 	{
 		UE::Mutable::Private::NodeScalarSwitch* ScalarSwitchVar = StaticCast<UE::Mutable::Private::NodeScalarSwitch*>(ParentNode);
 		AddChildFunc(ScalarSwitchVar->Parameter.get(), TEXT("PARAM"));
@@ -885,7 +907,7 @@ void SMutableGraphViewer::GetChildrenForInfo(TSharedPtr<FMutableGraphTreeElement
 		}
 	}
 	
-	else if (ParentNode->GetType() == UE::Mutable::Private::NodeColourFromScalars::GetStaticType())
+	else if (ParentNodeType == UE::Mutable::Private::NodeColourFromScalars::GetStaticType())
 	{
 		UE::Mutable::Private::NodeColourFromScalars* ScalarTableVar = StaticCast<UE::Mutable::Private::NodeColourFromScalars*>(ParentNode);
 		AddChildFunc(ScalarTableVar->X.get(), TEXT("X"));
@@ -894,7 +916,7 @@ void SMutableGraphViewer::GetChildrenForInfo(TSharedPtr<FMutableGraphTreeElement
 		AddChildFunc(ScalarTableVar->W.get(), TEXT("W"));
 	}
 
-	else if (ParentNode->GetType() == UE::Mutable::Private::NodeMaterialModify::GetStaticType())
+	else if (ParentNodeType == UE::Mutable::Private::NodeMaterialModify::GetStaticType())
 	{
 		UE::Mutable::Private::NodeMaterialModify* MaterialModify = StaticCast<UE::Mutable::Private::NodeMaterialModify*>(ParentNode);
 
@@ -914,12 +936,20 @@ void SMutableGraphViewer::GetChildrenForInfo(TSharedPtr<FMutableGraphTreeElement
 		}
 	}
 	
+	else if (ParentNodeType == UE::Mutable::Private::NodeMeshTransformWithBone::GetStaticType())
+	{
+		const UE::Mutable::Private::NodeMeshTransformWithBone* NodeMeshTransformWithBone = StaticCast<UE::Mutable::Private::NodeMeshTransformWithBone*>(ParentNode);
+		
+		AddChildFunc(NodeMeshTransformWithBone->Source.get(), TEXT("SOURCE"));
+		AddChildFunc( NodeMeshTransformWithBone->MatrixNode.get(), TEXT("TRANSFORM"));
+	}
+	
 	else
 	{
-		UE_LOG(LogMutable,Error,TEXT("The node of type %d has not been implemented, so its children won't be added to the tree."), int32(ParentNode->GetType()->Type));
+		UE_LOG(LogMutable,Error,TEXT("The node of type %d has not been implemented, so its children won't be added to the tree."), int32(ParentNodeType->Type));
 
 		// Add a placeholder to the tree
-		const FString Prefix =  FString::Printf(TEXT("[%d] NODE TYPE NOT IMPLEMENTED"), int32(ParentNode->GetType()->Type));
+		const FString Prefix =  FString::Printf(TEXT("[%d] NODE TYPE NOT IMPLEMENTED"), int32(ParentNodeType->Type));
 		AddChildFunc(nullptr, Prefix);
 	}
 #endif
diff --git a/Engine/Plugins/Mutable/Source/CustomizableObjectEditor/Public/MuCOE/Nodes/CustomizableObjectNode.h b/Engine/Plugins/Mutable/Source/CustomizableObjectEditor/Public/MuCOE/Nodes/CustomizableObjectNode.h
--- a/Engine/Plugins/Mutable/Source/CustomizableObjectEditor/Public/MuCOE/Nodes/CustomizableObjectNode.h
+++ b/Engine/Plugins/Mutable/Source/CustomizableObjectEditor/Public/MuCOE/Nodes/CustomizableObjectNode.h
@@ -368,7 +368,7 @@ protected:
 	 * 
 	 * Given the node context (this), create all the required pins.
 	 *
-	 * @param RemapPins remap pins action. Use this object to save any information required to later perfrom the remap pins action. Can be null.
+	 * @param RemapPins remap pins action. Use this object to save any information required to later perform the remap pins action. Can be null.
 	 */
 	virtual void AllocateDefaultPins(UCustomizableObjectNodeRemapPins* RemapPins) {}
 
diff --git a/Engine/Plugins/Mutable/Source/MutableRuntime/Private/MuR/OpMeshTransformWithBone.cpp b/Engine/Plugins/Mutable/Source/MutableRuntime/Private/MuR/OpMeshTransformWithBone.cpp
--- a/Engine/Plugins/Mutable/Source/MutableRuntime/Private/MuR/OpMeshTransformWithBone.cpp
+++ b/Engine/Plugins/Mutable/Source/MutableRuntime/Private/MuR/OpMeshTransformWithBone.cpp
@@ -2,7 +2,6 @@
 
 #include "MuR/OpMeshTransformWithBone.h"
 
-#include "OpMeshClipWithMesh.h"
 #include "MuR/Mesh.h"
 #include "MuR/MeshPrivate.h"
 #include "MuR/MutableTrace.h"
diff --git a/Engine/Plugins/Mutable/Source/MutableTools/Internal/MuT/Node.h b/Engine/Plugins/Mutable/Source/MutableTools/Internal/MuT/Node.h
--- a/Engine/Plugins/Mutable/Source/MutableTools/Internal/MuT/Node.h
+++ b/Engine/Plugins/Mutable/Source/MutableTools/Internal/MuT/Node.h
@@ -45,6 +45,7 @@ namespace UE::Mutable::Private
 			MeshClipMorphPlane,
 			MeshClipWithMesh,
 			MeshRemoveMesh,
+			MeshTransformWithBone,
 			MeshApplyPose,
 			MeshVariation,
 			MeshReshape,
diff --git a/Engine/Plugins/Mutable/Source/MutableTools/Internal/MuT/NodeMeshTransformWithBone.h b/Engine/Plugins/Mutable/Source/MutableTools/Internal/MuT/NodeMeshTransformWithBone.h
new file mode 100644
--- /dev/null
+++ b/Engine/Plugins/Mutable/Source/MutableTools/Internal/MuT/NodeMeshTransformWithBone.h
@@ -0,0 +1,44 @@
+// Copyright Epic Games, Inc. All Rights Reserved.
+
+#pragma once
+
+#include "MuR/Ptr.h"
+#include "MuT/Node.h"
+#include "MuT/NodeMesh.h"
+#include "MuT/NodeMatrix.h"
+
+#define UE_API MUTABLETOOLS_API
+
+namespace UE::Mutable::Private
+{
+	/** */
+	class NodeMeshTransformWithBone : public NodeMesh
+	{
+	public:
+
+		Ptr<NodeMesh> Source;
+		
+		FName BoneName = NAME_None;
+		
+		float ThresholdFactor = 0.0f;
+		
+		Ptr<NodeMatrix> MatrixNode;
+		
+		// Node interface
+		virtual const FNodeType* GetType() const override { return &StaticType; }
+		static const FNodeType* GetStaticType() { return &StaticType; }
+
+	protected:
+
+		/** Forbidden. Manage with the Ptr<> template. */
+		virtual ~NodeMeshTransformWithBone() override = default;
+
+	private:
+
+		static UE_API FNodeType StaticType;
+
+	};
+
+}
+
+#undef UE_API
\ No newline at end of file
diff --git a/Engine/Plugins/Mutable/Source/MutableTools/Private/MuT/CodeGenerator.h b/Engine/Plugins/Mutable/Source/MutableTools/Private/MuT/CodeGenerator.h
--- a/Engine/Plugins/Mutable/Source/MutableTools/Private/MuT/CodeGenerator.h
+++ b/Engine/Plugins/Mutable/Source/MutableTools/Private/MuT/CodeGenerator.h
@@ -133,6 +133,7 @@ namespace UE::Mutable::Private
 	class NodeMatrixConstant;
 	class NodeMatrixParameter;
 	class NodeMaterialParameter;
+	class NodeMeshTransformWithBone;
 	struct FObjectState;
 	struct FProgram;
 
@@ -857,8 +858,9 @@ namespace UE::Mutable::Private
         FMeshTask GenerateMesh_Transform(const FMeshGenerationStaticOptions&, FMeshOptionsTask, const NodeMeshTransform*);
         FMeshTask GenerateMesh_ClipMorphPlane(const FMeshGenerationStaticOptions&, FMeshOptionsTask, const NodeMeshClipMorphPlane*);
         FMeshTask GenerateMesh_ClipWithMesh(const FMeshGenerationStaticOptions&, FMeshOptionsTask, const NodeMeshClipWithMesh*);
+        FMeshTask GenerateMesh_ApplyPose(const FMeshGenerationStaticOptions&, FMeshOptionsTask, const NodeMeshApplyPose*);
         FMeshTask GenerateMesh_RemoveMesh(const FMeshGenerationStaticOptions&, FMeshOptionsTask, const NodeMeshRemoveMesh*);
-    	FMeshTask GenerateMesh_ApplyPose(const FMeshGenerationStaticOptions&, FMeshOptionsTask, const NodeMeshApplyPose*);
+		FMeshTask GenerateMesh_TransformWithBone(const FMeshGenerationStaticOptions&, FMeshOptionsTask, const NodeMeshTransformWithBone*);
         FMeshTask GenerateMesh_Variation(const FMeshGenerationStaticOptions&, FMeshOptionsTask, const NodeMeshVariation*);
 		FMeshTask GenerateMesh_Table(const FMeshGenerationStaticOptions&, FMeshOptionsTask, const NodeMeshTable*);
 		FMeshTask GenerateMesh_Reshape(const FMeshGenerationStaticOptions&, FMeshOptionsTask, const NodeMeshReshape*);
diff --git a/Engine/Plugins/Mutable/Source/MutableTools/Private/MuT/CodeGenerator_Mesh.cpp b/Engine/Plugins/Mutable/Source/MutableTools/Private/MuT/CodeGenerator_Mesh.cpp
--- a/Engine/Plugins/Mutable/Source/MutableTools/Private/MuT/CodeGenerator_Mesh.cpp
+++ b/Engine/Plugins/Mutable/Source/MutableTools/Private/MuT/CodeGenerator_Mesh.cpp
@@ -6,6 +6,7 @@
 #include "ASTOpMeshMaskDiff.h"
 #include "ASTOpMeshRemoveMask.h"
 #include "ASTOpMeshSkeletalMeshBreak.h"
+#include "ASTOpMeshTransformWithBone.h"
 #include "Containers/Array.h"
 #include "Misc/AssertionMacros.h"
 #include "MuR/Layout.h"
@@ -59,6 +60,7 @@
 #include "MuT/NodeMeshTransform.h"
 #include "MuT/NodeMeshVariation.h"
 #include "MuT/NodeMeshSkeletalMeshBreak.h"
+#include "MuT/NodeMeshTransformWithBone.h"
 #include "MuT/Table.h"
 #include "MuT/TablePrivate.h"
 
@@ -99,6 +101,7 @@ namespace UE::Mutable::Private
 		case Node::EType::MeshClipMorphPlane: Result = GenerateMesh_ClipMorphPlane(StaticOptions, InOptions, static_cast<const NodeMeshClipMorphPlane*>(Node)); break;
 		case Node::EType::MeshClipWithMesh: Result = GenerateMesh_ClipWithMesh(StaticOptions, InOptions, static_cast<const NodeMeshClipWithMesh*>(Node)); break;
 		case Node::EType::MeshRemoveMesh: Result = GenerateMesh_RemoveMesh(StaticOptions, InOptions, static_cast<const NodeMeshRemoveMesh*>(Node)); break;
+		case Node::EType::MeshTransformWithBone: Result = GenerateMesh_TransformWithBone(StaticOptions, InOptions, static_cast<const NodeMeshTransformWithBone*>(Node)); break;
 		case Node::EType::MeshApplyPose: Result = GenerateMesh_ApplyPose(StaticOptions, InOptions, static_cast<const NodeMeshApplyPose*>(Node)); break;
 		case Node::EType::MeshVariation: Result = GenerateMesh_Variation(StaticOptions, InOptions, static_cast<const NodeMeshVariation*>(Node)); break;
 		case Node::EType::MeshTable: Result = GenerateMesh_Table(StaticOptions, InOptions, static_cast<const NodeMeshTable*>(Node)); break;
@@ -1202,6 +1205,53 @@ namespace UE::Mutable::Private
 		);
     }
 
+	
+	FMeshTask CodeGenerator::GenerateMesh_TransformWithBone(const FMeshGenerationStaticOptions& InStaticOptions, FMeshOptionsTask InOptions, const NodeMeshTransformWithBone* InNode)
+	{
+		const NodeMeshTransformWithBone* TransformWithBoneNode = static_cast<const NodeMeshTransformWithBone*>(InNode);
+		const void* MessageContext = TransformWithBoneNode->GetMessageContext();
+		
+		if (!TransformWithBoneNode->Source)
+		{
+			// This argument is required
+			ErrorLog->Add("Mesh Transform with Bone source node is not set.", ELMT_ERROR, MessageContext);
+			return UE::Tasks::MakeCompletedTask<FMeshGenerationResult>();
+		}
+		
+		if (!TransformWithBoneNode->MatrixNode)
+		{
+			// This argument is required
+			ErrorLog->Add("Mesh Transform with Bone transform matrix node is not set.", ELMT_ERROR, MessageContext);
+			return UE::Tasks::MakeCompletedTask<FMeshGenerationResult>();
+		}
+		
+		// Generate the source mesh
+		FMeshTask MeshGenerationTask = GenerateMesh(InStaticOptions, InOptions, TransformWithBoneNode->Source);
+		
+		// Transform matrix.
+		FMatrixGenerationResult ChildResult;
+		GenerateMatrix(ChildResult, InStaticOptions, TransformWithBoneNode->MatrixNode);
+		Ptr<ASTOp> MatrixOp = ChildResult.op;
+		
+		return UE::Tasks::Launch(TEXT("MutableMeshTransformWithBone"),
+			[MeshGenerationTask, MatrixOp, BoneName = TransformWithBoneNode->BoneName, ThresholdFactor = TransformWithBoneNode->ThresholdFactor]() mutable
+			{
+				FMeshGenerationResult SourceMeshGenerationResults = MeshGenerationTask.GetResult();
+
+				Ptr<ASTOpMeshTransformWithBone> TransformOp = new ASTOpMeshTransformWithBone();
+				TransformOp->SourceMesh = SourceMeshGenerationResults.MeshOp;
+				TransformOp->Matrix = MatrixOp;
+				TransformOp->BoneName = BoneName;
+				TransformOp->ThresholdFactor = ThresholdFactor;
+
+				SourceMeshGenerationResults.MeshOp = TransformOp;
+
+				return SourceMeshGenerationResults;
+			},
+			UE::Tasks::Prerequisites(MeshGenerationTask)
+		);
+	}
+
 
 	FMeshTask CodeGenerator::GenerateMesh_RemoveMesh(const FMeshGenerationStaticOptions& StaticOptions, FMeshOptionsTask InOptions, const NodeMeshRemoveMesh* RemoveMeshNode)
 	{
diff --git a/Engine/Plugins/Mutable/Source/MutableTools/Private/MuT/Node.cpp b/Engine/Plugins/Mutable/Source/MutableTools/Private/MuT/Node.cpp
--- a/Engine/Plugins/Mutable/Source/MutableTools/Private/MuT/Node.cpp
+++ b/Engine/Plugins/Mutable/Source/MutableTools/Private/MuT/Node.cpp
@@ -117,6 +117,7 @@
 #include "MuT/NodeSkeletalMeshParameter.h"
 #include "MuT/NodeMaterialModify.h"
 #include "MuT/NodeMeshRemoveMesh.h"
+#include "MuT/NodeMeshTransformWithBone.h"
 
 
 namespace UE::Mutable::Private
@@ -217,6 +218,7 @@ namespace UE::Mutable::Private
 	FNodeType NodeMeshClipDeform::StaticType = FNodeType(Node::EType::MeshClipDeform, NodeMesh::GetStaticType());
 	FNodeType NodeMeshClipWithMesh::StaticType = FNodeType(Node::EType::MeshClipWithMesh, NodeMesh::GetStaticType());
 	FNodeType NodeMeshRemoveMesh::StaticType = FNodeType(Node::EType::MeshRemoveMesh, NodeMesh::GetStaticType());
+	FNodeType NodeMeshTransformWithBone::StaticType = FNodeType(Node::EType::MeshTransformWithBone, NodeMesh::GetStaticType());
 	FNodeType NodeMeshSkeletalMeshBreak::StaticType = FNodeType(Node::EType::MeshSkeletalMeshBreak, NodeMesh::GetStaticType());
 	FNodeType NodeMeshMakeMorph::StaticType = FNodeType(Node::EType::MeshMakeMorph, NodeMesh::GetStaticType());
 	FNodeType NodeMeshApplyPose::StaticType = FNodeType(Node::EType::MeshApplyPose, NodeMesh::GetStaticType());
