[Mutable] Fix baked morphs not working when realtime morphs not enabled. Fix crash in OpMeshMorph for morph batches with zero size but number of elements > 0.

-Removed unused morph data from FMutableConvertSkeletalMeshContex.
-Removed unused mesh delta based morph function.

#rnx

#Synced-CL 51207105

diff --git a/Engine/Plugins/Mutable/Source/CustomizableObject/Private/MuCO/CustomizableObjectSystem.cpp b/Engine/Plugins/Mutable/Source/CustomizableObject/Private/MuCO/CustomizableObjectSystem.cpp
--- a/Engine/Plugins/Mutable/Source/CustomizableObject/Private/MuCO/CustomizableObjectSystem.cpp
+++ b/Engine/Plugins/Mutable/Source/CustomizableObject/Private/MuCO/CustomizableObjectSystem.cpp
@@ -4239,11 +4239,14 @@ void ReconstructMorphTargetsFromMeshCompressedData(const UE::Mutable::Private::F
 			OutMorphTargets[MorphIndex].Vertices.SetNumUninitialized(NumMorphVertices);
 			OutMorphTargets[MorphIndex].NumVertices = NumMorphVertices;
 			
+			TConstArrayView<uint32> MorphDataBufferView = Mesh.MorphDataBuffer;
+
 			int32 MorphVertexIndex = 0;
 			for (int32 BatchHeaderIndex = 0; BatchHeaderIndex < MorphNumBatches; ++BatchHeaderIndex)
 			{
+				check((BatchStartOffset + BatchHeaderIndex) * NumBatchHeaderDwords + NumBatchHeaderDwords <= (uint32)MorphDataBufferView.Num());
 				TConstArrayView<uint32> BatchHeaderData(
-						&Mesh.MorphDataBuffer[(BatchStartOffset + BatchHeaderIndex) * NumBatchHeaderDwords], 
+						MorphDataBufferView.GetData() + (BatchStartOffset + BatchHeaderIndex) * NumBatchHeaderDwords, 
 						NumBatchHeaderDwords);
 
 				FDeltaBatchHeader BatchHeader;
@@ -4257,9 +4260,12 @@ void ReconstructMorphTargetsFromMeshCompressedData(const UE::Mutable::Private::F
 				TArray<FQuantizedDelta, TInlineAllocator<UE::MorphTargetVertexCodec::BatchSize>> QuantizedDeltas;
 				QuantizedDeltas.SetNumUninitialized(BatchHeader.NumElements);
 
+				const uint32 BatchSizeInDwords = CalculateBatchDwords(BatchHeader); 
+
+				check(BatchHeader.DataOffset / sizeof(uint32) + BatchSizeInDwords <= (uint32)MorphDataBufferView.Num());
 				TConstArrayView<uint32> Data(
-						&Mesh.MorphDataBuffer[BatchHeader.DataOffset / sizeof(uint32)], 
-						CalculateBatchDwords(BatchHeader));
+						MorphDataBufferView.GetData() + (BatchHeader.DataOffset / sizeof(uint32)), 
+						BatchSizeInDwords);
 
 				ReadQuantizedDeltas(QuantizedDeltas, BatchHeader, Data);
 				for (FQuantizedDelta& QuantizedDelta : QuantizedDeltas)
diff --git a/Engine/Plugins/Mutable/Source/CustomizableObjectEditor/Private/MuCOE/GenerateMutableSource/GenerateMutableSourceMesh.cpp b/Engine/Plugins/Mutable/Source/CustomizableObjectEditor/Private/MuCOE/GenerateMutableSource/GenerateMutableSourceMesh.cpp
--- a/Engine/Plugins/Mutable/Source/CustomizableObjectEditor/Private/MuCOE/GenerateMutableSource/GenerateMutableSourceMesh.cpp
+++ b/Engine/Plugins/Mutable/Source/CustomizableObjectEditor/Private/MuCOE/GenerateMutableSource/GenerateMutableSourceMesh.cpp
@@ -86,9 +86,6 @@ struct FMutableConvertSkeletalMeshContext
 	UE::Mutable::Private::EMeshBufferFormat SourceSkinWeightFormat;
 	int32 NumBoneInfluences = 0;
 
-	// Bake Morph
-	TArray<FVector3f> MorphDeltaPosition;
-	TArray<FVector3f> MorphDeltaNormal;
 };
 
 void SetBuffersFormat(FMutableCompilationContext& Context,
@@ -142,7 +139,6 @@ void SetBuffersFormat(FMutableCompilationContext& Context,
 	const bool bHasTexCoords = !bIgnoreTexCoords;
 	const bool bHasSkinWeights = !bIgnoreSkinning && MutableBonesPerVertex > 0;
 	const bool bHasVertexColors = SourceMesh.GetHasVertexColors();
-	const bool bHasRealTimeMorphs = !bIgnoreMorphs && Context.Options.bRealTimeMorphTargetsEnabled && SourceMesh.GetMorphTargets().Num();
 	const bool bHasClothing = !bIgnoreSkinning && SourceMesh.HasActiveClothingAssetsForLOD(LODIndex);
 	const bool bHasSkinWeightProfiles = !bIgnoreSkinning && Context.Options.bSkinWeightProfilesEnabled && LODModel.SkinWeightProfiles.Num();
 
@@ -597,10 +593,15 @@ bool IsSkeletalMeshCompatibleWithRefSkeleton(FSkeletalMeshComponentInfo& Compone
 }
 
 
-TArray<FName> GetUsedRealTimeMorphsNames(const FMutableSourceMeshData& SourceMeshData, const TArray<FRealTimeMorphSelectionOverride>& Overrides)
+TArray<FName> GetUsedRealTimeMorphsNames(const FMutableSourceMeshData& SourceMeshData, bool bEnableRealTimeMorphTargets, const TArray<FRealTimeMorphSelectionOverride>& Overrides)
 {
 	TArray<FName> Result;
 	
+	if (!bEnableRealTimeMorphTargets)
+	{
+		return Result;
+	}
+
 	if (SourceMeshData.TableReferenceSkeletalMesh.IsNull())
 	{	
 		const USkeletalMesh* SourceSkeletalMesh = Cast<USkeletalMesh>(UE::Mutable::Private::LoadObject(SourceMeshData.Mesh));
@@ -777,25 +778,10 @@ namespace UE::Mutable::Private
 							check(NumChannels == 1);
 
 							FVector3f* DestPosition = reinterpret_cast<FVector3f*>(VertexBuffer.Data.GetData());
-							if (MeshContext->MorphDeltaPosition.IsEmpty())
-							{
-								for (const FSoftSkinVertex& Vertex : Vertices)
-								{
-									*DestPosition = Vertex.Position;
-									++DestPosition;
-								}
-							}
-							else
+							for (const FSoftSkinVertex& Vertex : Vertices)
 							{
-								check(MeshContext->MorphDeltaPosition.Num() == NumVertices);
-
-								const FVector3f* MorphDeltaPosition = MeshContext->MorphDeltaPosition.GetData();
-								for (const FSoftSkinVertex& Vertex : Vertices)
-								{
-									*DestPosition = Vertex.Position + *MorphDeltaPosition;
-									++DestPosition;
-									++MorphDeltaPosition;
-								}
+								*DestPosition = Vertex.Position;
+								++DestPosition;
 							}
 						}
 						else if (NumChannels == 2
@@ -817,43 +803,18 @@ namespace UE::Mutable::Private
 
 							FPackedNormal* TypedDestData = reinterpret_cast<FPackedNormal*>(DestData);
 
-							if (MeshContext->MorphDeltaNormal.IsEmpty())
-							{
-								for (const FSoftSkinVertex& Vertex : Vertices)
-								{
-									// Tangent
-									*TypedDestData = Vertex.TangentX;
-									(*TypedDestData).Vector.W = 0;
-									++TypedDestData;
-
-									// Normal
-									*TypedDestData = Vertex.TangentZ;
-									FMatrix44f Mat(Vertex.TangentX, Vertex.TangentY, FVector3f(Vertex.TangentZ), FVector3f(0, 0, 0));
-									(*TypedDestData).Vector.W = Mat.RotDeterminant() < 0 ? -128 : 127;
-									++TypedDestData;
-								}
-							}
-							else
+							for (const FSoftSkinVertex& Vertex : Vertices)
 							{
-								check(MeshContext->MorphDeltaNormal.Num() == NumVertices);
+								// Tangent
+								*TypedDestData = Vertex.TangentX;
+								(*TypedDestData).Vector.W = 0;
+								++TypedDestData;
 
-								FVector3f* MorphDeltaNormal = MeshContext->MorphDeltaNormal.GetData();
-								for (const FSoftSkinVertex& Vertex : Vertices)
-								{
-									// Tangent
-									*TypedDestData = Vertex.TangentX;
-									(*TypedDestData).Vector.W = 0;
-									++TypedDestData;
-
-									// Normal
-									*MorphDeltaNormal += Vertex.TangentZ;
-									*TypedDestData = *MorphDeltaNormal;
-
-									FMatrix44f Mat(Vertex.TangentX, Vertex.TangentY, *MorphDeltaNormal, FVector3f(0, 0, 0));
-									(*TypedDestData).Vector.W = Mat.RotDeterminant() < 0 ? -128 : 127;
-									++TypedDestData;
-									++MorphDeltaNormal;
-								}
+								// Normal
+								*TypedDestData = Vertex.TangentZ;
+								FMatrix44f Mat(Vertex.TangentX, Vertex.TangentY, FVector3f(Vertex.TangentZ), FVector3f(0, 0, 0));
+								(*TypedDestData).Vector.W = Mat.RotDeterminant() < 0 ? -128 : 127;
+								++TypedDestData;
 							}
 						}
 						else if (NumChannels == 2
@@ -1319,9 +1280,7 @@ UE::Tasks::FTask ConvertSkeletalMeshToMutable(TSharedRef<FMeshConversionContext>
 
 	int32 NextBufferIndex = MeshContext->NumBaseBuffers;
 
-	// RealTimeMorphs and Clothing buffers might be removed if unused. Convert them before launching the tasks to convert non-optional buffers.
-
-	if (!bIgnoreMorphs && Context.Options.bRealTimeMorphTargetsEnabled && SkeletalMesh->GetMorphTargets().Num())
+	if (!bIgnoreMorphs && SkeletalMesh->GetMorphTargets().Num())
 	{
 		UnrealConversionUtils::SetMorphData(
 				MutableMesh.Get(), SkeletalMesh, 
@@ -1345,33 +1304,6 @@ UE::Tasks::FTask ConvertSkeletalMeshToMutable(TSharedRef<FMeshConversionContext>
 		}
 	}
 
-	// Copy embedded Morph Data
-	if (!bIgnoreMorphs && !MorphName.IsEmpty())
-	{
-		UMorphTarget* MorphTarget = SkeletalMesh->FindMorphTarget(*MorphName);
-		if (MorphTarget && MorphTarget->GetMorphLODModels().IsValidIndex(LODIndex))
-		{
-			const int32 MorphVertexStart = SkeletalMesh->GetImportedModel()->LODModels[LODIndex].Sections[SectionIndex].GetVertexBufferIndex();
-
-			MeshContext->MorphDeltaPosition.SetNumZeroed(MeshSection.NumVertices);
-			MeshContext->MorphDeltaNormal.SetNumZeroed(MeshSection.NumVertices);
-
-			TArrayView<FVector3f> DeltaPosition(MeshContext->MorphDeltaPosition);
-			TArrayView<FVector3f> DeltaNormal(MeshContext->MorphDeltaNormal);
-
-			const FMorphTargetLODModel& MorphLODModel = MorphTarget->GetMorphLODModels()[LODIndex];
-			for (const FMorphTargetDelta& MorphDelta : MorphLODModel.Vertices)
-			{
-				const int32 VertexIndex = MorphDelta.SourceIdx - MorphVertexStart;
-				if (VertexIndex >= 0 && (uint32)VertexIndex < VertexCount)
-				{
-					DeltaPosition[VertexIndex] = MorphDelta.PositionDelta;
-					DeltaNormal[VertexIndex] = MorphDelta.TangentZDelta;
-				}
-			}
-		}
-	}
-
 	UE::Tasks::FTask ConvertSkeletalMeshBuffersTask = UE::Mutable::Private::ConvertSkeletalMeshDefaultBuffers(MutableMesh, LODModel, MeshSection, MeshContext);
 
 	// SkinWeightProfiles vertex info.
@@ -2770,7 +2702,8 @@ UE::Mutable::Private::Ptr<UE::Mutable::Private::NodeMesh> GenerateMutableSourceM
 				}	
 			}
 
-			MeshNode->RealTimeMorphNames = GetUsedRealTimeMorphsNames(Source, GenerationContext.CompilationContext->RealTimeMorphTargetsOverrides);
+			const bool bEnableRealTimeMorphTargets = GenerationContext.CompilationContext->Options.bRealTimeMorphTargetsEnabled;
+			MeshNode->RealTimeMorphNames = GetUsedRealTimeMorphsNames(Source, bEnableRealTimeMorphTargets, GenerationContext.CompilationContext->RealTimeMorphTargetsOverrides);
 
 			MeshNode->BonePosePriority = GenerationContext.BonePosePriorityStack.Last();
 			MeshNode->SocketPriority = GenerationContext.SocketPriorityStack.Last();
@@ -3535,9 +3468,10 @@ UE::Mutable::Private::Ptr<UE::Mutable::Private::NodeMesh> GenerateMutableSourceM
 
 							// Set real-time morph names.
 							{
+								const bool bEnableRealTimeMorphTargets = GenerationContext.CompilationContext->Options.bRealTimeMorphTargetsEnabled;
 								FMutableSourceMeshData Source;
 								Source.TableReferenceSkeletalMesh = DefaultSkeletalMesh;
-								MeshTableNode->RealTimeMorphNames = GetUsedRealTimeMorphsNames(Source, {});
+								MeshTableNode->RealTimeMorphNames = GetUsedRealTimeMorphsNames(Source, bEnableRealTimeMorphTargets, {});
 							}
 
 							// Applying Mesh Morph Nodes
@@ -3641,8 +3575,9 @@ UE::Mutable::Private::Ptr<UE::Mutable::Private::NodeMesh> GenerateMutableSourceM
 			}
 		}
 
+		const bool bEnableRealTimeMorphTargets = GenerationContext.CompilationContext->Options.bRealTimeMorphTargetsEnabled;
 		Source.bUseAllRealTimeMorphs = true; // TODO: Remove, this should be selected similarly to the constant case, meanwhile use all morphs from the Reference Mesh.
-		BreakNode->RealTimeMorphNames = GetUsedRealTimeMorphsNames(Source, GenerationContext.CompilationContext->RealTimeMorphTargetsOverrides);
+		BreakNode->RealTimeMorphNames = GetUsedRealTimeMorphsNames(Source, bEnableRealTimeMorphTargets, GenerationContext.CompilationContext->RealTimeMorphTargetsOverrides);
 
 		// For constants we also have:
 		// animbp instance?
diff --git a/Engine/Plugins/Mutable/Source/MutableRuntime/Private/MuR/OpMeshMerge.cpp b/Engine/Plugins/Mutable/Source/MutableRuntime/Private/MuR/OpMeshMerge.cpp
--- a/Engine/Plugins/Mutable/Source/MutableRuntime/Private/MuR/OpMeshMerge.cpp
+++ b/Engine/Plugins/Mutable/Source/MutableRuntime/Private/MuR/OpMeshMerge.cpp
@@ -488,7 +488,7 @@ namespace OpMeshMergeMorphsInternal
 				}
 			}
 
-			check(BatchAccumulatedOffsetInDwords == Result.MorphDataBuffer.Num());
+			check(BatchAccumulatedOffsetInDwords <= Result.MorphDataBuffer.Num());
 		}
 	}
 } // namespace OpMeshMergeMorphsInternal
diff --git a/Engine/Plugins/Mutable/Source/MutableRuntime/Private/MuR/OpMeshMorph.cpp b/Engine/Plugins/Mutable/Source/MutableRuntime/Private/MuR/OpMeshMorph.cpp
--- a/Engine/Plugins/Mutable/Source/MutableRuntime/Private/MuR/OpMeshMorph.cpp
+++ b/Engine/Plugins/Mutable/Source/MutableRuntime/Private/MuR/OpMeshMorph.cpp
@@ -8,6 +8,7 @@
 #include "MuR/SparseIndexMap.h"
 #include "PackedNormal.h"
 #include "MuR/Mesh.h"
+#include "MuR/MutableRuntimeModule.h"
 #include "UObject/NameTypes.h"
 #include "Animation/MorphTarget.h"
 #include "Rendering/MorphTargetVertexCodec.h"
@@ -87,11 +88,17 @@ void UE::Mutable::Private::MeshMorph(FMesh* Mesh, const FName& MorphName, const
 	
 	const uint32 BatchStartOffset = Mesh->Morph.BatchStartOffsetPerMorph[MorphIndex];
 	const uint32 MorphNumBatches  = Mesh->Morph.BatchesPerMorph[MorphIndex];
+
+	TConstArrayView<uint32> MorphDataBufferView = Mesh->MorphDataBuffer;
+
+	uint32 NumMalformedBatchesFound = 0;
+
 	for (uint32 BatchHeaderIndex = 0; BatchHeaderIndex < MorphNumBatches; ++BatchHeaderIndex)
 	{
-		TArrayView<uint32> BatchHeaderData(
-				&Mesh->MorphDataBuffer[(BatchStartOffset + BatchHeaderIndex) * NumBatchHeaderDwords], 
-				NumBatchHeaderDwords);
+		const uint32 BatchHeaderOffsetInDwords = (BatchStartOffset + BatchHeaderIndex) * NumBatchHeaderDwords;
+
+		check(BatchHeaderOffsetInDwords + NumBatchHeaderDwords <= (uint32)MorphDataBufferView.Num());
+		TConstArrayView<uint32> BatchHeaderData(MorphDataBufferView.GetData() + BatchHeaderOffsetInDwords, NumBatchHeaderDwords);
 
 		FDeltaBatchHeader BatchHeader;
 		ReadHeader(BatchHeader, BatchHeaderData);
@@ -104,9 +111,25 @@ void UE::Mutable::Private::MeshMorph(FMesh* Mesh, const FName& MorphName, const
 		TArray<FQuantizedDelta, TInlineAllocator<UE::MorphTargetVertexCodec::BatchSize>> QuantizedDeltas;
 		QuantizedDeltas.SetNumUninitialized(BatchHeader.NumElements);
 
-		TArrayView<uint32> Data(
-				&Mesh->MorphDataBuffer[BatchHeader.DataOffset / sizeof(uint32)], 
-				CalculateBatchDwords(BatchHeader));
+		uint32 BatchSizeInDwords = CalculateBatchDwords(BatchHeader);
+
+		TConstArrayView<uint32> Data;
+		if (BatchSizeInDwords > 0)
+		{
+			uint32 BatchDataOffsetInDwords = BatchHeader.DataOffset / sizeof(uint32);
+
+			// Don't assume data is good, prevent any carsh and record skipped batchs.
+			if ((BatchDataOffsetInDwords + BatchSizeInDwords <= (uint32)MorphDataBufferView.Num()) & 
+				((BatchHeader.DataOffset & (sizeof(uint32) - 1)) == 0))
+			{
+				Data = TConstArrayView<uint32>(MorphDataBufferView.GetData() + BatchDataOffsetInDwords, BatchSizeInDwords);
+			}
+			else
+			{
+				++NumMalformedBatchesFound;
+				continue;
+			}
+		}
 
 		ReadQuantizedDeltas(QuantizedDeltas, BatchHeader, Data);
 
@@ -118,7 +141,7 @@ void UE::Mutable::Private::MeshMorph(FMesh* Mesh, const FName& MorphName, const
 			// Positions
 			if (bHasPositions)
 			{
-				check(Delta.SourceIdx < NumVertices)
+				check(Delta.SourceIdx < NumVertices);
 				FVector3f& Position = *reinterpret_cast<FVector3f*>(*(MeshPositionIter + Delta.SourceIdx));
 				Position += Delta.PositionDelta * Factor;
 			}
@@ -127,11 +150,12 @@ void UE::Mutable::Private::MeshMorph(FMesh* Mesh, const FName& MorphName, const
 			if (bHasOptimizedNormals)
 			{
 				// Normal
-				check(Delta.SourceIdx < NumVertices)
+				check(Delta.SourceIdx < NumVertices);
+
 				FPackedNormal* PackedNormal = reinterpret_cast<FPackedNormal*>((MeshNormalIter + Delta.SourceIdx).ptr());
 				int8 W = PackedNormal->Vector.W;
 				const FVector3f BaseNormal = PackedNormal->ToFVector3f();
-				
+
 				const FVector3f Normal = (BaseNormal + Delta.TangentZDelta * Factor).GetSafeNormal();
 
 				*PackedNormal = Normal;
@@ -140,7 +164,6 @@ void UE::Mutable::Private::MeshMorph(FMesh* Mesh, const FName& MorphName, const
 				// Tangent
 				if (MeshTangentIter.ptr())
 				{
-					check(Delta.SourceIdx < NumVertices)
 					FPackedNormal* PackedTangent = reinterpret_cast<FPackedNormal*>((MeshTangentIter + Delta.SourceIdx).ptr());
 					const FVector3f BaseTangent = PackedTangent->ToFVector3f();
 
@@ -210,318 +233,9 @@ void UE::Mutable::Private::MeshMorph(FMesh* Mesh, const FName& MorphName, const
 			}
 		}
 	}
-}
-
-void UE::Mutable::Private::MeshMorph(FMesh* BaseMesh, const FMesh* MaxMesh, const float Factor)
-{
-    MUTABLE_CPUPROFILER_SCOPE(MeshMorph_Compiled);
 
-	if (!BaseMesh)
+	if (NumMalformedBatchesFound > 0)
 	{
-		return;
-	}
-
-	const auto MakeIndexMap = [](
-		MeshVertexIdIteratorConst BaseIdIter, int32 BaseNum,
-		MeshVertexIdIteratorConst MorphIdIter, int32 MorphNum)
-	-> SparseIndexMapSet
-	{	
-		MUTABLE_CPUPROFILER_SCOPE(MakeIndexMap);
-		TArray<SparseIndexMapSet::FRangeDesc> RangeDescs;
-
-		// Detect all ranges and their limits
-		{
-			for (int32 Index = 0; Index < BaseNum; ++Index, ++BaseIdIter)
-			{
-				const uint64 BaseId = BaseIdIter.Get();
-
-				uint32 Prefix = BaseId >> 32;
-				uint32 Id = BaseId & 0xffffffff;
-				bool bFound = false;
-				for (SparseIndexMapSet::FRangeDesc& Range : RangeDescs)
-				{
-					if (Range.Prefix == Prefix)
-					{
-						Range.MinIndex = FMath::Min(Id, Range.MinIndex);
-						Range.MaxIndex = FMath::Max(Id, Range.MaxIndex);
-						bFound = true;
-						break;
-					}
-				}
-
-				if (!bFound)
-				{
-					RangeDescs.Add({Prefix, Id, Id});
-				}
-			}
-		}
-
-		SparseIndexMapSet IndexMap(RangeDescs);
-       
-        for (int32 Index = 0; Index < MorphNum; ++Index, ++MorphIdIter)
-        {
-            const uint64 MorphId = MorphIdIter.Get();
-            
-            IndexMap.Insert(MorphId, Index);
-        }
-
-		return IndexMap;
-	};
-
-	// Number of vertices to modify
-	const int32 MaxNum = MaxMesh ? MaxMesh->GetVertexBuffers().GetElementCount() : 0;
-	const int32 BaseNum = BaseMesh ? BaseMesh->GetVertexBuffers().GetElementCount() : 0;
-	const FMesh* RefTarget = MaxMesh;
-
-	if (BaseNum == 0 || MaxNum == 0)
-	{
-		return;
-	}
-
-	if (!RefTarget)
-	{
-		return;
-	}
-	
-	constexpr int32 MorphBufferDataChannel = 0;
-	const int32 ChannelsNum = RefTarget->GetVertexBuffers().GetBufferChannelCount(MorphBufferDataChannel);
-
-	TArray<UntypedMeshBufferIterator> BaseGenericIters;
-	BaseGenericIters.SetNum(ChannelsNum);
-	TArray<UntypedMeshBufferIteratorConst> MinGenericIters;
-	MinGenericIters.SetNum(ChannelsNum);
-	TArray<UntypedMeshBufferIteratorConst> MaxGenericIters;
-	MaxGenericIters.SetNum(ChannelsNum);
-
-	// Get pointers to vertex position data
-	MeshBufferIterator<EMeshBufferFormat::Float32, float, 3> BasePositionIter(BaseMesh->VertexBuffers, EMeshBufferSemantic::Position, 0);	
-	MeshBufferIteratorConst<EMeshBufferFormat::Float32, float, 3> MaxPositionIter(RefTarget->VertexBuffers, EMeshBufferSemantic::Position, 0);
-
-	// {BiNormal, Tangent, Normal}
-	TStaticArray<UntypedMeshBufferIterator, 3> BaseTangentFrameChannelsIters;
-	UntypedMeshBufferIteratorConst MaxNormalChannelIter;
-
-	const bool bBaseHasNormals = UntypedMeshBufferIteratorConst(BaseMesh->GetVertexBuffers(), EMeshBufferSemantic::Normal, 0).ptr() != nullptr;
-	for (int32 ChannelIndex = 0; ChannelIndex < ChannelsNum; ++ChannelIndex)
-	{
-		const FMeshBufferSet& MBSPriv = RefTarget->GetVertexBuffers();
-		const FMeshBufferChannel& Channel = MBSPriv.Buffers[MorphBufferDataChannel].Channels[ChannelIndex];
-		EMeshBufferSemantic Sem = Channel.Semantic;
-		int32 SemIndex = Channel.SemanticIndex;
-		if (Sem == EMeshBufferSemantic::Normal && bBaseHasNormals)
-		{
-			BaseTangentFrameChannelsIters[2] = UntypedMeshBufferIterator(BaseMesh->GetVertexBuffers(), Sem, SemIndex);
-			MaxNormalChannelIter = UntypedMeshBufferIteratorConst(MaxMesh->GetVertexBuffers(), Sem, SemIndex);
-		}
-		else if (Sem == EMeshBufferSemantic::Tangent && bBaseHasNormals)
-		{
-			BaseTangentFrameChannelsIters[1] = UntypedMeshBufferIterator(BaseMesh->GetVertexBuffers(), Sem, SemIndex);
-		}
-		else if (Sem == EMeshBufferSemantic::Binormal && bBaseHasNormals)
-		{
-			BaseTangentFrameChannelsIters[0] = UntypedMeshBufferIterator(BaseMesh->GetVertexBuffers(), Sem, SemIndex);
-		}
-		else if(Sem != EMeshBufferSemantic::Position)
-		{
-			BaseGenericIters[ChannelIndex] = UntypedMeshBufferIterator(BaseMesh->GetVertexBuffers(), Sem, SemIndex);
-			MaxGenericIters[ChannelIndex] = UntypedMeshBufferIteratorConst(MaxMesh->GetVertexBuffers(), Sem, SemIndex);
-		}
-	}
-
-	MeshVertexIdIteratorConst BaseIdIter(BaseMesh);
-	
-	MeshVertexIdIteratorConst MaxIdIter(MaxMesh);
-	SparseIndexMapSet IndexMap = MakeIndexMap(BaseIdIter, BaseNum, MaxIdIter, MaxNum);
-    
-	MUTABLE_CPUPROFILER_SCOPE(ApplyMorph);
-
-	const UntypedMeshBufferIterator& BaseNormalIter = BaseTangentFrameChannelsIters[2];
-	const UntypedMeshBufferIterator& BaseTangentIter = BaseTangentFrameChannelsIters[1];
-	const UntypedMeshBufferIterator& BaseBiNormalIter = BaseTangentFrameChannelsIters[0];
-
-	const EMeshBufferFormat NormalFormat = BaseNormalIter.GetFormat();
-	const int32 NormalComps = BaseNormalIter.GetComponents();
-
-	const EMeshBufferFormat TangentFormat = BaseTangentIter.GetFormat();
-	const int32 TangentComps = BaseTangentIter.GetComponents();
-
-	const EMeshBufferFormat BiNormalFormat = BaseBiNormalIter.GetFormat();
-	const int32 BiNormalComps = BaseBiNormalIter.GetComponents();
-
-	const EMeshBufferFormat MorphNormalFormat = MaxNormalChannelIter.GetFormat();
-
-	const bool bHasPositions = BasePositionIter.ptr() && MaxPositionIter.ptr();
-	check((BasePositionIter.GetFormat() == EMeshBufferFormat::Float32 && BasePositionIter.GetComponents() == 3) || !bHasPositions);
-
-	const bool bHasOptimizedNormals = NormalFormat == EMeshBufferFormat::PackedDirS8_W_TangentSign && MorphNormalFormat == EMeshBufferFormat::Float32
-		&& (!BaseTangentIter.ptr() || TangentFormat == EMeshBufferFormat::PackedDirS8) && !BaseBiNormalIter.ptr();
-	
-	bool bHasGenericMorphs = false;
-	const int32 ChannelNum = MaxGenericIters.Num();
-	for (int32 ChannelIndex = 0; ChannelIndex < ChannelNum; ++ChannelIndex)
-	{
-		if (!(BaseGenericIters[ChannelIndex].ptr() && MaxGenericIters[ChannelIndex].ptr()))
-		{
-			continue;
-		}
-
-		bHasGenericMorphs = true;
-	}
-
-	for (int32 VertexIndex = 0; VertexIndex < BaseNum; ++VertexIndex)
-	{
-		const uint64 BaseId = (BaseIdIter + VertexIndex).Get();
-		const uint32 MorphIndex = IndexMap.Find(BaseId);
-
-		if (MorphIndex == SparseIndexMap::NotFoundValue)
-		{
-			continue;
-		}
-
-		// Find consecutive run.
-		MeshVertexIdIteratorConst RunBaseIter = BaseIdIter + VertexIndex;
-		MeshVertexIdIteratorConst RunMorphIter = MaxIdIter + MorphIndex;
-
-		int32 RunSize = 0;
-		for (; VertexIndex + RunSize < BaseNum && int32(MorphIndex) + RunSize < MaxNum && RunBaseIter.Get() == RunMorphIter.Get();
-			++RunSize, ++RunBaseIter, ++RunMorphIter);
-
-		// Positions
-		if (bHasPositions)
-		{
-			for (int32 RunIndex = 0; RunIndex < RunSize; ++RunIndex)
-			{
-				FVector3f& Position = *reinterpret_cast<FVector3f*>(*(BasePositionIter + VertexIndex + RunIndex));
-				const FVector3f& MorphPosition = *reinterpret_cast<const FVector3f*>(*(MaxPositionIter + MorphIndex + RunIndex));
-				Position += MorphPosition * Factor;
-			}
-		}
-
-		// Normals
-		if (bHasOptimizedNormals)
-		{
-			for (int32 RunIndex = 0; RunIndex < RunSize; ++RunIndex)
-			{
-				// Normal
-				FPackedNormal* PackedNormal = reinterpret_cast<FPackedNormal*>((BaseNormalIter + VertexIndex + RunIndex).ptr());
-				int8 W = PackedNormal->Vector.W;
-				const FVector3f BaseNormal = PackedNormal->ToFVector3f();
-
-				const FVector3f* MorphNormal = reinterpret_cast<const FVector3f*>((MaxNormalChannelIter + MorphIndex + RunIndex).ptr());
-
-				const FVector3f Normal = (BaseNormal + *MorphNormal * Factor).GetSafeNormal();
-
-				*PackedNormal = *reinterpret_cast<const FVector3f*>(&Normal);
-				PackedNormal->Vector.W = W;
-
-				// Tangent
-				if (BaseTangentIter.ptr())
-				{
-					FPackedNormal* PackedTangent = reinterpret_cast<FPackedNormal*>((BaseTangentIter + (VertexIndex + RunIndex)).ptr());
-					const FVector3f BaseTangent = PackedTangent->ToFVector3f();
-
-					// Orthogonalize Tangent based on new Normal. This assumes Normal and BaseTangent are normalized and different.
-					const FVector3f Tangent = (BaseTangent - FVector3f::DotProduct(Normal, BaseTangent) * Normal).GetSafeNormal();
-
-					*PackedTangent = *reinterpret_cast<const FVector3f*>(&Tangent);
-				}
-			}
-		}
-		else if (BaseNormalIter.ptr())
-		{
-			// When normal is packed, binormal channel is not expected. It is not a big deal if it's there but we would be doing extra unused work in that case. 
-			ensure(!(NormalFormat == EMeshBufferFormat::PackedDir8_W_TangentSign || NormalFormat == EMeshBufferFormat::PackedDirS8_W_TangentSign) || !BaseBiNormalIter.ptr());
-			
-			MUTABLE_CPUPROFILER_SCOPE(ApplyNormalMorph_SlowPath);
-
-			for (int32 RunIndex = 0; RunIndex < RunSize; ++RunIndex)
-			{
-				UntypedMeshBufferIterator NormalIter = BaseNormalIter + (VertexIndex + RunIndex);
-
-				const FVector3f BaseNormal = NormalIter.GetAsVec3f();
-				const FVector3f MorphNormal = (MaxNormalChannelIter + (MorphIndex + RunIndex)).GetAsVec3f();
-
-				const FVector3f Normal = (BaseNormal + MorphNormal * Factor).GetSafeNormal();
-
-				// Leave the tangent basis sign untouched for packed normals formats.
-				for (int32 C = 0; C < NormalComps && C < 3; ++C)
-				{
-					ConvertData(C, NormalIter.ptr(), NormalFormat, &Normal, EMeshBufferFormat::Float32);
-				}
-
-				// Tangent
-				if (BaseTangentIter.ptr())
-				{
-					UntypedMeshBufferIterator TangentIter = BaseTangentIter + (VertexIndex + RunIndex);
-
-					const FVector3f BaseTangent = TangentIter.GetAsVec3f();
-
-					// Orthogonalize Tangent based on new Normal. This assumes Normal and BaseTangent are normalized and different.
-					const FVector3f Tangent = (BaseTangent - FVector3f::DotProduct(Normal, BaseTangent) * Normal).GetSafeNormal();
-
-					for (int32 C = 0; C < TangentComps && C < 3; ++C)
-					{
-						ConvertData(C, TangentIter.ptr(), TangentFormat, &Tangent, EMeshBufferFormat::Float32);
-					}
-
-					// BiNormal
-					if (BaseBiNormalIter.ptr())
-					{
-						UntypedMeshBufferIterator BiNormalIter = BaseBiNormalIter + (VertexIndex + RunIndex);
-
-						const FVector3f& N = BaseNormal;
-						const FVector3f& T = BaseTangent;
-						const FVector3f  B = BiNormalIter.GetAsVec3f();
-
-						const float BaseTangentBasisDeterminant =
-							B.X * T.Y * N.Z + B.Z * T.X * N.Y + B.Y * T.Z * N.Y -
-							B.Z * T.Y * N.X - B.Y * T.X * N.Z - B.X * T.Z * N.Y;
-
-						const float BaseTangentBasisDeterminantSign = BaseTangentBasisDeterminant >= 0 ? 1.0f : -1.0f;
-
-						const FVector3f BiNormal = FVector3f::CrossProduct(Tangent, Normal) * BaseTangentBasisDeterminantSign;
-
-						for (int32 C = 0; C < BiNormalComps && C < 3; ++C)
-						{
-							ConvertData(C, BiNormalIter.ptr(), BiNormalFormat, &BiNormal, EMeshBufferFormat::Float32);
-						}
-					}
-				}
-			}
-		}
-		
-		// Generic Morphs
-		if (bHasGenericMorphs)
-		{
-			MUTABLE_CPUPROFILER_SCOPE(ApplyNormalMorph_Generic);
-			for (int32 ChannelIndex = 0; ChannelIndex < ChannelNum; ++ChannelIndex)
-			{
-				if (!(BaseGenericIters[ChannelIndex].ptr() && MaxGenericIters[ChannelIndex].ptr()))
-				{
-					continue;
-				}
-
-				UntypedMeshBufferIterator ChannelBaseIter = BaseGenericIters[ChannelIndex] + VertexIndex;
-				UntypedMeshBufferIteratorConst ChannelMorphIter = MaxGenericIters[ChannelIndex] + MorphIndex;
-
-				const EMeshBufferFormat DestChannelFormat = BaseGenericIters[ChannelIndex].GetFormat();
-				const int32 DestChannelComps = BaseGenericIters[ChannelIndex].GetComponents();
-
-				// Apply Morph to range found above.
-				for (int32 R = 0; R < RunSize; ++R, ++ChannelBaseIter, ++ChannelMorphIter)
-				{
-					const FVector4f Value = ChannelBaseIter.GetAsVec4f() + ChannelMorphIter.GetAsVec4f() * Factor;
-
-					// TODO: Optimize this for the specific components.
-					// Max 4 components
-					for (int32 Comp = 0; Comp < DestChannelComps && Comp < 4; ++Comp)
-					{
-						ConvertData(Comp, ChannelBaseIter.ptr(), DestChannelFormat, &Value, EMeshBufferFormat::Float32);
-					}
-				}
-			}
-		}
-
-		VertexIndex += FMath::Max(RunSize - 1, 0);
+		UE_LOGF(LogMutableCore, Warning, "Morph operation found invalid data, some batches were skipped.");
 	}
 }
diff --git a/Engine/Plugins/Mutable/Source/MutableRuntime/Private/MuR/OpMeshMorph.h b/Engine/Plugins/Mutable/Source/MutableRuntime/Private/MuR/OpMeshMorph.h
--- a/Engine/Plugins/Mutable/Source/MutableRuntime/Private/MuR/OpMeshMorph.h
+++ b/Engine/Plugins/Mutable/Source/MutableRuntime/Private/MuR/OpMeshMorph.h
@@ -14,8 +14,4 @@ namespace UE::Mutable::Private
 {
 	/** Morph Parameter Meshes. Mesh contains the deltas in morph buffer. Deltas are compressed using the engine format. */
 	void MeshMorph(FMesh* Mesh, const FName& MorphName, float Factor);
-
-	
-	/** Morph Compiled Meshes. MaxMesh contains de deltas in the position buffer. */
-	void MeshMorph(FMesh* BaseMesh, const FMesh* MaxMesh, float Factor); // TODO Use the engine compressed format. This means changing how be compile morphs.
 }
