[Mutable] Make MeshExtractLayoutBlock operate in place.

- This allows using CloneOrTakeOver for the operant meshes, potentially saving memory and compute time.

#rnx

#Synced-CL 51207105

diff --git a/Engine/Plugins/Mutable/Source/MutableRuntime/Private/MuR/CodeRunner.cpp b/Engine/Plugins/Mutable/Source/MutableRuntime/Private/MuR/CodeRunner.cpp
--- a/Engine/Plugins/Mutable/Source/MutableRuntime/Private/MuR/CodeRunner.cpp
+++ b/Engine/Plugins/Mutable/Source/MutableRuntime/Private/MuR/CodeRunner.cpp
@@ -1849,28 +1849,19 @@ namespace UE::Mutable::Private
 
 					if (Source)
 					{
-						TManagedPtr<FMesh> Result = CreateMesh();
-						bool bOutSuccess;
+						TManagedPtr<FMesh> Result = CloneOrTakeOver(MoveTemp(Source));
 
 						if (BlockCount > 0)
 						{
 							TArrayView<uint64> BlockIds(Blocks, BlockCount);
-							MeshExtractLayoutBlock(Result.Get(), Source.Get(), LayoutIndex, BlockIds, bOutSuccess);
-						}
-						else
-						{
-							MeshExtractLayoutBlock(Result.Get(), Source.Get(), LayoutIndex, bOutSuccess);
-						}
-						
-						if (!bOutSuccess)
-						{
-							StoreMesh(Item, Source);
+							MeshExtractLayoutBlock(Result.Get(), LayoutIndex, BlockIds);
 						}
 						else
 						{
-							StoreMesh(Item, Result);
+							MeshExtractLayoutBlock(Result.Get(), LayoutIndex);
 						}
 
+						StoreMesh(Item, Result);
 					}
 					else
 					{
diff --git a/Engine/Plugins/Mutable/Source/MutableRuntime/Private/MuR/OpMeshExtractLayoutBlock.cpp b/Engine/Plugins/Mutable/Source/MutableRuntime/Private/MuR/OpMeshExtractLayoutBlock.cpp
--- a/Engine/Plugins/Mutable/Source/MutableRuntime/Private/MuR/OpMeshExtractLayoutBlock.cpp
+++ b/Engine/Plugins/Mutable/Source/MutableRuntime/Private/MuR/OpMeshExtractLayoutBlock.cpp
@@ -7,29 +7,23 @@
 
 namespace UE::Mutable::Private
 {
-	void MeshExtractLayoutBlock(FMesh* Result, const FMesh* Source, uint32 LayoutIndex, TArrayView<uint64> BlockIds, bool& bOutSuccess)
+	void MeshExtractLayoutBlock(TNotNull<FMesh*> Result, uint32 LayoutIndex, TArrayView<uint64> BlockIds)
 	{
-		check(Source);
-		bOutSuccess = true;
-
-		// TODO: Optimise
-		Result->CopyFrom(*Source);
-
-		UntypedMeshBufferIteratorConst ItBlocks(Source->GetVertexBuffers(), EMeshBufferSemantic::LayoutBlock, LayoutIndex);
+		UntypedMeshBufferIteratorConst ItBlocks(Result->GetVertexBuffers(), EMeshBufferSemantic::LayoutBlock, LayoutIndex);
 
 		if (ItBlocks.GetFormat() != EMeshBufferFormat::None)
 		{
-			int32 ResultVertices = 0;
+			const int32 NumVertices = Result->GetVertexCount();
 
 			TBitArray<> VerticesToCull;
-			VerticesToCull.SetNumUninitialized(Source->GetVertexCount());
+			VerticesToCull.SetNumUninitialized(NumVertices);
 
 			if (ItBlocks.GetFormat() == EMeshBufferFormat::UInt16)
 			{
-				TArrayView<const uint16> VertexBlockIdsView(reinterpret_cast<const uint16*>(ItBlocks.ptr()), Source->GetVertexCount());
+				TArrayView<const uint16> VertexBlockIdsView(reinterpret_cast<const uint16*>(ItBlocks.ptr()), NumVertices);
 
-				const uint64 AbsoluteVertexBlockIdHighBits = uint64(Source->MeshIDPrefix) << 32;
-				for (int32 VertexIndex = VertexBlockIdsView.Num() - 1; VertexIndex >= 0; --VertexIndex)
+				const uint64 AbsoluteVertexBlockIdHighBits = uint64(Result->MeshIDPrefix) << 32;
+				for (int32 VertexIndex = NumVertices - 1; VertexIndex >= 0; --VertexIndex)
 				{
 					const uint64 AbsoluteVertexBlockId = AbsoluteVertexBlockIdHighBits | VertexBlockIdsView[VertexIndex];
 					VerticesToCull[VertexIndex] = BlockIds.Find(AbsoluteVertexBlockId) == INDEX_NONE;
@@ -37,8 +31,8 @@ namespace UE::Mutable::Private
 			}
 			else if (ItBlocks.GetFormat() == EMeshBufferFormat::UInt64)
 			{
-				TArrayView<const uint64> VertexBlockIdsView(reinterpret_cast<const uint64*>(ItBlocks.ptr()), Source->GetVertexCount());
-				for (int32 VertexIndex = VertexBlockIdsView.Num() - 1; VertexIndex >= 0; --VertexIndex)
+				TArrayView<const uint64> VertexBlockIdsView(reinterpret_cast<const uint64*>(ItBlocks.ptr()), NumVertices);
+				for (int32 VertexIndex = NumVertices - 1; VertexIndex >= 0; --VertexIndex)
 				{
 					VerticesToCull[VertexIndex] = BlockIds.Find(VertexBlockIdsView[VertexIndex]) == INDEX_NONE;
 				}
@@ -53,36 +47,29 @@ namespace UE::Mutable::Private
 		}
 	}
 
-	void MeshExtractLayoutBlock(FMesh* Result, const FMesh* Source, uint32 LayoutIndex, bool& bOutSuccess)
+	void MeshExtractLayoutBlock(TNotNull<FMesh*> Result, uint32 LayoutIndex)
 	{
-		check(Source);
-		bOutSuccess = true;
-
-		// TODO: Optimise
-		Result->CopyFrom(*Source);
-
-		UntypedMeshBufferIteratorConst ItBlocks(Source->GetVertexBuffers(), EMeshBufferSemantic::LayoutBlock, LayoutIndex);
+		UntypedMeshBufferIteratorConst ItBlocks(Result->GetVertexBuffers(), EMeshBufferSemantic::LayoutBlock, LayoutIndex);
 
 		if (ItBlocks.GetFormat() != EMeshBufferFormat::None)
 		{
-			int32 ResultVertices = 0;
-
+			const int32 NumVertices = Result->GetVertexCount();
 			TBitArray<> VerticesToCull;
-			VerticesToCull.SetNumUninitialized(Source->GetVertexCount());
+			VerticesToCull.SetNumUninitialized(NumVertices);
 
 			if (ItBlocks.GetFormat() == EMeshBufferFormat::UInt16)
 			{
-				TArrayView<const uint16> VertexBlockIdsView(reinterpret_cast<const uint16*>(ItBlocks.ptr()), Source->GetVertexCount());
+				TArrayView<const uint16> VertexBlockIdsView(reinterpret_cast<const uint16*>(ItBlocks.ptr()), NumVertices);
 
-				for (int32 VertexIndex = VertexBlockIdsView.Num() - 1; VertexIndex >= 0; --VertexIndex)
+				for (int32 VertexIndex = NumVertices - 1; VertexIndex >= 0; --VertexIndex)
 				{
 					VerticesToCull[VertexIndex] = VertexBlockIdsView[VertexIndex] == std::numeric_limits<uint16>::max();
 				}
 			}
 			else if (ItBlocks.GetFormat() == EMeshBufferFormat::UInt64)
 			{
-				TArrayView<const uint64> VertexBlockIdsView(reinterpret_cast<const uint64*>(ItBlocks.ptr()), Source->GetVertexCount());
-				for (int32 VertexIndex = VertexBlockIdsView.Num() - 1; VertexIndex >= 0; --VertexIndex)
+				TArrayView<const uint64> VertexBlockIdsView(reinterpret_cast<const uint64*>(ItBlocks.ptr()), NumVertices);
+				for (int32 VertexIndex = NumVertices - 1; VertexIndex >= 0; --VertexIndex)
 				{
 					VerticesToCull[VertexIndex] = VertexBlockIdsView[VertexIndex] == std::numeric_limits<uint64>::max();
 				}
diff --git a/Engine/Plugins/Mutable/Source/MutableRuntime/Private/MuR/OpMeshExtractLayoutBlock.h b/Engine/Plugins/Mutable/Source/MutableRuntime/Private/MuR/OpMeshExtractLayoutBlock.h
--- a/Engine/Plugins/Mutable/Source/MutableRuntime/Private/MuR/OpMeshExtractLayoutBlock.h
+++ b/Engine/Plugins/Mutable/Source/MutableRuntime/Private/MuR/OpMeshExtractLayoutBlock.h
@@ -4,14 +4,15 @@
 
 #include "HAL/Platform.h"
 #include "Containers/ArrayView.h"
+#include "Misc/NotNull.h"
 
 namespace UE::Mutable::Private
 {
     class FMesh;
 
-	/** Extract all vertices that have a layout block from the passed list on the given channel. */
-    void MeshExtractLayoutBlock(FMesh* Result, const FMesh* Source, uint32 LayoutIndex, TArrayView<uint64> BlockIds, bool& bOutSuccess);
+	/** Extract all vertices that have a layout block from the passed list on the given channel in place. */
+    void MeshExtractLayoutBlock(TNotNull<FMesh*> Result, uint32 LayoutIndex, TArrayView<uint64> BlockIds);
 
-	/** Extract all vertices that have a valid layout block on the given channel. */
-	void MeshExtractLayoutBlock(FMesh* Result, const FMesh* Source, uint32 LayoutIndex, bool& bOutSuccess);
+	/** Extract all vertices that have a valid layout block on the given channel in place. */
+	void MeshExtractLayoutBlock(TNotNull<FMesh*> Result, uint32 LayoutIndex);
 }
