[Mutable] OpImageLayer Refactor.

- Removed the use of templates to generate different code for every possible layering mode and format permutation.
- Now different layering modes and formats are described by a predefined sequence of operations that run on small batches of pixels (a kernel). Each operation in the sequence can generate temporary results stored in small buffers to be used by subsequent operations in the same run.
- Now blend operations operate on multiple elements instead of a single channel.
- Added auxiliary generic pixel manipulation functionality to accommodate for format mismatches and common operations.
- Added iterative decoding functionality to support RLE masks. This allows to fill fixed sized temporary buffers.
- Added vectorized SSE and NEON implementations of some of the operations.
- SubImage layer operations needed for patching are not implemented yet.

#rnx

#Synced-CL 51208470

diff --git a/Engine/Plugins/Mutable/Source/MutableRuntime/Private/MuR/CodeRunner.cpp b/Engine/Plugins/Mutable/Source/MutableRuntime/Private/MuR/CodeRunner.cpp
--- a/Engine/Plugins/Mutable/Source/MutableRuntime/Private/MuR/CodeRunner.cpp
+++ b/Engine/Plugins/Mutable/Source/MutableRuntime/Private/MuR/CodeRunner.cpp
@@ -21,11 +21,11 @@
 #include "MuR/OpImageApplyComposite.h"
 #include "MuR/OpImageBinarise.h"
 #include "MuR/OpImageBlend.h"
+#include "MuR/OpImageLayer.h"
 #include "MuR/OpImageColourMap.h"
 #include "MuR/OpImageDisplace.h"
 #include "MuR/OpImageInterpolate.h"
 #include "MuR/OpImageLuminance.h"
-#include "MuR/OpImageNormalCombine.h"
 #include "MuR/OpImageProject.h"
 #include "MuR/OpImageRasterMesh.h"
 #include "MuR/OpImageTransform.h"
@@ -3392,15 +3392,20 @@ namespace UE::Mutable::Private
 					Data.MultiLayer.Iterations = Iterations;
 					Data.MultiLayer.OriginalBaseFormat = InitialBaseFormat;
 					Data.MultiLayer.bBlendOnlyOneMip = false;
-					int32 DataPos = HeapData.Add(Data);
+					int32 HeapDataIndex = HeapData.Add(Data);
 
 					// Request the first layer
 					int32 CurrentIteration = 0;
-					FScheduledOp ItemCopy = Item;
+					FScheduledOp ItemForFirstLayerBlend = Item;
+
 					ExecutionIndex Index = GetMemory().GetRangeIndex(Item.ExecutionIndex);
 					Index.SetFromModelRangeIndex(Args.rangeId, CurrentIteration);
-					ItemCopy.ExecutionIndex = GetMemory().GetRangeIndexIndex(Index);
-					AddOp(FScheduledOp(Item.At, Item, 2, DataPos), FScheduledOp(Args.base, Item), FScheduledOp(Args.blended, ItemCopy), FScheduledOp(Args.mask, ItemCopy));
+					ItemForFirstLayerBlend.ExecutionIndex = GetMemory().GetRangeIndexIndex(Index);
+
+					AddOp(FScheduledOp(Item.At, Item, 2, HeapDataIndex), 
+							FScheduledOp(Args.base, Item), 
+							FScheduledOp(Args.blended, ItemForFirstLayerBlend), 
+							FScheduledOp(Args.mask, ItemForFirstLayerBlend));
 				}
 
                 break;
@@ -3422,169 +3427,113 @@ namespace UE::Mutable::Private
 
 				TManagedPtr<FImage> Base = StaticCastManagedPtr<FImage>(ConstCastManagedPtr<FResource>(Data.Resource));
 				 
-                FScheduledOp itemCopy = Item;
-                ExecutionIndex index = GetMemory().GetRangeIndex( Item.ExecutionIndex );
+                FScheduledOp ItemForLayerBlend = Item;
+                ExecutionIndex Index = GetMemory().GetRangeIndex(Item.ExecutionIndex);
 				
                 {
-                    index.SetFromModelRangeIndex( Args.rangeId, CurrentIteration);
-                    itemCopy.ExecutionIndex = GetMemory().GetRangeIndexIndex(index);
-					itemCopy.CustomState = 0;
-
-                    TManagedPtr<const FImage> Blended = LoadImage( FCacheAddress(Args.blended,itemCopy) );
-
-                    // This shouldn't happen in optimised models, but it could happen in editors, etc.
-                    // \todo: raise a performance warning?
-                    if (Blended && Blended->GetFormat() != Base->GetFormat())
-                    {
-						MUTABLE_CPUPROFILER_SCOPE(ImageResize_BlendedReformat);
+                    Index.SetFromModelRangeIndex(Args.rangeId, CurrentIteration);
+                    ItemForLayerBlend.ExecutionIndex = GetMemory().GetRangeIndexIndex(Index);
+					ItemForLayerBlend.CustomState = 0;
 
-						TManagedPtr<FImage> Formatted = CreateImage(Blended->GetSizeX(), Blended->GetSizeY(), Blended->GetLODCount(), Base->GetFormat(), EInitializationType::NotInitialized);
+                    TManagedPtr<const FImage> Blended = LoadImage(FCacheAddress(Args.blended, ItemForLayerBlend));
+                    TManagedPtr<const FImage> Mask = LoadImage(FCacheAddress(Args.mask, ItemForLayerBlend));
 
-						bool bSuccess = false;
-						ImOp.ImagePixelFormat(bSuccess, Settings.ImageCompressionQuality, Formatted.Get(), Blended.Get());
-						check(bSuccess);
-
-						Blended = Formatted;
-                    }
-
-					// TODO: This shouldn't happen, but be defensive.
-					FImageSize ResultSize = Base->GetSize();
-					if (Blended && Blended->GetSize() != ResultSize)
+					if (Blended)
 					{
-						MUTABLE_CPUPROFILER_SCOPE(ImageResize_BlendedFixForMultilayer);
-
-						TManagedPtr<FImage> Resized = CreateImage(ResultSize[0], ResultSize[1], Blended->GetLODCount(), Blended->GetFormat(), EInitializationType::NotInitialized);
-						ImOp.ImageResizeLinear(Resized.Get(), 0, Blended.Get());
+						// This shouldn't happen in optimised models, but it could happen in editors, etc.
+						// \todo: raise a performance warning?
+						if (Blended->GetFormat() != Base->GetFormat())
+						{
+							MUTABLE_CPUPROFILER_SCOPE(ImageFormat_BlendedFixForMultilayer);
 
-						Blended = Resized;
-					}
+							TManagedPtr<FImage> Formatted = CreateImage(Blended->GetSizeX(), Blended->GetSizeY(), Blended->GetLODCount(), Base->GetFormat(), EInitializationType::NotInitialized);
 
-					if (Blended->GetLODCount() < Base->GetLODCount())
-					{
-						Data.MultiLayer.bBlendOnlyOneMip = true;
-					}
+							bool bSuccess = false;
+							ImOp.ImagePixelFormat(bSuccess, Settings.ImageCompressionQuality, Formatted.Get(), Blended.Get());
+							check(bSuccess);
 
-					bool bApplyColorBlendToAlpha = false;
+							Blended = Formatted;
+						}
 
-					bool bDone = false;
+						// TODO: This shouldn't happen, but be defensive.
+						FImageSize ResultSize = Base->GetSize();
+						if (Blended->GetSize() != ResultSize)
+						{
+							MUTABLE_CPUPROFILER_SCOPE(ImageResize_BlendedFixForMultilayer);
 
-					// This becomes true if we need to update the mips of the resulting image
-					// This could happen in the base image has mips, but one of the blended one doesn't.
-					bool bBlendOnlyOneMip = Data.MultiLayer.bBlendOnlyOneMip;
-					bool bUseBlendSourceFromBlendAlpha = false; // (Args.flags& OP::ImageLayerArgs::F_BLENDED_RGB_FROM_ALPHA) != 0;
+							TManagedPtr<FImage> Resized = CreateImage(ResultSize[0], ResultSize[1], Blended->GetLODCount(), Blended->GetFormat(), EInitializationType::NotInitialized);
+							ImOp.ImageResizeLinear(Resized.Get(), 0, Blended.Get());
 
-					if (!Args.mask && Args.bUseMaskFromBlended
-						&&
-						Args.blendType == uint8(EBlendType::BT_BLEND)
-						&&
-						Args.blendTypeAlpha == uint8(EBlendType::BT_LIGHTEN) )
-					{
-						// This is a frequent critical-path case because of multilayer projectors.
-						bDone = true;
-					
-						constexpr bool bUseVectorImpl = false;
-						if constexpr (bUseVectorImpl)
-						{
-							BufferLayerCompositeVector<VectorBlendChannelMasked, VectorLightenChannel, false>(Base.Get(), Blended.Get(), bBlendOnlyOneMip, Args.BlendAlphaSourceChannel);
+							Blended = Resized;
 						}
-						else
+
+						if (Mask && !(Mask->GetFormat() == EImageFormat::L_UByte || Mask->GetFormat() == EImageFormat::L_UByteRLE))
 						{
-							BufferLayerComposite<BlendChannelMasked, LightenChannel, false>(Base.Get(), Blended.Get(), bBlendOnlyOneMip, Args.BlendAlphaSourceChannel);
-						}
-					}
+							MUTABLE_CPUPROFILER_SCOPE(ImageFormat_MaskFixForMultilayer);
+							
+							TManagedPtr<FImage> Formatted = CreateImage(Mask->GetSizeX(), Mask->GetSizeY(), Mask->GetLODCount(), EImageFormat::L_UByte, EInitializationType::NotInitialized);
 
-                    if (!bDone && Args.mask)
-                    {
-                        TManagedPtr<const FImage> Mask = LoadImage( FCacheAddress(Args.mask,itemCopy) );
+							bool bSuccess = false;
+							ImOp.ImagePixelFormat(bSuccess, Settings.ImageCompressionQuality, Formatted.Get(), Mask.Get());
+							check(bSuccess);
+
+							Mask = Formatted;
+						}
 
-						// TODO: This shouldn't happen, but be defensive.
 						if (Mask && Mask->GetSize() != ResultSize)
 						{
 							MUTABLE_CPUPROFILER_SCOPE(ImageResize_MaskFixForMultilayer);
 
 							TManagedPtr<FImage> Resized = CreateImage(ResultSize[0], ResultSize[1], Mask->GetLODCount(), Mask->GetFormat(), EInitializationType::NotInitialized);
 							ImOp.ImageResizeLinear(Resized.Get(), 0, Mask.Get());
+
 							Mask = Resized;
 						}
 
-						// Not implemented yet
-						check(!bUseBlendSourceFromBlendAlpha);
-
-                        switch (EBlendType(Args.blendType))
-                        {
-						case EBlendType::BT_NORMAL_COMBINE: check(false); break;
-                        case EBlendType::BT_SOFTLIGHT: BufferLayer<SoftLightChannelMasked, SoftLightChannel, false>(Base.Get(), Base.Get(), Mask.Get(), Blended.Get(), bApplyColorBlendToAlpha, bBlendOnlyOneMip); break;
-                        case EBlendType::BT_HARDLIGHT: BufferLayer<HardLightChannelMasked, HardLightChannel, false>(Base.Get(), Base.Get(), Mask.Get(), Blended.Get(), bApplyColorBlendToAlpha, bBlendOnlyOneMip); break;
-                        case EBlendType::BT_BURN: BufferLayer<BurnChannelMasked, BurnChannel, false>(Base.Get(), Base.Get(), Mask.Get(), Blended.Get(), bApplyColorBlendToAlpha, bBlendOnlyOneMip); break;
-                        case EBlendType::BT_DODGE: BufferLayer<DodgeChannelMasked, DodgeChannel, false>(Base.Get(), Base.Get(), Mask.Get(), Blended.Get(), bApplyColorBlendToAlpha, bBlendOnlyOneMip); break;
-                        case EBlendType::BT_SCREEN: BufferLayer<ScreenChannelMasked, ScreenChannel, false>(Base.Get(), Base.Get(), Mask.Get(), Blended.Get(), bApplyColorBlendToAlpha, bBlendOnlyOneMip); break;
-                        case EBlendType::BT_OVERLAY: BufferLayer<OverlayChannelMasked, OverlayChannel, false>(Base.Get(), Base.Get(), Mask.Get(), Blended.Get(), bApplyColorBlendToAlpha, bBlendOnlyOneMip); break;
-                        case EBlendType::BT_LIGHTEN: BufferLayer<LightenChannelMasked, LightenChannel, false>(Base.Get(), Base.Get(), Mask.Get(), Blended.Get(), bApplyColorBlendToAlpha, bBlendOnlyOneMip); break;
-                        case EBlendType::BT_MULTIPLY: BufferLayer<MultiplyChannelMasked, MultiplyChannel, false>(Base.Get(), Base.Get(), Mask.Get(), Blended.Get(), bApplyColorBlendToAlpha, bBlendOnlyOneMip); break;
-                        case EBlendType::BT_BLEND: BufferLayer<BlendChannelMasked, BlendChannel, false>(Base.Get(), Base.Get(), Mask.Get(), Blended.Get(), bApplyColorBlendToAlpha, bBlendOnlyOneMip); break;
-                        default: check(false);
+						if (Blended->GetLODCount() < Base->GetLODCount())
+						{
+							Data.MultiLayer.bBlendOnlyOneMip = true;
 						}
 
-						Mask = nullptr;
-                    }
-					else if (!bDone && Args.bUseMaskFromBlended)
-					{
-						// Not implemented yet
-						check(!bUseBlendSourceFromBlendAlpha);
+						bool bApplyColorBlendToAlpha = false;
+
+						// This becomes true if we need to update the mips of the resulting image
+						// This could happen in the base image has mips, but one of the blended one doesn't.
+						bool bBlendOnlyOneMip = Data.MultiLayer.bBlendOnlyOneMip;
+					
+						const int32 LODBegin = 0;
+						const int32 LODEnd   = bBlendOnlyOneMip ? 1 : Base->GetLODCount();
 
-						switch (EBlendType(Args.blendType))
+						bool bUseBlendSourceFromBlendAlpha = false; // (Args.flags& OP::ImageLayerArgs::F_BLENDED_RGB_FROM_ALPHA) != 0;
+					
+						if (EBlendType(Args.blendType) != EBlendType::BT_NORMAL_COMBINE)
 						{
-						case EBlendType::BT_NORMAL_COMBINE: check(false); break;
-						case EBlendType::BT_SOFTLIGHT: BufferLayerEmbeddedMask<SoftLightChannelMasked, SoftLightChannel, false>(Base.Get(), Base.Get(), Blended.Get(), bApplyColorBlendToAlpha, bBlendOnlyOneMip); break;
-						case EBlendType::BT_HARDLIGHT: BufferLayerEmbeddedMask<HardLightChannelMasked, HardLightChannel, false>(Base.Get(), Base.Get(), Blended.Get(), bApplyColorBlendToAlpha, bBlendOnlyOneMip); break;
-						case EBlendType::BT_BURN: BufferLayerEmbeddedMask<BurnChannelMasked, BurnChannel, false>(Base.Get(), Base.Get(), Blended.Get(), bApplyColorBlendToAlpha, bBlendOnlyOneMip); break;
-						case EBlendType::BT_DODGE: BufferLayerEmbeddedMask<DodgeChannelMasked, DodgeChannel, false>(Base.Get(), Base.Get(), Blended.Get(), bApplyColorBlendToAlpha, bBlendOnlyOneMip); break;
-						case EBlendType::BT_SCREEN: BufferLayerEmbeddedMask<ScreenChannelMasked, ScreenChannel, false>(Base.Get(), Base.Get(), Blended.Get(), bApplyColorBlendToAlpha, bBlendOnlyOneMip); break;
-						case EBlendType::BT_OVERLAY: BufferLayerEmbeddedMask<OverlayChannelMasked, OverlayChannel, false>(Base.Get(), Base.Get(), Blended.Get(), bApplyColorBlendToAlpha, bBlendOnlyOneMip); break;
-						case EBlendType::BT_LIGHTEN: BufferLayerEmbeddedMask<LightenChannelMasked, LightenChannel, false>(Base.Get(), Base.Get(), Blended.Get(), bApplyColorBlendToAlpha, bBlendOnlyOneMip); break;
-						case EBlendType::BT_MULTIPLY: BufferLayerEmbeddedMask<MultiplyChannelMasked, MultiplyChannel, false>(Base.Get(), Base.Get(), Blended.Get(), bApplyColorBlendToAlpha, bBlendOnlyOneMip); break;
-						case EBlendType::BT_BLEND: BufferLayerEmbeddedMask<BlendChannelMasked, BlendChannel, false>(Base.Get(), Base.Get(), Blended.Get(), bApplyColorBlendToAlpha, bBlendOnlyOneMip); break;
-						default: check(false);
-						}
-					}
-                    else if (!bDone)
-                    {
-                        switch (EBlendType(Args.blendType))
-                        {
-						case EBlendType::BT_NORMAL_COMBINE: check(false); break;
-                        case EBlendType::BT_SOFTLIGHT: BufferLayer<SoftLightChannel, false>(Base.Get(), Base.Get(), Blended.Get(), bApplyColorBlendToAlpha, bBlendOnlyOneMip, bUseBlendSourceFromBlendAlpha); break;
-                        case EBlendType::BT_HARDLIGHT: BufferLayer<HardLightChannel, false>(Base.Get(), Base.Get(), Blended.Get(), bApplyColorBlendToAlpha, bBlendOnlyOneMip, bUseBlendSourceFromBlendAlpha); break;
-                        case EBlendType::BT_BURN: BufferLayer<BurnChannel, false>(Base.Get(), Base.Get(), Blended.Get(), bApplyColorBlendToAlpha, bBlendOnlyOneMip, bUseBlendSourceFromBlendAlpha); break;
-                        case EBlendType::BT_DODGE: BufferLayer<DodgeChannel, false>(Base.Get(), Base.Get(), Blended.Get(), bApplyColorBlendToAlpha, bBlendOnlyOneMip, bUseBlendSourceFromBlendAlpha); break;
-                        case EBlendType::BT_SCREEN: BufferLayer<ScreenChannel, false>(Base.Get(), Base.Get(),  Blended.Get(), bApplyColorBlendToAlpha, bBlendOnlyOneMip, bUseBlendSourceFromBlendAlpha); break;
-                        case EBlendType::BT_OVERLAY: BufferLayer<OverlayChannel, false>(Base.Get(), Base.Get(), Blended.Get(), bApplyColorBlendToAlpha, bBlendOnlyOneMip, bUseBlendSourceFromBlendAlpha); break;
-                        case EBlendType::BT_LIGHTEN: BufferLayer<LightenChannel, false>(Base.Get(), Base.Get(), Blended.Get(), bApplyColorBlendToAlpha, bBlendOnlyOneMip, bUseBlendSourceFromBlendAlpha); break;
-                        case EBlendType::BT_MULTIPLY: BufferLayer<MultiplyChannel, false>(Base.Get(), Base.Get(), Blended.Get(), bApplyColorBlendToAlpha, bBlendOnlyOneMip, bUseBlendSourceFromBlendAlpha); break;
-                        case EBlendType::BT_BLEND: BufferLayer<BlendChannel, false>(Base.Get(), Base.Get(), Blended.Get(), bApplyColorBlendToAlpha, bBlendOnlyOneMip, bUseBlendSourceFromBlendAlpha); break;
-                        default: check(false);
-                        }
-                    }
+							bool bUseMaskFromBlendAlpha = static_cast<bool>(Args.bUseMaskFromBlended);
 
-					// Apply the separate blend operation for alpha
-					if (!bDone && !bApplyColorBlendToAlpha && Args.blendTypeAlpha != uint8(EBlendType::BT_NONE) )
-					{
-						// Separate alpha operation ignores the mask.
-						switch (EBlendType(Args.blendTypeAlpha))
+							FBlendFuncType* ColorBlendFunc = SelectBlendFunc(EBlendType(Args.blendType));
+							FBlendFuncType* AlphaBlendFunc = nullptr; //ColorBlendFunc;
+
+							if (!bApplyColorBlendToAlpha)
+							{
+								AlphaBlendFunc = SelectBlendFunc(EBlendType(Args.blendTypeAlpha));
+							}
+
+							constexpr uint8 BlendAlphaSourceChannel = 3;
+
+							ImageLayerBlend(
+									Base.Get(), Base.Get(), Blended.Get(), Mask.Get(), 
+									ColorBlendFunc, AlphaBlendFunc, 
+									LODBegin, LODEnd, 
+									bApplyColorBlendToAlpha, bUseMaskFromBlendAlpha, bUseBlendSourceFromBlendAlpha, BlendAlphaSourceChannel);
+						}
+						else
 						{
-						case EBlendType::BT_SOFTLIGHT: BufferLayerInPlace<SoftLightChannel, false, 1>(Base.Get(), Blended.Get(), bBlendOnlyOneMip, 3, Args.BlendAlphaSourceChannel); break;
-						case EBlendType::BT_HARDLIGHT: BufferLayerInPlace<HardLightChannel, false, 1>(Base.Get(), Blended.Get(), bBlendOnlyOneMip, 3, Args.BlendAlphaSourceChannel); break;
-						case EBlendType::BT_BURN: BufferLayerInPlace<BurnChannel, false, 1>(Base.Get(), Blended.Get(), bBlendOnlyOneMip, 3, Args.BlendAlphaSourceChannel); break;
-						case EBlendType::BT_DODGE: BufferLayerInPlace<DodgeChannel, false, 1>(Base.Get(), Blended.Get(), bBlendOnlyOneMip, 3, Args.BlendAlphaSourceChannel); break;
-						case EBlendType::BT_SCREEN: BufferLayerInPlace<ScreenChannel, false, 1>(Base.Get(), Blended.Get(), bBlendOnlyOneMip, 3, Args.BlendAlphaSourceChannel); break;
-						case EBlendType::BT_OVERLAY: BufferLayerInPlace<OverlayChannel, false, 1>(Base.Get(), Blended.Get(), bBlendOnlyOneMip, 3, Args.BlendAlphaSourceChannel); break;
-						case EBlendType::BT_LIGHTEN: BufferLayerInPlace<LightenChannel, false, 1>(Base.Get(), Blended.Get(), bBlendOnlyOneMip, 3, Args.BlendAlphaSourceChannel); break;
-						case EBlendType::BT_MULTIPLY: BufferLayerInPlace<MultiplyChannel, false, 1>(Base.Get(), Blended.Get(), bBlendOnlyOneMip, 3, Args.BlendAlphaSourceChannel); break;
-						case EBlendType::BT_BLEND: BufferLayerInPlace<BlendChannel, false, 1>(Base.Get(), Blended.Get(), bBlendOnlyOneMip, 3, Args.BlendAlphaSourceChannel); break;
-						default: check(false);
+							ImageLayerCombine(
+									Base.Get(), Base.Get(), Blended.Get(), Mask.Get(),
+									OpImageLayerCombineOps::NormalCombine,
+									LODBegin, LODEnd);
 						}
 					}
-
-					Blended = nullptr;
 				}
 
 				// Are we done?
@@ -3607,11 +3556,14 @@ namespace UE::Mutable::Private
 				{
 					// Request a new layer
 					++CurrentIteration;
-					FScheduledOp ItemCopy = Item;
-					ExecutionIndex Index = GetMemory().GetRangeIndex(Item.ExecutionIndex);
-					Index.SetFromModelRangeIndex(Args.rangeId, CurrentIteration);
-					ItemCopy.ExecutionIndex = GetMemory().GetRangeIndexIndex(Index);
-					AddOp(FScheduledOp(Item.At, Item, 2+CurrentIteration, Item.CustomState), FScheduledOp(Args.blended, ItemCopy), FScheduledOp(Args.mask, ItemCopy));
+					FScheduledOp ItemForNextLayerBlend = Item;
+					ExecutionIndex NextIndex = GetMemory().GetRangeIndex(Item.ExecutionIndex);
+					NextIndex.SetFromModelRangeIndex(Args.rangeId, CurrentIteration);
+					ItemForNextLayerBlend.ExecutionIndex = GetMemory().GetRangeIndexIndex(NextIndex);
+
+					AddOp(FScheduledOp(Item.At, Item, 2+CurrentIteration, Item.CustomState), 
+							FScheduledOp(Args.blended, ItemForNextLayerBlend), 
+							FScheduledOp(Args.mask, ItemForNextLayerBlend));
 
 				}
 
@@ -3629,45 +3581,72 @@ namespace UE::Mutable::Private
 			switch (Item.Stage)
 			{
 			case 0:
+			{
 				if (Args.base && Args.normal)
 				{
-					AddOp(FScheduledOp(Item.At, Item, 1),
-						FScheduledOp(Args.base, Item),
-						FScheduledOp(Args.normal, Item));
+					int32 OutChannel = -1;
+					switch (Args.mode)
+					{
+						case ECompositeImageMode::CIM_NormalRoughnessToRed   : OutChannel = 0; break;
+						case ECompositeImageMode::CIM_NormalRoughnessToGreen : OutChannel = 1; break;
+						case ECompositeImageMode::CIM_NormalRoughnessToBlue  : OutChannel = 2; break;
+						case ECompositeImageMode::CIM_NormalRoughnessToAlpha : OutChannel = 3; break;
+						default: OutChannel = -1; break;
+					}
+
+					if (OutChannel < 0)
+					{
+						AddOp(FScheduledOp(Item.At, Item, 1, (uint32)OutChannel),
+							FScheduledOp(Args.base, Item));
+					}
+					else
+					{
+						AddOp(FScheduledOp(Item.At, Item, 1, (uint32)OutChannel),
+							FScheduledOp(Args.base, Item),
+							FScheduledOp(Args.normal, Item));
+					}
 				}
 				else
 				{
 					StoreImage(Item, nullptr);
 				}
 				break;
-
+			}
 			case 1:
 			{
 				MUTABLE_CPUPROFILER_SCOPE(IM_NORMALCOMPOSITE_1)
 
-				TManagedPtr<const FImage> Base = LoadImage(FCacheAddress(Args.base, Item));
-				TManagedPtr<const FImage> Normal = LoadImage(FCacheAddress(Args.normal, Item));
-
-				if (Normal->GetLODCount() < Base->GetLODCount())
+				int32 OutChannel = (int32)Item.CustomState;
+				if (OutChannel < 0)
+				{
+					StoreImage(Item, LoadImage(FCacheAddress(Args.base, Item)));
+				}
+				else
 				{
-					MUTABLE_CPUPROFILER_SCOPE(ImageNormalComposite_EmergencyFix);
+					TManagedPtr<const FImage> Base = LoadImage(FCacheAddress(Args.base, Item));
+					TManagedPtr<const FImage> Normal = LoadImage(FCacheAddress(Args.normal, Item));
 
-					int32 StartLevel = Normal->GetLODCount() - 1;
-					int32 LevelCount = Base->GetLODCount();
-					
-					TManagedPtr<FImage> NormalFix = CloneOrTakeOver(MoveTemp(Normal));
-					NormalFix->DataStorage.SetNumLODs(LevelCount);
+					if (Normal->GetLODCount() < Base->GetLODCount())
+					{
+						MUTABLE_CPUPROFILER_SCOPE(ImageNormalComposite_EmergencyFix);
 
-					FMipmapGenerationSettings MipSettings{};
-					ImOp.ImageMipmap(Settings.ImageCompressionQuality, NormalFix.Get(), NormalFix.Get(), StartLevel, LevelCount, MipSettings);
+						int32 StartLevel = Normal->GetLODCount() - 1;
+						int32 LevelCount = Base->GetLODCount();
+						
+						TManagedPtr<FImage> NormalFix = CloneOrTakeOver(MoveTemp(Normal));
+						NormalFix->DataStorage.SetNumLODs(LevelCount);
 
-					Normal = NormalFix;
-				}
+						FMipmapGenerationSettings MipSettings{};
+						ImOp.ImageMipmap(Settings.ImageCompressionQuality, NormalFix.Get(), NormalFix.Get(), StartLevel, LevelCount, MipSettings);
 
-                TManagedPtr<FImage> Result = CreateImage(Base->GetSizeX(), Base->GetSizeY(), Base->GetLODCount(), Base->GetFormat(), EInitializationType::NotInitialized);
-				ImageNormalComposite(Result.Get(), Base.Get(), Normal.Get(), Args.mode, Args.power);
+						Normal = NormalFix;
+					}
 
-				StoreImage(Item, Result);
+					TManagedPtr<FImage> Result = CreateImage(Base->GetSizeX(), Base->GetSizeY(), Base->GetLODCount(), Base->GetFormat(), EInitializationType::NotInitialized);
+					ImageNormalComposite(Result.Get(), Base.Get(), Normal.Get(), OutChannel, Args.power);
+
+					StoreImage(Item, Result);
+				}
 				break;
 			}
 
diff --git a/Engine/Plugins/Mutable/Source/MutableRuntime/Private/MuR/CodeRunnerTasked.cpp b/Engine/Plugins/Mutable/Source/MutableRuntime/Private/MuR/CodeRunnerTasked.cpp
--- a/Engine/Plugins/Mutable/Source/MutableRuntime/Private/MuR/CodeRunnerTasked.cpp
+++ b/Engine/Plugins/Mutable/Source/MutableRuntime/Private/MuR/CodeRunnerTasked.cpp
@@ -21,9 +21,9 @@
 #include "MuR/MutableMath.h"
 #include "MuR/MutableTrace.h"
 #include "MuR/OpImageBlend.h"
-#include "MuR/OpImageNormalCombine.h"
 #include "MuR/OpImageSaturate.h"
 #include "MuR/OpImageInvert.h"
+#include "MuR/OpImageLayer.h"
 #include "MuR/Operations.h"
 #include "MuR/Platform.h"
 #include "MuR/Ptr.h"
@@ -626,128 +626,37 @@ namespace UE::Mutable::Private
 		MUTABLE_CPUPROFILER_SCOPE(FImageLayerTask);
 
 		// This runs in a worker thread.
+		bool bOnlyOneMip = Blended->GetLODCount() < Result->GetLODCount();
 
-		bool bOnlyOneMip = false;
-		if (Blended->GetLODCount() < Result->GetLODCount())
-		{
-			bOnlyOneMip = true;
-		}
-
-		bool bDone = false;
-
-		if (!Mask 
-			&& 
-			// Flags have to match exactly for this optimize case. Other flags are not supported.
-			Args.flags == OP::ImageLayerArgs::F_USE_MASK_FROM_BLENDED
-			&&
-			Args.blendType == uint8(EBlendType::BT_BLEND)
-			&&
-			Args.blendTypeAlpha == uint8(EBlendType::BT_LIGHTEN))
-		{
-			MUTABLE_CPUPROFILER_SCOPE(ImageLayerTask_Optimized);
-
-			// This is a frequent critical-path case because of multilayer projectors.
-			bDone = true;
-
-			constexpr bool bUseVectorImplementation = false;	
-			if constexpr (bUseVectorImplementation)
-			{
-				BufferLayerCompositeVector<VectorBlendChannelMasked, VectorLightenChannel, false>(Result.Get(), Blended.Get(), bOnlyOneMip, Args.BlendAlphaSourceChannel);
-			}
-			else
-			{
-				BufferLayerComposite<BlendChannelMasked, LightenChannel, false>(Result.Get(), Blended.Get(), bOnlyOneMip, Args.BlendAlphaSourceChannel);
-			}
-		}
-
-		bool bApplyColorBlendToAlpha = (Args.flags & OP::ImageLayerArgs::F_APPLY_TO_ALPHA) != 0;
-		bool bUseBlendSourceFromBlendAlpha = (Args.flags & OP::ImageLayerArgs::F_BLENDED_RGB_FROM_ALPHA) != 0;
-		bool bUseMaskFromBlendAlpha = (Args.flags & OP::ImageLayerArgs::F_USE_MASK_FROM_BLENDED);
-
-		if (!bDone && Mask)
+		const int32 LODBegin = 0;
+		const int32 LODEnd   = bOnlyOneMip ? 1 : Result->GetLODCount();
+		
+		if (EBlendType(Args.blendType) != EBlendType::BT_NORMAL_COMBINE)
 		{
-			// Not implemented yet
-			check(!bUseBlendSourceFromBlendAlpha);
+			bool bApplyColorBlendToAlpha       = (Args.flags & OP::ImageLayerArgs::F_APPLY_TO_ALPHA        ) != 0;
+			bool bUseBlendSourceFromBlendAlpha = (Args.flags & OP::ImageLayerArgs::F_BLENDED_RGB_FROM_ALPHA) != 0;
+			bool bUseMaskFromBlendAlpha        = (Args.flags & OP::ImageLayerArgs::F_USE_MASK_FROM_BLENDED ) != 0;
 
-			// TODO: in-place
+			FBlendFuncType* ColorBlendFunc = SelectBlendFunc(EBlendType(Args.blendType));
+			FBlendFuncType* AlphaBlendFunc = nullptr; //ColorBlendFunc;
 
-			switch (EBlendType(Args.blendType))
+			if (!bApplyColorBlendToAlpha)
 			{
-			case EBlendType::BT_NORMAL_COMBINE: ImageNormalCombine(Result.Get(), Result.Get(), Mask.Get(), Blended.Get(), bOnlyOneMip); break;
-			case EBlendType::BT_SOFTLIGHT: BufferLayer<SoftLightChannelMasked, SoftLightChannel, true>(Result.Get(), Result.Get(), Mask.Get(), Blended.Get(), bApplyColorBlendToAlpha, bOnlyOneMip); break;
-			case EBlendType::BT_HARDLIGHT: BufferLayer<HardLightChannelMasked, HardLightChannel, true>(Result.Get(), Result.Get(), Mask.Get(), Blended.Get(), bApplyColorBlendToAlpha, bOnlyOneMip); break;
-			case EBlendType::BT_BURN: BufferLayer<BurnChannelMasked, BurnChannel, true>(Result.Get(), Result.Get(), Mask.Get(), Blended.Get(), bApplyColorBlendToAlpha, bOnlyOneMip); break;
-			case EBlendType::BT_DODGE: BufferLayer<DodgeChannelMasked, DodgeChannel, true>(Result.Get(), Result.Get(), Mask.Get(), Blended.Get(), bApplyColorBlendToAlpha, bOnlyOneMip); break;
-			case EBlendType::BT_SCREEN: BufferLayer<ScreenChannelMasked, ScreenChannel, true>(Result.Get(), Result.Get(), Mask.Get(), Blended.Get(), bApplyColorBlendToAlpha, bOnlyOneMip); break;
-			case EBlendType::BT_OVERLAY: BufferLayer<OverlayChannelMasked, OverlayChannel, true>(Result.Get(), Result.Get(), Mask.Get(), Blended.Get(), bApplyColorBlendToAlpha, bOnlyOneMip); break;
-			case EBlendType::BT_LIGHTEN: BufferLayer<LightenChannelMasked, LightenChannel, true>(Result.Get(), Result.Get(), Mask.Get(), Blended.Get(), bApplyColorBlendToAlpha, bOnlyOneMip); break;
-			case EBlendType::BT_MULTIPLY: BufferLayer<MultiplyChannelMasked, MultiplyChannel, true>(Result.Get(), Result.Get(), Mask.Get(), Blended.Get(), bApplyColorBlendToAlpha, bOnlyOneMip); break;
-			case EBlendType::BT_BLEND: BufferLayer<BlendChannelMasked, BlendChannel, true>(Result.Get(), Result.Get(), Mask.Get(), Blended.Get(), bApplyColorBlendToAlpha, bOnlyOneMip); break;
-			case EBlendType::BT_NONE: break;
-			default: check(false);
+				AlphaBlendFunc = SelectBlendFunc(EBlendType(Args.blendTypeAlpha));
 			}
 
+			ImageLayerBlend(
+					Result.Get(), Result.Get(), Blended.Get(), Mask.Get(), 
+					ColorBlendFunc, AlphaBlendFunc, 
+					LODBegin, LODEnd, 
+					bApplyColorBlendToAlpha, bUseMaskFromBlendAlpha, bUseBlendSourceFromBlendAlpha, Args.BlendAlphaSourceChannel);
 		}
-		else if (!bDone && bUseMaskFromBlendAlpha)
-		{
-			// Not implemented yet
-			check(!bUseBlendSourceFromBlendAlpha);
-
-			// Apply blend without to RGB using mask in blended alpha
-			switch (EBlendType(Args.blendType))
-			{
-			case EBlendType::BT_NORMAL_COMBINE: check(false); break;
-			case EBlendType::BT_SOFTLIGHT: BufferLayerEmbeddedMask<SoftLightChannelMasked, SoftLightChannel, false>(Result.Get(), Result.Get(), Blended.Get(), bApplyColorBlendToAlpha, bOnlyOneMip); break;
-			case EBlendType::BT_HARDLIGHT: BufferLayerEmbeddedMask<HardLightChannelMasked, HardLightChannel, false>(Result.Get(), Result.Get(), Blended.Get(), bApplyColorBlendToAlpha, bOnlyOneMip); break;
-			case EBlendType::BT_BURN: BufferLayerEmbeddedMask<BurnChannelMasked, BurnChannel, false>(Result.Get(), Result.Get(), Blended.Get(), bApplyColorBlendToAlpha, bOnlyOneMip); break;
-			case EBlendType::BT_DODGE: BufferLayerEmbeddedMask<DodgeChannelMasked, DodgeChannel, false>(Result.Get(), Result.Get(), Blended.Get(), bApplyColorBlendToAlpha, bOnlyOneMip); break;
-			case EBlendType::BT_SCREEN: BufferLayerEmbeddedMask<ScreenChannelMasked, ScreenChannel, false>(Result.Get(), Result.Get(), Blended.Get(), bApplyColorBlendToAlpha, bOnlyOneMip); break;
-			case EBlendType::BT_OVERLAY: BufferLayerEmbeddedMask<OverlayChannelMasked, OverlayChannel, false>(Result.Get(), Result.Get(), Blended.Get(), bApplyColorBlendToAlpha, bOnlyOneMip); break;
-			case EBlendType::BT_LIGHTEN: BufferLayerEmbeddedMask<LightenChannelMasked, LightenChannel, false>(Result.Get(), Result.Get(), Blended.Get(), bApplyColorBlendToAlpha, bOnlyOneMip); break;
-			case EBlendType::BT_MULTIPLY: BufferLayerEmbeddedMask<MultiplyChannelMasked, MultiplyChannel, false>(Result.Get(), Result.Get(), Blended.Get(), bApplyColorBlendToAlpha, bOnlyOneMip); break;
-			case EBlendType::BT_BLEND: BufferLayerEmbeddedMask<BlendChannelMasked, BlendChannel, false>(Result.Get(), Result.Get(), Blended.Get(), bApplyColorBlendToAlpha, bOnlyOneMip); break;
-			case EBlendType::BT_NONE: break;
-			default: check(false);
-			}
-		}
-		else if (!bDone)
-		{
-			// Apply blend without mask to RGB
-			switch (EBlendType(Args.blendType))
-			{
-			case EBlendType::BT_NORMAL_COMBINE: ImageNormalCombine(Result.Get(), Result.Get(), Blended.Get(), bOnlyOneMip); check(!bUseBlendSourceFromBlendAlpha);  break;
-			case EBlendType::BT_SOFTLIGHT: BufferLayer<SoftLightChannel, true>(Result.Get(), Result.Get(), Blended.Get(), bApplyColorBlendToAlpha, bOnlyOneMip, bUseBlendSourceFromBlendAlpha); break;
-			case EBlendType::BT_HARDLIGHT: BufferLayer<HardLightChannel, true>(Result.Get(), Result.Get(), Blended.Get(), bApplyColorBlendToAlpha, bOnlyOneMip, bUseBlendSourceFromBlendAlpha); break;
-			case EBlendType::BT_BURN: BufferLayer<BurnChannel, true>(Result.Get(), Result.Get(), Blended.Get(), bApplyColorBlendToAlpha, bOnlyOneMip, bUseBlendSourceFromBlendAlpha); break;
-			case EBlendType::BT_DODGE: BufferLayer<DodgeChannel, true>(Result.Get(), Result.Get(), Blended.Get(), bApplyColorBlendToAlpha, bOnlyOneMip, bUseBlendSourceFromBlendAlpha); break;
-			case EBlendType::BT_SCREEN: BufferLayer<ScreenChannel, true>(Result.Get(), Result.Get(), Blended.Get(), bApplyColorBlendToAlpha, bOnlyOneMip, bUseBlendSourceFromBlendAlpha); break;
-			case EBlendType::BT_OVERLAY: BufferLayer<OverlayChannel, true>(Result.Get(), Result.Get(), Blended.Get(), bApplyColorBlendToAlpha, bOnlyOneMip, bUseBlendSourceFromBlendAlpha); break;
-			case EBlendType::BT_LIGHTEN: BufferLayer<LightenChannel, true>(Result.Get(), Result.Get(), Blended.Get(), bApplyColorBlendToAlpha, bOnlyOneMip, bUseBlendSourceFromBlendAlpha); break;
-			case EBlendType::BT_MULTIPLY: BufferLayer<MultiplyChannel, true>(Result.Get(), Result.Get(), Blended.Get(), bApplyColorBlendToAlpha, bOnlyOneMip, bUseBlendSourceFromBlendAlpha); break;
-			case EBlendType::BT_BLEND: BufferLayer<BlendChannel, true>(Result.Get(), Result.Get(), Blended.Get(), bApplyColorBlendToAlpha, bOnlyOneMip, bUseBlendSourceFromBlendAlpha); break;
-			case EBlendType::BT_NONE: break;
-			default: check(false);
-			}
-		}
-
-		// Apply the separate blend operation for alpha
-		if (!bDone && !bApplyColorBlendToAlpha)
+		else
 		{
-			// Separate alpha operation ignores the mask.
-			switch (EBlendType(Args.blendTypeAlpha))
-			{
-			case EBlendType::BT_SOFTLIGHT: BufferLayerInPlace<SoftLightChannel, false, 1>(Result.Get(), Blended.Get(), bOnlyOneMip, 3, Args.BlendAlphaSourceChannel); break;
-			case EBlendType::BT_HARDLIGHT: BufferLayerInPlace<HardLightChannel, false, 1>(Result.Get(), Blended.Get(), bOnlyOneMip, 3, Args.BlendAlphaSourceChannel); break;
-			case EBlendType::BT_BURN: BufferLayerInPlace<BurnChannel, false, 1>(Result.Get(), Blended.Get(), bOnlyOneMip, 3, Args.BlendAlphaSourceChannel); break;
-			case EBlendType::BT_DODGE: BufferLayerInPlace<DodgeChannel, false, 1>(Result.Get(), Blended.Get(), bOnlyOneMip, 3, Args.BlendAlphaSourceChannel); break;
-			case EBlendType::BT_SCREEN: BufferLayerInPlace<ScreenChannel, false, 1>(Result.Get(), Blended.Get(), bOnlyOneMip, 3, Args.BlendAlphaSourceChannel); break;
-			case EBlendType::BT_OVERLAY: BufferLayerInPlace<OverlayChannel, false, 1>(Result.Get(), Blended.Get(), bOnlyOneMip, 3, Args.BlendAlphaSourceChannel); break;
-			case EBlendType::BT_LIGHTEN: BufferLayerInPlace<LightenChannel, false, 1>(Result.Get(), Blended.Get(), bOnlyOneMip, 3, Args.BlendAlphaSourceChannel); break;
-			case EBlendType::BT_MULTIPLY: BufferLayerInPlace<MultiplyChannel, false, 1>(Result.Get(), Blended.Get(), bOnlyOneMip, 3, Args.BlendAlphaSourceChannel); break;
-			case EBlendType::BT_BLEND: BufferLayerInPlace<BlendChannel, false, 1>(Result.Get(), Blended.Get(), bOnlyOneMip, 3, Args.BlendAlphaSourceChannel); break;
-			case EBlendType::BT_NONE: break;
-			default: check(false);
-			}
+			ImageLayerCombine(
+					Result.Get(), Result.Get(), Blended.Get(), Mask.Get(), 
+					OpImageLayerCombineOps::NormalCombine, 
+					LODBegin, LODEnd);
 		}
 
 		if (bOnlyOneMip)
@@ -761,7 +670,6 @@ namespace UE::Mutable::Private
 		Result->Flags &= ~FImage::EImageFlags::IF_HAS_RELEVANCY_MAP;
 	}
 
-
 	bool FImageLayerTask::Complete( CodeRunner* Runner )
 	{
 		// This runs in the Runner thread
@@ -906,120 +814,34 @@ namespace UE::Mutable::Private
 		// This runs in a worker thread
 		MUTABLE_CPUPROFILER_SCOPE(FImageLayerColourTask);
 
-		bool bOnlyOneMip = false;
+		const int32 LODBegin = 0;
+		const int32 LODEnd   = Result->GetLODCount();
 
-		// Does it apply to colour?
-		if (EBlendType(Args.blendType) != EBlendType::BT_NONE)
+		if (EBlendType(Args.blendType) != EBlendType::BT_NORMAL_COMBINE)
 		{
-			// TODO: It could be done "in-place"
-			if (Args.mask && Mask)
-			{			
-				// Not implemented yet
-				check(Args.flags==0);
-
-				// \todo: precalculated tables for softlight
-				switch (EBlendType(Args.blendType))
-				{
-				case EBlendType::BT_NORMAL_COMBINE: ImageNormalCombine(Result.Get(), Result.Get(), Mask.Get(), Color); break;
-				case EBlendType::BT_SOFTLIGHT: BufferLayerColour<SoftLightChannelMasked, SoftLightChannel>(Result.Get(), Result.Get(), Mask.Get(), Color); break;
-				case EBlendType::BT_HARDLIGHT: BufferLayerColour<HardLightChannelMasked, HardLightChannel>(Result.Get(), Result.Get(), Mask.Get(), Color); break;
-				case EBlendType::BT_BURN: BufferLayerColour<BurnChannelMasked, BurnChannel>(Result.Get(), Result.Get(), Mask.Get(), Color); break;
-				case EBlendType::BT_DODGE: BufferLayerColour<DodgeChannelMasked, DodgeChannel>(Result.Get(), Result.Get(), Mask.Get(), Color); break;
-				case EBlendType::BT_SCREEN: BufferLayerColour<ScreenChannelMasked, ScreenChannel>(Result.Get(), Result.Get(), Mask.Get(), Color); break;
-				case EBlendType::BT_OVERLAY: BufferLayerColour<OverlayChannelMasked, OverlayChannel>(Result.Get(), Result.Get(), Mask.Get(), Color); break;
-				case EBlendType::BT_LIGHTEN: BufferLayerColour<LightenChannelMasked, LightenChannel>(Result.Get(), Result.Get(), Mask.Get(), Color); break;
-				case EBlendType::BT_MULTIPLY: BufferLayerColour<MultiplyChannelMasked, MultiplyChannel>(Result.Get(), Result.Get(), Mask.Get(), Color); break;
-				case EBlendType::BT_BLEND: BufferLayerColour<BlendChannelMasked, BlendChannel>(Result.Get(), Result.Get(), Mask.Get(), Color); break;
-				default: check(false);
-				}
+			FBlendFuncType* ColorBlendFunc = SelectBlendFunc(EBlendType(Args.blendType));
+			FBlendFuncType* AlphaBlendFunc = ColorBlendFunc;
+			
+			bool bApplyColorBlendToAlpha = (Args.flags & OP::ImageLayerArgs::F_APPLY_TO_ALPHA) != 0;
 
-			}
-			else
+			if (!bApplyColorBlendToAlpha)
 			{
-				if (Args.flags & OP::ImageLayerArgs::FLAGS::F_BASE_RGB_FROM_ALPHA)
-				{
-					switch (EBlendType(Args.blendType))
-					{
-					case EBlendType::BT_NORMAL_COMBINE: check(false); break;
-					case EBlendType::BT_SOFTLIGHT: BufferLayerColourFromAlpha<SoftLightChannel>(Result.Get(), Result.Get(), Color); break;
-					case EBlendType::BT_HARDLIGHT: BufferLayerColourFromAlpha<HardLightChannel>(Result.Get(), Result.Get(), Color); break;
-					case EBlendType::BT_BURN: BufferLayerColourFromAlpha<BurnChannel>(Result.Get(), Result.Get(), Color); break;
-					case EBlendType::BT_DODGE: BufferLayerColourFromAlpha<DodgeChannel>(Result.Get(), Result.Get(), Color); break;
-					case EBlendType::BT_SCREEN: BufferLayerColourFromAlpha<ScreenChannel>(Result.Get(), Result.Get(), Color); break;
-					case EBlendType::BT_OVERLAY: BufferLayerColourFromAlpha<OverlayChannel>(Result.Get(), Result.Get(), Color); break;
-					case EBlendType::BT_LIGHTEN: BufferLayerColourFromAlpha<LightenChannel>(Result.Get(), Result.Get(), Color); break;
-					case EBlendType::BT_MULTIPLY: BufferLayerColourFromAlpha<MultiplyChannel>(Result.Get(), Result.Get(), Color); break;
-					case EBlendType::BT_BLEND: check(false); break;
-					default: check(false);
-					}
-				}
-				else
-				{
-					switch (EBlendType(Args.blendType))
-					{
-					case EBlendType::BT_NORMAL_COMBINE: ImageNormalCombine(Result.Get(), Result.Get(), Color); break;
-					case EBlendType::BT_SOFTLIGHT: BufferLayerColour<SoftLightChannel>(Result.Get(), Result.Get(), Color); break;
-					case EBlendType::BT_HARDLIGHT: BufferLayerColour<HardLightChannel>(Result.Get(), Result.Get(), Color); break;
-					case EBlendType::BT_BURN: BufferLayerColour<BurnChannel>(Result.Get(), Result.Get(), Color); break;
-					case EBlendType::BT_DODGE: BufferLayerColour<DodgeChannel>(Result.Get(), Result.Get(), Color); break;
-					case EBlendType::BT_SCREEN: BufferLayerColour<ScreenChannel>(Result.Get(), Result.Get(), Color); break;
-					case EBlendType::BT_OVERLAY: BufferLayerColour<OverlayChannel>(Result.Get(), Result.Get(), Color); break;
-					case EBlendType::BT_LIGHTEN: BufferLayerColour<LightenChannel>(Result.Get(), Result.Get(), Color); break;
-					case EBlendType::BT_MULTIPLY: BufferLayerColour<MultiplyChannel>(Result.Get(), Result.Get(), Color); break;
-					case EBlendType::BT_BLEND: 
-					{
-						// In this case we know it is already an uncompressed image, and we won't need additional allocations;
-						FImageOperator ImOp = FImageOperator::GetDefault(FImageOperator::FImagePixelFormatFunc());
-						ImOp.FillColor( Result.Get(), Color); 
-						break;
-					}
-					default: check(false);
-					}
-				}
+				AlphaBlendFunc = SelectBlendFunc(EBlendType(Args.blendTypeAlpha));
 			}
-		}
 
-		// Does it apply to alpha?
-		if (EBlendType(Args.blendTypeAlpha) != EBlendType::BT_NONE)
+			ImageLayerBlendConstant(
+					Result.Get(), Result.Get(), Mask.Get(), Color, 
+					ColorBlendFunc, AlphaBlendFunc, 
+					LODBegin, LODEnd, 
+					bApplyColorBlendToAlpha, false, false, Args.BlendAlphaSourceChannel);
+		}
+		else
 		{
-			if (Args.mask && Mask)
-			{
-				// \todo: precalculated tables for softlight
-				switch (EBlendType(Args.blendTypeAlpha))
-				{
-				case EBlendType::BT_NORMAL_COMBINE: check(false); break;
-				case EBlendType::BT_SOFTLIGHT: BufferLayerColourInPlace<SoftLightChannelMasked, SoftLightChannel, 1>(Result.Get(), Mask.Get(), Color, bOnlyOneMip, 3, Args.BlendAlphaSourceChannel); break;
-				case EBlendType::BT_HARDLIGHT: BufferLayerColourInPlace<HardLightChannelMasked, HardLightChannel, 1>(Result.Get(), Mask.Get(), Color, bOnlyOneMip, 3, Args.BlendAlphaSourceChannel); break;
-				case EBlendType::BT_BURN: BufferLayerColourInPlace<BurnChannelMasked, BurnChannel, 1>(Result.Get(), Mask.Get(), Color, bOnlyOneMip, 3, Args.BlendAlphaSourceChannel); break;
-				case EBlendType::BT_DODGE: BufferLayerColourInPlace<DodgeChannelMasked, DodgeChannel, 1>(Result.Get(), Mask.Get(), Color, bOnlyOneMip, 3, Args.BlendAlphaSourceChannel); break;
-				case EBlendType::BT_SCREEN: BufferLayerColourInPlace<ScreenChannelMasked, ScreenChannel, 1>(Result.Get(), Mask.Get(), Color, bOnlyOneMip, 3, Args.BlendAlphaSourceChannel); break;
-				case EBlendType::BT_OVERLAY: BufferLayerColourInPlace<OverlayChannelMasked, OverlayChannel, 1>(Result.Get(), Mask.Get(), Color, bOnlyOneMip, 3, Args.BlendAlphaSourceChannel); break;
-				case EBlendType::BT_LIGHTEN: BufferLayerColourInPlace<LightenChannelMasked, LightenChannel, 1>(Result.Get(), Mask.Get(), Color, bOnlyOneMip, 3, Args.BlendAlphaSourceChannel); break;
-				case EBlendType::BT_MULTIPLY: BufferLayerColourInPlace<MultiplyChannelMasked, MultiplyChannel, 1>(Result.Get(), Mask.Get(), Color, bOnlyOneMip, 3, Args.BlendAlphaSourceChannel); break;
-				case EBlendType::BT_BLEND: BufferLayerColourInPlace<BlendChannelMasked, BlendChannel, 1>(Result.Get(), Mask.Get(), Color, bOnlyOneMip, 3, Args.BlendAlphaSourceChannel); break;
-				default: check(false);
-				}
-
-			}
-			else
-			{
-				switch (EBlendType(Args.blendTypeAlpha))
-				{
-				case EBlendType::BT_NORMAL_COMBINE: check(false); break;
-				case EBlendType::BT_SOFTLIGHT: BufferLayerColourInPlace<SoftLightChannel, 1>(Result.Get(), Color, bOnlyOneMip, 3, Args.BlendAlphaSourceChannel); break;
-				case EBlendType::BT_HARDLIGHT: BufferLayerColourInPlace<HardLightChannel, 1>(Result.Get(), Color, bOnlyOneMip, 3, Args.BlendAlphaSourceChannel); break;
-				case EBlendType::BT_BURN: BufferLayerColourInPlace<BurnChannel, 1>(Result.Get(), Color, bOnlyOneMip, 3, Args.BlendAlphaSourceChannel); break;
-				case EBlendType::BT_DODGE: BufferLayerColourInPlace<DodgeChannel, 1>(Result.Get(), Color, bOnlyOneMip, 3, Args.BlendAlphaSourceChannel); break;
-				case EBlendType::BT_SCREEN: BufferLayerColourInPlace<ScreenChannel, 1>(Result.Get(), Color, bOnlyOneMip, 3, Args.BlendAlphaSourceChannel); break;
-				case EBlendType::BT_OVERLAY: BufferLayerColourInPlace<OverlayChannel, 1>(Result.Get(), Color, bOnlyOneMip, 3, Args.BlendAlphaSourceChannel); break;
-				case EBlendType::BT_LIGHTEN: BufferLayerColourInPlace<LightenChannel, 1>(Result.Get(), Color, bOnlyOneMip, 3, Args.BlendAlphaSourceChannel); break;
-				case EBlendType::BT_MULTIPLY: BufferLayerColourInPlace<MultiplyChannel, 1>(Result.Get(), Color, bOnlyOneMip, 3, Args.BlendAlphaSourceChannel); break;
-				case EBlendType::BT_BLEND: BufferLayerColourInPlace<BlendChannel, 1>(Result.Get(), Color, bOnlyOneMip, 3, Args.BlendAlphaSourceChannel); break;
-				default: check(false);
-				}
-			}
+			ImageLayerCombineConstant(
+					Result.Get(), Result.Get(), Mask.Get(), Color,
+					OpImageLayerCombineOps::NormalCombine, 
+					LODBegin, LODEnd);
 		}
-
 		// Reset relevancy map.
 		Result->Flags &= ~FImage::EImageFlags::IF_HAS_RELEVANCY_MAP;
 	}
diff --git a/Engine/Plugins/Mutable/Source/MutableRuntime/Private/MuR/ImageComputeUtils.cpp b/Engine/Plugins/Mutable/Source/MutableRuntime/Private/MuR/ImageComputeUtils.cpp
new file mode 100644
--- /dev/null
+++ b/Engine/Plugins/Mutable/Source/MutableRuntime/Private/MuR/ImageComputeUtils.cpp
@@ -0,0 +1,596 @@
+// Copyright Epic Games, Inc. All Rights Reserved.
+
+#include "MuR/ImageComputeUtils.h"
+
+#include "Math/VectorRegister.h"
+
+namespace UE::Mutable::Private
+{
+namespace ImageComputeUtils
+{
+
+	struct alignas(4) FVec4Mask_U8
+	{
+		uint8 X;
+		uint8 Y;
+		uint8 Z;
+		uint8 W;
+	};
+
+#if PLATFORM_ENABLE_VECTORINTRINSICS && !PLATFORM_ENABLE_VECTORINTRINSICS_NEON
+	using VectorRegisterInt8x16 = __m128i;
+
+  	FORCEINLINE VectorRegisterInt8x16 MakeShuffle_Int8x16(uint8 V0 = 0, uint8 V1 = 1, uint8 V2 = 2, uint8 V3 = 3, uint8 V4 = 4, uint8 V5 = 5, uint8 V6 = 6, uint8 V7 = 7, uint8 V8 = 8, uint8 V9 = 9, uint8 V10 = 10, uint8 V11 = 11, uint8 V12 = 12, uint8 V13 = 13, uint8 V14 = 14, uint8 V15 = 15)
+	{
+		return _mm_setr_epi8(V0, V1, V2, V3, V4, V5, V6, V7, V8, V9, V10, V11, V12, V13, V14, V15);
+	}
+
+	FORCEINLINE void StoreUnaligned_Int8x16(uint8* Ptr, const VectorRegisterInt8x16 Reg)
+	{
+		_mm_storeu_si128((__m128i*)Ptr, Reg);
+	}
+
+	FORCEINLINE VectorRegisterInt8x16 LoadUnaligned_Int8x16(const uint8* Ptr)
+	{
+		return _mm_loadu_si128((__m128i*)Ptr);
+	}
+
+	FORCEINLINE VectorRegisterInt8x16 Shuffle_Int8x16(const VectorRegisterInt8x16 V, const VectorRegisterInt8x16 Shuffle)
+	{
+		return _mm_shuffle_epi8(V, Shuffle);
+	}
+
+	FORCEINLINE VectorRegisterInt8x16 MakeMask_Int8x16(FVec4Mask_U8 Mask)
+	{
+		return _mm_setr_epi8(Mask.X, Mask.Y, Mask.Z, Mask.W, Mask.X, Mask.Y, Mask.Z, Mask.W, Mask.X, Mask.Y, Mask.Z, Mask.W, Mask.X, Mask.Y, Mask.Z, Mask.W);	
+	}
+
+	FORCEINLINE VectorRegisterInt8x16 Select_Int8x16(const VectorRegisterInt8x16 A, const VectorRegisterInt8x16 B, const VectorRegisterInt8x16 Mask)
+	{
+		return _mm_blendv_epi8(B, A, Mask);
+	}
+
+	FORCEINLINE VectorRegisterInt8x16 Unorm255Lerp_Int8x16(const VectorRegisterInt8x16 A, const VectorRegisterInt8x16 B, const VectorRegisterInt8x16 Factor)
+	{
+		__m128i OneMinusFactor = _mm_xor_si128(Factor, _mm_set1_epi8((uint8)0xFF)); 
+		__m128i OneMinusFactorLo = _mm_cvtepu8_epi16(OneMinusFactor);
+		__m128i OneMinusFactorHi = _mm_cvtepu8_epi16(_mm_srli_si128(OneMinusFactor, 8));
+
+		__m128i FactorLo = _mm_cvtepu8_epi16(Factor);
+		__m128i FactorHi = _mm_cvtepu8_epi16(_mm_srli_si128(Factor, 8));
+
+		__m128i ALo = _mm_mullo_epi16(_mm_cvtepu8_epi16(A), OneMinusFactorLo);
+		__m128i AHi = _mm_mullo_epi16(_mm_cvtepu8_epi16(_mm_srli_si128(A, 8)), OneMinusFactorHi);
+
+		__m128i BLo = _mm_mullo_epi16(_mm_cvtepu8_epi16(B), FactorLo);
+		__m128i BHi = _mm_mullo_epi16(_mm_cvtepu8_epi16(_mm_srli_si128(B, 8)), FactorHi);
+
+		__m128i Int128 = _mm_set1_epi16(128);
+		__m128i Int257 = _mm_set1_epi16(257);
+
+		__m128i ResultLo = _mm_mulhi_epu16(_mm_add_epi16(_mm_add_epi16(ALo, BLo), Int128), Int257);
+		__m128i ResultHi = _mm_mulhi_epu16(_mm_add_epi16(_mm_add_epi16(AHi, BHi), Int128), Int257);
+
+		return _mm_packus_epi16(ResultLo, ResultHi);
+	}
+
+#endif
+
+#if PLATFORM_ENABLE_VECTORINTRINSICS && PLATFORM_ENABLE_VECTORINTRINSICS_NEON
+	using VectorRegisterInt8x16 = uint8x16_t;
+
+  	FORCEINLINE VectorRegisterInt8x16 MakeShuffle_Int8x16(uint8 V0 = 0, uint8 V1 = 1, uint8 V2 = 2, uint8 V3 = 3, uint8 V4 = 4, uint8 V5 = 5, uint8 V6 = 6, uint8 V7 = 7, uint8 V8 = 8, uint8 V9 = 9, uint8 V10 = 10, uint8 V11 = 11, uint8 V12 = 12, uint8 V13 = 13, uint8 V14 = 14, uint8 V15 = 15)
+	{
+		int8x16_t Value {V0, V1, V2, V3, V4, V5, V6, V7, V8, V9, V10, V11, V12, V13, V14, V15};
+		return Value;
+	}
+
+	FORCEINLINE void StoreUnaligned_Int8x16(uint8* Ptr, VectorRegisterInt8x16 Reg)
+	{
+		vst1q_u8(Ptr, Reg);
+	}
+
+	FORCEINLINE VectorRegisterInt8x16 LoadUnaligned_Int8x16(const uint8* Ptr)
+	{
+		return vld1q_u8(Ptr);
+	}
+
+	FORCEINLINE VectorRegisterInt8x16 Shuffle_Int8x16(VectorRegisterInt8x16 V, VectorRegisterInt8x16 Shuffle)
+	{
+		return vqtbl1q_u8(V, Shuffle);
+	}
+
+	FORCEINLINE VectorRegisterInt8x16 MakeMask_Int8x16(FVec4Mask_U8 Mask)
+	{
+		return uint8x16_t { Mask.X, Mask.Y, Mask.Z, Mask.W, Mask.X, Mask.Y, Mask.Z, Mask.W, Mask.X, Mask.Y, Mask.Z, Mask.W, Mask.X, Mask.Y, Mask.Z, Mask.W };
+	}
+
+	FORCEINLINE VectorRegisterInt8x16 Select_Int8x16(VectorRegisterInt8x16 A, VectorRegisterInt8x16 B, VectorRegisterInt8x16 Mask)
+	{
+		return vbslq_u8(Mask, A, B);
+	}
+
+	FORCEINLINE VectorRegisterInt8x16 Unorm255Lerp_Int8x16(VectorRegisterInt8x16 A, VectorRegisterInt8x16 B, VectorRegisterInt8x16 Factor)
+	{
+		uint8x16_t OneMinusFactor = veorq_u8(Factor, vdupq_n_u8(0xFF)); 
+
+		uint16x8_t BTimesFactorLo = vmull_u8(vget_low_u8(B), vget_low_u8(Factor));
+		uint16x8_t BTimesFactorHi = vmull_u8(vget_high_u8(B), vget_high_u8(Factor));
+
+		uint16x8_t AddLo = vmlal_u8(BTimesFactorLo, vget_low_u8(A), vget_low_u8(OneMinusFactor));
+		uint16x8_t AddHi = vmlal_u8(BTimesFactorHi, vget_high_u8(A), vget_high_u8(OneMinusFactor));
+
+		uint8x8_t ResultLo = vqrshrn_n_u16(vrsraq_n_u16(AddLo, AddLo, 8), 8);
+		uint8x8_t ResultHi = vqrshrn_n_u16(vrsraq_n_u16(AddHi, AddHi, 8), 8);
+
+		return vcombine_u8(ResultLo, ResultHi);
+	}
+#endif
+
+
+#ifndef UE_MUTABLE_IMAGE_COMPUTE_UTILS_USE_VECTOR_IMPL
+	#define UE_MUTABLE_IMAGE_COMPUTE_UTILS_USE_VECTOR_IMPL (PLATFORM_ENABLE_VECTORINTRINSICS || PLATFORM_ENABLE_VECTORINTRINSICS_NEON)
+#endif
+
+#if UE_MUTABLE_IMAGE_COMPUTE_UTILS_USE_VECTOR_IMPL
+	
+	struct FConvertDescriptor
+	{
+		uint32 SrcBytesPerIter = 0;
+		uint32 DstBytesPerIter = 0;
+
+		VectorRegisterInt8x16 Shuffle;
+	};
+
+	// Applies the shuffle described by ConvDesc, it is the responsibility of the calling code to ensure NumIters
+	// does not overflow Dest or Src with the values in ConvDesc.
+	UE_REWRITE void VecConvertImpl(uint8* Dest, const uint8* Src, int32 NumIters, const FConvertDescriptor ConvDesc)
+	{
+		for (int32 I = 0; I < NumIters; ++I)
+		{
+			VectorRegisterInt8x16 SrcRegister = LoadUnaligned_Int8x16(Src);
+			VectorRegisterInt8x16 DstRegister = Shuffle_Int8x16(SrcRegister, ConvDesc.Shuffle);
+			StoreUnaligned_Int8x16(Dest, DstRegister);
+
+			Src  += ConvDesc.SrcBytesPerIter;
+			Dest += ConvDesc.DstBytesPerIter;
+		}
+	}
+
+#endif
+
+	void Vec4ToVec1_U8(uint8* Dest, const uint8* Src, int32 NumElems)
+	{
+		int32 ElemIndex = 0;
+
+#if UE_MUTABLE_IMAGE_COMPUTE_UTILS_USE_VECTOR_IMPL
+		FConvertDescriptor ConvDesc
+		{
+			.SrcBytesPerIter = 16,
+			.DstBytesPerIter = 4,
+
+			.Shuffle = MakeShuffle_Int8x16(0, 4, 8, 12),
+		};
+
+		constexpr int32 ElemsPerIter = 4;
+
+		int32 NumIters = FMath::Max<int32>(0, NumElems/ElemsPerIter - 3);
+		VecConvertImpl(Dest, Src, NumIters, ConvDesc);
+
+		ElemIndex = NumIters*ElemsPerIter;
+#endif
+
+		for (; ElemIndex < NumElems; ++ElemIndex)
+		{
+			Dest[ElemIndex] = Src[ElemIndex*4 + 0];
+		}
+	}
+
+	void Vec4ToVec3_U8(uint8* Dest, const uint8* Src, int32 NumElems)
+	{	
+		int32 ElemIndex = 0;
+
+#if UE_MUTABLE_IMAGE_COMPUTE_UTILS_USE_VECTOR_IMPL
+		FConvertDescriptor ConvDesc
+		{
+			.SrcBytesPerIter = 16,	
+			.DstBytesPerIter = 12,
+
+			.Shuffle = MakeShuffle_Int8x16(0, 1, 2, 4, 5, 6, 8, 9, 10, 12, 13, 14),
+		};
+
+		constexpr int32 ElemsPerIter = 4;
+
+		int32 NumIters = FMath::Max<int32>(0, NumElems/ElemsPerIter - 1);
+		VecConvertImpl(Dest, Src, NumIters, ConvDesc);
+		
+		ElemIndex = NumIters*ElemsPerIter;
+#endif
+
+		for (; ElemIndex < NumElems; ++ElemIndex)
+		{
+			Dest[ElemIndex*3 + 0] = Src[ElemIndex*4 + 0];
+			Dest[ElemIndex*3 + 1] = Src[ElemIndex*4 + 1];
+			Dest[ElemIndex*3 + 2] = Src[ElemIndex*4 + 2];
+		}
+	}
+
+	void Vec3ToVec1_U8(uint8* Dest, const uint8* Src, int32 NumElems)
+	{
+		int32 ElemIndex = 0;
+
+#if UE_MUTABLE_IMAGE_COMPUTE_UTILS_USE_VECTOR_IMPL
+		FConvertDescriptor ConvDesc
+		{
+			.SrcBytesPerIter = 15,	
+			.DstBytesPerIter = 5,
+
+			.Shuffle = MakeShuffle_Int8x16(0, 3, 6, 9, 12),
+		};
+
+		constexpr int32 ElemsPerIter = 5;
+
+		int32 NumIters = FMath::Max<int32>(0, NumElems/ElemsPerIter - 3);
+		VecConvertImpl(Dest, Src, NumIters, ConvDesc);
+
+		ElemIndex = NumIters*ElemsPerIter;
+#endif
+
+		for (; ElemIndex < NumElems; ++ElemIndex)
+		{
+			Dest[ElemIndex] = Src[ElemIndex*3 + 0];
+		}
+	}
+
+	void Vec3ToVec4_U8(uint8* Dest, const uint8* Src, int32 NumElems)
+	{
+		int32 ElemIndex = 0;
+
+#if UE_MUTABLE_IMAGE_COMPUTE_UTILS_USE_VECTOR_IMPL
+		FConvertDescriptor ConvDesc
+		{
+			.SrcBytesPerIter = 12,	
+			.DstBytesPerIter = 16,
+
+			.Shuffle = MakeShuffle_Int8x16(0, 1, 2, 0x80, 3, 4, 5, 0x80, 6, 7, 8, 0x80, 9, 10, 11, 0x80),
+		};
+
+		constexpr int32 ElemsPerIter = 4;
+
+		int32 NumIters = FMath::Max<int32>(0, NumElems/ElemsPerIter - 1);
+		VecConvertImpl(Dest, Src, NumIters, ConvDesc);
+		
+		ElemIndex = NumIters*ElemsPerIter;
+#endif
+
+		for (; ElemIndex < NumElems; ++ElemIndex)
+		{
+			Dest[ElemIndex*4 + 0] = Src[ElemIndex*3 + 0];
+			Dest[ElemIndex*4 + 1] = Src[ElemIndex*3 + 1];
+			Dest[ElemIndex*4 + 2] = Src[ElemIndex*3 + 2];
+			Dest[ElemIndex*4 + 3] = 0;
+		}
+	}
+	
+	void Vec1ToVec3_U8(uint8* Dest, const uint8* Src, int32 NumElems)
+	{
+		int32 ElemIndex = 0;
+
+#if UE_MUTABLE_IMAGE_COMPUTE_UTILS_USE_VECTOR_IMPL
+		FConvertDescriptor ConvDesc
+		{
+			.SrcBytesPerIter = 5,	
+			.DstBytesPerIter = 15,
+
+			.Shuffle = MakeShuffle_Int8x16(0, 0, 0, 1, 1, 1, 2, 2, 2, 3, 3, 3, 4, 4, 4),
+		};
+
+		constexpr int32 ElemsPerIter = 5;
+
+		int32 NumIters = FMath::Max<int32>(0, NumElems/ElemsPerIter - 3);
+		VecConvertImpl(Dest, Src, NumIters, ConvDesc);
+		
+		ElemIndex = NumIters*ElemsPerIter;
+#endif
+
+		for (; ElemIndex < NumElems; ++ElemIndex)
+		{
+			Dest[ElemIndex*3 + 0] = Src[ElemIndex];
+			Dest[ElemIndex*3 + 1] = Src[ElemIndex];
+			Dest[ElemIndex*3 + 2] = Src[ElemIndex];
+		}
+	}
+
+	void Vec1ToVec4_U8(uint8* Dest, const uint8* Src, int32 NumElems)
+	{
+		int32 ElemIndex = 0;
+
+#if UE_MUTABLE_IMAGE_COMPUTE_UTILS_USE_VECTOR_IMPL
+		FConvertDescriptor ConvDesc
+		{
+			.SrcBytesPerIter = 4,	
+			.DstBytesPerIter = 16,
+
+			.Shuffle = MakeShuffle_Int8x16(0, 0, 0, 0, 1, 1, 1, 1, 2, 2, 2, 2, 3, 3, 3, 3),
+		};
+
+		constexpr int32 ElemsPerIter = 4;
+
+		int32 NumIters = FMath::Max<int32>(0, NumElems/ElemsPerIter - 3);
+		VecConvertImpl(Dest, Src, NumIters, ConvDesc);
+
+		ElemIndex = NumIters*ElemsPerIter;
+#endif
+
+		for (; ElemIndex < NumElems; ++ElemIndex)
+		{
+			Dest[ElemIndex*4 + 0] = Src[ElemIndex];
+			Dest[ElemIndex*4 + 1] = Src[ElemIndex];
+			Dest[ElemIndex*4 + 2] = Src[ElemIndex];
+			Dest[ElemIndex*4 + 3] = Src[ElemIndex];
+		}
+	}
+
+	void Vec4SwizzleZYXW_U8(uint8* Dest, const uint8* Src, int32 NumElems)
+	{
+		int32 ElemIndex = 0;
+
+#if UE_MUTABLE_IMAGE_COMPUTE_UTILS_USE_VECTOR_IMPL
+		FConvertDescriptor ConvDesc
+		{
+			.SrcBytesPerIter = 16,	
+			.DstBytesPerIter = 16,
+
+			.Shuffle = MakeShuffle_Int8x16(2, 1, 0, 3, 6, 5, 4, 7, 10, 9, 8, 11, 14, 13, 12, 15),
+		};
+
+		constexpr int32 ElemsPerIter = 4;
+
+		int32 NumIters = NumElems/ElemsPerIter;
+		VecConvertImpl(Dest, Src, NumIters, ConvDesc);
+		
+		ElemIndex = NumIters*ElemsPerIter; 
+#endif
+
+		for (; ElemIndex < NumElems; ++ElemIndex)
+		{
+			Dest[ElemIndex*4 + 0] = Src[ElemIndex*4 + 2]; 
+			Dest[ElemIndex*4 + 1] = Src[ElemIndex*4 + 1]; 
+			Dest[ElemIndex*4 + 2] = Src[ElemIndex*4 + 0]; 
+			Dest[ElemIndex*4 + 3] = Src[ElemIndex*4 + 3]; 
+		}
+	}
+
+	void Vec4SwizzleWZYXAndConvertToVec1_U8(uint8* Dest, const uint8* Src, int32 NumElems)
+	{
+		int32 ElemIndex = 0;
+
+#if UE_MUTABLE_IMAGE_COMPUTE_UTILS_USE_VECTOR_IMPL
+		FConvertDescriptor ConvDesc
+		{
+			.SrcBytesPerIter = 16,	
+			.DstBytesPerIter = 4,
+
+			.Shuffle = MakeShuffle_Int8x16(3, 7, 11, 15),
+		};
+
+		constexpr int32 ElemsPerIter = 4;
+
+		int32 NumIters = FMath::Max<int32>(0, NumElems/ElemsPerIter - 3); 
+		VecConvertImpl(Dest, Src, NumIters, ConvDesc);
+
+		ElemIndex = NumIters*ElemsPerIter;
+#endif
+
+		for (; ElemIndex < NumElems; ++ElemIndex)
+		{
+			Dest[ElemIndex] = Src[ElemIndex*4 + 3];
+		}
+	}
+
+	void Vec4SwizzleXXXX_U8(uint8* Dest, const uint8* Src, int32 NumElems)
+	{
+		int32 ElemIndex = 0;
+
+#if UE_MUTABLE_IMAGE_COMPUTE_UTILS_USE_VECTOR_IMPL
+		FConvertDescriptor ConvDesc
+		{
+			.SrcBytesPerIter = 16,	
+			.DstBytesPerIter = 16,
+
+			.Shuffle = MakeShuffle_Int8x16(0,0,0,0, 4,4,4,4, 8,8,8,8, 12,12,12,12),
+		};
+
+		constexpr int32 ElemsPerIter = 4;
+
+		int32 NumIters = NumElems/ElemsPerIter; 
+		VecConvertImpl(Dest, Src, NumIters, ConvDesc);
+
+		ElemIndex = NumIters*ElemsPerIter;
+#endif
+
+		for (; ElemIndex < NumElems; ++ElemIndex)
+		{
+			Dest[ElemIndex*4 + 0] = Src[ElemIndex*4 + 0];
+			Dest[ElemIndex*4 + 1] = Src[ElemIndex*4 + 0];
+			Dest[ElemIndex*4 + 2] = Src[ElemIndex*4 + 0];
+			Dest[ElemIndex*4 + 3] = Src[ElemIndex*4 + 0];
+		}
+	}
+
+	void Vec4SwizzleWWWW_U8(uint8* Dest, const uint8* Src, int32 NumElems)
+	{
+		int32 ElemIndex = 0;
+
+#if UE_MUTABLE_IMAGE_COMPUTE_UTILS_USE_VECTOR_IMPL
+		FConvertDescriptor ConvDesc
+		{
+			.SrcBytesPerIter = 16,	
+			.DstBytesPerIter = 16,
+
+			.Shuffle = MakeShuffle_Int8x16(3,3,3,3, 7,7,7,7, 11,11,11,11, 15,15,15,15),
+		};
+
+		constexpr int32 ElemsPerIter = 4;
+
+		int32 NumIters = NumElems/ElemsPerIter; 
+		VecConvertImpl(Dest, Src, NumIters, ConvDesc);
+
+		ElemIndex = NumIters*ElemsPerIter;
+#endif
+
+		for (; ElemIndex < NumElems; ++ElemIndex)
+		{
+			Dest[ElemIndex*4 + 0] = Src[ElemIndex*4 + 3];
+			Dest[ElemIndex*4 + 1] = Src[ElemIndex*4 + 3];
+			Dest[ElemIndex*4 + 2] = Src[ElemIndex*4 + 3];
+			Dest[ElemIndex*4 + 3] = Src[ElemIndex*4 + 3];
+		}
+	}
+
+	void Vec4SelectMask_U8(uint8* Dest, const uint8* A, const uint8* B, FVec4Mask_U8 Mask, int32 NumElems)
+	{
+		int32 ElemIndex = 0;
+
+#if UE_MUTABLE_IMAGE_COMPUTE_UTILS_USE_VECTOR_IMPL
+		VectorRegisterInt8x16 VecMask = MakeMask_Int8x16(Mask);
+		
+		constexpr int32 ElemsPerIter = 4;
+		int32 NumIters = NumElems / ElemsPerIter;
+
+		for (int32 Iter = 0; Iter < NumIters; ++Iter)
+		{
+			VectorRegisterInt8x16 VecA = LoadUnaligned_Int8x16(A + Iter*16);	
+			VectorRegisterInt8x16 VecB = LoadUnaligned_Int8x16(B + Iter*16);	
+			VectorRegisterInt8x16 Result = Select_Int8x16(VecA, VecB, VecMask); 
+
+			StoreUnaligned_Int8x16(Dest + Iter*16, Result);
+		}
+
+		ElemIndex = NumIters*ElemsPerIter; 
+#endif
+
+		for (; ElemIndex < NumElems; ++ElemIndex)
+		{
+			Dest[ElemIndex*4 + 0] = A[ElemIndex*4 + 0] & Mask.X;
+			Dest[ElemIndex*4 + 1] = A[ElemIndex*4 + 1] & Mask.Y;
+			Dest[ElemIndex*4 + 2] = A[ElemIndex*4 + 2] & Mask.Z;
+			Dest[ElemIndex*4 + 3] = A[ElemIndex*4 + 3] & Mask.W;
+			
+			Dest[ElemIndex*4 + 0] |= B[ElemIndex*4 + 0] & ~Mask.X;
+			Dest[ElemIndex*4 + 1] |= B[ElemIndex*4 + 1] & ~Mask.Y;
+			Dest[ElemIndex*4 + 2] |= B[ElemIndex*4 + 2] & ~Mask.Z;
+			Dest[ElemIndex*4 + 3] |= B[ElemIndex*4 + 3] & ~Mask.W;
+		}
+	}
+
+	void Vec4SelectX_U8(uint8* Dest, const uint8* A, const uint8* B, int32 NumElems)
+	{
+		constexpr FVec4Mask_U8 Mask = FVec4Mask_U8{0x00, 0xFF, 0xFF, 0xFF};
+		Vec4SelectMask_U8(Dest, A, B, Mask, NumElems);
+	}
+
+	void Vec4SelectY_U8(uint8* Dest, const uint8* A, const uint8* B, int32 NumElems)
+	{
+		constexpr FVec4Mask_U8 Mask = FVec4Mask_U8{0xFF, 0x00, 0xFF, 0xFF};	
+		Vec4SelectMask_U8(Dest, A, B, Mask, NumElems);
+	}
+
+	void Vec4SelectZ_U8(uint8* Dest, const uint8* A, const uint8* B, int32 NumElems)
+	{
+		constexpr FVec4Mask_U8 Mask = FVec4Mask_U8{0xFF, 0xFF, 0x00, 0xFF};	
+		Vec4SelectMask_U8(Dest, A, B, Mask, NumElems);
+	}
+
+	void Vec4SelectW_U8(uint8* Dest, const uint8* A, const uint8* B, int32 NumElems)
+	{
+		constexpr FVec4Mask_U8 Mask = FVec4Mask_U8{0xFF, 0xFF, 0xFF, 0x00};	
+		Vec4SelectMask_U8(Dest, A, B, Mask, NumElems);
+	}
+
+	void Vec4Fill_U8(uint8* Dest, FColor Color, int32 NumElems)
+	{
+		uint32 ColorValue = Color.ToPackedABGR();
+
+		for (int32 I = 0; I < NumElems; ++I)
+		{
+			FMemory::Memcpy(Dest + I*4, &ColorValue, sizeof(uint32));
+		}
+	}
+
+	void Vec3Fill_U8(uint8* Dest, FColor Color, int32 NumElems)
+	{
+		for (int32 ElemIndex = 0; ElemIndex < NumElems; ++ElemIndex)
+		{
+			Dest[ElemIndex*3 + 0] = Color.R;
+			Dest[ElemIndex*3 + 1] = Color.G;
+			Dest[ElemIndex*3 + 2] = Color.B;
+		}
+	}
+	
+	void Vec1Fill_U8(uint8* Dest, FColor Color, int32 NumElems)
+	{
+		FMemory::Memset(Dest, Color.R, NumElems);
+	}
+
+	void Vec4Fill_U16(uint16* Dest, Math::TIntVector4<uint16> Value, int32 NumElems)
+	{
+		for (int32 ElemIndex = 0; ElemIndex < NumElems; ++ElemIndex)
+		{
+			Dest[ElemIndex*4 + 0] = Value.X;
+			Dest[ElemIndex*4 + 1] = Value.Y;
+			Dest[ElemIndex*4 + 2] = Value.Z;
+			Dest[ElemIndex*4 + 3] = Value.W;
+		}
+	}
+
+	void Vec3Fill_U16(uint16* Dest, Math::TIntVector4<uint16> Value, int32 NumElems)
+	{
+		for (int32 ElemIndex = 0; ElemIndex < NumElems; ++ElemIndex)
+		{
+			Dest[ElemIndex*3 + 0] = Value.X;
+			Dest[ElemIndex*3 + 1] = Value.Y;
+			Dest[ElemIndex*3 + 2] = Value.Z;
+		}
+	}
+
+	void Vec1Fill_U16(uint16* Dest, Math::TIntVector4<uint16> Value, int32 NumElems)
+	{
+		for (int32 ElemIndex = 0; ElemIndex < NumElems; ++ElemIndex)
+		{
+			Dest[ElemIndex*1 + 0] = Value.X;
+		}
+	}
+
+	void Lerp_U8(uint8* Dest, const uint8* A, const uint8* B, const uint8* Factor, int32 NumElems)
+	{
+		int32 ElemIndex = 0;
+
+#if UE_MUTABLE_IMAGE_COMPUTE_UTILS_USE_VECTOR_IMPL
+		constexpr int32 ElemsPerIter = 16;
+		int32 NumIters = NumElems / ElemsPerIter;
+
+		for (int32 Iter = 0; Iter < NumIters; ++Iter)
+		{
+			VectorRegisterInt8x16 VecA = LoadUnaligned_Int8x16(A + Iter*ElemsPerIter);	
+			VectorRegisterInt8x16 VecB = LoadUnaligned_Int8x16(B + Iter*ElemsPerIter);	
+			VectorRegisterInt8x16 VecFactor = LoadUnaligned_Int8x16(Factor + Iter*ElemsPerIter);	
+	
+			VectorRegisterInt8x16 Result = Unorm255Lerp_Int8x16(VecA, VecB, VecFactor);
+
+			StoreUnaligned_Int8x16(Dest + Iter*ElemsPerIter, Result);
+		}
+
+		ElemIndex = NumIters*ElemsPerIter; 
+#endif
+
+		for (; ElemIndex < NumElems; ++ElemIndex)
+		{
+			uint8 FactorValue = Factor[ElemIndex];
+			uint16 Value = (255 - FactorValue)*A[ElemIndex] + FactorValue*B[ElemIndex] + 128;
+
+			Dest[ElemIndex] = static_cast<uint8>((Value * 257) >> 16);
+		}
+	}
+} // namespace ImageComputeUtils
+} // namespace UE::Mutable::Private
diff --git a/Engine/Plugins/Mutable/Source/MutableRuntime/Private/MuR/ImageComputeUtils.h b/Engine/Plugins/Mutable/Source/MutableRuntime/Private/MuR/ImageComputeUtils.h
new file mode 100644
--- /dev/null
+++ b/Engine/Plugins/Mutable/Source/MutableRuntime/Private/MuR/ImageComputeUtils.h
@@ -0,0 +1,39 @@
+// Copyright Epic Games, Inc. All Rights Reserved.
+
+#include "HAL/Platform.h"
+#include "Math/Color.h"
+#include "Math/IntVector.h"
+
+namespace UE::Mutable::Private
+{
+namespace ImageComputeUtils
+{
+	void Lerp_U8(uint8* Dest, const uint8* A, const uint8* B, const uint8* Factor, int32 NumElems);
+
+	void Vec4ToVec1_U8(uint8* Dest, const uint8* Src, int32 NumElems);	
+	void Vec4ToVec3_U8(uint8* Dest, const uint8* Src, int32 NumElems);	
+	void Vec3ToVec1_U8(uint8* Dest, const uint8* Src, int32 NumElems);	
+	void Vec3ToVec4_U8(uint8* Dest, const uint8* Src, int32 NumElems);
+	void Vec1ToVec3_U8(uint8* Dest, const uint8* Src, int32 NumElems);
+	void Vec1ToVec4_U8(uint8* Dest, const uint8* Src, int32 NumElems);
+
+	void Vec4SelectX_U8(uint8* Dest, const uint8* A, const uint8* B, int32 NumElems);
+	void Vec4SelectY_U8(uint8* Dest, const uint8* A, const uint8* B, int32 NumElems);
+	void Vec4SelectZ_U8(uint8* Dest, const uint8* A, const uint8* B, int32 NumElems);
+	void Vec4SelectW_U8(uint8* Dest, const uint8* A, const uint8* B, int32 NumElems);
+	
+	void Vec4SwizzleWZYXAndConvertToVec1_U8(uint8* Dest, const uint8* Src, int32 NumElems);
+	void Vec4SwizzleZYXW_U8(uint8* Dest, const uint8* Src, int32 NumElems);
+	void Vec4SwizzleXXXX_U8(uint8* Dest, const uint8* Src, int32 NumElems);
+	void Vec4SwizzleWWWW_U8(uint8* Dest, const uint8* Src, int32 NumElems);
+
+	void Vec4Fill_U8(uint8* Dest, FColor Value, int32 NumElems);
+	void Vec3Fill_U8(uint8* Dest, FColor Value, int32 NumElems);	
+	void Vec1Fill_U8(uint8* Dest, FColor Value, int32 NumElems);
+
+	void Vec4Fill_U16(uint16* Dest, Math::TIntVector4<uint16> Value, int32 NumElems);
+	void Vec3Fill_U16(uint16* Dest, Math::TIntVector4<uint16> Value, int32 NumElems);	
+	void Vec1Fill_U16(uint16* Dest, Math::TIntVector4<uint16> Value, int32 NumElems);
+
+} // namespace ImageComputeUtils
+} // namespace UE::Mutable::Private
diff --git a/Engine/Plugins/Mutable/Source/MutableRuntime/Private/MuR/ImageRLE.cpp b/Engine/Plugins/Mutable/Source/MutableRuntime/Private/MuR/ImageRLE.cpp
--- a/Engine/Plugins/Mutable/Source/MutableRuntime/Private/MuR/ImageRLE.cpp
+++ b/Engine/Plugins/Mutable/Source/MutableRuntime/Private/MuR/ImageRLE.cpp
@@ -4,6 +4,7 @@
 
 #include "HAL/UnrealMemory.h"
 #include "Misc/AssertionMacros.h"
+#include "Templates/TypeCompatibleBytes.h"
 
 namespace UE::Mutable::Private
 {
@@ -120,6 +121,108 @@ namespace UE::Mutable::Private
 		OutCompressedSize = offset;
     }
 
+	struct FIterativeDecodeTokenRLE_L
+	{
+		uint32 SrcDataOffset = 0;
+		uint16 NumEqual      = 0;
+		uint8 NumDifferent   = 0;
+		uint8 EqualValue     = 0;
+	};
+	static_assert(sizeof(FIterativeDecodeTokenRLE_L) == sizeof(uint64));
+
+	struct FSequenceHeaderRLE_L
+	{
+		uint16 NumEqual;
+		uint8 NumDifferent;
+		uint8 EqualValue;
+	};
+	static_assert(sizeof(FSequenceHeaderRLE_L) == 4);
+
+	uint64 GetBeginTokenIterativeDecodeRLE_L(const uint8* Src, int32 NumRows)
+	{
+		// Skip RLE size and rows header.
+		Src += (NumRows + 1) * sizeof(uint32);
+
+		// Read first header.
+		FSequenceHeaderRLE_L FirstHeader;
+		FMemory::Memcpy(&FirstHeader, Src, sizeof(FSequenceHeaderRLE_L));
+
+		FIterativeDecodeTokenRLE_L Token
+		{
+			.SrcDataOffset = (NumRows + 1)*sizeof(uint32) + sizeof(FSequenceHeaderRLE_L),
+			.NumEqual      = FirstHeader.NumEqual,
+			.NumDifferent  = FirstHeader.NumDifferent,
+			.EqualValue    = FirstHeader.EqualValue,
+		};
+
+		return BitCast<uint64>(Token);
+	}
+
+	void IterativeDecodeRLE_L(uint8* Dest, const uint8* Src, uint64& InOutDecodeToken, int32 NumElems)
+	{
+		FIterativeDecodeTokenRLE_L DecodeToken = BitCast<FIterativeDecodeTokenRLE_L>(InOutDecodeToken);
+
+		// DecodeToken has the information about the previous data header and points to the first differnt value.
+		Src = Src + DecodeToken.SrcDataOffset;
+
+		FSequenceHeaderRLE_L SeqHeader
+		{
+			.NumEqual     = DecodeToken.NumEqual,
+			.NumDifferent = DecodeToken.NumDifferent,
+			.EqualValue   = DecodeToken.EqualValue,
+		};
+
+		for (int32 NumElemsDecoded = 0;;)
+		{
+			int32 NumEqualToDecode = SeqHeader.NumEqual; 
+			int32 NumDifferentToDecode = SeqHeader.NumDifferent;
+
+			if (NumElemsDecoded + NumEqualToDecode > NumElems)
+			{
+				NumEqualToDecode = NumElems - NumElemsDecoded;
+				NumDifferentToDecode = 0;
+			}
+			else if (NumElemsDecoded + NumEqualToDecode + NumDifferentToDecode > NumElems)
+			{
+				NumDifferentToDecode = NumElems - (NumElemsDecoded + NumEqualToDecode);
+			}
+
+			FMemory::Memset(Dest + NumElemsDecoded, SeqHeader.EqualValue, NumEqualToDecode);
+			FMemory::Memcpy(Dest + NumElemsDecoded + NumEqualToDecode, Src, NumDifferentToDecode);
+
+			DecodeToken.NumEqual = SeqHeader.NumEqual - NumEqualToDecode; 
+			DecodeToken.NumDifferent = SeqHeader.NumDifferent - NumDifferentToDecode;
+			DecodeToken.SrcDataOffset += NumDifferentToDecode;
+
+			Src += NumDifferentToDecode;
+			NumElemsDecoded += NumEqualToDecode + NumDifferentToDecode;
+
+			if (NumElemsDecoded >= NumElems)
+			{
+				break;
+			}
+		
+            // Read next header, at this point Src is already pointing to the next sequence header as all
+            // sequence data has been consumed.
+            FMemory::Memcpy(&SeqHeader, Src, sizeof(FSequenceHeaderRLE_L));
+       
+            // We don't allow empty sequences. Stop decoding.
+            if (SeqHeader.NumEqual + SeqHeader.NumDifferent == 0)
+            {
+                check(false);
+                break;
+            }
+
+            // Point Src to the sequence data.
+			Src += sizeof(FSequenceHeaderRLE_L);
+
+			DecodeToken.SrcDataOffset += sizeof(FSequenceHeaderRLE_L);
+			DecodeToken.EqualValue = SeqHeader.EqualValue;
+			
+		}
+
+		InOutDecodeToken = BitCast<uint64>(DecodeToken);
+	}
 
     uint32 UncompressRLE_L( int32 width, int32 rows, const uint8* pStartBaseData, uint8* pStartDestData )
     {
diff --git a/Engine/Plugins/Mutable/Source/MutableRuntime/Private/MuR/ImageRLE.h b/Engine/Plugins/Mutable/Source/MutableRuntime/Private/MuR/ImageRLE.h
--- a/Engine/Plugins/Mutable/Source/MutableRuntime/Private/MuR/ImageRLE.h
+++ b/Engine/Plugins/Mutable/Source/MutableRuntime/Private/MuR/ImageRLE.h
@@ -36,6 +36,8 @@ namespace UE::Mutable::Private
                                      const uint8* BaseData,
                                      uint8* pDestData);
 
+    MUTABLERUNTIME_API extern uint64 GetBeginTokenIterativeDecodeRLE_L(const uint8* Src, int32 NumRows);
+    MUTABLERUNTIME_API extern void IterativeDecodeRLE_L(uint8* Dest, const uint8* Src, uint64& InOutDecodeToken, int32 NumElems);
     //---------------------------------------------------------------------------------------------
     //! This RLE format compresses the data per row. At the beginning of the data there are the row
     //! offsets for direct access.
diff --git a/Engine/Plugins/Mutable/Source/MutableRuntime/Private/MuR/OpImageApplyComposite.cpp b/Engine/Plugins/Mutable/Source/MutableRuntime/Private/MuR/OpImageApplyComposite.cpp
new file mode 100644
--- /dev/null
+++ b/Engine/Plugins/Mutable/Source/MutableRuntime/Private/MuR/OpImageApplyComposite.cpp
@@ -0,0 +1,415 @@
+// Copyright Epic Games, Inc. All Rights Reserved.
+
+#include "MuR/OpImageApplyComposite.h"
+
+#include "MuR/ImageComputeUtils.h"
+#include "MuR/ParallelExecutionUtils.h"
+
+
+namespace UE::Mutable::Private
+{
+namespace OpImageApplyCompositeInternal
+{
+	struct FBatchViews
+	{
+		int32 NumElems = 0;
+
+		TArrayView<uint8> Result;
+		TArrayView<const uint8> Base;
+		TArrayView<const uint8> Normal;
+	};
+
+	struct FBatchArgs
+	{
+		int32 BatchNumElems  = 0;
+		int32 LODBegin       = 0;
+		int32 LODEnd         = 0;
+		int32 FirstLODOffsetInElems = 0;
+
+		FImage* Result       = nullptr;
+		const FImage* Base   = nullptr;
+		const FImage* Normal = nullptr;
+
+		uint32 ResultBytesPerElem = 0;
+		uint32 BaseBytesPerElem   = 0;
+		uint32 NormalBytesPerElem = 0;
+	};
+
+	struct FNormalCompositeKernel
+	{
+		using FFormatTypeFunc = void(uint8*, const uint8*, int32);
+		using FComputeRoughnessTypeFunc = void(uint8*, const uint8*, const uint8*, float, int32); 
+		using FSelectTypeFunc = void(uint8*, const uint8*, const uint8*, int32);
+
+		FFormatTypeFunc* BaseFormat                 = nullptr;
+		FFormatTypeFunc* NormalFormat               = nullptr;
+		FComputeRoughnessTypeFunc* ComputeRoughness = nullptr;
+		FSelectTypeFunc* SelectChannel              = nullptr;
+		FFormatTypeFunc* ResultFormat               = nullptr;
+
+		float Power =  1.0f;
+	};
+
+	struct FMakeNormalCompositeKernelArgs
+	{
+		const FImage* ResultImage = nullptr;
+		const FImage* BaseImage   = nullptr;
+		const FImage* NormalImage = nullptr;
+		
+		int32 NormalRoughnessOutputChannel = 0;
+		float Power = 1.0f;
+	};
+
+	void Transpose(VectorRegister4Float& InOutA, VectorRegister4Float& InOutB, VectorRegister4Float& InOutC, VectorRegister4Float& InOutD)
+	{
+		VectorRegister4Float Temp0;
+		VectorRegister4Float Temp1;
+		VectorRegister4Float Temp2;
+		VectorRegister4Float Temp3;
+
+		VectorDeinterleave(Temp0, Temp1, InOutA, InOutB);
+		VectorDeinterleave(Temp2, Temp3, InOutC, InOutD);
+
+		// Temp0 = [XA,ZA,XB,ZB]
+		// Temp1 = [YA,WA,YB,WB]
+		// Temp2 = [XC,ZC,XD,ZD]
+		// Temp3 = [YC,WC,YD,WD]
+
+		VectorDeinterleave(InOutA, InOutC, Temp0, Temp2);
+		VectorDeinterleave(InOutB, InOutD, Temp1, Temp3);
+	}
+
+	void ComputeRoughnessFromVarianceImpl(uint8* Result, const uint8* Base, const VectorRegister4Float& TwoTimesVariance)
+	{
+		constexpr VectorRegister4Float FloatOne        = GlobalVectorConstants::FloatOne;
+		constexpr VectorRegister4Float Float255        = GlobalVectorConstants::Float255;
+		constexpr VectorRegister4Float FloatOneOver255 = MakeVectorRegisterFloatConstant(1.0f/255.0f, 1.0f/255.0f,  1.0f/255.0f, 1.0f/255.0f);
+
+		VectorRegister4Float Value = VectorMultiply(VectorLoadByte4(Base), FloatOneOver255);
+
+		VectorRegister4Float ValueSquared = VectorMultiply(Value, Value);
+		Value = VectorMultiply(ValueSquared, ValueSquared);
+
+		VectorRegister4Float B = VectorMultiply(VectorSubtract(Value, FloatOne), TwoTimesVariance);
+
+		Value = VectorDivide(VectorSubtract(B, Value), VectorSubtract(B, FloatOne));
+		Value = VectorSqrt(VectorSqrt(Value));
+
+		VectorStoreByte4(VectorMultiply(Value, Float255), Result);
+	}
+
+	/** 
+	 * It computes the Toksvig estimation of the normal variance for all channels similar to
+	 * TextureCompressorModule.cpp::ApplyCompositeTexture.
+	 *
+	 * Base, Normal and Result are expected to be 4 channel wide. 
+	 *
+	 * This implementation has not been validated.
+	 */
+	void ComputeRoughness(uint8* Result, const uint8* Base, const uint8* Normal, float PowerValue, int32 NumElems)
+	{
+		constexpr VectorRegister4Float FloatOne = GlobalVectorConstants::FloatOne;
+		constexpr VectorRegister4Float Float255 = GlobalVectorConstants::Float255;
+
+		constexpr VectorRegister4Float FloatVarianceTolerance = MakeVectorRegisterFloatConstant(0.00004f, 0.00004f, 0.00004f, 0.00004f);
+
+		VectorRegister4Float Power = VectorSetFloat1(PowerValue);
+
+		int32 I = 0;
+		for (; I < NumElems - 4; I += 4)
+		{
+			// Compute variance.
+			VectorRegister4Float Vec0 = VectorLoadByte4((Normal + I*4 + 0 ));
+			VectorRegister4Float Vec1 = VectorLoadByte4((Normal + I*4 + 4 ));
+			VectorRegister4Float Vec2 = VectorLoadByte4((Normal + I*4 + 8 ));
+			VectorRegister4Float Vec3 = VectorLoadByte4((Normal + I*4 + 12));
+
+			Transpose(Vec0, Vec1, Vec2, Vec3);
+		
+			VectorRegister4Float SquaredLength =
+					VectorMultiplyAdd(Vec0, Vec0, VectorMultiplyAdd(Vec1, Vec1, VectorMultiply(Vec2, Vec2)));	
+
+			VectorRegister4Float OneOverLength = VectorMultiply(VectorReciprocalSqrt(SquaredLength), Float255);
+
+			VectorRegister4Float Variance = VectorSubtract(OneOverLength, FloatOne);
+
+			Variance = VectorMax(FloatVarianceTolerance, VectorSubtract(Variance, FloatVarianceTolerance));
+			Variance = VectorMultiply(Variance, Power);
+
+			VectorRegister4Float TwoTimesVariance = VectorAdd(Variance, Variance);
+			
+			ComputeRoughnessFromVarianceImpl(Result + I*4 + 0 , Base + I*4 + 0 , VectorReplicate(TwoTimesVariance, 0));
+			ComputeRoughnessFromVarianceImpl(Result + I*4 + 4 , Base + I*4 + 4 , VectorReplicate(TwoTimesVariance, 1));
+			ComputeRoughnessFromVarianceImpl(Result + I*4 + 8 , Base + I*4 + 8 , VectorReplicate(TwoTimesVariance, 2));
+			ComputeRoughnessFromVarianceImpl(Result + I*4 + 12, Base + I*4 + 12, VectorReplicate(TwoTimesVariance, 3));
+		}
+
+		for (; I < NumElems; ++I)
+		{
+			VectorRegister4Float Vec = VectorLoadByte4((Normal + I*4 + 0));
+			VectorRegister4Float OneOverLength = VectorMultiply(VectorReciprocalSqrt(VectorDot3(Vec, Vec)), Float255);
+
+			VectorRegister4Float Variance = VectorSubtract(OneOverLength, FloatOne);
+
+			Variance = VectorMax(FloatVarianceTolerance, VectorSubtract(Variance, FloatVarianceTolerance));
+			Variance = VectorMultiply(Variance, Power);
+
+			VectorRegister4Float TwoTimesVariance = VectorAdd(Variance, Variance);
+
+			ComputeRoughnessFromVarianceImpl(Result + I*4 + 0, Base + I*4 + 0, TwoTimesVariance);
+		}
+
+	}
+
+	FNormalCompositeKernel MakeNormalCompositeKernel(FMakeNormalCompositeKernelArgs Args)
+	{
+		using namespace ImageComputeUtils;
+
+		FNormalCompositeKernel Kernel;
+		Kernel.Power = Args.Power;
+		Kernel.ComputeRoughness = ComputeRoughness;
+
+		if (Args.BaseImage->GetFormat() == EImageFormat::RGB_UByte)
+		{
+			Kernel.BaseFormat = Vec3ToVec4_U8;
+		}
+
+		if (Args.NormalImage->GetFormat() == EImageFormat::RGB_UByte)
+		{
+			Kernel.NormalFormat = Vec3ToVec4_U8;
+		}
+
+		switch (Args.NormalRoughnessOutputChannel)
+		{
+			case 0: Kernel.SelectChannel = Vec4SelectX_U8; break;
+			case 1: Kernel.SelectChannel = Vec4SelectY_U8; break;
+			case 2: Kernel.SelectChannel = Vec4SelectZ_U8; break;
+			case 3: Kernel.SelectChannel = Vec4SelectW_U8; break;
+			default: Kernel.SelectChannel = nullptr; break;
+		}
+
+		if (Args.ResultImage->GetFormat() == EImageFormat::RGB_UByte)
+		{
+			Kernel.ResultFormat = Vec4ToVec3_U8;
+		}
+
+		return Kernel;
+	}
+
+	void RunKernelOnBatch(FNormalCompositeKernel Kernel, FBatchViews BatchViews)
+	{
+		MUTABLE_CPUPROFILER_SCOPE(OpImageNormalComposite_RunKernelOnBatch);
+
+		constexpr int32 MaxElemSize = 4;
+		constexpr int32 NumCacheLinesPerWave = 8;
+		constexpr int32 WaveSizeInElems = (PLATFORM_CACHE_LINE_SIZE * NumCacheLinesPerWave)/MaxElemSize;
+
+		struct alignas(PLATFORM_CACHE_LINE_SIZE) FTemporalRegister
+		{
+			uint8 Data[WaveSizeInElems*MaxElemSize];	
+		};
+
+		FTemporalRegister Register0;
+		FTemporalRegister Register1;
+		FTemporalRegister Register2;
+		//FTemporalRegister Register3;
+		//FTemporalRegister Register4;
+
+		if (BatchViews.NumElems == 0)
+		{
+			return;
+		}
+
+		int32 BaseBytesPerElem = BatchViews.Base.Num() / BatchViews.NumElems;
+		int32 NormalBytesPerElem = BatchViews.Normal.Num() / BatchViews.NumElems;
+
+		check(BaseBytesPerElem <= MaxElemSize);
+		check(NormalBytesPerElem <= MaxElemSize);
+
+		uint8* ResultBuffer = BatchViews.Result.GetData();
+		const uint8* BaseBuffer = BatchViews.Base.GetData();
+		const uint8* NormalBuffer = BatchViews.Normal.GetData();
+
+		const int32 NumWaves = FMath::DivideAndRoundUp<int32>(BatchViews.NumElems, WaveSizeInElems);
+		for (int32 Idx = 0; Idx < NumWaves; ++Idx)
+		{
+			int32 WaveBeginInElems = Idx * WaveSizeInElems;
+			int32 WaveEndInElems   = FMath::Min(WaveBeginInElems + WaveSizeInElems, BatchViews.NumElems); 
+
+			int32 NumWaveElems = WaveEndInElems - WaveBeginInElems;
+
+			uint8* ResultDataBuffer       = ResultBuffer + WaveBeginInElems*BaseBytesPerElem;
+			const uint8* BaseDataBuffer   = BaseBuffer + WaveBeginInElems*BaseBytesPerElem;
+			const uint8* NormalDataBuffer = NormalBuffer + WaveBeginInElems*NormalBytesPerElem;
+
+			if (Kernel.BaseFormat)
+			{
+				Kernel.BaseFormat(Register0.Data, BaseDataBuffer, NumWaveElems);
+				BaseDataBuffer = Register0.Data;
+			}
+
+			if (Kernel.NormalFormat)
+			{
+				Kernel.NormalFormat(Register1.Data, NormalDataBuffer, NumWaveElems);
+				NormalDataBuffer = Register1.Data;
+			}
+
+			const uint8* RoughnessDataBuffer = BaseDataBuffer;
+			if (Kernel.ComputeRoughness)
+			{
+				Kernel.ComputeRoughness(Register2.Data, BaseDataBuffer, NormalDataBuffer, Kernel.Power, NumWaveElems);
+				RoughnessDataBuffer = Register2.Data;
+			}
+
+			const uint8* SelectDataBuffer = BaseDataBuffer;
+			if (Kernel.SelectChannel)
+			{
+				uint8* MutableSelectDataBuffer = Kernel.ResultFormat ? Register1.Data : ResultDataBuffer;
+				SelectDataBuffer = MutableSelectDataBuffer;
+				
+				Kernel.SelectChannel(MutableSelectDataBuffer, BaseDataBuffer, RoughnessDataBuffer, NumWaveElems); 
+			}
+
+			if (Kernel.ResultFormat)
+			{
+				Kernel.ResultFormat(ResultDataBuffer, SelectDataBuffer, NumWaveElems); 
+			}
+			else if (!Kernel.SelectChannel && ResultDataBuffer != BaseDataBuffer)
+			{
+				FMemory::Memcpy(ResultDataBuffer, BaseDataBuffer, NumWaveElems*BaseBytesPerElem);
+			}
+		}
+	}
+
+	FORCENOINLINE int32 GetNumBatchesLODRangeOffsetViews(const FBatchArgs& Args)
+	{
+		const bool bOnlyFirstLOD = Args.FirstLODOffsetInElems >= 0;
+
+		check(Args.Result);
+		const int32 NumBatches = bOnlyFirstLOD
+				? Args.Result->DataStorage.GetNumBatchesFirstLODOffset(Args.BatchNumElems, Args.ResultBytesPerElem, Args.FirstLODOffsetInElems*Args.ResultBytesPerElem)
+				: Args.Result->DataStorage.GetNumBatchesLODRange(Args.BatchNumElems, Args.ResultBytesPerElem, Args.LODBegin, Args.LODEnd);
+
+#if DO_CHECK
+		if (Args.Base)
+		{
+			const int32 BaseNumBatches = bOnlyFirstLOD
+					? Args.Base->DataStorage.GetNumBatchesFirstLODOffset(Args.BatchNumElems, Args.BaseBytesPerElem, Args.FirstLODOffsetInElems*Args.BaseBytesPerElem)
+					: Args.Base->DataStorage.GetNumBatchesLODRange(Args.BatchNumElems, Args.BaseBytesPerElem, Args.LODBegin, Args.LODEnd);
+
+			check(NumBatches == BaseNumBatches);
+		}
+	
+		if (Args.Normal)
+		{
+			const int32 BlendNumBatches = bOnlyFirstLOD
+					? Args.Normal->DataStorage.GetNumBatchesFirstLODOffset(Args.BatchNumElems, Args.NormalBytesPerElem, Args.FirstLODOffsetInElems*Args.NormalBytesPerElem)
+					: Args.Normal->DataStorage.GetNumBatchesLODRange(Args.BatchNumElems, Args.NormalBytesPerElem, Args.LODBegin, Args.LODEnd);
+
+			check(NumBatches == BlendNumBatches);
+		}
+
+#endif
+		return NumBatches;
+	}
+
+	FORCENOINLINE FBatchViews GetBatchLODRangeOffsetViews(int32 BatchId, const FBatchArgs& Args)
+	{
+		FBatchViews BatchViewsResult;
+		
+		const bool bOnlyFirstLOD = Args.FirstLODOffsetInElems >= 0;
+
+		check(Args.Result);
+		BatchViewsResult.Result = bOnlyFirstLOD
+				? Args.Result->DataStorage.GetBatchFirstLODOffset(BatchId, Args.BatchNumElems, Args.ResultBytesPerElem, Args.FirstLODOffsetInElems*Args.ResultBytesPerElem) 
+				: Args.Result->DataStorage.GetBatchLODRange(BatchId, Args.BatchNumElems, Args.ResultBytesPerElem, Args.LODBegin, Args.LODEnd);
+
+		BatchViewsResult.NumElems = BatchViewsResult.Result.Num() / Args.ResultBytesPerElem;
+		
+		if (Args.Base)
+		{
+			BatchViewsResult.Base = bOnlyFirstLOD
+					? Args.Base->DataStorage.GetBatchFirstLODOffset(BatchId, Args.BatchNumElems, Args.BaseBytesPerElem, Args.FirstLODOffsetInElems*Args.BaseBytesPerElem) 
+					: Args.Base->DataStorage.GetBatchLODRange(BatchId, Args.BatchNumElems, Args.BaseBytesPerElem, Args.LODBegin, Args.LODEnd); 
+
+			check(BatchViewsResult.NumElems == BatchViewsResult.Base.Num() / Args.BaseBytesPerElem);
+		}
+
+		if (Args.Normal)
+		{
+			BatchViewsResult.Normal = bOnlyFirstLOD
+					? Args.Normal->DataStorage.GetBatchFirstLODOffset(BatchId, Args.BatchNumElems, Args.NormalBytesPerElem, Args.FirstLODOffsetInElems*Args.NormalBytesPerElem) 
+					: Args.Normal->DataStorage.GetBatchLODRange(BatchId, Args.BatchNumElems, Args.NormalBytesPerElem, Args.LODBegin, Args.LODEnd); 
+
+			check(BatchViewsResult.NumElems == BatchViewsResult.Normal.Num() / Args.NormalBytesPerElem);
+		}
+
+		return BatchViewsResult;
+	}
+
+} // namespace OpImageApplyCompositeInternal
+
+	void ImageNormalComposite(FImage* ResultImage, const FImage* BaseImage, const FImage* NormalImage, int32 NormalRoughnessOutputChannel, float Power)
+	{
+		check(BaseImage);
+		check(ResultImage);
+		check(NormalImage);
+
+		check(BaseImage->GetSize() == ResultImage->GetSize());
+		check(BaseImage->GetSize() == NormalImage->GetSize());
+
+		check(
+			NormalImage->GetFormat() == EImageFormat::RGB_UByte || 
+			NormalImage->GetFormat() == EImageFormat::RGBA_UByte);
+		check(
+			BaseImage->GetFormat() == EImageFormat::RGB_UByte || 
+			BaseImage->GetFormat() == EImageFormat::RGBA_UByte);
+		
+		check(BaseImage->GetFormat() == ResultImage->GetFormat());
+
+		OpImageApplyCompositeInternal::FMakeNormalCompositeKernelArgs Args =
+		{
+			.ResultImage = ResultImage,
+			.BaseImage   = BaseImage,
+			.NormalImage = NormalImage,
+
+			.NormalRoughnessOutputChannel = NormalRoughnessOutputChannel,
+			.Power = Power,
+		};
+
+		OpImageApplyCompositeInternal::FNormalCompositeKernel Kernel = 
+				OpImageApplyCompositeInternal::MakeNormalCompositeKernel(Args);
+
+
+		const uint32 BaseBytesPerElem = GetImageFormatData(BaseImage->GetFormat()).BytesPerBlock;
+		const uint32 NormalBytesPerElem = GetImageFormatData(NormalImage->GetFormat()).BytesPerBlock;
+
+		OpImageApplyCompositeInternal::FBatchArgs BatchArgs
+		{
+			.BatchNumElems = 1 << 12,
+			.LODBegin      = 0,
+			.LODEnd        = BaseImage->GetLODCount(),
+			.FirstLODOffsetInElems = -1,
+
+			.Result = ResultImage,
+			.Base   = BaseImage  ,
+			.Normal = NormalImage,
+
+		    .ResultBytesPerElem = BaseBytesPerElem,
+		    .BaseBytesPerElem   = BaseBytesPerElem,
+		    .NormalBytesPerElem = NormalBytesPerElem,
+		};
+
+		const int32 NumBatches = OpImageApplyCompositeInternal::GetNumBatchesLODRangeOffsetViews(BatchArgs);
+		ParallelExecutionUtils::InvokeBatchParallelFor(NumBatches, 
+		[
+			&BatchArgs, &Kernel
+		] (int32 BatchId)
+		{
+			OpImageApplyCompositeInternal::FBatchViews BatchViews = 
+					OpImageApplyCompositeInternal::GetBatchLODRangeOffsetViews(BatchId, BatchArgs);
+
+			OpImageApplyCompositeInternal::RunKernelOnBatch(Kernel, BatchViews);
+		});
+	}
+} // namespace UE::Mutable::Private
diff --git a/Engine/Plugins/Mutable/Source/MutableRuntime/Private/MuR/OpImageApplyComposite.h b/Engine/Plugins/Mutable/Source/MutableRuntime/Private/MuR/OpImageApplyComposite.h
--- a/Engine/Plugins/Mutable/Source/MutableRuntime/Private/MuR/OpImageApplyComposite.h
+++ b/Engine/Plugins/Mutable/Source/MutableRuntime/Private/MuR/OpImageApplyComposite.h
@@ -2,41 +2,9 @@
 
 #pragma once
 
-#include "Math/UnrealMathUtility.h"
-#include "Math/Vector.h"
-
 #include "MuR/ImagePrivate.h"
-#include "MuR/OpImageLayer.h"
-#include "MuR/OpImageBlend.h"
 
 namespace UE::Mutable::Private
 {
-
-	inline void ImageNormalComposite(FImage* pResult, const FImage* pBase, const FImage* pNormal, ECompositeImageMode mode, float power)
-	{
-		if (mode == ECompositeImageMode::CIM_Disabled)
-		{
-			// Copy base to result. Using functor layer combine with identity to do the copy for simplicity.
-			// This could be optimized but will never happen since it is checked at compile time. 	
-			ImageLayerCombineFunctor( pResult, pBase, pNormal, FNormalCompositeIdentityFunctor{}); 
-
-			return;
-		}	
-
-		const uint8 channel = [mode]() 
-			{
-				switch (mode)
-				{
-					case ECompositeImageMode::CIM_NormalRoughnessToRed   : return 0u;
-					case ECompositeImageMode::CIM_NormalRoughnessToGreen : return 1u;
-					case ECompositeImageMode::CIM_NormalRoughnessToBlue  : return 2u;
-					case ECompositeImageMode::CIM_NormalRoughnessToAlpha : return 3u;
-				}
-			
-				check(false);
-				return 0u;		
-			}();
-
-		ImageLayerCombineFunctor( pResult, pBase, pNormal, FNormalCompositeFunctor{channel, power});
-	}
+	void ImageNormalComposite(FImage* ResultImage, const FImage* BaseImage, const FImage* NormalImage, int32 NormalRoughnessOutputChannel, float Power);
 }
diff --git a/Engine/Plugins/Mutable/Source/MutableRuntime/Private/MuR/OpImageLayer.cpp b/Engine/Plugins/Mutable/Source/MutableRuntime/Private/MuR/OpImageLayer.cpp
--- a/Engine/Plugins/Mutable/Source/MutableRuntime/Private/MuR/OpImageLayer.cpp
+++ b/Engine/Plugins/Mutable/Source/MutableRuntime/Private/MuR/OpImageLayer.cpp
@@ -2,1633 +2,2181 @@
 
 #include "MuR/OpImageLayer.h"
 
-#include "MuR/ImagePrivate.h"
 #include "MuR/ImageRLE.h"
+#include "MuR/ImageComputeUtils.h"
 #include "MuR/ParallelExecutionUtils.h"
-#include "Templates/UnrealTemplate.h"
 
 #include "MuR/OpImageBlend.h"
-#include "MuR/MutableRuntimeModule.h"
+
+#include "Math/UnrealMathUtility.h"
+#include "Math/VectorRegister.h"
 
 namespace UE::Mutable::Private
 {
-
-namespace OpImageLayerInternal
+namespace OpImageLayerBlendOps
 {
 
-	struct FOpLayerBatchArgs
+#if PLATFORM_ENABLE_VECTORINTRINSICS && !PLATFORM_ENABLE_VECTORINTRINSICS_NEON
+	using VectorRegisterUInt8x16 = __m128i;
+	using VectorRegisterUInt16x8 = __m128i;
+
+	FORCEINLINE void StoreUnaligned_UInt8x16(uint8* Ptr, const VectorRegisterUInt8x16 Reg)
 	{
-		int32 BatchNumElems  = 0;
-		int32 LODBegin       = 0;
-		int32 LODEnd         = 0;
-		int32 FirstLODOffset = 0;
+		_mm_storeu_si128((__m128i*)Ptr, Reg);
+	}
 
-		FImage* Result       = nullptr;
-		const FImage* Base   = nullptr;
-		const FImage* Blend  = nullptr;
-		const FImage* Mask   = nullptr;
+	FORCEINLINE VectorRegisterUInt8x16 LoadUnaligned_UInt8x16(const uint8* Ptr)
+	{
+		return _mm_loadu_si128((__m128i*)Ptr);
+	}
 
-		uint32 ResultBytesPerElem = 0;
-		uint32 BaseBytesPerElem   = 0;
-		uint32 BlendBytesPerElem  = 0;
-		uint32 MaskBytesPerElem   = 0;
-	};
+	FORCEINLINE VectorRegisterUInt8x16 LoadUnaligned_UInt16x8(const uint16* Ptr)
+	{
+		return _mm_loadu_si128((__m128i*)Ptr);
+	}
 
-	struct FOpLayerBatchViews
+	FORCEINLINE VectorRegisterUInt16x8 GetLow16_UInt8x16(VectorRegisterUInt8x16 Reg)
 	{
-		int32 NumElems = 0;
+		return _mm_cvtepu8_epi16(Reg);
+	}	
 
-		TArrayView<uint8> Result;
-		TArrayView<const uint8> Base;
-		TArrayView<const uint8> Blend;
-		TArrayView<const uint8> Mask;
-	};
-	
-	FORCENOINLINE int32 GetOpLayerNumBatches(const FOpLayerBatchArgs& Args)
-	{	
-		check(Args.Result);
-		const int32 NumBatches = Args.Result->DataStorage.GetNumBatches(Args.BatchNumElems, Args.ResultBytesPerElem);
+	FORCEINLINE VectorRegisterUInt16x8 GetHigh16_UInt8x16(VectorRegisterUInt8x16 Reg)
+	{
+		return _mm_cvtepu8_epi16(_mm_srli_si128(Reg, 8));
+	}
 
-#if DO_CHECK
-		if (Args.Base)
-		{
-			check(NumBatches == Args.Base->DataStorage.GetNumBatches(Args.BatchNumElems, Args.BaseBytesPerElem));
-		}
-	
-		if (Args.Blend)
-		{
-			check(NumBatches == Args.Blend->DataStorage.GetNumBatches(Args.BatchNumElems, Args.BlendBytesPerElem));
-		}
+	FORCEINLINE VectorRegisterUInt8x16 Saturate_UInt16x8(VectorRegisterUInt16x8 Lo, VectorRegisterUInt16x8 Hi)
+	{
+		return _mm_packus_epi16(Lo, Hi); 
+	}
 
-		if (Args.Mask)
-		{
-			check(NumBatches == Args.Mask->DataStorage.GetNumBatches(Args.BatchNumElems, Args.MaskBytesPerElem));
-		}
-#endif
-		return NumBatches;
+	/* Multply 16 bit 255 normalized values and saturate */
+	FORCEINLINE VectorRegisterUInt8x16 Mult255Unorm_UInt16x8(VectorRegisterUInt16x8 ALo, VectorRegisterUInt16x8 AHi, VectorRegisterUInt16x8 BLo, VectorRegisterUInt16x8 BHi)
+	{
+		unimplemented();
+		return Saturate_UInt16x8(ALo, AHi);
 	}
 
-	FORCENOINLINE int32 GetOpLayerNumBatchesLODRange(const FOpLayerBatchArgs& Args)
+	FORCEINLINE VectorRegisterUInt8x16 Mult255Unorm_UInt8x16(VectorRegisterUInt8x16 A, VectorRegisterUInt8x16 B)
 	{
-		check(Args.Result);
-		const int32 NumBatches = Args.Result->DataStorage.GetNumBatchesLODRange(Args.BatchNumElems, Args.ResultBytesPerElem, Args.LODBegin, Args.LODEnd);
+		__m128i MultLo = _mm_add_epi16(_mm_mullo_epi16(GetLow16_UInt8x16(A), GetLow16_UInt8x16(B)), _mm_set1_epi16(0x80));
+		__m128i MultHi = _mm_add_epi16(_mm_mullo_epi16(GetHigh16_UInt8x16(A), GetHigh16_UInt8x16(B)), _mm_set1_epi16(0x80));
 
-#if DO_CHECK
-		if (Args.Base)
-		{
-			check(NumBatches == Args.Base->DataStorage.GetNumBatchesLODRange(Args.BatchNumElems, Args.BaseBytesPerElem, Args.LODBegin, Args.LODEnd));
-		}
-	
-		if (Args.Blend)
-		{
-			check(NumBatches == Args.Blend->DataStorage.GetNumBatchesLODRange(Args.BatchNumElems, Args.BlendBytesPerElem, Args.LODBegin, Args.LODEnd));
-		}
+		MultLo = _mm_mulhi_epu16(MultLo, _mm_set1_epi16(257));
+		MultHi = _mm_mulhi_epu16(MultHi, _mm_set1_epi16(257));
 
-		if (Args.Mask)
-		{
-			check(NumBatches == Args.Mask->DataStorage.GetNumBatchesLODRange(Args.BatchNumElems, Args.MaskBytesPerElem, Args.LODBegin, Args.LODEnd));
-		}
-#endif
-		return NumBatches;
+		return Saturate_UInt16x8(MultLo, MultHi);
 	}
-	
-	FORCENOINLINE int32 GetOpLayerNumBatchesLODRangeOffsetViews(const FOpLayerBatchArgs& Args)
+
+	FORCEINLINE VectorRegisterUInt8x16 Screen_UInt8x16(VectorRegisterUInt8x16 A, VectorRegisterUInt8x16 B)
 	{
-		const bool bOnlyFirstLOD = Args.FirstLODOffset >= 0;
+		__m128i ALo = GetLow16_UInt8x16(A);
+		__m128i AHi = GetHigh16_UInt8x16(A);
+		
+		__m128i BLo = GetLow16_UInt8x16(B);
+		__m128i BHi = GetHigh16_UInt8x16(B);
 
-		check(Args.Result);
-		const int32 NumBatches = bOnlyFirstLOD
-				? Args.Result->DataStorage.GetNumBatchesFirstLODOffset(Args.BatchNumElems, Args.ResultBytesPerElem, Args.FirstLODOffset)
-				: Args.Result->DataStorage.GetNumBatchesLODRange(Args.BatchNumElems, Args.ResultBytesPerElem, Args.LODBegin, Args.LODEnd);
+		__m128i MultLo = _mm_mullo_epi16(ALo, BLo);
+		__m128i MultHi = _mm_mullo_epi16(AHi, BHi);
 
-#if DO_CHECK
-		if (Args.Base)
-		{
-			const int32 BaseNumBatches = bOnlyFirstLOD
-					? Args.Base->DataStorage.GetNumBatchesFirstLODOffset(Args.BatchNumElems, Args.BaseBytesPerElem, Args.FirstLODOffset)
-					: Args.Base->DataStorage.GetNumBatchesLODRange(Args.BatchNumElems, Args.BaseBytesPerElem, Args.LODBegin, Args.LODEnd);
+		MultLo = _mm_mulhi_epu16(_mm_add_epi16(MultLo, _mm_set1_epi16(0x80)), _mm_set1_epi16(257));
+		MultHi = _mm_mulhi_epu16(_mm_add_epi16(MultHi, _mm_set1_epi16(0x80)), _mm_set1_epi16(257));
 
-			check(NumBatches == BaseNumBatches);
-		}
-	
-		if (Args.Blend)
-		{
+		__m128i ResultLo = _mm_sub_epi16(_mm_add_epi16(ALo, BLo), MultLo);
+		__m128i ResultHi = _mm_sub_epi16(_mm_add_epi16(AHi, BHi), MultHi);
 
-			const int32 BlendNumBatches = bOnlyFirstLOD
-					? Args.Blend->DataStorage.GetNumBatchesFirstLODOffset(Args.BatchNumElems, Args.BlendBytesPerElem, Args.FirstLODOffset)
-					: Args.Blend->DataStorage.GetNumBatchesLODRange(Args.BatchNumElems, Args.BlendBytesPerElem, Args.LODBegin, Args.LODEnd);
+		return Saturate_UInt16x8(ResultLo, ResultHi);
+	}
 
-			check(NumBatches == BlendNumBatches);
-		}
+	FORCEINLINE VectorRegisterUInt8x16 SoftLight_UInt8x16(VectorRegisterUInt8x16 A, VectorRegisterUInt8x16 B)
+	{
+		__m128i ALo = GetLow16_UInt8x16(A);
+		__m128i AHi = GetHigh16_UInt8x16(A);
+		
+		__m128i BLo = GetLow16_UInt8x16(B);
+		__m128i BHi = GetHigh16_UInt8x16(B);
 
-		if (Args.Mask)
-		{
-			const int32 MaskNumBatches = bOnlyFirstLOD
-					? Args.Mask->DataStorage.GetNumBatchesFirstLODOffset(Args.BatchNumElems, Args.MaskBytesPerElem, Args.FirstLODOffset)
-					: Args.Mask->DataStorage.GetNumBatchesLODRange(Args.BatchNumElems, Args.MaskBytesPerElem, Args.LODBegin, Args.LODEnd);
+		__m128i MultLo = _mm_mullo_epi16(ALo, BLo);
+		__m128i MultHi = _mm_mullo_epi16(AHi, BHi);
 
-			check(NumBatches == MaskNumBatches);
-		}
-#endif
-		return NumBatches;
-	}
+		MultLo = _mm_mulhi_epu16(_mm_add_epi16(MultLo, _mm_set1_epi16(0x80)), _mm_set1_epi16(257));
+		MultHi = _mm_mulhi_epu16(_mm_add_epi16(MultHi, _mm_set1_epi16(0x80)), _mm_set1_epi16(257));
 
+		__m128i ScreenLo = _mm_sub_epi16(_mm_add_epi16(ALo, BLo), MultLo);
+		__m128i ScreenHi = _mm_sub_epi16(_mm_add_epi16(AHi, BHi), MultHi);
 
-	FORCENOINLINE FOpLayerBatchViews GetOpLayerBatchViews(int32 BatchId, const FOpLayerBatchArgs& Args)
+		__m128i LerpLo = _mm_mullo_epi16(MultLo, _mm_xor_si128(ALo, _mm_set1_epi16(0xFF)));  
+		__m128i LerpHi = _mm_mullo_epi16(MultHi, _mm_xor_si128(AHi, _mm_set1_epi16(0xFF)));
+
+		LerpLo = _mm_add_epi16(LerpLo, _mm_add_epi16(_mm_mullo_epi16(ScreenLo, ALo), _mm_set1_epi16(0x80)));  
+		LerpHi = _mm_add_epi16(LerpHi, _mm_add_epi16(_mm_mullo_epi16(ScreenHi, AHi), _mm_set1_epi16(0x80)));
+
+		LerpLo = _mm_mulhi_epu16(LerpLo, _mm_set1_epi16(257));
+		LerpHi = _mm_mulhi_epu16(LerpHi, _mm_set1_epi16(257));
+
+		return Saturate_UInt16x8(LerpLo, LerpHi);
+	}
+
+	FORCEINLINE VectorRegisterUInt8x16 HardLight_UInt8x16(VectorRegisterUInt8x16 A, VectorRegisterUInt8x16 B)
 	{
-		FOpLayerBatchViews BatchViewsResult;
-	
-		check(Args.Result);
-		BatchViewsResult.Result = Args.Result->DataStorage.GetBatch(BatchId, Args.BatchNumElems, Args.ResultBytesPerElem);
+		__m128i Xor = _mm_cmpgt_epi8(_mm_setzero_si128(), B);
 
-		BatchViewsResult.NumElems = BatchViewsResult.Result.Num() / Args.ResultBytesPerElem;
+		__m128i AXor = _mm_xor_si128(A, Xor);
 
-		if (Args.Base)
-		{
-			BatchViewsResult.Base = Args.Base->DataStorage.GetBatch(BatchId, Args.BatchNumElems, Args.BaseBytesPerElem);
-			check(BatchViewsResult.NumElems == BatchViewsResult.Base.Num() / Args.BaseBytesPerElem);
-		}
+		__m128i TwoTimesBMod255Xor = _mm_xor_si128(_mm_add_epi8(B, B), Xor);
 
-		if (Args.Blend)
-		{
-			BatchViewsResult.Blend = Args.Blend->DataStorage.GetBatch(BatchId, Args.BatchNumElems, Args.BlendBytesPerElem);
-			check(BatchViewsResult.NumElems == BatchViewsResult.Blend.Num() / Args.BlendBytesPerElem);
-		}
+		__m128i MultLo = _mm_add_epi16(_mm_mullo_epi16(GetLow16_UInt8x16(AXor), GetLow16_UInt8x16(TwoTimesBMod255Xor)), _mm_set1_epi16(0x80));
+		__m128i MultHi = _mm_add_epi16(_mm_mullo_epi16(GetHigh16_UInt8x16(AXor), GetHigh16_UInt8x16(TwoTimesBMod255Xor)), _mm_set1_epi16(0x80));
 
-		if (Args.Mask)
-		{
-			BatchViewsResult.Mask = Args.Mask->DataStorage.GetBatch(BatchId, Args.BatchNumElems, Args.MaskBytesPerElem);
-			check(BatchViewsResult.NumElems == BatchViewsResult.Mask.Num() / Args.MaskBytesPerElem);
-		}
+		MultLo = _mm_mulhi_epu16(MultLo, _mm_set1_epi16(257));
+		MultHi = _mm_mulhi_epu16(MultHi, _mm_set1_epi16(257));
 
-		return BatchViewsResult;
+		return _mm_xor_si128(Saturate_UInt16x8(MultLo, MultHi), Xor);
 	}
 
-	FORCENOINLINE FOpLayerBatchViews GetOpLayerBatchLODRangeViews(int32 BatchId, const FOpLayerBatchArgs& Args)
+	FORCEINLINE VectorRegisterUInt8x16 Overlay_UInt8x16(VectorRegisterUInt8x16 A, VectorRegisterUInt8x16 B)
 	{
-		FOpLayerBatchViews BatchViewsResult;
-		
-		check(Args.Result);
-		BatchViewsResult.Result = 
-					Args.Result->DataStorage.GetBatchLODRange(BatchId, Args.BatchNumElems, Args.ResultBytesPerElem, Args.LODBegin, Args.LODEnd);
+		__m128i Xor = _mm_cmpgt_epi8(_mm_setzero_si128(), A);
 
-		BatchViewsResult.NumElems = BatchViewsResult.Result.Num() / Args.ResultBytesPerElem;
-
-		if (Args.Base)
-		{
-			BatchViewsResult.Base = 
-					Args.Base->DataStorage.GetBatchLODRange(BatchId, Args.BatchNumElems, Args.BaseBytesPerElem, Args.LODBegin, Args.LODEnd);
-			check(BatchViewsResult.NumElems == BatchViewsResult.Base.Num() / Args.BaseBytesPerElem);
-		}
+		__m128i TwoTimesAMod255Xor = _mm_xor_si128(_mm_add_epi8(A, A), Xor);
+		__m128i BXor = _mm_xor_si128(B, Xor);
 
-		if (Args.Blend)
-		{
-			BatchViewsResult.Blend = 
-					Args.Blend->DataStorage.GetBatchLODRange(BatchId, Args.BatchNumElems, Args.BlendBytesPerElem, Args.LODBegin, Args.LODEnd);
-			check(BatchViewsResult.NumElems == BatchViewsResult.Blend.Num() / Args.BlendBytesPerElem);
-		}
+		__m128i MultLo = _mm_add_epi16(_mm_mullo_epi16(GetLow16_UInt8x16(TwoTimesAMod255Xor), GetLow16_UInt8x16(BXor)), _mm_set1_epi16(0x80));
+		__m128i MultHi = _mm_add_epi16(_mm_mullo_epi16(GetHigh16_UInt8x16(TwoTimesAMod255Xor), GetHigh16_UInt8x16(BXor)), _mm_set1_epi16(0x80));
 
-		if (Args.Mask)
-		{
-			BatchViewsResult.Mask = 
-					Args.Mask->DataStorage.GetBatchLODRange(BatchId, Args.BatchNumElems, Args.MaskBytesPerElem, Args.LODBegin, Args.LODEnd);
-			check(BatchViewsResult.NumElems == BatchViewsResult.Mask.Num() / Args.MaskBytesPerElem);
-		}
+		MultLo = _mm_mulhi_epu16(MultLo, _mm_set1_epi16(257));
+		MultHi = _mm_mulhi_epu16(MultHi, _mm_set1_epi16(257));
 
-		return BatchViewsResult;
+		return _mm_xor_si128(Saturate_UInt16x8(MultLo, MultHi), Xor);
 	}
 
-	FORCENOINLINE FOpLayerBatchViews GetOpLayerBatchLODRangeOffsetViews(int32 BatchId, const FOpLayerBatchArgs& Args)
+	FORCEINLINE VectorRegisterUInt8x16 Dodge_UInt8x16(VectorRegisterUInt8x16 A, VectorRegisterUInt8x16 B)
 	{
-		FOpLayerBatchViews BatchViewsResult;
-		
-		const bool bOnlyFirstLOD = Args.FirstLODOffset >= 0;
+		//__m128i A0 = _mm_slli_epi16(GetLow16_UInt8x16(A), 8);
+		//__m128i A1 = _mm_slli_epi16(GetHigh16_UInt8x16(A), 8);
+
+		//__m128i B0 = _mm_sub_epi16(_mm_set1_epi16(256), GetLow16_UInt8x16(B));
+		//__m128i B1 = _mm_sub_epi16(_mm_set1_epi16(256), GetHigh16_UInt8x16(B));
+		//
+		//__m128 FloatA0 = _mm_cvtepi32_ps(_mm_cvtepu16_epi32(A0));
+		//__m128 FloatA1 = _mm_cvtepi32_ps(_mm_cvtepu16_epi32(_mm_srli_si128(A0, 8)));
+		//__m128 FloatA2 = _mm_cvtepi32_ps(_mm_cvtepu16_epi32(A1));
+		//__m128 FloatA3 = _mm_cvtepi32_ps(_mm_cvtepu16_epi32(_mm_srli_si128(A1, 8)));
+
+		//__m128 FloatB0 = _mm_cvtepi32_ps(_mm_cvtepu16_epi32(B0));
+		//__m128 FloatB1 = _mm_cvtepi32_ps(_mm_cvtepu16_epi32(_mm_srli_si128(B0, 8)));
+		//__m128 FloatB2 = _mm_cvtepi32_ps(_mm_cvtepu16_epi32(B1));
+		//__m128 FloatB3 = _mm_cvtepi32_ps(_mm_cvtepu16_epi32(_mm_srli_si128(B1, 8)));
+
+		//__m128 FloatResult0 = _mm_div_ps(FloatA0, FloatB0);
+		//__m128 FloatResult1 = _mm_div_ps(FloatA1, FloatB1);
+		//__m128 FloatResult2 = _mm_div_ps(FloatA2, FloatB2);
+		//__m128 FloatResult3 = _mm_div_ps(FloatA3, FloatB3);
+
+		//__m128i ResultLo = _mm_packus_epi32(_mm_cvtps_epi32(FloatResult0), _mm_cvtps_epi32(FloatResult1));
+		//__m128i ResultHi = _mm_packus_epi32(_mm_cvtps_epi32(FloatResult2), _mm_cvtps_epi32(FloatResult3));
+
+		//return Saturate_UInt16x8(ResultLo, ResultHi);
+
+		unimplemented();
+		return A;
+	}
 
-		check(Args.Result);
-		BatchViewsResult.Result = bOnlyFirstLOD
-				? Args.Result->DataStorage.GetBatchFirstLODOffset(BatchId, Args.BatchNumElems, Args.ResultBytesPerElem, Args.FirstLODOffset) 
-				: Args.Result->DataStorage.GetBatchLODRange(BatchId, Args.BatchNumElems, Args.ResultBytesPerElem, Args.LODBegin, Args.LODEnd);
+	FORCEINLINE VectorRegisterUInt8x16 Burn_UInt8x16(VectorRegisterUInt8x16 A, VectorRegisterUInt8x16 B)
+	{
+		//A = _mm_xor_si128(A, _mm_set1_epi8(0xFF));
+		//
+		//__m128i A0 = _mm_slli_epi16(GetLow16_UInt8x16(A), 8);
+		//__m128i A1 = _mm_slli_epi16(GetHigh16_UInt8x16(A), 8);
 
-		BatchViewsResult.NumElems = BatchViewsResult.Result.Num() / Args.ResultBytesPerElem;
-		
-		if (Args.Base)
-		{
-			BatchViewsResult.Base = bOnlyFirstLOD
-					? Args.Base->DataStorage.GetBatchFirstLODOffset(BatchId, Args.BatchNumElems, Args.BaseBytesPerElem, Args.FirstLODOffset) 
-					: Args.Base->DataStorage.GetBatchLODRange(BatchId, Args.BatchNumElems, Args.BaseBytesPerElem, Args.LODBegin, Args.LODEnd); 
+		//__m128 FloatA0 = _mm_cvtepi32_ps(_mm_cvtepu16_epi32(A0));
+		//__m128 FloatA1 = _mm_cvtepi32_ps(_mm_cvtepu16_epi32(_mm_srli_si128(A0, 8)));
+		//__m128 FloatA2 = _mm_cvtepi32_ps(_mm_cvtepu16_epi32(A1));
+		//__m128 FloatA3 = _mm_cvtepi32_ps(_mm_cvtepu16_epi32(_mm_srli_si128(A1, 8)));
 
-			check(BatchViewsResult.NumElems == BatchViewsResult.Base.Num() / Args.BaseBytesPerElem);
-		}
+		//__m128i B0 = _mm_add_epi16(GetLow16_UInt8x16(B), _mm_set1_epi16(1));
+		//__m128i B1 = _mm_add_epi16(GetHigh16_UInt8x16(B), _mm_set1_epi16(1));
 
-		if (Args.Blend)
-		{
-			BatchViewsResult.Blend = bOnlyFirstLOD
-					? Args.Blend->DataStorage.GetBatchFirstLODOffset(BatchId, Args.BatchNumElems, Args.BlendBytesPerElem, Args.FirstLODOffset) 
-					: Args.Blend->DataStorage.GetBatchLODRange(BatchId, Args.BatchNumElems, Args.BlendBytesPerElem, Args.LODBegin, Args.LODEnd); 
+		//__m128 FloatB0 = _mm_cvtepi32_ps(_mm_cvtepu16_epi32(B0));
+		//__m128 FloatB1 = _mm_cvtepi32_ps(_mm_cvtepu16_epi32(_mm_srli_si128(B0, 8)));
+		//__m128 FloatB2 = _mm_cvtepi32_ps(_mm_cvtepu16_epi32(B1));
+		//__m128 FloatB3 = _mm_cvtepi32_ps(_mm_cvtepu16_epi32(_mm_srli_si128(B1, 8)));
 
-			check(BatchViewsResult.NumElems == BatchViewsResult.Blend.Num() / Args.BlendBytesPerElem);
-		}
+		//__m128 FloatResult0 = _mm_div_ps(FloatA0, FloatB0);
+		//__m128 FloatResult1 = _mm_div_ps(FloatA1, FloatB1);
+		//__m128 FloatResult2 = _mm_div_ps(FloatA2, FloatB2);
+		//__m128 FloatResult3 = _mm_div_ps(FloatA3, FloatB3);
 
-		if (Args.Mask)
-		{
-			BatchViewsResult.Mask = bOnlyFirstLOD
-					? Args.Mask->DataStorage.GetBatchFirstLODOffset(BatchId, Args.BatchNumElems, Args.MaskBytesPerElem, Args.FirstLODOffset) 
-					: Args.Mask->DataStorage.GetBatchLODRange(BatchId, Args.BatchNumElems, Args.MaskBytesPerElem, Args.LODBegin, Args.LODEnd); 
+		//__m128i ResultLo = _mm_packus_epi32(_mm_cvtps_epi32(FloatResult0), _mm_cvtps_epi32(FloatResult1));
+		//__m128i ResultHi = _mm_packus_epi32(_mm_cvtps_epi32(FloatResult2), _mm_cvtps_epi32(FloatResult3));
 
-			check(BatchViewsResult.NumElems == BatchViewsResult.Mask.Num() / Args.MaskBytesPerElem);
-		}
+		//return _mm_xor_si128(Saturate_UInt16x8(ResultLo, ResultHi), _mm_set1_epi8(0xFF));
 
-		return BatchViewsResult;
+		unimplemented();
+		return A;
 	}
 
-	FORCEINLINE bool IsAnyComponentLargerThan1(FVector4f Value)
+
+#endif
+	
+#if PLATFORM_ENABLE_VECTORINTRINSICS_NEON
+	using VectorRegisterUInt8x16 = int8x16_t;
+	using VectorRegisterUInt16x8 = int16x8_t;
+
+
+	FORCEINLINE void StoreUnaligned_UInt8x16(uint8* Ptr, VectorRegisterUInt8x16 Reg)
 	{
-		return (Value[0] > 1) | (Value[1] > 1) | (Value[2] > 1) | (Value[3] > 1);
+		vst1q_u8(Ptr, Reg);
 	}
 
-	/** 
-	 * Apply a blending function to an image with a colour source.
-	 * It only affects the RGB or L channels, leaving alpha untouched.
-	 */
-	template<uint32 NumChannels, uint32 (*BLEND_FUNC)(uint32, uint32), bool bClamp>
-	FORCENOINLINE void BufferLayerColourGenericChannel(uint8* DestBuf, const uint8* BaseBuf, int32 NumElems, const FIntVector& Color)
+	FORCEINLINE VectorRegisterUInt8x16 LoadUnaligned_UInt8x16(const uint8* Ptr)
 	{
-		static_assert(NumChannels > 0 && NumChannels <= 4);
+		return vld1q_u8(Ptr);
+	}	
 
-		for (int32 I = 0; I < NumElems; ++I)
-		{
-			for (uint32 C = 0; C < NumChannels; ++C)
-			{
-				uint32 Base = BaseBuf[NumChannels * I + C];
-				uint32 Result = BLEND_FUNC(Base, Color[C]);
-				if constexpr (bClamp)
-				{
-					DestBuf[NumChannels * I + C] = (uint8)FMath::Min(255u, Result);
-				}
-				else
-				{
-					DestBuf[NumChannels * I + C] = (uint8)Result;
-				}
-			}
-		}
-	}
+	FORCEINLINE VectorRegisterUInt16x8 LoadUnaligned_UInt16x8(const uint16* Ptr)
+	{
+		return vld1q_u16(Ptr);
+	}	
 
-	template<uint32 NumChannels, uint32 (*BLEND_FUNC)(uint32,uint32), bool Clamp>
-	FORCENOINLINE void BufferLayerColourFromAlphaGenericChannel(uint8* DestBuf, const uint8* BaseBuf, int32 NumElems, const FIntVector& Color)
+	FORCEINLINE VectorRegisterUInt16x8 GetLow16_UInt8x16(VectorRegisterUInt8x16 Reg)
 	{
-		static_assert(NumChannels > 0 && NumChannels <= 4);
+		return vmovl_u8(vget_low_u8(Reg));
+	}	
 
-		for (int32 I = 0; I < NumElems; ++I)
-		{
-			const uint32 Alpha = Invoke([&]() -> uint32
-			{
-				if constexpr (NumChannels <= 3)
-				{
-					return 255;
-				}
-				else
-				{
-					return BaseBuf[NumChannels * I + 3];
-				}
-			});
+	FORCEINLINE VectorRegisterUInt16x8 GetHigh16_UInt8x16(VectorRegisterUInt8x16 Reg)
+	{
+		return vmovl_u8(vget_high_u8(Reg));
+	}
 
-			for (uint32 C = 0; C < NumChannels; ++C)
-			{
-				uint32 Result = BLEND_FUNC(Alpha, Color[C]);
-				if constexpr (Clamp)
-				{
-					DestBuf[NumChannels * I + C] = (uint8)FMath::Min(255u, Result);
-				}
-				else
-				{
-					DestBuf[NumChannels * I + C] = (uint8)Result;
-				}
-			}
-		}
+	FORCEINLINE VectorRegisterUInt8x16 Saturate_UInt16x8(VectorRegisterUInt16x8 Lo, VectorRegisterUInt16x8 Hi)
+	{
+		return vcombine_u8(vqmovn_u16(Lo), vqmovn_u16(Hi)); 
 	}
-} // namespace OpImageLayerInternal
 
-	template< uint32 (*BLEND_FUNC)(uint32,uint32), bool CLAMP >
-    FORCENOINLINE void BufferLayerColourImpl(FImage* ResultImage, const FImage* BaseImage, FVector4f Color)
+	FORCEINLINE VectorRegisterUInt8x16 Mult255Unorm_UInt8x16(VectorRegisterUInt8x16 A, VectorRegisterUInt8x16 B)
 	{
-        check(ResultImage->GetFormat() == BaseImage->GetFormat());
-        check(ResultImage->GetSizeX() == BaseImage->GetSizeX());
-        check(ResultImage->GetSizeY() == BaseImage->GetSizeY());
-        check(ResultImage->GetLODCount() == BaseImage->GetLODCount());
+		uint16x8_t MultLo = vmull_u8(vget_low_u8(A), vget_low_u8(B));
+		uint16x8_t MultHi = vmull_u8(vget_high_u8(A), vget_high_u8(B));
 
-		const EImageFormat BaseFormat = BaseImage->GetFormat();
-		// Generic implementation
-		const int32 BytesPerElem = GetImageFormatData(BaseFormat).BytesPerBlock;	
+		uint8x8_t ResultLo = vqrshrn_n_u16(vrsraq_n_u16(MultLo, MultLo, 8), 8);
+		uint8x8_t ResultHi = vqrshrn_n_u16(vrsraq_n_u16(MultHi, MultHi, 8), 8);
 
-		constexpr int32 BatchNumElems = 4096*4;
-		OpImageLayerInternal::FOpLayerBatchArgs BatchArgs;
-		BatchArgs.BatchNumElems      = BatchNumElems;
-		BatchArgs.Result             = ResultImage;
-		BatchArgs.Base               = BaseImage;
-		BatchArgs.BaseBytesPerElem   = BytesPerElem;
-		BatchArgs.ResultBytesPerElem = BytesPerElem;
+		return vcombine_u8(ResultLo, ResultHi);
+	}
 
-		const int32 NumBatches = OpImageLayerInternal::GetOpLayerNumBatches(BatchArgs);
+	FORCEINLINE VectorRegisterUInt8x16 Mult255Unorm_UInt16x8(VectorRegisterUInt16x8 A0, VectorRegisterUInt16x8 A1, VectorRegisterUInt16x8 B0, VectorRegisterUInt16x8 B1)
+	{
+		unimplemented();
+		return Saturate_UInt16x8(A0, A1);
+	}
 
-		const FIntVector ColorValue = FIntVector(Color.X * 255.0f, Color.Y * 255.0f, Color.Z * 255.0f);
-		ParallelExecutionUtils::InvokeBatchParallelFor(NumBatches,
-		[
-			&BatchArgs, ColorValue, BaseFormat
-		](int32 BatchId)
-		{
-			OpImageLayerInternal::FOpLayerBatchViews BatchViews = 
-					OpImageLayerInternal::GetOpLayerBatchViews(BatchId, BatchArgs);
-	
-			switch (BaseFormat)
-			{
-			case EImageFormat::L_UByte:
-			{
-				check(BatchArgs.ResultBytesPerElem == 1);
-				OpImageLayerInternal::BufferLayerColourGenericChannel<1, BLEND_FUNC, CLAMP>(
-						BatchViews.Result.GetData(), BatchViews.Base.GetData(), BatchViews.NumElems, ColorValue);
-				break;
-			}
-			case EImageFormat::RGB_UByte:
-			{
-				check(BatchArgs.ResultBytesPerElem == 3);
-				OpImageLayerInternal::BufferLayerColourGenericChannel<3, BLEND_FUNC, CLAMP>(
-						BatchViews.Result.GetData(), BatchViews.Base.GetData(), BatchViews.NumElems, ColorValue);
-				break;
-			}
-			case EImageFormat::RGBA_UByte:
-			{
-				check(BatchArgs.ResultBytesPerElem == 4);
-				OpImageLayerInternal::BufferLayerColourGenericChannel<4, BLEND_FUNC, CLAMP>(
-						BatchViews.Result.GetData(), BatchViews.Base.GetData(), BatchViews.NumElems, ColorValue);
-				break;
-			}
-			case EImageFormat::BGRA_UByte:
-			{
-				check(BatchArgs.ResultBytesPerElem == 4);
-				FIntVector BGRAColorValue = FIntVector(ColorValue.Z, ColorValue.Y, ColorValue.X);
-				OpImageLayerInternal::BufferLayerColourGenericChannel<4, BLEND_FUNC, CLAMP>(
-						BatchViews.Result.GetData(), BatchViews.Base.GetData(), BatchViews.NumElems, ColorValue);
-				break;
-			}
-			default:
-			{
-				checkf(false, TEXT("Unsupported format."));
-				break;
-			}
-			}
-		});
-    }
+	FORCEINLINE VectorRegisterUInt8x16 Screen_UInt8x16(VectorRegisterUInt8x16 A, VectorRegisterUInt8x16 B)
+	{
+		uint8x8_t ALo = vget_low_u8(A);
+		uint8x8_t AHi = vget_high_u8(A);
 
+		uint8x8_t BLo = vget_low_u8(B);
+		uint8x8_t BHi = vget_high_u8(B);
 
-	template< uint32 (*BLEND_FUNC)(uint32, uint32) >
-	void BufferLayerColour(FImage* ResultImage, const FImage* BaseImage, FVector4f Color)
-	{
-		bool bIsClampNeeded = OpImageLayerInternal::IsAnyComponentLargerThan1(Color);
-		if (bIsClampNeeded)
-		{
-			BufferLayerColourImpl<BLEND_FUNC, true>(ResultImage, BaseImage, Color);
-		}
-		else
-		{
-			BufferLayerColourImpl<BLEND_FUNC, false>(ResultImage, BaseImage, Color);
-		}
+		uint16x8_t MultLo = vmull_u8(ALo, BLo);
+		uint16x8_t MultHi = vmull_u8(AHi, BHi);
+
+		uint16x8_t ProdLo = vmovl_u8(vqrshrn_n_u16(vrsraq_n_u16(MultLo, MultLo, 8), 8));
+		uint16x8_t ProdHi = vmovl_u8(vqrshrn_n_u16(vrsraq_n_u16(MultHi, MultHi, 8), 8));
+		
+		uint16x8_t ResultLo = vsubq_u16(vaddl_u8(ALo, BLo), ProdLo);
+		uint16x8_t ResultHi = vsubq_u16(vaddl_u8(AHi, BHi), ProdHi);
+
+		return vcombine_u8(vmovn_u16(ResultLo), vmovn_u16(ResultHi));
 	}
 
-	template<
-		uint32 (*BLEND_FUNC)(uint32, uint32), 
-		bool CLAMP>
-    FORCENOINLINE void BufferLayerColourFromAlpha(FImage* ResultImage, const FImage* BaseImage, FVector4f Color)
+	FORCEINLINE VectorRegisterUInt8x16 SoftLight_UInt8x16(VectorRegisterUInt8x16 A, VectorRegisterUInt8x16 B)
 	{
-        check(ResultImage->GetFormat() == BaseImage->GetFormat());
-        check(ResultImage->GetSizeX() == BaseImage->GetSizeX());
-        check(ResultImage->GetSizeY() == BaseImage->GetSizeY());
-        check(ResultImage->GetLODCount() == BaseImage->GetLODCount());
+		uint8x8_t ALo = vget_low_u8(A);
+		uint8x8_t AHi = vget_high_u8(A);
 
-		const EImageFormat BaseFormat = BaseImage->GetFormat();
-		// Generic implementation
-		constexpr int32 BatchNumElems = 4096*2;
-		const int32 BytesPerElem = GetImageFormatData(BaseFormat).BytesPerBlock;	
+		uint8x8_t BLo = vget_low_u8(B);
+		uint8x8_t BHi = vget_high_u8(B);
 
-		OpImageLayerInternal::FOpLayerBatchArgs BatchArgs;
-		BatchArgs.BatchNumElems      = BatchNumElems;
-		BatchArgs.Result             = ResultImage;
-		BatchArgs.Base               = BaseImage;
-		BatchArgs.BaseBytesPerElem   = BytesPerElem;
-		BatchArgs.ResultBytesPerElem = BytesPerElem;
+		uint16x8_t MultLo = vmull_u8(ALo, BLo);
+		uint16x8_t MultHi = vmull_u8(AHi, BHi);
 
-		const int32 NumBatches = OpImageLayerInternal::GetOpLayerNumBatches(BatchArgs);
+		uint8x8_t ProdLo = vqrshrn_n_u16(vrsraq_n_u16(MultLo, MultLo, 8), 8);
+		uint8x8_t ProdHi = vqrshrn_n_u16(vrsraq_n_u16(MultHi, MultHi, 8), 8);
+		
+		uint8x8_t ScreenLo = vmovn_u16(vsubq_u16(vaddl_u8(ALo, BLo), vmovl_u8(ProdLo)));
+		uint8x8_t ScreenHi = vmovn_u16(vsubq_u16(vaddl_u8(AHi, BHi), vmovl_u8(ProdHi)));
 
-		const FIntVector ColorValue = FIntVector(Color.X * 255.0f, Color.Y * 255.0f, Color.Z * 255.0f);
-		ParallelExecutionUtils::InvokeBatchParallelFor(NumBatches,
-		[
-			&BatchArgs, ColorValue, BaseFormat
-		](int32 BatchId)
-		{
-			OpImageLayerInternal::FOpLayerBatchViews BatchViews = 
-					OpImageLayerInternal::GetOpLayerBatchViews(BatchId, BatchArgs);
-	
-			switch (BaseFormat)
-			{
-			case EImageFormat::L_UByte:
-			{
-				check(BatchArgs.ResultBytesPerElem == 1);
-				OpImageLayerInternal::BufferLayerColourFromAlphaGenericChannel<1, BLEND_FUNC, CLAMP>(
-						BatchViews.Result.GetData(), BatchViews.Base.GetData(), BatchViews.NumElems, ColorValue);
-				break;
-			}
-			case EImageFormat::RGB_UByte:
-			{
-				check(BatchArgs.ResultBytesPerElem == 3);
-				OpImageLayerInternal::BufferLayerColourFromAlphaGenericChannel<3, BLEND_FUNC, CLAMP>(
-						BatchViews.Result.GetData(), BatchViews.Base.GetData(), BatchViews.NumElems, ColorValue);
-				break;
-			}
-			case EImageFormat::RGBA_UByte:
-			{
-				check(BatchArgs.ResultBytesPerElem == 4);
-				OpImageLayerInternal::BufferLayerColourFromAlphaGenericChannel<4, BLEND_FUNC, CLAMP>(
-						BatchViews.Result.GetData(), BatchViews.Base.GetData(), BatchViews.NumElems, ColorValue);
-				break;
-			}
-			case EImageFormat::BGRA_UByte:
-			{
-				check(BatchArgs.ResultBytesPerElem == 4);
-				FIntVector BGRAColorValues = FIntVector(ColorValue.Z, ColorValue.Y, ColorValue.X);
-				OpImageLayerInternal::BufferLayerColourFromAlphaGenericChannel<4, BLEND_FUNC, CLAMP>(
-						BatchViews.Result.GetData(), BatchViews.Base.GetData(), BatchViews.NumElems, ColorValue);
-				break;
-			}
-			default:
-			{
-				checkf(false, TEXT("Unsupported format."));
-				break;
-			}
-			}
-		});
-    }
+		uint8x16_t OneMinusA = veorq_u8(A, vdupq_n_u8(0xFF)); 
+
+		uint16x8_t ScreenTimesALo = vmull_u8(ScreenLo, ALo);
+		uint16x8_t ScreenTimesAHi = vmull_u8(ScreenHi, AHi);
 
+		uint16x8_t AddLo = vmlal_u8(ScreenTimesALo, ProdLo, vget_low_u8(OneMinusA));
+		uint16x8_t AddHi = vmlal_u8(ScreenTimesAHi, ProdHi, vget_high_u8(OneMinusA));
 
-	template<uint32 (*BLEND_FUNC)(uint32, uint32)>
-	void BufferLayerColourFromAlpha(FImage* ResultImage, const FImage* BaseImage, FVector4f Color)
+		uint8x8_t ResultLo = vqrshrn_n_u16(vrsraq_n_u16(AddLo, AddLo, 8), 8);
+		uint8x8_t ResultHi = vqrshrn_n_u16(vrsraq_n_u16(AddLo, AddHi, 8), 8);
+
+		return vcombine_u8(ResultLo, ResultHi);
+	}
+
+	FORCEINLINE VectorRegisterUInt8x16 HardLight_UInt8x16(VectorRegisterUInt8x16 A, VectorRegisterUInt8x16 B)
 	{
-		bool bIsClampNeeded = OpImageLayerInternal::IsAnyComponentLargerThan1(Color);
-		if (bIsClampNeeded)
-		{
-			BufferLayerColourFromAlpha<BLEND_FUNC, true>(ResultImage, BaseImage, Color);
-		}
-		else
-		{
-			BufferLayerColourFromAlpha<BLEND_FUNC, false>(ResultImage, BaseImage, Color);
-		}
+		uint8x16_t Xor = vcltzq_s8(B);
+
+		uint8x16_t AXor = veorq_u8(A, Xor);
+		uint8x16_t TwoTimesBMod255Xor = veorq_u8(vaddq_u8(B, B), Xor);
+
+		uint16x8_t MultLo = vmull_u8(vget_low_u8(AXor), vget_low_u8(TwoTimesBMod255Xor));
+		uint16x8_t MultHi = vmull_u8(vget_high_u8(AXor), vget_high_u8(TwoTimesBMod255Xor));
+
+		uint8x8_t ResultLo = vqrshrn_n_u16(vrsraq_n_u16(MultLo, MultLo, 8), 8);
+		uint8x8_t ResultHi = vqrshrn_n_u16(vrsraq_n_u16(MultHi, MultHi, 8), 8);
+
+		return veorq_u8(vcombine_u8(ResultLo, ResultHi), Xor);
 	}
 
+	FORCEINLINE VectorRegisterUInt8x16 Overlay_UInt8x16(VectorRegisterUInt8x16 A, VectorRegisterUInt8x16 B)
+	{
+		uint8x16_t Xor = vcltzq_s8(A);
 
-	template< 
-		uint32 (*BLEND_FUNC_MASKED)(uint32, uint32, uint32),
-		uint32 (*BLEND_FUNC)(uint32, uint32),
-		bool CLAMP,
-		// Number of total channels to actually process
-		uint32 CHANNELS_TO_BLEND,
-		// Number of total channels in the base image
-		int32 BASE_CHANNEL_STRIDE>
-    void BufferLayerColourFormat(FImage* DestImage, const FImage* BaseImage, const FImage* MaskImage, FVector4f Col, uint32 BaseChannelOffset, uint8 ColorChannelOffset, bool bOnlyFirstLOD)
-	{
-		check(CHANNELS_TO_BLEND + BaseChannelOffset <= BASE_CHANNEL_STRIDE);
-		check(DestImage->GetLODCount() <= BaseImage->GetLODCount());
-
-		FUintVector4 TopColor = FUintVector4(Col.X * 255.0f, Col.Y * 255.0f, Col.Z * 255.0f, Col.W * 255);
-
-		const EImageFormat BaseFormat = BaseImage->GetFormat();
-		const EImageFormat MaskFormat = MaskImage->GetFormat();
-
-		int32 NumLODs = bOnlyFirstLOD ? 1 : BaseImage->GetLODCount();
-
-        const bool bIsMaskUncompressed = (MaskFormat == EImageFormat::L_UByte);
-
-		constexpr uint32 NumColorChannels = FMath::Min(CHANNELS_TO_BLEND, 3u);
-		check(NumColorChannels + ColorChannelOffset < 4);
-
-        if (bIsMaskUncompressed)
-        {
-			const int32 BytesPerElem = GetImageFormatData(BaseFormat).BytesPerBlock;
-			check(GetImageFormatData(MaskFormat).BytesPerBlock == 1);
-			check(GetImageFormatData(DestImage->GetFormat()).BytesPerBlock == BytesPerElem);
-
-			constexpr int32 BatchNumElems = 4096 * 2;
-			OpImageLayerInternal::FOpLayerBatchArgs BatchArgs;
-			BatchArgs.BatchNumElems      = BatchNumElems; 
-			BatchArgs.LODBegin           = 0;
-			BatchArgs.LODEnd             = NumLODs;
-			BatchArgs.Result             = DestImage;
-			BatchArgs.Base               = BaseImage;
-			BatchArgs.Mask               = MaskImage;
-			BatchArgs.ResultBytesPerElem = BytesPerElem;
-			BatchArgs.BaseBytesPerElem   = BytesPerElem;
-			BatchArgs.MaskBytesPerElem   = 1;
-
-			const int32 NumBatches = OpImageLayerInternal::GetOpLayerNumBatchesLODRange(BatchArgs);
-
-			ParallelExecutionUtils::InvokeBatchParallelFor(NumBatches,
-			[
-				&BatchArgs, BaseChannelOffset, TopColor, ColorChannelOffset, NumColorChannels
-			](int32 BatchId)
-			{
-				OpImageLayerInternal::FOpLayerBatchViews BatchViews = 
-						OpImageLayerInternal::GetOpLayerBatchLODRangeViews(BatchId, BatchArgs);
+		uint8x16_t BXor = veorq_u8(B, Xor);
+		uint8x16_t TwoTimesAMod255Xor = veorq_u8(vaddq_u8(A, A), Xor);
 
-				// This could happen in case of missing data files.
-				if (!BatchViews.Base.GetData() || !BatchViews.Mask.GetData() || !BatchViews.Result.GetData())
-				{
-					return;
-				}
+		uint16x8_t MultLo = vmull_u8(vget_low_u8(BXor), vget_low_u8(TwoTimesAMod255Xor));
+		uint16x8_t MultHi = vmull_u8(vget_high_u8(BXor), vget_high_u8(TwoTimesAMod255Xor));
 
-				const uint8* BaseBuf = BatchViews.Base.GetData() + BaseChannelOffset;
-				const uint8* MaskBuf = BatchViews.Mask.GetData();
-				uint8* DestBuf       = BatchViews.Result.GetData() + BaseChannelOffset;
+		uint8x8_t ResultLo = vqrshrn_n_u16(vrsraq_n_u16(MultLo, MultLo, 8), 8);
+		uint8x8_t ResultHi = vqrshrn_n_u16(vrsraq_n_u16(MultHi, MultHi, 8), 8);
 
-				for (int32 I = 0; I < BatchViews.NumElems; ++I)
-				{
-					uint32 MaskData = MaskBuf[I];
-					for (int32 C = 0; C < NumColorChannels; ++C)
-					{
-						const uint32 Base = BaseBuf[BASE_CHANNEL_STRIDE*I + C];
-						const uint32 Result = BLEND_FUNC_MASKED(Base, TopColor[C + ColorChannelOffset], MaskData);
-						if constexpr (CLAMP)
-						{
-							DestBuf[BASE_CHANNEL_STRIDE*I + C] = (uint8)FMath::Min(255u, Result);
-						}
-						else
-						{
-							DestBuf[BASE_CHANNEL_STRIDE*I + C] = (uint8)Result;
-						}
-					}
+		return veorq_u8(vcombine_u8(ResultLo, ResultHi), Xor);
+	}
 
-					constexpr bool bIsNC4 = (BASE_CHANNEL_STRIDE == 4);
-					if constexpr (bIsNC4)
-					{
-						DestBuf[BASE_CHANNEL_STRIDE*I + 3] = BaseBuf[BASE_CHANNEL_STRIDE*I + 3];
-					}
-				}
-			});
-        }
-        else if (MaskFormat == EImageFormat::L_UByteRLE)
-        {
-            int32 Rows = BaseImage->GetSizeY();
-            int32 Width = BaseImage->GetSizeX();
-
-            for (int32 Lod = 0; Lod < NumLODs; ++Lod)
-            {
-                const uint8* BaseBuf = BaseImage->GetLODData(Lod);
-                const uint8* MaskBuf = MaskImage->GetLODData(Lod);
-				uint8* DestBuf = DestImage->GetLODData(Lod);
-
-				// This could happen in case of missing data files.
-				if (!BaseBuf || !MaskBuf || !DestBuf)
-				{
-					continue;
-				}
+	FORCEINLINE VectorRegisterUInt8x16 Dodge_UInt8x16(VectorRegisterUInt8x16 A, VectorRegisterUInt8x16 B)
+	{
+		unimplemented();
+		return A;
+	}
 
-				// Remove RLE header, mip size and row sizes.
-				MaskBuf += sizeof(uint32);
-				MaskBuf += Rows * sizeof(uint32);
+	FORCEINLINE VectorRegisterUInt8x16 Burn_UInt8x16(VectorRegisterUInt8x16 A, VectorRegisterUInt8x16 B)
+	{
+		unimplemented();
+		return A;
+	}
 
-                for (int32 RowIndex = 0; RowIndex < Rows; ++RowIndex)
-                {
-                    const uint8* DestRowEnd = DestBuf + Width * BASE_CHANNEL_STRIDE;
-                    while (DestBuf != DestRowEnd)
-                    {
-                        // Decode header
-						uint16 Equal = 0;
-						FMemory::Memmove(&Equal, MaskBuf, sizeof(uint16));
-                        MaskBuf += 2;
+#endif
 
-                        uint8 Different = *MaskBuf;
-                        ++MaskBuf;
-
-                        uint8 EqualPixel = *MaskBuf;
-                        ++MaskBuf;
-
-                        // Equal pixels
-						//check(DestBuf + BASE_CHANNEL_STRIDE * Equal <= BaseImage->GetDataSize(Lod));
-                        if (EqualPixel == 255)
-                        {
-                            for (int32 I = 0; I < Equal; ++I)
-                            {
-                                for (int32 C = 0; C < NumColorChannels; ++C)
-                                {
-                                    uint32 Base = BaseBuf[BASE_CHANNEL_STRIDE * I + C];
-                                    uint32 Result = BLEND_FUNC(Base, TopColor[C + ColorChannelOffset]);
-                                    if constexpr (CLAMP)
-                                    {
-                                        DestBuf[BASE_CHANNEL_STRIDE * I + C] = (uint8)FMath::Min(255u, Result);
-                                    }
-                                    else
-                                    {
-                                        DestBuf[BASE_CHANNEL_STRIDE * I + C] = (uint8)Result;
-                                    }
-                                }
-
-                                constexpr bool bIsNC4 = (BASE_CHANNEL_STRIDE == 4);
-                                if (bIsNC4)
-                                {
-                                    DestBuf[BASE_CHANNEL_STRIDE * I + 3] = BaseBuf[BASE_CHANNEL_STRIDE * I + 3];
-                                }
-                            }
-                        }
-                        else if (EqualPixel > 0)
-                        {
-                            for (int32 I = 0; I < Equal; ++I)
-                            {
-                                for (int32 C = 0; C < NumColorChannels; ++C)
-                                {
-                                    uint32 Base = BaseBuf[BASE_CHANNEL_STRIDE * I + C];
-                                    uint32 Result = BLEND_FUNC_MASKED(Base, TopColor[C + ColorChannelOffset], EqualPixel);
-                                    if constexpr (CLAMP)
-                                    {
-                                        DestBuf[BASE_CHANNEL_STRIDE * I + C] = (uint8)FMath::Min(255u, Result);
-                                    }
-                                    else
-                                    {
-                                        DestBuf[BASE_CHANNEL_STRIDE * I + C] = (uint8)Result;
-                                    }
-                                }
-
-                                constexpr bool bIsNC4 = (BASE_CHANNEL_STRIDE == 4);
-                                if (bIsNC4)
-                                {
-                                    DestBuf[BASE_CHANNEL_STRIDE * I + 3] = BaseBuf[BASE_CHANNEL_STRIDE * I + 3];
-                                }
-                            }
-                        }
-                        else
-                        {
-                            // It could happen if xxxxxOnBase
-                            if (DestBuf != BaseBuf)
-                            {
-                                FMemory::Memmove(DestBuf, BaseBuf, BASE_CHANNEL_STRIDE*Equal);
-                            }
-                        }
-
-                        DestBuf += BASE_CHANNEL_STRIDE * Equal;
-                        BaseBuf += BASE_CHANNEL_STRIDE * Equal;
-
-                        // Different pixels
-						//check(DestBuf + BASE_CHANNEL_STRIDE * Different <= StartDestBuf + BaseImage->GetDataSize(Lod));
-                        for (int32 I = 0; I < Different; ++I)
-                        {
-                            for (int32 C = 0; C < NumColorChannels; ++C)
-                            {
-                                uint32 Mask = MaskBuf[I];
-                                uint32 Base = BaseBuf[BASE_CHANNEL_STRIDE * I + C];
-                                uint32 Result = BLEND_FUNC_MASKED(Base, TopColor[C + ColorChannelOffset], Mask);
-                                if constexpr (CLAMP)
-                                {
-                                    DestBuf[BASE_CHANNEL_STRIDE * I + C] = (uint8)FMath::Min(255u, Result);
-                                }
-                                else
-                                {
-                                    DestBuf[BASE_CHANNEL_STRIDE * I + C] = (uint8)Result;
-                                }
-                            }
-
-                            constexpr bool bIsNC4 = (BASE_CHANNEL_STRIDE == 4);
-                            if (bIsNC4)
-                            {
-                                DestBuf[BASE_CHANNEL_STRIDE * I + 3] = BaseBuf[BASE_CHANNEL_STRIDE * I + 3];
-                            }
-                        }
-
-                        DestBuf += BASE_CHANNEL_STRIDE * Different;
-                        BaseBuf += BASE_CHANNEL_STRIDE * Different;
-                        MaskBuf += Different;
-                    }
-                }
+	FORCEINLINE uint8 Saturate_U16(uint16 V)
+	{
+		return (uint8)FMath::Min<uint16>(V, 255u);
+	}
+
+	FORCEINLINE uint8 Mult255Unorm_U8(uint8 A, uint8 B)
+	{
+		uint32 Value = A * B + 128;
+		return (Value + (Value >> 8)) >> 8;
+	}
 
-                Rows = FMath::DivideAndRoundUp(Rows, 2);
-                Width = FMath::DivideAndRoundUp(Width, 2);
-            }
-        }
-        else
-        {
-            checkf( false, TEXT("Unsupported mask format.") );
-        }
+	/* Multply 16 bit 255 normalized values and saturate */
+	FORCEINLINE uint8 Mult255Unorm_U16(uint16 A, uint16 B)
+	{
+		uint32 Value = A * B + 128;
+		return Saturate_U16(uint16((Value + (Value >> 8)) >> 8));
 	}
 
+	FORCEINLINE uint8 Lerp_U8(uint8 A, uint8 B, uint8 Factor)
+	{
+		uint32 Value = (255 - Factor)*A + Factor*B + 128;
+		return (Value + (Value >> 8)) >> 8; 
+	}
 
-	/**
-	* Apply a blending function to an image with a colour source and a mask
-	*/
-	template<
-		uint32 (*BLEND_FUNC_MASKED)(uint32,uint32,uint32),
-		uint32 (*BLEND_FUNC)(uint32,uint32),
-		bool CLAMP>
-    FORCENOINLINE void BufferLayerColourImpl(FImage* ResultImage, const FImage* BaseImage, const FImage* MaskImage, FVector4f Col)
-	{
-		check(BaseImage->GetSizeX() == MaskImage->GetSizeX());
-		check(BaseImage->GetSizeY() == MaskImage->GetSizeY());
-		check(MaskImage->GetFormat() == EImageFormat::L_UByte ||
-			  MaskImage->GetFormat() == EImageFormat::L_UByteRLE);
+	FORCEINLINE uint8 Screen_U8(uint8 A, uint8 B)
+	{
+		return (A + B) - Mult255Unorm_U8(A, B);
+	}
+
+	FORCEINLINE uint8 SoftLight_U8(uint8 A, uint8 B)
+	{
+		uint8 Prod = Mult255Unorm_U8(A, B);
+		return Lerp_U8(Prod, (A + B) - Prod, A);
+	}
+
+	FORCEINLINE uint8 HardLight_U8(uint8 A, uint8 B)
+	{
+		uint8 Xor = (B >= 128) ? 255 : 0;
+		return Xor ^ Mult255Unorm_U8(A ^ Xor, ((B << 1) & 0xFF) ^ Xor);
+	}
+
+	FORCEINLINE uint8 Overlay_U8(uint8 A, uint8 B)
+	{
+		uint8 Xor = (A >= 128) ? 255 : 0;
+		return Xor ^ Mult255Unorm_U8(((A << 1) & 0xFF) ^ Xor, B ^ Xor);
+	}
+
+	FORCEINLINE uint8 Dodge_U8(uint8 A, uint8 B)
+	{
+		return Saturate_U16(uint16((A << 8) / (256 - B)));
+	}
+
+	FORCEINLINE uint8 Burn_U8(uint8 A, uint8 B)
+	{
+		return 255 - Saturate_U16(uint16(((255 - A) << 8) / (B + 1)));
+	}
+//#define UE_MUTABLE_IMAGE_LAYER_USE_VECTOR_IMPL 0
+#ifndef UE_MUTABLE_IMAGE_LAYER_USE_VECTOR_IMPL
+	#define UE_MUTABLE_IMAGE_LAYER_USE_VECTOR_IMPL (PLATFORM_ENABLE_VECTORINTRINSICS || PLATFORM_ENABLE_VECTORINTRINSICS_NEON)
+#endif
+
+	void Blend(uint8* Dest, const uint8*, const uint8* Blen, int32 NumElems)
+	{
+		FMemory::Memcpy(Dest, Blen, NumElems);
+	}
+
+	void BlendWithWord(uint8* Dest, const uint8*, const uint16* Blen, int32 NumElems)
+	{
+		int32 ElemIndex = 0;
+
+#if UE_MUTABLE_IMAGE_LAYER_USE_VECTOR_IMPL
+		int32 NumIters = NumElems / 16;
 
-		const bool bValid = (BaseImage->GetSizeX() == MaskImage->GetSizeX()) &&
-						    (BaseImage->GetSizeY() == MaskImage->GetSizeY()) &&
-							(MaskImage->GetFormat() == EImageFormat::L_UByte || MaskImage->GetFormat() == EImageFormat::L_UByteRLE);
-		if (!bValid)
+		for (int32 Iter = 0; Iter < NumIters; ++Iter)
 		{
-			return;
+			VectorRegisterUInt8x16 BLo = LoadUnaligned_UInt16x8(Blen + Iter*16);
+			VectorRegisterUInt8x16 BHi = LoadUnaligned_UInt16x8(Blen + Iter*16 + 8);
+
+			VectorRegisterUInt8x16 Result = Saturate_UInt16x8(BLo, BHi);
+
+			StoreUnaligned_UInt8x16(Dest + Iter*16, Result);
 		}
 
-		EImageFormat BaseFormat = BaseImage->GetFormat();
-		if (BaseFormat == EImageFormat::RGB_UByte)
+		ElemIndex = NumIters*16;
+#endif
+
+		for (; ElemIndex < NumElems; ++ElemIndex)
 		{
-			BufferLayerColourFormat<BLEND_FUNC_MASKED, BLEND_FUNC, CLAMP, 3, 3>(ResultImage, BaseImage, MaskImage, Col, 0, 0, false);
+			Dest[ElemIndex] = Saturate_U16(Blen[ElemIndex]);
 		}
-        else if (BaseFormat == EImageFormat::RGBA_UByte)
-        {
-            BufferLayerColourFormat<BLEND_FUNC_MASKED, BLEND_FUNC, CLAMP, 3, 4>(ResultImage, BaseImage, MaskImage, Col, 0, 0, false);
-        }
-        else if (BaseFormat == EImageFormat::BGRA_UByte)
-        {
-            float Temp = Col[0];
-            Col[0] = Col[2];
-            Col[2] = Temp;
-            BufferLayerColourFormat<BLEND_FUNC_MASKED, BLEND_FUNC, CLAMP, 3, 4>(ResultImage, BaseImage, MaskImage, Col, 0, 0, false);
-        }
-        else if (BaseFormat == EImageFormat::L_UByte)
+	}
+
+	void Screen(uint8* Dest, const uint8* Base, const uint8* Blen, int32 NumElems)
+	{
+		int32 ElemIndex = 0;
+
+#if UE_MUTABLE_IMAGE_LAYER_USE_VECTOR_IMPL
+		int32 NumIters = NumElems / 16;
+
+		for (int32 Iter = 0; Iter < NumIters; ++Iter)
 		{
-			BufferLayerColourFormat<BLEND_FUNC_MASKED, BLEND_FUNC, CLAMP, 1, 1>(ResultImage, BaseImage, MaskImage, Col, 0, 0, false);
+			VectorRegisterUInt8x16 A = LoadUnaligned_UInt8x16(Base + Iter*16);
+			VectorRegisterUInt8x16 B = LoadUnaligned_UInt8x16(Blen + Iter*16);
+
+			VectorRegisterUInt8x16 Result = Screen_UInt8x16(A, B);
+
+			StoreUnaligned_UInt8x16(Dest + Iter*16, Result);
 		}
-		else
+
+		ElemIndex = NumIters*16;
+#endif
+
+		for (; ElemIndex < NumElems; ++ElemIndex)
 		{
-			checkf(false, TEXT("Unsupported format."));
+			Dest[ElemIndex] = Screen_U8(Base[ElemIndex], Blen[ElemIndex]);
 		}
 	}
 
-
-	template< 
-		uint32 (*BLEND_FUNC_MASKED)(uint32, uint32, uint32),
-		uint32 (*BLEND_FUNC)(uint32, uint32)>
-	void BufferLayerColour(FImage* DestImage, const FImage* BaseImage, const FImage* MaskImage, FVector4f Col)
+	void ScreenWithWord(uint8* Dest, const uint8* Base, const uint16* Blen, int32 NumElems)
 	{
-		bool bIsClampNeeded = OpImageLayerInternal::IsAnyComponentLargerThan1(Col);
-		if (bIsClampNeeded)
+		int32 ElemIndex = 0;
+
+#if UE_MUTABLE_IMAGE_LAYER_USE_VECTOR_IMPL
+		int32 NumIters = NumElems / 16;
+
+		for (int32 Iter = 0; Iter < NumIters; ++Iter)
 		{
-			BufferLayerColourImpl<BLEND_FUNC_MASKED, BLEND_FUNC, true>(DestImage, BaseImage, MaskImage, Col);
+			VectorRegisterUInt8x16 BLo = LoadUnaligned_UInt16x8(Blen + Iter*16);
+			VectorRegisterUInt8x16 BHi = LoadUnaligned_UInt16x8(Blen + Iter*16 + 8);
+
+			VectorRegisterUInt8x16 A = LoadUnaligned_UInt8x16(Base + Iter*16);
+			VectorRegisterUInt8x16 B = Saturate_UInt16x8(BLo, BHi);
+
+			VectorRegisterUInt8x16 Result = Screen_UInt8x16(A, B);
+
+			StoreUnaligned_UInt8x16(Dest + Iter*16, Result);
 		}
-		else
+
+		ElemIndex = NumIters*16;
+#endif
+
+		for (; ElemIndex < NumElems; ++ElemIndex)
 		{
-			BufferLayerColourImpl<BLEND_FUNC_MASKED, BLEND_FUNC, false>(DestImage, BaseImage, MaskImage, Col);
+			Dest[ElemIndex] = Screen_U8(Base[ElemIndex], Saturate_U16(Blen[ElemIndex]));
 		}
 	}
 
-
-	/**
-	* Apply a blending function to an image with another image as blending layer
-	*/
-	template< 
-		uint32 (*BLEND_FUNC)(uint32, uint32), 
-		bool CLAMP,
-		// Number of total channels to actually process
-		uint32 CHANNELS_TO_BLEND,
-		// Number of total channels in the base image
-		int32 BASE_CHANNEL_STRIDE>
-	void BufferLayerColourFormatInPlace(FImage* BaseImage, FVector4f Color, uint32 BaseChannelOffset, uint8 ColorChannelOffset, bool bOnlyFirstLOD)
+	void SoftLight(uint8* Dest, const uint8* Base, const uint8* Blen, int32 NumElems)
 	{
-		FUintVector4 TopColor = FUintVector4(Color.X * 255.0f, Color.Y * 255.0f, Color.Z * 255.0f, Color.W * 255.0f);
+		int32 ElemIndex = 0;
 
-		int32 NumLODs = bOnlyFirstLOD ? 1 : BaseImage->GetLODCount();
+#if 0//UE_MUTABLE_IMAGE_LAYER_USE_VECTOR_IMPL
+		int32 NumIters = NumElems / 16;
 
-		const int32 BytesPerElem = GetImageFormatData(BaseImage->GetFormat()).BytesPerBlock;
-		
-		constexpr int32 BatchNumElems = 4096*2;
-		OpImageLayerInternal::FOpLayerBatchArgs BatchArgs;
-		BatchArgs.BatchNumElems      = BatchNumElems;
-		BatchArgs.LODBegin           = 0;
-		BatchArgs.LODEnd             = NumLODs;
-		BatchArgs.Result             = BaseImage;
-		BatchArgs.ResultBytesPerElem = BytesPerElem;
+		for (int32 Iter = 0; Iter < NumIters; ++Iter)
+		{
+			VectorRegisterUInt8x16 A = LoadUnaligned_UInt8x16(Base + Iter*16);
+			VectorRegisterUInt8x16 B = LoadUnaligned_UInt8x16(Blen + Iter*16);
 
-		const int32 NumBatches = OpImageLayerInternal::GetOpLayerNumBatchesLODRange(BatchArgs);
+			VectorRegisterUInt8x16 Result = SoftLight_UInt8x16(A, B);
 
-		ParallelExecutionUtils::InvokeBatchParallelFor(NumBatches,
-		[
-			&BatchArgs, TopColor, BaseChannelOffset, ColorChannelOffset
-		] (uint32 BatchId)
+			StoreUnaligned_UInt8x16(Dest + Iter*16, Result);
+		}
+
+		ElemIndex = NumIters*16;
+#endif
+
+		for (; ElemIndex < NumElems; ++ElemIndex)
 		{
-			OpImageLayerInternal::FOpLayerBatchViews BatchViews = 
-					OpImageLayerInternal::GetOpLayerBatchLODRangeViews(BatchId, BatchArgs);
-			
-			uint8* BaseBuf = BatchViews.Result.GetData() + BaseChannelOffset;
-			
-			for (int32 I = 0; I < BatchViews.NumElems; ++I)
-			{
-				for (int32 C = 0; C < CHANNELS_TO_BLEND; ++C)
-				{
-					uint32 Base = BaseBuf[BASE_CHANNEL_STRIDE*I + C];
-					uint32 Blended = TopColor[C + ColorChannelOffset];
-					uint32 Result = BLEND_FUNC(Base, Blended);
-					if constexpr (CLAMP)
-					{
-						BaseBuf[BASE_CHANNEL_STRIDE*I + C] = (uint8)FMath::Min(255u, Result);
-					}
-					else
-					{
-						BaseBuf[BASE_CHANNEL_STRIDE*I + C] = (uint8)Result;
-					}
-				}
-			}
-		});
+			Dest[ElemIndex] = SoftLight_U8(Base[ElemIndex], Blen[ElemIndex]);
+		}
 	}
 
-
-	template<
-		uint32 (*BLEND_FUNC)(uint32, uint32), 
-		bool CLAMP, 
-		uint32 CHANNEL_COUNT>
-	FORCENOINLINE void BufferLayerColourInPlaceImpl(FImage* BaseImage, FVector4f Col, bool bOnlyOneMip, uint32 BaseOffset, uint8 ColOffset)
+	void SoftLightWithWord(uint8* Dest, const uint8* Base, const uint16* Blen, int32 NumElems)
 	{
-		EImageFormat BaseFormat = BaseImage->GetFormat();
+		int32 ElemIndex = 0;
 
-		if (BaseFormat == EImageFormat::RGB_UByte)
-		{
-			check(BaseOffset + CHANNEL_COUNT <= 3);
-			BufferLayerColourFormatInPlace<BLEND_FUNC, CLAMP, CHANNEL_COUNT, 3>
-					(BaseImage, Col, BaseOffset, ColOffset, bOnlyOneMip);
-		}
-		else if (BaseFormat == EImageFormat::RGBA_UByte)
+#if UE_MUTABLE_IMAGE_LAYER_USE_VECTOR_IMPL
+		int32 NumIters = NumElems / 16;
+
+		for (int32 Iter = 0; Iter < NumIters; ++Iter)
 		{
-			check(BaseOffset + CHANNEL_COUNT <= 4);
-			BufferLayerColourFormatInPlace<BLEND_FUNC, CLAMP, CHANNEL_COUNT, 4>
-					(BaseImage, Col, BaseOffset, ColOffset, bOnlyOneMip);
+			VectorRegisterUInt8x16 BLo = LoadUnaligned_UInt16x8(Blen + Iter*16);
+			VectorRegisterUInt8x16 BHi = LoadUnaligned_UInt16x8(Blen + Iter*16 + 8);
+
+			VectorRegisterUInt8x16 A = LoadUnaligned_UInt8x16(Base + Iter*16);
+			VectorRegisterUInt8x16 B = Saturate_UInt16x8(BLo, BHi);
+
+			VectorRegisterUInt8x16 Result = SoftLight_UInt8x16(A, B);
+
+			StoreUnaligned_UInt8x16(Dest + Iter*16, Result);
 		}
-		else if (BaseFormat == EImageFormat::BGRA_UByte)
+
+		ElemIndex = NumIters*16;
+#endif
+
+		for (; ElemIndex < NumElems; ++ElemIndex)
 		{
-			float Temp = Col[0];
-			Col[0] = Col[2];
-			Col[2] = Temp;
-			BufferLayerColourFormatInPlace<BLEND_FUNC, CLAMP, CHANNEL_COUNT, 4>
-					(BaseImage, Col, BaseOffset, ColOffset, bOnlyOneMip);
+			Dest[ElemIndex] = SoftLight_U8(Base[ElemIndex], Saturate_U16(Blen[ElemIndex]));
 		}
+	}
+
+	void HardLight(uint8* Dest, const uint8* Base, const uint8* Blen, int32 NumElems)
+	{
+		int32 ElemIndex = 0;
+
+#if UE_MUTABLE_IMAGE_LAYER_USE_VECTOR_IMPL
+		int32 NumIters = NumElems / 16;
 
-		else if (BaseFormat == EImageFormat::L_UByte)
+		for (int32 Iter = 0; Iter < NumIters; ++Iter)
 		{
-			check(BaseOffset + CHANNEL_COUNT <= 1);
-			BufferLayerColourFormatInPlace< BLEND_FUNC, CLAMP, CHANNEL_COUNT, 1>
-					(BaseImage, Col, BaseOffset, ColOffset, bOnlyOneMip);
+			VectorRegisterUInt8x16 A = LoadUnaligned_UInt8x16(Base + Iter*16);
+			VectorRegisterUInt8x16 B = LoadUnaligned_UInt8x16(Blen + Iter*16);
+
+			VectorRegisterUInt8x16 Result = HardLight_UInt8x16(A, B);
+
+			StoreUnaligned_UInt8x16(Dest + Iter*16, Result);
 		}
-		else
+
+		ElemIndex = NumIters*16;
+#endif
+
+		for (; ElemIndex < NumElems; ++ElemIndex)
 		{
-			checkf(false, TEXT("Unsupported format."));
+			Dest[ElemIndex] = HardLight_U8(Base[ElemIndex], Blen[ElemIndex]);
 		}
 	}
 
-
-	template< 
-		uint32 (*BLEND_FUNC)(uint32, uint32), 
-		uint32 CHANNEL_COUNT>
-	void BufferLayerColourInPlace(FImage* BaseImage, FVector4f Color, bool bOnlyOneMip, uint32 BaseOffset, uint8 ColorOffset)
+	void HardLightWithWord(uint8* Dest, const uint8* Base, const uint16* Blen, int32 NumElems)
 	{
-		bool bIsClampNeeded = OpImageLayerInternal::IsAnyComponentLargerThan1(Color);
-		if (bIsClampNeeded)
+		int32 ElemIndex = 0;
+
+#if UE_MUTABLE_IMAGE_LAYER_USE_VECTOR_IMPL
+		int32 NumIters = NumElems / 16;
+
+		for (int32 Iter = 0; Iter < NumIters; ++Iter)
 		{
-			BufferLayerColourInPlaceImpl<BLEND_FUNC, true, CHANNEL_COUNT>(BaseImage, Color, bOnlyOneMip, BaseOffset, ColorOffset);
+			VectorRegisterUInt8x16 BLo = LoadUnaligned_UInt16x8(Blen + Iter*16);
+			VectorRegisterUInt8x16 BHi = LoadUnaligned_UInt16x8(Blen + Iter*16 + 8);
+
+			VectorRegisterUInt8x16 A = LoadUnaligned_UInt8x16(Base + Iter*16);
+			VectorRegisterUInt8x16 B = Saturate_UInt16x8(BLo, BHi);
+
+			VectorRegisterUInt8x16 Result = HardLight_UInt8x16(A, B);
+
+			StoreUnaligned_UInt8x16(Dest + Iter*16, Result);
 		}
-		else
+
+		ElemIndex = NumIters*16;
+#endif
+
+		for (; ElemIndex < NumElems; ++ElemIndex)
 		{
-			BufferLayerColourInPlaceImpl<BLEND_FUNC, false, CHANNEL_COUNT>(BaseImage, Color, bOnlyOneMip, BaseOffset, ColorOffset);
+			Dest[ElemIndex] = HardLight_U8(Base[ElemIndex], Saturate_U16(Blen[ElemIndex]));
 		}
 	}
 
-
-	template< 
-		uint32 (*BLEND_FUNC_MASKED)(uint32, uint32, uint32),
-		uint32 (*BLEND_FUNC)(uint32, uint32), 
-		bool CLAMP, 
-		int32 CHANNEL_COUNT>
-	FORCENOINLINE void BufferLayerColourInPlaceImpl(FImage* BaseImage, const FImage* MaskImage, FVector4f Color, bool bOnlyOneMip, uint32 BaseOffset, uint8 ColorOffset)
+	void Dodge(uint8* Dest, const uint8* Base, const uint8* Blen, int32 NumElems)
 	{
-		check(BaseImage->GetSizeX() == MaskImage->GetSizeX());
-		check(BaseImage->GetSizeY() == MaskImage->GetSizeY());
+		int32 ElemIndex = 0;
 
-		EImageFormat BaseFormat = BaseImage->GetFormat();
+#if 0 //UE_MUTABLE_IMAGE_LAYER_USE_VECTOR_IMPL
+		int32 NumIters = NumElems / 16;
 
-		if (BaseFormat == EImageFormat::RGB_UByte)
+		for (int32 Iter = 0; Iter < NumIters; ++Iter)
 		{
-			check(BaseOffset + CHANNEL_COUNT <= 3);
-			BufferLayerColourFormat<BLEND_FUNC_MASKED, BLEND_FUNC, CLAMP, CHANNEL_COUNT, 3>
-					(BaseImage, BaseImage, MaskImage, Color, BaseOffset, ColorOffset, bOnlyOneMip);
+			VectorRegisterUInt8x16 A = LoadUnaligned_UInt8x16(Base + Iter*16);
+			VectorRegisterUInt8x16 B = LoadUnaligned_UInt8x16(Blen + Iter*16);
+
+			VectorRegisterUInt8x16 Result = Dodge_UInt8x16(A, B);
+
+			StoreUnaligned_UInt8x16(Dest + Iter*16, Result);
 		}
-		else if (BaseFormat == EImageFormat::RGBA_UByte)
+
+		ElemIndex = NumIters*16;
+#endif
+
+		for (; ElemIndex < NumElems; ++ElemIndex)
 		{
-			check(BaseOffset + CHANNEL_COUNT <= 4);
-			BufferLayerColourFormat<BLEND_FUNC_MASKED, BLEND_FUNC, CLAMP, CHANNEL_COUNT, 4>
-					(BaseImage, BaseImage, MaskImage, Color, BaseOffset, ColorOffset, bOnlyOneMip);
+			Dest[ElemIndex] = Dodge_U8(Base[ElemIndex], Blen[ElemIndex]);
 		}
-		else if (BaseFormat == EImageFormat::BGRA_UByte)
+	}
+
+	void DodgeWithWord(uint8* Dest, const uint8* Base, const uint16* Blen, int32 NumElems)
+	{
+		int32 ElemIndex = 0;
+
+#if 0 //UE_MUTABLE_IMAGE_LAYER_USE_VECTOR_IMPL
+		int32 NumIters = NumElems / 16;
+
+		for (int32 Iter = 0; Iter < NumIters; ++Iter)
 		{
-			check(BaseOffset + CHANNEL_COUNT <= 4);
-			float Temp = Color[0];
-			Color[0] = Color[2];
-			Color[2] = Temp;
+			VectorRegisterUInt8x16 BLo = LoadUnaligned_UInt16x8(Blen + Iter*16);
+			VectorRegisterUInt8x16 BHi = LoadUnaligned_UInt16x8(Blen + Iter*16 + 8);
+
+			VectorRegisterUInt8x16 A = LoadUnaligned_UInt8x16(Base + Iter*16);
+			VectorRegisterUInt8x16 B = Saturate_UInt16x8(BLo, BHi);
 
-			BufferLayerColourFormat<BLEND_FUNC_MASKED, BLEND_FUNC, CLAMP, CHANNEL_COUNT, 4>
-					(BaseImage, BaseImage, MaskImage, Color, BaseOffset, ColorOffset, bOnlyOneMip);
+			VectorRegisterUInt8x16 Result = Dodge_UInt8x16(A, B);
+
+			StoreUnaligned_UInt8x16(Dest + Iter*16, Result);
 		}
-		else if (BaseFormat == EImageFormat::L_UByte)
+
+		ElemIndex = NumIters*16;
+#endif
+
+		for (; ElemIndex < NumElems; ++ElemIndex)
 		{
-			check(BaseOffset + CHANNEL_COUNT <= 1);
-			BufferLayerColourFormat<BLEND_FUNC_MASKED, BLEND_FUNC, CLAMP, CHANNEL_COUNT, 1>
-					(BaseImage, BaseImage, MaskImage, Color, BaseOffset, ColorOffset, bOnlyOneMip);
+			Dest[ElemIndex] = Dodge_U8(Base[ElemIndex], Saturate_U16(Blen[ElemIndex]));
 		}
-		else
+	}
+
+	void Lighten(uint8* Dest, const uint8* Base, const uint8* Blen, int32 NumElems)
+	{
+		Screen(Dest, Base, Blen, NumElems);
+	}
+
+	void LightenWithWord(uint8* Dest, const uint8* Base, const uint16* Blen, int32 NumElems)
+	{
+		ScreenWithWord(Dest, Base, Blen, NumElems);
+	}
+
+	void Burn(uint8* Dest, const uint8* Base, const uint8* Blen, int32 NumElems)
+	{
+		int32 ElemIndex = 0;
+
+#if 0 //UE_MUTABLE_IMAGE_LAYER_USE_VECTOR_IMPL
+		int32 NumIters = NumElems / 16;
+
+		for (int32 Iter = 0; Iter < NumIters; ++Iter)
+		{
+			VectorRegisterUInt8x16 A = LoadUnaligned_UInt8x16(Base + Iter*16);
+			VectorRegisterUInt8x16 B = LoadUnaligned_UInt8x16(Blen + Iter*16);
+
+			VectorRegisterUInt8x16 Result = Burn_UInt8x16(A, B);
+
+			StoreUnaligned_UInt8x16(Dest + Iter*16, Result);
+		}
+
+		ElemIndex = NumIters*16;
+#endif
+
+		for (; ElemIndex < NumElems; ++ElemIndex)
 		{
-			checkf(false, TEXT("Unsupported format."));
+			Dest[ElemIndex] = Burn_U8(Base[ElemIndex], Blen[ElemIndex]);
 		}
 	}
 
-	template< 
-		uint32 (*BLEND_FUNC_MASKED)(uint32, uint32, uint32), 
-		uint32 (*BLEND_FUNC)(uint32, uint32), uint32 CHANNEL_COUNT>
-	void BufferLayerColourInPlace(FImage* BaseImage, const FImage* MaskImage, FVector4f Color, bool bOnlyOneMip, uint32 BaseOffset, uint8 ColorOffset)
+	void BurnWithWord(uint8* Dest, const uint8* Base, const uint16* Blen, int32 NumElems)
 	{
-		bool bIsClampNeeded = OpImageLayerInternal::IsAnyComponentLargerThan1(Color);
-		if (bIsClampNeeded)
+		int32 ElemIndex = 0;
+
+#if 0 //UE_MUTABLE_IMAGE_LAYER_USE_VECTOR_IMPL
+		int32 NumIters = NumElems / 16;
+
+		for (int32 Iter = 0; Iter < NumIters; ++Iter)
 		{
-			BufferLayerColourInPlaceImpl<BLEND_FUNC_MASKED, BLEND_FUNC, true, CHANNEL_COUNT>
-					(BaseImage, MaskImage, Color, bOnlyOneMip, BaseOffset, ColorOffset);
+			VectorRegisterUInt8x16 BLo = LoadUnaligned_UInt16x8(Blen + Iter*16);
+			VectorRegisterUInt8x16 BHi = LoadUnaligned_UInt16x8(Blen + Iter*16 + 8);
+
+			VectorRegisterUInt8x16 A = LoadUnaligned_UInt8x16(Base + Iter*16);
+			VectorRegisterUInt8x16 B = Saturate_UInt16x8(BLo, BHi);
+
+			VectorRegisterUInt8x16 Result = Burn_UInt8x16(A, B);
+			StoreUnaligned_UInt8x16(Dest + Iter*16, Result);
 		}
-		else
+
+		ElemIndex = NumIters*16;
+#endif
+
+		for (; ElemIndex < NumElems; ++ElemIndex)
 		{
-			BufferLayerColourInPlaceImpl<BLEND_FUNC_MASKED, BLEND_FUNC, false, CHANNEL_COUNT>
-					(BaseImage, MaskImage, Color, bOnlyOneMip, BaseOffset, ColorOffset);
+			Dest[ElemIndex] = Burn_U8(Base[ElemIndex], Saturate_U16(Blen[ElemIndex]));
 		}
 	}
 
+	void Overlay(uint8* Dest, const uint8* Base, const uint8* Blen, int32 NumElems)
+	{	
+		int32 ElemIndex = 0;
 
-	/**	
-	* Apply a blending function to an image with another image as blending layer
-	*/
-	template< uint32 (*BLEND_FUNC)(uint32, uint32), bool CLAMP,
-		// Number of total channels to actually process
-		int32 CHANNELS_TO_BLEND,
-		// Number of total channels in the base image
-		int32 BASE_CHANNEL_STRIDE,
-		// Number of total channels in the blend image
-		int32 BLENDED_CHANNEL_STRIDE>
-    void BufferLayerFormatInPlace(
-		FImage* BaseImage, 
-		const FImage* BlendedImage,
-		uint32 BaseChannelOffset,
-		uint32 BlendedChannelOffset,
-		bool bOnlyFirstLOD)
-	{
-		check(BaseImage->GetSizeX() == BlendedImage->GetSizeX());
-		check(BaseImage->GetSizeY() == BlendedImage->GetSizeY());
-
-		// No longer required.
-		//check(BaseImage->GetFormat() == BlendedImage->GetFormat());
-		check(BaseChannelOffset + CHANNELS_TO_BLEND <= GetImageFormatData(BaseImage->GetFormat()).Channels);
-		check(BlendedChannelOffset + CHANNELS_TO_BLEND <= GetImageFormatData(BlendedImage->GetFormat()).Channels);
-
-		check(bOnlyFirstLOD || BaseImage->GetLODCount() <= BlendedImage->GetLODCount());
+#if UE_MUTABLE_IMAGE_LAYER_USE_VECTOR_IMPL
+		int32 NumIters = NumElems / 16;
 
-		const int32 BaseBytesPerElem = GetImageFormatData(BaseImage->GetFormat()).BytesPerBlock;
-		const int32 BlendBytesPerElem = GetImageFormatData(BlendedImage->GetFormat()).BytesPerBlock;
-	
-		const int32 NumLODs = bOnlyFirstLOD ? 1 : BaseImage->GetLODCount();
+		for (int32 Iter = 0; Iter < NumIters; ++Iter)
+		{
+			VectorRegisterUInt8x16 A = LoadUnaligned_UInt8x16(Base + Iter*16);
+			VectorRegisterUInt8x16 B = LoadUnaligned_UInt8x16(Blen + Iter*16);
+			VectorRegisterUInt8x16 Result = Overlay_UInt8x16(A, B);
+			StoreUnaligned_UInt8x16(Dest + Iter*16, Result);
+		}
 
-		constexpr int32 BatchNumElems = 4096*2;
+		ElemIndex = NumIters*16;
+#endif
 
-		OpImageLayerInternal::FOpLayerBatchArgs BatchArgs;
-		BatchArgs.BatchNumElems      = BatchNumElems;
-		BatchArgs.LODBegin           = 0;
-		BatchArgs.LODEnd             = NumLODs;
-		BatchArgs.Result             = BaseImage;
-		BatchArgs.Blend              = BlendedImage;
-		BatchArgs.ResultBytesPerElem = BaseBytesPerElem;
-		BatchArgs.BlendBytesPerElem  = BlendBytesPerElem;
+		for (; ElemIndex < NumElems; ++ElemIndex)
+		{
+			Dest[ElemIndex] = Overlay_U8(Base[ElemIndex], Blen[ElemIndex]);
+		}
+	}
 
-		const int32 NumBatches = OpImageLayerInternal::GetOpLayerNumBatchesLODRange(BatchArgs);
+	void OverlayWithWord(uint8* Dest, const uint8* Base, const uint16* Blen, int32 NumElems)
+	{
+		int32 ElemIndex = 0;
 
-		ParallelExecutionUtils::InvokeBatchParallelFor(NumBatches,
-		[
-			&BatchArgs, BaseChannelOffset, BlendedChannelOffset
-		] (int32 BatchId)
+#if UE_MUTABLE_IMAGE_LAYER_USE_VECTOR_IMPL
+		int32 NumIters = NumElems / 16;
+
+		for (int32 Iter = 0; Iter < NumIters; ++Iter)
 		{
-			OpImageLayerInternal::FOpLayerBatchViews BatchViews = 
-					OpImageLayerInternal::GetOpLayerBatchLODRangeViews(BatchId, BatchArgs);
+			VectorRegisterUInt8x16 BLo = LoadUnaligned_UInt16x8(Blen + Iter*16);
+			VectorRegisterUInt8x16 BHi = LoadUnaligned_UInt16x8(Blen + Iter*16 + 8);
 
-        	uint8* BaseBuf = BatchViews.Result.GetData() + BaseChannelOffset;
-        	const uint8* BlendedBuf = BatchViews.Blend.GetData() + BlendedChannelOffset;
+			VectorRegisterUInt8x16 A = LoadUnaligned_UInt8x16(Base + Iter*16);
+			VectorRegisterUInt8x16 B = Saturate_UInt16x8(BLo, BHi);
 
-			for (int32 I = 0; I < BatchViews.NumElems; ++I)
-			{
-				for (int32 C = 0; C < CHANNELS_TO_BLEND; ++C)
-				{
-					uint32 Base = BaseBuf[BASE_CHANNEL_STRIDE * I + C];
-					uint32 Blended = BlendedBuf[BLENDED_CHANNEL_STRIDE * I + C];
-					uint32 Result = BLEND_FUNC(Base, Blended);
-					
-					if constexpr (CLAMP)
-					{
-						BaseBuf[BASE_CHANNEL_STRIDE * I + C] = (uint8)FMath::Min(255u, Result);
-					}
-					else
-					{
-						BaseBuf[BASE_CHANNEL_STRIDE * I + C] = (uint8)Result;
-					}
-				}
-			}
-		});
+			VectorRegisterUInt8x16 Result = Overlay_UInt8x16(A, B);
+			StoreUnaligned_UInt8x16(Dest + Iter*16, Result);
+		}
+
+		ElemIndex = NumIters*16;
+#endif
+
+		for (; ElemIndex < NumElems; ++ElemIndex)
+		{
+			Dest[ElemIndex] = Overlay_U8(Base[ElemIndex], Saturate_U16(Blen[ElemIndex]));
+		}
 	}
-	
-	/**
-	* Apply a blending function to an image with another image as blending layer
-	*/
-	template<
-		uint32 (*BLEND_FUNC)(uint32, uint32), 
-		bool CLAMP,
-		// Number of total channels to actually process
-		int32 CHANNELS_TO_BLEND,
-		// Number of total channels in the base image
-		int32 BASE_CHANNEL_STRIDE,
-		// Number of total channels in the blend image
-		int32 BLENDED_CHANNEL_STRIDE,
-		int32 BLENDED_CHANNEL_OFFSET>
-	void BufferLayerFormat(FImage* DestImage, const FImage* BaseImage, const FImage* BlendedImage, bool bOnlyFirstLOD)
-	{
-		check(BaseImage->GetSizeX() == BlendedImage->GetSizeX());
-		check(BaseImage->GetSizeY() == BlendedImage->GetSizeY());
-		// Not true anymore, since the BLENDED_CHANNEL_OFFSET has been added.
-		// check(BaseImage->GetFormat() == BlendedImage->GetFormat());
-		check(bOnlyFirstLOD || BaseImage->GetLODCount() <= BlendedImage->GetLODCount());
-
-		constexpr int32 UnblendedChannels = BASE_CHANNEL_STRIDE - CHANNELS_TO_BLEND;
-		constexpr int32 BatchNumElems = 4096*2;
-		
-		const int32 BaseBytesPerElem = GetImageFormatData(BaseImage->GetFormat()).BytesPerBlock;
-		const int32 BlendBytesPerElem = GetImageFormatData(BlendedImage->GetFormat()).BytesPerBlock;
-		const int32 DestBytesPerElem = GetImageFormatData(DestImage->GetFormat()).BytesPerBlock;
 
-		const int32 NumLODs = bOnlyFirstLOD ? 1 : BaseImage->GetLODCount();
+	void Multiply(uint8* Dest, const uint8* Base, const uint8* Blen, int32 NumElems)
+	{
+		int32 ElemIndex = 0;
 
-		OpImageLayerInternal::FOpLayerBatchArgs BatchArgs;
-		BatchArgs.BatchNumElems      = BatchNumElems;
-		BatchArgs.LODBegin           = 0;
-		BatchArgs.LODEnd             = NumLODs;
-		BatchArgs.Result             = DestImage;
-		BatchArgs.Base               = BaseImage;
-		BatchArgs.Blend              = BlendedImage;
-		BatchArgs.ResultBytesPerElem = DestBytesPerElem;
-		BatchArgs.BaseBytesPerElem   = BaseBytesPerElem;
-		BatchArgs.BlendBytesPerElem  = BlendBytesPerElem;
-
-		const int32 NumBatches = OpImageLayerInternal::GetOpLayerNumBatchesLODRange(BatchArgs);
-
-		ParallelExecutionUtils::InvokeBatchParallelFor(NumBatches,
-		[
-			&BatchArgs, UnblendedChannels
-		] (int32 BatchId)
+#if UE_MUTABLE_IMAGE_LAYER_USE_VECTOR_IMPL
+		int32 NumIters = NumElems / 16;
+
+		for (int32 Iter = 0; Iter < NumIters; ++Iter)
 		{
-			OpImageLayerInternal::FOpLayerBatchViews BatchViews = 
-					OpImageLayerInternal::GetOpLayerBatchLODRangeViews(BatchId, BatchArgs);
-	
-			const uint8* BaseBuf = BatchViews.Base.GetData();
-			const uint8* BlendedBuf = BatchViews.Blend.GetData() + BLENDED_CHANNEL_OFFSET;
-			uint8* DestBuf = BatchViews.Result.GetData();
+			VectorRegisterUInt8x16 A = LoadUnaligned_UInt8x16(Base + Iter*16);
+			VectorRegisterUInt8x16 B = LoadUnaligned_UInt8x16(Blen + Iter*16);
+			VectorRegisterUInt8x16 Result = Mult255Unorm_UInt8x16(A, B);
+			StoreUnaligned_UInt8x16(Dest + Iter*16, Result);
+		}
 
-			for (int32 I = 0; I < BatchViews.NumElems; ++I)
-			{
-				for (int32 C = 0; C < CHANNELS_TO_BLEND; ++C)
-				{
-					uint32 Base = BaseBuf[BASE_CHANNEL_STRIDE * I + C];
-					uint32 Blended = BlendedBuf[BLENDED_CHANNEL_STRIDE * I + C];
-					uint32 Result = BLEND_FUNC(Base, Blended);
+		ElemIndex = NumIters*16;
+#endif
 
-					if constexpr (CLAMP)
-					{
-						DestBuf[BASE_CHANNEL_STRIDE * I + C] = (uint8)FMath::Min(255u, Result);
-					}
-					else
-					{
-						DestBuf[BASE_CHANNEL_STRIDE * I + C] = (uint8)Result;
-					}
-				}
+		for (; ElemIndex < NumElems; ++ElemIndex)
+		{
+			Dest[ElemIndex] = Mult255Unorm_U8(Base[ElemIndex], Blen[ElemIndex]);
+		}
+	}
 
-				// Copy the unblended channels
-				// \TODO: unnecessary when doing it in-place?
-				if constexpr (UnblendedChannels > 0)
-				{
-					for (int32 C = 0; C < UnblendedChannels; ++C)
-					{
-						DestBuf[BASE_CHANNEL_STRIDE * I + CHANNELS_TO_BLEND + C] 
-							= BaseBuf[BASE_CHANNEL_STRIDE * I + CHANNELS_TO_BLEND + C];
-					}
-				}
-			}
-		});
+	void MultiplyWithWord(uint8* Dest, const uint8* Base, const uint16* Blen, int32 NumElems)
+	{
+		int32 ElemIndex = 0;
+
+#if 0 //UE_MUTABLE_IMAGE_LAYER_USE_VECTOR_IMPL
+		int32 NumIters = NumElems / 16;
+
+		for (int32 Iter = 0; Iter < NumIters; ++Iter)
+		{
+			VectorRegisterUInt16x8 BLo = LoadUnaligned_UInt16x8(Blen + Iter*16);
+			VectorRegisterUInt16x8 BHi = LoadUnaligned_UInt16x8(Blen + Iter*16 + 8);
+
+			VectorRegisterUInt8x16 A = LoadUnaligned_UInt8x16(Base + Iter*16);
+
+			VectorRegisterUInt16x8 ALo = GetLow16_UInt8x16(A);
+			VectorRegisterUInt16x8 AHi = GetHigh16_UInt8x16(A);
+
+			VectorRegisterUInt8x16 Result = Mult255Unorm_UInt16x8(ALo, AHi, BLo, BHi);
+			StoreUnaligned_UInt8x16(Dest + Iter*16, Result);
+		}
+
+		ElemIndex = NumIters*16;
+#endif
+
+		for (; ElemIndex < NumElems; ++ElemIndex)
+		{
+			Dest[ElemIndex] = Mult255Unorm_U16(Base[ElemIndex], Blen[ElemIndex]);
+		}
 	}
 
-	/**
-	* Apply a blending function to an image with another image as blending layer
-	*/
-	template< uint32 (*BLEND_FUNC)(uint32, uint32), bool CLAMP >
-	void BufferLayer(FImage* ResultImage, const FImage* BaseImage, const FImage* BlendedImage, bool bApplyToAlpha, bool bOnlyOneMip, bool bUseBlendSourceFromBlendAlpha)
+
+} // namespace OpImageLayerBlendOps
+
+namespace OpImageLayerCombineOps
+{
+
+	void Transpose(VectorRegister4Float& InOutA, VectorRegister4Float& InOutB, VectorRegister4Float& InOutC, VectorRegister4Float& InOutD)
 	{
-		check(ResultImage->GetFormat() == BaseImage->GetFormat());
-		check(ResultImage->GetSizeX() == BaseImage->GetSizeX());
-		check(ResultImage->GetSizeY() == BaseImage->GetSizeY());
-		check(bOnlyOneMip || ResultImage->GetLODCount() == BaseImage->GetLODCount());
-		check(BaseImage->GetSizeX() == BlendedImage->GetSizeX());
-		check(BaseImage->GetSizeY() == BlendedImage->GetSizeY());
-		check(bOnlyOneMip || ResultImage->GetLODCount() <= BlendedImage->GetLODCount());
+		VectorRegister4Float Temp0;
+		VectorRegister4Float Temp1;
+		VectorRegister4Float Temp2;
+		VectorRegister4Float Temp3;
+
+		VectorDeinterleave(Temp0, Temp1, InOutA, InOutB);
+		VectorDeinterleave(Temp2, Temp3, InOutC, InOutD);
 
-		EImageFormat BaseFormat = BaseImage->GetFormat();
-		EImageFormat BlendedFormat = BlendedImage->GetFormat();
+		// Temp0 = [XA,ZA,XB,ZB]
+		// Temp1 = [YA,WA,YB,WB]
+		// Temp2 = [XC,ZC,XD,ZD]
+		// Temp3 = [YC,WC,YD,WD]
+
+		VectorDeinterleave(InOutA, InOutC, Temp0, Temp2);
+		VectorDeinterleave(InOutB, InOutD, Temp1, Temp3);
+	}
 
-		if (bUseBlendSourceFromBlendAlpha)
+	void NormalCombine(uint8* Dest, const uint8* Base, const uint8* Blen, int32 NumElems)
+	{
+		constexpr VectorRegister4Float Float127AndAHalf = MakeVectorRegisterFloatConstant(127.5f, 127.5f, 127.5f, 127.5f);
+		constexpr VectorRegister4Float FloatOneOver127AndAHalf = MakeVectorRegisterFloatConstant(1.0f/127.5f, 1.0f/127.5f, 1.0f/127.5f, 1.0f/127.5f);
+		constexpr VectorRegister4Float FloatMinusOneOver127AndAHalf = MakeVectorRegisterFloatConstant(-1.0f/127.5f, -1.0f/127.5f, -1.0f/127.5f, -1.0f/127.5f);
+		constexpr VectorRegister4Float FloatMinusOne = GlobalVectorConstants::FloatMinusOne;
+		constexpr VectorRegister4Float FloatOne = GlobalVectorConstants::FloatOne;
+		constexpr VectorRegister4Float FloatZero = GlobalVectorConstants::FloatZero;
+		
+		int32 I = 0;
+		for (; I < NumElems - 4; I += 4)
 		{
-			if (BlendedFormat == EImageFormat::RGBA_UByte || BlendedFormat == EImageFormat::BGRA_UByte)
-			{
-				if (BaseFormat == EImageFormat::L_UByte)
-				{
-					BufferLayerFormat<BLEND_FUNC, CLAMP, 1, 1, 4, 3>(ResultImage, BaseImage, BlendedImage, bOnlyOneMip);
-				}
-				else
-				{
-					checkf(false, TEXT("Unsupported format."));
-				}
-			}
-			else if (BlendedFormat == EImageFormat::L_UByte)
-			{
-				BufferLayerFormat<BLEND_FUNC, CLAMP, 1, 1, 1, 0>(ResultImage, BaseImage, BlendedImage, bOnlyOneMip);
-			}
+			VectorRegister4Float A0 = VectorLoadByte4((Base + I*4 + 0 ));
+			VectorRegister4Float A1 = VectorLoadByte4((Base + I*4 + 4 ));
+			VectorRegister4Float A2 = VectorLoadByte4((Base + I*4 + 8 ));
+			VectorRegister4Float A3 = VectorLoadByte4((Base + I*4 + 12));
+
+			Transpose(A0, A1, A2, A3);
+
+			VectorRegister4Float B0 = VectorLoadByte4((Blen + I*4 + 0 ));
+			VectorRegister4Float B1 = VectorLoadByte4((Blen + I*4 + 4 ));
+			VectorRegister4Float B2 = VectorLoadByte4((Blen + I*4 + 8 ));
+			VectorRegister4Float B3 = VectorLoadByte4((Blen + I*4 + 12));
+			
+			Transpose(B0, B1, B2, B3);
+
+			A0 = VectorMultiplyAdd(A0, FloatOneOver127AndAHalf, FloatMinusOne);
+			A1 = VectorMultiplyAdd(A1, FloatOneOver127AndAHalf, FloatMinusOne);
+			A2 = VectorMultiply   (A2, FloatOneOver127AndAHalf               );
+			
+			VectorRegister4Float NegateA2 = VectorNegate(A2);
+
+			B0 = VectorMultiplyAdd(B0, FloatMinusOneOver127AndAHalf, FloatOne);
+			B1 = VectorMultiplyAdd(B1, FloatMinusOneOver127AndAHalf, FloatOne);
+			B2 = VectorMultiplyAdd(B2, FloatOneOver127AndAHalf, FloatMinusOne);
+
+			VectorRegister4Float DotAB =
+					VectorMultiplyAdd(A0, B0, VectorMultiplyAdd(A1, B1, VectorMultiply(A2, B2)));
+
+			VectorRegister4Float N0 = VectorMultiplyAdd(A0, DotAB, VectorMultiply(B0, NegateA2));
+			VectorRegister4Float N1 = VectorMultiplyAdd(A1, DotAB, VectorMultiply(B1, NegateA2));
+			VectorRegister4Float N2 = VectorMultiplyAdd(A2, DotAB, VectorMultiply(B2, NegateA2));
+			VectorRegister4Float N3 = FloatZero;
+
+			VectorRegister4Float ReciprocalLen = VectorReciprocalSqrt(
+					VectorMultiplyAdd(N0, N0, VectorMultiplyAdd(N1, N1, VectorMultiply(N2, N2))));
+		
+			N0 = VectorMultiplyAdd(VectorMultiply(N0, ReciprocalLen), Float127AndAHalf, Float127AndAHalf);
+			N1 = VectorMultiplyAdd(VectorMultiply(N1, ReciprocalLen), Float127AndAHalf, Float127AndAHalf);
+			N2 = VectorMultiplyAdd(VectorMultiply(N2, ReciprocalLen), Float127AndAHalf, Float127AndAHalf);
+
+			Transpose(N0, N1, N2, N3);
+
+			VectorStoreByte4(N0, Dest + I*4 + 0 );
+			VectorStoreByte4(N1, Dest + I*4 + 4 );
+			VectorStoreByte4(N2, Dest + I*4 + 8 );
+			VectorStoreByte4(N3, Dest + I*4 + 12);
 		}
-		else
+
+		constexpr VectorRegister4Float AOffset = MakeVectorRegisterFloatConstant(-1.0f, -1.0f, 0.0f, 0.0f);
+		constexpr VectorRegister4Float BOffset = MakeVectorRegisterFloatConstant(1.0f, 1.0f, -1.0f, 0.0f);
+		constexpr VectorRegister4Float AScale = MakeVectorRegisterFloatConstant(1.0f/127.5f, 1.0f/127.5f, 1.0f/127.5f, 0.0f);
+		constexpr VectorRegister4Float BScale = MakeVectorRegisterFloatConstant(-1.0f/127.5f, -1.0f/127.5f, 1.0f/127.5f, 0.0f);
+
+		for (; I < NumElems; ++I)
 		{
-			check(BaseFormat == BlendedFormat);
-			if (BaseFormat == EImageFormat::RGB_UByte)
-			{
-				check(!bUseBlendSourceFromBlendAlpha);
-				BufferLayerFormat<BLEND_FUNC, CLAMP, 3, 3, 3, 0>(ResultImage, BaseImage, BlendedImage, bOnlyOneMip);
-			}
-			else if (BaseFormat == EImageFormat::RGBA_UByte || BaseFormat == EImageFormat::BGRA_UByte)
-			{
-				check(!bUseBlendSourceFromBlendAlpha);
-				if (bApplyToAlpha)
-				{
-					BufferLayerFormat<BLEND_FUNC, CLAMP, 4, 4, 4, 0>(ResultImage, BaseImage, BlendedImage, bOnlyOneMip);
-				}
-				else
-				{
-					BufferLayerFormat<BLEND_FUNC, CLAMP, 3, 4, 4, 0>(ResultImage, BaseImage, BlendedImage, bOnlyOneMip);
-				}
-			}
-			else if (BaseFormat == EImageFormat::L_UByte)
-			{
-				BufferLayerFormat<BLEND_FUNC, CLAMP, 1, 1, 1, 0>(ResultImage, BaseImage, BlendedImage, bOnlyOneMip);
-			}
-			else
-			{
-				checkf(false, TEXT("Unsupported format."));
-			}
+			VectorRegister4Float A = VectorLoadByte4((Base + I*4));
+			VectorRegister4Float B = VectorLoadByte4((Blen + I*4));
+
+			A = VectorMultiplyAdd(A, AScale, AOffset);
+			B = VectorMultiplyAdd(B, BScale, BOffset);
+
+			VectorRegister4Float NegateA2 = VectorNegate(VectorReplicate(A, 2));
+		
+			VectorRegister4Float N = VectorMultiplyAdd(A, VectorDot3(A, B), VectorMultiply(B, NegateA2));
+			N = VectorMultiplyAdd(VectorMultiply(N, VectorReciprocalLen(N)), Float127AndAHalf, Float127AndAHalf);
+
+			VectorStoreByte4(N, Dest + I*4);
 		}
 	}
+} //namespace OpImageLayerCombineOps
+
+namespace OpImageLayerInternal
+{
+	struct FOpLayerBatchArgs
+	{
+		int32 BatchNumElems         = 0;
+		int32 LODBegin              = 0;
+		int32 LODEnd                = 0;
+		int32 FirstLODOffsetInElems = 0;
+
+		FImage* Result      = nullptr;
+		const FImage* Base  = nullptr;
+		const FImage* Blend = nullptr;
+		const FImage* Mask  = nullptr;
+
+		uint32 ResultBytesPerElem = 0;
+		uint32 BaseBytesPerElem   = 0;
+		uint32 BlendBytesPerElem  = 0;
+		uint32 MaskBytesPerElem   = 0;
+	};
+
+	struct FOpLayerBatchViews
+	{
+		int32 NumElems = 0;
+
+		TArrayView<uint8> Result;
+		TArrayView<const uint8> Base;
+		TArrayView<const uint8> Blend;
+		TArrayView<const uint8> Mask;
+
+		uint64 DecodeTokenRLE;
+	};
+
+	struct FOpBlendKernel
+	{
+		using FConvertFuncType         = void(uint8*, const uint8*, int32);
+		using FMixFuncType             = void(uint8*, const uint8*, const uint8*, const uint8*, int32);
+		using FMakeConstantU8FuncType  = void(uint8*, FColor, int32);
+		using FMakeConstantU16FuncType = void(uint16*, Math::TIntVector4<uint16>, int32);
+		using FMaskDecodeFuncType      = void(uint8*, const uint8*, uint64&, int32);
+		using FBlendWithWordFuncType   = void(uint8*, const uint8*, const uint16*, int32);
+
+		FColor ConstantU8;
+		Math::TIntVector4<uint16> ConstantU16;
+		FMakeConstantU8FuncType* MakeConstantU8   = nullptr;
+		FMakeConstantU16FuncType* MakeConstantU16 = nullptr;
+
+		FConvertFuncType* BlendFormat = nullptr;
+		FBlendFuncType* ColorBlend    = nullptr;
+		FBlendFuncType* AlphaBlend    = nullptr;
+
+		FBlendWithWordFuncType* ColorBlendWithWord = nullptr;
+		FBlendWithWordFuncType* AlphaBlendWithWord = nullptr;
+
+		FConvertFuncType* ExtractMaskFromBlendAlpha = nullptr; 
+		FConvertFuncType* MaskFormat                = nullptr;
+		FMixFuncType* Mix                           = nullptr;
+		FBlendFuncType* SelectAlpha                 = nullptr;
+
+		FMaskDecodeFuncType* MaskDecode = nullptr;
+	};
+
+	struct FMakeOpBlendKernelArgs
+	{
+		FImage* Dest  = nullptr;
+		const FImage* Base  = nullptr;
+		const FImage* Blend = nullptr;
+		const FImage* Mask  = nullptr;
+
+		FBlendFuncType* ColorBlendFunc = nullptr;
+		FBlendFuncType* AlphaBlendFunc = nullptr;
+
+		bool bApplyColorBlendToAlpha       = false;
+		bool bUseBlendSourceFromBlendAlpha = false;
+		bool bUseMaskFromBlendAlpha        = false;
+		uint8 BlendAlphaSourceChannel      = 3;
+
+		bool bBlendWithConstant = false;
+		FVector4f RealConstant  = FVector4f{};
+	};
+
+	struct FOpCombineKernel
+	{
+		using FConvertFuncType      = void(uint8*, const uint8*, int32);
+		using FMixFuncType          = void(uint8*, const uint8*, const uint8*, const uint8*, int32);
+		using FMakeConstantFuncType = void(uint8*, FColor, int32);
+		using FMaskDecodeFuncType   = void(uint8*, const uint8*, uint64&, int32);
+
+		FMakeConstantFuncType* MakeConstant = nullptr;
+		FColor Constant;
+
+		FConvertFuncType* BaseFormat  = nullptr;
+		FConvertFuncType* BlendFormat = nullptr;
+		FConvertFuncType* MaskFormat  = nullptr;
 
+		FCombineFuncType* Combine     = nullptr;
+		FMixFuncType* Mix             = nullptr;
 
-	template< uint32 (*BLEND_FUNC)(uint32, uint32), bool CLAMP, int32 CHANNEL_COUNT >
-	void BufferLayerInPlace(FImage* BaseImage, const FImage* BlendedImage, bool bOnlyOneMip, uint32 BaseOffset, uint32 BlendedOffset)
+		FConvertFuncType* ResultFormat = nullptr;
+		
+		FMaskDecodeFuncType* MaskDecode = nullptr;
+		
+		uint32 CombineBytesPerElem = 4;
+	};
+
+	struct FMakeOpCombineKernelArgs
 	{
-		check(BaseImage->GetSizeX() == BlendedImage->GetSizeX());
-		check(BaseImage->GetSizeY() == BlendedImage->GetSizeY());
-		// Not required since we have the CHANNEL_COUNT and offsets. 
-		// check(BaseImage->GetFormat() == BlendedImage->GetFormat());
+		FImage* Dest  = nullptr;
+		const FImage* Base  = nullptr;
+		const FImage* Blend = nullptr;
+		const FImage* Mask  = nullptr;
 
-		EImageFormat BaseFormat = BaseImage->GetFormat();
-		EImageFormat BlendFormat = BlendedImage->GetFormat();
+		FCombineFuncType* CombineFunc = nullptr;
+		uint32 CombineBytesPerElem = 4;
 
-		if (BaseFormat == EImageFormat::RGB_UByte && BlendFormat==EImageFormat::RGB_UByte)
+		bool bCombineWithConstant = false;
+		FColor Constant           = FColor::White;
+	};
+
+	FOpBlendKernel::FBlendWithWordFuncType* ConvertBlendFuncToWordVariant(FBlendFuncType* Func)
+	{
+		if (Func == nullptr)
 		{
-			check(BaseOffset + CHANNEL_COUNT <= 3);
-			check(BlendedOffset + CHANNEL_COUNT <= 3);
-			BufferLayerFormatInPlace<BLEND_FUNC, CLAMP, CHANNEL_COUNT, 3, 3>(BaseImage, BlendedImage, BaseOffset, BlendedOffset, bOnlyOneMip);
+			return nullptr;
 		}
-		else if ((BaseFormat == EImageFormat::RGBA_UByte || BaseFormat == EImageFormat::BGRA_UByte) &&
-				 (BlendFormat == EImageFormat::RGBA_UByte || BlendFormat == EImageFormat::BGRA_UByte) )
+
+		if (Func == OpImageLayerBlendOps::Blend)
 		{
-			check(BaseOffset + CHANNEL_COUNT <= 4);
-			check(BlendedOffset + CHANNEL_COUNT <= 4);
-			BufferLayerFormatInPlace< BLEND_FUNC, CLAMP, CHANNEL_COUNT, 4, 4>
-					(BaseImage, BlendedImage, BaseOffset, BlendedOffset, bOnlyOneMip);
+			return OpImageLayerBlendOps::BlendWithWord;
 		}
-		else if ((BaseFormat == EImageFormat::RGBA_UByte || BaseFormat == EImageFormat::BGRA_UByte) &&
-				 BlendFormat == EImageFormat::L_UByte )
+
+		if (Func == OpImageLayerBlendOps::Screen)
 		{
-			check(BaseOffset + CHANNEL_COUNT <= 4);
-			check(BlendedOffset + CHANNEL_COUNT <= 1);
-			BufferLayerFormatInPlace< BLEND_FUNC, CLAMP, CHANNEL_COUNT, 4, 1>
-					(BaseImage, BlendedImage, BaseOffset, BlendedOffset, bOnlyOneMip);
+			return OpImageLayerBlendOps::ScreenWithWord;
 		}
-		else if (BaseFormat == EImageFormat::L_UByte && BlendFormat==EImageFormat::L_UByte)
+
+		if (Func == OpImageLayerBlendOps::SoftLight)
 		{
-			check(BaseOffset + CHANNEL_COUNT <= 1);
-			BufferLayerFormatInPlace< BLEND_FUNC, CLAMP, CHANNEL_COUNT, 1, 1>
-					(BaseImage, BlendedImage, BaseOffset, BlendedOffset, bOnlyOneMip);
+			return OpImageLayerBlendOps::BlendWithWord;
 		}
-		else
+
+		if (Func == OpImageLayerBlendOps::HardLight)
 		{
-			checkf(false, TEXT("Unsupported format."));
+			return OpImageLayerBlendOps::HardLightWithWord;
 		}
-	}
 
+		if (Func == OpImageLayerBlendOps::Dodge)
+		{
+			return OpImageLayerBlendOps::DodgeWithWord;
+		}
 
-	template< uint32 (*BLEND_FUNC_MASKED)(uint32, uint32, uint32),
-		uint32 (*BLEND_FUNC)(uint32, uint32),
-		bool CLAMP,
-		// Number of total channels to actually process
-		int32 CHANNELS_TO_BLEND,
-		// Number of total channels in the base image
-		int32 BASE_CHANNEL_STRIDE,
-		// Number of total channels in the blend image
-		int32 BLENDED_CHANNEL_STRIDE>
-	void BufferLayerFormat(
-		FImage* DestImage,
-		const FImage* BaseImage,
-		const FImage* MaskImage,
-		const FImage* BlendImage,
-		uint32 DestOffset,
-		uint32 BaseChannelOffset,
-		uint32 BlendedChannelOffset,
-		bool bOnlyFirstLOD)
-	{
-		check(BaseImage->GetSizeX() == MaskImage->GetSizeX() && BaseImage->GetSizeY() == MaskImage->GetSizeY());
-		check(BaseImage->GetSizeX() == BlendImage->GetSizeX() && BaseImage->GetSizeY() == BlendImage->GetSizeY());
-		check(bOnlyFirstLOD || BaseImage->GetLODCount() <= MaskImage->GetLODCount());
-		check(bOnlyFirstLOD || BaseImage->GetLODCount() <= BlendImage->GetLODCount());
-		check(BaseImage->GetFormat() == BlendImage->GetFormat());
-		check(MaskImage->GetFormat() == EImageFormat::L_UByte ||
-			  MaskImage->GetFormat() == EImageFormat::L_UByteRLE);
+		if (Func == OpImageLayerBlendOps::Lighten)
+		{
+			return OpImageLayerBlendOps::LightenWithWord;
+		}
+
+		if (Func == OpImageLayerBlendOps::Burn)
+		{
+			return OpImageLayerBlendOps::BurnWithWord;
+		}
 
-		const EImageFormat MaskFormat = MaskImage->GetFormat();
+		if (Func == OpImageLayerBlendOps::Multiply)
+		{
+			return OpImageLayerBlendOps::MultiplyWithWord;
+		}
 
-		const int32 BaseBytesPerElem = GetImageFormatData(BaseImage->GetFormat()).BytesPerBlock; 
-		const int32 BlendBytesPerElem = GetImageFormatData(BlendImage->GetFormat()).BytesPerBlock; 
+		if (Func == OpImageLayerBlendOps::Overlay)
+		{
+			return OpImageLayerBlendOps::OverlayWithWord;
+		}
 
-		const int32 NumLODs = bOnlyFirstLOD ? 1 : BaseImage->GetLODCount();
+		check(false);
+		return nullptr;
+	}
 
-		const bool bIsMaskUncompressed = MaskFormat == EImageFormat::L_UByte;
-		constexpr int32 UnblendedChannels = BASE_CHANNEL_STRIDE - CHANNELS_TO_BLEND;
+	FOpBlendKernel MakeKernel(const FMakeOpBlendKernelArgs& Args)
+	{
+		using namespace ImageComputeUtils;
+		
+		FOpBlendKernel Kernel;
 
-		if (bIsMaskUncompressed)
-		{	
-			constexpr int32 BatchNumElems = 4096*2;
+		Kernel.ColorBlend = Args.ColorBlendFunc;
+		Kernel.AlphaBlend = Args.AlphaBlendFunc;
 
-			OpImageLayerInternal::FOpLayerBatchArgs BatchArgs;
-			BatchArgs.BatchNumElems      = BatchNumElems;
-			BatchArgs.LODBegin           = 0;
-			BatchArgs.LODEnd             = NumLODs;
-			BatchArgs.Result             = DestImage;
-			BatchArgs.Base               = BaseImage;
-			BatchArgs.Blend              = BlendImage;
-			BatchArgs.Mask               = MaskImage;
-			BatchArgs.ResultBytesPerElem = BaseBytesPerElem;
-			BatchArgs.BaseBytesPerElem   = BaseBytesPerElem;
-			BatchArgs.BlendBytesPerElem  = BlendBytesPerElem;
-			BatchArgs.MaskBytesPerElem   = 1;
+		if (Args.bBlendWithConstant)
+		{
+			Kernel.ConstantU8  = FLinearColor(Args.RealConstant).QuantizeRound();
+			Kernel.ConstantU16 = Math::TIntVector4<uint16>(
+					FMath::Clamp(Args.RealConstant.X * 255.0f + 0.5f, 0.0f, 65535.0f),
+					FMath::Clamp(Args.RealConstant.Y * 255.0f + 0.5f, 0.0f, 65535.0f),
+					FMath::Clamp(Args.RealConstant.Z * 255.0f + 0.5f, 0.0f, 65535.0f),
+					FMath::Clamp(Args.RealConstant.W * 255.0f + 0.5f, 0.0f, 65535.0f));
 
-			const int32 NumBatches = OpImageLayerInternal::GetOpLayerNumBatchesLODRange(BatchArgs);
+			bool bAllLessThanEqOne = 
+				(Args.RealConstant.X <= 1.0f) && (Args.RealConstant.Y <= 1.0f) && 
+				(Args.RealConstant.Z <= 1.0f) && (Args.RealConstant.W <= 1.0f);
 
-			ParallelExecutionUtils::InvokeBatchParallelFor(NumBatches,
-			[
-				&BatchArgs, UnblendedChannels
-			] (uint32 BatchId)
+			bool bUseU16Constant = !bAllLessThanEqOne && 
+					(Kernel.ColorBlend == OpImageLayerBlendOps::Multiply || Kernel.AlphaBlend == OpImageLayerBlendOps::Multiply);
+			
+			switch (Args.Base->GetFormat())
 			{
-				OpImageLayerInternal::FOpLayerBatchViews BatchViews = 
-						OpImageLayerInternal::GetOpLayerBatchLODRangeViews(BatchId, BatchArgs);
-
-				const uint8* BaseBuf = BatchViews.Base.GetData();
-				const uint8* BlendedBuf = BatchViews.Blend.GetData();
-				const uint8* MaskBuf = BatchViews.Mask.GetData();
-				uint8* DestBuf = BatchViews.Result.GetData();
+			case EImageFormat::L_UByte:
+			{
+				if (!bUseU16Constant)
+				{
+					Kernel.MakeConstantU8 = Vec1Fill_U8;
+				}
+				else
+				{
+					Kernel.MakeConstantU16 = Vec1Fill_U16;
+				}
 
-				// This could happen in case of missing data files.
-				if (!BaseBuf || !BlendedBuf || !MaskBuf || !DestBuf)
+				break;
+			}
+			case EImageFormat::RGB_UByte:
+			{
+				if (!bUseU16Constant)
+				{
+					Kernel.MakeConstantU8 = Vec3Fill_U8;
+				}
+				else
+				{
+					Kernel.MakeConstantU16 = Vec3Fill_U16;
+				}
+				break;
+			}
+			case EImageFormat::RGBA_UByte:
+			{
+				if (!bUseU16Constant)
+				{
+					Kernel.MakeConstantU8 = Vec4Fill_U8;
+				}
+				else
+				{
+					Kernel.MakeConstantU16 = Vec4Fill_U16;
+				}
+				break;
+			}
+			case EImageFormat::BGRA_UByte:
+			{
+				if (!bUseU16Constant)
 				{
-					return;
+					Kernel.ConstantU8 = FColor(Kernel.ConstantU8.B, Kernel.ConstantU8.G, Kernel.ConstantU8.R, Kernel.ConstantU8.A);
+					Kernel.MakeConstantU8 = Vec4Fill_U8;
 				}
+				else
+				{
+					Kernel.ConstantU16 = Math::TIntVector4<uint16>(Kernel.ConstantU16.X, Kernel.ConstantU16.Y, Kernel.ConstantU16.Z, Kernel.ConstantU16.W);
+					Kernel.MakeConstantU16 = Vec4Fill_U16;
+				}
+
+				break;
+			}
+			default: check(false);
+			}
+
+			if (bUseU16Constant)
+			{
+				Kernel.ColorBlendWithWord = ConvertBlendFuncToWordVariant(Kernel.ColorBlend);  
+				Kernel.AlphaBlendWithWord = ConvertBlendFuncToWordVariant(Kernel.AlphaBlend);  
+			}
+		}
 
-				for (int32 I = 0; I < BatchViews.NumElems; ++I)
+		if (!Args.bBlendWithConstant && Args.Base->GetFormat() != Args.Blend->GetFormat())
+		{
+			switch (Args.Base->GetFormat())
+			{
+			case EImageFormat::L_UByte:
+			{
+				switch (Args.Blend->GetFormat())
+				{
+				case EImageFormat::RGB_UByte:
 				{
-					uint32 Mask = MaskBuf[I];
-					for (int32 C = 0; C < CHANNELS_TO_BLEND; ++C)
+					Kernel.BlendFormat = Vec3ToVec1_U8;
+					break;
+				}
+				case EImageFormat::RGBA_UByte:
+				case EImageFormat::BGRA_UByte:
+				{
+					if (Args.bUseBlendSourceFromBlendAlpha)
 					{
-						const uint32 Base = BaseBuf[BASE_CHANNEL_STRIDE * I + C];
-						const uint32 Blended = BlendedBuf[BLENDED_CHANNEL_STRIDE * I + C];
-						const uint32 Result = BLEND_FUNC_MASKED(Base, Blended, Mask);
-						if constexpr (CLAMP)
-						{
-							DestBuf[BASE_CHANNEL_STRIDE * I + C] = (uint8)FMath::Min(255u, Result);
-						}
-						else
-						{
-							DestBuf[BASE_CHANNEL_STRIDE * I + C] = (uint8)Result;
-						}
+						Kernel.BlendFormat = Vec4SwizzleWZYXAndConvertToVec1_U8;
 					}
-					// Copy the unblended channels
-					// \TODO: unnecessary when doing it in-place?
-					if constexpr (UnblendedChannels > 0)
+					else
 					{
-						for (int32 C = 0; C < UnblendedChannels; ++C)
-						{
-							DestBuf[BASE_CHANNEL_STRIDE * I + CHANNELS_TO_BLEND + C] = BaseBuf[BASE_CHANNEL_STRIDE * I + CHANNELS_TO_BLEND + C];
-						}
+						Kernel.BlendFormat = Vec4ToVec1_U8;
 					}
+					break;
 				}
-			});
-		}
-		else if (MaskFormat == EImageFormat::L_UByteRLE)
-		{
-			int32 Rows = BaseImage->GetSizeY();
-			int32 Width = BaseImage->GetSizeX();
+				default: check(Args.Blend->GetFormat() == EImageFormat::L_UByte); break;
+				}
+
+				break;
+			}
+			case EImageFormat::RGB_UByte:
+			{
+				switch (Args.Blend->GetFormat())
+				{
+				case EImageFormat::L_UByte:
+				{
+					Kernel.BlendFormat = Vec1ToVec3_U8;
+					break;
+				}
+				case EImageFormat::RGBA_UByte:
+				case EImageFormat::BGRA_UByte:
+				{
+					Kernel.BlendFormat = Vec4ToVec3_U8;
+					break;
+				}
+				default: check(Args.Blend->GetFormat() == EImageFormat::RGB_UByte); break;
+				}
+				break;
+			}
+			case EImageFormat::BGRA_UByte:
+			case EImageFormat::RGBA_UByte:
+			{
+				switch (Args.Blend->GetFormat())
+				{
+				case EImageFormat::L_UByte:
+				{
+					Kernel.BlendFormat = Vec1ToVec4_U8;
+					break;
+				}
+				case EImageFormat::RGB_UByte:
+				{
+					Kernel.BlendFormat = Vec3ToVec4_U8;
+					break;
+				}
+				default: check(Args.Blend->GetFormat() == EImageFormat::RGBA_UByte || Args.Blend->GetFormat() == EImageFormat::BGRA_UByte); break;
+				}
+				break;
+			}
+			default: check(Args.Base->GetFormat() == EImageFormat::RGBA_UByte || Args.Base->GetFormat() == EImageFormat::BGRA_UByte); break;
+			}
+		}
+
+		if (Args.Mask || Args.bUseMaskFromBlendAlpha)
+		{
+			if (Args.bUseMaskFromBlendAlpha)
+			{
+				check(Args.BlendAlphaSourceChannel == 3 || Args.BlendAlphaSourceChannel == 0);
+
+				if (Args.BlendAlphaSourceChannel == 3)
+				{
+					Kernel.ExtractMaskFromBlendAlpha = Vec4SwizzleWZYXAndConvertToVec1_U8;
+				}
+				else
+				{
+					Kernel.ExtractMaskFromBlendAlpha = Vec4ToVec1_U8;
+				}
+			}
+
+			if (Args.Mask)
+			{
+				check(Args.Mask->GetFormat() == EImageFormat::L_UByteRLE || 
+					  Args.Mask->GetFormat() == EImageFormat::L_UByte)
+
+				if (Args.Mask->GetFormat() == EImageFormat::L_UByteRLE)
+				{
+					Kernel.MaskDecode = IterativeDecodeRLE_L;
+				}
+			}
+
+			if (Args.Mask || Args.bUseMaskFromBlendAlpha)
+			{
+				switch (Args.Base->GetFormat())
+				{
+				case EImageFormat::RGB_UByte:
+				{
+					Kernel.MaskFormat = Vec1ToVec3_U8;
+					break;
+				}
+				case EImageFormat::BGRA_UByte:
+				case EImageFormat::RGBA_UByte:
+				{
+					if (Kernel.ExtractMaskFromBlendAlpha == Vec4SwizzleWZYXAndConvertToVec1_U8)
+					{
+						Kernel.ExtractMaskFromBlendAlpha = Vec4SwizzleWWWW_U8;
+					}
+					else if (Kernel.ExtractMaskFromBlendAlpha == Vec4ToVec1_U8)
+					{
+						Kernel.ExtractMaskFromBlendAlpha = Vec4SwizzleXXXX_U8;
+					}
+					else
+					{
+						Kernel.MaskFormat = Vec1ToVec4_U8;
+					}
+					break;
+				}
+				}
+				
+				Kernel.Mix = Lerp_U8;
+			}
+		}
+
+		if (Args.bApplyColorBlendToAlpha)
+		{
+			Kernel.AlphaBlend = nullptr;
+			Kernel.SelectAlpha = nullptr;
+		}
+		else
+		{
+			if (Args.Base->GetFormat() == EImageFormat::RGBA_UByte || Args.Base->GetFormat() == EImageFormat::BGRA_UByte)
+			{
+				Kernel.SelectAlpha = Vec4SelectW_U8;
+			}
+		}
+
+		return Kernel;
+	}
+
+	FOpCombineKernel MakeKernel(const FMakeOpCombineKernelArgs& Args)
+	{
+		using namespace ImageComputeUtils;
+		
+		FOpCombineKernel Kernel;
+
+		Kernel.Combine = Args.CombineFunc;
+		Kernel.CombineBytesPerElem = Args.CombineBytesPerElem;
+
+		switch (Args.Base->GetFormat())
+		{
+			case EImageFormat::L_UByte:
+			{
+				Kernel.BaseFormat = Vec1ToVec4_U8;
+				break;
+			}
+			case EImageFormat::RGB_UByte:
+			{
+				Kernel.BaseFormat = Vec3ToVec4_U8;
+				break;
+			}
+			case EImageFormat::BGRA_UByte:
+			{
+				Kernel.BaseFormat = Vec4SwizzleZYXW_U8;
+				break;
+			}
+			default: check(Args.Base->GetFormat() == EImageFormat::RGBA_UByte); break;
+		}
+
+		if (Args.bCombineWithConstant)
+		{
+			Kernel.Constant = Args.Constant;
+			Kernel.MakeConstant = Vec4Fill_U8;
+		}
+		else
+		{
+			switch (Args.Blend->GetFormat())
+			{
+				case EImageFormat::L_UByte:
+				{
+					Kernel.BlendFormat = Vec1ToVec4_U8;
+					break;
+				}
+				case EImageFormat::RGB_UByte:
+				{
+					Kernel.BlendFormat = Vec3ToVec4_U8;
+					break;
+				}
+				case EImageFormat::BGRA_UByte:
+				{
+					Kernel.BlendFormat = Vec4SwizzleZYXW_U8;
+					break;
+				}
+				default: check(Args.Blend->GetFormat() == EImageFormat::RGBA_UByte); break;
+			}
+		}
+
+		if (Args.Mask)
+		{
+			Kernel.MaskFormat = Vec1ToVec4_U8;
+			Kernel.Mix = Lerp_U8;
+		}
+
+		switch (Args.Base->GetFormat())
+		{
+			case EImageFormat::L_UByte:
+			{
+				Kernel.ResultFormat = Vec4ToVec1_U8;
+				break;
+			}
+			case EImageFormat::RGB_UByte:
+			{
+				Kernel.ResultFormat = Vec4ToVec3_U8;
+				break;
+			}
+			case EImageFormat::BGRA_UByte:
+			{
+				Kernel.ResultFormat = Vec4SwizzleZYXW_U8;
+				break;
+			}
+			default: check(Args.Base->GetFormat() == EImageFormat::RGBA_UByte); break;
+		}
+
+		return Kernel;
+	}
+
+	void RunKernelOnBatch(const FOpBlendKernel& Kernel, const FOpLayerBatchViews BatchViews)
+	{
+		//MUTABLE_CPUPROFILER_SCOPE(OpImageLayer_RunKernelOnBatch_Blend);
+
+		constexpr int32 MaxElemSize = 4;
+		constexpr int32 NumCacheLinesPerWave = 16;
+		constexpr int32 WaveSizeInByteElems = (PLATFORM_CACHE_LINE_SIZE * NumCacheLinesPerWave)/MaxElemSize;
+
+		struct alignas(PLATFORM_CACHE_LINE_SIZE) FTemporalRegister
+		{
+			uint8 Data[WaveSizeInByteElems*MaxElemSize];
+		};
+
+		FTemporalRegister Register0;
+		FTemporalRegister Register1;
+		FTemporalRegister Register2;
+		FTemporalRegister Register3;
+		FTemporalRegister Register4;
+
+		int32 BaseBytesPerElem = BatchViews.Base.Num() / BatchViews.NumElems;
+		int32 BlenBytesPerElem = BatchViews.Blend.Num() / BatchViews.NumElems;
+
+		check(BaseBytesPerElem <= MaxElemSize);
+		check(BlenBytesPerElem <= MaxElemSize);
+
+		uint8* ResultBuffer = BatchViews.Result.GetData();
+		const uint8* BaseBuffer = BatchViews.Base.GetData();
+		const uint8* BlenBuffer = BatchViews.Blend.GetData();
+		const uint8* MaskBuffer = BatchViews.Mask.GetData();
+
+		uint64 InOutDecodeToken = BatchViews.DecodeTokenRLE;
+
+		const int32 WaveSizeInElems = Kernel.MakeConstantU16 != nullptr ? WaveSizeInByteElems / 2 : WaveSizeInByteElems; 
+		
+		if (Kernel.MakeConstantU16)
+		{
+			Kernel.MakeConstantU16((uint16*)Register0.Data, Kernel.ConstantU16, WaveSizeInElems);
+		}
+		else if (Kernel.MakeConstantU8)
+		{
+			Kernel.MakeConstantU8(Register0.Data, Kernel.ConstantU8, WaveSizeInElems);
+		}
+
+		const int32 NumWaves = FMath::DivideAndRoundUp<int32>(BatchViews.NumElems, WaveSizeInElems);
+		for (int32 Idx = 0; Idx < NumWaves; ++Idx)
+		{
+			int32 WaveBeginInElems = Idx * WaveSizeInElems;
+			int32 WaveEndInElems   = FMath::Min(WaveBeginInElems + WaveSizeInElems, BatchViews.NumElems);
+
+			int32 NumWaveElems = WaveEndInElems - WaveBeginInElems;
+
+			uint8* ResultDataBuffer     = ResultBuffer + WaveBeginInElems*BaseBytesPerElem;
+			const uint8* BaseDataBuffer = BaseBuffer + WaveBeginInElems*BaseBytesPerElem;
+			const uint8* MaskDataBuffer = MaskBuffer + WaveBeginInElems;
+			const uint8* BlenDataBuffer = BlenBuffer + WaveBeginInElems*BlenBytesPerElem;
+
+			if (Kernel.MakeConstantU8 || Kernel.MakeConstantU16)
+			{
+				BlenDataBuffer = Register0.Data;
+			}
+			else if (Kernel.BlendFormat)
+			{
+				Kernel.BlendFormat(Register0.Data, BlenDataBuffer, NumWaveElems);
+				BlenDataBuffer = Register0.Data;
+			}
+
+			const uint8* BlendedResultBuffer = BaseDataBuffer; 
+			if (Kernel.ColorBlend)
+			{
+				uint8* MutableColorBlendResultBuffer = Kernel.Mix || Kernel.SelectAlpha ? Register1.Data : ResultDataBuffer;
+				BlendedResultBuffer = MutableColorBlendResultBuffer;
+
+				if (Kernel.ColorBlend == OpImageLayerBlendOps::Blend)
+				{
+					BlendedResultBuffer = BlenDataBuffer;
+				}
+				else if (Kernel.ColorBlendWithWord)
+				{
+					check(Kernel.MakeConstantU16);
+					Kernel.ColorBlendWithWord(MutableColorBlendResultBuffer, BaseDataBuffer, (uint16*)BlenDataBuffer, BaseBytesPerElem*NumWaveElems);
+				}
+				else
+				{
+					Kernel.ColorBlend(MutableColorBlendResultBuffer, BaseDataBuffer, BlenDataBuffer, BaseBytesPerElem*NumWaveElems);
+				}
+			}
+
+			const uint8* AlphaBlendedResultBuffer = BaseDataBuffer;
+			if (Kernel.AlphaBlend)
+			{
+				uint8* MutableAlphaBlendedResult = Register2.Data;
+				AlphaBlendedResultBuffer = MutableAlphaBlendedResult;
+
+				if (Kernel.AlphaBlend == OpImageLayerBlendOps::Blend)
+				{
+					AlphaBlendedResultBuffer = BlenDataBuffer;
+				}
+				else if (Kernel.AlphaBlendWithWord)
+				{
+					check(Kernel.MakeConstantU16);
+					Kernel.AlphaBlendWithWord(MutableAlphaBlendedResult, BaseDataBuffer, (uint16*)BlenDataBuffer, BaseBytesPerElem*NumWaveElems);
+				}
+				else
+				{
+					Kernel.AlphaBlend(MutableAlphaBlendedResult, BaseDataBuffer, BlenDataBuffer, BaseBytesPerElem*NumWaveElems);
+				}
+
+			}
+
+			if (Kernel.ExtractMaskFromBlendAlpha)
+			{
+				Kernel.ExtractMaskFromBlendAlpha(Register3.Data, BlenDataBuffer, NumWaveElems);
+				MaskDataBuffer = Register3.Data;
+			}
+			else if (Kernel.MaskDecode)
+			{
+				check(InOutDecodeToken != 0);
+				//NOTE: Decode operation, unlike the other operations, use the batch buffer and not the wave buffer 
+				//      because the data offset is carried by the DecodeToken.
+				Kernel.MaskDecode(Register3.Data, MaskBuffer, InOutDecodeToken, NumWaveElems);
+
+				MaskDataBuffer = Register3.Data;
+			}
+
+			if (Kernel.MaskFormat)
+			{
+				Kernel.MaskFormat(Register4.Data, MaskDataBuffer, NumWaveElems);
+				MaskDataBuffer = Register4.Data;
+			}
+
+			const uint8* MixResultBuffer = BlendedResultBuffer;
+			if (Kernel.Mix)
+			{
+				uint8* MutableMixResultBuffer = Kernel.SelectAlpha ? Register3.Data : ResultDataBuffer;
+				MixResultBuffer = MutableMixResultBuffer;
+
+				Kernel.Mix(MutableMixResultBuffer, BaseDataBuffer, BlendedResultBuffer, MaskDataBuffer, NumWaveElems*BaseBytesPerElem);
+			}
+
+			if (Kernel.SelectAlpha)
+			{
+				Kernel.SelectAlpha(ResultDataBuffer, MixResultBuffer, AlphaBlendedResultBuffer, NumWaveElems); 
+			}
+			else if (MixResultBuffer != ResultDataBuffer)
+			{
+				FMemory::Memcpy(ResultDataBuffer, MixResultBuffer, NumWaveElems*BaseBytesPerElem);
+			}
+		}
+	}
+
+	void RunKernelOnBatch(const FOpCombineKernel& Kernel, const FOpLayerBatchViews& BatchViews)
+	{
+		constexpr int32 MaxElemSize = 4;
+		constexpr int32 NumCacheLinesPerWave = 8;
+		constexpr int32 WaveSizeInElems = (PLATFORM_CACHE_LINE_SIZE * NumCacheLinesPerWave)/MaxElemSize;
+
+		struct alignas(PLATFORM_CACHE_LINE_SIZE) FTemporalRegister
+		{
+			uint8 Data[WaveSizeInElems*MaxElemSize];	
+		};
+
+		FTemporalRegister Register0;
+		FTemporalRegister Register1;
+		FTemporalRegister Register2;
+		FTemporalRegister Register3;
+		FTemporalRegister Register4;
+
+		int32 BaseBytesPerElem = BatchViews.Base.Num() / BatchViews.NumElems;
+		int32 BlenBytesPerElem = BatchViews.Blend.Num() / BatchViews.NumElems;
+
+		check(BaseBytesPerElem <= MaxElemSize);
+		check(BlenBytesPerElem <= MaxElemSize);
+
+		uint8* ResultBuffer = BatchViews.Result.GetData();
+		const uint8* BaseBuffer = BatchViews.Base.GetData();
+		const uint8* BlenBuffer = BatchViews.Blend.GetData();
+		const uint8* MaskBuffer = BatchViews.Mask.GetData();
+
+		uint64 InOutDecodeToken = BatchViews.DecodeTokenRLE;
+
+		if (Kernel.MakeConstant)
+		{
+			Kernel.MakeConstant(Register1.Data, Kernel.Constant, WaveSizeInElems);
+		}
+
+		const int32 NumWaves = FMath::DivideAndRoundUp<int32>(BatchViews.NumElems, WaveSizeInElems);
+		for (int32 Idx = 0; Idx < NumWaves; ++Idx)
+		{
+			int32 WaveBeginInElems = Idx * WaveSizeInElems;
+			int32 WaveEndInElems   = FMath::Min(WaveBeginInElems + WaveSizeInElems, BatchViews.NumElems); 
+
+			int32 NumWaveElems = WaveEndInElems - WaveBeginInElems;
+
+			uint8* ResultDataBuffer     = ResultBuffer + WaveBeginInElems*BaseBytesPerElem;
+			const uint8* BaseDataBuffer = BaseBuffer + WaveBeginInElems*BaseBytesPerElem;
+			const uint8* MaskDataBuffer = MaskBuffer + WaveBeginInElems;
+			const uint8* BlenDataBuffer = BlenBuffer + WaveBeginInElems*BlenBytesPerElem;
+
+			const uint8* FormatedBaseBuffer = BaseDataBuffer;
+			if (Kernel.BaseFormat)
+			{
+				Kernel.BaseFormat(Register0.Data, BaseDataBuffer, NumWaveElems);
+				FormatedBaseBuffer = Register0.Data;
+			}
+
+			const uint8* FormatedBlenBuffer = Kernel.MakeConstant ? Register1.Data : BlenDataBuffer;
+			if (Kernel.BlendFormat)
+			{
+				Kernel.BlendFormat(Register1.Data, BlenDataBuffer, NumWaveElems);
+				FormatedBlenBuffer = Register1.Data;
+			}
+
+			const uint8* CombineBuffer = FormatedBaseBuffer;
+			if (Kernel.Combine)
+			{
+				Kernel.Combine(Register2.Data, FormatedBaseBuffer, FormatedBlenBuffer, NumWaveElems);
+				CombineBuffer = Register2.Data;
+			}
+
+			const uint8* FormatedMaskBuffer = MaskDataBuffer;
+			if (Kernel.MaskDecode)
+			{
+				Kernel.MaskDecode(Register3.Data, MaskDataBuffer, InOutDecodeToken, NumWaveElems);
+				FormatedMaskBuffer = Register3.Data;
+			}
+
+			if (Kernel.MaskFormat)
+			{
+				Kernel.MaskFormat(Register4.Data, FormatedMaskBuffer, NumWaveElems);
+				FormatedMaskBuffer = Register4.Data;
+			}
+
+			const uint8* MixedBuffer = CombineBuffer;
+			if (Kernel.Mix)
+			{
+				uint8* MutableMixedBuffer = Kernel.ResultFormat ? Register1.Data : ResultDataBuffer;
+				MixedBuffer = MutableMixedBuffer;
+
+				Kernel.Mix(MutableMixedBuffer, FormatedBaseBuffer, CombineBuffer, FormatedMaskBuffer, Kernel.CombineBytesPerElem*NumWaveElems);
+			}
+
+			if (Kernel.ResultFormat)
+			{
+				Kernel.ResultFormat(ResultDataBuffer, MixedBuffer, NumWaveElems);
+			}
+		}
+	}
+
+
+	int32 GetOpLayerNumBathesLODRangeForMaskRLE(const FOpLayerBatchArgs& Args)
+	{
+		// For masks with L_UByteRLE format, Parallel execution is not supported. One batch per lod.
+		check(Args.Result);
+		check(Args.Result->GetLODCount() >= Args.LODEnd);
+		return Args.LODEnd - Args.LODBegin;
+	}
+
+	FOpLayerBatchViews GetOpLayerBatchViewsLODRangeForMaskRLE(int32 BatchId, const FOpLayerBatchArgs& Args)
+	{
+		FOpLayerBatchViews BatchViewsResult;
+
+		check(Args.Result);
+		BatchViewsResult.Result = Args.Result->DataStorage.GetLOD(BatchId);
+
+		BatchViewsResult.NumElems = BatchViewsResult.Result.Num() / Args.ResultBytesPerElem;
+
+		if (Args.Base)
+		{
+			BatchViewsResult.Base = Args.Base->DataStorage.GetLOD(BatchId);
+			check(BatchViewsResult.NumElems == BatchViewsResult.Base.Num() / Args.BaseBytesPerElem);
+		}
+
+		if (Args.Blend)
+		{
+			BatchViewsResult.Blend = Args.Blend->DataStorage.GetLOD(BatchId);
+			check(BatchViewsResult.NumElems == BatchViewsResult.Blend.Num() / Args.BlendBytesPerElem);
+		}
+
+		check(Args.Mask)
+		{
+			BatchViewsResult.Mask = Args.Mask->DataStorage.GetLOD(BatchId);
+
+			FIntVector2 LODSize = Args.Mask->CalculateMipSize(BatchId);	
+
+			BatchViewsResult.DecodeTokenRLE = GetBeginTokenIterativeDecodeRLE_L(
+					Args.Mask->GetLODData(BatchId), LODSize.Y);
+		}
+
+		return BatchViewsResult;
+	}
+
+	FORCENOINLINE int32 GetOpLayerNumBatchesLODRange(const FOpLayerBatchArgs& Args)
+	{
+		check(Args.Result);
+
+		if (Args.Mask && Args.Mask->GetFormat() == EImageFormat::L_UByteRLE)
+		{
+			return GetOpLayerNumBathesLODRangeForMaskRLE(Args);
+		}
+
+		const int32 NumBatches = Args.Result->DataStorage.GetNumBatchesLODRange(Args.BatchNumElems, Args.ResultBytesPerElem, Args.LODBegin, Args.LODEnd);
+
+#if DO_CHECK
+		if (Args.Base)
+		{
+			check(NumBatches == Args.Base->DataStorage.GetNumBatchesLODRange(Args.BatchNumElems, Args.BaseBytesPerElem, Args.LODBegin, Args.LODEnd));
+		}
+	
+		if (Args.Blend)
+		{
+			check(NumBatches == Args.Blend->DataStorage.GetNumBatchesLODRange(Args.BatchNumElems, Args.BlendBytesPerElem, Args.LODBegin, Args.LODEnd));
+		}
+
+		if (Args.Mask)
+		{
+			check(NumBatches == Args.Mask->DataStorage.GetNumBatchesLODRange(Args.BatchNumElems, Args.MaskBytesPerElem, Args.LODBegin, Args.LODEnd));
+		}
+#endif
+		return NumBatches;
+	}
+	
+	FORCENOINLINE int32 GetOpLayerNumBatchesLODRangeOffsetViews(const FOpLayerBatchArgs& Args)
+	{
+		const bool bOnlyFirstLOD = Args.FirstLODOffsetInElems >= 0;
+
+		check(Args.Result);
+		const int32 NumBatches = bOnlyFirstLOD
+				? Args.Result->DataStorage.GetNumBatchesFirstLODOffset(Args.BatchNumElems, Args.ResultBytesPerElem, Args.FirstLODOffsetInElems*Args.ResultBytesPerElem)
+				: Args.Result->DataStorage.GetNumBatchesLODRange(Args.BatchNumElems, Args.ResultBytesPerElem, Args.LODBegin, Args.LODEnd);
+
+#if DO_CHECK
+		if (Args.Base)
+		{
+			const int32 BaseNumBatches = bOnlyFirstLOD
+					? Args.Base->DataStorage.GetNumBatchesFirstLODOffset(Args.BatchNumElems, Args.BaseBytesPerElem, Args.FirstLODOffsetInElems*Args.BaseBytesPerElem)
+					: Args.Base->DataStorage.GetNumBatchesLODRange(Args.BatchNumElems, Args.BaseBytesPerElem, Args.LODBegin, Args.LODEnd);
+
+			check(NumBatches == BaseNumBatches);
+		}
+	
+		if (Args.Blend)
+		{
+			const int32 BlendNumBatches = bOnlyFirstLOD
+					? Args.Blend->DataStorage.GetNumBatchesFirstLODOffset(Args.BatchNumElems, Args.BlendBytesPerElem, Args.FirstLODOffsetInElems*Args.BlendBytesPerElem)
+					: Args.Blend->DataStorage.GetNumBatchesLODRange(Args.BatchNumElems, Args.BlendBytesPerElem, Args.LODBegin, Args.LODEnd);
+
+			check(NumBatches == BlendNumBatches);
+		}
+
+		if (Args.Mask)
+		{
+			const int32 MaskNumBatches = bOnlyFirstLOD
+					? Args.Mask->DataStorage.GetNumBatchesFirstLODOffset(Args.BatchNumElems, Args.MaskBytesPerElem, Args.FirstLODOffsetInElems*Args.MaskBytesPerElem)
+					: Args.Mask->DataStorage.GetNumBatchesLODRange(Args.BatchNumElems, Args.MaskBytesPerElem, Args.LODBegin, Args.LODEnd);
+
+			check(NumBatches == MaskNumBatches);
+		}
+#endif
+		return NumBatches;
+	}
+
+
+	FORCENOINLINE FOpLayerBatchViews GetOpLayerBatchViews(int32 BatchId, const FOpLayerBatchArgs& Args)
+	{
+		FOpLayerBatchViews BatchViewsResult;
+
+		check(Args.Result);
+		BatchViewsResult.Result = Args.Result->DataStorage.GetBatch(BatchId, Args.BatchNumElems, Args.ResultBytesPerElem);
+
+		BatchViewsResult.NumElems = BatchViewsResult.Result.Num() / Args.ResultBytesPerElem;
+
+		if (Args.Base)
+		{
+			BatchViewsResult.Base = Args.Base->DataStorage.GetBatch(BatchId, Args.BatchNumElems, Args.BaseBytesPerElem);
+			check(BatchViewsResult.NumElems == BatchViewsResult.Base.Num() / Args.BaseBytesPerElem);
+		}
+
+		if (Args.Blend)
+		{
+			BatchViewsResult.Blend = Args.Blend->DataStorage.GetBatch(BatchId, Args.BatchNumElems, Args.BlendBytesPerElem);
+			check(BatchViewsResult.NumElems == BatchViewsResult.Blend.Num() / Args.BlendBytesPerElem);
+		}
+
+		if (Args.Mask)
+		{
+			BatchViewsResult.Mask = Args.Mask->DataStorage.GetBatch(BatchId, Args.BatchNumElems, Args.MaskBytesPerElem);
+			check(BatchViewsResult.NumElems == BatchViewsResult.Mask.Num() / Args.MaskBytesPerElem);
+		}
+
+		return BatchViewsResult;
+	}
+
+	FORCENOINLINE FOpLayerBatchViews GetOpLayerBatchLODRangeViews(int32 BatchId, const FOpLayerBatchArgs& Args)
+	{
+		FOpLayerBatchViews BatchViewsResult;
+		
+		if (Args.Mask)
+		{
+			if (Args.Mask->GetFormat() == EImageFormat::L_UByteRLE)
+			{
+				return GetOpLayerBatchViewsLODRangeForMaskRLE(BatchId, Args);
+			}
+		}
+
+		check(Args.Result);
+		BatchViewsResult.Result = 
+					Args.Result->DataStorage.GetBatchLODRange(BatchId, Args.BatchNumElems, Args.ResultBytesPerElem, Args.LODBegin, Args.LODEnd);
+
+		BatchViewsResult.NumElems = BatchViewsResult.Result.Num() / Args.ResultBytesPerElem;
+
+		if (Args.Base)
+		{
+			BatchViewsResult.Base = 
+					Args.Base->DataStorage.GetBatchLODRange(BatchId, Args.BatchNumElems, Args.BaseBytesPerElem, Args.LODBegin, Args.LODEnd);
+			check(BatchViewsResult.NumElems == BatchViewsResult.Base.Num() / Args.BaseBytesPerElem);
+		}
+
+		if (Args.Blend)
+		{
+			BatchViewsResult.Blend = 
+					Args.Blend->DataStorage.GetBatchLODRange(BatchId, Args.BatchNumElems, Args.BlendBytesPerElem, Args.LODBegin, Args.LODEnd);
+			check(BatchViewsResult.NumElems == BatchViewsResult.Blend.Num() / Args.BlendBytesPerElem);
+		}
+
+		if (Args.Mask)
+		{
+			BatchViewsResult.Mask = 
+					Args.Mask->DataStorage.GetBatchLODRange(BatchId, Args.BatchNumElems, Args.MaskBytesPerElem, Args.LODBegin, Args.LODEnd);
+			check(BatchViewsResult.NumElems == BatchViewsResult.Mask.Num() / Args.MaskBytesPerElem);
+		}
+
+		return BatchViewsResult;
+	}
+
+	FORCENOINLINE FOpLayerBatchViews GetOpLayerBatchLODRangeOffsetViews(int32 BatchId, const FOpLayerBatchArgs& Args)
+	{
+		FOpLayerBatchViews BatchViewsResult;
+		
+		const bool bOnlyFirstLOD = Args.FirstLODOffsetInElems >= 0;
+
+		check(Args.Result);
+		BatchViewsResult.Result = bOnlyFirstLOD
+				? Args.Result->DataStorage.GetBatchFirstLODOffset(BatchId, Args.BatchNumElems, Args.ResultBytesPerElem, Args.FirstLODOffsetInElems*Args.ResultBytesPerElem) 
+				: Args.Result->DataStorage.GetBatchLODRange(BatchId, Args.BatchNumElems, Args.ResultBytesPerElem, Args.LODBegin, Args.LODEnd);
+
+		BatchViewsResult.NumElems = BatchViewsResult.Result.Num() / Args.ResultBytesPerElem;
+		
+		if (Args.Base)
+		{
+			BatchViewsResult.Base = bOnlyFirstLOD
+					? Args.Base->DataStorage.GetBatchFirstLODOffset(BatchId, Args.BatchNumElems, Args.BaseBytesPerElem, Args.FirstLODOffsetInElems*Args.BaseBytesPerElem) 
+					: Args.Base->DataStorage.GetBatchLODRange(BatchId, Args.BatchNumElems, Args.BaseBytesPerElem, Args.LODBegin, Args.LODEnd); 
+
+			check(BatchViewsResult.NumElems == BatchViewsResult.Base.Num() / Args.BaseBytesPerElem);
+		}
+
+		if (Args.Blend)
+		{
+			BatchViewsResult.Blend = bOnlyFirstLOD
+					? Args.Blend->DataStorage.GetBatchFirstLODOffset(BatchId, Args.BatchNumElems, Args.BlendBytesPerElem, Args.FirstLODOffsetInElems*Args.BlendBytesPerElem) 
+					: Args.Blend->DataStorage.GetBatchLODRange(BatchId, Args.BatchNumElems, Args.BlendBytesPerElem, Args.LODBegin, Args.LODEnd); 
+
+			check(BatchViewsResult.NumElems == BatchViewsResult.Blend.Num() / Args.BlendBytesPerElem);
+		}
+
+		if (Args.Mask)
+		{
+			BatchViewsResult.Mask = bOnlyFirstLOD
+					? Args.Mask->DataStorage.GetBatchFirstLODOffset(BatchId, Args.BatchNumElems, Args.MaskBytesPerElem, Args.FirstLODOffsetInElems*Args.MaskBytesPerElem) 
+					: Args.Mask->DataStorage.GetBatchLODRange(BatchId, Args.BatchNumElems, Args.MaskBytesPerElem, Args.LODBegin, Args.LODEnd); 
+
+			check(BatchViewsResult.NumElems == BatchViewsResult.Mask.Num() / Args.MaskBytesPerElem);
+		}
+
+		return BatchViewsResult;
+	}
+
+
+	void ImageLayerImpl(
+			FImage* DestImage, const FImage* BaseImage, const FImage* BlenImage, const FImage* MaskImage, 
+			FBlendFuncType* ColorBlendFunc, FBlendFuncType* AlphaBlendFunc,
+			bool bBlendWithConstant, const FVector4f& BlendConstant,
+			int32 LODBegin, int32 LODEnd, 
+			bool bApplyColorBlendToAlpha, bool bUseMaskFromBlendAlpha, bool bUseBlendSourceFromBlendAlpha, uint8 BlendAlphaSourceChannel)
+	{
+		MUTABLE_CPUPROFILER_SCOPE(ImageLayer)
+
+		check(DestImage && DestImage->GetFormat() == GetUncompressedFormat(DestImage->GetFormat()));
+		check(BaseImage && BaseImage->GetFormat() == GetUncompressedFormat(BaseImage->GetFormat()));
+		
+		check(LODEnd <= BaseImage->GetLODCount());
+
+		if (!bBlendWithConstant)
+		{
+			check(BlenImage);
+			check(BlenImage->GetFormat() == GetUncompressedFormat(BlenImage->GetFormat()));
+			check(BlenImage->GetLODCount() >= LODEnd);
+
+			check(BlenImage->GetSize() == BaseImage->GetSize());
+		}
+
+		check(BaseImage->GetFormat() == DestImage->GetFormat());
+		check(BaseImage->GetSize() == DestImage->GetSize());
+
+		if (MaskImage)
+		{
+			check(MaskImage->GetFormat() == EImageFormat::L_UByte || MaskImage->GetFormat() == EImageFormat::L_UByteRLE);
+		}
+
+		const int32 BaseBytesPerElem = GetImageFormatData(BaseImage->GetFormat()).BytesPerBlock;
+		const int32 BlenBytesPerElem = BlenImage ? GetImageFormatData(BlenImage->GetFormat()).BytesPerBlock : 0;
+
+		bool bOnlyFirstLOD = (LODBegin == 0 || LODEnd == 1) || BaseImage->GetLODCount() == 1;
+		
+		int32 FirstLODOffsetInElems = -1;	
+		int32 NumRelevantElems = -1;
+
+		if (BlenImage && BlenImage->Flags & FImage::IF_HAS_RELEVANCY_MAP && bOnlyFirstLOD)
+		{
+			check(BlenImage->RelevancyMaxY < BaseImage->GetSizeY());
+			check(BlenImage->RelevancyMaxY >= BlenImage->RelevancyMinY);
+
+			uint16 SizeX = BaseImage->GetSizeX();
+			NumRelevantElems = (BlenImage->RelevancyMaxY - BlenImage->RelevancyMinY + 1) * SizeX;
+			
+			FirstLODOffsetInElems = BlenImage->RelevancyMinY * SizeX;
+		}
+	
+		constexpr int32 BatchNumElems = 4096*2;
+
+		OpImageLayerInternal::FOpLayerBatchArgs BatchArgs;
+		BatchArgs.BatchNumElems         = BatchNumElems;
+		BatchArgs.LODBegin              = LODBegin;
+		BatchArgs.LODEnd                = LODEnd;
+		BatchArgs.FirstLODOffsetInElems = bOnlyFirstLOD ? FirstLODOffsetInElems : -1;
+		BatchArgs.Result                = DestImage;
+		BatchArgs.Base                  = BaseImage;
+		BatchArgs.Blend                 = BlenImage;
+		BatchArgs.Mask                  = MaskImage;
+		BatchArgs.ResultBytesPerElem    = BaseBytesPerElem;
+		BatchArgs.BaseBytesPerElem      = BaseBytesPerElem;
+		BatchArgs.BlendBytesPerElem     = BlenBytesPerElem;
+		BatchArgs.MaskBytesPerElem      = 1;
+
+		OpImageLayerInternal::FMakeOpBlendKernelArgs MakeKernelArgs
+		{
+			.Dest  = DestImage,
+			.Base  = BaseImage,
+			.Blend = BlenImage,
+			.Mask  = MaskImage,
+
+			.ColorBlendFunc = ColorBlendFunc,
+			.AlphaBlendFunc = AlphaBlendFunc,
+
+			.bApplyColorBlendToAlpha       = bApplyColorBlendToAlpha,
+			.bUseBlendSourceFromBlendAlpha = bUseBlendSourceFromBlendAlpha,
+			.bUseMaskFromBlendAlpha        = bUseMaskFromBlendAlpha,
+			.BlendAlphaSourceChannel       = BlendAlphaSourceChannel,
+
+			.bBlendWithConstant = bBlendWithConstant, 
+			.RealConstant       = BlendConstant,
+		};
+
+		int32 NumBatches = 0;
+		if (BatchArgs.Mask && BatchArgs.Mask->GetFormat() == EImageFormat::L_UByteRLE)
+		{
+			NumBatches = OpImageLayerInternal::GetOpLayerNumBatchesLODRange(BatchArgs);
+		}
+		else
+		{
+			NumBatches = OpImageLayerInternal::GetOpLayerNumBatchesLODRangeOffsetViews(BatchArgs);
+		}
+
+		bool bUseRelevancy = NumRelevantElems != -1;
+		if (BatchArgs.Mask)
+		{
+			bUseRelevancy &= BatchArgs.Mask->GetFormat() != EImageFormat::L_UByteRLE;
+		}
+
+		// This will always be an upper-bound for bOnlyFirtsLODs, check if it performs as expected or it needs more fine tune.
+		const int32 NumRelevantBatches = bUseRelevancy
+				? FMath::DivideAndRoundUp(NumRelevantElems, BatchNumElems)
+				: NumBatches;
+
+		const OpImageLayerInternal::FOpBlendKernel Kernel = OpImageLayerInternal::MakeKernel(MakeKernelArgs);
+		ParallelExecutionUtils::InvokeBatchParallelFor(NumRelevantBatches, 
+		[
+			&BatchArgs, &Kernel
+		] (int32 BatchId)
+		{
+
+			OpImageLayerInternal::FOpLayerBatchViews BatchViews;
+			
+			if (BatchArgs.Mask && BatchArgs.Mask->GetFormat() == EImageFormat::L_UByteRLE)
+			{
+				BatchViews = OpImageLayerInternal::GetOpLayerBatchLODRangeViews(BatchId, BatchArgs);
+			}
+			else
+			{
+				BatchViews = OpImageLayerInternal::GetOpLayerBatchLODRangeOffsetViews(BatchId, BatchArgs);
+			}
+
+			OpImageLayerInternal::RunKernelOnBatch(Kernel, BatchViews);
+		});
+
+	}
+
+	void ImageLayerCombineImpl(
+			FImage* DestImage, const FImage* BaseImage, const FImage* BlenImage, const FImage* MaskImage, 
+			FCombineFuncType* CombineFunc,
+			bool bCombineWithConstant, const FVector4f& Constant,
+			int32 LODBegin, int32 LODEnd)
+	{
+		MUTABLE_CPUPROFILER_SCOPE(ImageLayerNormalCombine)
+
+		check(BaseImage->GetFormat() == DestImage->GetFormat());
+		check(DestImage->GetFormat() == GetUncompressedFormat(DestImage->GetFormat()));
+		check(BaseImage->GetFormat() == GetUncompressedFormat(BaseImage->GetFormat()));
+		
+		check(BaseImage->GetLODCount() >= LODEnd);
 
-			for (int32 LOD = 0; LOD < NumLODs; ++LOD)
-			{
-				const uint8* BaseBuf = BaseImage->GetLODData(LOD);
-				uint8* DestBuf = DestImage->GetLODData(LOD);
-				const uint8* BlendedBuf = BlendImage->GetLODData(LOD);
-				const uint8* MaskBuf = MaskImage->GetLODData(LOD);
+		if (!bCombineWithConstant)
+		{
+			check(BlenImage)
+			check(BlenImage->GetFormat() == GetUncompressedFormat(BlenImage->GetFormat()));
+			check(BlenImage->GetLODCount() >= LODEnd);
+		}
 
-				// This could happen in case of missing data files.
-				if (!BaseBuf || !BlendedBuf || !MaskBuf || !DestBuf)
-				{
-					continue;
-				}
+		if (MaskImage)
+		{
+			check(MaskImage->GetFormat() == GetUncompressedFormat(MaskImage->GetFormat()));
+			check(MaskImage->GetFormat() == EImageFormat::L_UByte || MaskImage->GetFormat() == EImageFormat::L_UByteRLE);
+		}
 
-				// Remove RLE header, mip size and row sizes.
-				MaskBuf += sizeof(uint32) + Rows * sizeof(uint32);
 
-				for (int32 RowIndex = 0; RowIndex < Rows; ++RowIndex)
-				{
-					const uint8* DestRowEnd = DestBuf + Width * BASE_CHANNEL_STRIDE;
-					while (DestBuf != DestRowEnd)
-					{
-						// Decode header
-						uint16 Equal;
-						FMemory::Memcpy(&Equal, MaskBuf, 2);
-						MaskBuf += 2;
+		const int32 BaseBytesPerElem = GetImageFormatData(BaseImage->GetFormat()).BytesPerBlock;
+		const int32 BlenBytesPerElem = BlenImage ? GetImageFormatData(BlenImage->GetFormat()).BytesPerBlock : 0;
 
-						uint8 Different = *MaskBuf;
-						++MaskBuf;
+		bool bOnlyFirstLOD = (LODBegin == 0 && LODEnd == 1) || BaseImage->GetLODCount() == 1;
+		
+		int32 FirstLODDataOffset = -1;	
+		int32 NumRelevantElems = -1;
 
-						uint8 EqualPixel = *MaskBuf;
-						++MaskBuf;
+		if (BlenImage && BlenImage->Flags & FImage::IF_HAS_RELEVANCY_MAP && bOnlyFirstLOD)
+		{
+			check(BlenImage->RelevancyMaxY < BaseImage->GetSizeY());
+			check(BlenImage->RelevancyMaxY >= BlenImage->RelevancyMinY);
 
-						// Equal pixels
-						//check(DestBuf + BASE_CHANNEL_STRIDE * Equal <= BaseImage->GetLODDataSize(0));
-						if (EqualPixel == 255)
-						{
-							for (int32 I = 0; I < Equal; ++I)
-							{
-								for (int32 C = 0; C < CHANNELS_TO_BLEND; ++C)
-								{
-									uint32 Base = BaseBuf[BASE_CHANNEL_STRIDE * I + C];
-									uint32 Blended = BlendedBuf[BLENDED_CHANNEL_STRIDE * I + C];
-									uint32 Result = BLEND_FUNC(Base, Blended);
-									if constexpr (CLAMP)
-									{
-										DestBuf[BASE_CHANNEL_STRIDE * I + C] = (uint8)FMath::Min(255u, Result);
-									}
-									else
-									{
-										DestBuf[BASE_CHANNEL_STRIDE * I + C] = (uint8)Result;
-									}
-								}
-
-								// Copy the unblended channels
-								// \TODO: unnecessary when doing it in-place?
-								if constexpr (UnblendedChannels > 0)
-								{
-									for (int32 C = 0; C < UnblendedChannels; ++C)
-									{
-										DestBuf[BASE_CHANNEL_STRIDE * I + CHANNELS_TO_BLEND + C] = BaseBuf[BASE_CHANNEL_STRIDE * I + CHANNELS_TO_BLEND + C];
-									}
-								}
-							}
-						}
-						else if (EqualPixel > 0)
-						{
-							for (int32 I = 0; I < Equal; ++I)
-							{
-								for (int32 C = 0; C < CHANNELS_TO_BLEND; ++C)
-								{
-									uint32 Base = BaseBuf[BASE_CHANNEL_STRIDE * I + C];
-									uint32 Blended = BlendedBuf[BLENDED_CHANNEL_STRIDE * I + C];
-									uint32 Result = BLEND_FUNC_MASKED(Base, Blended, EqualPixel);
-
-									if constexpr (CLAMP)
-									{
-										DestBuf[BASE_CHANNEL_STRIDE * I + C] = (uint8)FMath::Min(255u, Result);
-									}
-									else
-									{
-										DestBuf[BASE_CHANNEL_STRIDE * I + C] = (uint8)Result;
-									}
-								}
-
-								// Copy the unblended channels
-								// \TODO: unnecessary when doing it in-place?
-								if constexpr (UnblendedChannels > 0)
-								{
-									for (int32 C = 0; C < UnblendedChannels; ++C)
-									{
-										DestBuf[BASE_CHANNEL_STRIDE * I + CHANNELS_TO_BLEND + C] = BaseBuf[BASE_CHANNEL_STRIDE * I + CHANNELS_TO_BLEND + C];
-									}
-								}
-							}
-						}
-						else
-						{
-							// It could happen if xxxxxOnBase
-							if (DestBuf != BaseBuf)
-							{
-								FMemory::Memmove(DestBuf, BaseBuf, BASE_CHANNEL_STRIDE * Equal);
-							}
-						}
+			uint16 SizeX = BaseImage->GetSizeX();
+			NumRelevantElems = (BlenImage->RelevancyMaxY - BlenImage->RelevancyMinY + 1) * SizeX;
+			
+			FirstLODDataOffset = BlenImage->RelevancyMinY * SizeX * BlenBytesPerElem;
+		}
+	
+		constexpr int32 BatchNumElems = 4096*2;
 
-						DestBuf += BASE_CHANNEL_STRIDE * Equal;
-						BaseBuf += BASE_CHANNEL_STRIDE * Equal;
-						BlendedBuf += BLENDED_CHANNEL_STRIDE * Equal;
+		OpImageLayerInternal::FOpLayerBatchArgs BatchArgs;
+		BatchArgs.BatchNumElems         = BatchNumElems;
+		BatchArgs.LODBegin              = LODBegin;
+		BatchArgs.LODEnd                = LODEnd;
+		BatchArgs.FirstLODOffsetInElems = bOnlyFirstLOD ? FirstLODDataOffset : -1;
+		BatchArgs.Result                = DestImage;
+		BatchArgs.Base                  = BaseImage;
+		BatchArgs.Blend                 = BlenImage;
+		BatchArgs.Mask                  = MaskImage;
+		BatchArgs.ResultBytesPerElem    = BaseBytesPerElem;
+		BatchArgs.BaseBytesPerElem      = BaseBytesPerElem;
+		BatchArgs.BlendBytesPerElem     = BlenBytesPerElem;
+		BatchArgs.MaskBytesPerElem      = 1;
 
-						// Different pixels
-						//check(DestBuf + BASE_CHANNEL_STRIDE * Different <= BaseImage->GetDataSize(0));
-						for (int32 I = 0; I < Different; ++I)
-						{
-							for (int32 C = 0; C < CHANNELS_TO_BLEND; ++C)
-							{
-								uint32 Mask = MaskBuf[I];
-								uint32 Base = BaseBuf[BASE_CHANNEL_STRIDE * I + C];
-								uint32 Blended = BlendedBuf[BLENDED_CHANNEL_STRIDE * I + C];
-								uint32 Result = BLEND_FUNC_MASKED(Base, Blended, Mask);
-								if (CLAMP)
-								{
-									DestBuf[BASE_CHANNEL_STRIDE * I + C] = (uint8)FMath::Min(255u, Result);
-								}
-								else
-								{
-									DestBuf[BASE_CHANNEL_STRIDE * I + C] = (uint8)Result;
-								}
-							}
+		OpImageLayerInternal::FMakeOpCombineKernelArgs MakeKernelArgs
+		{
+			.Dest  = DestImage,
+			.Base  = BaseImage,
+			.Blend = BlenImage,
+			.Mask  = MaskImage,
 
-							// Copy the unblended channels
-							// \TODO: unnecessary when doing it in-place?
-							if constexpr (UnblendedChannels > 0)
-							{
-								for (int32 C = 0; C < UnblendedChannels; ++C)
-								{
-									DestBuf[BASE_CHANNEL_STRIDE * I + CHANNELS_TO_BLEND + C] = BaseBuf[BASE_CHANNEL_STRIDE * I + CHANNELS_TO_BLEND + C];
-								}
-							}
-						}
+			.CombineFunc = CombineFunc,
 
-						DestBuf += BASE_CHANNEL_STRIDE * Different;
-						BaseBuf += BASE_CHANNEL_STRIDE * Different;
-						BlendedBuf += BLENDED_CHANNEL_STRIDE * Different;
-						MaskBuf += Different;
-					}
-				}
+			.bCombineWithConstant = bCombineWithConstant, 
+			.Constant             = FLinearColor(Constant).QuantizeRound(),
+		};
 
-				Rows = FMath::DivideAndRoundUp(Rows, 2);
-				Width = FMath::DivideAndRoundUp(Width, 2);
-			}
+		int32 NumBatches = 0;
+		if (BatchArgs.Mask && BatchArgs.Mask->GetFormat() == EImageFormat::L_UByteRLE)
+		{
+			NumBatches = OpImageLayerInternal::GetOpLayerNumBatchesLODRange(BatchArgs);
 		}
 		else
 		{
-			checkf(false, TEXT("Unsupported mask format."));
+			NumBatches = OpImageLayerInternal::GetOpLayerNumBatchesLODRangeOffsetViews(BatchArgs);
 		}
-	}
-
-	template< uint32 (*BLEND_FUNC_MASKED)(uint32, uint32, uint32),
-		uint32 (*BLEND_FUNC)(uint32, uint32),
-		bool CLAMP,
-		// Number of total channels to actually process
-		int32 CHANNELS_TO_BLEND,
-		// Number of total channels in the base image
-		int32 BASE_CHANNEL_STRIDE,
-		// Number of total channels in the blend image
-		int32 BLENDED_CHANNEL_STRIDE>
-	void BufferLayerFormatEmbeddedMask(
-		FImage* DestImage,
-		const FImage* BaseImage,
-		const FImage* BlendImage,
-		uint32 DestOffset,
-		uint32 BaseChannelOffset,
-		bool bOnlyFirstLOD)
-	{
-		check(BaseImage->GetSizeX() == BlendImage->GetSizeX() && BaseImage->GetSizeY() == BlendImage->GetSizeY());
-		check(bOnlyFirstLOD || BaseImage->GetLODCount() <= BlendImage->GetLODCount());
-
-		constexpr int32 UnblendedChannels = BASE_CHANNEL_STRIDE - CHANNELS_TO_BLEND;
 
-		const int32 BytesPerElem = GetImageFormatData(BaseImage->GetFormat()).BytesPerBlock; 
-	
-		const int32 NumLODs = bOnlyFirstLOD ? 1 : BaseImage->GetLODCount();
+		// This will always be an upper-bound for bOnlyFirtsLODs, check if it performs as expected or it needs more fine tune.
+		const int32 NumRelevantBatches = NumRelevantElems != -1
+				? FMath::DivideAndRoundUp(NumRelevantElems, BatchNumElems)
+				: NumBatches;
 
-		constexpr int32 BatchNumElems = 4096*2;
-		OpImageLayerInternal::FOpLayerBatchArgs BatchArgs;
-		BatchArgs.BatchNumElems      = BatchNumElems;
-		BatchArgs.LODBegin           = 0;
-		BatchArgs.LODEnd             = NumLODs;
-		BatchArgs.Result             = DestImage;
-		BatchArgs.Base               = BaseImage;
-		BatchArgs.Blend              = BlendImage;
-		BatchArgs.ResultBytesPerElem = BytesPerElem;
-		BatchArgs.BaseBytesPerElem   = BytesPerElem;
-		BatchArgs.BlendBytesPerElem  = BytesPerElem;
-
-		const int32 NumBatches = OpImageLayerInternal::GetOpLayerNumBatchesLODRange(BatchArgs);
-
-		ParallelExecutionUtils::InvokeBatchParallelFor(NumBatches, 
+		const OpImageLayerInternal::FOpCombineKernel Kernel = OpImageLayerInternal::MakeKernel(MakeKernelArgs);
+		ParallelExecutionUtils::InvokeBatchParallelFor(NumRelevantBatches, 
 		[
-			&BatchArgs, UnblendedChannels, BlendImageFormat = BlendImage->GetFormat()
+			&BatchArgs, &Kernel
 		] (int32 BatchId)
 		{
-			OpImageLayerInternal::FOpLayerBatchViews BatchViews = 
-					OpImageLayerInternal::GetOpLayerBatchLODRangeViews(BatchId, BatchArgs);
-
-			uint8* const DestBuf = BatchViews.Result.GetData();
-			uint8 const * const BlendedBuf = BatchViews.Blend.GetData();
-			uint8 const * const BaseBuf = BatchViews.Base.GetData(); 
-
-			// This could happen in case of missing data files.
-			if (!BaseBuf || !BlendedBuf)
-			{
-				return;
-			}
 
-			if (BlendImageFormat == EImageFormat::RGBA_UByte)
+			OpImageLayerInternal::FOpLayerBatchViews BatchViews;
+			
+			if (BatchArgs.Mask && BatchArgs.Mask->GetFormat() == EImageFormat::L_UByteRLE)
 			{
-				const uint8* MaskBuf = BlendedBuf + 3;
-
-				for (int32 I = 0; I < BatchViews.NumElems; ++I)
-				{
-					uint32 Mask = MaskBuf[BLENDED_CHANNEL_STRIDE * I];
-					for (int32 C = 0; C < CHANNELS_TO_BLEND; ++C)
-					{
-						uint32 Base = BaseBuf[BASE_CHANNEL_STRIDE * I + C];
-						uint32 Blended = BlendedBuf[BLENDED_CHANNEL_STRIDE * I + C];
-						uint32 Result = BLEND_FUNC_MASKED(Base, Blended, Mask);
-						if constexpr (CLAMP)
-						{
-							DestBuf[BASE_CHANNEL_STRIDE * I + C] = (uint8)FMath::Min(255u, Result);
-						}
-						else
-						{
-							DestBuf[BASE_CHANNEL_STRIDE * I + C] = (uint8)Result;
-						}
-					}
-					// Copy the unblended channels
-					// \TODO: unnecessary when doing it in-place?
-					if constexpr (UnblendedChannels > 0)
-					{
-						for (int32 C = 0; C < UnblendedChannels; ++C)
-						{
-							DestBuf[BASE_CHANNEL_STRIDE * I + CHANNELS_TO_BLEND + C] = BaseBuf[BASE_CHANNEL_STRIDE * I + CHANNELS_TO_BLEND + C];
-						}
-					}
-				}
+				BatchViews = OpImageLayerInternal::GetOpLayerBatchLODRangeViews(BatchId, BatchArgs);
 			}
 			else
 			{
-				for (int32 I = 0; I < BatchViews.NumElems; ++I)
-				{
-					for (int32 C = 0; C < CHANNELS_TO_BLEND; ++C)
-					{
-						uint32 Base = BaseBuf[BASE_CHANNEL_STRIDE * I + C];
-						uint32 Blended = BlendedBuf[BLENDED_CHANNEL_STRIDE * I + C];
-						uint32 Result = BLEND_FUNC(Base, Blended);
-						if constexpr (CLAMP)
-						{
-							DestBuf[BASE_CHANNEL_STRIDE * I + C] = (uint8)FMath::Min(255u, Result);
-						}
-						else
-						{
-							DestBuf[BASE_CHANNEL_STRIDE * I + C] = (uint8)Result;
-						}
-					}
-					// Copy the unblended channels
-					// \TODO: unnecessary when doing it in-place?
-					if constexpr (UnblendedChannels > 0)
-					{
-						for (int32 C = 0; C < UnblendedChannels; ++C)
-						{
-							DestBuf[BASE_CHANNEL_STRIDE * I + CHANNELS_TO_BLEND + C] = BaseBuf[BASE_CHANNEL_STRIDE * I + CHANNELS_TO_BLEND + C];
-						}
-					}
-				}
+				BatchViews = OpImageLayerInternal::GetOpLayerBatchLODRangeOffsetViews(BatchId, BatchArgs);
 			}
+
+			OpImageLayerInternal::RunKernelOnBatch(Kernel, BatchViews);
 		});
 
 	}
+} // namespece OpImageLayerInternal;
 
+	//TODO: Remove this, use the kernel execution aproatch. Needs some masking to be implemented. 
 	template< 
 		uint32 (*BLEND_FUNC)(uint32, uint32),
 		bool CLAMP,
@@ -1699,402 +2247,46 @@ namespace OpImageLayerInternal
 
             //for (int32 lod = 0; lod < LODCount; ++lod)
             //{
-				// Remove RLE header.
-                MaskBuf += 4 + Rows*sizeof(uint32);
-
-                for (int32 RowIndex = 0; RowIndex < Rows; ++RowIndex)
-                {
-                    const uint8* DestRowEnd = DestBuf + Width*NC;
-                    while (DestBuf != DestRowEnd)
-                    {
-                        // Decode header
-                        uint16 Zeros = *(const uint16*)MaskBuf;
-                        MaskBuf += 2;
-
-                        uint16 Ones = *(const uint16*)MaskBuf;
-                        MaskBuf += 2;
-
-                        // Skip
-                        DestBuf += Zeros*NC;
-                        BlendedBuf += Zeros*NC;
-
-                        // Copy
-                        FMemory::Memmove(DestBuf, BlendedBuf, Ones*NC);
-
-                        DestBuf += NC*Ones;
-                        BlendedBuf += NC*Ones;
-                    }
-
-                    DestBuf += Stride;
-                }
-
-                //Rows = FMath::DivideAndRoundUp(Rows, 2);
-                //Width = FMath::DivideAndRoundUp(Width, 2);
-            //}
-		}
-		else
-		{
-			checkf( false, TEXT("Unsupported mask format.") );
-		}
-	}
-
-	template<uint32(*BLEND_FUNC_MASKED)(uint32, uint32, uint32),
-			 uint32(*BLEND_FUNC)(uint32, uint32),
-			 bool CLAMP>
-    void BufferLayer(FImage* DestImage,
-							const FImage* BaseImage,
-							const FImage* MaskImage,
-                            const FImage* BlendImage,
-                            bool bApplyToAlpha,
-							bool bOnlyFirstLOD)
-	{
-		if (BaseImage->GetFormat() == EImageFormat::RGB_UByte)
-		{
-            BufferLayerFormat<BLEND_FUNC_MASKED, BLEND_FUNC, CLAMP, 3, 3, 3>
-                    (DestImage, BaseImage, MaskImage, BlendImage, 0, 0, 0, bOnlyFirstLOD);
-		}
-        else if (BaseImage->GetFormat() == EImageFormat::RGBA_UByte || 
-				 BaseImage->GetFormat() == EImageFormat::BGRA_UByte)
-		{
-			if (bApplyToAlpha)
-			{
-				BufferLayerFormat<BLEND_FUNC_MASKED, BLEND_FUNC, CLAMP, 4, 4, 4>
-					(DestImage, BaseImage, MaskImage, BlendImage, 0, 0, 0, bOnlyFirstLOD);
-			}
-			else
-			{
-				BufferLayerFormat<BLEND_FUNC_MASKED, BLEND_FUNC, CLAMP, 3, 4, 4>
-					(DestImage, BaseImage, MaskImage, BlendImage, 0, 0, 0, bOnlyFirstLOD);
-			}
-		}
-		else if (BaseImage->GetFormat() == EImageFormat::L_UByte)
-		{
-            BufferLayerFormat<BLEND_FUNC_MASKED, BLEND_FUNC, CLAMP, 1, 1, 1>
-                    (DestImage, BaseImage, MaskImage, BlendImage, 0, 0, 0, bOnlyFirstLOD);
-		}
-        else
-		{
-			checkf(false, TEXT("Unsupported format."));
-		}
-	}
-
-
-	template< uint32 (*BLEND_FUNC_MASKED)(uint32, uint32, uint32),
-		uint32 (*BLEND_FUNC)(uint32, uint32),
-		bool CLAMP >
-	void BufferLayerEmbeddedMask(
-		FImage* DestImage,
-		const FImage* BaseImage,
-		const FImage* BlendImage,
-		bool bApplyToAlpha,
-		bool bOnlyFirstLOD)
-	{
-		if (BaseImage->GetFormat() == EImageFormat::RGB_UByte)
-		{
-			BufferLayerFormatEmbeddedMask<BLEND_FUNC_MASKED, BLEND_FUNC, CLAMP, 3, 3, 3>
-				(DestImage, BaseImage, BlendImage, 0, 0, bOnlyFirstLOD);
-		}
-		else if (BaseImage->GetFormat() == EImageFormat::RGBA_UByte || 
-				 BaseImage->GetFormat() == EImageFormat::BGRA_UByte)
-		{
-			if (bApplyToAlpha)
-			{
-				BufferLayerFormatEmbeddedMask<BLEND_FUNC_MASKED, BLEND_FUNC, CLAMP, 4, 4, 4>
-					(DestImage, BaseImage, BlendImage, 0, 0, bOnlyFirstLOD);
-			}
-			else
-			{
-				BufferLayerFormatEmbeddedMask<BLEND_FUNC_MASKED, BLEND_FUNC, CLAMP, 3, 4, 4>
-					(DestImage, BaseImage, BlendImage, 0, 0, bOnlyFirstLOD);
-			}
-		}
-		else
-		{
-			checkf(false, TEXT("Unsupported format."));
-		}
-	}
-
-
-	template< 
-		uint32 (*RGB_FUNC_MASKED)(uint32, uint32, uint32),
-		uint32 (*A_FUNC)(uint32, uint32),
-		bool CLAMP >
-	void BufferLayerComposite(
-		FImage* BaseImage,
-		const FImage* BlendImage,
-		bool bOnlyFirstLOD,
-		uint8 BlendAlphaSourceChannel)
-	{
-		check(BaseImage->GetFormat() == EImageFormat::RGBA_UByte);
-		check(BlendImage->GetFormat() == EImageFormat::RGBA_UByte);
-		check(BaseImage->GetSizeX() == BlendImage->GetSizeX() && BaseImage->GetSizeY() == BlendImage->GetSizeY());
-		check(bOnlyFirstLOD || BaseImage->GetLODCount() <= BlendImage->GetLODCount());
-
-		bOnlyFirstLOD = bOnlyFirstLOD || BaseImage->GetLODCount() == 1;
-		
-		int32 FirstLODDataOffset = 0;	
-		int32 NumRelevantElems = -1;
-
-		if (BlendImage->Flags & FImage::IF_HAS_RELEVANCY_MAP && bOnlyFirstLOD)
-		{
-			check(BlendImage->RelevancyMaxY < BaseImage->GetSizeY());
-			check(BlendImage->RelevancyMaxY >= BlendImage->RelevancyMinY);
-
-			uint16 SizeX = BaseImage->GetSizeX();
-			NumRelevantElems = (BlendImage->RelevancyMaxY - BlendImage->RelevancyMinY + 1) * SizeX * 4;
-			
-			FirstLODDataOffset = BlendImage->RelevancyMinY * SizeX * 4;
-		}
-
-		const int32 BytesPerElem = GetImageFormatData(BaseImage->GetFormat()).BytesPerBlock;
-		check(BytesPerElem == GetImageFormatData(BlendImage->GetFormat()).BytesPerBlock);
-
-		const int32 LODBegin = 0;
-		const int32 LODEnd = BaseImage->GetLODCount();
-
-		constexpr int32 BatchNumElems = 4096*2;
-
-		OpImageLayerInternal::FOpLayerBatchArgs BatchArgs;
-		BatchArgs.BatchNumElems      = BatchNumElems;
-		BatchArgs.LODBegin           = LODBegin;
-		BatchArgs.LODEnd             = LODEnd;
-		BatchArgs.FirstLODOffset     = bOnlyFirstLOD ? FirstLODDataOffset : -1;
-		BatchArgs.Result             = BaseImage;
-		BatchArgs.Blend              = BlendImage;
-		BatchArgs.ResultBytesPerElem = BytesPerElem;
-		BatchArgs.BlendBytesPerElem  = BytesPerElem;
-
-		const int32 NumBatches = OpImageLayerInternal::GetOpLayerNumBatchesLODRangeOffsetViews(BatchArgs);
-
-		// This will always be an upper-bound for bOnlyFirtsLODs, check if it performs as expected or it needs more fine tune.
-		const int32 NumRelevantBatches = NumRelevantElems != -1
-				? FMath::DivideAndRoundUp(NumRelevantElems, BatchNumElems)
-				: NumBatches;
-
-		ParallelExecutionUtils::InvokeBatchParallelFor(NumBatches, 
-		[
-			&BatchArgs, BlendAlphaSourceChannel
-		] (int32 BatchId)
-		{
-			OpImageLayerInternal::FOpLayerBatchViews BatchViews = 
-					OpImageLayerInternal::GetOpLayerBatchLODRangeOffsetViews(BatchId, BatchArgs);
-
-			uint8* BaseBuf = BatchViews.Result.GetData();
-			const uint8* BlendBuf = BatchViews.Blend.GetData();
-
-			for (int32 I = 0; I < BatchViews.NumElems; ++I)
-			{
-				// TODO: Optimize this (SIMD?)
-				uint32 Mask = BlendBuf[4 * I + 3];
-
-				// RGB
-				for (int32 C = 0; C < 3; ++C)
-				{
-					uint32 Base = BaseBuf[4 * I + C];
-					uint32 Blended = BlendBuf[4 * I + C];
-					uint32 Result = RGB_FUNC_MASKED(Base, Blended, Mask);
-					if constexpr (CLAMP)
-					{
-						BaseBuf[4 * I + C] = (uint8)FMath::Min(255u, Result);
-					}
-					else
-					{
-						BaseBuf[4 * I + C] = (uint8)Result;
-					}
-				}
-
-				// A
-				{
-					uint32 Base = BaseBuf[4 * I + 3];
-					uint32 Blended = BlendBuf[4 * I + BlendAlphaSourceChannel];
-					uint32 Result = A_FUNC(Base, Blended);
-					if constexpr (CLAMP)
-					{
-						BaseBuf[4 * I + 3] = (uint8)FMath::Min(255u, Result);
-					}
-					else
-					{
-						BaseBuf[4 * I + 3] = (uint8)Result;
-					}
-				}
-			}
-		});
-	}
-
-	template<>
-	void BufferLayerComposite<BlendChannelMasked, LightenChannel, false>
-	(
-		FImage* BaseImage,
-		const FImage* BlendImage,
-		bool bOnlyFirstLOD,
-		uint8 BlendAlphaSourceChannel)
-	{
-		check(BaseImage->GetFormat() == EImageFormat::RGBA_UByte);
-		check(BlendImage->GetFormat() == EImageFormat::RGBA_UByte);
-		check(BaseImage->GetSizeX() == BlendImage->GetSizeX() && BaseImage->GetSizeY() == BlendImage->GetSizeY());
-		check(bOnlyFirstLOD || BaseImage->GetLODCount() <= BlendImage->GetLODCount());
-		check(BlendAlphaSourceChannel == 3);
-
-		bOnlyFirstLOD = bOnlyFirstLOD || BaseImage->GetLODCount() == 1;
-
-		int32 FirstLODDataOffset = 0;
-		int32 NumRelevantElems = -1;
-		if (BlendImage->Flags & FImage::IF_HAS_RELEVANCY_MAP && bOnlyFirstLOD)
-		{
-			check(BlendImage->RelevancyMaxY < BaseImage->GetSizeY());
-			check(BlendImage->RelevancyMaxY >= BlendImage->RelevancyMinY);
-
-			uint16 SizeX = BaseImage->GetSizeX();
-			NumRelevantElems = (BlendImage->RelevancyMaxY - BlendImage->RelevancyMinY + 1) * SizeX;
-			
-			FirstLODDataOffset = BlendImage->RelevancyMinY * SizeX * 4;
-		}
-
-		const int32 NumLODs = BaseImage->GetLODCount();
-	
-		constexpr int32 BatchNumElems = 4096 * 2;
-		constexpr int32 BytesPerElem = 4;
-		OpImageLayerInternal::FOpLayerBatchArgs BatchArgs;
-		BatchArgs.BatchNumElems      = BatchNumElems;
-		BatchArgs.LODBegin           = 0;
-		BatchArgs.LODEnd             = NumLODs;
-		BatchArgs.FirstLODOffset     = bOnlyFirstLOD ? FirstLODDataOffset : -1;
-		BatchArgs.Result             = BaseImage;
-		BatchArgs.Blend              = BlendImage;
-		BatchArgs.ResultBytesPerElem = BytesPerElem;
-		BatchArgs.BlendBytesPerElem  = BytesPerElem;
-
-		const int32 NumBatches = OpImageLayerInternal::GetOpLayerNumBatchesLODRangeOffsetViews(BatchArgs);
-
-		// This will always be an upper-bound for bOnlyFirtsLODs, check if it performs as expected or it needs more fine tune.
-		const int32 NumRelevantBatches = NumRelevantElems != -1
-				? FMath::DivideAndRoundUp(NumRelevantElems, BatchNumElems)
-				: NumBatches;
-
-		ParallelExecutionUtils::InvokeBatchParallelFor(NumRelevantBatches,
-		[
-			&BatchArgs
-		] (uint32 BatchId)
-		{
-			OpImageLayerInternal::FOpLayerBatchViews BatchViews = 
-					OpImageLayerInternal::GetOpLayerBatchLODRangeOffsetViews(BatchId, BatchArgs);
-	
-			uint8* BaseBuf = BatchViews.Result.GetData();
-			const uint8* BlendBuf = BatchViews.Blend.GetData();
-			
-			for (int32 I = 0; I < BatchViews.NumElems; ++I)
-			{
-				// TODO: Optimize this (SIMD?)
-				uint32 FullBase;
-				FMemory::Memcpy(&FullBase, BaseBuf + I*sizeof(uint32), sizeof(uint32)); 
-
-				uint32 FullBlended;
-				FMemory::Memcpy(&FullBlended, BlendBuf + I*sizeof(uint32), sizeof(uint32)); 
-				uint32 Mask = (FullBlended & 0xff000000) >> 24;
-
-				uint32 FullResult = 0;
-				FullResult |= BlendChannelMasked((FullBase >>  0) & 0xff, (FullBlended >>  0) & 0xff, Mask) << 0;
-				FullResult |= BlendChannelMasked((FullBase >>  8) & 0xff, (FullBlended >>  8) & 0xff, Mask) << 8;
-				FullResult |= BlendChannelMasked((FullBase >> 16) & 0xff, (FullBlended >> 16) & 0xff, Mask) << 16;
-				FullResult |= LightenChannel	((FullBase >> 24) & 0xff, (FullBlended >> 24) & 0xff) << 24;
-
-				FMemory::Memcpy(BaseBuf + I*sizeof(uint32), &FullResult, sizeof(uint32));
-			}
-		});
-	}
-
-	template< 
-		VectorRegister4Int (*RGB_FUNC_MASKED)(const VectorRegister4Int&, const VectorRegister4Int&, const VectorRegister4Int&),
-		int32 (*A_FUNC)(int32, int32),
-		bool CLAMP >
-	void BufferLayerCompositeVector(
-		FImage* BaseImage,
-		const FImage* BlendImage,
-		bool bOnlyFirstLOD,
-		uint8 BlendAlphaSourceChannel)
-	{
-		check(BaseImage->GetFormat() == EImageFormat::RGBA_UByte);
-		check(BlendImage->GetFormat() == EImageFormat::RGBA_UByte);
-		check(BaseImage->GetSizeX() == BlendImage->GetSizeX() && BaseImage->GetSizeY() == BlendImage->GetSizeY());
-		check(bOnlyFirstLOD || BaseImage->GetLODCount() <= BlendImage->GetLODCount());
-
-		const int32 NumLODs = bOnlyFirstLOD ? 1 : BaseImage->GetLODCount();
-
-		constexpr int32 BatchNumElems = 4096*2;
-		constexpr int32 BytesPerElem = 4;
-		OpImageLayerInternal::FOpLayerBatchArgs BatchArgs;
-		BatchArgs.BatchNumElems      = BatchNumElems;
-		BatchArgs.LODBegin           = 0;
-		BatchArgs.LODEnd             = NumLODs;
-		BatchArgs.Result             = BaseImage;
-		BatchArgs.Blend              = BlendImage;
-		BatchArgs.ResultBytesPerElem = BytesPerElem;
-		BatchArgs.BlendBytesPerElem  = BytesPerElem;
-
-		const int32 NumBatches = OpImageLayerInternal::GetOpLayerNumBatchesLODRange(BatchArgs);
-
-		ParallelExecutionUtils::InvokeBatchParallelFor(NumBatches, 
-		[
-			&BatchArgs, BlendAlphaSourceChannel
-		] (uint32 BatchId)
-		{
-			OpImageLayerInternal::FOpLayerBatchViews BatchViews = 
-					OpImageLayerInternal::GetOpLayerBatchLODRangeViews(BatchId, BatchArgs);
-
-			uint8* BaseBuf = BatchViews.Result.GetData();
-			const uint8* BlendBuf = BatchViews.Blend.GetData();
-
-			for (int32 I = 0; I < BatchViews.NumElems; ++I)
-			{
-				// TODO: Optimize this (SIMD?)
-				const int32 BaseAlpha = BaseBuf[4 * I + BlendAlphaSourceChannel];
-				const int32 BlendedAlpha = BlendBuf[4 * I + BlendAlphaSourceChannel];
-
-				const VectorRegister4Int Mask = VectorIntSet1(BlendedAlpha);
-
-				uint32 BlendPixel;
-				FMemory::Memcpy(&BlendPixel, BlendBuf, sizeof(uint32));
-
-				uint32 BasePixel;
-				FMemory::Memcpy(&BasePixel, BaseBuf, sizeof(uint32));
+				// Remove RLE header.
+                MaskBuf += 4 + Rows*sizeof(uint32);
 
-				const VectorRegister4Int Blended = MakeVectorRegisterInt(
-						(BlendPixel >> 0) & 0xFF, 
-						(BlendPixel >> 8) & 0xFF, 
-						(BlendPixel >> 16) & 0xFF, 
-						(BlendPixel >> 24) & 0xFF);
+                for (int32 RowIndex = 0; RowIndex < Rows; ++RowIndex)
+                {
+                    const uint8* DestRowEnd = DestBuf + Width*NC;
+                    while (DestBuf != DestRowEnd)
+                    {
+                        // Decode header
+                        uint16 Zeros = *(const uint16*)MaskBuf;
+                        MaskBuf += 2;
 
-				const VectorRegister4Int Base = MakeVectorRegisterInt(
-						(BasePixel >> 0) & 0xFF, 
-						(BasePixel >> 8) & 0xFF, 
-						(BasePixel >> 16) & 0xFF, 
-						(BasePixel >> 24) & 0xFF);
+                        uint16 Ones = *(const uint16*)MaskBuf;
+                        MaskBuf += 2;
 
-				VectorRegister4Int Result = RGB_FUNC_MASKED(Base, Blended, Mask);
-				if constexpr (CLAMP)
-				{
-					Result = VectorIntMin(MakeVectorRegisterIntConstant(255, 255, 255, 255), Result);
-				}
+                        // Skip
+                        DestBuf += Zeros*NC;
+                        BlendedBuf += Zeros*NC;
 
-				int32 AlphaResult = A_FUNC(BaseAlpha, BlendedAlpha);
-				if constexpr (CLAMP)
-				{
-					AlphaResult = FMath::Min(255, AlphaResult);
-				}
+                        // Copy
+                        FMemory::Memmove(DestBuf, BlendedBuf, Ones*NC);
 
-				alignas(VectorRegister4Int) int32 IndexableRegister[4];
-				VectorIntStoreAligned(Result, &IndexableRegister);
+                        DestBuf += NC*Ones;
+                        BlendedBuf += NC*Ones;
+                    }
 
-				BaseBuf[4 * I + 0] = static_cast<uint8>(IndexableRegister[0]);
-				BaseBuf[4 * I + 1] = static_cast<uint8>(IndexableRegister[1]);
-				BaseBuf[4 * I + 2] = static_cast<uint8>(IndexableRegister[2]);
-				BaseBuf[4 * I + 3] = static_cast<uint8>(IndexableRegister[3]);
+                    DestBuf += Stride;
+                }
 
-				BaseBuf[4 * I + BlendAlphaSourceChannel] = static_cast<uint8>(AlphaResult);
-			}
-		});
+                //Rows = FMath::DivideAndRoundUp(Rows, 2);
+                //Width = FMath::DivideAndRoundUp(Width, 2);
+            //}
+		}
+		else
+		{
+			checkf( false, TEXT("Unsupported mask format.") );
+		}
 	}
 
+
 	template<uint32 (*BLEND_FUNC)(uint32,uint32),
 			 bool CLAMP>
     void BufferLayerStrideNoAlpha(FImage* DestImage, int32 DestOffset, int32 Stride, const FImage* MaskImage, const FImage* BlendImage/*, int32 LODCount*/)
@@ -2157,657 +2349,84 @@ namespace OpImageLayerInternal
 		BufferLayerStrideNoAlpha<BLEND_FUNC, CLAMP>(BaseImage, Start, Stride, MaskImage, BlendedImage/*, BaseImage->GetLODCount()*/);
 	}
 
-	template<uint32 NC>
-	FORCEINLINE uint32 PackPixel(const uint8* PixelPtr)
-	{
-		static_assert(NC > 0 && NC <= 4);
-
-		uint32 PixelPack = 0;
-
-		// The compiler should be able to optimize this given that NC is a constant expression.
-		FMemory::Memcpy(&PixelPack, PixelPtr, NC);
-
-		return PixelPack;
-	}
-
-	template<uint32 NC>
-	FORCEINLINE void UnpackPixel(uint8* PixelPtr, uint32 PixelData)
-	{
-		static_assert(NC > 0 && NC <= 4);
-
-		// The compiler should be able to optimize this given that NC is a constant expression
-		FMemory::Memcpy(PixelPtr, &PixelData, NC);
-	}
-
-	template< 
-		uint32 (*BLEND_FUNC)(uint32, uint32), 
-		uint32 NC>
-	void BufferLayerCombineColour(FImage* ResultImage, const FImage* BaseImage, FVector4f Color, bool bOnlyFirstLOD = false)
-	{
-		static_assert(NC > 0 && NC <= 4);
-
-		check(BaseImage->GetSizeX() == ResultImage->GetSizeX());
-		check(BaseImage->GetSizeY() == ResultImage->GetSizeY());
-
-		const uint32 TopColor = 
-			static_cast<uint32>(255.0f * Color[0]) << 0 |  
-			static_cast<uint32>(255.0f * Color[1]) << 8 |  
-			static_cast<uint32>(255.0f * Color[2]) << 16;  
-
-		const int32 NumLODs = bOnlyFirstLOD ? 1 : ResultImage->GetLODCount();
-		
-		constexpr int32 BatchNumElems = 4098*2;
-		OpImageLayerInternal::FOpLayerBatchArgs BatchArgs;
-		BatchArgs.BatchNumElems      = BatchNumElems;
-		BatchArgs.LODBegin           = 0;
-		BatchArgs.LODEnd             = NumLODs;
-		BatchArgs.Result             = ResultImage;
-		BatchArgs.Base               = BaseImage;
-		BatchArgs.ResultBytesPerElem = NC;
-		BatchArgs.BaseBytesPerElem   = NC;
-
-		const int32 NumBatches = OpImageLayerInternal::GetOpLayerNumBatchesLODRange(BatchArgs);
-
-		ParallelExecutionUtils::InvokeBatchParallelFor(NumBatches,
-		[
-			&BatchArgs, TopColor
-		](int32 BatchId)
-		{
-
-			OpImageLayerInternal::FOpLayerBatchViews BatchViews = 
-					OpImageLayerInternal::GetOpLayerBatchLODRangeViews(BatchId, BatchArgs);
-	
-			const uint8* BaseBuf = BatchViews.Base.GetData();
-			uint8* ResultBuf = BatchViews.Result.GetData();
-			
-			for (int32 I = 0; I < BatchViews.NumElems; ++I)
-			{
-				const uint32 Base = PackPixel<NC>(&BaseBuf[NC * I]);
-
-				const uint32 Result = BLEND_FUNC(Base, TopColor);
-				UnpackPixel<NC>(&ResultBuf[NC * I], Result);
-			}
-		});
-	}
-
-	template< 
-		uint32 (*BLEND_FUNC)(uint32, uint32), 
-		uint32 NC>
-	void BufferLayerCombine(FImage* ResultImage, const FImage* BaseImage, const FImage* BlendImage, bool bOnlyFirstLOD)
-	{
-		static_assert(NC > 0 && NC <= 4);
-
-		check(BaseImage->GetSizeX() == BlendImage->GetSizeX());
-		check(BaseImage->GetSizeY() == BlendImage->GetSizeY());
-		check(BaseImage->GetSizeX() == ResultImage->GetSizeX());
-		check(BaseImage->GetSizeY() == ResultImage->GetSizeY());
-		check(BaseImage->GetFormat() == BlendImage->GetFormat());
-		check(BaseImage->GetFormat() == ResultImage->GetFormat());
-
-		const int32 NumLODs = bOnlyFirstLOD ? 1 : ResultImage->GetLODCount();
-
-		constexpr int32 BatchNumElems = 4098*2;
-		OpImageLayerInternal::FOpLayerBatchArgs BatchArgs;
-		BatchArgs.BatchNumElems      = BatchNumElems;
-		BatchArgs.LODBegin           = 0;
-		BatchArgs.LODEnd             = NumLODs;
-		BatchArgs.Result             = ResultImage;
-		BatchArgs.Base               = BaseImage;
-		BatchArgs.Blend              = BlendImage;
-		BatchArgs.ResultBytesPerElem = NC;
-		BatchArgs.BaseBytesPerElem   = NC;
-		BatchArgs.BlendBytesPerElem  = NC;
-
-		const int32 NumBatches = OpImageLayerInternal::GetOpLayerNumBatchesLODRange(BatchArgs);
-
-		ParallelExecutionUtils::InvokeBatchParallelFor(NumBatches,
-		[
-			&BatchArgs
-		](int32 BatchId)
-		{
-			OpImageLayerInternal::FOpLayerBatchViews BatchViews = 
-					OpImageLayerInternal::GetOpLayerBatchLODRangeViews(BatchId, BatchArgs);
-			
-			const uint8* BaseBuf = BatchViews.Base.GetData();
-			const uint8* BlendBuf = BatchViews.Blend.GetData();
-			uint8* ResultBuf = BatchViews.Result.GetData();
-			
-			for (int32 I = 0; I < BatchViews.NumElems; ++I)
-			{
-				const uint32 Base = PackPixel<NC>(&BaseBuf[NC * I]);
-				const uint32 Blend = PackPixel<NC>(&BlendBuf[NC * I]);
-
-				const uint32 Result = BLEND_FUNC(Base, Blend);
-				UnpackPixel<NC>(&ResultBuf[NC * I], Result);
-			}
-		});
-	}
-
-	template< 
-		uint32 (*BLEND_FUNC_MASKED)(uint32, uint32, uint32), 
-		uint32 NC>
-	void BufferLayerCombine(FImage* ResultImage, const FImage* BaseImage, const FImage* MaskImage, const FImage* BlendImage, bool bOnlyFirstLOD)
-	{
-		static_assert(NC > 0 && NC <= 4);
-
-		check(BaseImage->GetSizeX() == BlendImage->GetSizeX());
-		check(BaseImage->GetSizeY() == BlendImage->GetSizeY());
-		check(BaseImage->GetSizeX() == MaskImage->GetSizeX());
-		check(BaseImage->GetSizeY() == MaskImage->GetSizeY());
-		check(BaseImage->GetFormat() == BlendImage->GetFormat());
-
-		const int32 NumLODs = bOnlyFirstLOD ? 1 : ResultImage->GetLODCount();
-
-		constexpr int32 BatchNumElems = 4098*2;
-		OpImageLayerInternal::FOpLayerBatchArgs BatchArgs;
-		BatchArgs.BatchNumElems      = BatchNumElems;
-		BatchArgs.LODBegin           = 0;
-		BatchArgs.LODEnd             = NumLODs;
-		BatchArgs.Result             = ResultImage;
-		BatchArgs.Base               = BaseImage;
-		BatchArgs.Blend              = BlendImage;
-		BatchArgs.Mask               = MaskImage;
-		BatchArgs.ResultBytesPerElem = NC;
-		BatchArgs.BaseBytesPerElem   = NC;
-		BatchArgs.BlendBytesPerElem  = NC;
-		BatchArgs.MaskBytesPerElem   = 1;
-
-		const int32 NumBatches = OpImageLayerInternal::GetOpLayerNumBatchesLODRange(BatchArgs);
-		
-		ParallelExecutionUtils::InvokeBatchParallelFor(NumBatches,
-		[
-			&BatchArgs
-		](int32 BatchId)
-		{
-			OpImageLayerInternal::FOpLayerBatchViews BatchViews = 
-					OpImageLayerInternal::GetOpLayerBatchLODRangeViews(BatchId, BatchArgs);
-
-			const uint8* BaseBuf = BatchViews.Base.GetData();
-			const uint8* BlendBuf = BatchViews.Blend.GetData();
-			const uint8* MaskBuf = BatchViews.Mask.GetData();
-			uint8* ResultBuf = BatchViews.Result.GetData();
-			
-			// This could happen in case of missing data files.
-			if (!BaseBuf || !BlendBuf || !MaskBuf)
-			{
-				return;
-			}
-
-			for (int32 I = 0; I < BatchViews.NumElems; ++I)
-			{
-				const uint32 Base = PackPixel<NC>(&BaseBuf[NC * I]);
-				const uint32 Blend = PackPixel<NC>(&BlendBuf[NC * I]);
-				const uint32 Mask = PackPixel<1>(&MaskBuf[1 * I]);
-
-				const uint32 Result = BLEND_FUNC_MASKED(Base, Blend, Mask);
-				UnpackPixel<NC>(&ResultBuf[NC * I], Result);
-			}
-		});
-	}
-
-	template< 
-		uint32 (*BLEND_FUNC_MASKED)(uint32, uint32, uint32), 
-		uint32 NC>
-	void BufferLayerCombineColour(FImage* ResultImage, const FImage* BaseImage, const FImage* MaskImage, FVector4f Color)
+	//! Blend a subimage on the base using a mask.
+	void ImageBlendOnBaseNoAlpha(FImage* BaseImage, const FImage* MaskImage, const FImage* BlendedImage, const box<FIntVector2>& Rect)
 	{
-		static_assert(NC > 0 && NC <= 4);
-
-		check(BaseImage->GetSizeX() == MaskImage->GetSizeX());
-		check(BaseImage->GetSizeY() == MaskImage->GetSizeY());
-
-		const uint32 TopColor = 
-			static_cast<uint32>(255.0f * Color[0]) << 0 |  
-			static_cast<uint32>(255.0f * Color[1]) << 8 |  
-			static_cast<uint32>(255.0f * Color[2]) << 16;  
-
-		constexpr int32 BatchNumElems = 4098*2;
-
-		OpImageLayerInternal::FOpLayerBatchArgs BatchArgs;
-		BatchArgs.BatchNumElems      = BatchNumElems;
-		BatchArgs.Result             = ResultImage;
-		BatchArgs.Base               = BaseImage;
-		BatchArgs.Mask               = MaskImage;
-		BatchArgs.ResultBytesPerElem = NC;
-		BatchArgs.BaseBytesPerElem   = NC;
-		BatchArgs.MaskBytesPerElem   = 1;
-		
-		const int32 NumBatches = OpImageLayerInternal::GetOpLayerNumBatches(BatchArgs);
-
-		ParallelExecutionUtils::InvokeBatchParallelFor(NumBatches,
-		[
-			&BatchArgs, TopColor
-		](int32 BatchId)
-		{
-			OpImageLayerInternal::FOpLayerBatchViews BatchViews = 
-					OpImageLayerInternal::GetOpLayerBatchViews(BatchId, BatchArgs);
-	
-			const uint8* BaseBuf = BatchViews.Base.GetData();
-			const uint8* MaskBuf = BatchViews.Mask.GetData();
-			uint8* ResultBuf = BatchViews.Result.GetData();
-
-			// This could happen in case of missing data files.
-			if (!BaseBuf || !MaskBuf)
-			{
-				return;
-			}
-
-			for (int32 I = 0; I < BatchViews.NumElems; ++I)
-			{
-				const uint32 Base = PackPixel<NC>(&BaseBuf[NC * I]);
-				const uint32 Mask = PackPixel<1>(&MaskBuf[1 * I]);
-
-				const uint32 Result = BLEND_FUNC_MASKED(Base, TopColor, Mask);
-				UnpackPixel<NC>(&ResultBuf[NC * I], Result);
-			}
-		});
+		ImageLayerOnBaseNoAlpha<BlendChannel, false>(BaseImage, MaskImage, BlendedImage, Rect);
 	}
 
-	template<uint32 (*BLEND_FUNC)(uint32, uint32)>
-	void ImageLayerCombine(FImage* ResultImage, const FImage* BaseImage, const FImage* BlendedImage, bool bOnlyFirstLOD)
+	void ImageLayerBlend(
+			FImage* DestImage, const FImage* BaseImage, const FImage* BlenImage, const FImage* MaskImage, 
+			FBlendFuncType* ColorBlendFunc, FBlendFuncType* AlphaBlendFunc,
+			int32 LODBegin, int32 LODEnd,
+			bool bApplyColorBlendToAlpha, bool bUseMaskFromBlendAlpha, bool bUseBlendSourceFromBlendAlpha, uint8 BlendAlphaSourceChannel)
 	{
-		check(ResultImage->GetFormat() == BaseImage->GetFormat());
-		check(ResultImage->GetSizeX() == BaseImage->GetSizeX());
-		check(ResultImage->GetSizeY() == BaseImage->GetSizeY());
-		check(bOnlyFirstLOD || ResultImage->GetLODCount() == BaseImage->GetLODCount());
-		check(BaseImage->GetSizeX() == BlendedImage->GetSizeX());
-		check(BaseImage->GetSizeY() == BlendedImage->GetSizeY());
-		check(BaseImage->GetFormat() == BlendedImage->GetFormat());
-		check(bOnlyFirstLOD || ResultImage->GetLODCount() <= BlendedImage->GetLODCount());
-
-		const EImageFormat BaseFormat = BaseImage->GetFormat();
-
-		if (BaseFormat == EImageFormat::L_UByte)
-		{
-			BufferLayerCombine<BLEND_FUNC, 1>(ResultImage, BaseImage, BlendedImage, bOnlyFirstLOD);
-		}
-		else if (BaseFormat == EImageFormat::RGB_UByte)
-		{
-			BufferLayerCombine<BLEND_FUNC, 3>(ResultImage, BaseImage, BlendedImage, bOnlyFirstLOD);
-		}
-		else if (BaseFormat == EImageFormat::RGBA_UByte || BaseFormat == EImageFormat::BGRA_UByte)
-		{
-			// \todo: pass swizzle template argument if BGRA_UBYTE, not yet supported.
-			BufferLayerCombine<BLEND_FUNC, 4>(ResultImage, BaseImage, BlendedImage, bOnlyFirstLOD);
-		}
-		else
-		{
-			checkf(false, TEXT("Unsupported format."));
-		}
-	}
-
-	template< 
-		uint32 (*BLEND_FUNC)(uint32, uint32), 
-		uint32 (*BLEND_FUNC_MASKED)(uint32, uint32, uint32)>
-	void ImageLayerCombine(FImage* ResultImage, const FImage* BaseImage, const FImage* MaskImage, const FImage* BlendedImage, bool bOnlyFirstLOD)
-	{
-		check(ResultImage->GetFormat() == BaseImage->GetFormat());
-		check(ResultImage->GetSizeX() == BaseImage->GetSizeX());
-		check(ResultImage->GetSizeY() == BaseImage->GetSizeY());
-		check(bOnlyFirstLOD || ResultImage->GetLODCount() == BaseImage->GetLODCount());
-		check(BaseImage->GetSizeX() == BlendedImage->GetSizeX());
-		check(BaseImage->GetSizeY() == BlendedImage->GetSizeY());
-		check(BaseImage->GetFormat() == BlendedImage->GetFormat());
-		check(bOnlyFirstLOD || ResultImage->GetLODCount() <= BlendedImage->GetLODCount());
-
-		const EImageFormat BaseFormat = BaseImage->GetFormat();
-
-		TManagedPtr<FImage> TempMaskImage;
-		if (MaskImage->GetFormat() != EImageFormat::L_UByte)
-		{
-			UE_LOG(LogMutableCore, Log, TEXT("Image layer format not supported. A generic one will be used. "));
-
-			FImageOperator ImOp = FImageOperator::GetDefault(nullptr);
-			constexpr int32 Quality = 4;
-			TempMaskImage = ImOp.ImagePixelFormat( Quality, MaskImage, EImageFormat::L_UByte );
-			MaskImage = TempMaskImage.Get();
-		}
-
-		if (BaseFormat == EImageFormat::L_UByte)
-		{
-			BufferLayerCombine<BLEND_FUNC_MASKED, 1>(ResultImage, BaseImage, MaskImage, BlendedImage, bOnlyFirstLOD);
-		}
-		else if (BaseFormat == EImageFormat::RGB_UByte)
-		{
-			BufferLayerCombine<BLEND_FUNC_MASKED, 3>(ResultImage, BaseImage, MaskImage, BlendedImage, bOnlyFirstLOD);
-		}
-		else if (BaseFormat == EImageFormat::RGBA_UByte || BaseFormat == EImageFormat::BGRA_UByte)
-		{
-			// \todo: pass swizzle template argument if BGRA_UBYTE, not yet supported.
-			BufferLayerCombine<BLEND_FUNC_MASKED, 4>(ResultImage, BaseImage, MaskImage, BlendedImage, bOnlyFirstLOD);
-		}
-		else
-		{
-			UE_LOG(LogMutableCore, Log, TEXT("Image layer format not supported. A generic one will be used. "));
-
-			FImageOperator ImOp = FImageOperator::GetDefault(nullptr);
-			constexpr int32 Quality = 4;
-			TManagedPtr<FImage> TempBaseImage = ImOp.ImagePixelFormat(Quality, BaseImage, EImageFormat::RGBA_UByte);
-			TManagedPtr<FImage> TempBlededImage = ImOp.ImagePixelFormat(Quality, BlendedImage, EImageFormat::RGBA_UByte);
-			BufferLayerCombine<BLEND_FUNC_MASKED, 4>(ResultImage, TempBaseImage.Get(), MaskImage, TempBlededImage.Get(), bOnlyFirstLOD);
-		}
+		constexpr bool bBlendWithConstant = false;
+		OpImageLayerInternal::ImageLayerImpl(
+				DestImage, BaseImage, BlenImage, MaskImage, 
+				ColorBlendFunc, AlphaBlendFunc, 
+				bBlendWithConstant, FVector4f{}, 
+				LODBegin, LODEnd,  
+				bApplyColorBlendToAlpha, bUseMaskFromBlendAlpha, bUseBlendSourceFromBlendAlpha, BlendAlphaSourceChannel);
 	}
 
-	template<uint32 (*BLEND_FUNC)(uint32, uint32)>
-	void ImageLayerCombineColour(FImage* ResultImage, const FImage* BaseImage, FVector4f Color)
+	void ImageLayerBlendConstant(
+			FImage* DestImage, const FImage* BaseImage, const FImage* MaskImage, const FVector4f& Constant,  
+			FBlendFuncType* ColorBlendFunc, FBlendFuncType* AlphaBlendFunc,
+			int32 LODBegin, int32 LODEnd, 
+			bool bApplyColorBlendToAlpha, bool bUseMaskFromBlendAlpha, bool bUseBlendSourceFromBlendAlpha, uint8 BlendAlphaSourceChannel)
 	{
-		check(ResultImage->GetFormat() == BaseImage->GetFormat());
-		check(ResultImage->GetSizeX() == BaseImage->GetSizeX());
-		check(ResultImage->GetSizeY() == BaseImage->GetSizeY());
-		check(ResultImage->GetLODCount() == BaseImage->GetLODCount());
-
-		const EImageFormat BaseFormat = BaseImage->GetFormat();
-
-		if (BaseFormat == EImageFormat::L_UByte)
-		{
-			BufferLayerCombineColour<BLEND_FUNC, 1>(ResultImage, BaseImage, Color);
-		}
-		else if (BaseFormat == EImageFormat::RGB_UByte)
-		{
-			BufferLayerCombineColour<BLEND_FUNC, 3>(ResultImage, BaseImage, Color);
-		}
-		else if (BaseFormat == EImageFormat::RGBA_UByte || BaseFormat == EImageFormat::BGRA_UByte)
-		{
-			// \todo: pass swizzle template argument if BGRA_UBYTE, not yet supported.
- 			BufferLayerCombineColour<BLEND_FUNC, 4>(ResultImage, BaseImage, Color);
-		}
-		else
-		{
-			checkf(false, TEXT("Unsupported format."));
-		}
+		constexpr bool bBlendWithConstant = true;
+		OpImageLayerInternal::ImageLayerImpl(
+				DestImage, BaseImage, nullptr, MaskImage, 
+				ColorBlendFunc, AlphaBlendFunc, 
+				bBlendWithConstant, Constant, 
+				LODBegin, LODEnd,
+				bApplyColorBlendToAlpha, bUseMaskFromBlendAlpha, bUseBlendSourceFromBlendAlpha, BlendAlphaSourceChannel);
 	}
 
-	template< 
-		uint32 (*BLEND_FUNC)(uint32, uint32),
-		uint32 (*BLEND_FUNC_MASKED)(uint32, uint32, uint32)>
-	void ImageLayerCombineColour(FImage* ResultImage, const FImage* BaseImage, const FImage* MaskImage, FVector4f Color)
+	void ImageLayerCombine(
+			FImage* DestImage, const FImage* BaseImage, const FImage* BlendImage, const FImage* MaskImage,
+			FCombineFuncType* CombineFunc,
+			int32 LODBegin, int32 LODEnd)
 	{
-		check(ResultImage->GetFormat() == BaseImage->GetFormat());
-		check(ResultImage->GetSizeX() == BaseImage->GetSizeX());
-		check(ResultImage->GetSizeY() == BaseImage->GetSizeY());
-		check(ResultImage->GetLODCount() == BaseImage->GetLODCount());
-
-		const EImageFormat BaseFormat = BaseImage->GetFormat();
-
-		if (MaskImage->GetFormat() != EImageFormat::L_UByte)
-		{
-			checkf(false, TEXT("Unsupported mask format."));
-
-			BufferLayerCombineColour<BLEND_FUNC, 1>(ResultImage, BaseImage, Color);
-		}
-
-		if (BaseFormat == EImageFormat::L_UByte)
-		{
-			BufferLayerCombineColour<BLEND_FUNC_MASKED, 1>(ResultImage, BaseImage, MaskImage, Color);
-		}
-		else if (BaseFormat == EImageFormat::RGB_UByte)
-		{
-			BufferLayerCombineColour<BLEND_FUNC_MASKED, 3>(ResultImage, BaseImage, MaskImage, Color);
-		}
-		else if (BaseFormat == EImageFormat::RGBA_UByte || BaseFormat == EImageFormat::BGRA_UByte)
-		{
-			// \todo: pass swizzle template argument if BGRA_UBYTE, not yet supported.
- 			BufferLayerCombineColour<BLEND_FUNC_MASKED, 4>(ResultImage, BaseImage, MaskImage, Color);
-		}
-		else
-		{
-			checkf(false, TEXT("Unsupported format."));
-		}
+		constexpr bool bCombineWithConstant = false;
+		OpImageLayerInternal::ImageLayerCombineImpl(
+				DestImage, BaseImage, BlendImage, MaskImage, 
+				CombineFunc, 
+				bCombineWithConstant, FVector4f{},
+				LODBegin, LODEnd);
 	}
 
-	template<uint32 NCBase, uint32 NCBlend, class ImageCombineFn>
-	void BufferLayerCombineFunctor(FImage* DestImage, const FImage* BaseImage, const FImage* BlendImage, ImageCombineFn&& ImageCombine)
+	void ImageLayerCombineConstant(
+			FImage* DestImage, const FImage* BaseImage, const FImage* MaskImage, const FVector4f& Constant,
+			FCombineFuncType* CombineFunc,
+			int32 LODBegin, int32 LODEnd)
 	{
-		static_assert(NCBase > 0 && NCBase <= 4);
-		static_assert(NCBlend > 0 && NCBlend <= 4);
-
-		check(BaseImage->GetFormat() == DestImage->GetFormat());
-		check(BaseImage->GetSizeX() == BlendImage->GetSizeX());
-		check(BaseImage->GetSizeY() == BlendImage->GetSizeY());
-		check(BaseImage->GetSizeX() == DestImage->GetSizeX());
-		check(BaseImage->GetSizeY() == DestImage->GetSizeY());
-		check(BaseImage->GetLODCount() <= BlendImage->GetLODCount());
-		check(BaseImage->GetLODCount() <= DestImage->GetLODCount());
-	
-		constexpr int32 BatchNumElems = 4098*2;
-
-		OpImageLayerInternal::FOpLayerBatchArgs BatchArgs;
-		BatchArgs.BatchNumElems      = BatchNumElems;
-		BatchArgs.Result             = DestImage;
-		BatchArgs.Base               = BaseImage;
-		BatchArgs.Blend              = BlendImage;
-		BatchArgs.ResultBytesPerElem = NCBase;
-		BatchArgs.BaseBytesPerElem   = NCBase;
-		BatchArgs.BlendBytesPerElem  = NCBlend;
-		
-		const int32 NumBatches = OpImageLayerInternal::GetOpLayerNumBatches(BatchArgs);
-
-		ParallelExecutionUtils::InvokeBatchParallelFor(NumBatches, 
-		[
-			&BatchArgs, ImageCombine
-		](int32 BatchId)
-		{
-			OpImageLayerInternal::FOpLayerBatchViews BatchViews = 
-					OpImageLayerInternal::GetOpLayerBatchViews(BatchId, BatchArgs);
-
-			const uint8* BaseBuf = BatchViews.Base.GetData();
-			const uint8* BlendBuf = BatchViews.Blend.GetData();
-			uint8* DestBuf = BatchViews.Result.GetData();
-	
-			for (int32 I = 0; I < BatchViews.NumElems; ++I)
-			{
-				const uint32 Base = PackPixel<NCBase>(&BaseBuf[NCBase * I]);
-				const uint32 Blend = PackPixel<NCBlend>(&BlendBuf[NCBlend * I]);
-
-				const uint32 Result = ImageCombine(Base, Blend);
-				UnpackPixel<NCBase>(&DestBuf[NCBase* I], Result);
-			}
-		});
+		constexpr bool bCombineWithConstant = true;
+		OpImageLayerInternal::ImageLayerCombineImpl(
+				DestImage, BaseImage, nullptr, MaskImage, 
+				CombineFunc, 
+				bCombineWithConstant, Constant,
+				LODBegin, LODEnd);
 	}
 
-	// Same functionality as above, in this case we use a functor which allows to pass user data. 
-	template<class ImageCombineFn>
-	void ImageLayerCombineFunctor(FImage* ResultImage, const FImage* BaseImage, const FImage* BlendedImage, ImageCombineFn&& ImageCombine)
+	FBlendFuncType* SelectBlendFunc(EBlendType BlendType) 
 	{
-		check(ResultImage->GetFormat() == BaseImage->GetFormat());
-		check(ResultImage->GetSizeX() == BaseImage->GetSizeX());
-		check(ResultImage->GetSizeY() == BaseImage->GetSizeY());
-		check(ResultImage->GetLODCount() == BaseImage->GetLODCount());
-		check(BaseImage->GetSizeX() == BlendedImage->GetSizeX());
-		check(BaseImage->GetSizeY() == BlendedImage->GetSizeY());
-		check(ResultImage->GetLODCount() <= BlendedImage->GetLODCount());
-
-		const EImageFormat BaseFormat = BaseImage->GetFormat();
-		const EImageFormat BlendFormat = BlendedImage->GetFormat();
-
-		if (BaseFormat == EImageFormat::L_UByte )
-		{
-			if (BlendFormat == EImageFormat::L_UByte )
-			{
-				BufferLayerCombineFunctor<1, 1>(ResultImage, BaseImage, BlendedImage, Forward<ImageCombineFn>(ImageCombine));
-			}
-			else if (BlendFormat == EImageFormat::RGB_UByte )
-			{
-				BufferLayerCombineFunctor<1, 3>(ResultImage, BaseImage, BlendedImage, Forward<ImageCombineFn>(ImageCombine));
-			}
-			else if (BlendFormat == EImageFormat::RGBA_UByte || BlendFormat == EImageFormat::BGRA_UByte)
-			{
-				BufferLayerCombineFunctor<1, 4>(ResultImage, BaseImage, BlendedImage, Forward<ImageCombineFn>(ImageCombine));
-			}
-			else
-			{
-				checkf(false, TEXT("Unsupported format."));
-			}
-		}
-		else if (BaseFormat == EImageFormat::RGB_UByte)
-		{
-			if (BlendFormat == EImageFormat::L_UByte )
-			{
-				BufferLayerCombineFunctor<3, 1>(ResultImage, BaseImage, BlendedImage, Forward<ImageCombineFn>(ImageCombine));
-			}
-			else if (BlendFormat == EImageFormat::RGB_UByte )
-			{
-				BufferLayerCombineFunctor<3, 3>(ResultImage, BaseImage, BlendedImage, Forward<ImageCombineFn>(ImageCombine));
-			}
-			else if (BlendFormat == EImageFormat::RGBA_UByte || BlendFormat == EImageFormat::BGRA_UByte )
-			{
-				BufferLayerCombineFunctor<3, 4>(ResultImage, BaseImage, BlendedImage, Forward<ImageCombineFn>(ImageCombine));
-			}
-			else
-			{
-				checkf(false, TEXT("Unsupported format."));
-			}
-		}
-		else if (BaseFormat == EImageFormat::RGBA_UByte || BaseFormat == EImageFormat::BGRA_UByte)
-		{
-			if (BlendFormat == EImageFormat::L_UByte )
-			{
-				BufferLayerCombineFunctor<4, 1>(ResultImage, BaseImage, BlendedImage, Forward<ImageCombineFn>(ImageCombine));
-			}
-			else if (BlendFormat == EImageFormat::RGB_UByte )
-			{
-				BufferLayerCombineFunctor<4, 3>(ResultImage, BaseImage, BlendedImage, Forward<ImageCombineFn>(ImageCombine));
-			}
-			else if (BlendFormat == EImageFormat::RGBA_UByte || BlendFormat == EImageFormat::BGRA_UByte )
-			{
-				BufferLayerCombineFunctor<4, 4>(ResultImage, BaseImage, BlendedImage, Forward<ImageCombineFn>(ImageCombine));
-			}
-			else
-			{
-				checkf(false, TEXT("Unsupported format."));
-			}
-		}
-		else
-		{
-			checkf(false, TEXT("Unsupported format."));
+		switch (BlendType)
+		{
+			case EBlendType::BT_SOFTLIGHT: return OpImageLayerBlendOps::SoftLight;
+			case EBlendType::BT_HARDLIGHT: return OpImageLayerBlendOps::HardLight;
+			case EBlendType::BT_BURN     : return OpImageLayerBlendOps::Burn;
+			case EBlendType::BT_DODGE    : return OpImageLayerBlendOps::Dodge;
+			case EBlendType::BT_SCREEN   : return OpImageLayerBlendOps::Screen;
+			case EBlendType::BT_OVERLAY  : return OpImageLayerBlendOps::Overlay;
+			case EBlendType::BT_LIGHTEN  : return OpImageLayerBlendOps::Lighten;
+			case EBlendType::BT_MULTIPLY : return OpImageLayerBlendOps::Multiply;
+			case EBlendType::BT_BLEND    : return OpImageLayerBlendOps::Blend;
+			case EBlendType::BT_NONE     : return nullptr;
+			default                      : check(false); return nullptr;
 		}
 	}
-
-	//! Blend a subimage on the base using a mask.
-	void ImageBlendOnBaseNoAlpha(FImage* BaseImage, const FImage* MaskImage, const FImage* BlendedImage, const box<FIntVector2>& Rect)
-	{
-		ImageLayerOnBaseNoAlpha<BlendChannel, false>(BaseImage, MaskImage, BlendedImage, Rect);
-	}
-
-	template void BufferLayer<SoftLightChannelMasked, SoftLightChannel, false>(FImage*, const FImage*, const FImage*, const FImage*, bool, bool);
-    template void BufferLayer<HardLightChannelMasked, HardLightChannel, false>(FImage*, const FImage*, const FImage*, const FImage*, bool, bool);
-    template void BufferLayer<BurnChannelMasked     , BurnChannel     , false>(FImage*, const FImage*, const FImage*, const FImage*, bool, bool);
-    template void BufferLayer<DodgeChannelMasked    , DodgeChannel    , false>(FImage*, const FImage*, const FImage*, const FImage*, bool, bool);
-    template void BufferLayer<ScreenChannelMasked   , ScreenChannel   , false>(FImage*, const FImage*, const FImage*, const FImage*, bool, bool);
-    template void BufferLayer<OverlayChannelMasked  , OverlayChannel  , false>(FImage*, const FImage*, const FImage*, const FImage*, bool, bool);
-    template void BufferLayer<LightenChannelMasked  , LightenChannel  , false>(FImage*, const FImage*, const FImage*, const FImage*, bool, bool);
-    template void BufferLayer<MultiplyChannelMasked , MultiplyChannel , false>(FImage*, const FImage*, const FImage*, const FImage*, bool, bool);
-    template void BufferLayer<BlendChannelMasked    , BlendChannel    , false>(FImage*, const FImage*, const FImage*, const FImage*, bool, bool);
-
-	template void BufferLayer<SoftLightChannelMasked, SoftLightChannel, true>(FImage*, const FImage*, const FImage*, const FImage*, bool, bool);
-    template void BufferLayer<HardLightChannelMasked, HardLightChannel, true>(FImage*, const FImage*, const FImage*, const FImage*, bool, bool);
-    template void BufferLayer<BurnChannelMasked     , BurnChannel     , true>(FImage*, const FImage*, const FImage*, const FImage*, bool, bool);
-    template void BufferLayer<DodgeChannelMasked    , DodgeChannel    , true>(FImage*, const FImage*, const FImage*, const FImage*, bool, bool);
-    template void BufferLayer<ScreenChannelMasked   , ScreenChannel   , true>(FImage*, const FImage*, const FImage*, const FImage*, bool, bool);
-    template void BufferLayer<OverlayChannelMasked  , OverlayChannel  , true>(FImage*, const FImage*, const FImage*, const FImage*, bool, bool);
-    template void BufferLayer<LightenChannelMasked  , LightenChannel  , true>(FImage*, const FImage*, const FImage*, const FImage*, bool, bool);
-    template void BufferLayer<MultiplyChannelMasked , MultiplyChannel , true>(FImage*, const FImage*, const FImage*, const FImage*, bool, bool);
-    template void BufferLayer<BlendChannelMasked    , BlendChannel    , true>(FImage*, const FImage*, const FImage*, const FImage*, bool, bool);
-
-	template void BufferLayerEmbeddedMask<SoftLightChannelMasked, SoftLightChannel, false>(FImage*, const FImage*, const FImage*, bool, bool); 
-	template void BufferLayerEmbeddedMask<HardLightChannelMasked, HardLightChannel, false>(FImage*, const FImage*, const FImage*, bool, bool); 
-	template void BufferLayerEmbeddedMask<BurnChannelMasked     , BurnChannel     , false>(FImage*, const FImage*, const FImage*, bool, bool); 
-	template void BufferLayerEmbeddedMask<DodgeChannelMasked    , DodgeChannel    , false>(FImage*, const FImage*, const FImage*, bool, bool); 
-	template void BufferLayerEmbeddedMask<ScreenChannelMasked   , ScreenChannel   , false>(FImage*, const FImage*, const FImage*, bool, bool); 
-	template void BufferLayerEmbeddedMask<OverlayChannelMasked  , OverlayChannel  , false>(FImage*, const FImage*, const FImage*, bool, bool); 
-	template void BufferLayerEmbeddedMask<LightenChannelMasked  , LightenChannel  , false>(FImage*, const FImage*, const FImage*, bool, bool); 
-	template void BufferLayerEmbeddedMask<MultiplyChannelMasked , MultiplyChannel , false>(FImage*, const FImage*, const FImage*, bool, bool); 
-	template void BufferLayerEmbeddedMask<BlendChannelMasked    , BlendChannel    , false>(FImage*, const FImage*, const FImage*, bool, bool); 
-   
-	template void BufferLayer<SoftLightChannel, false>(FImage*, const FImage*, const FImage*, bool, bool, bool); 
-    template void BufferLayer<HardLightChannel, false>(FImage*, const FImage*, const FImage*, bool, bool, bool); 
-    template void BufferLayer<BurnChannel     , false>(FImage*, const FImage*, const FImage*, bool, bool, bool); 
-    template void BufferLayer<DodgeChannel    , false>(FImage*, const FImage*, const FImage*, bool, bool, bool); 
-    template void BufferLayer<ScreenChannel   , false>(FImage*, const FImage*, const FImage*, bool, bool, bool); 
-    template void BufferLayer<OverlayChannel  , false>(FImage*, const FImage*, const FImage*, bool, bool, bool); 
-    template void BufferLayer<LightenChannel  , false>(FImage*, const FImage*, const FImage*, bool, bool, bool); 
-    template void BufferLayer<MultiplyChannel , false>(FImage*, const FImage*, const FImage*, bool, bool, bool); 
-    template void BufferLayer<BlendChannel    , false>(FImage*, const FImage*, const FImage*, bool, bool, bool); 
-
-	template void BufferLayer<SoftLightChannel, true>(FImage*, const FImage*, const FImage*, bool, bool, bool); 
-    template void BufferLayer<HardLightChannel, true>(FImage*, const FImage*, const FImage*, bool, bool, bool); 
-    template void BufferLayer<BurnChannel     , true>(FImage*, const FImage*, const FImage*, bool, bool, bool); 
-    template void BufferLayer<DodgeChannel    , true>(FImage*, const FImage*, const FImage*, bool, bool, bool); 
-    template void BufferLayer<ScreenChannel   , true>(FImage*, const FImage*, const FImage*, bool, bool, bool); 
-    template void BufferLayer<OverlayChannel  , true>(FImage*, const FImage*, const FImage*, bool, bool, bool); 
-    template void BufferLayer<LightenChannel  , true>(FImage*, const FImage*, const FImage*, bool, bool, bool); 
-    template void BufferLayer<MultiplyChannel , true>(FImage*, const FImage*, const FImage*, bool, bool, bool);
-    template void BufferLayer<BlendChannel    , true>(FImage*, const FImage*, const FImage*, bool, bool, bool);
-
-	template void BufferLayerInPlace<SoftLightChannel, false, 1>(FImage*, const FImage*, bool, uint32, uint32); 
-	template void BufferLayerInPlace<HardLightChannel, false, 1>(FImage*, const FImage*, bool, uint32, uint32); 
-	template void BufferLayerInPlace<BurnChannel     , false, 1>(FImage*, const FImage*, bool, uint32, uint32); 
-	template void BufferLayerInPlace<DodgeChannel    , false, 1>(FImage*, const FImage*, bool, uint32, uint32); 
-	template void BufferLayerInPlace<ScreenChannel   , false, 1>(FImage*, const FImage*, bool, uint32, uint32); 
-	template void BufferLayerInPlace<OverlayChannel  , false, 1>(FImage*, const FImage*, bool, uint32, uint32); 
-	template void BufferLayerInPlace<LightenChannel  , false, 1>(FImage*, const FImage*, bool, uint32, uint32); 
-	template void BufferLayerInPlace<MultiplyChannel , false, 1>(FImage*, const FImage*, bool, uint32, uint32); 
-	template void BufferLayerInPlace<BlendChannel    , false, 1>(FImage*, const FImage*, bool, uint32, uint32);
-
-	template void BufferLayerColour<SoftLightChannelMasked, SoftLightChannel>(FImage*, const FImage*, const FImage*, FVector4f); 
-	template void BufferLayerColour<HardLightChannelMasked, HardLightChannel>(FImage*, const FImage*, const FImage*, FVector4f); 
-	template void BufferLayerColour<BurnChannelMasked     , BurnChannel     >(FImage*, const FImage*, const FImage*, FVector4f); 
-	template void BufferLayerColour<DodgeChannelMasked    , DodgeChannel    >(FImage*, const FImage*, const FImage*, FVector4f); 
-	template void BufferLayerColour<ScreenChannelMasked   , ScreenChannel   >(FImage*, const FImage*, const FImage*, FVector4f); 
-	template void BufferLayerColour<OverlayChannelMasked  , OverlayChannel  >(FImage*, const FImage*, const FImage*, FVector4f); 
-	template void BufferLayerColour<LightenChannelMasked  , LightenChannel  >(FImage*, const FImage*, const FImage*, FVector4f); 
-	template void BufferLayerColour<MultiplyChannelMasked , MultiplyChannel >(FImage*, const FImage*, const FImage*, FVector4f); 
-	template void BufferLayerColour<BlendChannelMasked    , BlendChannel    >(FImage*, const FImage*, const FImage*, FVector4f); 
-
-	template void BufferLayerColourFromAlpha<SoftLightChannel>(FImage*, const FImage*, FVector4f); 
-	template void BufferLayerColourFromAlpha<HardLightChannel>(FImage*, const FImage*, FVector4f); 
-	template void BufferLayerColourFromAlpha<BurnChannel     >(FImage*, const FImage*, FVector4f); 
-	template void BufferLayerColourFromAlpha<DodgeChannel    >(FImage*, const FImage*, FVector4f); 
-	template void BufferLayerColourFromAlpha<ScreenChannel   >(FImage*, const FImage*, FVector4f); 
-	template void BufferLayerColourFromAlpha<OverlayChannel  >(FImage*, const FImage*, FVector4f); 
-	template void BufferLayerColourFromAlpha<LightenChannel  >(FImage*, const FImage*, FVector4f); 
-	template void BufferLayerColourFromAlpha<MultiplyChannel >(FImage*, const FImage*, FVector4f); 
-
-	template void BufferLayerColour<SoftLightChannel>(FImage*, const FImage*, FVector4f); 
-	template void BufferLayerColour<HardLightChannel>(FImage*, const FImage*, FVector4f); 
-	template void BufferLayerColour<BurnChannel     >(FImage*, const FImage*, FVector4f); 
-	template void BufferLayerColour<DodgeChannel    >(FImage*, const FImage*, FVector4f); 
-	template void BufferLayerColour<ScreenChannel   >(FImage*, const FImage*, FVector4f); 
-	template void BufferLayerColour<OverlayChannel  >(FImage*, const FImage*, FVector4f); 
-	template void BufferLayerColour<LightenChannel  >(FImage*, const FImage*, FVector4f); 
-	template void BufferLayerColour<MultiplyChannel >(FImage*, const FImage*, FVector4f); 
-
-	template void BufferLayerColourInPlace<SoftLightChannelMasked, SoftLightChannel, 1>(FImage*, const FImage*, FVector4f, bool, uint32, uint8); 
-	template void BufferLayerColourInPlace<HardLightChannelMasked, HardLightChannel, 1>(FImage*, const FImage*, FVector4f, bool, uint32, uint8); 
-	template void BufferLayerColourInPlace<BurnChannelMasked     , BurnChannel     , 1>(FImage*, const FImage*, FVector4f, bool, uint32, uint8); 
-	template void BufferLayerColourInPlace<DodgeChannelMasked    , DodgeChannel    , 1>(FImage*, const FImage*, FVector4f, bool, uint32, uint8); 
-	template void BufferLayerColourInPlace<ScreenChannelMasked   , ScreenChannel   , 1>(FImage*, const FImage*, FVector4f, bool, uint32, uint8); 
-	template void BufferLayerColourInPlace<OverlayChannelMasked  , OverlayChannel  , 1>(FImage*, const FImage*, FVector4f, bool, uint32, uint8); 
-	template void BufferLayerColourInPlace<LightenChannelMasked  , LightenChannel  , 1>(FImage*, const FImage*, FVector4f, bool, uint32, uint8); 
-	template void BufferLayerColourInPlace<MultiplyChannelMasked , MultiplyChannel , 1>(FImage*, const FImage*, FVector4f, bool, uint32, uint8); 
-	template void BufferLayerColourInPlace<BlendChannelMasked    , BlendChannel    , 1>(FImage*, const FImage*, FVector4f, bool, uint32, uint8); 
-	
-	template void BufferLayerColourInPlace<SoftLightChannel, 1>(FImage*, FVector4f, bool, uint32, uint8); 
-	template void BufferLayerColourInPlace<HardLightChannel, 1>(FImage*, FVector4f, bool, uint32, uint8); 
-	template void BufferLayerColourInPlace<BurnChannel     , 1>(FImage*, FVector4f, bool, uint32, uint8); 
-	template void BufferLayerColourInPlace<DodgeChannel    , 1>(FImage*, FVector4f, bool, uint32, uint8); 
-	template void BufferLayerColourInPlace<ScreenChannel   , 1>(FImage*, FVector4f, bool, uint32, uint8); 
-	template void BufferLayerColourInPlace<OverlayChannel  , 1>(FImage*, FVector4f, bool, uint32, uint8); 
-	template void BufferLayerColourInPlace<LightenChannel  , 1>(FImage*, FVector4f, bool, uint32, uint8); 
-	template void BufferLayerColourInPlace<MultiplyChannel , 1>(FImage*, FVector4f, bool, uint32, uint8); 
-	template void BufferLayerColourInPlace<BlendChannel    , 1>(FImage*, FVector4f, bool, uint32, uint8); 
-
-	// BufferLayerComposite* uses are specializations, no need to explicitly instanciate.
-	//template void BufferLayerComposite<BlendChannelMasked, LightenChannel, false>(FImage*, const FImage*, bool, uint8);
-
-	//template void BufferLayerCompositeVector<VectorBlendChannelMasked, VectorLightenChannel, false>(FImage*, const FImage*, bool, uint8);
-
-	template void ImageLayerCombine<CombineNormal>(FImage*, const FImage*, const FImage*, bool);
-	template void ImageLayerCombine<CombineNormal, CombineNormalMasked>(FImage*, const FImage*, const FImage*, const FImage*, bool);
-
-	template void ImageLayerCombineColour<CombineNormal>(FImage*, const FImage*, FVector4f);
-	template void ImageLayerCombineColour<CombineNormal, CombineNormalMasked>(FImage*, const FImage*, const FImage*, FVector4f);
-
-	template void ImageLayerCombineFunctor<FNormalCompositeFunctor>(FImage*, const FImage*, const FImage*, FNormalCompositeFunctor&&);
-	template void ImageLayerCombineFunctor<FNormalCompositeIdentityFunctor>(FImage*, const FImage*, const FImage*, FNormalCompositeIdentityFunctor&&);
-
 }
+
diff --git a/Engine/Plugins/Mutable/Source/MutableRuntime/Private/MuR/OpImageLayer.h b/Engine/Plugins/Mutable/Source/MutableRuntime/Private/MuR/OpImageLayer.h
--- a/Engine/Plugins/Mutable/Source/MutableRuntime/Private/MuR/OpImageLayer.h
+++ b/Engine/Plugins/Mutable/Source/MutableRuntime/Private/MuR/OpImageLayer.h
@@ -3,67 +3,58 @@
 #pragma once
 
 #include "MuR/ImagePrivate.h"
-#include "MuR/OpImageBlend.h"
+#include "MuR/ImageTypes.h"
 
 namespace UE::Mutable::Private
-{	
-	/** Used function templates need to be manually instanciated in OpImageLayer.cpp */ 
-
-	template<uint32(*BlendFuncMasked)(uint32, uint32, uint32), uint32(*BlendFunc)(uint32, uint32), bool bClamp>
-    void BufferLayer(FImage* DestImage, const FImage* BaseImage, const FImage* MaskImage, const FImage* BlendImage, bool bApplyToAlpha, bool bOnlyFirstLOD);
-
-	template<uint32(*BlendFunc)(uint32, uint32), bool bClamp>
-	void BufferLayer(FImage* ResultImage, const FImage* BaseImage, const FImage* BlendedImage, bool bApplyToAlpha, bool bOnlyOneMip, bool bUseBlendSourceFromBlendAlpha);
-
-	template<uint32(*BlendFuncMasked)(uint32, uint32, uint32), uint32 (*BlendFunc)(uint32, uint32), bool bClamp>
-	void BufferLayerEmbeddedMask(FImage* DestImage, const FImage* BaseImage, const FImage* BlendImage, bool bApplyToAlpha, bool bOnlyFirstLOD);
-	
-	template<uint32(*BlendFunc)(uint32, uint32), bool bClamp, int32 ChannelCount>
-	void BufferLayerInPlace(FImage* BaseImage, const FImage* BlendedImage, bool bOnlyOneMip, uint32 BaseOffset, uint32 BlendedOffset);
-
-	template<uint32(*RGBFuncMasked)(uint32, uint32, uint32), uint32(*AlphaFunc)(uint32, uint32), bool bClamp>
-	void BufferLayerComposite(FImage* BaseImage, const FImage* BlendImage, bool bOnlyFirstLOD, uint8 BlendAlphaSourceChannel);
-
-	template<> 
-	void BufferLayerComposite<BlendChannelMasked, LightenChannel, false>(FImage*, const FImage*, bool, uint8);
-
-	template<VectorRegister4Int(*RGBFuncMasked)(const VectorRegister4Int&, const VectorRegister4Int&, const VectorRegister4Int&), int32(*AlphaFunc)(int32, int32), bool bClamp>
-	void BufferLayerCompositeVector(FImage* BaseImage, const FImage* BlendImage, bool bOnlyFirstLOD, uint8 BlendAlphaSourceChannel);
-
-	template<> 
-	void BufferLayerCompositeVector<VectorBlendChannelMasked, VectorLightenChannel, false>(FImage*, const FImage*, bool, uint8);
-
-	/** Blend a subimage on the base using a mask */
-	void ImageBlendOnBaseNoAlpha(FImage* BaseImage, const FImage* MaskImage, const FImage* BlendedImage, const box<FIntVector2>& Rect);
-
-	template<uint32 (*BlendFuncMasked)(uint32, uint32, uint32), uint32 (*BlendFunc)(uint32, uint32), uint32 ChannelCount>
-	void BufferLayerColourInPlace(FImage* BaseImage, const FImage* MaskImage, FVector4f Color, bool bOnlyOneMip, uint32 BaseOffset, uint8 ColorOffset);
-
-	template<uint32 (*BlendFunc)(uint32, uint32), uint32 ChannelCount>
-	void BufferLayerColourInPlace(FImage* BaseImage, FVector4f Color, bool bOnlyOneMip, uint32 BaseOffset, uint8 ColorOffset);
-
-	template<uint32 (*BlendFuncMasked)(uint32, uint32, uint32), uint32 (*BlendFunc)(uint32, uint32)>
-	void BufferLayerColour(FImage* DestImage, const FImage* BaseImage, const FImage* MaskImage, FVector4f Color);
-
-	template<uint32 (*BlendFunc)(uint32,uint32)>
-    void BufferLayerColour(FImage* ResultImage, const FImage* BaseImage, FVector4f Color);
-
-	template<uint32 (*BlendFunc)(uint32, uint32)>
-	void BufferLayerColourFromAlpha(FImage* ResultImage, const FImage* BaseImage, FVector4f Color);
+{
+
+using FBlendFuncType = void(uint8*, const uint8*, const uint8*, int32);
+
+namespace OpImageLayerBlendOps
+{
+	void Blend    (uint8* Dest, const uint8* Base, const uint8* Blen, int32 NumElems);
+	void Screen   (uint8* Dest, const uint8* Base, const uint8* Blen, int32 NumElems);
+	void SoftLight(uint8* Dest, const uint8* Base, const uint8* Blen, int32 NumElems);
+	void HardLight(uint8* Dest, const uint8* Base, const uint8* Blen, int32 NumElems);
+	void Dodge    (uint8* Dest, const uint8* Base, const uint8* Blen, int32 NumElems);
+	void Lighten  (uint8* Dest, const uint8* Base, const uint8* Blen, int32 NumElems);
+	void Burn     (uint8* Dest, const uint8* Base, const uint8* Blen, int32 NumElems);
+	void Multiply (uint8* Dest, const uint8* Base, const uint8* Blen, int32 NumElems);
+	void Overlay  (uint8* Dest, const uint8* Base, const uint8* Blen, int32 NumElems);
+} // namespace OpImageLayerBlendOps
 	
-	template<uint32 (*BlendFunc)(uint32, uint32)>
-	void ImageLayerCombine(FImage* ResultImage, const FImage* BaseImage, const FImage* BlendedImage, bool bOnlyFirstLOD);
-
-	template<uint32 (*BlendFunc)(uint32, uint32), uint32 (*BlendFuncMasked)(uint32, uint32, uint32)>
-	void ImageLayerCombine(FImage* ResultImage, const FImage* BaseImage, const FImage* MaskImage, const FImage* BlendedImage, bool bOnlyFirstLOD);
-
-	template<uint32 (*BlendFunc)(uint32, uint32)>
-	void ImageLayerCombineColour(FImage* ResultImage, const FImage* BaseImage, FVector4f Color);
-
-	template<uint32 (*BlendFunc)(uint32, uint32), uint32 (*BlendFuncMasked)(uint32, uint32, uint32)>
-	void ImageLayerCombineColour(FImage* ResultImage, const FImage* BaseImage, const FImage* MaskImage, FVector4f Color);
-
-	template<class ImageCombineFn>
-	void ImageLayerCombineFunctor(FImage* ResultImage, const FImage* BaseImage, const FImage* BlendedImage, ImageCombineFn&& ImageCombine);
+using FCombineFuncType = void(uint8*, const uint8*, const uint8*, int32);
+namespace OpImageLayerCombineOps
+{
+	void NormalCombine(uint8* Dest, const uint8* Base, const uint8* Blen, int32 NumElems);
+} // namespace OpImageLayerCombineOps
+
+
+/** Blend a subimage on the base using a mask */
+void ImageBlendOnBaseNoAlpha(FImage* BaseImage, const FImage* MaskImage, const FImage* BlendedImage, const box<FIntVector2>& Rect);
+
+FBlendFuncType* SelectBlendFunc(EBlendType BlendType); 
+
+void ImageLayerBlend(
+		FImage* Dest, const FImage* Base, const FImage* Blend, const FImage* Mask, 
+		FBlendFuncType* ColorBlendFunc, FBlendFuncType* AlphaBlendFunc,
+		int32 LODBegin, int32 LODEnd, 
+		bool bApplyColorBlendToAlpha, bool bUseMaskFromBlendAlpha, bool bUseBlendSourceFromBlendAlpha, uint8 BlendAlphaSourceChannel);
+
+void ImageLayerBlendConstant(
+		FImage* Dest, const FImage* Base, const FImage* Mask, const FVector4f& Constant, 
+		FBlendFuncType* ColorBlendFunc, FBlendFuncType* AlphaBlendFunc,
+		int32 LODBegin, int32 LODEnd, 
+		bool bApplyColorBlendToAlpha, bool bUseMaskFromBlendAlpha, bool bUseBlendSourceFromBlendAlpha, uint8 BlendAlphaSourceChannel);
+
+void ImageLayerCombine(
+		FImage* Dest, const FImage* Base, const FImage* Blend, const FImage* Mask,
+		FCombineFuncType* CombineFunc,
+		int32 LODBegin, int32 LODEnd);
+
+void ImageLayerCombineConstant(
+		FImage* Dest, const FImage* Base, const FImage* Mask, const FVector4f& Constant,
+		FCombineFuncType* CombineFunc,
+		int32 LODBegin, int32 LODEnd);
 
 }
diff --git a/Engine/Plugins/Mutable/Source/MutableRuntime/Private/MuR/OpImageNormalCombine.h b/Engine/Plugins/Mutable/Source/MutableRuntime/Private/MuR/OpImageNormalCombine.h
deleted file mode 100644
