[mutable] Fixed modifier nodes not showing the custom tags explorer UI element.
- The details for those nodes were removed, but the base one was not set to be used instead.

#Synced-CL 51208415

#rnx
#Synced-CL 51208415

diff --git a/Engine/Plugins/Mutable/Source/CustomizableObjectEditor/Private/MuCOE/CustomizableObjectEditorModule.cpp b/Engine/Plugins/Mutable/Source/CustomizableObjectEditor/Private/MuCOE/CustomizableObjectEditorModule.cpp
--- a/Engine/Plugins/Mutable/Source/CustomizableObjectEditor/Private/MuCOE/CustomizableObjectEditorModule.cpp
+++ b/Engine/Plugins/Mutable/Source/CustomizableObjectEditor/Private/MuCOE/CustomizableObjectEditorModule.cpp
@@ -96,8 +96,11 @@
 #include "MuR/External/Operation.h"
 #include "Nodes/CONodeExternalOperation.h"
 #include "Nodes/CONodeExternalOperationDetails.h"
+#include "Nodes/CONodeModifierTransformWithBone.h"
 #include "Nodes/CustomizableObjectNodeComponentPassthroughMesh.h"
 #include "Nodes/CONodeSkeletalMeshBreak.h"
+#include "Nodes/CustomizableObjectNodeModifierClipWithMesh.h"
+#include "Nodes/CustomizableObjectNodeModifierTransformInMesh.h"
 #include "Nodes/CustomizableObjectNodeTextureParameter.h"
 #include "Nodes/ExternalTypeCustomization.h"
 #include "PhysicsEngine/PhysicsAsset.h"
@@ -246,8 +249,11 @@ void FCustomizableObjectEditorModule::StartupModule()
 	RegisterCustomDetails(PropertyModule, UCustomizableObjectNodeModifierRemoveMeshBlocks::StaticClass(), FOnGetDetailCustomizationInstance::CreateStatic(&FCustomizableObjectNodeModifierRemoveMeshBlocksDetails::MakeInstance));
 	RegisterCustomDetails(PropertyModule, UCustomizableObjectNodeModifierMorphMeshSection::StaticClass(), FOnGetDetailCustomizationInstance::CreateStatic(&FCustomizableObjectNodeModifierMorphMeshSectionDetails::MakeInstance));
 	RegisterCustomDetails(PropertyModule, UCustomizableObjectNodeModifierClipMorph::StaticClass(), FOnGetDetailCustomizationInstance::CreateStatic(&FCustomizableObjectNodeModifierClipMorphDetails::MakeInstance));
+	RegisterCustomDetails(PropertyModule, UCustomizableObjectNodeModifierClipWithMesh::StaticClass(), FOnGetDetailCustomizationInstance::CreateStatic(&FCustomizableObjectNodeModifierBaseDetails::MakeInstance));
 	RegisterCustomDetails(PropertyModule, UCustomizableObjectNodeModifierClipWithUVMask::StaticClass(), FOnGetDetailCustomizationInstance::CreateStatic(&FCustomizableObjectNodeModifierBaseDetails::MakeInstance));
 	RegisterCustomDetails(PropertyModule, UCustomizableObjectNodeModifierClipDeform::StaticClass(), FOnGetDetailCustomizationInstance::CreateStatic(&FCustomizableObjectNodeModifierBaseDetails::MakeInstance));
+	RegisterCustomDetails(PropertyModule, UCustomizableObjectNodeModifierTransformInMesh::StaticClass(), FOnGetDetailCustomizationInstance::CreateStatic(&FCustomizableObjectNodeModifierBaseDetails::MakeInstance));
+	RegisterCustomDetails(PropertyModule, UCONodeModifierTransformWithBone::StaticClass(), FOnGetDetailCustomizationInstance::CreateStatic(&FCustomizableObjectNodeModifierBaseDetails::MakeInstance));
 	RegisterCustomDetails(PropertyModule, UCustomizableObjectNodeObject::StaticClass(), FOnGetDetailCustomizationInstance::CreateStatic(&FCustomizableObjectNodeObjectDetails::MakeInstance));
 	RegisterCustomDetails(PropertyModule, UCustomizableObjectNodeObjectGroup::StaticClass(), FOnGetDetailCustomizationInstance::CreateStatic(&FCustomizableObjectNodeObjectGroupDetails::MakeInstance));
 	RegisterCustomDetails(PropertyModule, UCustomizableObjectNodeProjectorParameter::StaticClass(), FOnGetDetailCustomizationInstance::CreateStatic(&FCustomizableObjectNodeProjectorParameterDetails::MakeInstance));
diff --git a/Engine/Plugins/Mutable/Source/CustomizableObjectEditor/Private/MuCOE/Nodes/CustomizableObjectNodeModifierTransformInMesh.cpp b/Engine/Plugins/Mutable/Source/CustomizableObjectEditor/Private/MuCOE/Nodes/CustomizableObjectNodeModifierTransformInMesh.cpp
--- a/Engine/Plugins/Mutable/Source/CustomizableObjectEditor/Private/MuCOE/Nodes/CustomizableObjectNodeModifierTransformInMesh.cpp
+++ b/Engine/Plugins/Mutable/Source/CustomizableObjectEditor/Private/MuCOE/Nodes/CustomizableObjectNodeModifierTransformInMesh.cpp
@@ -18,7 +18,7 @@ FText UCustomizableObjectNodeModifierTransformInMesh::GetNodeTitle(ENodeTitleTyp
 
 FText UCustomizableObjectNodeModifierTransformInMesh::GetTooltipText() const
 {
-	return LOCTEXT("Transform_Mesh_In_Mesh_Tooltip", "Applies a transform to the vertices of a mesh that is contained within the given bounding mesh");
+	return LOCTEXT("Transform_Mesh_In_Mesh_Tooltip", "Applies a transform to the vertices of a mesh that are contained within the given bounding mesh");
 }
 
 
