[mutable] Fixed LocText conflicts.

#rnx
#Synced-CL 51208415

diff --git a/Engine/Plugins/Mutable/Source/CustomizableObjectEditor/Private/MuCOE/Nodes/CONodeClipMeshWithMesh.cpp b/Engine/Plugins/Mutable/Source/CustomizableObjectEditor/Private/MuCOE/Nodes/CONodeClipMeshWithMesh.cpp
--- a/Engine/Plugins/Mutable/Source/CustomizableObjectEditor/Private/MuCOE/Nodes/CONodeClipMeshWithMesh.cpp
+++ b/Engine/Plugins/Mutable/Source/CustomizableObjectEditor/Private/MuCOE/Nodes/CONodeClipMeshWithMesh.cpp
@@ -58,7 +58,7 @@ FLinearColor UCONodeClipMeshWithMesh::GetNodeTitleColor() const
 
 FText UCONodeClipMeshWithMesh::GetTooltipText() const
 {
-	return LOCTEXT("Clip_Mesh_Mesh_Tooltip", "Removes the part of a mesh section that is completely enclosed in a mesh volume.");
+	return LOCTEXT("Clip_Mesh_With_Mesh_Tooltip", "Removes the part of a mesh section that is completely enclosed in a mesh volume.");
 }
 
 #undef LOCTEXT_NAMESPACE
diff --git a/Engine/Plugins/Mutable/Source/CustomizableObjectEditor/Private/MuCOE/Nodes/CONodeModifierTransformWithBone.cpp b/Engine/Plugins/Mutable/Source/CustomizableObjectEditor/Private/MuCOE/Nodes/CONodeModifierTransformWithBone.cpp
--- a/Engine/Plugins/Mutable/Source/CustomizableObjectEditor/Private/MuCOE/Nodes/CONodeModifierTransformWithBone.cpp
+++ b/Engine/Plugins/Mutable/Source/CustomizableObjectEditor/Private/MuCOE/Nodes/CONodeModifierTransformWithBone.cpp
@@ -11,13 +11,13 @@
 
 FText UCONodeModifierTransformWithBone::GetNodeTitle(ENodeTitleType::Type TitleType) const
 {
-	return LOCTEXT("Transform_Mesh_In_Bone_Hierarchy", "Transform Mesh In Bone Hierarchy");
+	return LOCTEXT("Transform_Mesh_In_Bone_Hierarchy_Modifier", "Transform Mesh In Bone Hierarchy");
 }
 
 
 FText UCONodeModifierTransformWithBone::GetTooltipText() const
 {
-	return LOCTEXT("Transform_Mesh_In_Bone_Hierarchy_Tooltip", "Applies a transform to the vertices of a mesh that are skinned to the target bone or its child bones");
+	return LOCTEXT("Transform_Mesh_In_Bone_Hierarchy_Modifier_Tooltip", "Applies a transform to the vertices of a mesh that are skinned to the target bone or its child bones");
 }
 
 
diff --git a/Engine/Plugins/Mutable/Source/CustomizableObjectEditor/Private/MuCOE/Nodes/CustomizableObjectNodeMeshClipWithMesh.cpp b/Engine/Plugins/Mutable/Source/CustomizableObjectEditor/Private/MuCOE/Nodes/CustomizableObjectNodeMeshClipWithMesh.cpp
--- a/Engine/Plugins/Mutable/Source/CustomizableObjectEditor/Private/MuCOE/Nodes/CustomizableObjectNodeMeshClipWithMesh.cpp
+++ b/Engine/Plugins/Mutable/Source/CustomizableObjectEditor/Private/MuCOE/Nodes/CustomizableObjectNodeMeshClipWithMesh.cpp
@@ -110,7 +110,7 @@ void UCustomizableObjectNodeModifierClipWithMesh::PostEditChangeProperty(struct
 
 FText UCustomizableObjectNodeModifierClipWithMesh::GetNodeTitle(ENodeTitleType::Type TitleType) const
 {
-	return LOCTEXT("Clip_Mesh_With_Mesh", "Clip Mesh With Mesh");
+	return LOCTEXT("Clip_Mesh_With_Mesh_Modifier", "Clip Mesh With Mesh");
 }
 
 
@@ -128,7 +128,7 @@ UEdGraphPin* UCustomizableObjectNodeModifierClipWithMesh::GetClipMeshPin() const
 
 FText UCustomizableObjectNodeModifierClipWithMesh::GetTooltipText() const
 {
-	return LOCTEXT("Clip_Mesh_Mesh_Tooltip", "Removes the part of a mesh section that is completely enclosed in a mesh volume.\nIt only removes the faces that fall completely inside the cutting volume, along with the vertices and edges that define only faces that are deleted.");
+	return LOCTEXT("Clip_Mesh_With_Mesh_Modifier_Tooltip", "Removes the part of a mesh section that is completely enclosed in a mesh volume.\nIt only removes the faces that fall completely inside the cutting volume, along with the vertices and edges that define only faces that are deleted.");
 }
 
 #undef LOCTEXT_NAMESPACE
diff --git a/Engine/Plugins/Mutable/Source/CustomizableObjectEditor/Private/MuCOE/Nodes/CustomizableObjectNodeModifierRemoveMesh.cpp b/Engine/Plugins/Mutable/Source/CustomizableObjectEditor/Private/MuCOE/Nodes/CustomizableObjectNodeModifierRemoveMesh.cpp
--- a/Engine/Plugins/Mutable/Source/CustomizableObjectEditor/Private/MuCOE/Nodes/CustomizableObjectNodeModifierRemoveMesh.cpp
+++ b/Engine/Plugins/Mutable/Source/CustomizableObjectEditor/Private/MuCOE/Nodes/CustomizableObjectNodeModifierRemoveMesh.cpp
@@ -27,7 +27,7 @@ void UCustomizableObjectNodeModifierRemoveMesh::AllocateDefaultPins(UCustomizabl
 
 FText UCustomizableObjectNodeModifierRemoveMesh::GetNodeTitle(ENodeTitleType::Type TitleType) const
 {
-	return LOCTEXT("Remove_Mesh", "Remove Mesh");
+	return LOCTEXT("Remove_Mesh_Modifier", "Remove Mesh");
 }
 
 
@@ -48,7 +48,7 @@ void UCustomizableObjectNodeModifierRemoveMesh::PinConnectionListChanged(UEdGrap
 
 FText UCustomizableObjectNodeModifierRemoveMesh::GetTooltipText() const
 {
-	return LOCTEXT("Remove_Mesh_Tooltip",
+	return LOCTEXT("Remove_Mesh_Modifier_Tooltip",
 	"Removes the faces of a mesh section that are defined only by the vertexes shared by said mesh section and the input mesh. \n"
 	"It also removes any vertex and edge that only define deleted faces. \n"
 	"If the removed mesh covers all the faces included in one or more layout blocks those blocks get removed, freeing layout space in the final texture.");
