[Mutable] Small BuildOrCopyClothingData optimization.

#Synced-CL 51099946

diff --git a/Engine/Plugins/Mutable/Source/CustomizableObject/Private/MuCO/CustomizableObjectInstance.cpp b/Engine/Plugins/Mutable/Source/CustomizableObject/Private/MuCO/CustomizableObjectInstance.cpp
--- a/Engine/Plugins/Mutable/Source/CustomizableObject/Private/MuCO/CustomizableObjectInstance.cpp
+++ b/Engine/Plugins/Mutable/Source/CustomizableObject/Private/MuCO/CustomizableObjectInstance.cpp
@@ -4064,7 +4064,7 @@ void BuildOrCopyClothingData(
 			if (Cloth.IsValid())
 			{
 				RenderSection.ClothMappingDataLODs.SetNum(1);
-				RenderSection.ClothMappingDataLODs[0].SetNum(RenderSection.NumVertices);
+				RenderSection.ClothMappingDataLODs[0].SetNumUninitialized(RenderSection.NumVertices);
 			}
 		}
 			
