[Mutable] Fix compilation error in ImageComputeUtils neon code path on some platforms.

#rnx

#Synced-CL 51258789

diff --git a/Engine/Plugins/Mutable/Source/MutableRuntime/Private/MuR/ImageComputeUtils.cpp b/Engine/Plugins/Mutable/Source/MutableRuntime/Private/MuR/ImageComputeUtils.cpp
--- a/Engine/Plugins/Mutable/Source/MutableRuntime/Private/MuR/ImageComputeUtils.cpp
+++ b/Engine/Plugins/Mutable/Source/MutableRuntime/Private/MuR/ImageComputeUtils.cpp
@@ -81,8 +81,8 @@ namespace ImageComputeUtils
 
   	FORCEINLINE VectorRegisterInt8x16 MakeShuffle_Int8x16(uint8 V0 = 0, uint8 V1 = 1, uint8 V2 = 2, uint8 V3 = 3, uint8 V4 = 4, uint8 V5 = 5, uint8 V6 = 6, uint8 V7 = 7, uint8 V8 = 8, uint8 V9 = 9, uint8 V10 = 10, uint8 V11 = 11, uint8 V12 = 12, uint8 V13 = 13, uint8 V14 = 14, uint8 V15 = 15)
 	{
-		int8x16_t Value {V0, V1, V2, V3, V4, V5, V6, V7, V8, V9, V10, V11, V12, V13, V14, V15};
-		return Value;
+		const uint8 Data[16] = { V0, V1, V2, V3, V4, V5, V6, V7, V8, V9, V10, V11, V12, V13, V14, V15 };
+		return vld1q_u8(Data);
 	}
 
 	FORCEINLINE void StoreUnaligned_Int8x16(uint8* Ptr, VectorRegisterInt8x16 Reg)
@@ -102,7 +102,8 @@ namespace ImageComputeUtils
 
 	FORCEINLINE VectorRegisterInt8x16 MakeMask_Int8x16(FVec4Mask_U8 Mask)
 	{
-		return uint8x16_t { Mask.X, Mask.Y, Mask.Z, Mask.W, Mask.X, Mask.Y, Mask.Z, Mask.W, Mask.X, Mask.Y, Mask.Z, Mask.W, Mask.X, Mask.Y, Mask.Z, Mask.W };
+		const uint8 Data[16] = { Mask.X, Mask.Y, Mask.Z, Mask.W, Mask.X, Mask.Y, Mask.Z, Mask.W, Mask.X, Mask.Y, Mask.Z, Mask.W, Mask.X, Mask.Y, Mask.Z, Mask.W };
+		return vld1q_u8(Data);
 	}
 
 	FORCEINLINE VectorRegisterInt8x16 Select_Int8x16(VectorRegisterInt8x16 A, VectorRegisterInt8x16 B, VectorRegisterInt8x16 Mask)
