[Mutable] Simplify Material Parameters.

Avoid storing parameters in an intermediate structure. Images still can not be skipped.

#Synced-CL 51258030

diff --git a/Engine/Plugins/Mutable/Source/CustomizableObject/Internal/MuCO/CustomizableObjectSystemPrivate.h b/Engine/Plugins/Mutable/Source/CustomizableObject/Internal/MuCO/CustomizableObjectSystemPrivate.h
--- a/Engine/Plugins/Mutable/Source/CustomizableObject/Internal/MuCO/CustomizableObjectSystemPrivate.h
+++ b/Engine/Plugins/Mutable/Source/CustomizableObject/Internal/MuCO/CustomizableObjectSystemPrivate.h
@@ -252,36 +252,12 @@ struct FInstanceUpdateData
 
 		int32 MaterialLayer = INDEX_NONE; // TODO MATERIAL_PIN This does not belong to the Image but the Material
 	};
-
-	struct FVector
-	{
-		FName Name;
-		FLinearColor Vector;
-
-		int32 MaterialLayer;
-	};
-
-	struct FScalar
-	{
-		FName Name;
-		float Scalar;
-		
-		int32 MaterialLayer;
-	};
-
+	
 	struct FMaterial
 	{
 		/** Range in the Images array */
 		uint16 FirstImage = 0;
 		uint16 ImageCount = 0;
-
-		/** Range in the Vectors array */
-		uint16 FirstVector = 0;
-		uint16 VectorCount = 0;
-
-		/** Range in the Scalar array */
-		uint16 FirstScalar = 0;
-		uint16 ScalarCount = 0;
 				
 		UE::Mutable::Private::FMaterialId MaterialId;
 		UE::Mutable::Private::TManagedPtr<const UE::Mutable::Private::FMaterial> Material;
@@ -367,8 +343,6 @@ struct FInstanceUpdateData
 	TArray<TSharedRef<FMaterial>> Materials;
 	TArray<TSharedRef<FSurface>> Surfaces;
 	TArray<TSharedRef<FImage>> Images;
-	TArray<FVector> Vectors;
-	TArray<FScalar> Scalars;
 
 	TArray<FName> ActiveBones;
 	TArray<FName> BoneMaps;
@@ -394,8 +368,6 @@ struct FInstanceUpdateData
 		OverlayMaterials.Empty();
 		Materials.Empty();
 		Images.Empty();
-		Scalars.Empty();
-		Vectors.Empty();
 		RealTimeMorphTargets.Empty();
 		ExtendedInputPins.Empty();
 	}
diff --git a/Engine/Plugins/Mutable/Source/CustomizableObject/Private/MuCO/CustomizableObjectInstance.cpp b/Engine/Plugins/Mutable/Source/CustomizableObject/Private/MuCO/CustomizableObjectInstance.cpp
--- a/Engine/Plugins/Mutable/Source/CustomizableObject/Private/MuCO/CustomizableObjectInstance.cpp
+++ b/Engine/Plugins/Mutable/Source/CustomizableObject/Private/MuCO/CustomizableObjectInstance.cpp
@@ -5113,46 +5113,21 @@ FBuildMaterialOut BuildMaterial(const TSharedRef<FUpdateContextPrivate>& Context
 		}
 	}
 
-	const bool bHasParameters = Material->ScalarCount || Material->VectorCount || Material->ImageCount;
+	const bool bHasParameters = Material->Material->ColorParameters.Num() || Material->Material->ScalarParameters.Num() || Material->ImageCount;
 	
 	if (!MaterialInterface &&
 		bHasParameters)
 	{
-#if WITH_EDITOR
-		FString MIDName = MaterialTemplate->GetName();
-		{
-			// Remove the MI_ or M_ prefixes from the material string so we use it as the name of the MID
-			const FString MaterialPrefix = TEXT("M_");						// Material
-			const FString MaterialInstancePrefix = TEXT("MI_");				// Material Instance
-			const FString MaterialInstanceConstantPrefix = TEXT("MIC_");	// Material Instance Constant
-			
-			if (MIDName.Find(MaterialInstancePrefix, ESearchCase::CaseSensitive) == 0)
-			{
-				MIDName = MIDName.RightChop(MaterialInstancePrefix.Len());
-			}
-			else if (MIDName.Find(MaterialPrefix, ESearchCase::CaseSensitive) == 0)
-			{
-				MIDName = MIDName.RightChop(MaterialPrefix.Len());
-			}
-			else if (MIDName.Find(MaterialInstanceConstantPrefix, ESearchCase::CaseSensitive) == 0)
-			{
-				MIDName = MIDName.RightChop(MaterialInstanceConstantPrefix.Len());
-			}
-
-			MIDName = FString::Printf(TEXT("MID_%s_h%lu"), *MIDName, UE::Mutable::Private::GetTypeHashPersistent(Material->MaterialId));
+		FName Name = NAME_None;
 
+		if (Context->bBake)
+		{
 			// Add a "tag" to better identify this asset during the bake
-			if (Context->bBake)
-			{
-				MIDName = FBakingConfiguration::BakedResourcePrefix + MIDName;
-			}
+			FString BakeName = FString::Printf(TEXT("%sMID_h%lu"), *FBakingConfiguration::BakedResourcePrefix, UE::Mutable::Private::GetTypeHashPersistent(Material->MaterialId));
+			Name = MakeUniqueObjectName(GetTransientPackage(), UMaterialInstanceDynamic::StaticClass(), *BakeName);
 		}
-		
-		const FName UniqueMIDName = MakeUniqueObjectName(GetTransientPackage(), UMaterialInstanceDynamic::StaticClass(), *MIDName);
-		MaterialInterface = UMaterialInstanceDynamic::Create(MaterialTemplate, GetTransientPackage(), UniqueMIDName);
-#else
-		MaterialInterface = UMaterialInstanceDynamic::Create(MaterialTemplate, GetTransientPackage());
-#endif
+
+		MaterialInterface = UMaterialInstanceDynamic::Create(MaterialTemplate, GetTransientPackage(), Name);
 	}
 	
 	if (!MaterialInterface)
@@ -5164,41 +5139,39 @@ FBuildMaterialOut BuildMaterial(const TSharedRef<FUpdateContextPrivate>& Context
 	{
 		UMaterialInstanceDynamic* MaterialInstance = CastChecked<UMaterialInstanceDynamic>(MaterialInterface);
 		
-		for (int32 VectorIndex = 0; VectorIndex < Material->VectorCount; ++VectorIndex)
+		for (const TTuple<UE::Mutable::Private::FParameterKey, UE::Math::TVector4<float>>& Parameter : Material->Material->ColorParameters)
 		{
-			const FInstanceUpdateData::FVector& Vector = Context->InstanceUpdateData.Vectors[Material->FirstVector + VectorIndex];
-			if (Vector.MaterialLayer < 0)
+			if (Parameter.Key.LayerIndex == INDEX_NONE)
 			{
-				FLinearColor Color = Vector.Vector;
+				FLinearColor Color = Parameter.Value;
 
 				// HACK: We encode an invalid value (Nan) for table option "None.
 				// Decoding "None" color parameters that use the material color
 				if (FMath::IsNaN(Color.R))
 				{
-					FMaterialParameterInfo ParameterInfo(Vector.Name);
+					FMaterialParameterInfo ParameterInfo(Parameter.Key.ParameterName);
 					MaterialTemplate->GetVectorParameterValue(ParameterInfo, Color);
 				}
 
-				MaterialInstance->SetVectorParameterValue(Vector.Name, Color);
+				MaterialInstance->SetVectorParameterValue(Parameter.Key.ParameterName, Color);
 			}
 			else
 			{
-				FMaterialParameterInfo ParameterInfo = FMaterialParameterInfo(Vector.Name, EMaterialParameterAssociation::LayerParameter, Vector.MaterialLayer);
-				MaterialInstance->SetVectorParameterValueByInfo(ParameterInfo, Vector.Vector);
+				FMaterialParameterInfo ParameterInfo = FMaterialParameterInfo(Parameter.Key.ParameterName, EMaterialParameterAssociation::LayerParameter, Parameter.Key.LayerIndex);
+				MaterialInstance->SetVectorParameterValueByInfo(ParameterInfo, Parameter.Value);
 			}
 		}
 
-		for (int32 ScalarIndex = 0; ScalarIndex < Material->ScalarCount; ++ScalarIndex)
+		for (const TTuple<UE::Mutable::Private::FParameterKey, float>& Parameter : Material->Material->ScalarParameters)
 		{
-			const FInstanceUpdateData::FScalar& Scalar = Context->InstanceUpdateData.Scalars[Material->FirstScalar + ScalarIndex];
-			if (Scalar.MaterialLayer < 0)
+			if (Parameter.Key.LayerIndex == INDEX_NONE)
 			{
-				MaterialInstance->SetScalarParameterValue(FName(Scalar.Name), Scalar.Scalar);
+				MaterialInstance->SetScalarParameterValue(Parameter.Key.ParameterName, Parameter.Value);
 			}
 			else
 			{
-				FMaterialParameterInfo ParameterInfo = FMaterialParameterInfo(Scalar.Name, EMaterialParameterAssociation::LayerParameter, Scalar.MaterialLayer);
-				MaterialInstance->SetScalarParameterValueByInfo(ParameterInfo, Scalar.Scalar);
+				FMaterialParameterInfo ParameterInfo = FMaterialParameterInfo(Parameter.Key.ParameterName, EMaterialParameterAssociation::LayerParameter, Parameter.Key.LayerIndex);
+				MaterialInstance->SetScalarParameterValueByInfo(ParameterInfo, Parameter.Value);
 			}
 		}
 
diff --git a/Engine/Plugins/Mutable/Source/CustomizableObject/Private/MuCO/CustomizableObjectSystem.cpp b/Engine/Plugins/Mutable/Source/CustomizableObject/Private/MuCO/CustomizableObjectSystem.cpp
--- a/Engine/Plugins/Mutable/Source/CustomizableObject/Private/MuCO/CustomizableObjectSystem.cpp
+++ b/Engine/Plugins/Mutable/Source/CustomizableObject/Private/MuCO/CustomizableObjectSystem.cpp
@@ -2103,29 +2103,6 @@ namespace impl
 			{
 				Material->Material = GetMaterialTask.GetResult();
 
-				// Scalars
-				Material->FirstScalar = Context->InstanceUpdateData.Scalars.Num();
-				Material->ScalarCount = Material->Material->ScalarParameters.Num();
-				for (TTuple<UE::Mutable::Private::FParameterKey, float> ScalarParameter : Material->Material->ScalarParameters)
-				{
-					FInstanceUpdateData::FScalar& Scalar = Context->InstanceUpdateData.Scalars.AddDefaulted_GetRef();
-					Scalar.Name = ScalarParameter.Key.ParameterName;
-					Scalar.MaterialLayer = (int32)ScalarParameter.Key.LayerIndex;
-					Scalar.Scalar = ScalarParameter.Value;
-					
-				}
-
-				// Vectors
-				Material->FirstVector = Context->InstanceUpdateData.Vectors.Num();
-				Material->VectorCount = Material->Material->ColorParameters.Num();
-				for (TTuple<UE::Mutable::Private::FParameterKey, UE::Math::TVector4<float>> VectorParameter : Material->Material->ColorParameters)
-				{
-					FInstanceUpdateData::FVector& Vector = Context->InstanceUpdateData.Vectors.AddDefaulted_GetRef();
-					Vector.Name = VectorParameter.Key.ParameterName;
-					Vector.MaterialLayer = (int32)VectorParameter.Key.LayerIndex;
-					Vector.Vector = VectorParameter.Value;
-				}
-
 				// Images
 				Material->FirstImage = Context->InstanceUpdateData.Images.Num();
 				Material->ImageCount = Material->Material->ImageParameters.Num();
