[Mutable] Add support for UAssetUserData in SKM Parameter and Constants nodes.

- UAssetUserData removed from Refeference Skeletal Mesh.
- Removed ExtensionData.
- Removed RigLogicDNA extension. No longer needed

#Synced-CL 51263765

#Synced-CL 51263765

diff --git a/Engine/Plugins/Experimental/HairStrandsMutable/Source/HairStrandsMutable/Private/HairStrandsMutableExtension.cpp b/Engine/Plugins/Experimental/HairStrandsMutable/Source/HairStrandsMutable/Private/HairStrandsMutableExtension.cpp
--- a/Engine/Plugins/Experimental/HairStrandsMutable/Source/HairStrandsMutable/Private/HairStrandsMutableExtension.cpp
+++ b/Engine/Plugins/Experimental/HairStrandsMutable/Source/HairStrandsMutable/Private/HairStrandsMutableExtension.cpp
@@ -43,32 +43,9 @@ TArray<FObjectNodeInputPin> UHairStrandsMutableExtension::GetAdditionalObjectNod
 	return Result;
 }
 
-FInstancedStruct UHairStrandsMutableExtension::GenerateExtensionInstanceData(
-	const TArray<FInputPinDataContainer>& InputPinData) const
+void UHairStrandsMutableExtension::OnCustomizableObjectInstanceUsageUpdated(UCustomizableObjectInstanceUsage& InstanceUsage, TArray<const UObject*>& ExtensionData) const
 {
-	FInstancedStruct Result(FGroomInstanceData::StaticStruct());
-	FGroomInstanceData* InstanceData = Result.GetMutablePtr<FGroomInstanceData>();
-
-	for (const FInputPinDataContainer& Container : InputPinData)
-	{
-		if (Container.Pin.PinName == GroomsBaseNodePinName)
-		{
-			const FGroomPinData* Data = Container.Data.GetPtr<FGroomPinData>();
-			if (Data
-				&& Data->GroomAsset != nullptr)
-			{
-				InstanceData->Grooms.Add(*Data);
-			}
-		}
-	}
-
-	return Result;
-}
-
-
-void UHairStrandsMutableExtension::OnCustomizableObjectInstanceUsageUpdated(UCustomizableObjectInstanceUsage& InstanceUsage) const
-{
-	Super::OnCustomizableObjectInstanceUsageUpdated(InstanceUsage);
+	Super::OnCustomizableObjectInstanceUsageUpdated(InstanceUsage, ExtensionData);
 	
 	UCustomizableObjectInstance* Instance = InstanceUsage.GetCustomizableObjectInstance();
 	if (!Instance)
@@ -82,20 +59,19 @@ void UHairStrandsMutableExtension::OnCustomizableObjectInstanceUsageUpdated(UCus
 		return;
 	}
 	
-	const FInstancedStruct UntypedInstanceData = Instance->GetExtensionInstanceData(this);
-	const FGroomInstanceData* InstanceData = UntypedInstanceData.GetPtr<FGroomInstanceData>();
-	if (!InstanceData)
-	{
-		return;
-	}
-
 	USceneComponent* AttachParent = InstanceUsage.GetAttachParent();
 	
 	TArray<UGroomComponent*> UpdateGroomComponents; // Grooms that must be present at the end of this update.
 
-	for (const FGroomPinData& GroomPinData : InstanceData->Grooms)
+	for (const UObject* Extension : ExtensionData)
 	{
-		if (GroomPinData.ComponentName != InstanceUsage.GetComponentName())
+		const UGroomCompiledData* GroomData = Cast<UGroomCompiledData>(Extension);
+		if (!GroomData)
+		{
+			continue;
+		}
+		
+		if (GroomData->ComponentName != InstanceUsage.GetComponentName())
 		{
 			continue;
 		}
@@ -110,12 +86,12 @@ void UHairStrandsMutableExtension::OnCustomizableObjectInstanceUsageUpdated(UCus
 				continue;
 			}
 
-			if (GroomComponent->GroomAsset == GroomPinData.GroomAsset &&
-				GroomComponent->GroomCache == GroomPinData.GroomCache &&
-				GroomComponent->BindingAsset == GroomPinData.BindingAsset &&
-				GroomComponent->PhysicsAsset == GroomPinData.PhysicsAsset &&
-				GroomComponent->AttachmentName == GroomPinData.AttachmentName &&
-				GroomComponent->OverrideMaterials == GroomPinData.OverrideMaterials)
+			if (GroomComponent->GroomAsset == GroomData->GroomAsset &&
+				GroomComponent->GroomCache == GroomData->GroomCache &&
+				GroomComponent->BindingAsset == GroomData->BindingAsset &&
+				GroomComponent->PhysicsAsset == GroomData->PhysicsAsset &&
+				GroomComponent->AttachmentName == GroomData->AttachmentName &&
+				GroomComponent->OverrideMaterials == GroomData->OverrideMaterials)
 			{
 				UpdateGroomComponents.Add(GroomComponent);
 				bAlreadyExists = true;
@@ -128,13 +104,13 @@ void UHairStrandsMutableExtension::OnCustomizableObjectInstanceUsageUpdated(UCus
 			continue;
 		}
 		
-		UGroomComponent* GroomComponent = NewObject<UGroomComponent>(AttachParent, GroomPinData.GroomComponentName);
-		GroomComponent->GroomAsset = GroomPinData.GroomAsset;
-		GroomComponent->GroomCache = GroomPinData.GroomCache;
-		GroomComponent->BindingAsset = GroomPinData.BindingAsset;
-		GroomComponent->PhysicsAsset = GroomPinData.PhysicsAsset;
-		GroomComponent->AttachmentName = GroomPinData.AttachmentName;
-		GroomComponent->OverrideMaterials = GroomPinData.OverrideMaterials;
+		UGroomComponent* GroomComponent = NewObject<UGroomComponent>(AttachParent, GroomData->GroomComponentName);
+		GroomComponent->GroomAsset = GroomData->GroomAsset;
+		GroomComponent->GroomCache = GroomData->GroomCache;
+		GroomComponent->BindingAsset = GroomData->BindingAsset;
+		GroomComponent->PhysicsAsset = GroomData->PhysicsAsset;
+		GroomComponent->AttachmentName = GroomData->AttachmentName;
+		GroomComponent->OverrideMaterials = GroomData->OverrideMaterials;
 		
 		// Work around UE-158069
 		GroomComponent->CreationMethod = EComponentCreationMethod::Instance;
diff --git a/Engine/Plugins/Experimental/HairStrandsMutable/Source/HairStrandsMutable/Public/HairStrandsMutableExtension.h b/Engine/Plugins/Experimental/HairStrandsMutable/Source/HairStrandsMutable/Public/HairStrandsMutableExtension.h
--- a/Engine/Plugins/Experimental/HairStrandsMutable/Source/HairStrandsMutable/Public/HairStrandsMutableExtension.h
+++ b/Engine/Plugins/Experimental/HairStrandsMutable/Source/HairStrandsMutable/Public/HairStrandsMutableExtension.h
@@ -3,6 +3,7 @@
 #pragma once
 
 #include "MuCO/CustomizableObjectExtension.h"
+#include "MuR/ExtensionData.h"
 
 #include "HairStrandsMutableExtension.generated.h"
 
@@ -52,14 +53,43 @@ public:
 	TArray<TObjectPtr<UMaterialInterface>> OverrideMaterials;
 };
 
-USTRUCT()
-struct HAIRSTRANDSMUTABLE_API FGroomInstanceData
+UCLASS()
+class UGroomCompiledData : public UObject
 {
 	GENERATED_BODY()
-
+	
 public:
-	UPROPERTY()
-	TArray<FGroomPinData> Grooms;
+	/** Name of the mesh component this groom will be attached to */
+	UPROPERTY(EditAnywhere, Category = CustomizableObject)
+	FName ComponentName;
+
+	/** The groom asset to use. See UGroomComponent::GroomAsset. */
+	UPROPERTY(EditAnywhere, Category = Groom)
+	TObjectPtr<UGroomAsset> GroomAsset;
+
+	/** See UGroomComponent::GroomCache. */
+	UPROPERTY(EditAnywhere, Category = GroomCache, meta = (EditCondition = "BindingAsset == nullptr"))
+	TObjectPtr<UGroomCache> GroomCache;
+
+	/** See UGroomComponent::BindingAsset. */
+	UPROPERTY(EditAnywhere, Category = Groom, meta = (EditCondition = "GroomCache == nullptr"))
+	TObjectPtr<UGroomBindingAsset> BindingAsset;
+
+	/** See UGroomComponent::PhysicsAsset. */
+	UPROPERTY(EditAnywhere, Category = Simulation)
+	TObjectPtr<UPhysicsAsset> PhysicsAsset;
+
+	/** See UGroomComponent::AttachmentName. */
+	UPROPERTY(EditAnywhere, Category = Groom)
+	FString AttachmentName;
+
+	/** Name of the created groom component */
+	UPROPERTY(EditAnywhere, Category = Groom)
+	FName GroomComponentName;
+
+	/** See UMeshComponent::OverrideMaterials. */
+	UPROPERTY(EditAnywhere, Category = Rendering)
+	TArray<TObjectPtr<UMaterialInterface>> OverrideMaterials;
 };
 
 /**
@@ -75,8 +105,7 @@ public:
 	/** UCustomizableObjectExtension interface */
 	virtual TArray<FCustomizableObjectPinType> GetPinTypes() const override;
 	virtual TArray<FObjectNodeInputPin> GetAdditionalObjectNodePins() const override;
-	virtual FInstancedStruct GenerateExtensionInstanceData(const TArray<FInputPinDataContainer>& InputPinData) const override;
-	virtual void OnCustomizableObjectInstanceUsageUpdated(UCustomizableObjectInstanceUsage& Usage) const override;
+	virtual void OnCustomizableObjectInstanceUsageUpdated(UCustomizableObjectInstanceUsage& Usage, TArray<const UObject*>& ExtensionData) const override;
 	virtual void OnCustomizableObjectInstanceUsageDiscarded(UCustomizableObjectInstanceUsage& Usage) const override;
 
 	HAIRSTRANDSMUTABLE_API static const FName GroomPinType;
diff --git a/Engine/Plugins/Experimental/HairStrandsMutable/Source/HairStrandsMutableEditor/Private/CustomizableObjectNodeGroomConstant.cpp b/Engine/Plugins/Experimental/HairStrandsMutable/Source/HairStrandsMutableEditor/Private/CustomizableObjectNodeGroomConstant.cpp
--- a/Engine/Plugins/Experimental/HairStrandsMutable/Source/HairStrandsMutableEditor/Private/CustomizableObjectNodeGroomConstant.cpp
+++ b/Engine/Plugins/Experimental/HairStrandsMutable/Source/HairStrandsMutableEditor/Private/CustomizableObjectNodeGroomConstant.cpp
@@ -3,11 +3,10 @@
 #include "CustomizableObjectNodeGroomConstant.h"
 
 #include "MuCOE/EdGraphSchema_CustomizableObject.h"
-#include "MuCOE/ExtensionDataCompilerInterface.h"
 #include "MuCOE/RemapPins/CustomizableObjectNodeRemapPins.h"
-#include "MuR/ExtensionData.h"
 #include "MuT/NodeExtensionDataConstant.h"
 #include "HairStrandsMutableExtension.h"
+#include "MuCOE/ExtensionDataCompilerInterface.h"
 
 #include UE_INLINE_GENERATED_CPP_BY_NAME(CustomizableObjectNodeGroomConstant)
 
@@ -29,9 +28,24 @@ FText UCustomizableObjectNodeGroomConstant::GetTooltipText() const
 	return LOCTEXT("Groom_Constant_Tooltip", "Imports a Groom");
 }
 
+void UCustomizableObjectNodeGroomConstant::PostLoad()
+{
+	Super::PostLoad();
+	
+	CompiledData = NewObject<UGroomCompiledData>(this, *FString::Printf(TEXT("UGroomData_%s"), *NodeGuid.ToString()));
+	CopyCompiledData();
+}
+
+void UCustomizableObjectNodeGroomConstant::PostEditChangeProperty(FPropertyChangedEvent& PropertyChangedEvent)
+{
+	Super::PostEditChangeProperty(PropertyChangedEvent);
+	
+	CopyCompiledData();
+}
+
 void UCustomizableObjectNodeGroomConstant::AllocateDefaultPins(UCustomizableObjectNodeRemapPins* RemapPins)
 {
-	UEdGraphPin* OutputPin = CustomCreatePin(EGPD_Output, UHairStrandsMutableExtension::GroomPinType, UHairStrandsMutableExtension::GroomsBaseNodePinName, UHairStrandsMutableExtension::GroomNodeCategory);
+	CustomCreatePin(EGPD_Output, UHairStrandsMutableExtension::GroomPinType, UHairStrandsMutableExtension::GroomsBaseNodePinName, UHairStrandsMutableExtension::GroomNodeCategory);
 }
 
 bool UCustomizableObjectNodeGroomConstant::ShouldAddToContextMenu(FText& OutCategory) const
@@ -40,31 +54,31 @@ bool UCustomizableObjectNodeGroomConstant::ShouldAddToContextMenu(FText& OutCate
 	return true;
 }
 
-
 bool UCustomizableObjectNodeGroomConstant::IsExperimental() const
 {
 	return true;
 }
 
-
 UE::Mutable::Private::Ptr<UE::Mutable::Private::NodeExtensionData> UCustomizableObjectNodeGroomConstant::GenerateMutableNode(FExtensionDataCompilerInterface& CompilerInterface) const
 {
 	check(IsInGameThread());
 
-	// Create node and extension data container
 	UE::Mutable::Private::Ptr<UE::Mutable::Private::NodeExtensionDataConstant> Result = new UE::Mutable::Private::NodeExtensionDataConstant();
-
-	FInstancedStruct Struct;
-	Struct.InitializeAs<FGroomPinData>(GroomData);
-
-	// Grooms are usually quite large, so set it up as a streaming constant to allow it to be
-	// loaded on demand.
-	//
-	// If needed we could expose an editable UPROPERTY to give the user the option of making this
-	// an always-loaded constant.
-	Result->SetValue(CompilerInterface.MakeStreamedExtensionData(MoveTemp(Struct)));
+	Result->ExtensionDataId = CompilerInterface.MakeExtensionData(*CompiledData, true);
 
 	return Result;
 }
 
+void UCustomizableObjectNodeGroomConstant::CopyCompiledData()
+{
+	CompiledData->ComponentName = GroomData.ComponentName;
+	CompiledData->GroomAsset = GroomData.GroomAsset;
+	CompiledData->GroomCache = GroomData.GroomCache;
+	CompiledData->BindingAsset = GroomData.BindingAsset;
+	CompiledData->PhysicsAsset = GroomData.PhysicsAsset;
+	CompiledData->AttachmentName = GroomData.AttachmentName;
+	CompiledData->GroomComponentName = GroomData.GroomComponentName;
+	CompiledData->OverrideMaterials = GroomData.OverrideMaterials;
+}
+
 #undef LOCTEXT_NAMESPACE
diff --git a/Engine/Plugins/Experimental/HairStrandsMutable/Source/HairStrandsMutableEditor/Public/CustomizableObjectNodeGroomConstant.h b/Engine/Plugins/Experimental/HairStrandsMutable/Source/HairStrandsMutableEditor/Public/CustomizableObjectNodeGroomConstant.h
--- a/Engine/Plugins/Experimental/HairStrandsMutable/Source/HairStrandsMutableEditor/Public/CustomizableObjectNodeGroomConstant.h
+++ b/Engine/Plugins/Experimental/HairStrandsMutable/Source/HairStrandsMutableEditor/Public/CustomizableObjectNodeGroomConstant.h
@@ -19,6 +19,8 @@ public:
 	virtual FText GetNodeTitle(ENodeTitleType::Type TitleType) const override;
 	virtual FLinearColor GetNodeTitleColor() const override;
 	virtual FText GetTooltipText() const override;
+	virtual void PostLoad() override;
+	virtual void PostEditChangeProperty(FPropertyChangedEvent& PropertyChangedEvent) override;
 
 	/** UCustomizableObjectNode interface */
 	virtual void AllocateDefaultPins(UCustomizableObjectNodeRemapPins* RemapPins) override;
@@ -28,7 +30,13 @@ public:
 	/** ICustomizableObjectExtensionNode interface */
 	virtual UE::Mutable::Private::Ptr<UE::Mutable::Private::NodeExtensionData> GenerateMutableNode(FExtensionDataCompilerInterface& CompilerInterface) const override;
 
+	/** Own interface */
+	void CopyCompiledData();
+	
 	UPROPERTY(EditAnywhere, Category = CustomizableObject)
 	FGroomPinData GroomData;
+	
+	UPROPERTY()
+	TObjectPtr<UGroomCompiledData> CompiledData;
 };
 
diff --git a/Engine/Plugins/Experimental/RigLogicMutable/RigLogicMutable.uplugin b/Engine/Plugins/Experimental/RigLogicMutable/RigLogicMutable.uplugin
deleted file mode 100644
diff --git a/Engine/Plugins/Experimental/RigLogicMutable/Source/RigLogicMutable/Private/RigLogicMutableExtension.cpp b/Engine/Plugins/Experimental/RigLogicMutable/Source/RigLogicMutable/Private/RigLogicMutableExtension.cpp
deleted file mode 100644
diff --git a/Engine/Plugins/Experimental/RigLogicMutable/Source/RigLogicMutable/Private/RigLogicMutableModule.cpp b/Engine/Plugins/Experimental/RigLogicMutable/Source/RigLogicMutable/Private/RigLogicMutableModule.cpp
deleted file mode 100644
diff --git a/Engine/Plugins/Experimental/RigLogicMutable/Source/RigLogicMutable/Public/RigLogicMutableExtension.h b/Engine/Plugins/Experimental/RigLogicMutable/Source/RigLogicMutable/Public/RigLogicMutableExtension.h
deleted file mode 100644
diff --git a/Engine/Plugins/Experimental/RigLogicMutable/Source/RigLogicMutable/Public/RigLogicMutableModule.h b/Engine/Plugins/Experimental/RigLogicMutable/Source/RigLogicMutable/Public/RigLogicMutableModule.h
deleted file mode 100644
diff --git a/Engine/Plugins/Experimental/RigLogicMutable/Source/RigLogicMutable/RigLogicMutable.Build.cs b/Engine/Plugins/Experimental/RigLogicMutable/Source/RigLogicMutable/RigLogicMutable.Build.cs
deleted file mode 100644
diff --git a/Engine/Plugins/Experimental/RigLogicMutable/Source/RigLogicMutableEditor/Private/CustomizableObjectNodeDNAConstant.cpp b/Engine/Plugins/Experimental/RigLogicMutable/Source/RigLogicMutableEditor/Private/CustomizableObjectNodeDNAConstant.cpp
deleted file mode 100644
diff --git a/Engine/Plugins/Experimental/RigLogicMutable/Source/RigLogicMutableEditor/Private/RigLogicMutableEditorModule.cpp b/Engine/Plugins/Experimental/RigLogicMutable/Source/RigLogicMutableEditor/Private/RigLogicMutableEditorModule.cpp
deleted file mode 100644
diff --git a/Engine/Plugins/Experimental/RigLogicMutable/Source/RigLogicMutableEditor/Public/CustomizableObjectNodeDNAConstant.h b/Engine/Plugins/Experimental/RigLogicMutable/Source/RigLogicMutableEditor/Public/CustomizableObjectNodeDNAConstant.h
deleted file mode 100644
diff --git a/Engine/Plugins/Experimental/RigLogicMutable/Source/RigLogicMutableEditor/Public/RigLogicMutableEditorModule.h b/Engine/Plugins/Experimental/RigLogicMutable/Source/RigLogicMutableEditor/Public/RigLogicMutableEditorModule.h
deleted file mode 100644
diff --git a/Engine/Plugins/Experimental/RigLogicMutable/Source/RigLogicMutableEditor/RigLogicMutableEditor.Build.cs b/Engine/Plugins/Experimental/RigLogicMutable/Source/RigLogicMutableEditor/RigLogicMutableEditor.Build.cs
deleted file mode 100644
diff --git a/Engine/Plugins/Mutable/Source/CustomizableObject/Internal/MuCO/CustomizableObjectInstancePrivate.h b/Engine/Plugins/Mutable/Source/CustomizableObject/Internal/MuCO/CustomizableObjectInstancePrivate.h
--- a/Engine/Plugins/Mutable/Source/CustomizableObject/Internal/MuCO/CustomizableObjectInstancePrivate.h
+++ b/Engine/Plugins/Mutable/Source/CustomizableObject/Internal/MuCO/CustomizableObjectInstancePrivate.h
@@ -83,23 +83,11 @@ struct FCustomizableInstanceComponentData
 	FName ComponentName;
 	
 	UE::Mutable::Private::FComponentId ComponentId = INDEX_NONE;
-	
-	// AnimBP data gathered for a component from its constituent meshes
-	UPROPERTY(Transient, Category = CustomizableObjectInstance, editfixedsize, VisibleAnywhere)
-	TMap<FName, TSoftClassPtr<UAnimInstance>> AnimSlotToBP;
-
-	// AssetUserData gathered for a component from its constituent meshes
-	UPROPERTY(Transient, Category = CustomizableObjectInstance, editfixedsize, VisibleAnywhere)
-	TSet<TObjectPtr<UAssetUserData>> AssetUserDataArray;
-
-	// Index of the resource in the StreamedResourceData array of the CustomizableObject.
-	TArray<int32> StreamedResourceIndex;
 
 #if WITH_EDITORONLY_DATA
 	// Just used for mutable.EnableMutableAnimInfoDebugging command
 	TArray<FString> MeshPartPaths;
 #endif
-
 	
 	FComponentDescriptor LastComponentDescriptor;
 
@@ -113,6 +101,7 @@ struct FCustomizableInstanceComponentData
 	TArray<TObjectPtr<UMaterialInterface>> OverlayMaterials;
 };
 
+
 USTRUCT()
 struct FAnimInstanceOverridePhysicsAsset
 {
@@ -125,6 +114,7 @@ struct FAnimInstanceOverridePhysicsAsset
 	TObjectPtr<UPhysicsAsset> PhysicsAsset;
 };
 
+
 USTRUCT()
 struct FAnimBpGeneratedPhysicsAssets
 {
@@ -135,19 +125,6 @@ struct FAnimBpGeneratedPhysicsAssets
 };
 
 
-USTRUCT()
-struct FExtensionInstanceData
-{
-	GENERATED_BODY()
- 
-	UPROPERTY()
-	TWeakObjectPtr<const UCustomizableObjectExtension> Extension;
- 
-	UPROPERTY()
-	FInstancedStruct Data;
-};
-
-
 struct FMaterialReuseCacheKey
 {
 	TStrongObjectPtr<UMaterialInterface> MaterialTemplate;
@@ -223,10 +200,6 @@ public:
 	UE_API void OnObjectStatusChanged(FCustomizableObjectStatus::EState Previous, FCustomizableObjectStatus::EState Next);
 #endif
 	
-	/** Invalidates the previously generated data and retrieves information from the CObject after specific actions.
-	 *  It'll be called in the PostLoad, after Compiling the CO, and after changing the CO of the Instance. */
-	UE_API void InitCustomizableObjectData(const UCustomizableObject* InCustomizableObject);
-	
 	ECOInstanceFlags GetCOInstanceFlags() const { return InstanceFlagsPrivate; }
 	void SetCOInstanceFlags(ECOInstanceFlags FlagsToSet) { InstanceFlagsPrivate = (ECOInstanceFlags)(InstanceFlagsPrivate | FlagsToSet); }
 	void ClearCOInstanceFlags(ECOInstanceFlags FlagsToClear) { InstanceFlagsPrivate = (ECOInstanceFlags)(InstanceFlagsPrivate & ~FlagsToClear); }
@@ -307,20 +280,10 @@ public:
 	
 	UPROPERTY(Transient)
 	TArray<FCustomizableInstanceComponentData> ComponentsData;
-	
-	// To keep loaded AnimBPs referenced and prevent GC
-	UPROPERTY(Transient, Category = Animation, editfixedsize, VisibleAnywhere)
-	TArray<TSubclassOf<UAnimInstance>> GatheredAnimBPs;
-
-	UPROPERTY(Transient, Category = Animation, editfixedsize, VisibleAnywhere)
-	FGameplayTagContainer AnimBPGameplayTags;
 
 	UPROPERTY(Transient, Category = Animation, editfixedsize, VisibleAnywhere)
 	TMap<TSubclassOf<UAnimInstance>, FAnimBpGeneratedPhysicsAssets> AnimBpPhysicsAssets;
 
-	UPROPERTY(Transient)
-	TArray<FExtensionInstanceData> ExtensionInstanceData;
-
 private:
 	ECOInstanceFlags InstanceFlagsPrivate = ECOInstanceFlags::ECONone;
 
@@ -381,7 +344,7 @@ UE_API UPhysicsAsset* GetOrBuildMainPhysicsAsset(const TSharedRef<FUpdateContext
 
 UE_API USkeleton* MergeSkeletons(const TArray<TObjectPtr<USkeleton>>& SkeletonsToMerge, UCustomizableObject& CustomizableObject, bool& bOutCreatedNewSkeleton);
 
-UE_API void InitSkeletalMeshData(const TSharedRef<FUpdateContextPrivate>& Context, USkeletalMesh* SkeletalMesh, const FMutableRefSkeletalMeshData& RefSkeletalMeshData, const UCustomizableObject& CustomizableObject, UE::Mutable::Private::FComponentId ComponentId);
+UE_API void InitSkeletalMeshData(const TSharedRef<FUpdateContextPrivate>& Context, USkeletalMesh* SkeletalMesh, const FMutableRefSkeletalMeshData& RefSkeletalMeshData, const UCustomizableObject& CustomizableObject, const TSharedPtr<FInstanceUpdateData::FComponent>& Component);
 
 UE_API void BuildOrCopyClothingData(const TSharedRef<FUpdateContextPrivate>& Context, USkeletalMesh* SkeletalMesh, const UModelResources& CustomizableObjectInstance, int32 InstanceComponentIndex);
 
diff --git a/Engine/Plugins/Mutable/Source/CustomizableObject/Internal/MuCO/CustomizableObjectPrivate.h b/Engine/Plugins/Mutable/Source/CustomizableObject/Internal/MuCO/CustomizableObjectPrivate.h
--- a/Engine/Plugins/Mutable/Source/CustomizableObject/Internal/MuCO/CustomizableObjectPrivate.h
+++ b/Engine/Plugins/Mutable/Source/CustomizableObject/Internal/MuCO/CustomizableObjectPrivate.h
@@ -4,7 +4,6 @@
 
 #include "MuCO/CustomizableObject.h"
 #include "MuCO/CustomizableObjectCompilerTypes.h"
-#include "MuCO/CustomizableObjectStreamedResourceData.h"
 #include "MuCO/StateMachine.h"
 #include "MuCO/CustomizableObjectUIData.h"
 #include "MuCO/CustomizableObjectIdentifier.h"
@@ -84,24 +83,6 @@ struct MutableCompiledDataStreamHeader
 	}
 };
 
-struct FCustomizableObjectStreameableResourceId
-{
-	enum class EType : uint8
-	{
-		None                  = 0,
-		AssetUserData         = 1,
-	};
-
-	uint64 Id   : 64 - 8;
-	uint64 Type : 8;
-
-	friend bool operator==(FCustomizableObjectStreameableResourceId A, FCustomizableObjectStreameableResourceId B)
-	{
-		return BitCast<uint64>(A) == BitCast<uint64>(B);
-	}
-};
-static_assert(sizeof(FCustomizableObjectStreameableResourceId) == sizeof(uint64));
-
 
 USTRUCT()
 struct FMutableRemappedBone
@@ -460,10 +441,6 @@ struct FMutableRefSkeletalMeshData
 	UPROPERTY()
 	TObjectPtr<UPhysicsAsset> ShadowPhysicsAsset;
 
-	// Asset user data
-	UPROPERTY()
-	TArray<int32> AssetUserDataIndices;
-
 #if WITH_EDITORONLY_DATA
 	friend FArchive& operator<<(FArchive& Ar, FMutableRefSkeletalMeshData& Data);
 #endif
@@ -1015,12 +992,11 @@ public:
 	
 	UPROPERTY()
 	TMap<uint32, TSoftObjectPtr<UObject>> PassthroughObjects;
-
-	/** Objects that we can not directly reference since they are private. Pointing the original package is not of our interest since we would be including unnecessary assets. */
-	UPROPERTY()
-	TArray<TObjectPtr<UObject>> StolenPrivateObjects;
 	
 #if WITH_EDITORONLY_DATA
+	UPROPERTY()
+	TArray<uint32> DuplicateObjects;
+	
 	/** Runtime referenced textures used by the UE::Mutable::Private::FModel. */
 	UPROPERTY()
 	TArray<TSoftObjectPtr<UTexture2D>> RuntimeReferencedTextures;
@@ -1092,26 +1068,8 @@ public:
 	  * Updated each time the CO is compiled and saved in the Derived Data. */
 	UPROPERTY()
 	TMap<FName, FGuid> ParticipatingObjects;
-
-	UPROPERTY()
-	TArray<FCustomizableObjectResourceData> StreamedResourceDataEditor;
-
-	UPROPERTY()
-	TArray<FCustomizableObjectResourceData> StreamedExtensionDataEditor;
 #endif
-
-	// Constant Resources streamed in on demand when generating meshes
-	UPROPERTY()
-	TArray<FCustomizableObjectStreamedResourceData> StreamedResourceData;
-
-	// UE::Mutable::Private::FExtensionData::Index is an index into this array when UE::Mutable::Private::ExtensionData::Origin is ConstantAlwaysLoaded
-	UPROPERTY()
-	TArray<FCustomizableObjectResourceData> AlwaysLoadedExtensionData;
-
-	// UE::Mutable::Private::FExtensionData::Index is an index into this array when UE::Mutable::Private::ExtensionData::Origin is ConstantStreamed
-	UPROPERTY()
-	TArray<FCustomizableObjectStreamedResourceData> StreamedExtensionData;
-
+	
 	// Key is UE::Mutable::Private::OP::ADDRESS. UPROPERTY does not support typedefs. 
 	UPROPERTY()
 	TMap<uint32, FInstancedStruct> ExternalOperations;
@@ -1171,8 +1129,6 @@ public:
 		* compilation setting than the stored in the COE.*/
 	UPROPERTY()
 	bool bCompiledWithHDTextureCompression = false;
-	
-	CUSTOMIZABLEOBJECT_API void InitCookData(UObject& InCustomizableObject);
 #endif
 };
 
@@ -1491,6 +1447,8 @@ enum class ECustomizableObjectVersions : int32
 
 	RemoveMorphTarget,
 	
+	AssetUserDataPassthroughObject2,
+	
 	// -----<new versions can be added above this line>--------
 	LastCustomizableObjectVersion
 };
@@ -1660,8 +1618,6 @@ public:
 	CUSTOMIZABLEOBJECT_API EMutableCompileMeshType GetMeshCompileType() const;
 	
 	CUSTOMIZABLEOBJECT_API const TArray<TSoftObjectPtr<UCustomizableObject>>& GetWorkingSet() const;
-
-	CUSTOMIZABLEOBJECT_API bool IsAssetUserDataMergeEnabled() const;
 	
 	CUSTOMIZABLEOBJECT_API bool IsTableMaterialsParentCheckDisabled() const;
 
diff --git a/Engine/Plugins/Mutable/Source/CustomizableObject/Internal/MuCO/CustomizableObjectStreamedResourceData.h b/Engine/Plugins/Mutable/Source/CustomizableObject/Internal/MuCO/CustomizableObjectStreamedResourceData.h
deleted file mode 100644
diff --git a/Engine/Plugins/Mutable/Source/CustomizableObject/Internal/MuCO/CustomizableObjectSystemPrivate.h b/Engine/Plugins/Mutable/Source/CustomizableObject/Internal/MuCO/CustomizableObjectSystemPrivate.h
--- a/Engine/Plugins/Mutable/Source/CustomizableObject/Internal/MuCO/CustomizableObjectSystemPrivate.h
+++ b/Engine/Plugins/Mutable/Source/CustomizableObject/Internal/MuCO/CustomizableObjectSystemPrivate.h
@@ -351,13 +351,6 @@ struct FInstanceUpdateData
 	/** Key is the Component Name. Value is the final Morph Target data to copy into the Skeletal Mesh. */
 	TMap<UE::Mutable::Private::FComponentId, FSkeletalMeshMorphTargets> RealTimeMorphTargets;
 
-	struct FNamedExtensionData
-	{
-		UE::Mutable::Private::TManagedPtr<const UE::Mutable::Private::FExtensionData> Data;
-		FName Name;
-	};
-	TArray<FNamedExtensionData> ExtendedInputPins;
-
 	/** */
 	void Clear()
 	{
@@ -369,7 +362,6 @@ struct FInstanceUpdateData
 		Materials.Empty();
 		Images.Empty();
 		RealTimeMorphTargets.Empty();
-		ExtendedInputPins.Empty();
 	}
 };
 
@@ -436,6 +428,8 @@ public:
 	TArray<TStrongObjectPtr<const UTexture>> LoadedPassThroughTexturesPendingSetMaterial;
 	
 	TArray<UE::Mutable::Private::FImageId> CachedTextures;
+
+	TArray<UE::Mutable::Private::TManagedPtr<const UE::Mutable::Private::FExtensionData>> ExtensionData;
 	
 	/** Used to know if the updated instances' meshes are different from the previous ones. 
 	  * The index of the array is the instance component's index.
@@ -445,12 +439,6 @@ public:
 	/** List of component names. Key is the FComponentId. */
 	TArray<FName> ComponentNames;
 
-	/** Index of the resource in the StreamedResourceData array of the Model Resources. */
-	TArray<int32> StreamedResourceIndex;
-
-	/** Index of the resource in the ExtensionStreamedResourceData array of the Model Resources. */
-	TArray<int32> ExtensionStreamedResourceIndex;
-	
 private:
 	TArray<FComponentDescriptor> ComponentDescriptors;
 
@@ -755,8 +743,6 @@ public:
 
 	/** Update the last set amount of internal memory Mutable can use to build objects. */
 	void UpdateMemoryLimit();
-
-	bool IsMutableAnimInfoDebuggingEnabled() const;
 	
 	/** Start the actual work of Update Skeletal Mesh process (Update Skeletal Mesh without the queue). */
 	void StartUpdateSkeletalMesh(const TSharedRef<FUpdateContextPrivate>& Context);
diff --git a/Engine/Plugins/Mutable/Source/CustomizableObject/Internal/MuCO/CustomizableSkeletalComponentPrivate.h b/Engine/Plugins/Mutable/Source/CustomizableObject/Internal/MuCO/CustomizableSkeletalComponentPrivate.h
--- a/Engine/Plugins/Mutable/Source/CustomizableObject/Internal/MuCO/CustomizableSkeletalComponentPrivate.h
+++ b/Engine/Plugins/Mutable/Source/CustomizableObject/Internal/MuCO/CustomizableSkeletalComponentPrivate.h
@@ -20,9 +20,6 @@ class UCustomizableSkeletalComponentPrivate : public UObject
 public:
 	UE_API UCustomizableSkeletalComponentPrivate();
 	
-	/** Common end point of all updates. Even those which failed. */
-	UE_API void Callbacks() const;
-	
 	UE_API USkeletalMesh* GetSkeletalMesh() const;
 
 	UE_API USkeletalMesh* GetAttachedSkeletalMesh() const;
diff --git a/Engine/Plugins/Mutable/Source/CustomizableObject/Internal/MuCO/UnrealToMutableTextureConversionUtils.h b/Engine/Plugins/Mutable/Source/CustomizableObject/Internal/MuCO/UnrealToMutableTextureConversionUtils.h
--- a/Engine/Plugins/Mutable/Source/CustomizableObject/Internal/MuCO/UnrealToMutableTextureConversionUtils.h
+++ b/Engine/Plugins/Mutable/Source/CustomizableObject/Internal/MuCO/UnrealToMutableTextureConversionUtils.h
@@ -2,6 +2,7 @@
 
 #pragma once
 
+#include "GameplayTagContainer.h"
 #include "HAL/Platform.h"
 #include "Engine/Texture.h"
 #include "MuR/ComponentId.h"
@@ -118,7 +119,11 @@ struct FMutableSourceMeshData
 	const UObject* MessageContext = nullptr;
 
 	FString OptionalMessageInfo;
+	
+	FName AnimBPSlotName;
 
+	FGameplayTagContainer GameplayTags;
+	
 	bool operator==(const FMutableSourceMeshData&) const = default;
 };
 
@@ -132,6 +137,8 @@ inline uint32 GetTypeHash(const FMutableSourceMeshData& Key)
 	GuidHash = HashCombineFast(GuidHash, GetTypeHash(Key.bOnlyConnectedLOD));
 	GuidHash = HashCombineFast(GuidHash, GetTypeHash(Key.ComponentId));
 	GuidHash = HashCombineFast(GuidHash, GetTypeHash(Key.Flags));
+	GuidHash = HashCombineFast(GuidHash, GetTypeHash(Key.AnimBPSlotName));
+	//GuidHash = HashCombineFast(GuidHash, GetTypeHash(Key.GameplayTags)); // Hash func not implemented, GetGameplayTagArray order is not deterministic.
 	return GuidHash;
 }
 
diff --git a/Engine/Plugins/Mutable/Source/CustomizableObject/Private/MuCO/CustomizableObject.cpp b/Engine/Plugins/Mutable/Source/CustomizableObject/Private/MuCO/CustomizableObject.cpp
--- a/Engine/Plugins/Mutable/Source/CustomizableObject/Private/MuCO/CustomizableObject.cpp
+++ b/Engine/Plugins/Mutable/Source/CustomizableObject/Private/MuCO/CustomizableObject.cpp
@@ -46,9 +46,9 @@
 #include "DerivedDataCacheKey.h"
 #include "DerivedDataRequestOwner.h"
 #include "MuCO/LoadUtils.h"
+#include "AssetToolsModule.h"
 #endif
 
-
 #include "ExternalOperationProvider.h"
 #include "MuCO/CustomizableObjectCustomVersion.h"
 #include "MuCO/CustomizableObjectResourceDataTypes.h"
@@ -513,115 +513,6 @@ const TSharedPtr<UE::Mutable::Private::FModel, ESPMode::ThreadSafe> LoadModel_In
 }
 
 
-void UModelResources::InitCookData(UObject& InCustomizableObject)
-{
-	const FString ObjectName = InCustomizableObject.GetName();
-
-	const int32 NumStreamedResources = StreamedResourceDataEditor.Num();
-	for (int32 Index = 0; Index < NumStreamedResources; ++Index)
-	{
-		FCustomizableObjectResourceData& ResourceData = StreamedResourceDataEditor[Index];
-
-		const FString& ContainerName = ObjectName + FString::Printf(TEXT("_SR_%d"), Index);
-
-		UCustomizableObjectResourceDataContainer* Container = FindObject<UCustomizableObjectResourceDataContainer>(&InCustomizableObject, *ContainerName);
-		if (!Container)
-		{
-			Container = NewObject<UCustomizableObjectResourceDataContainer>(
-				&InCustomizableObject,
-				FName(*ContainerName),
-				RF_Public);
-
-			Container->Data = ResourceData;
-
-			if (FCustomizableObjectAssetUserData* AUDResource = Container->Data.Data.GetMutablePtr<FCustomizableObjectAssetUserData>())
-			{
-				// Find or duplicate the AUD replacing the outer
-				UAssetUserData* SourceAssetUserData = UE::Mutable::Private::LoadObject(AUDResource->AssetUserDataEditor);
-				const FString AssetName = GetNameSafe(SourceAssetUserData);
-				check(SourceAssetUserData != nullptr);
-
-				UAssetUserData* AssetUserData = FindObject<UAssetUserData>(Container, *AssetName);
-				if (!AssetUserData)
-				{
-					// AUD may be private objects within meshes. Duplicate changing the outer to avoid including meshes into the builds.
-					AssetUserData = DuplicateObject<UAssetUserData>(SourceAssetUserData, Container, FName(*AssetName));
-				}
-
-				AUDResource->AssetUserData = AssetUserData;
-			}
-		}
-
-		StreamedResourceData.Emplace(Container);
-	}
-
-	const TArrayView<const UCustomizableObjectExtension* const> AllExtensions = ICustomizableObjectModule::Get().GetRegisteredExtensions();
-
-	const int32 NumStreamedExtensionResources = StreamedExtensionDataEditor.Num();
-	for (int32 Index = 0; Index < NumStreamedExtensionResources; ++Index)
-	{
-		const FString& ContainerName = ObjectName + FString::Printf(TEXT("_SE_%d"), Index);
-		FCustomizableObjectResourceData& ResourceData = StreamedExtensionDataEditor[Index];
-
-		UCustomizableObjectResourceDataContainer* Container = FindObject<UCustomizableObjectResourceDataContainer>(&InCustomizableObject, *ContainerName);
-		if (!Container)
-		{
-			Container = NewObject<UCustomizableObjectResourceDataContainer>(
-				&InCustomizableObject,
-				FName(*ContainerName),
-				RF_Public);
-
-			for (const UCustomizableObjectExtension* Extension : AllExtensions)
-			{
-				Extension->MovePrivateReferencesToContainer(ResourceData.Data, Container);
-			}
-
-			Container->Data = MoveTemp(ResourceData);
-		}
-
-		StreamedExtensionData.Emplace(Container);
-	}
-
-	for (FCustomizableObjectResourceData& ResourceData : AlwaysLoadedExtensionData)
-	{
-		for (const UCustomizableObjectExtension* Extension : AllExtensions)
-		{
-			Extension->MovePrivateReferencesToContainer(ResourceData.Data, &InCustomizableObject);
-		}
-	}
-	
-	for (TMap<uint32, TSoftObjectPtr<UObject>>::TIterator It = PassthroughObjects.CreateIterator(); It; ++It)
-	{
-		TSoftObjectPtr<UObject>& PassthroughObject = It->Value;
-		
-		UObject* Object = UE::Mutable::Private::LoadObject(PassthroughObject);
-		check(Object);
-		
-		const bool bDuplicate = Object->GetOuter() != &InCustomizableObject && !Object->HasAnyFlags(RF_Public);
-		if (bDuplicate)
-		{
-			UObject* DuplicatedObject = DuplicateObject(Object, &InCustomizableObject);
-
-			PassthroughObject = DuplicatedObject;
-			StolenPrivateObjects.Add(DuplicatedObject);
-
-			if (UClothingAssetCommon* ClothingAsset = Cast<UClothingAssetCommon>(DuplicatedObject))
-			{
-				for (TMap<FName, TObjectPtr<UClothConfigBase>>::TIterator ConfigIt = ClothingAsset->ClothConfigs.CreateIterator(); ConfigIt; ++ConfigIt)
-				{
-					TObjectPtr<UClothConfigBase>& ClothConfig = ConfigIt->Value;
-					
-					UClothConfigBase* DuplicatedClothConfig = DuplicateObject<UClothConfigBase>(ClothConfig, &InCustomizableObject);
-					
-					ClothConfig = DuplicatedClothConfig;
-					StolenPrivateObjects.Add(DuplicatedClothConfig);
-				}
-			}
-		}
-	}
-}
-
-
 void UCustomizableObjectPrivate::LoadCompiledDataFromDisk()
 {
 	MUTABLE_CPUPROFILER_SCOPE(UCustomizableObjectPrivate::LoadCompiledDataFromDisk)
@@ -1157,14 +1048,6 @@ const TArray<TSoftObjectPtr<UCustomizableObject>>& UCustomizableObjectPrivate::G
 }
 
 
-bool UCustomizableObjectPrivate::IsAssetUserDataMergeEnabled() const
-{
-	PRAGMA_DISABLE_DEPRECATION_WARNINGS // TODO Remove >5.6
-	return GetPublic()->bEnableAssetUserDataMerge;
-	PRAGMA_ENABLE_DEPRECATION_WARNINGS
-}
-
-
 bool UCustomizableObjectPrivate::IsTableMaterialsParentCheckDisabled() const
 {
 	PRAGMA_DISABLE_DEPRECATION_WARNINGS // TODO Remove >5.6
@@ -3303,13 +3186,6 @@ void UCustomizableObjectPrivate::LogMemory()
 				Arch << ModelResources->ComponentNamesPerObjectComponent;
 				Arch << ModelResources->MinLODs;
 				Arch << ModelResources->MinQualityLevelLODs;
-				
-				ModeleResData += ModelResources->StreamedResourceData.GetAllocatedSize();
-				for (FCustomizableObjectStreamedResourceData& Data : ModelResources->StreamedResourceData)
-				{
-					TSoftObjectPtr<UCustomizableObjectResourceDataContainer> PathCopy = Data.GetPath();
-					Arch << PathCopy;
-				}
 			}
 			ModeleResData += Arch.TotalSize();
 		}
@@ -3500,7 +3376,6 @@ FArchive& operator<<(FArchive& Ar, FMutableRefSkeletalMeshData& Data)
 	Ar << Data.SkeletalMeshLODSettings;
 	Ar << Data.Skeleton;
 	Ar << Data.ShadowPhysicsAsset;
-	Ar << Data.AssetUserDataIndices;
 
 	return Ar;
 }
diff --git a/Engine/Plugins/Mutable/Source/CustomizableObject/Private/MuCO/CustomizableObjectInstance.cpp b/Engine/Plugins/Mutable/Source/CustomizableObject/Private/MuCO/CustomizableObjectInstance.cpp
--- a/Engine/Plugins/Mutable/Source/CustomizableObject/Private/MuCO/CustomizableObjectInstance.cpp
+++ b/Engine/Plugins/Mutable/Source/CustomizableObject/Private/MuCO/CustomizableObjectInstance.cpp
@@ -44,8 +44,6 @@
 #include "PhysicsEngine/PhysicsConstraintTemplate.h"
 #include "PhysicsEngine/SkeletalBodySetup.h"
 #include "MuCO/CustomizableObjectCustomVersion.h"
-#include "MuCO/CustomizableObjectResourceData.h"
-#include "MuCO/CustomizableObjectStreamedResourceData.h"
 #include "MuCO/CustomizableObjectResourceDataTypes.h"
 #include "MuCO/LoadUtils.h"
 #include "Serialization/BulkData.h"
@@ -110,43 +108,6 @@ UTexture2D* CreateTexture(const FName& TextureName)
 }
 
 
-void AddComponentAnimBPData(FCustomizableInstanceComponentData& ComponentData, UCustomizableInstancePrivate* InstancePrivate, TArray<FSoftObjectPath>& AssetsToStream, const FCustomizableObjectAnimationSlot& AnimSlot, const TSharedRef<FUpdateContextPrivate>& Context)
-{
-	check(InstancePrivate);
-
-	if (!AnimSlot.AnimInstance.IsNull())
-	{
-		const TSoftClassPtr<UAnimInstance>* FoundAnimBpSlot = ComponentData.AnimSlotToBP.Find(AnimSlot.Name);
-		bool bIsSameAnimBp = FoundAnimBpSlot && AnimSlot.AnimInstance == *FoundAnimBpSlot;
-		if (!FoundAnimBpSlot)
-		{
-			ComponentData.AnimSlotToBP.Add(AnimSlot.Name, AnimSlot.AnimInstance);
-
-			if (AnimSlot.AnimInstance.Get())
-			{
-				InstancePrivate->GatheredAnimBPs.Add(AnimSlot.AnimInstance.Get());
-			}
-			else
-			{
-				AssetsToStream.AddUnique(AnimSlot.AnimInstance.ToSoftObjectPath());
-			}
-		}
-		else if (!bIsSameAnimBp)
-		{
-			// Two submeshes should not have the same animation slot index
-			Context->UpdateResult = EUpdateResult::Warning;
-
-			FString WarningMessage = FString::Printf(TEXT("Two submeshes have the same anim slot index [%s] in a Mutable Instance."), *AnimSlot.Name.ToString());
-			UE_LOG(LogMutable, Warning, TEXT("%s"), *WarningMessage);
-#if WITH_EDITOR
-			FMessageLog MessageLog("Mutable");
-			MessageLog.Notify(FText::FromString(WarningMessage), EMessageSeverity::Warning, true);
-#endif
-		}
-	}
-}
-
-
 void UCustomizableInstancePrivate::InvalidateGeneratedData()
 {
 	SkeletalMeshStatus = ESkeletalMeshStatus::NotGenerated;
@@ -162,19 +123,6 @@ void UCustomizableInstancePrivate::InvalidateGeneratedData()
 }
 
 
-void UCustomizableInstancePrivate::InitCustomizableObjectData(const UCustomizableObject* InCustomizableObject)
-{
-	InvalidateGeneratedData();
-
-	if (!InCustomizableObject || !InCustomizableObject->IsCompiled())
-	{
-		return;
-	}
-	
-	ExtensionInstanceData.Empty();
-}
-
-
 FCustomizableInstanceComponentData* UCustomizableInstancePrivate::GetComponentData(const FName& ComponentName)
 {
 	return ComponentsData.FindByPredicate([&](const FCustomizableInstanceComponentData& ComponentData)
@@ -216,7 +164,7 @@ void UCustomizableInstancePrivate::SetDescriptor(const FCustomizableObjectInstan
 
 	if (bCustomizableObjectChanged)
 	{
-		InitCustomizableObjectData(InCustomizableObject);
+		InvalidateGeneratedData();
 	}
 }
 
@@ -249,7 +197,7 @@ void UCustomizableInstancePrivate::PostDuplicate(bool bDuplicateForPIE)
 void UCustomizableInstancePrivate::OnPostCompile()
 {
 	GetDescriptor().ReloadParameters();
-	InitCustomizableObjectData(GetPublic()->GetCustomizableObject());
+	InvalidateGeneratedData();
 }
 
 
@@ -451,7 +399,7 @@ void UCustomizableObjectInstance::PostLoad()
 	// Skip the cost of ReloadParameters in the cook commandlet; it will be reloaded during PreSave. For cooked runtime
 	// and editor UI, reload on load because it will not otherwise reload unless the CustomizableObject recompiles.
 	Descriptor.ReloadParameters();
-	PrivateData->InitCustomizableObjectData(GetCustomizableObject());
+	GetPrivate()->InvalidateGeneratedData();
 }
 
 
@@ -499,7 +447,7 @@ void UCustomizableObjectInstance::SetObject(UCustomizableObject* InObject)
 #endif
 
 	Descriptor.SetCustomizableObject(InObject);
-	PrivateData->InitCustomizableObjectData(InObject);
+	GetPrivate()->InvalidateGeneratedData();
 }
 
 
@@ -1906,122 +1854,6 @@ bool UCustomizableInstancePrivate::UpdateSkeletalMesh_PostBeginUpdate0(UCustomiz
 
 	const UModelResources& ModelResources = CustomizableObject->GetPrivate()->GetModelResourcesChecked();
 
-	// Collate the Extension Data on the instance into groups based on the extension that produced
-	// it, so that we only need to call extension functions such as OnSkeletalMeshCreated once for
-	// each extension.
-	TMap<const UCustomizableObjectExtension*, TArray<FInputPinDataContainer>> ExtensionToExtensionData;
-	{
-		const TArrayView<const UCustomizableObjectExtension* const> AllExtensions = ICustomizableObjectModule::Get().GetRegisteredExtensions();
-
-		// Pre-populate ExtensionToExtensionData with empty entries for all extensions.
-		//
-		// This ensures that extension functions such as OnSkeletalMeshCreated are called for each
-		// extension, even if they didn't produce any extension data.
-		Instance->GetPrivate()->ExtensionInstanceData.Empty(AllExtensions.Num());
-		for (const UCustomizableObjectExtension* Extension : AllExtensions)
-		{
-			ExtensionToExtensionData.Add(Extension);
-		}
-
-		const TArrayView<const FRegisteredObjectNodeInputPin> ExtensionPins = ICustomizableObjectModule::Get().GetAdditionalObjectNodePins();
-
-		for (FInstanceUpdateData::FNamedExtensionData& ExtensionOutput : Context->InstanceUpdateData.ExtendedInputPins)
-		{
-			const FRegisteredObjectNodeInputPin* FoundPin =
-				Algo::FindBy(ExtensionPins, ExtensionOutput.Name, &FRegisteredObjectNodeInputPin::GlobalPinName);
-
-			if (!FoundPin)
-			{
-				// Failed to find the corresponding pin for this output
-				// 
-				// This may indicate that a plugin has been removed or renamed since the CO was compiled
-				UE_LOG(LogMutable, Error, TEXT("Failed to find Object node input pin with name %s"), *ExtensionOutput.Name.ToString());
-				continue;
-			}
-
-			const UCustomizableObjectExtension* Extension = FoundPin->Extension.Get();
-			if (!Extension)
-			{
-				// Extension is not loaded or not found
-				UE_LOG(LogMutable, Error, TEXT("Extension for Object node input pin %s is no longer valid"), *ExtensionOutput.Name.ToString());
-				continue;
-			}
-
-			if (ExtensionOutput.Data->Origin == UE::Mutable::Private::FExtensionData::EOrigin::Invalid)
-			{
-				// Null data was produced
-				//
-				// This can happen if a node produces an FExtensionData but doesn't initialize it
-				UE_LOG(LogMutable, Error, TEXT("Invalid data sent to Object node input pin %s"), *ExtensionOutput.Name.ToString());
-				continue;
-			}
-
-			// All registered extensions were added to the map above, so if the extension is still
-			// registered it should be found.
-			TArray<FInputPinDataContainer>* ContainerArray = ExtensionToExtensionData.Find(Extension);
-			if (!ContainerArray)
-			{
-				UE_LOG(LogMutable, Error, TEXT("Object node input pin %s received data for unregistered extension %s"),
-					*ExtensionOutput.Name.ToString(), *Extension->GetPathName());
-				continue;
-			}
-
-			const FCustomizableObjectResourceData* ReferencedExtensionData = nullptr;
-			switch (ExtensionOutput.Data->Origin)
-			{
-				case UE::Mutable::Private::FExtensionData::EOrigin::ConstantAlwaysLoaded:
-				{
-					check(ModelResources.AlwaysLoadedExtensionData.IsValidIndex(ExtensionOutput.Data->Index));
-					ReferencedExtensionData = &ModelResources.AlwaysLoadedExtensionData[ExtensionOutput.Data->Index];
-				}
-				break;
-
-				case UE::Mutable::Private::FExtensionData::EOrigin::ConstantStreamed:
-				{
-#if WITH_EDITOR
-					check(ModelResources.StreamedExtensionDataEditor.IsValidIndex(ExtensionOutput.Data->Index));
-					ReferencedExtensionData = &ModelResources.StreamedExtensionDataEditor[ExtensionOutput.Data->Index];
-#else
-					check(ModelResources.StreamedExtensionData.IsValidIndex(ExtensionOutput.Data->Index));
-					const FCustomizableObjectStreamedResourceData& StreamedData = ModelResources.StreamedExtensionData[ExtensionOutput.Data->Index];
-					if (!StreamedData.IsLoaded())
-					{
-						// The data should have been loaded as part of executing the CO program.
-						//
-						// This could indicate a bug in the streaming logic.
-						UE_LOG(LogMutable, Error, TEXT("Customizable Object produced a streamed extension data that is not loaded: %s"),
-							*StreamedData.GetPath().ToString());
-
-						continue;
-					}
-
-					ReferencedExtensionData = &StreamedData.GetLoadedData();
-#endif
-				}
-				break;
-
-				default:
-					unimplemented();
-			}
-
-			check(ReferencedExtensionData);
-
-			ContainerArray->Emplace(FoundPin->InputPin, ReferencedExtensionData->Data);
-		}
-	}
-
-	// Give each extension the chance to generate Extension Instance Data
-	for (const TPair<const UCustomizableObjectExtension*, TArray<FInputPinDataContainer>>& Pair : ExtensionToExtensionData)
-	{
-		FInstancedStruct NewExtensionInstanceData = Pair.Key->GenerateExtensionInstanceData(Pair.Value);
-		if (NewExtensionInstanceData.IsValid())
-		{
-			FExtensionInstanceData& NewData = Instance->GetPrivate()->ExtensionInstanceData.AddDefaulted_GetRef();
-			NewData.Extension = Pair.Key;
-			NewData.Data = MoveTemp(NewExtensionInstanceData);
-		}
-	}
-
 	// None of the current meshes requires a mesh update. Continue to BuildResources
 	if (!bUpdateMeshes)
 	{
@@ -2150,7 +1982,7 @@ bool UCustomizableInstancePrivate::UpdateSkeletalMesh_PostBeginUpdate0(UCustomiz
 		const FMutableRefSkeletalMeshData& RefSkeletalMeshData = ModelResources.ReferenceSkeletalMeshesData[Component->Id];
 
 		// Set up the default information any mesh from this component will have (LODArrayInfos, RenderData, Mesh settings, etc). 
-		InitSkeletalMeshData(Context, SkeletalMesh, RefSkeletalMeshData, *CustomizableObject, Component->Id);
+		InitSkeletalMeshData(Context, SkeletalMesh, RefSkeletalMeshData, *CustomizableObject, Component);
 		
 		// Construct a new skeleton, fix up ActiveBones and Bonemap arrays and recompute the RefInvMatrices
 		const bool bBuildSkeletonDataSuccess = BuildSkeletonData(Context, *SkeletalMesh, ComponentIndex);
@@ -2221,9 +2053,9 @@ bool UCustomizableInstancePrivate::UpdateSkeletalMesh_PostBeginUpdate0(UCustomiz
 		// Add sockets from the SkeletalMesh of reference and from the MutableMesh
 		BuildMeshSockets(SkeletalMesh, ModelResources, RefSkeletalMeshData, ComponentMesh);
 		
-		for (const TPair<const UCustomizableObjectExtension*, TArray<FInputPinDataContainer>>& Pair : ExtensionToExtensionData)
+		for (const UCustomizableObjectExtension* Extension : ICustomizableObjectModule::Get().GetRegisteredExtensions())
 		{
-			Pair.Key->OnSkeletalMeshCreated(Pair.Value, ComponentName, SkeletalMesh);
+			Extension->OnSkeletalMeshCreated(ComponentName, SkeletalMesh);
 		}
 		
 		BuildOrCopyElementData(Context, SkeletalMesh, ComponentIndex);
@@ -3510,7 +3342,7 @@ FAutoConsoleVariableRef CVarMutableEnableRayTracingFix(
 	TEXT("If 2, Enable fix for all generated meshes. Meshes will have ray tracing disabled.")
 	);
 
-void InitSkeletalMeshData(const TSharedRef<FUpdateContextPrivate>& Context, USkeletalMesh* SkeletalMesh, const FMutableRefSkeletalMeshData& RefSkeletalMeshData, const UCustomizableObject& CustomizableObject, UE::Mutable::Private::FComponentId ComponentId)
+void InitSkeletalMeshData(const TSharedRef<FUpdateContextPrivate>& Context, USkeletalMesh* SkeletalMesh, const FMutableRefSkeletalMeshData& RefSkeletalMeshData, const UCustomizableObject& CustomizableObject, const TSharedPtr<FInstanceUpdateData::FComponent>& Component)
 {
 	MUTABLE_CPUPROFILER_SCOPE(UCustomizableInstancePrivate::InitSkeletalMesh);
 
@@ -3519,7 +3351,7 @@ void InitSkeletalMeshData(const TSharedRef<FUpdateContextPrivate>& Context, USke
 	UCustomizableObjectInstance* Instance = Context->Instance.Get();
 	UCustomizableInstancePrivate* InstancePrivate = Instance->GetPrivate();
 	
-	const FName ComponentName = Context->ComponentNames[ComponentId];
+	const FName ComponentName = Context->ComponentNames[Component->Id];
 
 	SkeletalMesh->NeverStream = !Context->bStreamMeshLODs;
 
@@ -3538,35 +3370,101 @@ void InitSkeletalMeshData(const TSharedRef<FUpdateContextPrivate>& Context, USke
 	// Set the default Physics Assets
 	SkeletalMesh->SetPhysicsAsset(RefSkeletalMeshData.PhysicsAsset.Get());
 	SkeletalMesh->SetEnablePerPolyCollision(RefSkeletalMeshData.Settings.bEnablePerPolyCollision);
-
+	
+	const int32 FirstLOD = Context->bStreamMeshLODs ?
+			Context->FirstLODAvailable[ComponentName] :
+			Context->GetFirstRequestedLOD()[ComponentName];
+	
 	// Asset User Data
 	{
-		const FCustomizableInstanceComponentData& ComponentData = InstancePrivate->GetComponentData(ComponentId);
-		for (TObjectPtr<UAssetUserData> AssetUserData : ComponentData.AssetUserDataArray)
+		TMap<const UClass*, TArray<UAssetUserData*>> AssetUserDataByType; 
+		
+		// Hack. Get the asset user data from the first LOD.
+		for (int32 LODIndex = FirstLOD; LODIndex < Component->LODCount; ++LODIndex)
 		{
-			SkeletalMesh->AddAssetUserData(AssetUserData);
+			const TSharedRef<FInstanceUpdateData::FLOD>& LOD = Context->InstanceUpdateData.LODs[Component->FirstLOD + LODIndex];
+
+			if (LOD->Mesh)
+			{
+				for (const UE::Mutable::Private::TPassthroughObjectPtr<UAssetUserData>& PassthroughAssetUserData : LOD->Mesh->AssetUserData)
+				{
+					if (UAssetUserData* AssetUserData = PassthroughAssetUserData.Get())
+					{
+						TArray<UAssetUserData*>& Result = AssetUserDataByType.FindOrAdd(AssetUserData->GetClass(), {});
+						Result.AddUnique(AssetUserData);
+					}
+				}
+
+				break;
+			}
 		}
+	
+		for (const TTuple<const UClass*, TArray<UAssetUserData*>>& Pair : AssetUserDataByType)
+		{
+			if (Pair.Value.Num() == 1)
+			{
+				SkeletalMesh->AddAssetUserData(Pair.Value.Top());
+			}
+			else if (Pair.Value.Num() > 1)
+			{
+				UAssetUserData* NewAssetUserData = DuplicateObject(Pair.Value[0], SkeletalMesh);
 
-		//Custom Asset User Data
-		if (Context->Instance->GetAnimationGameplayTags().Num() ||
-			ComponentData.AnimSlotToBP.Num())
+				bool bSuccess = true;
+				for (int32 Index = 1; Index < Pair.Value.Num(); ++Index)
+				{
+					bSuccess = NewAssetUserData->Merge(Pair.Value[Index]);
+					if (!bSuccess)
+					{
+						UE_LOGF(LogMutable, Error, "Unable to merge Asset User Data %ls", *Pair.Key->GetName());
+						break;
+					}
+				}
+				
+				if (bSuccess)
+				{
+					SkeletalMesh->AddAssetUserData(NewAssetUserData);
+				}
+			}
+		}
+	}
+	
+	// Anim Slots
+	{
+		// Hack. Get the Anim Slot data from the first LOD.
+		for (int32 LODIndex = FirstLOD; LODIndex < Component->LODCount; ++LODIndex)
 		{
-			UCustomizableObjectInstanceUserData* InstanceData = NewObject<UCustomizableObjectInstanceUserData>(SkeletalMesh, NAME_None, RF_Public | RF_Transactional);
-			InstanceData->AnimationGameplayTag = Context->Instance->GetAnimationGameplayTags();
+			const TSharedRef<FInstanceUpdateData::FLOD>& LOD = Context->InstanceUpdateData.LODs[Component->FirstLOD + LODIndex];
+
+			if (!LOD->Mesh)
+			{
+				continue;
+			}
 			
-			for (const TTuple<FName, TSoftClassPtr<UAnimInstance>>& AnimSlot : ComponentData.AnimSlotToBP)
+			if (LOD->Mesh->GameplayTags.Num() || LOD->Mesh->AnimationSlots.Num())
 			{
-				FCustomizableObjectAnimationSlot AnimationSlot;
-				AnimationSlot.Name = AnimSlot.Key;
-				AnimationSlot.AnimInstance = AnimSlot.Value;
+				UCustomizableObjectInstanceUserData* UserData = NewObject<UCustomizableObjectInstanceUserData>(SkeletalMesh, NAME_None, RF_Public | RF_Transactional);
+					
+				for (const FName& GameplayTag : LOD->Mesh->GameplayTags)
+				{
+					UserData->AnimationGameplayTag.AddTag(FGameplayTag::RequestGameplayTag(GameplayTag, false));
+				}
+
+				for (const TPair<FName, UE::Mutable::Private::TPassthroughObjectPtr<UClass>>& Pair : LOD->Mesh->AnimationSlots)
+				{
+					FCustomizableObjectAnimationSlot Slot;
+					Slot.Name = Pair.Key;
+					Slot.AnimInstance = Pair.Value.Get();
+
+					UserData->AnimationSlots.Add(Slot);
+				}	
 				
-				InstanceData->AnimationSlots.Add(AnimationSlot);
+				SkeletalMesh->AddAssetUserData(UserData);
 			}
-			
-			SkeletalMesh->AddAssetUserData(InstanceData);
+
+			break;
 		}
 	}
-
+	
 	// Allocate resources for rendering and add LOD Info
 	{
 		MUTABLE_CPUPROFILER_SCOPE(InitSkeletalMesh_AddLODData);
@@ -3650,8 +3548,8 @@ void InitSkeletalMeshData(const TSharedRef<FUpdateContextPrivate>& Context, USke
 
 	// Set Min LOD (Override the Reference Skeletal Mesh LOD Settings)
 	const UModelResources& ModelResources = CustomizableObject.GetPrivate()->GetModelResourcesChecked();
-	SkeletalMesh->SetMinLod(FMath::Max(ModelResources.MinLODs.FindChecked(ComponentId).GetDefault(), static_cast<int32>(Context->FirstLODAvailable[ComponentName])));
-	SkeletalMesh->SetQualityLevelMinLod(ModelResources.MinQualityLevelLODs.FindChecked(ComponentId));
+	SkeletalMesh->SetMinLod(FMath::Max(ModelResources.MinLODs.FindChecked(Component->Id).GetDefault(), static_cast<int32>(Context->FirstLODAvailable[ComponentName])));
+	SkeletalMesh->SetQualityLevelMinLod(ModelResources.MinQualityLevelLODs.FindChecked(Component->Id));
 
 	// Set up unreal's default material, will be replaced when building materials
 	{
@@ -4440,217 +4338,9 @@ UE::Tasks::FTask LoadAdditionalAssetsAndData(const TSharedRef<FUpdateContextPriv
 
 	TArray<FSoftObjectPath> AssetsToStream;
 
-	
-	InstancePrivate->GatheredAnimBPs.Empty();
-	InstancePrivate->AnimBPGameplayTags.Reset();
 	InstancePrivate->AnimBpPhysicsAssets.Reset();
 
-	bool bHasInvalidMesh = false;
-	const bool bUpdateMeshes = DoComponentsNeedUpdate(Context, bHasInvalidMesh);
-
-	// Load assets coming from SubMeshes of the newly generated Mesh
-	TArray<TSharedRef<FInstanceUpdateData::FComponent>>& Components = Context->InstanceUpdateData.Components;
-	for (int32 ComponentIndex = 0; ComponentIndex < Components.Num(); ++ComponentIndex)
-	{
-		const TSharedRef<const FInstanceUpdateData::FComponent>& Component = Components[ComponentIndex];
-		FCustomizableInstanceComponentData& ComponentData = InstancePrivate->GetComponentData(Component->Id);
-		
-		if (Component->SkeletalMesh->PassthroughObject)
-		{
-			// Pass-through components don't have a Reference Mesh so don't access it
-			continue;
-		}
-		
-		const FMutableRefSkeletalMeshData& RefSkeletalMeshData = ModelResources.ReferenceSkeletalMeshesData[Component->Id];
-
-		for (int32 AssetUserDataIndex : RefSkeletalMeshData.AssetUserDataIndices)
-		{
-#if !WITH_EDITOR
-			Context->StreamedResourceIndex.AddUnique(AssetUserDataIndex); // Used to hold/release streamed resources in non-editor builds.
-#endif
-			ComponentData.StreamedResourceIndex.AddUnique(AssetUserDataIndex);
-		}
-
-		const FName ComponentName = Context->ComponentNames[Component->Id];
-		
-		for (int32 LODIndex = Context->FirstLODAvailable[ComponentName]; LODIndex < Component->LODCount; ++LODIndex)
-		{
-			const TSharedRef<FInstanceUpdateData::FLOD>& LOD = Context->InstanceUpdateData.LODs[Component->FirstLOD + LODIndex];
-
-			UE::Mutable::Private::TManagedPtr<const UE::Mutable::Private::FMesh> MutableMesh = LOD->Mesh;
-			if (!MutableMesh)
-			{
-				continue;
-			}
-			
-			const TArray<uint64>& StreamedResources = MutableMesh->GetStreamedResources();
-
-			for (uint64 ResourceId : StreamedResources)
-			{
-				FCustomizableObjectStreameableResourceId TypedResourceId = BitCast<FCustomizableObjectStreameableResourceId>(ResourceId);
-
-				if (TypedResourceId.Type == (uint8)FCustomizableObjectStreameableResourceId::EType::AssetUserData)
-				{
-					const uint32 ResourceIndex = TypedResourceId.Id;
-#if !WITH_EDITOR
-					Context->StreamedResourceIndex.AddUnique(ResourceIndex); // Used to hold/release streamed resources in non-editor builds.
-#endif
-					ComponentData.StreamedResourceIndex.AddUnique(ResourceIndex);
-				}
-			}
-
-			// Add AssetUserData from Mesh Parameters
-			for (const TStrongObjectPtr<USkeletalMesh>& StrongMesh : MutableMesh->SkeletalMeshes)
-			{
-				USkeletalMesh* SkeletalMesh = StrongMesh.Get();
-
-				if (!SkeletalMesh)
-				{
-					continue;
-				}
-
-				if (const TArray<UAssetUserData*>* AssetUserDataArray = SkeletalMesh->GetAssetUserDataArray())
-				{
-					for (UAssetUserData* AssetUserData : *AssetUserDataArray)
-					{
-						if (!AssetUserData)
-						{
-							continue;
-						}
-
-						if (const UCustomizableObjectInstanceUserData* COIAssetUserData = Cast<UCustomizableObjectInstanceUserData>(AssetUserData))
-						{
-							for (const FCustomizableObjectAnimationSlot& AnimSlot : COIAssetUserData->AnimationSlots)
-							{
-								AddComponentAnimBPData(ComponentData, InstancePrivate, AssetsToStream, AnimSlot, Context);
-							}
-
-							InstancePrivate->AnimBPGameplayTags.AppendTags(COIAssetUserData->AnimationGameplayTag);
-						}
-						else
-						{
-							ComponentData.AssetUserDataArray.Add(AssetUserData);
-						}
-					}
-				}
-			}
-			
-			for (int32 TagIndex = 0; TagIndex < MutableMesh->GetTagCount(); ++TagIndex)
-			{
-				FString Tag = MutableMesh->GetTag(TagIndex);
-
-				if (Tag.RemoveFromStart("__AnimBP:"))
-				{
-					FString SlotIndexString, AnimBpIndexString;
-
-					if (Tag.Split(TEXT("_Slot_"), &SlotIndexString, &AnimBpIndexString))
-					{
-						if (SlotIndexString.IsEmpty() || AnimBpIndexString.IsEmpty())
-						{
-							continue;
-						}
-
-						const int32 AnimBpIndex = FCString::Atoi(*AnimBpIndexString);
-						if (!ModelResources.AnimBPs.IsValidIndex(AnimBpIndex))
-						{
-							continue;
-						}
-
-						AddComponentAnimBPData(ComponentData, InstancePrivate, AssetsToStream, FCustomizableObjectAnimationSlot{ *SlotIndexString, ModelResources.AnimBPs[AnimBpIndex] }, Context);
-					}
-				}
-				else if (Tag.RemoveFromStart("__AnimBPTag:"))
-				{
-					InstancePrivate->AnimBPGameplayTags.AddTag(FGameplayTag::RequestGameplayTag(*Tag));
-				}
-#if WITH_EDITORONLY_DATA
-				else if (Tag.RemoveFromStart("__MeshPath:"))
-				{
-					ComponentData.MeshPartPaths.Add(Tag);
-				}
-#endif
-			}
-		}
-
-
-		for (int32 ResourceIndex : ComponentData.StreamedResourceIndex)
-		{
-#if WITH_EDITOR
-			if (!ModelResources.StreamedResourceDataEditor.IsValidIndex(ResourceIndex))
-			{
-				UE_LOG(LogMutable, Error, TEXT("Invalid streamed resource index. Max Index [%d]. Resource Index [%d]."), ModelResources.StreamedResourceDataEditor.Num(), ResourceIndex);
-				continue;
-			}
-
-			if (const FCustomizableObjectAssetUserData* AUDResource = ModelResources.StreamedResourceDataEditor[ResourceIndex].Data.GetPtr<FCustomizableObjectAssetUserData>())
-			{
-				AssetsToStream.AddUnique(AUDResource->AssetUserDataEditor.ToSoftObjectPath());
-			}
-#else
-			if (!ModelResources.StreamedResourceData.IsValidIndex(ResourceIndex))
-			{
-				UE_LOG(LogMutable, Error, TEXT("Invalid streamed resource index. Max Index [%d]. Resource Index [%d]."), ModelResources.StreamedResourceData.Num(), ResourceIndex);
-				continue;
-			}
-
-			const FCustomizableObjectStreamedResourceData& StreamedResource = ModelResources.StreamedResourceData[ResourceIndex];
-			if (!StreamedResource.IsLoaded())
-			{
-				AssetsToStream.AddUnique(StreamedResource.GetPath().ToSoftObjectPath());
-			}
-#endif
-		}
-	}
-	
-
 	Context->PassthroughObjectLoader->Load(ModelResources.PassthroughObjects, AssetsToStream);
-	
-	// Copy FExtensionData Object node input from the Instance to the InstanceUpdateData
-	for (int32 ExtensionDataIndex = 0; ExtensionDataIndex < Context->MutableInstance->GetExtensionDataCount(); ExtensionDataIndex++)
-	{
-		UE::Mutable::Private::TManagedPtr<const UE::Mutable::Private::FExtensionData> ExtensionData;
-		FName Name;
-		Context->MutableInstance->GetExtensionData(ExtensionDataIndex, ExtensionData, Name);
-
-		check(ExtensionData);
-
-		FInstanceUpdateData::FNamedExtensionData& NewEntry = Context->InstanceUpdateData.ExtendedInputPins.AddDefaulted_GetRef();
-		NewEntry.Data = ExtensionData;
-		NewEntry.Name = Name;
-		check(NewEntry.Name != NAME_None);
-
-#if WITH_EDITOR
-		if (!ModelResources.StreamedExtensionDataEditor.IsValidIndex(ExtensionData->Index))
-		{
-			// The compiled data appears to be out of sync with the CO's properties
-
-			UE_LOG(LogMutable, Error, TEXT("Couldn't find streamed Extension Data with index %d in %s. Compiled data may be stale."),
-				ExtensionData->Index, *CustomizableObject->GetFullName());
-		}
-#else
-		if (!ModelResources.StreamedExtensionData.IsValidIndex(ExtensionData->Index))
-		{
-			// The compiled data appears to be out of sync with the CO's properties
-
-			UE_LOG(LogMutable, Error, TEXT("Couldn't find streamed Extension Data with index %d in %s. Compiled data may be stale."),
-				ExtensionData->Index, *CustomizableObject->GetFullName());
-
-			continue;
-		}
-
-		const FCustomizableObjectStreamedResourceData& StreamedData = ModelResources.StreamedExtensionData[ExtensionData->Index];
-		if (StreamedData.IsLoaded())
-		{
-			continue;
-		}
-
-		// Note that this just checks if the path is non-null, NOT if the object is loaded
-		check(!StreamedData.GetPath().IsNull());
-
-		Context->ExtensionStreamedResourceIndex.Add(ExtensionData->Index);
-		AssetsToStream.Add(StreamedData.GetPath().ToSoftObjectPath());
-#endif
-	}
 
 	TArray<UE::Tasks::FTask, TInlineAllocator<2>> Prerequisites;
 	
@@ -4715,74 +4405,11 @@ void AdditionalAssetsAsyncLoaded(const TSharedRef<FUpdateContextPrivate> Context
 	check(IsInGameThread())
 
 	Event.Trigger();
-
-	UCustomizableObjectInstance* Instance = Context->Instance.Get();
-	UCustomizableInstancePrivate* InstancePrivate = Instance->GetPrivate();
 	
-	UCustomizableObjectPrivate* CustomizableObjectPrivate = Instance->GetCustomizableObject()->GetPrivate();
-
-	UModelResources& ModelResources = *CustomizableObjectPrivate->GetModelResources();
+	UModelResources& ModelResources = *Context->Object->GetPrivate()->GetModelResources();
 
 	Context->PassthroughObjectLoader->Resolve(ModelResources.PassthroughObjects);
-	
-	for (int32 ResourceIndex : Context->StreamedResourceIndex)
-	{
-		ModelResources.StreamedResourceData[ResourceIndex].Hold();
-	}
-	
-	for (int32 ResourceIndex : Context->ExtensionStreamedResourceIndex)
-	{
-		ModelResources.StreamedExtensionData[ResourceIndex].Hold();
-	}
 
-	for (FCustomizableInstanceComponentData& ComponentData : InstancePrivate->ComponentsData)
-	{
-		for (int32 ResourceIndex : ComponentData.StreamedResourceIndex)
-		{
-#if WITH_EDITOR
-			if (ModelResources.StreamedResourceDataEditor.IsValidIndex(ResourceIndex))
-			{
-				if (const FCustomizableObjectAssetUserData* AUDResource = ModelResources.StreamedResourceDataEditor[ResourceIndex].Data.GetPtr<FCustomizableObjectAssetUserData>())
-				{
-					ComponentData.AssetUserDataArray.Add(UE::Mutable::Private::LoadObject(AUDResource->AssetUserDataEditor)); // Already loaded
-				}
-			}
-#else
-			if (ModelResources.StreamedResourceData.IsValidIndex(ResourceIndex) && ModelResources.StreamedResourceData[ResourceIndex].IsLoaded())
-			{
-				const FCustomizableObjectResourceData& ResourceData = ModelResources.StreamedResourceData[ResourceIndex].GetLoadedData();
-
-				if (const FCustomizableObjectAssetUserData* AUDResource = ResourceData.Data.GetPtr<FCustomizableObjectAssetUserData>())
-				{
-					ComponentData.AssetUserDataArray.Add(AUDResource->AssetUserData);
-				}
-			}
-#endif
-		}
-
-		// Loaded anim BPs
-		for (TPair<FName, TSoftClassPtr<UAnimInstance>>& SlotAnimBP : ComponentData.AnimSlotToBP)
-		{
-			if (TSubclassOf<UAnimInstance> AnimBP = SlotAnimBP.Value.Get())
-			{
-				if (!InstancePrivate->GatheredAnimBPs.Contains(AnimBP))
-				{
-					InstancePrivate->GatheredAnimBPs.Add(AnimBP);
-				}
-			}
-#if WITH_EDITOR
-			else
-			{
-				FString ErrorMsg = FString::Printf(TEXT("Mutable couldn't load the AnimBlueprint [%s]. If it has been deleted or renamed, please recompile all the mutable objects that use it."), *SlotAnimBP.Value.GetAssetName());
-				UE_LOG(LogMutable, Error, TEXT("%s"), *ErrorMsg);
-
-				FMessageLog MessageLog("Mutable");
-				MessageLog.Notify(FText::FromString(ErrorMsg), EMessageSeverity::Error, true);
-			}
-#endif
-		}
-	}
-	
 #if WITH_EDITOR
 	// TODO: Remove with UE-217665 when the underlying bug in the ColorPicker is solved
 	// Reenable the throttling which disabled when launching the Async Load
@@ -5849,30 +5476,45 @@ TArray<FName> UCustomizableObjectInstance::GetComponentNames() const
 
 TSubclassOf<UAnimInstance> UCustomizableObjectInstance::GetAnimBP(FName ComponentName, const FName& SlotName) const
 {
-	FCustomizableInstanceComponentData* ComponentData =	GetPrivate()->GetComponentData(ComponentName);
-	if (!ComponentData)
+	TObjectPtr<USkeletalMesh>* SkeletalMesh = GetPrivate()->SkeletalMeshes.Find(ComponentName);
+	if (!SkeletalMesh)
 	{
-		FString ErrorMsg = FString::Printf(TEXT("Tried to access an invalid component index [%s] in a Mutable Instance."), *ComponentName.ToString());
-		UE_LOG(LogMutable, Error, TEXT("%s"), *ErrorMsg);
-#if WITH_EDITOR
-		FMessageLogModule& MessageLogModule = FModuleManager::LoadModuleChecked<FMessageLogModule>("MessageLog");
-		MessageLogModule.RegisterLogListing(FName("Mutable"), FText::FromString(FString("Mutable")));
-		FMessageLog MessageLog("Mutable");
-
-		MessageLog.Notify(FText::FromString(ErrorMsg), EMessageSeverity::Error, true);
-#endif
-
 		return nullptr;
 	}
+	
+	UCustomizableObjectInstanceUserData* AssetUserData = Cast<UCustomizableObjectInstanceUserData>((*SkeletalMesh)->GetAssetUserDataOfClass(UCustomizableObjectInstanceUserData::StaticClass()));
+	if (!AssetUserData)
+	{
+		return nullptr;
+	}
+	
+	FCustomizableObjectAnimationSlot* Slot = AssetUserData->AnimationSlots.FindByPredicate([&](const FCustomizableObjectAnimationSlot& Element)
+	{
+		return Element.Name == SlotName;
+	});
 
-	TSoftClassPtr<UAnimInstance>* Result = ComponentData->AnimSlotToBP.Find(SlotName);
-
-	return Result ? Result->Get() : nullptr;
+	if (!Slot)
+	{
+		return nullptr;	
+	}
+	
+	return Slot->AnimInstance.Get();
 }
 
-const FGameplayTagContainer& UCustomizableObjectInstance::GetAnimationGameplayTags() const
+
+FGameplayTagContainer UCustomizableObjectInstance::GetAnimationGameplayTags() const
 {
-	return GetPrivate()->AnimBPGameplayTags;
+	FGameplayTagContainer Container;
+	
+	for (const TTuple<FName, TObjectPtr<USkeletalMesh>>& Pair : GetPrivate()->SkeletalMeshes)
+	{
+		if (UCustomizableObjectInstanceUserData* AssetUserData = Cast<UCustomizableObjectInstanceUserData>(Pair.Value->GetAssetUserDataOfClass(UCustomizableObjectInstanceUserData::StaticClass())))
+		{
+			Container.AppendTags(AssetUserData->AnimationGameplayTag);
+		}
+	}
+	
+	return Container;
 }
 
 namespace UE::Mutable::Private
@@ -5897,46 +5539,39 @@ namespace UE::Mutable::Private
 			bAnyErrors = true;
 		}
 
-		const FCustomizableInstanceComponentData* ComponentData = Private->GetComponentData(ComponentName);
-
-		if (!ComponentData)
+		if (bAnyErrors)
 		{
-			FString ErrorMsg = FString::Printf(TEXT("Tried to access an invalid component [%s] in a Mutable Instance."), *ComponentName.ToString());
-			UE_LOG(LogMutable, Error, TEXT("%s"), *ErrorMsg);
-#if WITH_EDITOR
-			FMessageLogModule& MessageLogModule = FModuleManager::LoadModuleChecked<FMessageLogModule>("MessageLog");
-			MessageLogModule.RegisterLogListing(FName("Mutable"), FText::FromString(FString("Mutable")));
-			FMessageLog MessageLog("Mutable");
-
-			MessageLog.Notify(FText::FromString(ErrorMsg), EMessageSeverity::Error, true);
-#endif
-
-			bAnyErrors = true;
+			return;
 		}
-
-		if (bAnyErrors)
+		
+		const TObjectPtr<USkeletalMesh>* SkeletalMesh = Private->SkeletalMeshes.Find(ComponentName);
+		if (!SkeletalMesh)
 		{
 			return;
 		}
-
-		for (const TPair<FName, TSoftClassPtr<UAnimInstance>>& MapElem : ComponentData->AnimSlotToBP)
+		
+		UCustomizableObjectInstanceUserData* AssetUserData = Cast<UCustomizableObjectInstanceUserData>((*SkeletalMesh)->GetAssetUserDataOfClass(UCustomizableObjectInstanceUserData::StaticClass()));
+		if (!AssetUserData)
+		{
+			return;
+		}
+		
+		for (const FCustomizableObjectAnimationSlot& Slot :  AssetUserData->AnimationSlots)
 		{
-			const FName& Index = MapElem.Key;
-			const TSoftClassPtr<UAnimInstance>& AnimBP = MapElem.Value;
+			const TSoftClassPtr<UAnimInstance>& AnimBP = Slot.AnimInstance;
 
 			// if this _can_ resolve to a real AnimBP
 			if (!AnimBP.IsNull())
 			{
 				// force a load right now - we don't know whether we would have loaded already - this could be called in editor
-				const TSubclassOf<UAnimInstance> LiveAnimBP = UE::Mutable::Private::LoadClass(AnimBP);
+				const TSubclassOf<UAnimInstance> LiveAnimBP = Private::LoadClass(AnimBP);
 				if (LiveAnimBP)
 				{
-					Delegate.Execute(Index, LiveAnimBP);
+					Delegate.Execute(Slot.Name, LiveAnimBP);
 				}
 			}
 		}
 	}
-
 }
 
 
@@ -6042,30 +5677,17 @@ const TArray<FAnimInstanceOverridePhysicsAsset>* UCustomizableInstancePrivate::G
 }
 
 
-FInstancedStruct UCustomizableObjectInstance::GetExtensionInstanceData(const UCustomizableObjectExtension* Extension) const
-{
-	const FExtensionInstanceData* FoundData = Algo::FindBy(PrivateData->ExtensionInstanceData, Extension, &FExtensionInstanceData::Extension);
-	if (FoundData)
-	{
-		return FoundData->Data;
-	}
-
-	// Data not found. Return an empty instance.
-	return FInstancedStruct();
-}
-
-
 TSet<UAssetUserData*> UCustomizableObjectInstance::GetMergedAssetUserData(const FName& ComponentName) const
 {
-	FCustomizableInstanceComponentData* ComponentData = GetPrivate()->GetComponentData(ComponentName);
-	if (!ComponentData)
+	TObjectPtr<USkeletalMesh>* SkeletalMesh = GetPrivate()->SkeletalMeshes.Find(ComponentName);
+	if (!SkeletalMesh)
 	{
 		return {};
 	}
 	
 	TSet<UAssetUserData*> Set;
 		
-	for (const TObjectPtr<UAssetUserData>& Elem : ComponentData->AssetUserDataArray)
+	for (UAssetUserData* Elem : *(*SkeletalMesh)->GetAssetUserDataArray())
 	{
 		// Have to convert to UAssetUserData* because BP functions don't support TObjectPtr
 		Set.Add(Elem);
diff --git a/Engine/Plugins/Mutable/Source/CustomizableObject/Private/MuCO/CustomizableObjectInstanceAssetUserData.cpp b/Engine/Plugins/Mutable/Source/CustomizableObject/Private/MuCO/CustomizableObjectInstanceAssetUserData.cpp
new file mode 100644
--- /dev/null
+++ b/Engine/Plugins/Mutable/Source/CustomizableObject/Private/MuCO/CustomizableObjectInstanceAssetUserData.cpp
@@ -0,0 +1,20 @@
+﻿// Copyright Epic Games, Inc. All Rights Reserved.
+
+#include "MuCO/CustomizableObjectInstanceAssetUserData.h"
+
+#if WITH_EDITOR
+#include "Logging/MessageLog.h"
+#endif
+
+
+const FGameplayTagContainer& UCustomizableObjectInstanceUserData::GetAnimationGameplayTags() const
+{
+	return AnimationGameplayTag;
+}
+
+
+void UCustomizableObjectInstanceUserData::SetAnimationGameplayTags(const FGameplayTagContainer& InstanceTags)
+{
+	AnimationGameplayTag = InstanceTags;
+}
+
diff --git a/Engine/Plugins/Mutable/Source/CustomizableObject/Private/MuCO/CustomizableObjectInstanceUsage.cpp b/Engine/Plugins/Mutable/Source/CustomizableObject/Private/MuCO/CustomizableObjectInstanceUsage.cpp
--- a/Engine/Plugins/Mutable/Source/CustomizableObject/Private/MuCO/CustomizableObjectInstanceUsage.cpp
+++ b/Engine/Plugins/Mutable/Source/CustomizableObject/Private/MuCO/CustomizableObjectInstanceUsage.cpp
@@ -26,10 +26,6 @@
 
 void UCustomizableObjectInstanceUsagePrivate::Callbacks()
 {
-	for (const UCustomizableObjectExtension* Extension : ICustomizableObjectModule::Get().GetRegisteredExtensions())
-	{
-		Extension->OnCustomizableObjectInstanceUsageUpdated(*GetPublic());
-	}
 	
 	if (UCustomizableSkeletalComponent* CustomizableSkeletalComponent = GetCustomizableSkeletalComponent())
 	{
diff --git a/Engine/Plugins/Mutable/Source/CustomizableObject/Private/MuCO/CustomizableObjectStreamedResourceData.cpp b/Engine/Plugins/Mutable/Source/CustomizableObject/Private/MuCO/CustomizableObjectStreamedResourceData.cpp
deleted file mode 100644
diff --git a/Engine/Plugins/Mutable/Source/CustomizableObject/Private/MuCO/CustomizableObjectSystem.cpp b/Engine/Plugins/Mutable/Source/CustomizableObject/Private/MuCO/CustomizableObjectSystem.cpp
--- a/Engine/Plugins/Mutable/Source/CustomizableObject/Private/MuCO/CustomizableObjectSystem.cpp
+++ b/Engine/Plugins/Mutable/Source/CustomizableObject/Private/MuCO/CustomizableObjectSystem.cpp
@@ -895,15 +895,6 @@ FString UCustomizableObjectSystem::GetDesc()
 }
 
 
-int32 UCustomizableObjectSystemPrivate::EnableMutableAnimInfoDebugging = 0;
-
-static FAutoConsoleVariableRef CVarEnableMutableAnimInfoDebugging(
-	TEXT("mutable.EnableMutableAnimInfoDebugging"), UCustomizableObjectSystemPrivate::EnableMutableAnimInfoDebugging,
-	TEXT("If set to 1 or greater print on screen the animation info of the pawn's Customizable Object Instance. Anim BPs, slots and tags will be displayed."
-	"If the root Customizable Object is recompiled after this command is run, the used skeletal meshes will also be displayed."),
-	ECVF_Default);
-
-
 UCustomizableObjectSystem* UCustomizableObjectSystemPrivate::GetPublic() const
 {
 	UCustomizableObjectSystem* Public = StaticCast<UCustomizableObjectSystem*>(GetOuter());
@@ -1047,6 +1038,16 @@ void FinishUpdateGlobal(const TSharedRef<FUpdateContextPrivate>& Context)
 	{
 		MUTABLE_CPUPROFILER_SCOPE(FinishUpdateGlobal_Callbacks);
 
+		TArray<const UObject*> ExtensionData;
+		ExtensionData.Reserve(Context->ExtensionData.Num());
+		for (UE::Mutable::Private::TManagedPtr<const UE::Mutable::Private::FExtensionData>& Data : Context->ExtensionData)
+		{
+			if (UObject* PassthroughObject = Data->PassthroughObject.Get())
+			{
+				ExtensionData.Add(PassthroughObject);
+			}
+		}
+		
 		// Call CustomizableObjectInstanceUsages updated callbacks.
 		if (Context->bOptimizedUpdate)
 		{
@@ -1064,6 +1065,12 @@ void FinishUpdateGlobal(const TSharedRef<FUpdateContextPrivate>& Context)
 					continue;
 				}
 #endif
+				
+				for (const UCustomizableObjectExtension* Extension : ICustomizableObjectModule::Get().GetRegisteredExtensions())
+				{
+					Extension->OnCustomizableObjectInstanceUsageUpdated(*InstanceUsage.Get(), ExtensionData);
+				}
+				
 				InstanceUsage->GetPrivate()->Callbacks();
 			}
 		}
@@ -1086,6 +1093,11 @@ void FinishUpdateGlobal(const TSharedRef<FUpdateContextPrivate>& Context)
 
 				if (InstanceUsage->GetCustomizableObjectInstance() == Instance)
 				{
+					for (const UCustomizableObjectExtension* Extension : ICustomizableObjectModule::Get().GetRegisteredExtensions())
+					{
+						Extension->OnCustomizableObjectInstanceUsageUpdated(*InstanceUsage, ExtensionData);
+					}
+					
 					InstanceUsage->GetPrivate()->Callbacks();
 				}
 			}
@@ -1098,19 +1110,6 @@ void FinishUpdateGlobal(const TSharedRef<FUpdateContextPrivate>& Context)
 	
 	Context->UpdateCallback.ExecuteIfBound(ContextPublic);
 	Context->UpdateNativeCallback.Broadcast(ContextPublic);
-
-	if (ModelResources)
-	{
-		for (const int32 ResourceIndex : Context->StreamedResourceIndex)
-		{
-			ModelResources->StreamedResourceData[ResourceIndex].Release();
-		}
-
-		for (const int32 ResourceIndex : Context->ExtensionStreamedResourceIndex)
-		{
-			ModelResources->StreamedExtensionData[ResourceIndex].Release();
-		}
-	}
 	
 	if (CVarFixLowPriorityTasksOverlap.GetValueOnGameThread())
 	{
@@ -2165,6 +2164,8 @@ namespace impl
 			OperationData->bClearCacheLayer1);
 		
 		OperationData->MutableInstance = Instance;
+
+		OperationData->ExtensionData = Instance->ExtensionData;
 		
 		const bool bForceGeometryGeneration = CVarForceGeometryOnFirstGeneration.GetValueOnAnyThread();
 
@@ -2537,7 +2538,7 @@ namespace impl
 		}
 		
 		CustomizableObjectInstance->GetPrivate()->RelevantParameters = MoveTemp(Context->RelevantParameters);
-
+		
 		// Actual work
 		// TODO MTBL-391: Review This hotfix
 		UpdateSkeletalMesh(Context);
@@ -2697,15 +2698,6 @@ namespace impl
 
 		UCustomizableInstancePrivate* ObjectInstancePrivateData = ObjectInstance->GetPrivate();
 		
-		// TODO: If this is the first code that runs after the CO program has finished AND if it's
-		// guaranteed that the next CO program hasn't started yet, we need to call ClearActiveObject
-		// and CancelPendingLoads on SystemPrivate->ExtensionDataStreamer.
-		//
-		// FExtensionDataStreamer->AreAnyLoadsPending should return false if the program succeeded.
-		//
-		// If the program aborted, AreAnyLoadsPending may return true, as the program doesn't cancel
-		// its own loads on exit (maybe it should?)
-		
 		if (OperationData->bUseCaches)
 		{
 			FTextureCache& TextureCache = CustomizableObject->GetPrivate()->TextureCache;
@@ -3909,20 +3901,9 @@ uint64 UCustomizableObjectSystem::GetMaxChunkSizeForPlatform(const ITargetPlatfo
 
 	return MaxChunkSize;
 }
-
 #endif // WITH_EDITOR
 
 
-bool UCustomizableObjectSystemPrivate::IsMutableAnimInfoDebuggingEnabled() const
-{ 
-#if WITH_EDITORONLY_DATA
-	return EnableMutableAnimInfoDebugging > 0;
-#else
-	return false;
-#endif
-}
-
-
 void UCustomizableObjectSystemPrivate::OnMutableEnabledChanged(IConsoleVariable* MutableEnabled)
 {
 	if (!UCustomizableObjectSystem::IsCreated())
@@ -4011,16 +3992,6 @@ bool UCustomizableObjectSystemPrivate::IsUpdating(const UCustomizableObjectInsta
 }
 
 
-bool UCustomizableObjectSystem::IsMutableAnimInfoDebuggingEnabled() const
-{
-#if WITH_EDITOR
-	return GetPrivate()->IsMutableAnimInfoDebuggingEnabled();
-#else
-	return false;
-#endif
-}
-
-
 void UCustomizableObjectSystemPrivate::UpdateResourceStreaming(float DeltaTime, bool bProcessEverything)
 {
 	MainTick(false);
diff --git a/Engine/Plugins/Mutable/Source/CustomizableObject/Private/MuCO/CustomizableSkeletalComponent.cpp b/Engine/Plugins/Mutable/Source/CustomizableObject/Private/MuCO/CustomizableSkeletalComponent.cpp
--- a/Engine/Plugins/Mutable/Source/CustomizableObject/Private/MuCO/CustomizableSkeletalComponent.cpp
+++ b/Engine/Plugins/Mutable/Source/CustomizableObject/Private/MuCO/CustomizableSkeletalComponent.cpp
@@ -25,15 +25,6 @@ UCustomizableSkeletalComponentPrivate::UCustomizableSkeletalComponentPrivate()
 }
 
 
-void UCustomizableSkeletalComponentPrivate::Callbacks() const
-{
-	if (InstanceUsage)
-	{
-		InstanceUsage->GetPrivate()->Callbacks();
-	}
-}
-
-
 USkeletalMesh* UCustomizableSkeletalComponentPrivate::GetSkeletalMesh() const
 {
 	if (InstanceUsage)
diff --git a/Engine/Plugins/Mutable/Source/CustomizableObject/Private/MuCO/LogBenchmarkUtil.cpp b/Engine/Plugins/Mutable/Source/CustomizableObject/Private/MuCO/LogBenchmarkUtil.cpp
--- a/Engine/Plugins/Mutable/Source/CustomizableObject/Private/MuCO/LogBenchmarkUtil.cpp
+++ b/Engine/Plugins/Mutable/Source/CustomizableObject/Private/MuCO/LogBenchmarkUtil.cpp
@@ -177,100 +177,6 @@ void FLogBenchmarkUtil::UpdateStats()
 
 	NumAllocatedTextures = LocalNumAllocatedTextures;
 	TextureGPUSize = LocalTextureGPUSize;
-
-#if WITH_EDITORONLY_DATA
-	if (UCustomizableObjectSystem::GetInstance()->IsMutableAnimInfoDebuggingEnabled())
-	{
-		return;
-	}
-
-	if (GEngine)
-	{
-		return;
-	}
-
-	bool bFoundPlayer = false;
-	int32 MsgIndex = 15820; // Arbitrary big value to prevent collisions with other on-screen messages
-
-	for (TObjectIterator<UCustomizableObjectInstanceUsage> CustomizableObjectInstanceUsage; CustomizableObjectInstanceUsage; ++CustomizableObjectInstanceUsage)
-	{
-		if (!IsValid(*CustomizableObjectInstanceUsage) || CustomizableObjectInstanceUsage->GetPrivate()->IsNetMode(NM_DedicatedServer))
-		{
-			continue;
-		}
-
-		USkeletalMeshComponent* Parent = Cast<USkeletalMeshComponent>(CustomizableObjectInstanceUsage->GetAttachParent());
-		AActor* ParentActor = Parent ? 
-			Parent->GetAttachmentRootActor()
-			:  nullptr;
-		UCustomizableObjectInstance* Instance = CustomizableObjectInstanceUsage->GetCustomizableObjectInstance();
-
-		APawn* PlayerPawn = nullptr;
-		UWorld* World = Parent ? Parent->GetWorld() : nullptr;
-
-		if (World)
-		{
-			PlayerPawn = UGameplayStatics::GetPlayerPawn(World, 0);
-		}
-
-		if (ParentActor && (ParentActor == PlayerPawn) && Instance)
-		{
-			bFoundPlayer = true;
-
-			FString TagString;
-			const FGameplayTagContainer& Tags = Instance->GetAnimationGameplayTags();
-
-			for (const FGameplayTag& Tag : Tags)
-			{
-				TagString += !TagString.IsEmpty() ? FString(TEXT(", ")) : FString();
-				TagString += Tag.ToString();
-			}
-
-			GEngine->AddOnScreenDebugMessage(MsgIndex++, .0f, FColor::Green, TEXT("Animation tags: ") + TagString);
-
-			check(Instance->GetPrivate() != nullptr);
-			FCustomizableInstanceComponentData* ComponentData = Instance->GetPrivate()->GetComponentData(CustomizableObjectInstanceUsage->GetComponentName());
-
-			if (ComponentData)
-			{
-				for (TPair<FName, TSoftClassPtr<UAnimInstance>>& Entry : ComponentData->AnimSlotToBP)
-				{
-					FString AnimBPSlot;
-
-					AnimBPSlot += Entry.Key.ToString() + FString("-") + Entry.Value.GetAssetName();
-					GEngine->AddOnScreenDebugMessage(MsgIndex++, .0f, FColor::Green, AnimBPSlot);
-				}
-			}
-
-			GEngine->AddOnScreenDebugMessage(MsgIndex++, .0f, FColor::Green, TEXT("Slots-AnimBP: "));
-
-			if (ComponentData)
-			{
-				if (ComponentData->MeshPartPaths.IsEmpty())
-				{
-					GEngine->AddOnScreenDebugMessage(MsgIndex++, .0f, FColor::Magenta,
-						TEXT("No meshes found. In order to see the meshes compile the pawn's root CustomizableObject after the 'mutable.EnableMutableAnimInfoDebugging 1' command has been run."));
-				}
-
-				for (const FString& MeshPath : ComponentData->MeshPartPaths)
-				{
-					GEngine->AddOnScreenDebugMessage(MsgIndex++, .0f, FColor::Magenta, MeshPath);
-				}
-			}
-
-			GEngine->AddOnScreenDebugMessage(MsgIndex++, .0f, FColor::Magenta, TEXT("Meshes: "));
-
-			GEngine->AddOnScreenDebugMessage(MsgIndex++, .0f, FColor::Cyan,
-				FString::Printf(TEXT("Player Pawn Mutable Mesh/Animation info for component %s"),
-					*CustomizableObjectInstanceUsage->GetComponentName().ToString()));
-		}
-	}
-
-	if (!bFoundPlayer)
-	{
-		GEngine->AddOnScreenDebugMessage(MsgIndex, .0f, FColor::Yellow, TEXT("Mutable Animation info: N/A"));
-	}
-#endif
 }
 
 
@@ -327,7 +233,6 @@ void FLogBenchmarkUtil::FinishUpdateMesh(const TSharedRef<FUpdateContextPrivate>
 			}
 		}
 	}
-
 	
 	// Identifying data
 	FInstanceUpdateStats UpdateData;
diff --git a/Engine/Plugins/Mutable/Source/CustomizableObject/Private/MuCO/UnrealConversionUtils.cpp b/Engine/Plugins/Mutable/Source/CustomizableObject/Private/MuCO/UnrealConversionUtils.cpp
--- a/Engine/Plugins/Mutable/Source/CustomizableObject/Private/MuCO/UnrealConversionUtils.cpp
+++ b/Engine/Plugins/Mutable/Source/CustomizableObject/Private/MuCO/UnrealConversionUtils.cpp
@@ -1549,8 +1549,6 @@ namespace UnrealConversionUtils
 					return;
 				}
 					
-				Result->SkeletalMeshes.AddUnique(SkeletalMesh);
-
 				const FSkelMeshRenderSection& Section = Context->OriginalLOD->RenderSections[SectionIndex];
 				
 				const bool bIgnoreSkinning = EnumHasAnyFlags(ConversionFlags, EMutableMeshConversionFlags::IgnoreSkinning);
@@ -1960,6 +1958,12 @@ namespace UnrealConversionUtils
 						FMemory::Memcpy(Result->ClothSections[0].Data.GetData(), &ClothVertexBuffer.MappingData(SectionClothIndexMapping.BaseVertexIndex), NumVertices * sizeof(FMeshToMeshVertData));
 					}	
 				}
+					
+				// Asset User Data
+				for (UAssetUserData* AssetUserData : *SkeletalMesh->GetAssetUserDataArray())
+				{
+					Result->AssetUserData.Add(AssetUserData);
+				}
 			},
 			ConversionTask);
 
diff --git a/Engine/Plugins/Mutable/Source/CustomizableObject/Public/MuCO/CustomizableObject.h b/Engine/Plugins/Mutable/Source/CustomizableObject/Public/MuCO/CustomizableObject.h
--- a/Engine/Plugins/Mutable/Source/CustomizableObject/Public/MuCO/CustomizableObject.h
+++ b/Engine/Plugins/Mutable/Source/CustomizableObject/Public/MuCO/CustomizableObject.h
@@ -293,10 +293,6 @@ private:
 	UPROPERTY(EditAnywhere, Category = CompileOptions)
 	bool bEnableAnimBpPhysicsAssetsManipulation = false;
 	
-	// When this is enabled generated meshes will merge the AssetUserData from all of its constituent mesh parts
-	UPROPERTY(EditAnywhere, Category = CompileOptions)
-	bool bEnableAssetUserDataMerge = false;
-	
 	// Disabling the Table Materials parent material check will let the user use any material regardless of its parent when connecting a material from a table column to a material node.
 	// Warning! But it will not check if the table material channels connected to the Material node exist in the actual material used in the Instance, and will fail silently at runtime 
 	// when setting the value of those channels if they don't exist.
diff --git a/Engine/Plugins/Mutable/Source/CustomizableObject/Public/MuCO/CustomizableObjectExtension.h b/Engine/Plugins/Mutable/Source/CustomizableObject/Public/MuCO/CustomizableObjectExtension.h
--- a/Engine/Plugins/Mutable/Source/CustomizableObject/Public/MuCO/CustomizableObjectExtension.h
+++ b/Engine/Plugins/Mutable/Source/CustomizableObject/Public/MuCO/CustomizableObjectExtension.h
@@ -8,6 +8,11 @@
 
 #include "CustomizableObjectExtension.generated.h"
 
+namespace UE::Mutable::Private
+{
+	class FExtensionData;
+}
+
 class UCustomizableObjectInstance;
 class UCustomizableObjectInstanceUsage;
 class USkeletalMesh;
@@ -63,18 +68,6 @@ struct FObjectNodeInputPin
 	bool bIsArray = false;
 };
 
-/** An Object node input pin and the data that was passed into it by the Customizable Object graph */
-struct FInputPinDataContainer
-{
-	FInputPinDataContainer(const FObjectNodeInputPin& InPin, const FInstancedStruct& InData)
-		: Pin(InPin)
-		, Data(InData)
-	{
-	}
-
-	FObjectNodeInputPin Pin;
-	const FInstancedStruct& Data;
-};
 
 /**
  * An extension that adds functionality to the Customizable Object system
@@ -104,42 +97,12 @@ public:
 	 * 
 	 * @param SkeletalMesh - The Skeletal Mesh that was created.
 	 */
-	virtual void OnSkeletalMeshCreated(const TArray<FInputPinDataContainer>& InputPinData, FName ComponentName, USkeletalMesh* SkeletalMesh) const {}
+	virtual void OnSkeletalMeshCreated(FName ComponentName, USkeletalMesh* SkeletalMesh) const {}
 	
-	/**
-	 * Note that the data registered here is completely independent of any Extension Data used in
-	 * the Customizable Object graph. Even though Extension Data and this Extension Instance Data
-	 * both use FInstancedStruct to box an extension-defined struct, there's no requirement that
-	 * they use the same struct type, so they may be completely unrelated.
-	 * 
-	 * Note that GetExtensionInstanceData returns the struct by value to ensure memory safety, so
-	 * the struct should ideally be small and cheap to copy. If you need to reference large data
-	 * from this struct, consider wrapping it in a UObject or referencing it via a TSharedPtr so
-	 * that the large data itself isn't being copied.
-	 */
-	virtual FInstancedStruct GenerateExtensionInstanceData(const TArray<FInputPinDataContainer>& InputPinData) const { return FInstancedStruct(); }
-
-#if WITH_EDITOR
-	/**
-	 * Non-owned references to private UObjects must be converted to owned references at cook time.
-	 * Duplicate the private UObjects with the container as its new outer object, or the cook will 
-	 * fail to serialize them.
-	 * 
-	 * E.g. Any UObjects that are referenced by the Extension Data that aren't in their own asset
-	 * package should be copied with the Container as their Outer. This ensures that they get cooked
-	 * into the correct package for streaming.
-	 * 
-	 * @param Struct - Instanced struct of unknown type that could hold external private references.
-	 *
-	 * @param Container - Outer UObject for duplicated private references.
-	 */
-	virtual void MovePrivateReferencesToContainer(FInstancedStruct& Struct, UObject* Container) const {};
-#endif
-
 	/** Called when a Customizable Object Instance Usage is being updated.
 	 * 
 	 * @param Usage Customizable Object Instance Usage being updated. */
-	virtual void OnCustomizableObjectInstanceUsageUpdated(UCustomizableObjectInstanceUsage& Usage) const {}
+	virtual void OnCustomizableObjectInstanceUsageUpdated(UCustomizableObjectInstanceUsage& Usage, TArray<const UObject*>& ExtensionData) const {}
 
 	/** Called when a Customizable Object Instance Usage is being discarded.
 	 * 
diff --git a/Engine/Plugins/Mutable/Source/CustomizableObject/Public/MuCO/CustomizableObjectInstance.h b/Engine/Plugins/Mutable/Source/CustomizableObject/Public/MuCO/CustomizableObjectInstance.h
--- a/Engine/Plugins/Mutable/Source/CustomizableObject/Public/MuCO/CustomizableObjectInstance.h
+++ b/Engine/Plugins/Mutable/Source/CustomizableObject/Public/MuCO/CustomizableObjectInstance.h
@@ -749,7 +749,7 @@ public:
 
 	/** Return the list of tags for this instance. */
 	UFUNCTION(BlueprintCallable, Category = CustomizableObjectInstance)
-	CUSTOMIZABLEOBJECT_API const FGameplayTagContainer& GetAnimationGameplayTags() const;
+	CUSTOMIZABLEOBJECT_API FGameplayTagContainer GetAnimationGameplayTags() const;
 	
 	/** Execute a delegate for each animation instance involved in this customizable object instance. */
 	UFUNCTION(BlueprintCallable, Category = CustomizableObjectInstance)
@@ -762,9 +762,6 @@ public:
 
 	/** Fix the given UAnimInstance instance. */
 	CUSTOMIZABLEOBJECT_API void AnimInstanceFixup(UAnimInstance* AnimInstance) const;
-
-	UFUNCTION(BlueprintCallable, Category = CustomizableObjectInstance)
-	CUSTOMIZABLEOBJECT_API FInstancedStruct GetExtensionInstanceData(const UCustomizableObjectExtension* Extension) const;
 	
 	/** See FCustomizableObjectInstanceDescriptor::SaveDescriptor. */
 	CUSTOMIZABLEOBJECT_API void SaveDescriptor(FArchive &CustomizableObjectDescriptor, bool bUseCompactDescriptor);
diff --git a/Engine/Plugins/Mutable/Source/CustomizableObject/Public/MuCO/CustomizableObjectInstanceAssetUserData.h b/Engine/Plugins/Mutable/Source/CustomizableObject/Public/MuCO/CustomizableObjectInstanceAssetUserData.h
--- a/Engine/Plugins/Mutable/Source/CustomizableObject/Public/MuCO/CustomizableObjectInstanceAssetUserData.h
+++ b/Engine/Plugins/Mutable/Source/CustomizableObject/Public/MuCO/CustomizableObjectInstanceAssetUserData.h
@@ -4,6 +4,7 @@
 
 #include "Engine/AssetUserData.h"
 #include "GameplayTagContainer.h"
+#include "UObject/SoftObjectPtr.h"
 
 #include "CustomizableObjectInstanceAssetUserData.generated.h"
 
@@ -15,10 +16,10 @@ struct FCustomizableObjectAnimationSlot
 {
 	GENERATED_BODY()
 
-	UPROPERTY(BlueprintReadWrite, Category = CustomizableObjectInstance)
+	UPROPERTY(VisibleAnywhere, BlueprintReadWrite, Category = CustomizableObjectInstance)
 	FName Name;
 
-	UPROPERTY(BlueprintReadWrite, Category = CustomizableObjectInstance)
+	UPROPERTY(VisibleAnywhere, BlueprintReadWrite, Category = CustomizableObjectInstance)
 	TSoftClassPtr<UAnimInstance> AnimInstance;
 };
 
@@ -32,21 +33,15 @@ class UCustomizableObjectInstanceUserData : public UAssetUserData
 public:
 	/** Return the list of tags for this instance. */
 	UFUNCTION(BlueprintCallable, Category = CustomizableObjectInstance)
-	const FGameplayTagContainer& GetAnimationGameplayTags() const
-	{
-		return AnimationGameplayTag;
-	};
+	CUSTOMIZABLEOBJECT_API const FGameplayTagContainer& GetAnimationGameplayTags() const;
 
 	/** Sets the list of tags for this instance. */
 	UFUNCTION(BlueprintCallable, Category = CustomizableObjectInstance)
-	void SetAnimationGameplayTags(const FGameplayTagContainer& InstanceTags)
-	{
-		AnimationGameplayTag = InstanceTags;
-	};
+	CUSTOMIZABLEOBJECT_API void SetAnimationGameplayTags(const FGameplayTagContainer& InstanceTags);
 	
-	UPROPERTY(BlueprintReadWrite, Category = CustomizableObjectInstance)
+	UPROPERTY(VisibleAnywhere, BlueprintReadWrite, Category = CustomizableObjectInstance)
 	FGameplayTagContainer AnimationGameplayTag;
 
-	UPROPERTY(BlueprintReadWrite, Category = CustomizableObjectInstance)
+	UPROPERTY(VisibleAnywhere, BlueprintReadWrite, Category = CustomizableObjectInstance)
 	TArray<FCustomizableObjectAnimationSlot> AnimationSlots;
 };
diff --git a/Engine/Plugins/Mutable/Source/CustomizableObject/Public/MuCO/CustomizableObjectResourceData.h b/Engine/Plugins/Mutable/Source/CustomizableObject/Public/MuCO/CustomizableObjectResourceData.h
deleted file mode 100644
diff --git a/Engine/Plugins/Mutable/Source/CustomizableObject/Public/MuCO/CustomizableObjectSystem.h b/Engine/Plugins/Mutable/Source/CustomizableObject/Public/MuCO/CustomizableObjectSystem.h
--- a/Engine/Plugins/Mutable/Source/CustomizableObject/Public/MuCO/CustomizableObjectSystem.h
+++ b/Engine/Plugins/Mutable/Source/CustomizableObject/Public/MuCO/CustomizableObjectSystem.h
@@ -230,8 +230,6 @@ public:
 	UFUNCTION(BlueprintCallable, Category = Status)
 	UE_API int32 GetAverageBuildTime() const;
 
-	UE_API bool IsMutableAnimInfoDebuggingEnabled() const;
-
 #if WITH_EDITOR
 	// Get the maximum size a chunk can have on a specific platform. If unspecified return MUTABLE_STREAMED_DATA_MAXCHUNKSIZE
 	UE_API uint64 GetMaxChunkSizeForPlatform(const ITargetPlatform* InTargetPlatform);
diff --git a/Engine/Plugins/Mutable/Source/CustomizableObjectEditor/Private/MuCOE/CustomizableObjectCompiler.cpp b/Engine/Plugins/Mutable/Source/CustomizableObjectEditor/Private/MuCOE/CustomizableObjectCompiler.cpp
--- a/Engine/Plugins/Mutable/Source/CustomizableObjectEditor/Private/MuCOE/CustomizableObjectCompiler.cpp
+++ b/Engine/Plugins/Mutable/Source/CustomizableObjectEditor/Private/MuCOE/CustomizableObjectCompiler.cpp
@@ -1026,9 +1026,6 @@ void FCustomizableObjectCompiler::CompileInternal(bool bAsync)
 
 		ModelResources->bCompiledWithHDTextureCompression = CurrentOptions.TextureCompression == ECustomizableObjectTextureCompression::HighQuality;
 
-		ModelResources->AlwaysLoadedExtensionData = MoveTemp(GenerationContext.AlwaysLoadedExtensionData);
-		ModelResources->StreamedExtensionDataEditor = MoveTemp(GenerationContext.StreamedExtensionData);
-
 #if WITH_EDITORONLY_DATA
 		// Cache the tables that are used by more than one param so that CompileOnlySelected can work properly
 		ModelResources->TableToParamNames = MoveTemp(GenerationContext.TableToParamNames);
@@ -1412,11 +1409,6 @@ void FCustomizableObjectCompiler::FinishLoadingModelDataFromDDC()
 		{
 			PlatformData->ModelResources = TStrongObjectPtr<UModelResources>(LocalModelResources);
 
-			if (CurrentOptions.bIsCooking)
-			{
-				PlatformData->ModelResources->InitCookData(*CurrentObject);
-			}
-
 			FMemoryReaderView ModelStreamablesReader(DDCHeapMemory->ModelStreamablesBytesDDC.GetView());
 			PlatformData->ModelStreamableBulkData = LoadModelStreamableBulk_Internal(ModelStreamablesReader);
 			PlatformData->ModelStreamableBulkData->bIsStoredInDDC = true;
@@ -1670,6 +1662,7 @@ void FCustomizableObjectCompiler::FinishCompilationTask()
 		TObjectPtr<UModelResources> ModelResources = PlatformData->ModelResources.Get();
 
 		ModelResources->PassthroughObjects = MoveTemp(CompilationContext->PassthroughObjectFactory.GetAll());
+		ModelResources->DuplicateObjects = MoveTemp(CompilationContext->PassthroughObjectFactory.GetDuplicate());
 
 		ModelResources->ExternalOperations = MoveTemp(CompileTask->ExternalOperations);
 
@@ -1683,8 +1676,6 @@ void FCustomizableObjectCompiler::FinishCompilationTask()
 		
 		ModelResources->SkinWeightProfilesInfo = MoveTemp(CompilationContext->SkinWeightProfilesInfo);
 
-		ModelResources->StreamedResourceDataEditor = MoveTemp(CompilationContext->StreamedResourceData);
-
 		TSharedPtr<FModelStreamableBulkData> ModelStreamableBulkData = PlatformData->ModelStreamableBulkData;
 		
 		ModelResources->SurfaceMetadata = MoveTemp(CompilationContext->SurfaceMetadata);
@@ -1722,8 +1713,6 @@ void FCustomizableObjectCompiler::FinishSavingDerivedDataTask()
 
 		UE::Mutable::Private::FMutableCachedPlatformData& Data = CurrentObject->GetPrivate()->CachedPlatformsData.Add(PlatformName);
 		Data = MoveTemp(*PlatformData.Get());
-
-		Data.ModelResources->InitCookData(*CurrentObject);
 	}
 
 	// Order matters
diff --git a/Engine/Plugins/Mutable/Source/CustomizableObjectEditor/Private/MuCOE/CustomizableObjectCookPackageSplitter.cpp b/Engine/Plugins/Mutable/Source/CustomizableObjectEditor/Private/MuCOE/CustomizableObjectCookPackageSplitter.cpp
--- a/Engine/Plugins/Mutable/Source/CustomizableObjectEditor/Private/MuCOE/CustomizableObjectCookPackageSplitter.cpp
+++ b/Engine/Plugins/Mutable/Source/CustomizableObjectEditor/Private/MuCOE/CustomizableObjectCookPackageSplitter.cpp
@@ -2,7 +2,8 @@
 
 #include "MuCOE/CustomizableObjectCookPackageSplitter.h"
 
-#include "Algo/Find.h"
+#include "ClothConfigBase.h"
+#include "ClothingAsset.h"
 #include "MuCO/CustomizableObject.h"
 #include "MuCO/CustomizableObjectPrivate.h"
 #include "MuCO/LoadUtils.h"
@@ -11,8 +12,7 @@
 
 REGISTER_COOKPACKAGE_SPLITTER(FCustomizableObjectCookPackageSplitter, UCustomizableObject);
 
-namespace
-{
+
 UModelResources* FindModelResources(UCustomizableObject& Object)
 {
 	// All platforms should have the same resources
@@ -27,414 +27,148 @@ UModelResources* FindModelResources(UCustomizableObject& Object)
 
 	return nullptr;
 }
-
-// Look up a streamed Resource Data constant by name on a Customizable Object.
-//
-// Returns nullptr if not found.
-FCustomizableObjectStreamedResourceData* FindStreamedResourceData(
-	TArray<FCustomizableObjectStreamedResourceData>& StreamedResources,
-	const FString& ContainerName
-	)
+	
+	
+const UModelResources* FindModelResources(const UCustomizableObject& Object)
 {
-	return Algo::FindByPredicate(
-		StreamedResources,
-		[&ContainerName](const FCustomizableObjectStreamedResourceData& StreamedData)
-		{
-			const FSoftObjectPath& Path = StreamedData.GetPath().ToSoftObjectPath();
-
-			// ContainerName should match the last element of the path, which could be the
-			// sub-path string or the asset name.
-
-			if (Path.GetSubPathString().Len() > 0)
-			{
-				return Path.GetSubPathString() == ContainerName;
-			}
-
-			return Path.GetAssetName() == ContainerName;
-		});
+	return FindModelResources(*const_cast<UCustomizableObject*>(&Object));
 }
 
-enum class EMoveContainerError
-{
-	None,
-	FailedToLoadContainer,
-	NameCollision, // Object with that name already exists in the new outer
-	RenameFailed,
-};
-
-const TCHAR* LexToString(EMoveContainerError Error)
-{
-	switch (Error)
-	{
-	case EMoveContainerError::None: return TEXT("None");
-	case EMoveContainerError::FailedToLoadContainer: return TEXT("FailedToLoadContainer");
-	case EMoveContainerError::NameCollision: return TEXT("NameCollision");
-	case EMoveContainerError::RenameFailed: return TEXT("RenameFailed");
-	default: return TEXT("Unknown");
-	}
-}
 
-// Moves the StreamedResourceData's data container to the given Outer.
-EMoveContainerError MoveContainerToNewOuter(
-	UObject* NewOuter,
-	const FCustomizableObjectStreamedResourceData* StreamedResourceData,
-	UCustomizableObjectResourceDataContainer*& OutContainer
-)
+bool FCustomizableObjectCookPackageSplitter::ShouldSplit(UObject* SplitData)
 {
-	check(StreamedResourceData);
+	UCustomizableObject* Object = CastChecked<UCustomizableObject>(SplitData);
 	
-	OutContainer = nullptr;
-
-	UCustomizableObjectResourceDataContainer* Container = UE::Mutable::Private::LoadObject(StreamedResourceData->GetPath());
-	if (!Container)
+	if (Object->IsChildObject())
 	{
-		return EMoveContainerError::FailedToLoadContainer;
+		return false;
 	}
-
-	if (Container->GetOuter() != NewOuter)
+	
+	const UModelResources* ModelResources = FindModelResources(*Object);
+	if (!ModelResources)
 	{
-		// Ensure the target object doesn't exist
-		if (FindObject<UObject>(NewOuter, *Container->GetName()))
-		{
-			return EMoveContainerError::NameCollision;
-		}
-
-		// The Rename function moves the object into the given package
-		if (!Container->Rename(nullptr, NewOuter, REN_DontCreateRedirectors))
-		{
-			return EMoveContainerError::RenameFailed;
-		}
+		return false;
 	}
-
-	OutContainer = Container;
-	return EMoveContainerError::None;
+	
+	return !ModelResources->DuplicateObjects.IsEmpty();
 }
 
 
-void GenerateNewPackage(const FCustomizableObjectStreamedResourceData& StreamedData,
-	const UPackage* OwnerPackage,
-	const UObject* OwnerObject,
-	TArray<ICookPackageSplitter::FGeneratedPackage>& Result)
+FString FCustomizableObjectCookPackageSplitter::GetSplitterDebugName()
 {
-	// The StreamedData container path should be of the form
-	// OwnerPackageName.OwnerObjectName:ContainerName
-	const FSoftObjectPath& StreamedDataPath = StreamedData.GetPath().ToSoftObjectPath();
-
-	// Check that the StreamedData container has the OwnerObject as its Outer
-	check(StreamedDataPath.GetWithoutSubPath() == FSoftObjectPath(OwnerObject));
-
-	// Check that the ContainerName is valid and that there isn't another Outer level between
-	// the OwnerObject and the container.
-	check(StreamedDataPath.GetSubPathString().Len() > 0);
-	check(!StreamedDataPath.GetSubPathString().Contains(SUBOBJECT_DELIMITER));
-
-	ICookPackageSplitter::FGeneratedPackage& Package = Result.AddDefaulted_GetRef();
-	// Because of the checks above, the container name must be unique within this Customizable
-	// Object, so it's safe to use as a package path.
-	Package.RelativePath = StreamedDataPath.GetSubPathString();
-	Package.SetCreateAsMap(false);
-
-	// To support iterative cooking, GenerationHash should only change when OwnerPackage
-	// changes.
-	//
-	// The simplest and fastest way to do this is to set it to OwnerPackage's PackageSavedHash.
-	{
-		// Zero the hash, as we won't be writing all bytes of it below
-		Package.GenerationHash.Reset();
+	return TEXT("FCustomizableObjectCookPackageSplitter");
+}
 
-		FIoHash OwnerSavedHash = OwnerPackage->GetSavedHash();
-		static_assert(sizeof(Package.GenerationHash.GetBytes()) >= sizeof(OwnerSavedHash.GetBytes()));  // -V568
-		static_assert(sizeof(Package.GenerationHash.GetBytes()) > 8); // It should be a byte array, not a pointer // -V568
-		static_assert(sizeof(OwnerSavedHash.GetBytes()) > 8); // It should be a byte array, not a pointer // -V568
-		FMemory::Memcpy(Package.GenerationHash.GetBytes(), OwnerSavedHash.GetBytes(), sizeof(OwnerSavedHash.GetBytes())); // -V568
 
-	}
+bool FCustomizableObjectCookPackageSplitter::RequiresCachedCookedPlatformDataBeforeSplit()
+{
+	return true;
 }
 
-}
 
-bool FCustomizableObjectCookPackageSplitter::ShouldSplit(UObject* SplitData)
+FString GenerateRelativePath(const UCustomizableObject& Object, UE::Mutable::Private::PASSTHROUGH_ID Id)
 {
-	UCustomizableObject* Object = CastChecked<UCustomizableObject>(SplitData);
-	
-	if (!Object->IsChildObject())
-	{
-		if(const UModelResources* ModelResources = FindModelResources(*Object))
-		{
-			return ModelResources->StreamedResourceData.Num() || ModelResources->StreamedExtensionData.Num();
-		}
-	}
-
-	return false;
+	return FString::Printf(TEXT("%s-%u"), *Object.GetName(), static_cast<uint32>(Id));
 }
 
-ICookPackageSplitter::FGenerationManifest FCustomizableObjectCookPackageSplitter::ReportGenerationManifest(
-	const UPackage* OwnerPackage,
-	const UObject* OwnerObject)
+
+ICookPackageSplitter::FGenerationManifest FCustomizableObjectCookPackageSplitter::ReportGenerationManifest(const UPackage* OwnerPackage, const UObject* OwnerObject)
 {
-	// Keep a strong reference to the CO.
-	StrongObject.Reset(OwnerObject);
-	
-	UCustomizableObject* Object = const_cast<UCustomizableObject*>(CastChecked<UCustomizableObject>(OwnerObject));
+	const UCustomizableObject* Object = CastChecked<UCustomizableObject>(OwnerObject);
 
 	// All platforms should have the same resources
+
 	const UModelResources* ModelResources = FindModelResources(*Object);
 	check(ModelResources);
 
-	ICookPackageSplitter::FGenerationManifest Result;
+	FGenerationManifest Result;
 
 	// Generate a new package for each streamed Resource Data
-	for (const FCustomizableObjectStreamedResourceData& StreamedData : ModelResources->StreamedResourceData)
-	{
-		GenerateNewPackage(StreamedData, OwnerPackage, OwnerObject, Result.GeneratedPackages);
-	}
-
-	// Generate a new package for each streamed Extension Data
-	for (const FCustomizableObjectStreamedResourceData& StreamedData : ModelResources->StreamedExtensionData)
-	{
-		GenerateNewPackage(StreamedData, OwnerPackage, OwnerObject, Result.GeneratedPackages);
+	for (const UE::Mutable::Private::PASSTHROUGH_ID Id : ModelResources->DuplicateObjects)
+	{
+		FGeneratedPackage& GeneratedPackage = Result.GeneratedPackages.AddDefaulted_GetRef();
+		// Because of the checks above, the container name must be unique within this Customizable
+		// Object, so it's safe to use as a package path.
+		GeneratedPackage.RelativePath = GenerateRelativePath(*Object, Id);
+		GeneratedPackage.SetCreateAsMap(false);
+
+		// To support iterative cooking, GenerationHash should only change when OwnerPackage
+		// changes.
+		//
+		// The simplest and fastest way to do this is to set it to OwnerPackage's PackageSavedHash.
+		{
+			// Zero the hash, as we won't be writing all bytes of it below
+			GeneratedPackage.GenerationHash.Reset();
+
+			FIoHash OwnerSavedHash = OwnerPackage->GetSavedHash();
+			static_assert(sizeof(GeneratedPackage.GenerationHash.GetBytes()) >= sizeof(OwnerSavedHash.GetBytes()));  // -V568
+			static_assert(sizeof(GeneratedPackage.GenerationHash.GetBytes()) > 8); // It should be a byte array, not a pointer // -V568
+			static_assert(sizeof(OwnerSavedHash.GetBytes()) > 8); // It should be a byte array, not a pointer // -V568
+			FMemory::Memcpy(GeneratedPackage.GenerationHash.GetBytes(), OwnerSavedHash.GetBytes(), sizeof(OwnerSavedHash.GetBytes())); // -V568
+		}
 	}
 
 	return Result;
 }
 
-bool FCustomizableObjectCookPackageSplitter::PreSaveGeneratorPackage(FPopulateContext& PopulateContext)
-{
-	// The CO is just about to be saved (i.e. produce the cooked version of the asset), so this
-	// function needs to:
-	// 
-	// 1.	Move the streamed Data out of the CO's package, so that it doesn't get saved
-	//		into the cooked package.
-	// 
-	// 2.	Remove hard references to the streamed data, so that it doesn't get loaded as soon as
-	//		the CO is loaded
-	TConstArrayView<ICookPackageSplitter::FGeneratedPackageForPopulate>& PlaceholderPackages = PopulateContext.GetGeneratedPackages();
-	const auto& PreSavePackage = [] (const ICookPackageSplitter::FGeneratedPackageForPopulate& GeneratedPackage,
-		TArray<FCustomizableObjectStreamedResourceData>& StreamedResources
-		) -> bool
-	{
-		FCustomizableObjectStreamedResourceData* FoundData = FindStreamedResourceData(StreamedResources, GeneratedPackage.RelativePath);
-		if (!FoundData)
-		{
-			UE_LOG(LogMutable, Error, TEXT("Couldn't find streamed Resource Data container with name %s in array of %d entries"),
-				*GeneratedPackage.RelativePath, StreamedResources.Num());
-
-			return false;
-		}
-
-		// Move the streamed data to the generated package
-		UCustomizableObjectResourceDataContainer* Container = nullptr;
-		EMoveContainerError Error = MoveContainerToNewOuter(GeneratedPackage.Package, FoundData, Container);
-		if (Error != EMoveContainerError::None)
-		{
-			UE_LOG(LogMutable, Error, TEXT("Failed to move container %s to new outer %s - %s"), *FoundData->GetPath().ToSoftObjectPath().ToString(), *GetPathNameSafe(GeneratedPackage.Package), LexToString(Error));
-			return false;
-		}
-
-		return true;
-	};
 
+bool FCustomizableObjectCookPackageSplitter::PopulateGeneratorPackage(FPopulateContext& PopulateContext)
+{
 	UCustomizableObject* Object = CastChecked<UCustomizableObject>(PopulateContext.GetOwnerObject());
-
 	UModelResources* ModelResources = FindModelResources(*Object);
-	if (!ModelResources)
-	{
-		UE_LOG(LogMutable, Warning, TEXT("Couldn't find ModelResources. CO %s"), *GetNameSafe(Object));
-		return false;
-	}
-
-	// There should be one generated package per streamed Resource Data
-	const int32 NumStreamedData = ModelResources->StreamedResourceData.Num();
-	const int32 NumStreamedExtensionData = ModelResources->StreamedExtensionData.Num();
-	
-	check(NumStreamedData + NumStreamedExtensionData == PlaceholderPackages.Num());
-
-
-	// After the CO has been saved, the contract for ICookPackageSplitter states that we need to
-	// restore the CO back to how it was before, so we need to save some information to help with
-	// this.
-	SavedContainerNames.Reset();
-	SavedExtensionContainerNames.Reset();
-
-	for (int32 Index = 0; Index < NumStreamedData; ++Index)
-	{
-		const ICookPackageSplitter::FGeneratedPackageForPopulate& GeneratedPackage = PlaceholderPackages[Index];
-		if (!PreSavePackage(GeneratedPackage, ModelResources->StreamedResourceData))
-		{
-			return false;
-		}
-
-		SavedContainerNames.Add(GeneratedPackage.RelativePath);
-	}
-
-	for (int32 Index = 0; Index < NumStreamedExtensionData; ++Index)
-	{
-		const ICookPackageSplitter::FGeneratedPackageForPopulate& GeneratedPackage = PlaceholderPackages[NumStreamedData + Index];
-		if (!PreSavePackage(GeneratedPackage, ModelResources->StreamedExtensionData))
-		{
-			return false;
-		}
-
-		SavedExtensionContainerNames.Add(GeneratedPackage.RelativePath);
-	}
+	check(ModelResources);
 
-	// All platforms should have the same resources
-	const TMap<FString, UE::Mutable::Private::FMutableCachedPlatformData>& CachePlatforms = Object->GetPrivate()->CachedPlatformsData;
-	for (const TPair<FString, UE::Mutable::Private::FMutableCachedPlatformData>& PlatformData : CachePlatforms)
+	for (const UE::Mutable::Private::PASSTHROUGH_ID Id : ModelResources->DuplicateObjects)
 	{
-		if (PlatformData.Value.ModelResources)
+		TSoftObjectPtr<UObject>& SoftPassthroughObject = ModelResources->PassthroughObjects[Id];
+		
+		UObject* PassthroughObject = UE::Mutable::Private::LoadObject(SoftPassthroughObject);
+		if (PassthroughObject)
 		{
-			for (FCustomizableObjectStreamedResourceData& StreamedResourceData : PlatformData.Value.ModelResources->StreamedResourceData)
+			const FGeneratedPackageForPopulate* GeneratedPackage = PopulateContext.GetGeneratedPackages().FindByPredicate(
+			[&](const FGeneratedPackageForPopulate& Element)
 			{
-				// Remove the hard reference and set the soft reference to the streamed data's new location
-				StreamedResourceData.ConvertToSoftReferenceForCooking();
-			}
-
-			for (FCustomizableObjectStreamedResourceData& StreamedExtensionData : PlatformData.Value.ModelResources->StreamedExtensionData)
+				return Element.RelativePath == GenerateRelativePath(*Object, Id);
+			});
+			
+			check(GeneratedPackage);
+			
+			UObject* DuplicatedObject = DuplicateObject(PassthroughObject, GeneratedPackage->Package, *PassthroughObject->GetName());
+			DuplicatedObject->SetFlags(RF_Public);
+			PopulateContext.ReportObjectToMove(DuplicatedObject);
+			
+			SoftPassthroughObject = DuplicatedObject;
+
+			if (UClothingAssetCommon* ClothingAsset = Cast<UClothingAssetCommon>(DuplicatedObject))
 			{
-				// Remove the hard reference and set the soft reference to the streamed data's new location
-				StreamedExtensionData.ConvertToSoftReferenceForCooking();
+				for (TMap<FName, TObjectPtr<UClothConfigBase>>::TIterator ConfigIt = ClothingAsset->ClothConfigs.CreateIterator(); ConfigIt; ++ConfigIt)
+				{
+					TObjectPtr<UClothConfigBase>& ClothConfig = ConfigIt->Value;
+					
+					UClothConfigBase* DuplicatedClothConfig = DuplicateObject<UClothConfigBase>(ClothConfig, DuplicatedObject);
+					
+					ClothConfig = DuplicatedClothConfig;
+				}
 			}
 		}
 	}
-
+	
 	return true;
 }
 
-void FCustomizableObjectCookPackageSplitter::PostSaveGeneratorPackage(FPopulateContext& PopulateContext)
-{
-	// Move the streamed data back into the CO's package and restore the StreamedResourceData and StreamedExtensionData
-	// array on the CO to how it was before PreSaveGeneratorPackage.
-
-	UCustomizableObject* Object = CastChecked<UCustomizableObject>(PopulateContext.GetOwnerObject());
-	UModelResources* ModelResources = FindModelResources(*Object);
-	if (!ModelResources)
-	{
-		UE_LOG(LogMutable, Warning, TEXT("Couldn't find ModelResources. CO %s"), *GetNameSafe(Object));
-		return;
-	}
-
-	TArray<FCustomizableObjectStreamedResourceData> NewArray;
-	NewArray.Reset(SavedContainerNames.Num());
-
-	for (const FString& ContainerName : SavedContainerNames)
-	{
-		FCustomizableObjectStreamedResourceData* ResourceData = FindStreamedResourceData(ModelResources->StreamedResourceData, ContainerName);
-		if (!ResourceData)
-		{
-			UE_LOG(LogMutable, Error, TEXT("Couldn't find streamed Resource Data container with name %s in array of %d entries"),
-				*ContainerName, ModelResources->StreamedResourceData.Num());
-
-			continue;
-		}
-
-		UCustomizableObjectResourceDataContainer* Container = nullptr;
-		EMoveContainerError Error = MoveContainerToNewOuter(Object, ResourceData, Container);
-		UE_CLOG(Error != EMoveContainerError::None, LogMutable, Warning, TEXT("Failed to move container %s back to %s - %s"), *ContainerName, *GetPathNameSafe(Object), LexToString(Error));
-
-		NewArray.Emplace(Container);
-	}
-
-	ModelResources->StreamedResourceData = NewArray;
-
-	NewArray.Reset(SavedExtensionContainerNames.Num());
-
-	for (const FString& ContainerName : SavedExtensionContainerNames)
-	{
-		FCustomizableObjectStreamedResourceData* ResourceData = FindStreamedResourceData(ModelResources->StreamedExtensionData, ContainerName);
-		if (!ResourceData)
-		{
-			UE_LOG(LogMutable, Error, TEXT("Couldn't find streamed Extension Data container with name %s in array of %d entries"),
-				*ContainerName, ModelResources->StreamedExtensionData.Num());
-
-			continue;
-		}
-
-		UCustomizableObjectResourceDataContainer* Container = nullptr;
-		EMoveContainerError Error = MoveContainerToNewOuter(Object, ResourceData, Container);
-		UE_CLOG(Error != EMoveContainerError::None, LogMutable, Warning, TEXT("Failed to move container %s back to %s - %s"), *ContainerName, *GetPathNameSafe(Object), LexToString(Error));
 
-		NewArray.Emplace(Container);
-	}
-
-	ModelResources->StreamedExtensionData = NewArray;
-}
-
-bool FCustomizableObjectCookPackageSplitter::PopulateGeneratedPackage(FPopulateContext& PopulateContext)
+ICookPackageSplitter::EGeneratedRequiresGenerator FCustomizableObjectCookPackageSplitter::DoesGeneratedRequireGenerator()
 {
-	// Move the container into its newly generated package
-
-	const FGeneratedPackageForPopulate& GeneratedPackage = *PopulateContext.GetTargetGeneratedPackage();
-	UCustomizableObject* Object = CastChecked<UCustomizableObject>(PopulateContext.GetOwnerObject());
-	UModelResources* ModelResources = FindModelResources(*Object);
-
-	FCustomizableObjectStreamedResourceData* ResourceData = FindStreamedResourceData(ModelResources->StreamedResourceData, GeneratedPackage.RelativePath);
-	if (!ResourceData)
-	{
-		ResourceData = FindStreamedResourceData(ModelResources->StreamedExtensionData, GeneratedPackage.RelativePath);
-	}
-
-	if (!ResourceData)
-	{
-		UE_LOG(LogMutable, Error, TEXT("Couldn't find streamed resource Data container with name %s in arrays of %d and %d entries"),
-			*GeneratedPackage.RelativePath, ModelResources->StreamedResourceData.Num(), ModelResources->StreamedExtensionData.Num());
-
-		return false;
-	}
-
-	// [TEMP] Loading a package referencing the CO before PostSaveGeneratedPackage is called causes a name collision.
-	// Duplicate the object with the new outer instead of moving it until it is fixed.
-	UObject* Container = UE::Mutable::Private::LoadObject(ResourceData->GetPath());
-	EMoveContainerError Error = Container ? EMoveContainerError::None : EMoveContainerError::FailedToLoadContainer;
-	if (Container)
-	{
-		Container = StaticDuplicateObject(Container, GeneratedPackage.Package);
-	}
-
-	//UCustomizableObjectResourceDataContainer* Container = nullptr;
-	//EMoveContainerError Error = MoveContainerToNewOuter(GeneratedPackage.Package, ResourceData, Container);
-
-	if (Error != EMoveContainerError::None)
-	{
-		UE_LOG(LogMutable, Error, TEXT("Failed to move container %s to new outer %s - %s"), *ResourceData->GetPath().ToSoftObjectPath().ToString(), *GetPathNameSafe(GeneratedPackage.Package), LexToString(Error));
-		return false;
-	}
-
-	PopulateContext.ReportObjectToMove(Container);
-
-	return true;
+	return EGeneratedRequiresGenerator::Save;
 }
 
-void FCustomizableObjectCookPackageSplitter::PostSaveGeneratedPackage(FPopulateContext& PopulateContext)
-{
-	// Now that the generated package has been saved/cooked, move the container back to the CO, so
-	// that everything is the same as it was before cooking.
-
-	const FGeneratedPackageForPopulate& GeneratedPackage = *PopulateContext.GetTargetGeneratedPackage();
-	UCustomizableObject* Object = CastChecked<UCustomizableObject>(PopulateContext.GetOwnerObject());
-	UModelResources* ModelResources = FindModelResources(*Object);
-
-	FCustomizableObjectStreamedResourceData* ResourceData = FindStreamedResourceData(ModelResources->StreamedResourceData, GeneratedPackage.RelativePath);
-	if (!ResourceData)
-	{
-		ResourceData = FindStreamedResourceData(ModelResources->StreamedExtensionData, GeneratedPackage.RelativePath);
-	}
-
-	if (!ResourceData)
-	{
-		UE_LOG(LogMutable, Error, TEXT("Couldn't find streamed resource Data container with name %s in arrays of %d and %d entries"),
-			*GeneratedPackage.RelativePath, ModelResources->StreamedResourceData.Num(), ModelResources->StreamedExtensionData.Num());
-
-		return;
-	}
 
-	UCustomizableObjectResourceDataContainer* Container = nullptr;
-	EMoveContainerError Error = MoveContainerToNewOuter(Object, ResourceData, Container);
-	UE_CLOG(Error != EMoveContainerError::None, LogMutable, Warning, 
-		TEXT("Failed to move container %s back to %s - %s"), *ResourceData->GetPath().ToSoftObjectPath().ToString(), *GetPathNameSafe(Object), LexToString(Error));
+bool FCustomizableObjectCookPackageSplitter::UseInternalReferenceToAvoidGarbageCollect()
+{
+	return true;
 }
 
-void FCustomizableObjectCookPackageSplitter::Teardown(ETeardown Status)
+bool FCustomizableObjectCookPackageSplitter::RequiresEmptyPackageBeforePopulate(const UPackage* Package)
 {
-	StrongObject.Reset();
+	return false;
 }
diff --git a/Engine/Plugins/Mutable/Source/CustomizableObjectEditor/Private/MuCOE/CustomizableObjectCookPackageSplitter.h b/Engine/Plugins/Mutable/Source/CustomizableObjectEditor/Private/MuCOE/CustomizableObjectCookPackageSplitter.h
--- a/Engine/Plugins/Mutable/Source/CustomizableObjectEditor/Private/MuCOE/CustomizableObjectCookPackageSplitter.h
+++ b/Engine/Plugins/Mutable/Source/CustomizableObjectEditor/Private/MuCOE/CustomizableObjectCookPackageSplitter.h
@@ -3,52 +3,21 @@
 #pragma once
 
 #include "CookPackageSplitter.h"
-#include "UObject/StrongObjectPtr.h"
 #include "MuCO/CustomizableObject.h"
 
 /** Handles splitting the streamable Data constants into their own packages */
 class FCustomizableObjectCookPackageSplitter : public ICookPackageSplitter
 {
 public:
-	/** ICookPackageSplitter interface */
+	// REGISTER_COOKPACKAGE_SPLITTER interface
 	static bool ShouldSplit(UObject* SplitData);
-	static FString GetSplitterDebugName() { return TEXT("FCustomizableObjectCookPackageSplitter"); }
-	static bool RequiresCachedCookedPlatformDataBeforeSplit() { return true; }
-
-	virtual ICookPackageSplitter::FGenerationManifest ReportGenerationManifest(
-		const UPackage* OwnerPackage,
-		const UObject* OwnerObject) override;
-	
-	virtual bool PreSaveGeneratorPackage(FPopulateContext& PopulateContext) override;
-
-	virtual void PostSaveGeneratorPackage(FPopulateContext& PopulateContext) override;
-
-	virtual bool PopulateGeneratedPackage(FPopulateContext& PopulateContext) override;
-
-	virtual void PostSaveGeneratedPackage(FPopulateContext& PopulateContext) override;
-
-	/** Do teardown actions after all packages have saved, or when the cook is cancelled. Always called before destruction. */
-	virtual void Teardown(ETeardown Status) override;
-
-	/**
-	 * If true, this splitter forces the Generator package objects it needs to remain referenced, and the cooker
-	 * should expect them to still be in memory after a garbage collect so long as the splitter is alive.
-	 */
-	virtual bool UseInternalReferenceToAvoidGarbageCollect() override { return true; }
-	/**
-	 * Return capability setting which indicates which splitter functions acting on the parent generator package must
-	 * be called on the splitter before splitter functions acting on the generated packages can be called. 
-	 */
-	virtual EGeneratedRequiresGenerator DoesGeneratedRequireGenerator() override
-	{
-		return EGeneratedRequiresGenerator::Save;
-	}
-
-private:
-
-	TArray<FString> SavedContainerNames;
-	TArray<FString> SavedExtensionContainerNames;
-
-	// Keep a strong reference to the CO to protect it from garbage collector.
-	TStrongObjectPtr<const UObject> StrongObject;
+	static FString GetSplitterDebugName();
+	static bool RequiresCachedCookedPlatformDataBeforeSplit();
+
+	// ICookPackageSplitter interface
+	virtual FGenerationManifest ReportGenerationManifest(const UPackage* OwnerPackage, const UObject* OwnerObject) override;
+	virtual bool PopulateGeneratorPackage(FPopulateContext& PopulateContext) override;
+	virtual EGeneratedRequiresGenerator DoesGeneratedRequireGenerator() override;
+	virtual bool UseInternalReferenceToAvoidGarbageCollect() override;
+	virtual bool RequiresEmptyPackageBeforePopulate(const UPackage* Package) override;
 };
diff --git a/Engine/Plugins/Mutable/Source/CustomizableObjectEditor/Private/MuCOE/CustomizableObjectInstanceBakingUtils.cpp b/Engine/Plugins/Mutable/Source/CustomizableObjectEditor/Private/MuCOE/CustomizableObjectInstanceBakingUtils.cpp
--- a/Engine/Plugins/Mutable/Source/CustomizableObjectEditor/Private/MuCOE/CustomizableObjectInstanceBakingUtils.cpp
+++ b/Engine/Plugins/Mutable/Source/CustomizableObjectEditor/Private/MuCOE/CustomizableObjectInstanceBakingUtils.cpp
@@ -1066,29 +1066,6 @@ bool BakeCustomizableObjectInstance(
 				}
 			}
 
-			// Add Instance Info in a custom AssetUserData
-			{
-				const FCustomizableInstanceComponentData& ComponentData = InInstance.GetPrivate()->GetComponentData(ComponentId);
-				
-				if (InInstance.GetAnimationGameplayTags().Num() ||
-					ComponentData.AnimSlotToBP.Num())
-				{
-					UCustomizableObjectInstanceUserData* InstanceData = NewObject<UCustomizableObjectInstanceUserData>(SkeletalMesh, NAME_None, RF_Public | RF_Transactional);
-					InstanceData->AnimationGameplayTag = InInstance.GetAnimationGameplayTags();
-
-					for (const TTuple<FName, TSoftClassPtr<UAnimInstance>>& AnimSlot : ComponentData.AnimSlotToBP)
-					{
-						FCustomizableObjectAnimationSlot AnimationSlot;
-						AnimationSlot.Name = AnimSlot.Key;
-						AnimationSlot.AnimInstance = AnimSlot.Value;
-				
-						InstanceData->AnimationSlots.Add(AnimationSlot);
-					}
-
-					SkeletalMesh->AddAssetUserData(InstanceData);
-				}
-			}
-
 			// Copy LODSettings from the Reference Skeletal Mesh
 			{
 				if (ModelResources->ReferenceSkeletalMeshesData.IsValidIndex(ComponentId))
diff --git a/Engine/Plugins/Mutable/Source/CustomizableObjectEditor/Private/MuCOE/ExtensionDataCompilerInterface.cpp b/Engine/Plugins/Mutable/Source/CustomizableObjectEditor/Private/MuCOE/ExtensionDataCompilerInterface.cpp
--- a/Engine/Plugins/Mutable/Source/CustomizableObjectEditor/Private/MuCOE/ExtensionDataCompilerInterface.cpp
+++ b/Engine/Plugins/Mutable/Source/CustomizableObjectEditor/Private/MuCOE/ExtensionDataCompilerInterface.cpp
@@ -2,8 +2,6 @@
 
 #include "MuCOE/ExtensionDataCompilerInterface.h"
 
-#include "StructUtils/InstancedStruct.h"
-#include "MuCO/CustomizableObjectStreamedResourceData.h"
 #include "MuCOE/GenerateMutableSource/GenerateMutableSource.h"
 #include "MuCOE/CustomizableObjectCompiler.h"
 #include "MuR/ExtensionData.h"
@@ -13,35 +11,6 @@ FExtensionDataCompilerInterface::FExtensionDataCompilerInterface(FMutableGraphGe
 {
 }
 
-TSharedPtr<const UE::Mutable::Private::FExtensionData> FExtensionDataCompilerInterface::MakeStreamedExtensionData(FInstancedStruct&& Data)
-{
-	TSharedPtr<UE::Mutable::Private::FExtensionData> Result = MakeShared<UE::Mutable::Private::FExtensionData>();
-	Result->Origin = UE::Mutable::Private::FExtensionData::EOrigin::ConstantStreamed;
-	Result->Index = GenerationContext.StreamedExtensionData.Num();
-	
-	FCustomizableObjectResourceData* CompileTimeExtensionData = &GenerationContext.StreamedExtensionData.AddDefaulted_GetRef();
-	CompileTimeExtensionData->Data = MoveTemp(Data);
-
-	return Result;
-}
-
-TSharedPtr<const UE::Mutable::Private::FExtensionData> FExtensionDataCompilerInterface::MakeAlwaysLoadedExtensionData(FInstancedStruct&& Data)
-{
-	TSharedPtr<UE::Mutable::Private::FExtensionData> Result = MakeShared<UE::Mutable::Private::FExtensionData>();
-	Result->Origin = UE::Mutable::Private::FExtensionData::EOrigin::ConstantAlwaysLoaded;
-	Result->Index = GenerationContext.AlwaysLoadedExtensionData.Num();
-
-	FCustomizableObjectResourceData* CompileTimeExtensionData = &GenerationContext.AlwaysLoadedExtensionData.AddDefaulted_GetRef();
-	CompileTimeExtensionData->Data = MoveTemp(Data);
-
-	return Result;
-}
-
-const UObject* FExtensionDataCompilerInterface::GetOuterForAlwaysLoadedObjects()
-{
-	check(GenerationContext.CompilationContext->Object);
-	return GenerationContext.CompilationContext->Object.Get();
-}
 
 void FExtensionDataCompilerInterface::AddGeneratedNode(const UCustomizableObjectNode* InNode)
 {
@@ -54,6 +23,13 @@ void FExtensionDataCompilerInterface::AddGeneratedNode(const UCustomizableObject
 	GenerationContext.GeneratedNodes.Add(const_cast<UCustomizableObjectNode*>(InNode));
 }
 
+
+UE::Mutable::Private::PASSTHROUGH_ID FExtensionDataCompilerInterface::MakeExtensionData(UObject& ExtensionData, bool bDuplicate)
+{
+	return GenerationContext.CompilationContext->PassthroughObjectFactory.Add(ExtensionData, bDuplicate);
+}
+
+
 void FExtensionDataCompilerInterface::CompilerLog(const FText& InLogText, const UCustomizableObjectNode* InNode)
 {
 	GenerationContext.Log(InLogText, InNode);
diff --git a/Engine/Plugins/Mutable/Source/CustomizableObjectEditor/Private/MuCOE/GenerateMutableSource/GenerateMutableSource.cpp b/Engine/Plugins/Mutable/Source/CustomizableObjectEditor/Private/MuCOE/GenerateMutableSource/GenerateMutableSource.cpp
--- a/Engine/Plugins/Mutable/Source/CustomizableObjectEditor/Private/MuCOE/GenerateMutableSource/GenerateMutableSource.cpp
+++ b/Engine/Plugins/Mutable/Source/CustomizableObjectEditor/Private/MuCOE/GenerateMutableSource/GenerateMutableSource.cpp
@@ -263,31 +263,6 @@ UE::Mutable::Private::TManagedPtr<UE::Mutable::Private::FMesh> FMutableGraphGene
 }
 
 
-int32 FMutableCompilationContext::AddStreamedResource(UAssetUserData& AssetUserData)
-{
-	const uint32 AssetIdentifier = AssetUserData.GetUniqueID();
-
-	// Return resource index if found.
-	if (int32* ResourceIndex = StreamedResourceIndices.Find(AssetIdentifier))
-	{
-		return *ResourceIndex;
-	}
-
-	FCustomizableObjectAssetUserData ResourceData;
-	ResourceData.AssetUserDataEditor = &AssetUserData;
-
-	const int32 ResourceIndex = StreamedResourceData.Num();
-
-	FCustomizableObjectResourceData& Resource = StreamedResourceData.AddDefaulted_GetRef();
-	Resource.Type = ECOResourceDataType::AssetUserData;
-	Resource.Data.InitializeAs<FCustomizableObjectAssetUserData>(MoveTemp(ResourceData));
-
-	StreamedResourceIndices.Add({ AssetIdentifier, ResourceIndex });
-
-	return ResourceIndex;
-}
-
-
 const FGuid FMutableGraphGenerationContext::GetNodeIdUnique(const UCustomizableObjectNode* Node)
 {
 	FGuid NodeID = GetNodeIdUnchecked(Node);
@@ -1060,7 +1035,7 @@ UE::Mutable::Private::Ptr<UE::Mutable::Private::NodeObject> GenerateMutableSourc
 				{
 					if (UE::Mutable::Private::Ptr<UE::Mutable::Private::NodeExtensionData> GeneratedNode = ExtensionNode->GenerateMutableNode(GenerationContext.ExtensionDataCompilerInterface))
 					{
-						ObjectNode->AddExtensionDataNode(GeneratedNode, TCHAR_TO_ANSI(*ExtensionInputPin.GlobalPinName.ToString()));
+						ObjectNode->AddExtensionDataNode(GeneratedNode);
 					}
 				}
 			}
@@ -1420,36 +1395,6 @@ UE::Mutable::Private::Ptr<UE::Mutable::Private::NodeObject> GenerateMutableSourc
 }
 
 
-int32 AddTagToMutableMeshUnique(UE::Mutable::Private::FMesh& MutableMesh, const FString& Tag)
-{
-	const int32 TagCount = MutableMesh.GetTagCount();
-
-	for (int32 TagIndex = TagCount - 1; TagIndex >= 0; --TagIndex)
-	{
-		if (FString(MutableMesh.GetTag(TagIndex)) == Tag)
-		{
-			return TagIndex;
-		}
-	}
-
-	MutableMesh.SetTagCount(TagCount + 1);
-	MutableMesh.SetTag(TagCount, Tag);
-
-	return TagCount;
-}
-
-FString GenerateAnimationInstanceTag(const int32 AnimBpIndex, const FName& SlotIndex)
-{
-	return FString("__AnimBP:") + FString::Printf(TEXT("%s_Slot_"), *SlotIndex.ToString()) + FString::FromInt(AnimBpIndex);
-}
-
-
-FString GenerateGameplayTag(const FString& GameplayTag)
-{
-	return FString("__AnimBPTag:") + GameplayTag;
-}
-
-
 void PopulateReferenceSkeletalMeshesData(FMutableGraphGenerationContext& GenerationContext)
 {
 	const FString PlatformName = GenerationContext.CompilationContext->Options.TargetPlatform->IniPlatformName();
@@ -1569,18 +1514,6 @@ void PopulateReferenceSkeletalMeshesData(FMutableGraphGenerationContext& Generat
 		{
 			Data.ShadowPhysicsAsset = PhysicsAsset;
 		}
-
-		// Asset User Data
-		if (const TArray<UAssetUserData*>* AssetUserDataArray = RefSkeletalMesh->GetAssetUserDataArray())
-		{
-			for (UAssetUserData* AssetUserData : *AssetUserDataArray)
-			{
-				if (AssetUserData)
-				{
-					Data.AssetUserDataIndices.AddUnique(GenerationContext.CompilationContext->AddStreamedResource(*AssetUserData));
-				}
-			}
-		}
 	}
 }
 
@@ -1766,7 +1699,7 @@ UE::Mutable::Private::TManagedPtr<UE::Mutable::Private::FImage> GenerateImageCon
 
 	if (bIsReference)
 	{
-		Id = GenerationContext.CompilationContext->PassthroughObjectFactory.Add(*Texture);
+		Id = GenerationContext.CompilationContext->PassthroughObjectFactory.Add(*Texture, false);
 	}
 	else if (bIsCompileTime)
 	{
diff --git a/Engine/Plugins/Mutable/Source/CustomizableObjectEditor/Private/MuCOE/GenerateMutableSource/GenerateMutableSource.h b/Engine/Plugins/Mutable/Source/CustomizableObjectEditor/Private/MuCOE/GenerateMutableSource/GenerateMutableSource.h
--- a/Engine/Plugins/Mutable/Source/CustomizableObjectEditor/Private/MuCOE/GenerateMutableSource/GenerateMutableSource.h
+++ b/Engine/Plugins/Mutable/Source/CustomizableObjectEditor/Private/MuCOE/GenerateMutableSource/GenerateMutableSource.h
@@ -519,10 +519,6 @@ struct FMutableCompilationContext
 
 	TArray<FAnimBpOverridePhysicsAssetsInfo> AnimBpOverridePhysicsAssetsInfo;
 
-	/** Resource Data constants */
-	TMap<uint32, int32> StreamedResourceIndices;
-	TArray<FCustomizableObjectResourceData> StreamedResourceData;
-
 	/** Data used for SkinWeightProfiles reconstruction. */
 	TArray<FMutableSkinWeightProfileInfo> SkinWeightProfilesInfo;
 
@@ -565,11 +561,6 @@ public:
 
 	/** */
 	uint32 GetSkinWeightProfileIdUnique(const FName ProfileName);
-
-	/** Adds a streamed resource of type AssetUserData.
-	 * Returns resource index in the array of streamed resources. */
-	int32 AddStreamedResource(UAssetUserData& AssetUserData);
-
 };
 
 
@@ -727,22 +718,27 @@ struct FMutableGraphGenerationContext
 			/** Flag used to generate this mesh. Bit mask of EMutableMeshConversionFlags */
 			EMutableMeshConversionFlags Flags = EMutableMeshConversionFlags::None;
 
-			/** Tags added at the UE level that go through the Mutable core and are merged in the generated mesh.
-			 *  Only add the tags that make the mesh unique and require it not to be cached together with the 
-			 *  same exact mesh but with different tags.
-			*/
-			FString Tags;
-
 			/**
 			* SkeletalMeshNode is needed to disambiguate realtime morph selection from different nodes.
 			* TODO: Consider using the actual selection.
 			*/
 			const UCustomizableObjectNode* SkeletalMeshNode = nullptr;
 
+			FGameplayTagContainer GameplayTags;
+			FName AnimBPSlotName;
+			TSoftClassPtr<UAnimInstance> AnimInstance;
+
 			bool operator==( const FKey& OtherKey ) const
 			{
-				return Mesh == OtherKey.Mesh && LOD == OtherKey.LOD && CurrentLOD == OtherKey.CurrentLOD && MaterialIndex == OtherKey.MaterialIndex
-					&& Flags == OtherKey.Flags && Tags == OtherKey.Tags && SkeletalMeshNode == OtherKey.SkeletalMeshNode;
+				return Mesh == OtherKey.Mesh &&
+					LOD == OtherKey.LOD && 
+					CurrentLOD == OtherKey.CurrentLOD &&
+					MaterialIndex == OtherKey.MaterialIndex &&
+					Flags == OtherKey.Flags &&
+					GameplayTags == OtherKey.GameplayTags &&
+					AnimBPSlotName == OtherKey.AnimBPSlotName &&
+					AnimInstance == OtherKey.AnimInstance &&
+					SkeletalMeshNode == OtherKey.SkeletalMeshNode;
 			}
 		};
 
@@ -898,14 +894,11 @@ struct FMutableGraphGenerationContext
 	// Current material parameter id to find the corresponding column in a mutable table
 	FString CurrentMaterialParameterId;
 
-
 	// Material to SharedSurfaceId
 	TMap<FSharedSurfaceKey, FGuid> SurfaceGuids;
 
 	/** Extension Data constants are collected here */
 	FExtensionDataCompilerInterface ExtensionDataCompilerInterface;
-	TArray<FCustomizableObjectResourceData> AlwaysLoadedExtensionData;
-	TArray<FCustomizableObjectResourceData> StreamedExtensionData;
 
 	/** Map to relate a Composite Data Table Row and its original DataTable */
 	TMap<UDataTable*,TMap<FName, TArray<UDataTable*>>> CompositeDataTableRowToOriginalDataTableMap;
@@ -1006,15 +999,6 @@ UTexture2D* FindReferenceImage(const UEdGraphPin* Pin, FMutableGraphGenerationCo
 UE::Mutable::Private::Ptr<UE::Mutable::Private::NodeMeshApplyPose> CreateNodeMeshApplyPose(FMutableGraphGenerationContext& GenerationContext, UE::Mutable::Private::NodeMeshPtr InputMeshNode, const TArray<FName>& ArrayBoneName, const TArray<FTransform>& ArrayTransform);
 
 
-/** Adds Tag to MutableMesh uniquely, returns the index were the tag has been inserted or the index where an intance of the tag has been found */
-int32 AddTagToMutableMeshUnique(UE::Mutable::Private::FMesh& MutableMesh, const FString& Tag);
-
-// Generates the tag for an animation instance
-FString GenerateAnimationInstanceTag(const int32 AnimInstanceIndex, const FName& SlotIndex);
-
-
-FString GenerateGameplayTag(const FString& GameplayTag);
-
 uint32 GetBaseTextureSize(const FMutableGraphGenerationContext& GenerationContext, const UCONodeSkeletalMeshSectionBase* Material, const UE::Mutable::Private::FParameterKey& ImageKey);
 
 // Computes the LOD bias for a texture given the current mesh LOD and automatic LOD settings, the reference texture settings
diff --git a/Engine/Plugins/Mutable/Source/CustomizableObjectEditor/Private/MuCOE/GenerateMutableSource/GenerateMutableSourceMaterial.cpp b/Engine/Plugins/Mutable/Source/CustomizableObjectEditor/Private/MuCOE/GenerateMutableSource/GenerateMutableSourceMaterial.cpp
--- a/Engine/Plugins/Mutable/Source/CustomizableObjectEditor/Private/MuCOE/GenerateMutableSource/GenerateMutableSourceMaterial.cpp
+++ b/Engine/Plugins/Mutable/Source/CustomizableObjectEditor/Private/MuCOE/GenerateMutableSource/GenerateMutableSourceMaterial.cpp
@@ -155,7 +155,7 @@ UE::Mutable::Private::Ptr<UE::Mutable::Private::NodeMaterial> GenerateMutableSou
 		if (!BaseMaterial->MaterialSource)
 		{
 			UE::Mutable::Private::Ptr<UE::Mutable::Private::NodeMaterialConstant> MaterialNode = new UE::Mutable::Private::NodeMaterialConstant();
-			MaterialNode->MaterialId = GenerationContext.CompilationContext->PassthroughObjectFactory.Add(*GenerationContext.CurrentReferencedMaterial.Get());
+			MaterialNode->MaterialId = GenerationContext.CompilationContext->PassthroughObjectFactory.Add(*GenerationContext.CurrentReferencedMaterial.Get(), false);
 
 			BaseMaterial->MaterialSource = MaterialNode;
 		}
@@ -393,7 +393,7 @@ UE::Mutable::Private::Ptr<UE::Mutable::Private::NodeMaterial> GenerateMutableSou
 		{
 			MaterialNode = new UE::Mutable::Private::NodeMaterialConstant();
 			GenerationContext.CurrentReferencedMaterial = TStrongObjectPtr(Material);
-			MaterialNode->MaterialId = GenerationContext.CompilationContext->PassthroughObjectFactory.Add(*Material);
+			MaterialNode->MaterialId = GenerationContext.CompilationContext->PassthroughObjectFactory.Add(*Material, false);
 
 			// Store the parameter that needs to be processed
 			const FMaterialBreakParameter CurrentBreakParameter = GenerationContext.GetCurrentMaterialBreakParameter();
@@ -835,7 +835,7 @@ UE::Mutable::Private::Ptr<UE::Mutable::Private::NodeMaterial> GenerateMutableSou
 	{
 		//This node will add a default value in case of error
 		UE::Mutable::Private::Ptr<UE::Mutable::Private::NodeMaterialConstant> ConstantValue = new UE::Mutable::Private::NodeMaterialConstant();
-		ConstantValue->MaterialId = GenerationContext.CompilationContext->PassthroughObjectFactory.Add(*GenerationContext.CurrentReferencedMaterial.Get());
+		ConstantValue->MaterialId = GenerationContext.CompilationContext->PassthroughObjectFactory.Add(*GenerationContext.CurrentReferencedMaterial.Get(), false);
 
 		Result = ConstantValue;
 
diff --git a/Engine/Plugins/Mutable/Source/CustomizableObjectEditor/Private/MuCOE/GenerateMutableSource/GenerateMutableSourceMesh.cpp b/Engine/Plugins/Mutable/Source/CustomizableObjectEditor/Private/MuCOE/GenerateMutableSource/GenerateMutableSourceMesh.cpp
--- a/Engine/Plugins/Mutable/Source/CustomizableObjectEditor/Private/MuCOE/GenerateMutableSource/GenerateMutableSourceMesh.cpp
+++ b/Engine/Plugins/Mutable/Source/CustomizableObjectEditor/Private/MuCOE/GenerateMutableSource/GenerateMutableSourceMesh.cpp
@@ -62,6 +62,7 @@
 #include "PackedNormal.h"
 #include "MuCOE/Nodes/CONodeClipMeshWithMesh.h"
 #include "MuCOE/Nodes/CONodeTransformWithBone.h"
+#include "MuCO/CustomizableObjectInstanceAssetUserData.h"
 #include "MuCOE/Nodes/CONodeExternalOperation.h"
 #include "MuCOE/Nodes/CONodeRemoveMesh.h"
 #include "MuT/NodeMeshClipWithMesh.h"
@@ -1050,36 +1051,32 @@ UE::Tasks::FTask ConvertSkeletalMeshToMutable(TSharedRef<FMeshConversionContext>
 
 	UE::Mutable::Private::TManagedPtr<UE::Mutable::Private::FMesh> MutableMesh = UE::Mutable::Private::MakeManaged<UE::Mutable::Private::FMesh>();
 	TSharedPtr<FMutableConvertSkeletalMeshContext> MeshContext = MakeShared<FMutableConvertSkeletalMeshContext>();
-
-	if (UCustomizableObjectSystem::GetInstance()->IsMutableAnimInfoDebuggingEnabled())
-	{
-		FString MeshTag = FString("__MeshPath:") + Source.Mesh->GetPathName();;
-		AddTagToMutableMeshUnique(*MutableMesh, MeshTag);
-	}
-
-	if (!bIgnoreAUD && Context.Object->GetPrivate()->IsAssetUserDataMergeEnabled())
+	
+	if (!bIgnoreAUD)
 	{
-		if (const TArray<UAssetUserData*>* AssetUserDataArray = SkeletalMesh->GetAssetUserDataArray())
+		for (UAssetUserData* AssetUserData : *SkeletalMesh->GetAssetUserDataArray())
 		{
-			for (UAssetUserData* AssetUserData : *AssetUserDataArray)
+			if (AssetUserData)
 			{
-				if (!AssetUserData)
-				{
-					continue;
-				}
-
-				const int32 ResourceIndex = Context.AddStreamedResource(*AssetUserData);
-				if (ResourceIndex >= 0)
-				{
-					FCustomizableObjectStreameableResourceId ResourceId;
-					ResourceId.Id = (uint32)ResourceIndex;
-					ResourceId.Type = static_cast<uint8>(FCustomizableObjectStreameableResourceId::EType::AssetUserData);
-
-					MutableMesh->AddStreamedResource(BitCast<uint64>(ResourceId));
-				}
+				MutableMesh->AssetUserData.AddUnique(Context.PassthroughObjectFactory.Add(*AssetUserData, true));
 			}
 		}
 	}
+	
+	const int32 NumGameplayTags = Source.GameplayTags.Num();
+	
+	TArray<FName> GameplayTags;
+	GameplayTags.Reserve(NumGameplayTags);
+	for (int32 Index = 0; Index < NumGameplayTags; ++Index)
+	{
+		GameplayTags.Add(Source.GameplayTags.GetByIndex(Index).GetTagName());
+	}
+	
+	MutableMesh->GameplayTags = GameplayTags;
+	if (UClass* AnimInstance = UE::Mutable::Private::LoadClass(Source.AnimInstance))
+	{
+		MutableMesh->AnimationSlots.Emplace(Source.AnimBPSlotName, Context.PassthroughObjectFactory.Add(*AnimInstance, false));
+	}
 
 	bool bBoneMapModified = false;
 	TArray<FBoneIndexType> BoneMap;
@@ -1119,7 +1116,7 @@ UE::Tasks::FTask ConvertSkeletalMeshToMutable(TSharedRef<FMeshConversionContext>
 			}
 
 			// Add the skeleton object
-			MutableMesh->SkeletonObjects.Emplace(Context.PassthroughObjectFactory.Add(*const_cast<USkeleton*>(Skeleton)));
+			MutableMesh->SkeletonObjects.AddUnique(Context.PassthroughObjectFactory.Add(*const_cast<USkeleton*>(Skeleton), false));
 		}
 
 		// RefSkeleton check
@@ -1418,7 +1415,7 @@ UE::Tasks::FTask ConvertSkeletalMeshToMutable(TSharedRef<FMeshConversionContext>
 		if (PhysicsBody)
 		{
 			MutableMesh->SetPhysicsBody(PhysicsBody);
-			MutableMesh->PhysicsAssets.Emplace(Context.PassthroughObjectFactory.Add(PhysicsAsset));
+			MutableMesh->PhysicsAssets.Emplace(Context.PassthroughObjectFactory.Add(PhysicsAsset, false));
 
 			// Set Physics usage flag
 			const int32 NumBodySetups = PhysicsBody->GetBodyCount();
@@ -1460,7 +1457,7 @@ UE::Tasks::FTask ConvertSkeletalMeshToMutable(TSharedRef<FMeshConversionContext>
 			MutableBody->CustomId = PhysicsAssetId;
 
 			MutableMesh->AddAdditionalPhysicsBody(MutableBody);
-			MutableMesh->AdditionalPhysicsAssets.Emplace(Context.PassthroughObjectFactory.Add(*PropertyAsset));
+			MutableMesh->AdditionalPhysicsAssets.Emplace(Context.PassthroughObjectFactory.Add(*PropertyAsset, false));
 		}
 	}
 	
@@ -1518,7 +1515,7 @@ UE::Tasks::FTask ConvertSkeletalMeshToMutable(TSharedRef<FMeshConversionContext>
 		UClothingAssetCommon* ClothingAsset = Cast<UClothingAssetCommon>(SkeletalMesh->GetClothingAsset(Section.ClothingData.AssetGuid));
 		if (bHasClothData && ClothingAsset)
 		{
-			MutableMesh->ClothSections[0].ClothingAsset = Context.PassthroughObjectFactory.Add(*ClothingAsset);
+			MutableMesh->ClothSections[0].ClothingAsset = Context.PassthroughObjectFactory.Add(*ClothingAsset, true);
 
 			if (MeshSection.ClothMappingDataLODs.Num() > 1)
 			{
@@ -1553,7 +1550,7 @@ UE::Mutable::Private::TManagedPtr<UE::Mutable::Private::FMesh> ConvertSkeletalMe
 	}
 
 	// Prepare the data that is needed for the core mesh conversion.
-	Source.BaseLODIndex= LODIndex;
+	Source.BaseLODIndex = LODIndex;
 	Source.BaseSectionIndex = SectionIndex;
 	Source.LODOffset = GenerationContext.CurrentLOD - GenerationContext.FromLOD;
 	Source.Flags = GenerationContext.MeshGenerationFlags.Last();
@@ -1770,15 +1767,18 @@ UE::Mutable::Private::TManagedPtr<UE::Mutable::Private::FMesh> ConvertStaticMesh
 
 // Convert a Mesh constant to a mutable format. UniqueTags are the tags that make this Mesh unique that cannot be merged in the cache 
 // with the exact same Mesh with other tags
-UE::Mutable::Private::TManagedPtr<UE::Mutable::Private::FMesh> GenerateMutableSkeletalMesh(FMutableSourceMeshData& Source, int32 LODIndexConnected, int32 SectionIndexConnected,
-										const FString& UniqueTags, FMutableGraphGenerationContext & GenerationContext, 
-										const UCustomizableObjectNode* CurrentNode,
-										const bool bOnlyConnectedLOD)
+UE::Mutable::Private::TManagedPtr<UE::Mutable::Private::FMesh> GenerateMutableSkeletalMesh(
+	FMutableSourceMeshData& Source,
+	int32 LODIndexConnected,
+	int32 SectionIndexConnected,
+	FMutableGraphGenerationContext & GenerationContext, 
+	const UCustomizableObjectNode* CurrentNode,
+	const bool bOnlyConnectedLOD)
 {
 	// Get the mesh generation flags to use
 	EMutableMeshConversionFlags CurrentFlags = GenerationContext.MeshGenerationFlags.Last();
-
-	FMutableGraphGenerationContext::FGeneratedMeshData::FKey Key = { Source.Mesh, LODIndexConnected, GenerationContext.CurrentLOD, SectionIndexConnected, CurrentFlags, UniqueTags, CurrentNode };
+	
+	FMutableGraphGenerationContext::FGeneratedMeshData::FKey Key = { Source.Mesh, LODIndexConnected, GenerationContext.CurrentLOD, SectionIndexConnected, CurrentFlags, CurrentNode, Source.GameplayTags, Source.AnimBPSlotName, Source.AnimInstance};
 	UE::Mutable::Private::TManagedPtr<UE::Mutable::Private::FMesh> MutableMesh = GenerationContext.FindGeneratedMesh(Key);
 	if (MutableMesh)
 	{
@@ -1799,14 +1799,17 @@ UE::Mutable::Private::TManagedPtr<UE::Mutable::Private::FMesh> GenerateMutableSk
 
 // Convert a Mesh constant to a mutable format. UniqueTags are the tags that make this Mesh unique that cannot be merged in the cache 
 // with the exact same Mesh with other tags
-UE::Mutable::Private::TManagedPtr<UE::Mutable::Private::FMesh> GenerateMutableStaticMesh(TSoftObjectPtr<UStreamableRenderAsset> Mesh, const TSoftClassPtr<UAnimInstance>& AnimInstance,
-	int32 LODIndex, int32 SectionIndex, const FString& UniqueTags, FMutableGraphGenerationContext& GenerationContext,
-	const UCustomizableObjectNode* CurrentNode, USkeletalMesh* TableReferenceSkeletalMesh, bool bIsPassthrough)
+UE::Mutable::Private::TManagedPtr<UE::Mutable::Private::FMesh> GenerateMutableStaticMesh(
+	TSoftObjectPtr<UStreamableRenderAsset> Mesh,
+	int32 LODIndex, 
+	int32 SectionIndex,
+	FMutableGraphGenerationContext& GenerationContext,
+	const UCustomizableObjectNode* CurrentNode)
 {
 	// Get the mesh generation flags to use
 	EMutableMeshConversionFlags CurrentFlags = GenerationContext.MeshGenerationFlags.Last();
 
-	FMutableGraphGenerationContext::FGeneratedMeshData::FKey Key = { Mesh, LODIndex, GenerationContext.CurrentLOD, SectionIndex, CurrentFlags, UniqueTags, CurrentNode };
+	FMutableGraphGenerationContext::FGeneratedMeshData::FKey Key = { Mesh, LODIndex, GenerationContext.CurrentLOD, SectionIndex, CurrentFlags, CurrentNode };
 	UE::Mutable::Private::TManagedPtr<UE::Mutable::Private::FMesh> MutableMesh = GenerationContext.FindGeneratedMesh(Key);
 	if (MutableMesh)
 	{
@@ -1884,7 +1887,7 @@ UE::Mutable::Private::TManagedPtr<UE::Mutable::Private::FMesh> BuildMorphedMutab
 		// Get the base mesh
 		FMutableSourceMeshData Source;
 		Source.Mesh = SkeletalMesh;
-		UE::Mutable::Private::TManagedPtr<UE::Mutable::Private::FMesh> BaseSourceMesh = GenerateMutableSkeletalMesh(Source, BaseLODIndex, BaseSectionIndex, FString(), GenerationContext, Node, bOnlyConnectedLOD);
+		UE::Mutable::Private::TManagedPtr<UE::Mutable::Private::FMesh> BaseSourceMesh = GenerateMutableSkeletalMesh(Source, BaseLODIndex, BaseSectionIndex, GenerationContext, Node, bOnlyConnectedLOD);
 
 		if (BaseSourceMesh)
 		{
@@ -2614,49 +2617,15 @@ UE::Mutable::Private::Ptr<UE::Mutable::Private::NodeMesh> GenerateMutableSourceM
 			int32 BaseSectionIndex = -1;
 			TypedNodeSkel->GetPinSection(*Pin, BaseLODIndex, BaseSectionIndex);
 
-			// First process the mesh tags that are going to make the mesh unique and affect whether it's repeated in 
-			// the mesh cache or not
-			FString MeshUniqueTags;
-			FString AnimBPAssetTag;
-
-			if (!TypedNodeSkel->AnimInstance.IsNull())
-			{
-				FName SlotIndex = TypedNodeSkel->AnimBlueprintSlotName;
-				const int32 AnimInstanceIndex = GenerationContext.AnimBPAssets.AddUnique(TypedNodeSkel->AnimInstance);
-
-				AnimBPAssetTag = GenerateAnimationInstanceTag(AnimInstanceIndex, SlotIndex);
-				MeshUniqueTags += AnimBPAssetTag;
-			}
-
-			TArray<FString> ArrayAnimBPTags;
-
-			for (const FGameplayTag& GamePlayTag : TypedNodeSkel->AnimationGameplayTags)
-			{
-				const FString AnimBPTag = GenerateGameplayTag(GamePlayTag.ToString());
-				ArrayAnimBPTags.Add(AnimBPTag);
-				MeshUniqueTags += AnimBPTag;
-			}
-
 			FMutableSourceMeshData Source = BaseMeshData;
 			Source.Mesh = TypedNodeSkel->SkeletalMesh;
 			Source.AnimInstance = TypedNodeSkel->AnimInstance;
-			UE::Mutable::Private::TManagedPtr<UE::Mutable::Private::FMesh> MutableMesh = GenerateMutableSkeletalMesh(Source, BaseLODIndex, BaseSectionIndex, MeshUniqueTags, GenerationContext, TypedNodeSkel, bOnlyConnectedLOD);
-
-			if (MutableMesh)
-			{
-				MeshNode->Value = MutableMesh;
-
-				if (!TypedNodeSkel->AnimInstance.IsNull())
-				{
-					AddTagToMutableMeshUnique(*MutableMesh, AnimBPAssetTag);
-				}
-
-				for (const FString& GamePlayTag : ArrayAnimBPTags)
-				{
-					AddTagToMutableMeshUnique(*MutableMesh, GamePlayTag);
-				}
-			}
-
+			Source.GameplayTags = TypedNodeSkel->AnimationGameplayTags;
+			Source.AnimBPSlotName = TypedNodeSkel->AnimBlueprintSlotName;
+			
+			UE::Mutable::Private::TManagedPtr<UE::Mutable::Private::FMesh> MutableMesh = GenerateMutableSkeletalMesh(Source, BaseLODIndex, BaseSectionIndex, GenerationContext, TypedNodeSkel, bOnlyConnectedLOD);
+			MeshNode->Value = MutableMesh;
+			
 			// Layouts
 			// 
 			// Generate Layouts when using manual LODs or when processing the base Automatic LOD (the other automatic LODs will use the base layouts).
@@ -2768,8 +2737,7 @@ UE::Mutable::Private::Ptr<UE::Mutable::Private::NodeMesh> GenerateMutableSourceM
 		check(TypedNodeStatic->LODs.IsValidIndex(LODIndex) && SectionIndex < TypedNodeStatic->LODs[LODIndex].Materials.Num());
 
 		constexpr bool bIsPassthrough = false;
-		UE::Mutable::Private::TManagedPtr<UE::Mutable::Private::FMesh> MutableMesh = GenerateMutableStaticMesh(TypedNodeStatic->StaticMesh, TSoftClassPtr<UAnimInstance>(), LODIndex, SectionIndex,
-													  FString(), GenerationContext, TypedNodeStatic, nullptr, bIsPassthrough);
+		UE::Mutable::Private::TManagedPtr<UE::Mutable::Private::FMesh> MutableMesh = GenerateMutableStaticMesh(TypedNodeStatic->StaticMesh,  LODIndex, SectionIndex, GenerationContext, TypedNodeStatic);
 		if (MutableMesh)
 		{
 			MeshNode->Value = MutableMesh;
diff --git a/Engine/Plugins/Mutable/Source/CustomizableObjectEditor/Private/MuCOE/GenerateMutableSource/GenerateMutableSourceMesh.h b/Engine/Plugins/Mutable/Source/CustomizableObjectEditor/Private/MuCOE/GenerateMutableSource/GenerateMutableSourceMesh.h
--- a/Engine/Plugins/Mutable/Source/CustomizableObjectEditor/Private/MuCOE/GenerateMutableSource/GenerateMutableSourceMesh.h
+++ b/Engine/Plugins/Mutable/Source/CustomizableObjectEditor/Private/MuCOE/GenerateMutableSource/GenerateMutableSourceMesh.h
@@ -62,24 +62,26 @@ UE::Mutable::Private::TManagedPtr<UE::Mutable::Private::FMesh> ConvertStaticMesh
  * @param Source. Mesh data required to generate the mutable mesh
  * @param LODIndexConnected LOD which the pin is connected to.
  * @param SectionIndexConnected Section which the pin is connected to.
- * @param MeshUniqueTags
  * @param GenerationContext
  * @param CurrentNode The node being compiled, can be a skeletal mesh or a table node.
  * @param bOnlyConnectedLOD
  * @return Mutable Mesh. Nullptr if there has been an error. Empty mesh if the Skeletal Mesh does not contain the requested LOD + Section. */
-UE::Mutable::Private::TManagedPtr<UE::Mutable::Private::FMesh> GenerateMutableSkeletalMesh(FMutableSourceMeshData& Source,
-                                int32 LODIndexConnected, int32 SectionIndexConnected,
-                                const FString& MeshUniqueTags, 
-                                FMutableGraphGenerationContext& GenerationContext, 
-								const UCustomizableObjectNode* CurrentNode,
-								const bool bOnlyConnectedLOD);
+UE::Mutable::Private::TManagedPtr<UE::Mutable::Private::FMesh> GenerateMutableSkeletalMesh(
+	FMutableSourceMeshData& Source,
+	int32 LODIndexConnected,
+	int32 SectionIndexConnected,
+	FMutableGraphGenerationContext & GenerationContext, 
+	const UCustomizableObjectNode* CurrentNode,
+	const bool bOnlyConnectedLOD);
+
 
 /** */
-UE::Mutable::Private::TManagedPtr<UE::Mutable::Private::FMesh> GenerateMutableStaticMesh(TSoftObjectPtr<UStreamableRenderAsset> Mesh, const TSoftClassPtr<UAnimInstance>& AnimBp,
-	int32 LODIndex, int32 SectionIndex, const FString& MeshUniqueTags,
+UE::Mutable::Private::TManagedPtr<UE::Mutable::Private::FMesh> GenerateMutableStaticMesh(
+	TSoftObjectPtr<UStreamableRenderAsset> Mesh,
+	int32 LODIndex, 
+	int32 SectionIndex,
 	FMutableGraphGenerationContext& GenerationContext,
-	const UCustomizableObjectNode* CurrentNode, USkeletalMesh* TableReferenceSkeletalMesh,
-	bool bIsPassthrough);
+	const UCustomizableObjectNode* CurrentNode);
 
 
 UE::Mutable::Private::TManagedPtr<UE::Mutable::Private::FMesh> BuildMorphedMutableMesh(const UEdGraphPin* BaseSourcePin, const FString& MorphTargetName,
diff --git a/Engine/Plugins/Mutable/Source/CustomizableObjectEditor/Private/MuCOE/GenerateMutableSource/GenerateMutableSourceTable.cpp b/Engine/Plugins/Mutable/Source/CustomizableObjectEditor/Private/MuCOE/GenerateMutableSource/GenerateMutableSourceTable.cpp
--- a/Engine/Plugins/Mutable/Source/CustomizableObjectEditor/Private/MuCOE/GenerateMutableSource/GenerateMutableSourceTable.cpp
+++ b/Engine/Plugins/Mutable/Source/CustomizableObjectEditor/Private/MuCOE/GenerateMutableSource/GenerateMutableSourceTable.cpp
@@ -33,6 +33,8 @@ bool FillTableColumn(const UCustomizableObjectNodeTable* TableNode, UE::Mutable:
 {
 	int32 CurrentColumn;
 	UDataTable* DataTablePtr = GetDataTable(TableNode, GenerationContext);
+	check(DataTablePtr);
+	
 	FString ColumnPropertyName = ColumnProperty->GetAuthoredName();
 
 	const bool bIgnoreMissingData = GenerationContext.CompilationContext->Object->GetPrivate()->bDisableTableMissingDataWarning;
@@ -84,83 +86,56 @@ bool FillTableColumn(const UCustomizableObjectNodeTable* TableNode, UE::Mutable:
 
 			// Getting Animation information (Anim Blueprints, Animation Slot and Anim Tags)
 			FString AnimBP, AnimSlot, GameplayTag, AnimBPAssetTag;
-			TArray<FGameplayTag> GameplayTags;
 			TableNode->GetAnimationColumns(ColumnName, AnimBP, AnimSlot, GameplayTag);
 
-			if (!AnimBP.IsEmpty())
+			FName AnimBPSlotName;
+			if (!AnimSlot.IsEmpty())
 			{
-				if (!AnimSlot.IsEmpty())
+				if (uint8* AnimRowData = DataTablePtr->FindRowUnchecked(FName(*RowName)))
 				{
-					if (DataTablePtr)
+					// Getting animation slot row value from data table
+					if (FProperty* AnimSlotProperty = TableNode->FindColumnProperty(FName(*AnimSlot)))
 					{
-						uint8* AnimRowData = DataTablePtr->FindRowUnchecked(FName(*RowName));
-
-						if (AnimRowData)
+						if (uint8* AnimSlotData = AnimSlotProperty->ContainerPtrToValuePtr<uint8>(AnimRowData, 0))
 						{
-							FName SlotIndex;
-
-							// Getting animation slot row value from data table
-							if (FProperty* AnimSlotProperty = TableNode->FindColumnProperty(FName(*AnimSlot)))
+							if (const FIntProperty* IntProperty = CastField<FIntProperty>(AnimSlotProperty))
 							{
-								uint8* AnimSlotData = AnimSlotProperty->ContainerPtrToValuePtr<uint8>(AnimRowData, 0);
-
-								if (AnimSlotData)
-								{
-									if (const FIntProperty* IntProperty = CastField<FIntProperty>(AnimSlotProperty))
-									{
-										FString Message = FString::Printf(
-											TEXT("The column with name [%s] for the Anim Slot property should be an FName instead of an Integer, it will be internally converted to FName but should probaly be converted in the table itself."), 
-											*AnimBP);
-										GenerationContext.Log(FText::FromString(Message), TableNode, EMessageSeverity::Info);
+								FString Message = FString::Printf(
+									TEXT("The column with name [%s] for the Anim Slot property should be an FName instead of an Integer, it will be internally converted to FName but should probaly be converted in the table itself."), 
+									*AnimBP);
+								GenerationContext.Log(FText::FromString(Message), TableNode, EMessageSeverity::Info);
 
-										SlotIndex = FName(FString::FromInt(IntProperty->GetPropertyValue(AnimSlotData)));
-									}
-									else if (const FNameProperty* NameProperty = CastField<FNameProperty>(AnimSlotProperty))
-									{
-										SlotIndex = NameProperty->GetPropertyValue(AnimSlotData);
-									}
-								}
+								AnimBPSlotName = FName(FString::FromInt(IntProperty->GetPropertyValue(AnimSlotData)));
 							}
-
-							if (SlotIndex.GetStringLength() != 0)
+							else if (const FNameProperty* NameProperty = CastField<FNameProperty>(AnimSlotProperty))
 							{
-								// Getting animation instance soft class from data table
-								if (FProperty* AnimBPProperty = TableNode->FindColumnProperty(FName(*AnimBP)))
-								{
-									uint8* AnimBPData = AnimBPProperty->ContainerPtrToValuePtr<uint8>(AnimRowData, 0);
-
-									if (AnimBPData)
-									{
-										if (const FSoftClassProperty* SoftClassProperty = CastField<FSoftClassProperty>(AnimBPProperty))
-										{
-											TSoftClassPtr<UAnimInstance> AnimInstance(SoftClassProperty->GetPropertyValue(AnimBPData).ToSoftObjectPath());
-
-											if (!AnimInstance.IsNull())
-											{
-												const int32 AnimInstanceIndex = GenerationContext.AnimBPAssets.AddUnique(AnimInstance);
-
-												AnimBPAssetTag = GenerateAnimationInstanceTag(AnimInstanceIndex, SlotIndex);
-											}
-										}
-									}
-								}
-							}
-							else if (!bIgnoreMissingData)
-							{
-								FString msg = FString::Printf(TEXT("Could not find the Slot column of the animation blueprint column [%s] for the mesh column [%s] row [%s]."), *AnimBP, *ColumnName, *RowName);
-								LogRowGenerationMessage(TableNode, DataTablePtr, GenerationContext, msg, RowName);
+								AnimBPSlotName = NameProperty->GetPropertyValue(AnimSlotData);
 							}
 						}
 					}
 				}
-				else if(!bIgnoreMissingData)
+			}
+			
+			TSoftClassPtr<UAnimInstance> AnimInstance;
+			if (!AnimBP.IsEmpty())
+			{
+				if (uint8* AnimRowData = DataTablePtr->FindRowUnchecked(FName(*RowName)))
 				{
-					FString msg = FString::Printf(TEXT("Could not found the Slot column of the animation blueprint column [%s] for the mesh column [%s]."), *AnimBP, *ColumnName);
-					GenerationContext.Log(FText::FromString(msg), TableNode);
+					// Getting animation instance soft class from data table
+					if (FProperty* AnimBPProperty = TableNode->FindColumnProperty(FName(*AnimBP)))
+					{
+						if (uint8* AnimBPData = AnimBPProperty->ContainerPtrToValuePtr<uint8>(AnimRowData, 0))
+						{
+							if (const FSoftClassProperty* SoftClassProperty = CastField<FSoftClassProperty>(AnimBPProperty))
+							{
+								AnimInstance = SoftClassProperty->GetPropertyValue(AnimBPData).ToSoftObjectPath();
+							}
+						}
+					}
 				}
 			}
-
-			// Getting Gameplay tags
+			
+			FGameplayTagContainer GameplayTags;
 			if (!GameplayTag.IsEmpty())
 			{
 				if (DataTablePtr)
@@ -180,8 +155,7 @@ bool FillTableColumn(const UCustomizableObjectNodeTable* TableNode, UE::Mutable:
 								{
 									if (GameplayTagData)
 									{
-										const FGameplayTagContainer* TagContainer = reinterpret_cast<FGameplayTagContainer*>(GameplayTagData);
-										TagContainer->GetGameplayTagArray(GameplayTags);
+										GameplayTags = *reinterpret_cast<FGameplayTagContainer*>(GameplayTagData);
 									}
 								}
 							}
@@ -189,43 +163,21 @@ bool FillTableColumn(const UCustomizableObjectNodeTable* TableNode, UE::Mutable:
 					}
 				}
 			}
-
-			// First process the mesh tags that are going to make the mesh unique and affect whether it's repeated in 
-			// the mesh cache or not
-			FString MeshUniqueTags;
-
-			if (!AnimBPAssetTag.IsEmpty())
-			{
-				MeshUniqueTags += AnimBPAssetTag;
-			}
-
-			TArray<FString> ArrayAnimBPTags;
-
-			for (const FGameplayTag& Tag : GameplayTags)
-			{
-				MeshUniqueTags += GenerateGameplayTag(Tag.ToString());
-			}
-
-			//TODO: Add AnimBp physics to Tables.
+			
+			// TODO: Add AnimBp physics to Tables.
 
 			FMutableSourceMeshData Source = BaseMeshData;
 			Source.Mesh = SoftObject.ToSoftObjectPath();
 			Source.TableReferenceSkeletalMesh = ReferenceSkeletalMesh;
 			Source.OptionalMessageInfo = FString::Printf(TEXT(" SkeletalMesh from column [%s] row [%s]"), *ColumnName, *RowName);
-			UE::Mutable::Private::TManagedPtr<UE::Mutable::Private::FMesh> MutableMesh = GenerateMutableSkeletalMesh(Source, Source.BaseLODIndex, Source.BaseSectionIndex, MeshUniqueTags, GenerationContext, TableNode, Source.bOnlyConnectedLOD);
+			Source.GameplayTags = GameplayTags;
+			Source.AnimBPSlotName = AnimBPSlotName;
+			Source.AnimInstance = AnimInstance;
+
+			UE::Mutable::Private::TManagedPtr<UE::Mutable::Private::FMesh> MutableMesh = GenerateMutableSkeletalMesh(Source, Source.BaseLODIndex, Source.BaseSectionIndex, GenerationContext, TableNode, Source.bOnlyConnectedLOD);
 
 			if (MutableMesh)
 			{
-				if (!AnimBPAssetTag.IsEmpty())
-				{
-					AddTagToMutableMeshUnique(*MutableMesh, AnimBPAssetTag);
-				}
-
-				for (const FGameplayTag& Tag : GameplayTags)
-				{
-					AddTagToMutableMeshUnique(*MutableMesh, GenerateGameplayTag(Tag.ToString()));
-				}
-
 				MutableTable->SetCell(CurrentColumn, RowId, MutableMesh, TableNode);
 			}
 		}
@@ -287,8 +239,7 @@ bool FillTableColumn(const UCustomizableObjectNodeTable* TableNode, UE::Mutable:
 			}
 
 			constexpr bool bIsPassthrough = false;
-			UE::Mutable::Private::TManagedPtr<UE::Mutable::Private::FMesh> MutableMesh = GenerateMutableStaticMesh(StaticMesh, TSoftClassPtr<UAnimInstance>(), CurrentLOD, BaseMeshData.BaseSectionIndex,
-														  FString(), GenerationContext, TableNode, nullptr, bIsPassthrough);
+			UE::Mutable::Private::TManagedPtr<UE::Mutable::Private::FMesh> MutableMesh = GenerateMutableStaticMesh(StaticMesh, CurrentLOD, BaseMeshData.BaseSectionIndex, GenerationContext, TableNode);
 
 			if (MutableMesh)
 			{
@@ -398,7 +349,7 @@ bool FillTableColumn(const UCustomizableObjectNodeTable* TableNode, UE::Mutable:
 					CurrentColumn = MutableTable->AddColumn(ColumnName, UE::Mutable::Private::ETableColumnType::Material);
 				}
 
-				int32 ReferenceMaterialId = GenerationContext.CompilationContext->PassthroughObjectFactory.Add(*MaterialInstance);
+				int32 ReferenceMaterialId = GenerationContext.CompilationContext->PassthroughObjectFactory.Add(*MaterialInstance, false);
 				MutableTable->SetCell(CurrentColumn, RowId, ReferenceMaterialId);
 
 				return true;
@@ -1198,7 +1149,7 @@ bool GenerateNoneRow(const UCustomizableObjectNodeTable& TableNode, const UEdGra
 			ColumnId = MutableTable->AddColumn(ColumnName, UE::Mutable::Private::ETableColumnType::Material);
 		}
 
-		int32 MaterialId = GenerationContext.CompilationContext->PassthroughObjectFactory.Add(*GenerationContext.CurrentReferencedMaterial.Get());
+		int32 MaterialId = GenerationContext.CompilationContext->PassthroughObjectFactory.Add(*GenerationContext.CurrentReferencedMaterial.Get(), false);
 		MutableTable->SetCell(ColumnId, RowId, MaterialId);
 	}
 	else if (Pin->PinType.PinCategory == Schema->PC_Color)
diff --git a/Engine/Plugins/Mutable/Source/CustomizableObjectEditor/Private/MuCOE/PassthroughObjectFactory.cpp b/Engine/Plugins/Mutable/Source/CustomizableObjectEditor/Private/MuCOE/PassthroughObjectFactory.cpp
--- a/Engine/Plugins/Mutable/Source/CustomizableObjectEditor/Private/MuCOE/PassthroughObjectFactory.cpp
+++ b/Engine/Plugins/Mutable/Source/CustomizableObjectEditor/Private/MuCOE/PassthroughObjectFactory.cpp
@@ -5,7 +5,7 @@
 
 namespace UE::Mutable::Private
 {
-	PASSTHROUGH_ID FPassthroughObjectFactory::Add(UObject& Object)
+	PASSTHROUGH_ID FPassthroughObjectFactory::Add(UObject& Object, bool bDuplicate)
 	{
 		TSoftObjectPtr<UObject> SoftObjectPtr(&Object);
 		
@@ -23,6 +23,11 @@ namespace UE::Mutable::Private
 		}
 	
 		PassthroughObjects.Add(Id, SoftObjectPtr);
+		
+		if (bDuplicate)
+		{
+			DuplicateObjects.Add(Id);
+		}
 
 		return Id;
 	}
@@ -32,4 +37,10 @@ namespace UE::Mutable::Private
 	{
 		return PassthroughObjects;
 	}
+
+
+	TArray<PASSTHROUGH_ID>& FPassthroughObjectFactory::GetDuplicate()
+	{
+		return DuplicateObjects;
+	}
 }
diff --git a/Engine/Plugins/Mutable/Source/CustomizableObjectEditor/Private/MuCOE/PassthroughObjectFactory.h b/Engine/Plugins/Mutable/Source/CustomizableObjectEditor/Private/MuCOE/PassthroughObjectFactory.h
--- a/Engine/Plugins/Mutable/Source/CustomizableObjectEditor/Private/MuCOE/PassthroughObjectFactory.h
+++ b/Engine/Plugins/Mutable/Source/CustomizableObjectEditor/Private/MuCOE/PassthroughObjectFactory.h
@@ -9,11 +9,15 @@ namespace UE::Mutable::Private
 	class FPassthroughObjectFactory
 	{
 	public:
-		PASSTHROUGH_ID Add(UObject& Object);
+		PASSTHROUGH_ID Add(UObject& Object, bool bDuplicate);
 		
 		TMap<PASSTHROUGH_ID, TSoftObjectPtr<UObject>>& GetAll();
 		
+		TArray<PASSTHROUGH_ID>& GetDuplicate();
+		
 	private:
 		TMap<PASSTHROUGH_ID, TSoftObjectPtr<UObject>> PassthroughObjects;
+		
+		TArray<PASSTHROUGH_ID> DuplicateObjects;
 	};	
 }
diff --git a/Engine/Plugins/Mutable/Source/CustomizableObjectEditor/Private/MuCOE/SMutableCodeViewer.cpp b/Engine/Plugins/Mutable/Source/CustomizableObjectEditor/Private/MuCOE/SMutableCodeViewer.cpp
--- a/Engine/Plugins/Mutable/Source/CustomizableObjectEditor/Private/MuCOE/SMutableCodeViewer.cpp
+++ b/Engine/Plugins/Mutable/Source/CustomizableObjectEditor/Private/MuCOE/SMutableCodeViewer.cpp
@@ -1566,19 +1566,21 @@ void SMutableCodeViewer::GenerateElementRecursive(const int32& InStateIndex, UE:
 
 		using OpEnumFlags = UE::Mutable::Private::OP::MeshAddMetadataArgs::EnumFlags;
 
-		int32 TagCount = EnumHasAnyFlags(Args.Flags, OpEnumFlags::HasTags) ? 1 : 0;
-		if (EnumHasAnyFlags(Args.Flags, OpEnumFlags::IsTagList) &&
-			InProgram.ConstantUInt32Lists.IsValidIndex(Args.Tags.ListAddress))
+		int32 TagCount = EnumHasAnyFlags(Args.Flags, OpEnumFlags::HasGameplayTags) ? 1 : 0;
+		if (EnumHasAnyFlags(Args.Flags, OpEnumFlags::IsGameplayTagList) &&
+			ensure(InProgram.ConstantUInt32Lists.IsValidIndex(Args.GameplayTags.ListIndex)))
 		{
-			TagCount = InProgram.ConstantUInt32Lists[Args.Tags.ListAddress].Num();
+			TagCount = InProgram.ConstantUInt32Lists.IsValidIndex(Args.GameplayTags.ListIndex) 
+					? InProgram.ConstantUInt32Lists[Args.GameplayTags.ListIndex].Num() 
+					: 0;
 		}
 
-		int32 ResourceCount = EnumHasAnyFlags(Args.Flags, OpEnumFlags::HasResources) ? 1 : 0;
-		if (EnumHasAnyFlags(Args.Flags, OpEnumFlags::IsResourceList) &&
-			ensure(InProgram.ConstantUInt64Lists.IsValidIndex(Args.ResourceIds.ListAddress)))
+		int32 ResourceCount = EnumHasAnyFlags(Args.Flags, OpEnumFlags::HasAnimationSlots) ? 1 : 0;
+		if (EnumHasAnyFlags(Args.Flags, OpEnumFlags::IsAnimationSlotList) &&
+			ensure(InProgram.ConstantUInt32Lists.IsValidIndex(Args.AnimSlotNames.ListIndex)))
 		{
-			ResourceCount = InProgram.ConstantUInt64Lists.IsValidIndex(Args.ResourceIds.ListAddress) 
-					? InProgram.ConstantUInt64Lists[Args.ResourceIds.ListAddress].Num() 
+			ResourceCount = InProgram.ConstantUInt32Lists.IsValidIndex(Args.AnimSlotNames.ListIndex) 
+					? InProgram.ConstantUInt32Lists[Args.AnimSlotNames.ListIndex].Num() 
 					: 0;
 		}
 
@@ -1600,9 +1602,16 @@ void SMutableCodeViewer::GenerateElementRecursive(const int32& InStateIndex, UE:
 		{
 			AdditionalPhysicsAssetCount = InProgram.ConstantUInt32Lists[Args.AdditionalPhysicsAssetsIds.ListAddress].Num();
 		}
+			
+		int32 AssetUserDataCount = EnumHasAnyFlags(Args.Flags, OpEnumFlags::HasAssetUserData) ? 1 : 0;
+		if (EnumHasAnyFlags(Args.Flags, OpEnumFlags::IsAssetUserDataList) &&
+			ensure(InProgram.ConstantUInt32Lists.IsValidIndex(Args.AssetUserDataIds.ListAddress)))
+		{
+			AssetUserDataCount = InProgram.ConstantUInt32Lists[Args.AssetUserDataIds.ListAddress].Num();
+		}
 
-		FString Caption = FString::Printf(TEXT("add %d Tags, %d Resources, %d Skeleton, %d PhysicsAsset, %d RealTimeMorphs and %d Additional PhysicsAssets to "),
-			TagCount, ResourceCount, SkeletonCount, PhysicsAssetCount, RealTimeMorphNamesCount, AdditionalPhysicsAssetCount);
+		FString Caption = FString::Printf(TEXT("add %d Tags, %d Resources, %d Skeleton, %d PhysicsAsset, %d RealTimeMorphs, %d Additional PhysicsAssets, and %d AssetUserData to"),
+			TagCount, ResourceCount, SkeletonCount, PhysicsAssetCount, RealTimeMorphNamesCount, AdditionalPhysicsAssetCount, AssetUserDataCount);
 		AddOpFunc(Args.Source, Caption);
 
 		break;
@@ -2658,34 +2667,6 @@ bool SMutableCodeViewer::IsConstantResourceUsedByOperation(const int32 IndexOnCo
 						break;
 					}
 				}
-				else if (OperationType == UE::Mutable::Private::EOpType::IN_ADDEXTENSIONDATA)
-				{
-					bResourceLocated = IndexOnConstantsArray == InProgram.GetOpArgs<UE::Mutable::Private::OP::InstanceAddExtensionDataArgs>(OperationAddress).ExtensionDataName;
-				}
-				else if (OperationType == UE::Mutable::Private::EOpType::ME_ADDMETADATA)
-				{
-					 const uint8* OpData = InProgram.GetOpArgsPointer(OperationAddress);
-					
-					 UE::Mutable::Private::OP::ADDRESS SourceAddress;
-					 FMemory::Memcpy(&SourceAddress, OpData, sizeof(UE::Mutable::Private::OP::ADDRESS));
-					 OpData += sizeof(UE::Mutable::Private::OP::ADDRESS);
-					
-					 uint16 TagCount;
-					 FMemory::Memcpy(&TagCount, OpData, sizeof(uint16));
-					 OpData += sizeof(uint16);
-					
-					for (int32 TagIndex = 0; TagIndex < TagCount; ++TagIndex)
-					{
-						// Exit once we know that the data is pointing to the index provided
-						if (*OpData == IndexOnConstantsArray)
-						{
-							bResourceLocated = true;
-							continue;
-						}
-						
-						OpData += sizeof(uint16);
-					}
-				}
 
 				break;
 			}
diff --git a/Engine/Plugins/Mutable/Source/CustomizableObjectEditor/Private/MuCOE/SMutableMeshViewer.cpp b/Engine/Plugins/Mutable/Source/CustomizableObjectEditor/Private/MuCOE/SMutableMeshViewer.cpp
--- a/Engine/Plugins/Mutable/Source/CustomizableObjectEditor/Private/MuCOE/SMutableMeshViewer.cpp
+++ b/Engine/Plugins/Mutable/Source/CustomizableObjectEditor/Private/MuCOE/SMutableMeshViewer.cpp
@@ -205,15 +205,7 @@ void SMutableMeshViewer::SetMesh(const UE::Mutable::Private::TManagedPtr<const U
 	if (InMesh != MutableMesh)
 	{
 		MutableMesh = InMesh;
-
-		// Extract a copy of the tags
-		int32 TagCount = MutableMesh ? MutableMesh->GetTagCount() : 0;
-		MeshTagList.SetNum(TagCount, EAllowShrinking::No);
-		for (int32 TagIndex=0; TagIndex<TagCount; ++TagIndex)
-		{
-			MeshTagList[TagIndex] = MakeShared<FString>(MutableMesh->GetTag(TagIndex));
-		}
-
+		
 		if (MutableMesh)
 		{
 			OnMeshChanged();
@@ -483,24 +475,6 @@ TSharedRef<SWidget> SMutableMeshViewer::GenerateDataTableSlates()
 				SAssignNew(MutableSkeletonViewer, SMutableSkeletonViewer)
 			]
 			// ---------------------------------
-
-			// Tags
-			+ SVerticalBox::Slot()
-				.AutoHeight()
-				[
-					SNew(STextBlock).
-					Text(LOCTEXT("Tags", "Tags"))
-				]
-
-			+ SVerticalBox::Slot()
-				.Padding(IndentationSpace, AfterTitleSpacing)
-				.AutoHeight()
-				[
-					SNew(SListView< TSharedPtr<FString> >)
-					.ListItemsSource(&MeshTagList)
-					.OnGenerateRow(this, &SMutableMeshViewer::GenerateTagRow)
-				]
-			// ---------------------------------
 		]
 	]
 	
diff --git a/Engine/Plugins/Mutable/Source/CustomizableObjectEditor/Private/MuCOE/SMutableMeshViewer.h b/Engine/Plugins/Mutable/Source/CustomizableObjectEditor/Private/MuCOE/SMutableMeshViewer.h
--- a/Engine/Plugins/Mutable/Source/CustomizableObjectEditor/Private/MuCOE/SMutableMeshViewer.h
+++ b/Engine/Plugins/Mutable/Source/CustomizableObjectEditor/Private/MuCOE/SMutableMeshViewer.h
@@ -69,8 +69,6 @@ private:
 	/** Slate whose task is to display the skeleton found on this mesh as a slate tree view */
 	TSharedPtr<SMutableSkeletonViewer> MutableSkeletonViewer;
 
-	/** Widget-side copy of the tags in the mesh.  */
-	TArray<TSharedPtr<FString>> MeshTagList;
 	TSharedRef<ITableRow> GenerateTagRow(TSharedPtr<FString> InItem, const TSharedRef<STableViewBase>& OwnerTable);
 
 	/** Data backend for the widget. It represents the mesh that is being "displayed" */
diff --git a/Engine/Plugins/Mutable/Source/CustomizableObjectEditor/Public/MuCOE/ExtensionDataCompilerInterface.h b/Engine/Plugins/Mutable/Source/CustomizableObjectEditor/Public/MuCOE/ExtensionDataCompilerInterface.h
--- a/Engine/Plugins/Mutable/Source/CustomizableObjectEditor/Public/MuCOE/ExtensionDataCompilerInterface.h
+++ b/Engine/Plugins/Mutable/Source/CustomizableObjectEditor/Public/MuCOE/ExtensionDataCompilerInterface.h
@@ -2,10 +2,13 @@
 
 #pragma once
 
-#include "MuR/ExtensionData.h"
+#include "MuR/PassthroughObject.h"
 
 #define UE_API CUSTOMIZABLEOBJECTEDITOR_API
 
+class FText;
+class UCustomizableObject;
+class UCustomizableObjectNode;
 struct FInstancedStruct;
 struct FMutableGraphGenerationContext;
 
@@ -18,49 +21,17 @@ class FExtensionDataCompilerInterface
 public:
 	UE_API FExtensionDataCompilerInterface(FMutableGraphGenerationContext& InGenerationContext);
 
-	/**
-	 * Register a new Extension Data compile-time constant that will be streamed in on demand.
-	 * 
-	 * This constant will be cooked out to its own package and will be loaded as needed.
-	 * 
-	 * The provided FInstancedStruct will be visible to the garbage collector, so any object
-	 * references from UPROPERTYs in the instanced struct will be treated as hard references from
-	 * the Customizable Object.
-	 *
-	 * Don't save any references to objects under the container as string paths, e.g. using
-	 * TSoftObjectPtr, because the container will be moved to a different package during cooking
-	 * and your paths won't be automatically updated.
-	 */
-	UE_API TSharedPtr<const UE::Mutable::Private::FExtensionData> MakeStreamedExtensionData(FInstancedStruct&& Data);
+	/** Register a new Extension Data. 
+	 * @param bDuplicate if true, object will be duplicated during cook. */
+	UE_API UE::Mutable::Private::PASSTHROUGH_ID MakeExtensionData(UObject& ExtensionData, bool bDuplicate);
 
-	/**
-	 * Register a new Extension Data compile-time constant that will always be loaded.
-	 *
-	 * This constant will be embedded in the Customizable Object and hence will be loaded in memory
-	 * as long as the Customizable Object is loaded.
-	 * 
-	 * Any UObjects referenced by the provided instanced struct that aren't in an asset package
-	 * should be created with the return value of GetOuterForAlwaysLoadedObjects as their Outer.
-	 *
-	 * Move the constant data into this function and use the resulting UE::Mutable::Private::FExtensionData in the
-	 * UE::Mutable::Private::Node graph, e.g. set it as the value of a UE::Mutable::Private::NodeExtensionDataConstant.
-	 * 
-	 * As with MakeStreamedExtensionData, the provided instanced struct will be visible to the GC.
-	 */
-	UE_API TSharedPtr<const UE::Mutable::Private::FExtensionData> MakeAlwaysLoadedExtensionData(FInstancedStruct&& Data);
-
-	/** The Outer to use for objects owned by always-loaded Extension Data constants. */
-	UE_API const UObject* GetOuterForAlwaysLoadedObjects();
-
-	/**
-	* Adds a node to the Generation Context list of generated nodes. This function is meant to be called
+	/** Adds a node to the Generation Context list of generated nodes. This function is meant to be called
 	* from classes that implement ICustomizableObjectExtensionNode::GenerateMutableNode for any generated nodes
-	* so they are registered against the Mutable compiler
-	*/
-	UE_API void AddGeneratedNode(const class UCustomizableObjectNode* InNode);
+	* so they are registered against the Mutable compiler */
+	UE_API void AddGeneratedNode(const UCustomizableObjectNode* InNode);
 
 	/** Adds a compiler log message to be displayed at the end of the compilation process */
-	UE_API void CompilerLog(const FText& InLogText, const class UCustomizableObjectNode* InNode);
+	UE_API void CompilerLog(const FText& InLogText, const UCustomizableObjectNode* InNode);
 
 	FMutableGraphGenerationContext& GenerationContext;
 };
diff --git a/Engine/Plugins/Mutable/Source/MutableRuntime/Internal/MuR/ExtensionData.h b/Engine/Plugins/Mutable/Source/MutableRuntime/Internal/MuR/ExtensionData.h
--- a/Engine/Plugins/Mutable/Source/MutableRuntime/Internal/MuR/ExtensionData.h
+++ b/Engine/Plugins/Mutable/Source/MutableRuntime/Internal/MuR/ExtensionData.h
@@ -6,6 +6,8 @@
 #include "MuR/RefCounted.h"
 #include "MuR/Serialisation.h"
 #include "MuR/Types.h"
+#include "MuR/PassthroughObject.h"
+#include "UObject/Object.h"
 
 #define UE_API MUTABLERUNTIME_API
 
@@ -15,48 +17,10 @@ namespace UE::Mutable::Private
 	* Extensions can provide data, and functionality to operate on that data, without Mutable
 	* needing to know what the data refers to.
 	*/
-	class FExtensionData : public FResource
+	class FExtensionData
 	{
 	public:
-		static UE_API void Serialise(const FExtensionData* Data, FOutputArchive& Archive);
-		static UE_API TManagedPtr<FExtensionData> StaticUnserialise(FInputArchive& Archive);
-
-		// Resource interface
-		UE_API int32 GetDataSize() const override;
-
-		//! A stable hash of the contents
-		UE_API uint32 Hash() const;
-
-		UE_API void Serialise(FOutputArchive& Archive) const;
-		UE_API void Unserialise(FInputArchive& Archive);
-
-		inline bool operator==(const FExtensionData& Other) const
-		{
-			const bool bResult =
-				Other.Index == Index
-				&& Other.Origin == Origin;
-
-			return bResult;
-		}
-
-		enum class EOrigin : uint8
-		{
-			//! An invalid value used to indicate that this ExtensionData hasn't been initialized
-			Invalid,
-
-			//! This ExtensionData is a compile-time constant that's always loaded into memory
-			ConstantAlwaysLoaded,
-
-			//! This ExtensionData is a compile-time constant that's streamed in from disk when needed
-			ConstantStreamed,
-
-			//! This ExtensionData was generated at runtime
-			Runtime
-		};
-
-		int16 Index = -1;
-
-		EOrigin Origin = EOrigin::Invalid;
+		TPassthroughObjectPtr<UObject> PassthroughObject;
 	};
 }
 
diff --git a/Engine/Plugins/Mutable/Source/MutableRuntime/Internal/MuR/Instance.h b/Engine/Plugins/Mutable/Source/MutableRuntime/Internal/MuR/Instance.h
--- a/Engine/Plugins/Mutable/Source/MutableRuntime/Internal/MuR/Instance.h
+++ b/Engine/Plugins/Mutable/Source/MutableRuntime/Internal/MuR/Instance.h
@@ -5,7 +5,6 @@
 #include "MuR/ManagedPointer.h"
 #include "MuR/ResourceID.h"
 #include "HAL/PlatformMath.h"
-#include "MuR/ExtensionData.h"
 #include "MuR/Ptr.h"
 #include "MuR/RefCounted.h"
 #include "MuR/Types.h"
@@ -16,6 +15,8 @@
 
 namespace UE::Mutable::Private
 {
+	class FExtensionData;
+
 	struct FInstanceSurface
 	{
 		FName Name;
@@ -71,14 +72,7 @@ namespace UE::Mutable::Private
 
 		TArray<FInstanceLOD, TInlineAllocator<1>> LODs;
 	};
-
 	
-	struct NamedExtensionData
-	{
-		UE::Mutable::Private::TManagedPtr<const FExtensionData> Data;
-		FName Name;
-	};
-
 	
     /** A customised object created from a model and a set of parameter values.
     * It corresponds to an "engine object" but the contents of its data depends on the Model, and
@@ -176,7 +170,6 @@ namespace UE::Mutable::Private
 
 		UE_API FMaterialId GetOverrideMaterialId(int32 ComponentIndex, int32 MaterialIndex) const;
 		
-		
 		UE_API int32 GetOverlayMaterialCount(int32 ComponentIndex) const;
 
 		UE_API FName GetOverlayMaterialSlotSlotName(int32 ComponentIndex, int32 MaterialIndex) const;
@@ -187,12 +180,6 @@ namespace UE::Mutable::Private
 
 		//! Get the number of ExtensionData values in a component
 		UE_API int32 GetExtensionDataCount() const;
-
-		//! Get an ExtensionData value from a component
-		//! \param Index Index of the ExtensionData to fetch
-		//! \param OutExtensionData Receives the ExtensionData
-		//! \param OutName Receives the name associated with the ExtensionData. Guaranteed to be a valid string of non-zero length.
-		UE_API void GetExtensionData(int32 Index, TManagedPtr<const FExtensionData>& OutExtensionData, FName& OutName) const;
 		
 		int32 AddComponent();
 		
@@ -216,12 +203,10 @@ namespace UE::Mutable::Private
 
 		void AddOverlayMaterial(int32 ComponentIndex, const FName& SlotName, const FMaterialId& MaterialId);
 		
-		void SetExtensionData(const TManagedRef<const FExtensionData>& Data, FName Name);
-		
 		TArray<FInstanceComponent, TInlineAllocator<1>> Components;
 
 		// Every entry must have a valid ExtensionData and name
-		TArray<NamedExtensionData> ExtensionData;
+		TArray<TManagedPtr<const FExtensionData>> ExtensionData;
 	};
 
 
diff --git a/Engine/Plugins/Mutable/Source/MutableRuntime/Internal/MuR/Mesh.h b/Engine/Plugins/Mutable/Source/MutableRuntime/Internal/MuR/Mesh.h
--- a/Engine/Plugins/Mutable/Source/MutableRuntime/Internal/MuR/Mesh.h
+++ b/Engine/Plugins/Mutable/Source/MutableRuntime/Internal/MuR/Mesh.h
@@ -28,6 +28,7 @@
 
 #define UE_API MUTABLERUNTIME_API
 
+class UAnimInstance;
 class FString;
 class USkeletalMesh;
 class UPhysicsAsset;
@@ -227,7 +228,7 @@ namespace UE::Mutable::Private
 		WithSkeleton = 1 << 3,
 		WithPhysicsBody = 1 << 4,
 		WithFaceGroups = 1 << 5,
-		WithTags = 1 << 6,
+		WithGameplayTags = 1 << 6,
 		WithVertexBuffers = 1 << 7,
 		WithIndexBuffers = 1 << 8,
 		WithSockets = 1 << 9,
@@ -237,10 +238,11 @@ namespace UE::Mutable::Private
 		WithBoneMap = 1 << 13,
 		WithSkeletonIDs = 1 << 14,
 		WithAdditionalPhysics = 1 << 15,
-		WithStreamedResources = 1 << 16,
+		WithAssetUserData = 1 << 16,
 		WithMorphData = 1 << 17,
 		WithMorphMetaData = 1 << 18,
 		WithClothData = 1 << 19,
+		WithAnimSlots = 1 << 20,
 
 		AllFlags = 0xFFFFFFFF
 	};
@@ -404,30 +406,7 @@ namespace UE::Mutable::Private
 		UE_API int32 AddAdditionalPhysicsBody(const TManagedPtr<const FPhysicsBody>&);
 		UE_API TManagedPtr<const FPhysicsBody> GetAdditionalPhysicsBody(int32 I) const;
 		//int32 GetAdditionalPhysicsBodyExternalId(int32 I) const;
-
-        //! \}
-
-        //! \name Tags
-        //! \{
-
-        //!
-        UE_API void SetTagCount( int32 Count );
-
-        //!
-        UE_API int32 GetTagCount() const;
-
-        //!
-        UE_API const FString& GetTag( int32 TagIndex ) const;
-
-        //!
-        UE_API void SetTag( int32 TagIndex, const FString& Name );
-
-		//!
-		UE_API void AddStreamedResource(uint64 ResourceId);
-
-		//!
-		UE_API const TArray<uint64>& GetStreamedResources() const;
-
+    	
 		//!
 		UE_API int32 FindBonePoseByBoneIndex(int32 Index) const;
 		
@@ -477,9 +456,6 @@ namespace UE::Mutable::Private
 		 */
 		mutable EMeshFlags Flags = EMeshFlags::None;
 
-		/** Skeletal Meshes involved in the generation of the mesh. Meshes from Mesh Parameters.*/
-		TArray<TStrongObjectPtr<USkeletalMesh>> SkeletalMeshes;
-
 		/** Passthrough Id or TStrongObjectPtr. */
 		TPassthroughObjectPtr<USkeletalMesh> PassthroughObject; // TODO GMT It should not be a TPassthroughObject but instead directly UObject.
 
@@ -533,10 +509,11 @@ namespace UE::Mutable::Private
 		 */
 		TArray<TManagedPtr<const FLayout>> Layouts;		
 
-		TArray<FString> Tags;
+		TArray<FName> GameplayTags;
 
-		/** Opaque handle to external resources. */
-		TArray<uint64> StreamedResources;
+		TArray<TPassthroughObjectPtr<UAssetUserData>> AssetUserData;
+		
+    	TArray<TPair<FName, TPassthroughObjectPtr<UClass>>> AnimationSlots;
 
 		struct FBonePose
 		{
@@ -612,9 +589,10 @@ namespace UE::Mutable::Private
 					bEqual = false;
 				}
 			}
-			if (bEqual) bEqual = (StreamedResources == Other.StreamedResources);
+			if (bEqual) bEqual = (AssetUserData == Other.AssetUserData);
+			if (bEqual) bEqual = (AnimationSlots == Other.AnimationSlots);
 			if (bEqual) bEqual = (Surfaces == Other.Surfaces);
-			if (bEqual) bEqual = (Tags == Other.Tags);
+			if (bEqual) bEqual = (GameplayTags == Other.GameplayTags);
 			if (bEqual) bEqual = (SkeletonObjects == Other.SkeletonObjects);
 
 			for (int32 i = 0; bEqual && i < Layouts.Num(); ++i)
diff --git a/Engine/Plugins/Mutable/Source/MutableRuntime/Internal/MuR/Model.h b/Engine/Plugins/Mutable/Source/MutableRuntime/Internal/MuR/Model.h
--- a/Engine/Plugins/Mutable/Source/MutableRuntime/Internal/MuR/Model.h
+++ b/Engine/Plugins/Mutable/Source/MutableRuntime/Internal/MuR/Model.h
@@ -156,36 +156,6 @@ namespace UE::Mutable::Private
 	static_assert(sizeof(FMeshContentRange) == sizeof(uint32)*2);
 	MUTABLE_DEFINE_POD_SERIALISABLE(FMeshContentRange);
 
-	struct FExtensionDataConstant
-	{
-		// This should always be valid, but if the state is Unloaded it won't be usable.
-		//
-		// Avoid storing references to this Data in Memory while the state is Unloaded.
-		TManagedPtr<const FExtensionData> Data;
-
-		enum class ELoadState : uint8
-		{
-			Invalid,
-			Unloaded,
-			FailedToLoad,
-			CurrentlyLoaded,
-			AlwaysLoaded
-		};
-
-		void Serialise(FOutputArchive& Arch) const
-		{
-			Arch << Data;
-		}
-		
-		void Unserialise(FInputArchive& Arch)
-		{
-			Arch >> Data;
-
-			check(Data);
-			check(Data->Origin == FExtensionData::EOrigin::ConstantAlwaysLoaded || Data->Origin == FExtensionData::EOrigin::ConstantStreamed);
-		}
-	};
-
     //!
     struct FProgram
     {
@@ -271,9 +241,6 @@ namespace UE::Mutable::Private
 		*/
 		TArray<TManagedPtr<const FMesh>> ConstantMeshesPermanent;
     	
-		/** Constant FExtensionData */
-		TArray<FExtensionDataConstant> ConstantExtensionData;
-
         /** Constant string data */
 		TArray<FString> ConstantStrings;
 
@@ -380,7 +347,7 @@ namespace UE::Mutable::Private
 
         UE_API uint32 AddConstant(const FMeshSocket& Socket);
 
-        UE_API OP::LIST_UINT32_ADDRESS AddConstant(const TArray<FMeshSocket>& Sockets);
+        UE_API OP::CONSTANT_LIST_UINT32_INDEX AddConstant(const TArray<FMeshSocket>& Sockets);
     	
         UE_API OP::CONSTANT_STRING_INDEX AddConstant(const FString& Str);
     	
@@ -404,9 +371,7 @@ namespace UE::Mutable::Private
     	
         UE_API OP::ADDRESS AddConstant(const TManagedPtr<const FMaterial>& Material);
 #endif //WITH_EDITOR
-
-		void GetExtensionDataConstant(int32 ConstantIndex, TManagedPtr<const FExtensionData>& Result) const;
-
+    	
         UE_API EOpType GetOpType(OP::ADDRESS At) const;
     	
         template<typename ARGS>
@@ -892,9 +857,10 @@ namespace UE::Mutable::Private
 						Result->ClothSections.SetNum(MetaDataMesh->Surfaces.Num());
 					}
 
-					Result->Tags = MetaDataMesh->Tags; 
+					Result->GameplayTags = MetaDataMesh->GameplayTags; 
 					Result->SkeletonObjects = MetaDataMesh->SkeletonObjects;
-					Result->StreamedResources = MetaDataMesh->StreamedResources; 
+					Result->AssetUserData = MetaDataMesh->AssetUserData; 
+					Result->AnimationSlots = MetaDataMesh->AnimationSlots; 
 					Result->Morph = MetaDataMesh->Morph;
 				}
 			}
diff --git a/Engine/Plugins/Mutable/Source/MutableRuntime/Internal/MuR/Operations.h b/Engine/Plugins/Mutable/Source/MutableRuntime/Internal/MuR/Operations.h
--- a/Engine/Plugins/Mutable/Source/MutableRuntime/Internal/MuR/Operations.h
+++ b/Engine/Plugins/Mutable/Source/MutableRuntime/Internal/MuR/Operations.h
@@ -445,11 +445,13 @@ namespace UE::Mutable::Private
     {
         typedef ADDRESS_TYPE ADDRESS;
         typedef ADDRESS_TYPE PARAMETER;
-    	typedef ADDRESS_TYPE CONSTANT_SHAPE_INDEX;
+        typedef ADDRESS_TYPE CONSTANT_LIST_PASSTHROUGH_ID;
+        typedef ADDRESS_TYPE CONSTANT_SHAPE_INDEX;
     	typedef ADDRESS_TYPE CONSTANT_MATRIX_INDEX;
     	typedef ADDRESS_TYPE CONSTANT_EXTENSIONDATA_INDEX;
+    	typedef ADDRESS_TYPE CONSTANT_NAME_ID;
     	typedef ADDRESS_TYPE CONSTANT_STRING_INDEX;
-		typedef ADDRESS_TYPE LIST_UINT32_ADDRESS;
+		typedef ADDRESS_TYPE CONSTANT_LIST_UINT32_INDEX;
 
         /** Arguments for every operation type. */
         struct BoolConstantArgs
@@ -479,8 +481,13 @@ namespace UE::Mutable::Private
 
         struct ResourceConstantArgs
         {
-            CONSTANT_EXTENSIONDATA_INDEX value;
+            ADDRESS value;
         };
+    	
+    	struct ExternalDataConstantArgs
+    	{
+    		PASSTHROUGH_ID ExternalObjectId;
+    	};
 
         struct MeshConstantArgs
         {
@@ -510,8 +517,8 @@ namespace UE::Mutable::Private
 		{
 			PARAMETER Variable;
 
-			LIST_UINT32_ADDRESS ImageParameterNames;
-			LIST_UINT32_ADDRESS ImageParameterAddress;
+			CONSTANT_LIST_UINT32_INDEX ImageParameterNames;
+			CONSTANT_LIST_UINT32_INDEX ImageParameterAddress;
 		};
 
     	struct MaterialSkeletalMeshBreakArgs
@@ -1103,22 +1110,30 @@ namespace UE::Mutable::Private
             enum class EnumFlags : uint16
             {
 				None = 0,
-				HasTags = 1 << 0,
-				IsTagList = 1 << 1,
-				HasResources = 1 << 2,
-				IsResourceList = 1 << 3,
-				HasAdditionalPhysicsAsset = 1 << 6,
-				IsAdditionalPhysicsAssetList = 1 << 7,
-				HasRealTimeMorphNames = 1 << 8,
-				IsRealTimeMorphNamesList = 1 << 9,	
-				HasSockets = 1 << 10,	
-				IsSocketsList = 1 << 11,	
+				HasGameplayTags = 1 << 0,
+				IsGameplayTagList = 1 << 1,
+				HasAssetUserData = 1 << 2,
+            	IsAssetUserDataList = 1 << 4,
+            	HasAnimationSlots = 1 << 5,
+            	IsAnimationSlotList = 1 << 6,
+				HasAdditionalPhysicsAsset = 1 << 7,
+				IsAdditionalPhysicsAssetList = 1 << 8,
+				HasRealTimeMorphNames = 1 << 9,
+				IsRealTimeMorphNamesList = 1 << 10,	
+				HasSockets = 1 << 11,	
+				IsSocketsList = 1 << 12,	
             };
 
+        	union NameOrUInt32ListIndex
+        	{
+        		CONSTANT_NAME_ID NameIndex;
+        		CONSTANT_LIST_UINT32_INDEX ListIndex;
+        	};
+        	
             union StringAddressOrUInt32ListAddress
             {
                 ADDRESS StringAddress;
-				LIST_UINT32_ADDRESS ListAddress;
+				CONSTANT_LIST_UINT32_INDEX ListAddress;
             };
 
             union ResourceIdOrUInt64ListAddress
@@ -1129,14 +1144,14 @@ namespace UE::Mutable::Private
 
 			union PassthroughIDOrUInt32ListAddress
 			{
-				PASSTHROUGH_ID PassthroughID;
-				LIST_UINT32_ADDRESS ListAddress;
+				PASSTHROUGH_ID PassthroughId;
+				CONSTANT_LIST_UINT32_INDEX ListAddress;
 			};
 
 			union SocketIdOrUInt32ListAddress
 			{
 				uint32 SocketId;
-				LIST_UINT32_ADDRESS ListAddress;
+				CONSTANT_LIST_UINT32_INDEX ListAddress;
 			};
 
             EnumFlags Flags; 
@@ -1144,14 +1159,16 @@ namespace UE::Mutable::Private
 			int8 SocketPriority;
             ADDRESS Source;
             PASSTHROUGH_ID SkeletonId;
-            StringAddressOrUInt32ListAddress Tags;
-            ResourceIdOrUInt64ListAddress ResourceIds;
+            NameOrUInt32ListIndex GameplayTags;
+            PassthroughIDOrUInt32ListAddress AssetUserDataIds;
+        	NameOrUInt32ListIndex AnimSlotNames;
+        	PassthroughIDOrUInt32ListAddress AnimInstances;
 			PASSTHROUGH_ID PhysicsAssetId; // Main Physics Asset of the mesh
 			PassthroughIDOrUInt32ListAddress AdditionalPhysicsAssetsIds; // Physics Assets used by AnimBPs
             StringAddressOrUInt32ListAddress RealTimeMorphNames;
 			SocketIdOrUInt32ListAddress Sockets;
         };
-        static_assert(sizeof(MeshAddMetadataArgs) == sizeof(uint64)*5, "MeshAddMetadataArgs has an unexpected size.");
+        static_assert(sizeof(MeshAddMetadataArgs) == 44, "MeshAddMetadataArgs has an unexpected size.");
 
 
         struct InstanceAddArgs
@@ -1176,12 +1193,8 @@ namespace UE::Mutable::Private
 			// This is a reference to an op that produces the Instance that the FExtensionData will
 			// be added to.
 			ADDRESS Instance;
-			// An op that produces the FExtensionData to add to the Instance
+			
 			ADDRESS ExtensionData;
-			// The name to associate with the FExtensionData
-			//
-			// This is an index into the string table
-			ADDRESS ExtensionDataName;
 		};
 
 		struct InstanceAddMaterialArgs
diff --git a/Engine/Plugins/Mutable/Source/MutableRuntime/Private/MuR/CodeRunner.cpp b/Engine/Plugins/Mutable/Source/MutableRuntime/Private/MuR/CodeRunner.cpp
--- a/Engine/Plugins/Mutable/Source/MutableRuntime/Private/MuR/CodeRunner.cpp
+++ b/Engine/Plugins/Mutable/Source/MutableRuntime/Private/MuR/CodeRunner.cpp
@@ -794,11 +794,7 @@ namespace UE::Mutable::Private
 
 					if (TManagedPtr<const FExtensionData> ExtensionData = LoadExtensionData(FCacheAddress(Args.ExtensionData, Item)))
 					{
-						const OP::ADDRESS NameAddress = Args.ExtensionDataName;
-						check(NameAddress < (uint32)Program.ConstantStrings.Num());
-						const FString& NameString = Program.ConstantStrings[NameAddress];
-
-						Instance->SetExtensionData(ExtensionData.ToStrongRef(), FName(NameString));
+						Instance->ExtensionData.Add(ExtensionData);
 					}
 
 					StoreInstance(Item, Instance);
@@ -1086,15 +1082,13 @@ namespace UE::Mutable::Private
 		{
         	MUTABLE_CPUPROFILER_SCOPE(ED_CONSTANT);
         		
-			OP::ResourceConstantArgs Args = Program.GetOpArgs<OP::ResourceConstantArgs>(Item.At);
+			OP::ExternalDataConstantArgs Args = Program.GetOpArgs<OP::ExternalDataConstantArgs>(Item.At);
 
 			// Assume the ROM has been loaded previously
-			TManagedPtr<const UE::Mutable::Private::FExtensionData> SourceConst;
-			Program.GetExtensionDataConstant(Args.value, SourceConst);
+			TManagedPtr<UE::Mutable::Private::FExtensionData> ExtensionData = MakeManaged<UE::Mutable::Private::FExtensionData>();
+			ExtensionData->PassthroughObject = FindPassthrough<UObject>(Args.ExternalObjectId);
 
-			check(SourceConst);
-
-            StoreExtensionData(Item, SourceConst);
+            StoreExtensionData(Item, ExtensionData);
             break;
 		}
 
@@ -2790,46 +2784,63 @@ namespace UE::Mutable::Private
 					using OpEnumFlags = OP::MeshAddMetadataArgs::EnumFlags;
 
 					// Append Tags
-					if (EnumHasAnyFlags(OpArgs.Flags, OpEnumFlags::IsTagList) &&
-						Program.ConstantUInt32Lists.IsValidIndex(OpArgs.Tags.ListAddress))
+					if (EnumHasAnyFlags(OpArgs.Flags, OpEnumFlags::IsGameplayTagList))
 					{
-						const TArray<uint32>& ConstantTagList = Program.ConstantUInt32Lists[OpArgs.Tags.ListAddress];
-						Result->Tags.Reserve(Result->Tags.Num() + ConstantTagList.Num());
+						const TArray<uint32>& ConstantGameplayTagList = Program.ConstantUInt32Lists[OpArgs.GameplayTags.ListIndex];
+						Result->GameplayTags.Reserve(Result->GameplayTags.Num() + ConstantGameplayTagList.Num());
 
-						for (uint32 TagIndex : ConstantTagList)
+						for (uint32 Index : ConstantGameplayTagList)
 						{
-							if (ensure(Program.ConstantStrings.IsValidIndex(TagIndex)))
-							{
-								Result->Tags.Add(Program.ConstantStrings[TagIndex]);
-							}
+							Result->GameplayTags.Add(Program.ConstantNames[Index]);
 						}
 					}
-					else if(EnumHasAnyFlags(OpArgs.Flags, OpEnumFlags::HasTags) &&
-						Program.ConstantStrings.IsValidIndex(OpArgs.Tags.StringAddress))
+					else if (EnumHasAnyFlags(OpArgs.Flags, OpEnumFlags::HasGameplayTags))
 					{
-						Result->Tags.Add(Program.ConstantStrings[OpArgs.Tags.StringAddress]);
+						Result->GameplayTags.Add(Program.ConstantNames[OpArgs.GameplayTags.NameIndex]);
 					}
-					else
+					
+					// Asset User Data
+					if (EnumHasAnyFlags(OpArgs.Flags, OpEnumFlags::IsAssetUserDataList))
 					{
-						check(!EnumHasAnyFlags(OpArgs.Flags, OpEnumFlags::HasTags | OpEnumFlags::IsTagList));
+						for (uint32 PassthroughId : Program.ConstantUInt32Lists[OpArgs.AssetUserDataIds.ListAddress])
+						{
+							check(PassthroughId != PASSTHROUGH_ID_INVALID);
+							Result->AssetUserData.Emplace(FindPassthrough<UAssetUserData>(PassthroughId));
+						}
 					}
-					
-					// Append Streamed Resources
-					if (EnumHasAnyFlags(OpArgs.Flags, OpEnumFlags::IsResourceList) &&
-						Program.ConstantUInt64Lists.IsValidIndex(OpArgs.ResourceIds.ListAddress))
+					else if (EnumHasAnyFlags(OpArgs.Flags, OpEnumFlags::HasAssetUserData))
 					{
-						Result->StreamedResources.Append(Program.ConstantUInt64Lists[OpArgs.ResourceIds.ListAddress]);
+						check(OpArgs.AssetUserDataIds.PassthroughId != PASSTHROUGH_ID_INVALID);
+						Result->AssetUserData.Add(FindPassthrough<UAssetUserData>(OpArgs.AssetUserDataIds.PassthroughId));
 					}
-					else if (EnumHasAnyFlags(OpArgs.Flags, OpEnumFlags::HasResources))
+					
+					// Animation Slots
+					if (EnumHasAnyFlags(OpArgs.Flags, OpEnumFlags::IsAnimationSlotList))
 					{
-						Result->StreamedResources.Add(OpArgs.ResourceIds.ResourceId);
+						const TArray<uint32>& AnimSlotNames = Program.ConstantUInt32Lists[OpArgs.AnimSlotNames.ListIndex];
+						const TArray<uint32>& AnimInstances = Program.ConstantUInt32Lists[OpArgs.AnimInstances.ListAddress];
+						
+						for (int32 Index = 0; Index < AnimSlotNames.Num(); ++Index)
+						{
+							const FName& AnimationSlotName = Program.ConstantNames[AnimSlotNames[Index]];
+
+							PASSTHROUGH_ID AnimInstanceId = AnimInstances[Index];
+							check(AnimInstanceId != PASSTHROUGH_ID_INVALID);
+							TPassthroughObjectPtr<UClass> AnimInstance = FindPassthrough<UClass>(AnimInstanceId);
+							
+							Result->AnimationSlots.Emplace(AnimationSlotName, AnimInstance);
+						}
 					}
-					else
+					else if (EnumHasAnyFlags(OpArgs.Flags, OpEnumFlags::HasAnimationSlots))
 					{
-						check(!EnumHasAnyFlags(OpArgs.Flags, OpEnumFlags::IsResourceList));
+						check(OpArgs.AnimInstances.PassthroughId != PASSTHROUGH_ID_INVALID);
+
+						const FName& AnimationSlotName = Program.ConstantNames[OpArgs.AnimSlotNames.NameIndex];
+						TPassthroughObjectPtr<UClass> AnimInstance = FindPassthrough<UClass>(OpArgs.AnimInstances.PassthroughId);
+						Result->AnimationSlots.Emplace(AnimationSlotName, AnimInstance);
 					}
 					
-					
+					// Pose Priority
 					if (OpArgs.BonePosePriority != 0)
 					{
 						for (FMesh::FBonePose& BonePose : Result->BonePoses)
@@ -2838,6 +2849,7 @@ namespace UE::Mutable::Private
 						}
 					}
 
+					// Sockets
 					if (EnumHasAnyFlags(OpArgs.Flags, OpEnumFlags::IsSocketsList) &&
 						Program.ConstantUInt32Lists.IsValidIndex(OpArgs.Sockets.ListAddress))
 					{
@@ -2853,10 +2865,6 @@ namespace UE::Mutable::Private
 					{
 						Result->Sockets.Emplace(Program.ConstantSockets.FindChecked(OpArgs.Sockets.SocketId));
 					}
-					else
-					{	
-						check(!EnumHasAnyFlags(OpArgs.Flags, OpEnumFlags::IsSocketsList));
-					}
 
 					if (OpArgs.SocketPriority != 0)
 					{
@@ -2885,12 +2893,8 @@ namespace UE::Mutable::Private
 					}
 					else if (EnumHasAnyFlags(OpArgs.Flags, OpEnumFlags::HasAdditionalPhysicsAsset))
 					{
-						check(OpArgs.AdditionalPhysicsAssetsIds.PassthroughID != PASSTHROUGH_ID_INVALID);
-						Result->AdditionalPhysicsAssets.Emplace(FindPassthrough<UPhysicsAsset>(OpArgs.AdditionalPhysicsAssetsIds.PassthroughID));
-					}
-					else
-					{
-						check(!EnumHasAnyFlags(OpArgs.Flags, OpEnumFlags::IsAdditionalPhysicsAssetList));
+						check(OpArgs.AdditionalPhysicsAssetsIds.PassthroughId != PASSTHROUGH_ID_INVALID);
+						Result->AdditionalPhysicsAssets.Emplace(FindPassthrough<UPhysicsAsset>(OpArgs.AdditionalPhysicsAssetsIds.PassthroughId));
 					}
 
 					check(Result->Morph.Names.Num() == Result->Morph.UsageFlagsPerMorph.Num());
@@ -2927,11 +2931,6 @@ namespace UE::Mutable::Private
 							EnumAddFlags(Result->Morph.UsageFlagsPerMorph[FoundIndex], EMorphUsageFlags::RealTime);
 						}
 					}
-					else
-					{
-						check(!EnumHasAnyFlags(OpArgs.Flags, OpEnumFlags::HasRealTimeMorphNames | OpEnumFlags::IsRealTimeMorphNamesList));
-					}
-
 				}
 
 				StoreMesh(Item, Result);
diff --git a/Engine/Plugins/Mutable/Source/MutableRuntime/Private/MuR/CodeRunnerBegin.cpp b/Engine/Plugins/Mutable/Source/MutableRuntime/Private/MuR/CodeRunnerBegin.cpp
--- a/Engine/Plugins/Mutable/Source/MutableRuntime/Private/MuR/CodeRunnerBegin.cpp
+++ b/Engine/Plugins/Mutable/Source/MutableRuntime/Private/MuR/CodeRunnerBegin.cpp
@@ -193,24 +193,7 @@ namespace UE::Mutable::Private
 
 		ResultsString.Add(To, Value);
 	}
-
-
-	void CodeRunnerBegin::StoreExtensionData(const FCacheAddressInline& To, const TManagedPtr<const FExtensionData>& Value)
-	{
-		if (!To.At)
-		{
-			return;
-		}
-		
-		Executed.Add(To);
-
-		FStackValue StackValue;
-		StackValue.Set<TManagedPtr<const FExtensionData>>(Value);
-		Stack.Push(StackValue);
-
-		ResultsExtensionData.Add(To, Value);
-	}
-
+	
 
 	void CodeRunnerBegin::StoreNone(const FCacheAddressInline& To)
 	{
@@ -293,18 +276,7 @@ namespace UE::Mutable::Private
 		
 		return Stack.Pop().Get<TManagedPtr<const String>>();
 	}
-
-
-	TManagedPtr<const FExtensionData> CodeRunnerBegin::LoadExtensionData(const FCacheAddressInline& To)
-	{
-		if (!To.At)
-		{
-			return nullptr;	
-		}
-		
-		return Stack.Pop().Get<TManagedPtr<const FExtensionData>>();
-	}
-
+	
 
 	TManagedPtr<const FMaterial> CodeRunnerBegin::LoadMaterial(const FCacheAddressInline& To)
 	{
@@ -414,26 +386,6 @@ namespace UE::Mutable::Private
 						
 						break;
 					}
-								
-				case EDataType::ExtensionData:
-					{
-						if (Item.bEvaluate)
-						{
-							if (TManagedPtr<const FExtensionData>* Result = ResultsExtensionData.Find(Item))
-							{
-								FStackValue Value;
-								Value.Set<TManagedPtr<const FExtensionData>>(*Result);
-					
-								Stack.Push(Value);
-							}
-							else
-							{
-								bExecute = true;
-							}
-						}
-
-						break;
-					}
 
 				case EDataType::Scalar:
 					{
@@ -463,6 +415,7 @@ namespace UE::Mutable::Private
 						break;
 					}
 					
+				case EDataType::ExtensionData:
 				case EDataType::Color:		
 				case EDataType::Projector:	
 				case EDataType::Mesh:		
@@ -567,10 +520,6 @@ namespace UE::Mutable::Private
 									StoreInstance(Item, LoadInstance(Item.CustomState));
 									break;
 
-								case EDataType::ExtensionData:
-									StoreExtensionData(Item, LoadExtensionData(Item.CustomState));
-									break;
-							
 								case EDataType::Scalar:		
 									StoreFloat(Item, LoadFloat(Item.CustomState));
 									break;
@@ -579,6 +528,7 @@ namespace UE::Mutable::Private
 									StoreMaterial(Item, LoadMaterial(Item.CustomState));
 									break;
 
+								case EDataType::ExtensionData:
 								case EDataType::Color:		
 								case EDataType::Projector:	
 								case EDataType::Mesh:		
@@ -663,11 +613,7 @@ namespace UE::Mutable::Private
 									case EDataType::String:
 										StoreString(Item, nullptr);
 										break;
-								
-									case EDataType::ExtensionData:
-										StoreExtensionData(Item, MakeManaged<FExtensionData>());
-										break;
-
+									
 									case EDataType::Scalar:
 										StoreFloat(Item, 0.0f);
 										break;
@@ -676,6 +622,7 @@ namespace UE::Mutable::Private
 										StoreMaterial(Item, MakeManaged<const FMaterial>());
 										break;
 									
+									case EDataType::ExtensionData:
 									case EDataType::Color:			
 									case EDataType::Projector:		
 									case EDataType::Mesh:			
@@ -778,10 +725,6 @@ namespace UE::Mutable::Private
 									StoreString(Item, LoadString(Item.CustomState));
 									break;
 								
-								case EDataType::ExtensionData:
-									StoreExtensionData(Item, LoadExtensionData(Item.CustomState));
-									break;
-
 								case EDataType::Scalar:
 									StoreFloat(Item, LoadFloat(Item.CustomState));
 									break;
@@ -790,6 +733,7 @@ namespace UE::Mutable::Private
 									StoreMaterial(Item, LoadMaterial(Item.CustomState));
 									break;
 
+								case EDataType::ExtensionData:
 								case EDataType::Color:
 								case EDataType::Projector:
 								case EDataType::Mesh:
@@ -1384,15 +1328,6 @@ namespace UE::Mutable::Private
 							{
 								Instance = MakeManaged<FInstance>();
 							}
-							
-							if (TManagedPtr<const FExtensionData> ExtensionData = LoadExtensionData(Args.ExtensionData))
-							{
-								const OP::ADDRESS NameAddress = Args.ExtensionDataName;
-								check(NameAddress < (uint32)Program.ConstantStrings.Num());
-								const FString& NameString = Program.ConstantStrings[NameAddress];
-
-								Instance->SetExtensionData(ExtensionData.ToStrongRef(), FName(NameString));
-							}
 
 							StoreInstance(Item, Instance);
 							break;
@@ -1753,15 +1688,9 @@ namespace UE::Mutable::Private
 				{
 					check(Item.bEvaluate);
 
-					OP::ResourceConstantArgs Args = Program.GetOpArgs<OP::ResourceConstantArgs>(Item.At);
-
-					// Assume the ROM has been loaded previously
-					TManagedPtr<const FExtensionData> SourceConst;
-					Program.GetExtensionDataConstant(Args.value, SourceConst);
+					OP::ExternalDataConstantArgs Args = Program.GetOpArgs<OP::ExternalDataConstantArgs>(Item.At);
+					PassthroughObjectLoader->Add<UObject>(Args.ExternalObjectId);
 
-					check(SourceConst);
-					
-					StoreExtensionData(Item, SourceConst);
 					break;
 				}
 
@@ -1901,7 +1830,7 @@ namespace UE::Mutable::Private
 
 					OP::MeshConstantArgs Args = Program.GetOpArgs<OP::MeshConstantArgs>(Item.At);
 					PassthroughObjectLoader->Add<UMaterialInterface>(Args.ClothID);
-
+					
 					StoreNone(Item);
 					break;
 				}
@@ -2225,7 +2154,7 @@ namespace UE::Mutable::Private
 					const OP::MeshAddMetadataArgs OpArgs = Program.GetOpArgs<OP::MeshAddMetadataArgs>(Item.At);
 					PushOp(FScheduledOpInline(OpArgs.Source, Item));
 
-if (OpArgs.SkeletonId != PASSTHROUGH_ID_INVALID)
+					if (OpArgs.SkeletonId != PASSTHROUGH_ID_INVALID)
 					{
 						PassthroughObjectLoader->Add<USkeleton>(OpArgs.SkeletonId);
 					}
@@ -2236,30 +2165,48 @@ if (OpArgs.SkeletonId != PASSTHROUGH_ID_INVALID)
 						PassthroughObjectLoader->Add<UPhysicsAsset>(OpArgs.PhysicsAssetId);
 					}
 
-					if (EnumHasAnyFlags(OpArgs.Flags, OpEnumFlags::IsAdditionalPhysicsAssetList) &&
-						Program.ConstantUInt32Lists.IsValidIndex(OpArgs.AdditionalPhysicsAssetsIds.ListAddress))
+					if (EnumHasAnyFlags(OpArgs.Flags, OpEnumFlags::IsAdditionalPhysicsAssetList))
 					{
-						const TArray<uint32>& ConstantPassthroughIdList = Program.ConstantUInt32Lists[OpArgs.AdditionalPhysicsAssetsIds.ListAddress];
+						const TArray<PASSTHROUGH_ID>& ConstantPassthroughIdList = Program.ConstantUInt32Lists[OpArgs.AdditionalPhysicsAssetsIds.ListAddress];
 						
 						for (uint32 PassthroughId : ConstantPassthroughIdList)
 						{
-							if (PassthroughId != PASSTHROUGH_ID_INVALID)
-							{
-								PassthroughObjectLoader->Add<UPhysicsAsset>(PassthroughId);
-							}
+							PassthroughObjectLoader->Add<UPhysicsAsset>(PassthroughId);
 						}
 					}
 					else if (EnumHasAnyFlags(OpArgs.Flags, OpEnumFlags::HasAdditionalPhysicsAsset))
 					{
-						PassthroughObjectLoader->Add<UPhysicsAsset>(OpArgs.AdditionalPhysicsAssetsIds.PassthroughID);
+						PassthroughObjectLoader->Add<UPhysicsAsset>(OpArgs.AdditionalPhysicsAssetsIds.PassthroughId);
 					}
-					else
+
+					if (EnumHasAnyFlags(OpArgs.Flags, OpEnumFlags::IsAssetUserDataList))
 					{
-						check(!EnumHasAnyFlags(OpArgs.Flags, OpEnumFlags::IsAdditionalPhysicsAssetList));
+						const TArray<PASSTHROUGH_ID>& ConstantPassthroughIdList = Program.ConstantUInt32Lists[OpArgs.AssetUserDataIds.ListAddress];
+						
+						for (uint32 PassthroughId : ConstantPassthroughIdList)
+						{
+							PassthroughObjectLoader->Add<UAssetUserData>(PassthroughId);
+						}
+					}
+					else if (EnumHasAnyFlags(OpArgs.Flags, OpEnumFlags::HasAssetUserData))
+					{
+						PassthroughObjectLoader->Add<UAssetUserData>(OpArgs.AssetUserDataIds.PassthroughId);
+					}
+					
+					if (EnumHasAnyFlags(OpArgs.Flags, OpEnumFlags::IsAnimationSlotList))
+					{
+						const TArray<PASSTHROUGH_ID>& ConstantPassthroughIdList = Program.ConstantUInt32Lists[OpArgs.AnimInstances.ListAddress];
+						
+						for (uint32 PassthroughId : ConstantPassthroughIdList)
+						{
+							PassthroughObjectLoader->Add<UAnimInstance>(PassthroughId);
+						}
+					}
+					else if (EnumHasAnyFlags(OpArgs.Flags, OpEnumFlags::HasAnimationSlots))
+					{
+						PassthroughObjectLoader->Add<UAnimInstance>(OpArgs.AnimInstances.PassthroughId);
 					}
-
 					
-
 					StoreNone(Item);
 					break;
 				}
diff --git a/Engine/Plugins/Mutable/Source/MutableRuntime/Private/MuR/CodeRunnerBegin.h b/Engine/Plugins/Mutable/Source/MutableRuntime/Private/MuR/CodeRunnerBegin.h
--- a/Engine/Plugins/Mutable/Source/MutableRuntime/Private/MuR/CodeRunnerBegin.h
+++ b/Engine/Plugins/Mutable/Source/MutableRuntime/Private/MuR/CodeRunnerBegin.h
@@ -98,7 +98,7 @@ namespace UE::Mutable::Private
 		TSet<FCacheAddressInline> IndexOther;
 	};
 	
-	using FStackValue = TVariant<bool, int32, float, TManagedPtr<const String>, TManagedPtr<const FExtensionData>, TManagedPtr<const FInstance>, TManagedPtr<const FMaterial>> ;
+	using FStackValue = TVariant<bool, int32, float, TManagedPtr<const String>, TManagedPtr<const FInstance>, TManagedPtr<const FMaterial>>;
 
 	class FStack : public TArray<FStackValue, TInlineAllocator<32>>
 	{
@@ -130,8 +130,6 @@ namespace UE::Mutable::Private
 		
 		void StoreString(const FCacheAddressInline& To, const TManagedPtr<const String>& Value);
 		
-		void StoreExtensionData(const FCacheAddressInline& To, const TManagedPtr<const FExtensionData>& Value);
-
 		void StoreNone(const FCacheAddressInline& To);
 
 		void StoreMaterial(const FCacheAddressInline& To, const TManagedPtr<const FMaterial>& Value);
@@ -145,9 +143,7 @@ namespace UE::Mutable::Private
 		TManagedPtr<const FInstance> LoadInstance(const FCacheAddressInline& To);
 
 		TManagedPtr<const String> LoadString(const FCacheAddressInline& To);
-
-		TManagedPtr<const FExtensionData> LoadExtensionData(const FCacheAddressInline& To);
-
+		
 		TManagedPtr<const FMaterial> LoadMaterial(const FCacheAddressInline& To);
 		
 		TSharedPtr<const FModel> Model;
@@ -168,7 +164,6 @@ namespace UE::Mutable::Private
 		TMap<FCacheAddressInline, int32> ResultsInt;
 		TMap<FCacheAddressInline, TManagedPtr<const FInstance>> ResultsInstance;
 		TMap<FCacheAddressInline, TManagedPtr<const String>> ResultsString;
-		TMap<FCacheAddressInline, TManagedPtr<const FExtensionData>> ResultsExtensionData;
 		TMap<FCacheAddressInline, TManagedPtr<const FMaterial>> ResultsMaterial;
 
 		TSharedPtr<FProgramCache> ProgramCache;
diff --git a/Engine/Plugins/Mutable/Source/MutableRuntime/Private/MuR/ExtensionData.cpp b/Engine/Plugins/Mutable/Source/MutableRuntime/Private/MuR/ExtensionData.cpp
deleted file mode 100644
diff --git a/Engine/Plugins/Mutable/Source/MutableRuntime/Private/MuR/Instance.cpp b/Engine/Plugins/Mutable/Source/MutableRuntime/Private/MuR/Instance.cpp
--- a/Engine/Plugins/Mutable/Source/MutableRuntime/Private/MuR/Instance.cpp
+++ b/Engine/Plugins/Mutable/Source/MutableRuntime/Private/MuR/Instance.cpp
@@ -220,16 +220,7 @@ int32 FInstance::GetStringCount(int32 ComponentIndex, int32 LODIndex, int32 Surf
 	{
 		return ExtensionData.Num();
 	}
-
 	
-	void FInstance::GetExtensionData(int32 Index, TManagedPtr<const FExtensionData>& OutExtensionData, FName& OutName) const
-	{
-		check(ExtensionData.IsValidIndex(Index));
-
-		OutExtensionData = ExtensionData[Index].Data;
-		OutName = ExtensionData[Index].Name;
-	}
-
     
 	int32 FInstance::AddComponent()
 	{
@@ -388,13 +379,5 @@ int32 FInstance::GetStringCount(int32 ComponentIndex, int32 LODIndex, int32 Surf
 
     	Components[ComponentIndex].OverlayMaterials.Add({ SlotName, MaterialId });
 	}
-
-
-	void FInstance::SetExtensionData(const TManagedRef<const FExtensionData>& Data, FName Name)
-	{
-		NamedExtensionData& Entry = ExtensionData.AddDefaulted_GetRef();
-		Entry.Data = Data;
-		Entry.Name = Name;
-	}
 }
 
diff --git a/Engine/Plugins/Mutable/Source/MutableRuntime/Private/MuR/Mesh.cpp b/Engine/Plugins/Mutable/Source/MutableRuntime/Private/MuR/Mesh.cpp
--- a/Engine/Plugins/Mutable/Source/MutableRuntime/Private/MuR/Mesh.cpp
+++ b/Engine/Plugins/Mutable/Source/MutableRuntime/Private/MuR/Mesh.cpp
@@ -100,7 +100,6 @@ void FMesh::CopyFrom(const FMesh& From, EMeshCopyFlags InFlags)
 	PassthroughObject = From.PassthroughObject;
 
 	MeshIDPrefix = From.MeshIDPrefix;
-	SkeletalMeshes = From.SkeletalMeshes;
 
 	if (EnumHasAnyFlags(InFlags, EMeshCopyFlags::WithSurfaces))
 	{
@@ -118,14 +117,19 @@ void FMesh::CopyFrom(const FMesh& From, EMeshCopyFlags InFlags)
 		PhysicsBody = From.PhysicsBody;
 	}
 
-	if (EnumHasAnyFlags(InFlags, EMeshCopyFlags::WithTags))
+	if (EnumHasAnyFlags(InFlags, EMeshCopyFlags::WithGameplayTags))
 	{
-		Tags = From.Tags;
+		GameplayTags = From.GameplayTags;
 	}
 
-	if (EnumHasAnyFlags(InFlags, EMeshCopyFlags::WithStreamedResources))
+	if (EnumHasAnyFlags(InFlags, EMeshCopyFlags::WithAssetUserData))
 	{
-		StreamedResources = From.StreamedResources;
+		AssetUserData = From.AssetUserData;
+	}
+	
+	if (EnumHasAnyFlags(InFlags, EMeshCopyFlags::WithAnimSlots))
+	{
+		AnimationSlots = From.AnimationSlots;
 	}
 
     // Copy the main buffers
@@ -555,59 +559,6 @@ void FMesh::SetLayout(int32 LayoutIndex, TManagedPtr<const FLayout> InLayout)
 }
 
 
-int32 FMesh::GetTagCount() const
-{
-    return Tags.Num();
-}
-
-
-void FMesh::SetTagCount(int32 Count)
-{
-	LLM_SCOPE_BYNAME(TEXT("MutableRuntime"));
-    Tags.SetNum(Count);
-}
-
-
-const FString& FMesh::GetTag(int32 TagIndex) const
-{
-    check(Tags.IsValidIndex(TagIndex));
-
-    if (Tags.IsValidIndex(TagIndex))
-    {
-        return Tags[TagIndex];
-    }
-    else
-    {
-		static FString NullString;
-        return NullString;
-    }
-}
-
-
-void FMesh::SetTag(int32 TagIndex, const FString& Name)
-{
-    check(Tags.IsValidIndex(TagIndex));
-	LLM_SCOPE_BYNAME(TEXT("MutableRuntime"));
-
-    if (Tags.IsValidIndex(TagIndex))
-    {
-        Tags[TagIndex] = Name;
-    }
-}
-
-
-void FMesh::AddStreamedResource(uint64 ResourceId)
-{
-	StreamedResources.AddUnique(ResourceId);
-}
-
-
-const TArray<uint64>& FMesh::GetStreamedResources() const
-{
-	return StreamedResources;
-}
-
-
 int32 FMesh::FindBonePoseByBoneIndex(int32 Index) const
 {
 	return BonePoses.IndexOfByPredicate([Index](const FBonePose& Pose) { return Pose.BoneId.Index == Index; });
@@ -1105,7 +1056,6 @@ bool FMesh::IsSimilar(const FMesh& Other) const
 	}
 
 	bEqual = bEqual && (Surfaces == Other.Surfaces);
-	bEqual = bEqual && (Tags == Other.Tags);
 
 	bEqual = bEqual && (Morph == Other.Morph);
 	bEqual = bEqual && (MorphDataBuffer == Other.MorphDataBuffer);
diff --git a/Engine/Plugins/Mutable/Source/MutableRuntime/Private/MuR/Model.cpp b/Engine/Plugins/Mutable/Source/MutableRuntime/Private/MuR/Model.cpp
--- a/Engine/Plugins/Mutable/Source/MutableRuntime/Private/MuR/Model.cpp
+++ b/Engine/Plugins/Mutable/Source/MutableRuntime/Private/MuR/Model.cpp
@@ -118,7 +118,6 @@ namespace UE::Mutable::Private
 		Arch << ConstantMeshesPermanent;
 		Arch << ConstantMeshContentIndices;
 		Arch << ConstantMeshes;
-		Arch << ConstantExtensionData;
 		Arch << ConstantStrings;
 		Arch << ConstantUInt32Lists;
 		Arch << ConstantUInt64Lists;
@@ -151,7 +150,6 @@ namespace UE::Mutable::Private
 		Arch >> ConstantMeshesPermanent;
 		Arch >> ConstantMeshContentIndices;
 		Arch >> ConstantMeshes;
-		Arch >> ConstantExtensionData;
 		Arch >> ConstantStrings;
 		Arch >> ConstantUInt32Lists;
 		Arch >> ConstantUInt64Lists;
@@ -226,25 +224,6 @@ namespace UE::Mutable::Private
     }
 	
 #if WITH_EDITOR
-    OP::CONSTANT_EXTENSIONDATA_INDEX FProgram::AddConstant(const TManagedPtr<const FExtensionData>& Data)
-    {
-	    // Ensure unique
-	    for (int32 Index = 0; Index < ConstantExtensionData.Num(); Index++)
-	    {
-		    const FExtensionData* Candidate = ConstantExtensionData[Index].Data.Get();
-		    if (*Candidate == *Data)
-		    {
-			    return Index;
-		    }
-	    }
-
-	    FExtensionDataConstant& NewConstant = ConstantExtensionData.AddDefaulted_GetRef();
-	    NewConstant.Data = Data;
-
-	    return ConstantExtensionData.Num() - 1;
-    }
-
-
     OP::ADDRESS FProgram::AddConstant(const TManagedPtr<const FLayout>& Layout)
     {
 	    // Ensure unique
@@ -351,7 +330,7 @@ namespace UE::Mutable::Private
 	}
 
 
-	OP::LIST_UINT32_ADDRESS FProgram::AddConstant(const TArray<FMeshSocket>& Sockets)
+	OP::CONSTANT_LIST_UINT32_INDEX FProgram::AddConstant(const TArray<FMeshSocket>& Sockets)
 	{
 		TArray<uint32> SocketIds;
 		SocketIds.Reset(Sockets.Num());
@@ -505,16 +484,7 @@ namespace UE::Mutable::Private
 	    return Index;
     }
 #endif // WITH_EDITOR
-
-    void FProgram::GetExtensionDataConstant(int32 ConstantIndex, TManagedPtr<const FExtensionData>& Result) const
-    {
-	    const FExtensionDataConstant& Constant = ConstantExtensionData[ConstantIndex];
-
-	    check(Constant.Data);
-
-	    Result = Constant.Data;
-    }
-
+	
 
     EOpType FProgram::GetOpType(OP::ADDRESS At) const
     {
diff --git a/Engine/Plugins/Mutable/Source/MutableRuntime/Private/MuR/OpMeshMerge.cpp b/Engine/Plugins/Mutable/Source/MutableRuntime/Private/MuR/OpMeshMerge.cpp
--- a/Engine/Plugins/Mutable/Source/MutableRuntime/Private/MuR/OpMeshMerge.cpp
+++ b/Engine/Plugins/Mutable/Source/MutableRuntime/Private/MuR/OpMeshMerge.cpp
@@ -1658,33 +1658,43 @@ namespace OpMeshMergeClothInternal
 
 
 		// Tags
-		Result->Tags = pFirst->Tags;
+		Result->GameplayTags = pFirst->GameplayTags;
 
-		for (const FString& SecondTag : pSecond->Tags)
+		for (const FName& SecondTag : pSecond->GameplayTags)
 		{
-			Result->Tags.AddUnique(SecondTag);
+			Result->GameplayTags.AddUnique(SecondTag);
 		}
 
-
 		// Streamed Resources
-		Result->StreamedResources = pFirst->StreamedResources;
+		Result->AssetUserData = pFirst->AssetUserData;
 
-		const int32 NumStreamedResources = pSecond->StreamedResources.Num();
-		for (int32 Index = 0; Index < NumStreamedResources; ++Index)
+		const int32 NumAsserUserData = pSecond->AssetUserData.Num();
+		for (int32 Index = 0; Index < NumAsserUserData; ++Index)
 		{
-			Result->StreamedResources.AddUnique(pSecond->StreamedResources[Index]);
+			Result->AssetUserData.AddUnique(pSecond->AssetUserData[Index]);
 		}
-
-		// Skeletal Meshes
-		Result->SkeletalMeshes = pFirst->SkeletalMeshes;
-
-		const int32 NumSkeletalMeshes = pSecond->SkeletalMeshes.Num();
-		for (int32 Index = 0; Index < NumSkeletalMeshes; ++Index)
+	
+		// Animation slots
+		Result->AnimationSlots = pFirst->AnimationSlots;
+	
+		for (const TPair<FName, TPassthroughObjectPtr<UClass>>& Pair : pSecond->AnimationSlots)
 		{
-			Result->SkeletalMeshes.AddUnique(pSecond->SkeletalMeshes[Index]);
+			const TPair<FName, TPassthroughObjectPtr<UClass>>* FindResult = Result->AnimationSlots.FindByPredicate([&](const TPair<FName, TPassthroughObjectPtr<UClass>>& Element)
+			{
+				return Element.Key == Pair.Key;
+			});
+			
+			if (!FindResult)
+			{
+				Result->AnimationSlots.Add(Pair);
+			}
+			else if (FindResult->Value != Pair.Value)
+			{
+				UE_LOGF(LogMutableCore, Warning, "Unable to merge AnimBP Slots. AnimBP Slot Name [%ls] already exists", *FindResult->Key.ToString());
+			}
 		}
 
-		// Mesh morphs.
+		// Mesh morphs
 		constexpr EMorphUsageFlags MorphUsageFilter = EMorphUsageFlags::Merged;
 		if (pFirst->HasMorphs() && !pSecond->HasMorphs())
 		{
diff --git a/Engine/Plugins/Mutable/Source/MutableTools/Internal/MuT/NodeExtensionDataConstant.h b/Engine/Plugins/Mutable/Source/MutableTools/Internal/MuT/NodeExtensionDataConstant.h
--- a/Engine/Plugins/Mutable/Source/MutableTools/Internal/MuT/NodeExtensionDataConstant.h
+++ b/Engine/Plugins/Mutable/Source/MutableTools/Internal/MuT/NodeExtensionDataConstant.h
@@ -2,11 +2,10 @@
 
 #pragma once
 
-#include "MuR/ExtensionData.h"
-#include "MuR/Ptr.h"
-#include "MuR/RefCounted.h"
 #include "MuT/NodeExtensionData.h"
 
+#include "MuR/PassthroughObject.h"
+
 #define UE_API MUTABLETOOLS_API
 
 
@@ -18,27 +17,14 @@ namespace UE::Mutable::Private
 	class NodeExtensionDataConstant : public NodeExtensionData
 	{
 	public:
-
-		TManagedPtr<const FExtensionData> Value;
-
-	public:
+		PASSTHROUGH_ID ExtensionDataId = PASSTHROUGH_ID_INVALID;
 
 		//-----------------------------------------------------------------------------------------
 		// Node Interface
 		//-----------------------------------------------------------------------------------------
 		virtual const FNodeType* GetType() const override { return GetStaticType(); }
 		static const FNodeType* GetStaticType() { return &StaticType; }
-
-		//-----------------------------------------------------------------------------------------
-		// Own Interface
-		//-----------------------------------------------------------------------------------------
-
-		/** Deprecated. Access Value attribute directly. */
-		void SetValue(const TSharedPtr<const FExtensionData>& In) 
-		{
-			Value = MakeManaged<FExtensionData>(*In);
-		}
-
+	
 	protected:
 
 		/** Forbidden. Manage with the Ptr<> template. */
diff --git a/Engine/Plugins/Mutable/Source/MutableTools/Internal/MuT/NodeObjectNew.h b/Engine/Plugins/Mutable/Source/MutableTools/Internal/MuT/NodeObjectNew.h
--- a/Engine/Plugins/Mutable/Source/MutableTools/Internal/MuT/NodeObjectNew.h
+++ b/Engine/Plugins/Mutable/Source/MutableTools/Internal/MuT/NodeObjectNew.h
@@ -38,12 +38,7 @@ namespace UE::Mutable::Private
 		TArray<Ptr<NodeObject>> Children;
 
 		/** Extension data attached to this object. */
-		struct FNamedExtensionDataNode
-		{
-			Ptr<NodeExtensionData> Node;
-			FString Name;
-		};
-		TArray<FNamedExtensionDataNode> ExtensionDataNodes;
+		TArray<Ptr<NodeExtensionData>> ExtensionDataNodes;
 
 		/** States defined in this object. */
 		TArray<FObjectState> States;
@@ -83,7 +78,7 @@ namespace UE::Mutable::Private
 			uint8 NumExtraLODsToBuildAfterFirstLOD);
 
 		//! Connect a node that produces ExtensionData to be added to the final Instance, and provide a name to associate with the data
-		UE_API void AddExtensionDataNode(Ptr<NodeExtensionData> Node, const FString& Name);
+		UE_API void AddExtensionDataNode(Ptr<NodeExtensionData> Node);
 
 	protected:
 
diff --git a/Engine/Plugins/Mutable/Source/MutableTools/Private/MuT/ASTOpAddExtensionData.cpp b/Engine/Plugins/Mutable/Source/MutableTools/Private/MuT/ASTOpAddExtensionData.cpp
--- a/Engine/Plugins/Mutable/Source/MutableTools/Private/MuT/ASTOpAddExtensionData.cpp
+++ b/Engine/Plugins/Mutable/Source/MutableTools/Private/MuT/ASTOpAddExtensionData.cpp
@@ -32,9 +32,8 @@ bool ASTOpAddExtensionData::IsEqual(const ASTOp& OtherUntyped) const
 	if (OtherUntyped.GetOpType() == GetOpType())
 	{
 		const ASTOpAddExtensionData* Other = static_cast<const ASTOpAddExtensionData*>(&OtherUntyped);
-		return Instance == Other->Instance
-			&& ExtensionData == Other->ExtensionData
-			&& ExtensionDataName == Other->ExtensionDataName;
+		return Instance == Other->Instance &&
+			ExtensionData == Other->ExtensionData;
 	}
 
 	return false;
@@ -56,7 +55,6 @@ UE::Mutable::Private::Ptr<ASTOp> ASTOpAddExtensionData::Clone(MapChildFuncRef Ma
 
 	NewInstance->Instance = MapChild(Instance.child());
 	NewInstance->ExtensionData = MapChild(ExtensionData.child());
-	NewInstance->ExtensionDataName = ExtensionDataName;
 
 	return NewInstance;
 }
@@ -89,9 +87,6 @@ void ASTOpAddExtensionData::Link(FProgram& Program, FLinkerOptions*)
 	check(ExtensionData->linkedAddress);
 	Args.ExtensionData = ExtensionData->linkedAddress;
 
-	check(ExtensionDataName.Len() > 0);
-	Args.ExtensionDataName = Program.AddConstant(ExtensionDataName);
-
 	linkedAddress = (OP::ADDRESS)Program.OpAddress.Num();
 	Program.OpAddress.Add((uint32_t)Program.ByteCode.Num());
 	AppendCode(Program.ByteCode, EOpType::IN_ADDEXTENSIONDATA);
diff --git a/Engine/Plugins/Mutable/Source/MutableTools/Private/MuT/ASTOpAddExtensionData.h b/Engine/Plugins/Mutable/Source/MutableTools/Private/MuT/ASTOpAddExtensionData.h
--- a/Engine/Plugins/Mutable/Source/MutableTools/Private/MuT/ASTOpAddExtensionData.h
+++ b/Engine/Plugins/Mutable/Source/MutableTools/Private/MuT/ASTOpAddExtensionData.h
@@ -23,7 +23,6 @@ public:
 
 	ASTChild Instance;
 	ASTChild ExtensionData;
-	FString ExtensionDataName;
 
 public:
 
diff --git a/Engine/Plugins/Mutable/Source/MutableTools/Private/MuT/ASTOpConstantExtensionData.cpp b/Engine/Plugins/Mutable/Source/MutableTools/Private/MuT/ASTOpConstantExtensionData.cpp
--- a/Engine/Plugins/Mutable/Source/MutableTools/Private/MuT/ASTOpConstantExtensionData.cpp
+++ b/Engine/Plugins/Mutable/Source/MutableTools/Private/MuT/ASTOpConstantExtensionData.cpp
@@ -16,7 +16,7 @@ bool ASTOpConstantExtensionData::IsEqual(const ASTOp& OtherUntyped) const
 	if (OtherUntyped.GetOpType() == GetOpType())
 	{
 		const ASTOpConstantExtensionData* Other = static_cast<const ASTOpConstantExtensionData*>(&OtherUntyped);
-		return Other->Value == Value;
+		return Other->ExtensionDataId == ExtensionDataId;
 	}
 
 	return false;
@@ -27,8 +27,7 @@ bool ASTOpConstantExtensionData::IsEqual(const ASTOp& OtherUntyped) const
 UE::Mutable::Private::Ptr<ASTOp> ASTOpConstantExtensionData::Clone(MapChildFuncRef MapChild) const
 {
 	Ptr<ASTOpConstantExtensionData> Result = new ASTOpConstantExtensionData();
-
-	Result->Value = Value;
+	Result->ExtensionDataId = ExtensionDataId;
 
 	return Result;
 }
@@ -37,7 +36,11 @@ UE::Mutable::Private::Ptr<ASTOp> ASTOpConstantExtensionData::Clone(MapChildFuncR
 //-------------------------------------------------------------------------------------------------
 uint32 ASTOpConstantExtensionData::Hash() const
 {
-	return Value != nullptr ? Value->Hash() : 0;
+	uint32 Result = GetTypeHash(GetOpType());
+		
+	Result = HashCombineFast(Result, GetTypeHash(ExtensionDataId));
+		
+	return Result;
 }
 
 
@@ -46,14 +49,14 @@ void ASTOpConstantExtensionData::Link(FProgram& Program, FLinkerOptions* Options
 {
 	if (!linkedAddress)
 	{
-		OP::ResourceConstantArgs Args;
+		OP::ExternalDataConstantArgs Args;
 		FMemory::Memzero(Args);
 
-		Args.value = Program.AddConstant(Value);
+		Args.ExternalObjectId = ExtensionDataId;
 
 		linkedAddress = (OP::ADDRESS)Program.OpAddress.Num();
 		Program.OpAddress.Add((uint32_t)Program.ByteCode.Num());
-		AppendCode(Program.ByteCode, EOpType::ED_CONSTANT);
+		AppendCode(Program.ByteCode, GetOpType());
 		AppendCode(Program.ByteCode, Args);
 	}
 }
diff --git a/Engine/Plugins/Mutable/Source/MutableTools/Private/MuT/ASTOpConstantExtensionData.h b/Engine/Plugins/Mutable/Source/MutableTools/Private/MuT/ASTOpConstantExtensionData.h
--- a/Engine/Plugins/Mutable/Source/MutableTools/Private/MuT/ASTOpConstantExtensionData.h
+++ b/Engine/Plugins/Mutable/Source/MutableTools/Private/MuT/ASTOpConstantExtensionData.h
@@ -18,10 +18,7 @@ namespace UE::Mutable::Private
 	class ASTOpConstantExtensionData final : public ASTOp
 	{
 	public:
-
-		TManagedPtr<const FExtensionData> Value;
-
-	public:
+		PASSTHROUGH_ID ExtensionDataId = PASSTHROUGH_ID_INVALID;
 
 		// ASTOp interface
 		EOpType GetOpType() const override { return EOpType::ED_CONSTANT; }
diff --git a/Engine/Plugins/Mutable/Source/MutableTools/Private/MuT/ASTOpMeshAddMetadata.cpp b/Engine/Plugins/Mutable/Source/MutableTools/Private/MuT/ASTOpMeshAddMetadata.cpp
--- a/Engine/Plugins/Mutable/Source/MutableTools/Private/MuT/ASTOpMeshAddMetadata.cpp
+++ b/Engine/Plugins/Mutable/Source/MutableTools/Private/MuT/ASTOpMeshAddMetadata.cpp
@@ -32,8 +32,9 @@ namespace UE::Mutable::Private
 			const ASTOpMeshAddMetadata* Other = static_cast<const ASTOpMeshAddMetadata*>(&OtherUntyped);
 			return
 				Source == Other->Source &&
-				Tags == Other->Tags &&
-				ResourceIds == Other->ResourceIds &&
+				GameplayTags == Other->GameplayTags &&
+				AssetUserDataIds == Other->AssetUserDataIds &&
+				AnimationSlots == Other->AnimationSlots &&
 				SkeletonId == Other->SkeletonId &&
 				BonePosePriority == Other->BonePosePriority &&
 				SocketPriority == Other->SocketPriority &&
@@ -52,8 +53,9 @@ namespace UE::Mutable::Private
 		Ptr<ASTOpMeshAddMetadata> NewOp = new ASTOpMeshAddMetadata();
 		NewOp->Source = MapChild(Source.child());
 		NewOp->SkeletonId = SkeletonId;
-		NewOp->Tags = Tags;
-		NewOp->ResourceIds = ResourceIds;
+		NewOp->GameplayTags = GameplayTags;
+		NewOp->AssetUserDataIds = AssetUserDataIds;
+		NewOp->AnimationSlots = AnimationSlots;
 		NewOp->BonePosePriority = BonePosePriority;
 		NewOp->SocketPriority = SocketPriority;
 		NewOp->Sockets = Sockets;
@@ -96,33 +98,60 @@ namespace UE::Mutable::Private
 
 			Args.Flags = OpEnumFlags::None;
 			
-			if (Tags.Num())
+			if (GameplayTags.Num())
 			{
-				EnumAddFlags(Args.Flags, OpEnumFlags::HasTags);
+				EnumAddFlags(Args.Flags, OpEnumFlags::HasGameplayTags);
 
-				if (Tags.Num() > 1)
+				if (GameplayTags.Num() > 1)
 				{
-					EnumAddFlags(Args.Flags, OpEnumFlags::IsTagList);
-					Args.Tags.ListAddress = Program.AddConstant(Tags);
+					EnumAddFlags(Args.Flags, OpEnumFlags::IsGameplayTagList);
+					Args.GameplayTags.ListIndex = Program.AddConstant(GameplayTags);
 				}
 				else
 				{
-					Args.Tags.StringAddress = Program.AddConstant(Tags[0]);
+					Args.GameplayTags.NameIndex = Program.AddConstant(GameplayTags[0]);
 				}
 			}
 
-			if (ResourceIds.Num())
+			if (AssetUserDataIds.Num())
 			{
-				EnumAddFlags(Args.Flags, OpEnumFlags::HasResources);
+				EnumAddFlags(Args.Flags, OpEnumFlags::HasAssetUserData);
 
-				if (ResourceIds.Num() > 1)
+				if (AssetUserDataIds.Num() > 1)
 				{
-					EnumAddFlags(Args.Flags, OpEnumFlags::IsResourceList);
-					Args.ResourceIds.ListAddress = Program.AddConstant(ResourceIds);
+					EnumAddFlags(Args.Flags, OpEnumFlags::IsAssetUserDataList);
+					Args.AssetUserDataIds.ListAddress = Program.AddConstant(AssetUserDataIds);
 				}
 				else
 				{
-					Args.ResourceIds.ResourceId = ResourceIds[0];
+					Args.AssetUserDataIds.PassthroughId = AssetUserDataIds[0];
+				}
+			}
+			
+			if (AnimationSlots.Num())
+			{
+				EnumAddFlags(Args.Flags, OpEnumFlags::HasAnimationSlots);
+
+				if (AnimationSlots.Num() > 1)
+				{
+					TArray<FName> AnimSlotNames;
+					TArray<PASSTHROUGH_ID> AnimInstances;
+					
+					for (const TPair<FName, PASSTHROUGH_ID>& AnimationSlot : AnimationSlots)
+					{
+						AnimSlotNames.Add(AnimationSlot.Key);
+						AnimInstances.Add(AnimationSlot.Value);
+					}
+					
+					EnumAddFlags(Args.Flags, OpEnumFlags::IsAnimationSlotList);
+					Args.AnimSlotNames.ListIndex = Program.AddConstant(AnimSlotNames);
+					Args.AnimInstances.ListAddress = Program.AddConstant(AnimInstances);
+					
+				}
+				else
+				{
+					Args.AnimSlotNames.NameIndex = Program.AddConstant(AnimationSlots[0].Key);
+					Args.AnimInstances.PassthroughId = AnimationSlots[0].Value;
 				}
 			}
 
@@ -172,7 +201,7 @@ namespace UE::Mutable::Private
 				}
 				else
 				{
-					Args.AdditionalPhysicsAssetsIds.PassthroughID = AdditionalPhysicsAssetIds[0];
+					Args.AdditionalPhysicsAssetsIds.PassthroughId = AdditionalPhysicsAssetIds[0];
 				}
 			}
 
diff --git a/Engine/Plugins/Mutable/Source/MutableTools/Private/MuT/ASTOpMeshAddMetadata.h b/Engine/Plugins/Mutable/Source/MutableTools/Private/MuT/ASTOpMeshAddMetadata.h
--- a/Engine/Plugins/Mutable/Source/MutableTools/Private/MuT/ASTOpMeshAddMetadata.h
+++ b/Engine/Plugins/Mutable/Source/MutableTools/Private/MuT/ASTOpMeshAddMetadata.h
@@ -22,11 +22,14 @@ namespace UE::Mutable::Private
 		ASTChild Source;
 
 		/** Tags to add. */
-		TArray<FString> Tags;
+		TArray<FName> GameplayTags;
 	
-		/** ResourceIds to add. */
-		TArray<uint64> ResourceIds;
+		TArray<PASSTHROUGH_ID> AssetUserDataIds;
 
+		TArray<TPair<FName, PASSTHROUGH_ID>> AnimationSlots;
+		
+		PASSTHROUGH_ID AnimInstanceId = PASSTHROUGH_ID_INVALID;
+		
 		/** SkeletonId to add. */
 		PASSTHROUGH_ID SkeletonId = PASSTHROUGH_ID_INVALID;
 
@@ -46,10 +49,7 @@ namespace UE::Mutable::Private
 
 		/** Additional PhysicsAssets */
 		TArray<PASSTHROUGH_ID> AdditionalPhysicsAssetIds;
-
-
-	public:
-
+	
 		ASTOpMeshAddMetadata();
 		ASTOpMeshAddMetadata(const ASTOpMeshAddMetadata&) = delete;
 		virtual ~ASTOpMeshAddMetadata() override;
diff --git a/Engine/Plugins/Mutable/Source/MutableTools/Private/MuT/ASTOpSwitch.cpp b/Engine/Plugins/Mutable/Source/MutableTools/Private/MuT/ASTOpSwitch.cpp
--- a/Engine/Plugins/Mutable/Source/MutableTools/Private/MuT/ASTOpSwitch.cpp
+++ b/Engine/Plugins/Mutable/Source/MutableTools/Private/MuT/ASTOpSwitch.cpp
@@ -713,8 +713,9 @@ namespace UE::Mutable::Private
 			{
 				// Move the metadata out of the switch if all metadatas are the same
 				bool bAllMetadataIsTheSame = true;
-				TArray<FString> Tags;
-				TArray<uint64> ResourceIds;
+				TArray<FName> GameplayTags;
+				TArray<PASSTHROUGH_ID> AssetUserDataIds;
+				TArray<TPair<FName, PASSTHROUGH_ID>> AnimationSlots;
 				PASSTHROUGH_ID SkeletonId = PASSTHROUGH_ID_INVALID;
 				int16 BonePosePriority = 0;
 				PASSTHROUGH_ID PhysicsAssetId = PASSTHROUGH_ID_INVALID;
@@ -724,8 +725,9 @@ namespace UE::Mutable::Private
 				{
 					check(Default->GetOpType() == EOpType::ME_ADDMETADATA);
 					const ASTOpMeshAddMetadata* Typed = static_cast<const ASTOpMeshAddMetadata*>(Default.child().get());
-					Tags = Typed->Tags;
-					ResourceIds = Typed->ResourceIds;
+					GameplayTags = Typed->GameplayTags;
+					AssetUserDataIds = Typed->AssetUserDataIds;
+					AnimationSlots = Typed->AnimationSlots;
 					SkeletonId = Typed->SkeletonId;
 					BonePosePriority = Typed->BonePosePriority;
 					PhysicsAssetId = Typed->PhysicsAssetId;
@@ -741,18 +743,20 @@ namespace UE::Mutable::Private
 
 					check(Case.Branch->GetOpType() == EOpType::ME_ADDMETADATA);
 					const ASTOpMeshAddMetadata* Typed = static_cast<const ASTOpMeshAddMetadata*>(Case.Branch.child().get());
-					if (Tags.IsEmpty() && ResourceIds.IsEmpty() &&  SkeletonId == PASSTHROUGH_ID_INVALID && PhysicsAssetId == PASSTHROUGH_ID_INVALID && AdditionalPhysicsAssetIds.IsEmpty())
+					if (GameplayTags.IsEmpty() && AssetUserDataIds.IsEmpty() &&  SkeletonId == PASSTHROUGH_ID_INVALID && PhysicsAssetId == PASSTHROUGH_ID_INVALID && AdditionalPhysicsAssetIds.IsEmpty())
 					{
-						Tags = Typed->Tags;
-						ResourceIds = Typed->ResourceIds;
+						GameplayTags = Typed->GameplayTags;
+						AssetUserDataIds = Typed->AssetUserDataIds;
+						AnimationSlots = Typed->AnimationSlots;
 						SkeletonId = Typed->SkeletonId;
 						BonePosePriority = Typed->BonePosePriority;
 						PhysicsAssetId = Typed->PhysicsAssetId;
 						AdditionalPhysicsAssetIds = Typed->AdditionalPhysicsAssetIds;
 					}
 					else if (
-							Typed->Tags != Tags || 
-							Typed->ResourceIds != ResourceIds || 
+							Typed->GameplayTags != GameplayTags || 
+							Typed->AssetUserDataIds != AssetUserDataIds || 
+							Typed->AnimationSlots != AnimationSlots || 
 							Typed->SkeletonId != SkeletonId ||
 							Typed->BonePosePriority != BonePosePriority ||
 							Typed->PhysicsAssetId != PhysicsAssetId ||
@@ -766,8 +770,9 @@ namespace UE::Mutable::Private
 				if (bAllMetadataIsTheSame)
 				{
 					Ptr<ASTOpMeshAddMetadata> New = new ASTOpMeshAddMetadata();
-					New->Tags = Tags;
-					New->ResourceIds = ResourceIds;
+					New->GameplayTags = GameplayTags;
+					New->AssetUserDataIds = AssetUserDataIds;
+					New->AnimationSlots = AnimationSlots;
 					New->SkeletonId = SkeletonId;
 					New->BonePosePriority= BonePosePriority;
 					New->PhysicsAssetId = PhysicsAssetId;
diff --git a/Engine/Plugins/Mutable/Source/MutableTools/Private/MuT/CodeGenerator.cpp b/Engine/Plugins/Mutable/Source/MutableTools/Private/MuT/CodeGenerator.cpp
--- a/Engine/Plugins/Mutable/Source/MutableTools/Private/MuT/CodeGenerator.cpp
+++ b/Engine/Plugins/Mutable/Source/MutableTools/Private/MuT/CodeGenerator.cpp
@@ -1488,19 +1488,16 @@ namespace UE::Mutable::Private
 		}
 
 		// Add an ASTOpAddExtensionData for each connected ExtensionData node
-		for (const NodeObjectNew::FNamedExtensionDataNode& NamedNode : InNode->ExtensionDataNodes)
+		for (const Ptr<NodeExtensionData>& Node : InNode->ExtensionDataNodes)
 		{
-			if (!NamedNode.Node.get())
+			if (!Node.get())
 			{
 				// No node connected
 				continue;
 			}
 
-			// Name must be valid
-			check(NamedNode.Name.Len() > 0);
-
 			FExtensionDataGenerationResult ChildResult;
-			GenerateExtensionData(ChildResult, Options, NamedNode.Node);
+			GenerateExtensionData(ChildResult, Options, Node);
 
 			if (!ChildResult.Op.get())
 			{
@@ -1511,7 +1508,6 @@ namespace UE::Mutable::Private
 			FConditionalExtensionDataOp& SavedOp = ConditionalExtensionDataOps.AddDefaulted_GetRef();
 			SavedOp.Condition = Options.CurrentObjectCondition;
 			SavedOp.ExtensionDataOp = ChildResult.Op;
-			SavedOp.ExtensionDataName = NamedNode.Name;
 		}
 
 		Ptr<ASTOp> RootOp = LastCompOp;
@@ -1523,7 +1519,6 @@ namespace UE::Mutable::Private
 				Ptr<ASTOpAddExtensionData> ExtensionPinOp = new ASTOpAddExtensionData();
 				ExtensionPinOp->Instance = ASTChild(ExtensionPinOp, RootOp);
 				ExtensionPinOp->ExtensionData = ASTChild(ExtensionPinOp, SavedOp.ExtensionDataOp);
-				ExtensionPinOp->ExtensionDataName = SavedOp.ExtensionDataName;
 
 				if (SavedOp.Condition.get())
 				{
diff --git a/Engine/Plugins/Mutable/Source/MutableTools/Private/MuT/CodeGenerator.h b/Engine/Plugins/Mutable/Source/MutableTools/Private/MuT/CodeGenerator.h
--- a/Engine/Plugins/Mutable/Source/MutableTools/Private/MuT/CodeGenerator.h
+++ b/Engine/Plugins/Mutable/Source/MutableTools/Private/MuT/CodeGenerator.h
@@ -378,6 +378,7 @@ namespace UE::Mutable::Private
 			UE::FMutex Mutex;
 			TSet<uint32> Map;
 		};
+    	
 		FUniqueMeshIds UniqueMeshIds;
 
 		/** Any task accessing GenerateMeshConstantState should run in this pipe. */
@@ -416,7 +417,6 @@ namespace UE::Mutable::Private
 		{
 			Ptr<ASTOp> Condition;
 			Ptr<ASTOp> ExtensionDataOp;
-			FString ExtensionDataName;
 		};
 
 		TArray<FConditionalExtensionDataOp> ConditionalExtensionDataOps;
diff --git a/Engine/Plugins/Mutable/Source/MutableTools/Private/MuT/CodeGenerator_ExtensionData.cpp b/Engine/Plugins/Mutable/Source/MutableTools/Private/MuT/CodeGenerator_ExtensionData.cpp
--- a/Engine/Plugins/Mutable/Source/MutableTools/Private/MuT/CodeGenerator_ExtensionData.cpp
+++ b/Engine/Plugins/Mutable/Source/MutableTools/Private/MuT/CodeGenerator_ExtensionData.cpp
@@ -65,19 +65,9 @@ class Node;
 	void CodeGenerator::GenerateExtensionData_Constant(FExtensionDataGenerationResult& OutResult, const FGenericGenerationOptions& Options, const NodeExtensionDataConstant* Constant)
 	{
 		Ptr<ASTOpConstantExtensionData> Op = new ASTOpConstantExtensionData();
-		OutResult.Op = Op;
+		Op->ExtensionDataId = Constant->ExtensionDataId;
 
-		TManagedPtr<const FExtensionData> Data = Constant->Value;
-		if (!Data)
-		{
-			// Data can't be null, so make an empty one
-			Data = MakeManaged<FExtensionData>();
-			
-			// Log an error message
-			ErrorLog->Add("Constant extension data not set", ELMT_WARNING, Constant->GetMessageContext());
-		}
-
-		Op->Value = Data;
+		OutResult.Op = Op;
 	}
 
 	//---------------------------------------------------------------------------------------------
diff --git a/Engine/Plugins/Mutable/Source/MutableTools/Private/MuT/CodeGenerator_Mesh.cpp b/Engine/Plugins/Mutable/Source/MutableTools/Private/MuT/CodeGenerator_Mesh.cpp
--- a/Engine/Plugins/Mutable/Source/MutableTools/Private/MuT/CodeGenerator_Mesh.cpp
+++ b/Engine/Plugins/Mutable/Source/MutableTools/Private/MuT/CodeGenerator_Mesh.cpp
@@ -796,9 +796,9 @@ namespace UE::Mutable::Private
 				}
 
 				// Separate metadata from the mesh
-				TArray<FString> Tags = ReferenceMesh->Tags; // __AnimBP Tags are added to the ReferenceMesh
-
-				TArray<uint64> ResourceIds;
+				TArray<FName> GameplayTags;
+				TArray<PASSTHROUGH_ID> AssetUserDataIds;
+				TArray<TPair<FName, PASSTHROUGH_ID>> AnimationSlots;
 				PASSTHROUGH_ID SkeletonID = PASSTHROUGH_ID_INVALID;
 				const int8 BonePosePriority = NodeCopy->BonePosePriority;
 				TArray<FMeshSocket> Sockets;
@@ -806,10 +806,11 @@ namespace UE::Mutable::Private
 				PASSTHROUGH_ID PhysicsAssetId = PASSTHROUGH_ID_INVALID;
 				TArray<PASSTHROUGH_ID> AdditionalPhysicsAssetId;
 
-				const bool bMeshHasMetadata = ReferenceMesh->Tags.Num() ||
+				const bool bMeshHasMetadata = 
 					Mesh->SkeletonObjects.Num() > 0 ||
-					Mesh->Tags.Num() ||
-					Mesh->StreamedResources.Num() ||
+					Mesh->GameplayTags.Num() ||
+					Mesh->AssetUserData.Num() ||
+					Mesh->AnimationSlots.Num() ||
 					BonePosePriority != 0 ||
 					Mesh->Sockets.Num() != 0 ||
 					(Mesh->Sockets.Num() && SocketPriority != 0) ||
@@ -818,16 +819,18 @@ namespace UE::Mutable::Private
 
 				if (bMeshHasMetadata)
 				{
-					for (FString& Tag : Mesh->Tags)
+					GameplayTags = Mesh->GameplayTags;
+
+					for (const TPassthroughObjectPtr<UAssetUserData>& AssetUserData : Mesh->AssetUserData)
 					{
-						Tags.AddUnique(Tag);
+						AssetUserDataIds.AddUnique(AssetUserData.GetId());
 					}
-
-					for (uint64 ResourceId : Mesh->StreamedResources)
+					
+					for (const TPair<FName, TPassthroughObjectPtr<UClass>>& AnimationSlot : Mesh->AnimationSlots)
 					{
-						ResourceIds.AddUnique(ResourceId);
+						AnimationSlots.Add({ AnimationSlot.Key, AnimationSlot.Value.GetId() });
 					}
-
+					
 					SkeletonID = Mesh->SkeletonObjects.IsEmpty() ? PASSTHROUGH_ID_INVALID : Mesh->SkeletonObjects[0].GetId();
 
 					if (Mesh->PhysicsAssets.Num())
@@ -843,8 +846,9 @@ namespace UE::Mutable::Private
 
 					TManagedPtr<FMesh> MetadatalessMesh = CloneOrTakeOver(MoveTemp(Mesh));
 
-					MetadatalessMesh->Tags.SetNum(0, EAllowShrinking::No);
-					MetadatalessMesh->StreamedResources.SetNum(0, EAllowShrinking::No);
+					MetadatalessMesh->GameplayTags.SetNum(0, EAllowShrinking::No);
+					MetadatalessMesh->AssetUserData.SetNum(0, EAllowShrinking::No);
+					MetadatalessMesh->AnimationSlots.SetNum(0, EAllowShrinking::No);
 					MetadatalessMesh->SkeletonObjects.SetNum(0, EAllowShrinking::No);
 					MetadatalessMesh->PhysicsAssets.SetNum(0, EAllowShrinking::No);
 					MetadatalessMesh->AdditionalPhysicsAssets.SetNum(0, EAllowShrinking::No);
@@ -933,9 +937,13 @@ namespace UE::Mutable::Private
 				{
 					Ptr<ASTOpMeshAddMetadata> AddMetadataOp = new ASTOpMeshAddMetadata;
 					AddMetadataOp->Source = LastMeshOp;
-
-					AddMetadataOp->Tags        = MoveTemp(Tags);
-					AddMetadataOp->ResourceIds = MoveTemp(ResourceIds);
+					
+					AddMetadataOp->GameplayTags = MoveTemp(GameplayTags);
+					
+					AddMetadataOp->AssetUserDataIds = MoveTemp(AssetUserDataIds);
+					
+					AddMetadataOp->AnimationSlots = MoveTemp(AnimationSlots);
+					
 					AddMetadataOp->SkeletonId = SkeletonID;
 
 					AddMetadataOp->BonePosePriority = BonePosePriority;
diff --git a/Engine/Plugins/Mutable/Source/MutableTools/Private/MuT/NodeObjectNew.cpp b/Engine/Plugins/Mutable/Source/MutableTools/Private/MuT/NodeObjectNew.cpp
--- a/Engine/Plugins/Mutable/Source/MutableTools/Private/MuT/NodeObjectNew.cpp
+++ b/Engine/Plugins/Mutable/Source/MutableTools/Private/MuT/NodeObjectNew.cpp
@@ -92,11 +92,9 @@ namespace UE::Mutable::Private
 		Data.NumExtraLODsToBuildAfterFirstLOD = NumExtraLODsToBuildAfterFirstLOD;
     }
 
-	void NodeObjectNew::AddExtensionDataNode(Ptr<NodeExtensionData> Node, const FString& InName)
+	void NodeObjectNew::AddExtensionDataNode(Ptr<NodeExtensionData> Node)
 	{
-		NodeObjectNew::FNamedExtensionDataNode& Entry = ExtensionDataNodes.AddDefaulted_GetRef();
-		Entry.Node = Node;
-		Entry.Name = InName;
+		ExtensionDataNodes.Add(Node);
 	}
 }
 
