[mutable] Made the Passthrough Skeletal Mesh node support having an overlay and override material for the same slot.
- The prior implementation wrongfully assumed that the end user would only override or overlay a slot, but not both.

#rnx

#Synced-CL 51291550

diff --git a/Engine/Plugins/Mutable/Source/CustomizableObjectEditor/Private/MuCOE/GenerateMutableSource/GenerateMutableSourceComponent.cpp b/Engine/Plugins/Mutable/Source/CustomizableObjectEditor/Private/MuCOE/GenerateMutableSource/GenerateMutableSourceComponent.cpp
--- a/Engine/Plugins/Mutable/Source/CustomizableObjectEditor/Private/MuCOE/GenerateMutableSource/GenerateMutableSourceComponent.cpp
+++ b/Engine/Plugins/Mutable/Source/CustomizableObjectEditor/Private/MuCOE/GenerateMutableSource/GenerateMutableSourceComponent.cpp
@@ -117,7 +117,7 @@ UE::Mutable::Private::Ptr<UE::Mutable::Private::NodeComponent> GenerateMutableSo
 			// Get the parent name using the parent's macro context
 			GenerationContext.CurrentSkeletalMeshComponent = ComponentInfo->ComponentId;
 
-			// Restore GenerationContext MacroContex
+			// Restore GenerationContext MacroContext
 			GenerationContext.MacroNodesStack = MacroContextCopy;
 
 			UE::Mutable::Private::Ptr<UE::Mutable::Private::NodeComponentEdit> NodeComponentEdit = new UE::Mutable::Private::NodeComponentEdit();
@@ -162,57 +162,65 @@ UE::Mutable::Private::Ptr<UE::Mutable::Private::NodeComponent> GenerateMutableSo
 			}
 		}
 		
-		
 		// Override/Overlay Material slots
-		TArray<FName> ProcessedSlotNames;
+		TArray<FName> ProcessedOverrideMaterialSlotNames;
+		TArray<FName> ProcessedOverlayMaterialSlotNames;
 		for (const FEdGraphPinReference& InputMaterialPinReference : TypedComponentPassthroughMesh->MaterialSlotPins)
 		{
-			// Pin Data
-			FName TargetMaterialSlotName;
-			if (const UEdGraphPin* InputMaterialPin = InputMaterialPinReference.Get())
-			{
-				const UCONodeComponentPassthroughMeshMaterialPinData& OverrideMaterialPinData =
-					TypedComponentPassthroughMesh->GetPinData<UCONodeComponentPassthroughMeshMaterialPinData>(*InputMaterialPin);
-
-				TargetMaterialSlotName = OverrideMaterialPinData.TargetMaterialSlotName;
-			}
-			
 			// Material node
-			UE::Mutable::Private::Ptr<UE::Mutable::Private::NodeMaterial> MaterialNode = nullptr;
 			if (const UEdGraphPin* MaterialOverridePin = InputMaterialPinReference.Get())
 			{
 				if (const UEdGraphPin* ConnectedPin = FollowInputPin(*MaterialOverridePin))
 				{
+					const UCONodeComponentPassthroughMeshMaterialPinData& OverrideMaterialPinData =
+						TypedComponentPassthroughMesh->GetPinData<UCONodeComponentPassthroughMeshMaterialPinData>(*MaterialOverridePin);
+					
+					const FName TargetMaterialSlotName = OverrideMaterialPinData.TargetMaterialSlotName;
+					
 					// Skip invalid slot names (None)
 					if (TargetMaterialSlotName.IsNone())
 					{
-						GenerationContext.Log(LOCTEXT("NoneMaterialSlotName", "'None' is not a valid Material slot name."), Node, EMessageSeverity::Warning);
+						GenerationContext.Log(LOCTEXT("NoneMaterialSlotName", "'None' is not a valid Material slot name."), Node, EMessageSeverity::Error);
 						continue;
 					}
 					
-					// Skip repeated slot names
-					if (ProcessedSlotNames.Contains(TargetMaterialSlotName))
+					const FName TargetMaterialOperationName = InputMaterialPinReference.Get()->PinType.PinSubCategory;
+					
+					bool bSlotMaterialOperationAlreadyFound = false;
+					if (TargetMaterialOperationName == UEdGraphSchema_CustomizableObject::PSC_Material_Override)
+					{
+						bSlotMaterialOperationAlreadyFound = ProcessedOverrideMaterialSlotNames.Contains(TargetMaterialSlotName);
+					}
+					else if (TargetMaterialOperationName == UEdGraphSchema_CustomizableObject::PSC_Material_Overlay)
 					{
-						GenerationContext.Log(FText::Format(LOCTEXT("RepeatedMaterialSlotName", "The slot name '{0}' has been defined multiple times. Only the first entry will be used."), FText::FromName(TargetMaterialSlotName)), Node, EMessageSeverity::Warning);
+						bSlotMaterialOperationAlreadyFound = ProcessedOverlayMaterialSlotNames.Contains(TargetMaterialSlotName);
+					}
+					
+					// Skip repeated slot operations. 
+					if (bSlotMaterialOperationAlreadyFound)
+					{
+						const FText FriendlyOperationTypeName = UEdGraphSchema_CustomizableObject::GetPinSubCategoryFriendlyName(TargetMaterialOperationName);
+						GenerationContext.Log(FText::Format(LOCTEXT("RepeatedMaterialSlot", "The slot with '{0}' name and '{1}' type has been defined multiple times."), FText::FromName(TargetMaterialSlotName), FriendlyOperationTypeName), Node, EMessageSeverity::Error);
 						continue;
 					}
 				
-					MaterialNode = GenerateMutableSourceMaterial(ConnectedPin, GenerationContext, FGenerationMaterialOptions());
-				}
-			}
-
-			if (MaterialNode)
-			{
-				if (InputMaterialPinReference.Get()->PinType.PinSubCategory == UEdGraphSchema_CustomizableObject::PSC_Material_Override)
-				{
-					ComponentNode->OverrideMaterials.Add(TargetMaterialSlotName, MaterialNode);
-				}
-				else if (InputMaterialPinReference.Get()->PinType.PinSubCategory == UEdGraphSchema_CustomizableObject::PSC_Material_Overlay)
-				{
-					ComponentNode->OverlayMaterials.Add(TargetMaterialSlotName, MaterialNode);
+					const UE::Mutable::Private::Ptr<UE::Mutable::Private::NodeMaterial> MaterialNode = 
+						GenerateMutableSourceMaterial(ConnectedPin, GenerationContext, FGenerationMaterialOptions());
+					
+					if (MaterialNode)
+					{
+						if (TargetMaterialOperationName == UEdGraphSchema_CustomizableObject::PSC_Material_Override)
+						{
+							ComponentNode->OverrideMaterials.Add(TargetMaterialSlotName, MaterialNode);
+							ProcessedOverrideMaterialSlotNames.Add(TargetMaterialSlotName);
+						}
+						else if (TargetMaterialOperationName == UEdGraphSchema_CustomizableObject::PSC_Material_Overlay)
+						{
+							ComponentNode->OverlayMaterials.Add(TargetMaterialSlotName, MaterialNode);
+							ProcessedOverlayMaterialSlotNames.Add(TargetMaterialSlotName);
+						}
+					}
 				}
-				
-				ProcessedSlotNames.Add(TargetMaterialSlotName);
 			}
 		}
 		
