[Mutable] Fix compile-time logic deciding the number of LODs to stream.
Remove the redundant Enable Mesh Streaming setting at the CO level.

#rnx

#Synced-CL 51290764

diff --git a/Engine/Plugins/Mutable/Source/CustomizableObject/Private/MuCO/CustomizableObjectInstance.cpp b/Engine/Plugins/Mutable/Source/CustomizableObject/Private/MuCO/CustomizableObjectInstance.cpp
--- a/Engine/Plugins/Mutable/Source/CustomizableObject/Private/MuCO/CustomizableObjectInstance.cpp
+++ b/Engine/Plugins/Mutable/Source/CustomizableObject/Private/MuCO/CustomizableObjectInstance.cpp
@@ -3415,7 +3415,7 @@ void InitSkeletalMeshData(const TSharedRef<FUpdateContextPrivate>& Context, USke
 					bSuccess = NewAssetUserData->Merge(Pair.Value[Index]);
 					if (!bSuccess)
 					{
-						UE_LOGF(LogMutable, Error, "Unable to merge Asset User Data %ls", *Pair.Key->GetName());
+						//UE_LOGF(LogMutable, Error, "Unable to merge Asset User Data %ls", *Pair.Key->GetName());
 						break;
 					}
 				}
diff --git a/Engine/Plugins/Mutable/Source/CustomizableObject/Private/MuCO/CustomizableObjectSystem.cpp b/Engine/Plugins/Mutable/Source/CustomizableObject/Private/MuCO/CustomizableObjectSystem.cpp
--- a/Engine/Plugins/Mutable/Source/CustomizableObject/Private/MuCO/CustomizableObjectSystem.cpp
+++ b/Engine/Plugins/Mutable/Source/CustomizableObject/Private/MuCO/CustomizableObjectSystem.cpp
@@ -4236,7 +4236,7 @@ bool IsStreamingEnabled(const UCustomizableObject& Object, const int32 State)
 
 	const bool bStateAllowsStreaming = StateData ? !StateData->bDisableMeshStreaming : true;
 
-	return ((Object.bEnableMeshStreaming && bStateAllowsStreaming) || bForceStreamMeshLODs) &&
+	return (bStateAllowsStreaming || bForceStreamMeshLODs) &&
 		bStreamMeshLODs &&
 		IStreamingManager::Get().IsRenderAssetStreamingEnabled(EStreamableRenderAssetType::SkeletalMesh);
 }
diff --git a/Engine/Plugins/Mutable/Source/CustomizableObject/Public/MuCO/CustomizableObject.h b/Engine/Plugins/Mutable/Source/CustomizableObject/Public/MuCO/CustomizableObject.h
--- a/Engine/Plugins/Mutable/Source/CustomizableObject/Public/MuCO/CustomizableObject.h
+++ b/Engine/Plugins/Mutable/Source/CustomizableObject/Public/MuCO/CustomizableObject.h
@@ -267,10 +267,6 @@ public:
 	 * USkeletalMeshes are only reused between the same CO. */
 	UPROPERTY(EditAnywhere, Category = CustomizableObject)
 	bool bEnableMeshCache = false;
-
-	/** If true, Mesh LODs will be streamed on demand. It requires streaming of SkeletalMeshes and Mutable.StreamMeshLODsEnabled to be enabled. */
-	UPROPERTY(EditAnywhere, Category = CustomizableObject)
-	bool bEnableMeshStreaming = true;
 	
 #if WITH_EDITORONLY_DATA
 private:
diff --git a/Engine/Plugins/Mutable/Source/CustomizableObjectEditor/Private/MuCOE/CustomizableObjectCompiler.cpp b/Engine/Plugins/Mutable/Source/CustomizableObjectEditor/Private/MuCOE/CustomizableObjectCompiler.cpp
--- a/Engine/Plugins/Mutable/Source/CustomizableObjectEditor/Private/MuCOE/CustomizableObjectCompiler.cpp
+++ b/Engine/Plugins/Mutable/Source/CustomizableObjectEditor/Private/MuCOE/CustomizableObjectCompiler.cpp
@@ -1040,19 +1040,7 @@ void FCustomizableObjectCompiler::CompileInternal(bool bAsync)
 		}
 
 		ModelResources->NumLODs = GenerationContext.NumLODs;
-
-		if (GenerationContext.bEnableLODStreaming)
-		{
-			ModelResources->NumStreamingLODs = GenerationContext.NumMaxLODsToStream;
-		}
-		else
-		{
-			for (TTuple<UE::Mutable::Private::FComponentId, uint8>& Pair : ModelResources->NumStreamingLODs)
-			{
-				Pair.Value = 0;
-			}
-		}
-
+		ModelResources->NumStreamingLODs = GenerationContext.NumMaxLODsToStream;
 		ModelResources->FirstLODs = GenerationContext.FirstLODAvailable;
 		
 		ModelResources->ParticipatingObjects = MoveTemp(ParticipatingObjects);
diff --git a/Engine/Plugins/Mutable/Source/CustomizableObjectEditor/Private/MuCOE/GenerateMutableSource/GenerateMutableSource.h b/Engine/Plugins/Mutable/Source/CustomizableObjectEditor/Private/MuCOE/GenerateMutableSource/GenerateMutableSource.h
--- a/Engine/Plugins/Mutable/Source/CustomizableObjectEditor/Private/MuCOE/GenerateMutableSource/GenerateMutableSource.h
+++ b/Engine/Plugins/Mutable/Source/CustomizableObjectEditor/Private/MuCOE/GenerateMutableSource/GenerateMutableSource.h
@@ -844,8 +844,6 @@ struct FMutableGraphGenerationContext
 	
 	TMap<UE::Mutable::Private::FComponentId, uint8> NumMaxLODsToStream;
 
-	bool bEnableLODStreaming = true;
-
 	bool bPartialCompilation = false;
 	
 	// Based on the last object visited.
diff --git a/Engine/Plugins/Mutable/Source/CustomizableObjectEditor/Private/MuCOE/GenerateMutableSource/GenerateMutableSourceComponent.cpp b/Engine/Plugins/Mutable/Source/CustomizableObjectEditor/Private/MuCOE/GenerateMutableSource/GenerateMutableSourceComponent.cpp
--- a/Engine/Plugins/Mutable/Source/CustomizableObjectEditor/Private/MuCOE/GenerateMutableSource/GenerateMutableSourceComponent.cpp
+++ b/Engine/Plugins/Mutable/Source/CustomizableObjectEditor/Private/MuCOE/GenerateMutableSource/GenerateMutableSourceComponent.cpp
@@ -454,25 +454,21 @@ void FirstPass(UCustomizableObjectNodeComponentMesh& Node, FMutableGraphGenerati
 
 	FirstLODAvailable = FMath::Clamp(FirstLODAvailable, 0, NumLODs - 1);
 
-	uint8 NumMaxLODsToStream = MAX_MESH_LOD_COUNT;
+	uint8 NumMaxLODsToStream = 0;
 
 	// Find the streaming settings for the target platform
 	if (LODSettings.bOverrideLODStreamingSettings)
 	{
-		GenerationContext.bEnableLODStreaming = LODSettings.bEnableLODStreaming.GetValueForPlatform(*GenerationContext.CompilationContext->Options.TargetPlatform->IniPlatformName());
-		NumMaxLODsToStream = LODSettings.NumMaxStreamedLODs.GetValueForPlatform(*GenerationContext.CompilationContext->Options.TargetPlatform->IniPlatformName());
+		if (LODSettings.bEnableLODStreaming.GetValueForPlatform(*GenerationContext.CompilationContext->Options.TargetPlatform->IniPlatformName()))
+		{
+			NumMaxLODsToStream = LODSettings.NumMaxStreamedLODs.GetValueForPlatform(*GenerationContext.CompilationContext->Options.TargetPlatform->IniPlatformName());
+		}
 	}
-	else
+	else if(Node.ReferenceSkeletalMesh)
 	{
-		for (int32 MeshIndex = 0; MeshIndex < GenerationContext.CompilationContext->ComponentInfos.Num(); ++MeshIndex)
+		if (RefSkeletalMesh->GetEnableLODStreaming(GenerationContext.CompilationContext->Options.TargetPlatform))
 		{
-			RefSkeletalMesh = GenerationContext.CompilationContext->ComponentInfos[MeshIndex].RefSkeletalMesh.Get();
-			check(RefSkeletalMesh);
-
-			GenerationContext.bEnableLODStreaming = GenerationContext.bEnableLODStreaming &&
-				RefSkeletalMesh->GetEnableLODStreaming(GenerationContext.CompilationContext->Options.TargetPlatform);
-
-			NumMaxLODsToStream = FMath::Min(NumMaxLODsToStream, static_cast<uint8>(RefSkeletalMesh->GetMaxNumStreamedLODs(GenerationContext.CompilationContext->Options.TargetPlatform)));
+			NumMaxLODsToStream = static_cast<uint8>(RefSkeletalMesh->GetMaxNumStreamedLODs(GenerationContext.CompilationContext->Options.TargetPlatform));
 		}
 	}
 
