[Mutable] Initialize UTexture pointers to avoid potential invalid access if GetTextureParameterValue fails.

#rnx

#Synced-CL 51290764

diff --git a/Engine/Plugins/Mutable/Source/CustomizableObjectEditor/Private/MuCOE/CustomizableObjectInstanceBakingUtils.cpp b/Engine/Plugins/Mutable/Source/CustomizableObjectEditor/Private/MuCOE/CustomizableObjectInstanceBakingUtils.cpp
--- a/Engine/Plugins/Mutable/Source/CustomizableObjectEditor/Private/MuCOE/CustomizableObjectInstanceBakingUtils.cpp
+++ b/Engine/Plugins/Mutable/Source/CustomizableObjectEditor/Private/MuCOE/CustomizableObjectInstanceBakingUtils.cpp
@@ -620,7 +620,7 @@ bool BakeCustomizableObjectInstance(
 					
 					for (int32 ParameterInfoIndex = 0; ParameterInfoIndex < ParameterInfoObjects.Num(); ParameterInfoIndex++)
 					{
-						UTexture* Texture;
+						UTexture* Texture = nullptr;
 						const FName& ParameterName = ParameterInfoObjects[ParameterInfoIndex].Name;
 						if (MaterialInstance->GetTextureParameterValue(ParameterName, Texture))
 						{
diff --git a/Engine/Plugins/Mutable/Source/CustomizableObjectEditor/Private/MuCOE/Nodes/CONodeSkeletalMeshSection.cpp b/Engine/Plugins/Mutable/Source/CustomizableObjectEditor/Private/MuCOE/Nodes/CONodeSkeletalMeshSection.cpp
--- a/Engine/Plugins/Mutable/Source/CustomizableObjectEditor/Private/MuCOE/Nodes/CONodeSkeletalMeshSection.cpp
+++ b/Engine/Plugins/Mutable/Source/CustomizableObjectEditor/Private/MuCOE/Nodes/CONodeSkeletalMeshSection.cpp
@@ -1011,7 +1011,7 @@ UTexture2D* UCONodeSkeletalMeshSection::GetImageValue(const int32 ImageIndex) co
 	TextureInfo.Index = LayerIndex;
 	TextureInfo.Association = LayerIndex == INDEX_NONE ? GlobalParameter : LayerParameter; // TODO: Blend Parameters
 
-	UTexture* Texture;
+	UTexture* Texture = nullptr;
 	Material->GetTextureParameterValue(TextureInfo, Texture);
 
 	return Cast<UTexture2D>(Texture);
diff --git a/Engine/Plugins/Mutable/Source/CustomizableObjectEditor/Private/MuCOE/Nodes/CustomizableObjectNodeTable.cpp b/Engine/Plugins/Mutable/Source/CustomizableObjectEditor/Private/MuCOE/Nodes/CustomizableObjectNodeTable.cpp
--- a/Engine/Plugins/Mutable/Source/CustomizableObjectEditor/Private/MuCOE/Nodes/CustomizableObjectNodeTable.cpp
+++ b/Engine/Plugins/Mutable/Source/CustomizableObjectEditor/Private/MuCOE/Nodes/CustomizableObjectNodeTable.cpp
@@ -1013,7 +1013,7 @@ UTexture2D* UCustomizableObjectNodeTable::FindReferenceTextureParameter(const UE
 
 	if (Material)
 	{
-		UTexture* OutTexture;
+		UTexture* OutTexture = nullptr;
 		bool bFound = Material->GetTextureParameterValue(FName(*ParameterImageName), OutTexture);
 
 		if (bFound)
diff --git a/Engine/Plugins/Mutable/Source/MutableRuntime/Private/MuR/CodeRunner.cpp b/Engine/Plugins/Mutable/Source/MutableRuntime/Private/MuR/CodeRunner.cpp
--- a/Engine/Plugins/Mutable/Source/MutableRuntime/Private/MuR/CodeRunner.cpp
+++ b/Engine/Plugins/Mutable/Source/MutableRuntime/Private/MuR/CodeRunner.cpp
@@ -6647,7 +6647,7 @@ namespace UE::Mutable::Private
 					ImageParameterInfo.Association = ImageParameterInfo.Index == INDEX_NONE ? GlobalParameter : LayerParameter;
 
 					// Find Texture using the parameter info
-					UTexture* Texture;
+					UTexture* Texture = nullptr;
 					bool bParameterFound = MaterialInterface->GetTextureParameterValue(ImageParameterInfo, Texture);
 
 					// Store Texture into the FMaterial
diff --git a/Engine/Plugins/Mutable/Source/MutableRuntime/Private/MuR/CodeRunnerTasked.cpp b/Engine/Plugins/Mutable/Source/MutableRuntime/Private/MuR/CodeRunnerTasked.cpp
--- a/Engine/Plugins/Mutable/Source/MutableRuntime/Private/MuR/CodeRunnerTasked.cpp
+++ b/Engine/Plugins/Mutable/Source/MutableRuntime/Private/MuR/CodeRunnerTasked.cpp
@@ -2227,7 +2227,7 @@ namespace UE::Mutable::Private
 				ParameterInfo.Index = (int32)Args.LayerIndex;
 				ParameterInfo.Association = ParameterInfo.Index == INDEX_NONE ? GlobalParameter : LayerParameter;
 
-				UTexture* Texture;
+				UTexture* Texture = nullptr;
 				bool bParameterFound = Material->GetTextureParameterValue(ParameterInfo, Texture);
 
 				if (bParameterFound && Texture)
