[mutable] Inlined morphs
- Made the morph nodes be added to the mutable node tree when a morph node is reached instead of waiting to get to the leaf const or param to apply them all at once.
- Added missing Transform node type into the Mutable Graph Viewer.

#rnx

#Synced-CL 51334386

diff --git a/Engine/Plugins/Mutable/Source/CustomizableObjectEditor/Private/MuCOE/GenerateMutableSource/GenerateMutableSource.cpp b/Engine/Plugins/Mutable/Source/CustomizableObjectEditor/Private/MuCOE/GenerateMutableSource/GenerateMutableSource.cpp
--- a/Engine/Plugins/Mutable/Source/CustomizableObjectEditor/Private/MuCOE/GenerateMutableSource/GenerateMutableSource.cpp
+++ b/Engine/Plugins/Mutable/Source/CustomizableObjectEditor/Private/MuCOE/GenerateMutableSource/GenerateMutableSource.cpp
@@ -356,7 +356,6 @@ FGeneratedKey::FGeneratedKey(void* InFunctionAddress, const UEdGraphPin& InPin,
 		{
 			Flags = GenerationContext.MeshGenerationFlags.Last();
 			LayoutFlags = GenerationContext.LayoutGenerationFlags.Last();
-			MeshMorphStack = GenerationContext.MeshMorphStack;
 			bOnlyConnectedLOD = InbOnlyConnectedLOD;
 		}
 	}
diff --git a/Engine/Plugins/Mutable/Source/CustomizableObjectEditor/Private/MuCOE/GenerateMutableSource/GenerateMutableSource.h b/Engine/Plugins/Mutable/Source/CustomizableObjectEditor/Private/MuCOE/GenerateMutableSource/GenerateMutableSource.h
--- a/Engine/Plugins/Mutable/Source/CustomizableObjectEditor/Private/MuCOE/GenerateMutableSource/GenerateMutableSource.h
+++ b/Engine/Plugins/Mutable/Source/CustomizableObjectEditor/Private/MuCOE/GenerateMutableSource/GenerateMutableSource.h
@@ -130,7 +130,7 @@ struct FMorphNodeData
 	// Pin to the node that generates the factor of the morph
 	UE::Mutable::Private::Ptr<UE::Mutable::Private::NodeScalar> FactorNode;
 
-	// Pin of the mesh where the morphs will ble apllied
+	// Pin of the mesh where the morphs will ble applied
 	const UEdGraphPin* MeshPin;
 
 	bool operator==(const FMorphNodeData& Other) const
@@ -174,9 +174,6 @@ public:
 	/** Flag used to generate this mesh. Bit mask of EMutableMeshConversionFlags */
 	EMutableMeshConversionFlags Flags = EMutableMeshConversionFlags::None;
 
-	/** Active morphs at the time of mesh generation. */
-	TArray<FMorphNodeData> MeshMorphStack;
-
 	/** UV Layout modes */
 	FLayoutGenerationFlags LayoutFlags;
 	
@@ -880,9 +877,6 @@ struct FMutableGraphGenerationContext
 	TArray<const UEdGraphNode*> LimitedParameters;
 	int32 ParameterLimitationCount = 0;
 
-	// Stores all morphs to apply them directly to a skeletal mesh node
-	TArray<FMorphNodeData> MeshMorphStack;
-
 	// Current material parameter name to find the corresponding column in a mutable table
 	FString CurrentMaterialTableParameter;
 
diff --git a/Engine/Plugins/Mutable/Source/CustomizableObjectEditor/Private/MuCOE/GenerateMutableSource/GenerateMutableSourceMesh.cpp b/Engine/Plugins/Mutable/Source/CustomizableObjectEditor/Private/MuCOE/GenerateMutableSource/GenerateMutableSourceMesh.cpp
--- a/Engine/Plugins/Mutable/Source/CustomizableObjectEditor/Private/MuCOE/GenerateMutableSource/GenerateMutableSourceMesh.cpp
+++ b/Engine/Plugins/Mutable/Source/CustomizableObjectEditor/Private/MuCOE/GenerateMutableSource/GenerateMutableSourceMesh.cpp
@@ -2407,19 +2407,22 @@ bool GetAndValidateReshapePhysicsToDeform(
 UE::Mutable::Private::NodeMeshPtr GenerateMorphMesh(const UEdGraphPin* Pin,
 	TArray<FMorphNodeData> TypedNodeMorphs,
 	int32 MorphIndex,
-	UE::Mutable::Private::Ptr<UE::Mutable::Private::NodeMesh> SourceNode,
+	const UE::Mutable::Private::Ptr<UE::Mutable::Private::NodeMesh>& SourceNode,
 	FMutableGraphGenerationContext & GenerationContext,
 	const bool bOnlyConnectedLOD,
 	const FString& TableColumnName = "")
 {
 	MUTABLE_CPUPROFILER_SCOPE(GenerateMorphMesh);
 	
+	check(Pin);
+	
 	// SkeletalMesh node
 	const UEdGraphNode * MeshNode = Pin->GetOwningNode();
 	check(MeshNode);
 	
 	// Current morph node
-	const UCustomizableObjectNode* MorphNode = TypedNodeMorphs[MorphIndex].OwningNode;
+	check(TypedNodeMorphs.IsValidIndex(MorphIndex));
+	UCustomizableObjectNode* MorphNode = TypedNodeMorphs[MorphIndex].OwningNode;
 	check(MorphNode);
 	
 	UE::Mutable::Private::Ptr<UE::Mutable::Private::NodeMeshMorph> Result = new UE::Mutable::Private::NodeMeshMorph();
@@ -2440,7 +2443,7 @@ UE::Mutable::Private::NodeMeshPtr GenerateMorphMesh(const UEdGraphPin* Pin,
 		Result->Base = NextMorph;
 	}
 
-	if (UCustomizableObjectNodeMeshMorph* TypedMorphNode = Cast<UCustomizableObjectNodeMeshMorph>(TypedNodeMorphs[MorphIndex].OwningNode))
+	if (UCustomizableObjectNodeMeshMorph* TypedMorphNode = Cast<UCustomizableObjectNodeMeshMorph>(MorphNode))
 	{
 		Result->bReshapeSkeleton = TypedMorphNode->bReshapeSkeleton;
 		Result->bReshapePhysicsVolumes = TypedMorphNode->bReshapePhysicsVolumes;
@@ -2596,14 +2599,11 @@ UE::Mutable::Private::Ptr<UE::Mutable::Private::NodeMesh> GenerateMutableSourceM
 	}
 
 	// Bool that determines if a node can be added to the cache of nodes.
-	// Most nodes need to be added to the cache but there are some that don't. For exampel, MacroInstanceNodes
+	// Most nodes need to be added to the cache but there are some that don't. For example, MacroInstanceNodes
 	bool bCacheNode = true;
 
 	//SkeletalMesh Result
 	UE::Mutable::Private::Ptr<UE::Mutable::Private::NodeMesh> Result;
-
-	//SkeletalMesh + Morphs Result
-	UE::Mutable::Private::Ptr<UE::Mutable::Private::NodeMesh> MorphResult;
 	
 	if (const UCustomizableObjectNodeSkeletalMesh* TypedNodeSkel = Cast<UCustomizableObjectNodeSkeletalMesh>(Node))
 	{
@@ -2680,12 +2680,6 @@ UE::Mutable::Private::Ptr<UE::Mutable::Private::NodeMesh> GenerateMutableSourceM
 			const FString MeshName = TypedNodeSkel->SkeletalMesh.GetLongPackageName().ToLower();
 			MeshNode->SourceDataDescriptor.SourceId = CityHash32(reinterpret_cast<const char*>(*MeshName), MeshName.Len() * sizeof(FString::ElementType));
 			MeshNode->SourceDataDescriptor.OptionalMaxLODSize = 0;
-
-			// Applying Mesh Morph Nodes
-			if (GenerationContext.MeshMorphStack.Num())
-			{
-				MorphResult = GenerateMorphMesh(Pin, GenerationContext.MeshMorphStack, 0, Result, GenerationContext, bOnlyConnectedLOD);
-			}
 		}
 		else
 		{
@@ -2796,15 +2790,24 @@ UE::Mutable::Private::Ptr<UE::Mutable::Private::NodeMesh> GenerateMutableSourceM
 					}
 				}
 			}
-
-			// Factor
-			UE::Mutable::Private::Ptr<UE::Mutable::Private::NodeScalar> FactorNode = GenerateMorphFactor(Node, *TypedNodeMorph->FactorPin(), GenerationContext);
-
-			// Mesh Morph Stack Management
-			FMorphNodeData NewMorphData = { TypedNodeMorph, MorphTargetName, FactorNode, TypedNodeMorph->MeshPin() };
-			GenerationContext.MeshMorphStack.Push(NewMorphData);
-			Result = GenerateMutableSourceMesh(ConnectedPin, GenerationContext, BaseMeshData, false, bOnlyConnectedLOD);
-			GenerationContext.MeshMorphStack.Pop(EAllowShrinking::Yes);
+			
+			// Base Mesh
+			UE::Mutable::Private::Ptr<UE::Mutable::Private::NodeMesh> BaseMesh = GenerateMutableSourceMesh(ConnectedPin, GenerationContext, BaseMeshData,
+				false, bOnlyConnectedLOD);
+			
+			if (BaseMesh)
+			{
+				// Factor
+				UE::Mutable::Private::Ptr<UE::Mutable::Private::NodeScalar> FactorNode = GenerateMorphFactor(Node, *TypedNodeMorph->FactorPin(), GenerationContext);
+				
+				if (FactorNode)
+				{
+					FMorphNodeData NewMorphData = { TypedNodeMorph, MorphTargetName, FactorNode, TypedNodeMorph->MeshPin() };
+					
+					// Apply the morphs to the mesh
+					Result = GenerateMorphMesh(ConnectedPin, {NewMorphData}, 0, BaseMesh, GenerationContext, bOnlyConnectedLOD);
+				}
+			}
 		}
 		else
 		{
@@ -2821,29 +2824,30 @@ UE::Mutable::Private::Ptr<UE::Mutable::Private::NodeMesh> GenerateMutableSourceM
 
 		if (const UEdGraphPin* ConnectedPin = FollowInputPin(*TypedNodeMeshMorphStackApp->GetStackPin()))
 		{
-			const TArray<FString> MorphTargetNames = TypedNodeMeshMorphStackApp->GetMorphList();
-			TArray<FMorphNodeData> MorphDataArray = GenerateMeshMorphStackDefinition(ConnectedPin, TypedNodeMeshMorphStackApp->GetMeshPin(), MorphTargetNames, GenerationContext);
-			int32 AddedMorphs = 0;
-
-			for (const FMorphNodeData& MorphData : MorphDataArray)
-			{
-				GenerationContext.MeshMorphStack.Push(MorphData);
-				AddedMorphs++;
-			}
-
 			if (const UEdGraphPin* MeshConnectedPin = FollowInputPin(*TypedNodeMeshMorphStackApp->GetMeshPin()))
 			{
-				Result = GenerateMutableSourceMesh(MeshConnectedPin, GenerationContext, BaseMeshData, false, bOnlyConnectedLOD);
+				UE::Mutable::Private::NodeMeshPtr BaseMesh = GenerateMutableSourceMesh(MeshConnectedPin, GenerationContext, BaseMeshData, false, bOnlyConnectedLOD);
+				
+				const TArray<FString> MorphTargetNames = TypedNodeMeshMorphStackApp->GetMorphList();
+				TArray<FMorphNodeData> MorphDataArray = GenerateMeshMorphStackDefinition(ConnectedPin, TypedNodeMeshMorphStackApp->GetMeshPin(), MorphTargetNames, GenerationContext);
+				
+				if (MorphDataArray.Num())
+				{
+					Result = GenerateMorphMesh(ConnectedPin, MorphDataArray, 0, BaseMesh, GenerationContext, bOnlyConnectedLOD);
+				}
+				else
+				{
+					GenerationContext.Log(LOCTEXT("MorphStackMorphsNotSet", "No morphs could be found to be used."), Node, EMessageSeverity::Error);
+				}
 			}
-
-			for (int32 MorphIndex = 0; MorphIndex < AddedMorphs; ++MorphIndex)
+			else
 			{
-				GenerationContext.MeshMorphStack.Pop(EAllowShrinking::No);
+				GenerationContext.Log(LOCTEXT("MorphStackBaseMeshConnectionFailed", "Base Mesh connection not found."), Node, EMessageSeverity::Error);
 			}
 		}
 		else
 		{
-			GenerationContext.Log(LOCTEXT("MorphStackConnectionFailed", "Stack definition connection not found."), Node);
+			GenerationContext.Log(LOCTEXT("MorphStackStackConnectionFailed", "Stack definition connection not found."), Node, EMessageSeverity::Error);
 		}
 	}
 
@@ -3402,7 +3406,6 @@ UE::Mutable::Private::Ptr<UE::Mutable::Private::NodeMesh> GenerateMutableSourceM
 
 							if (bGenerateLayouts)
 							{
-
 								TArray<UCustomizableObjectLayout*> Layouts = TypedNodeTable->GetLayouts(Pin);
 								const int32 NumLayouts = Layouts.Num();
 
@@ -3441,12 +3444,6 @@ UE::Mutable::Private::Ptr<UE::Mutable::Private::NodeMesh> GenerateMutableSourceM
 								Source.TableReferenceSkeletalMesh = DefaultSkeletalMesh;
 								MeshTableNode->RealTimeMorphNames = GetUsedRealTimeMorphsNames(Source, bEnableRealTimeMorphTargets, {});
 							}
-
-							// Applying Mesh Morph Nodes
-							if (DefaultSkeletalMesh && GenerationContext.MeshMorphStack.Num())
-							{
-								MorphResult = GenerateMorphMesh(Pin, GenerationContext.MeshMorphStack, 0, Result, GenerationContext, bOnlyConnectedLOD, MutableColumnName);
-							}
 						}
 					}
 				}
@@ -3517,7 +3514,6 @@ UE::Mutable::Private::Ptr<UE::Mutable::Private::NodeMesh> GenerateMutableSourceM
 				const UCustomizableObjectLayout* Layout = Layouts.IsValidIndex(LayoutIndex) ? Layouts[LayoutIndex] : nullptr;
 				if (ensure(Layout))
 				{
-					
 					UE::Mutable::Private::Ptr<UE::Mutable::Private::NodeLayout> LayoutNode = CreateMutableLayoutNode(Layout, bLinkedToExtendMaterial);
 					LayoutNode->SetMessageContext(Node);
 
@@ -3552,12 +3548,6 @@ UE::Mutable::Private::Ptr<UE::Mutable::Private::NodeMesh> GenerateMutableSourceM
 		// gameplay tags?
 
 		Result = BreakNode;
-
-		// Applying Mesh Morph Node
-		if (GenerationContext.MeshMorphStack.Num())
-		{
-			Result = GenerateMorphMesh(Pin, GenerationContext.MeshMorphStack, 0, Result, GenerationContext, bOnlyConnectedLOD);
-		}
 	}
 
 	else if (Cast<UCONodeExternalOperation>(Node))
@@ -3638,12 +3628,6 @@ UE::Mutable::Private::Ptr<UE::Mutable::Private::NodeMesh> GenerateMutableSourceM
 		GenerationContext.GeneratedNodes.Add(Node);
 	}
 
-	// We return the mesh modified by morphs if there is any
-	if (MorphResult)
-	{
-		Result = MorphResult;
-	}
-
 	if (Result)
 	{
 		Result->SetMessageContext(Node);
diff --git a/Engine/Plugins/Mutable/Source/CustomizableObjectEditor/Private/MuCOE/SMutableGraphViewer.cpp b/Engine/Plugins/Mutable/Source/CustomizableObjectEditor/Private/MuCOE/SMutableGraphViewer.cpp
--- a/Engine/Plugins/Mutable/Source/CustomizableObjectEditor/Private/MuCOE/SMutableGraphViewer.cpp
+++ b/Engine/Plugins/Mutable/Source/CustomizableObjectEditor/Private/MuCOE/SMutableGraphViewer.cpp
@@ -944,6 +944,12 @@ void SMutableGraphViewer::GetChildrenForInfo(TSharedPtr<FMutableGraphTreeElement
 		AddChildFunc( NodeMeshTransformWithBone->MatrixNode.get(), TEXT("TRANSFORM"));
 	}
 	
+	else if (ParentNodeType == UE::Mutable::Private::NodeMeshTransform::GetStaticType())
+	{
+		const UE::Mutable::Private::NodeMeshTransform* NodeMeshTransform = StaticCast<UE::Mutable::Private::NodeMeshTransform*>(ParentNode);
+		AddChildFunc(NodeMeshTransform->Source.get(), TEXT("SOURCE"));
+	}
+	
 	else
 	{
 		UE_LOG(LogMutable,Error,TEXT("The node of type %d has not been implemented, so its children won't be added to the tree."), int32(ParentNodeType->Type));
