[Mutable] OpImageLayer fixes and optimizations.
- Add zero mask optimization.
- Enable Softlight vectorized implementation
- Fix Soflight Neon implementation
- Add support for use source from alpha flags.

#rnx

#Synced-CL 51338073

diff --git a/Engine/Plugins/Mutable/Source/MutableRuntime/Private/MuR/CodeRunner.cpp b/Engine/Plugins/Mutable/Source/MutableRuntime/Private/MuR/CodeRunner.cpp
--- a/Engine/Plugins/Mutable/Source/MutableRuntime/Private/MuR/CodeRunner.cpp
+++ b/Engine/Plugins/Mutable/Source/MutableRuntime/Private/MuR/CodeRunner.cpp
@@ -3503,7 +3503,8 @@ namespace UE::Mutable::Private
 						const int32 LODBegin = 0;
 						const int32 LODEnd   = bBlendOnlyOneMip ? 1 : Base->GetLODCount();
 
-						bool bUseBlendSourceFromBlendAlpha = false; // (Args.flags& OP::ImageLayerArgs::F_BLENDED_RGB_FROM_ALPHA) != 0;
+						bool bUseBaseSourceFromBaseAlpha = false;   //(Args.flags & OP::ImageLayerArgs::F_BASE_RGB_FROM_ALPHA) != 0;
+						bool bUseBlendSourceFromBlendAlpha = false; //(Args.flags & OP::ImageLayerArgs::F_BLENDED_RGB_FROM_ALPHA) != 0;
 					
 						if (EBlendType(Args.blendType) != EBlendType::BT_NORMAL_COMBINE)
 						{
@@ -3523,7 +3524,11 @@ namespace UE::Mutable::Private
 									Base.Get(), Base.Get(), Blended.Get(), Mask.Get(), 
 									ColorBlendFunc, AlphaBlendFunc, 
 									LODBegin, LODEnd, 
-									bApplyColorBlendToAlpha, bUseMaskFromBlendAlpha, bUseBlendSourceFromBlendAlpha, BlendAlphaSourceChannel);
+									bApplyColorBlendToAlpha, 
+									bUseMaskFromBlendAlpha, 
+									bUseBaseSourceFromBaseAlpha,
+									bUseBlendSourceFromBlendAlpha, 
+									BlendAlphaSourceChannel);
 						}
 						else
 						{
diff --git a/Engine/Plugins/Mutable/Source/MutableRuntime/Private/MuR/CodeRunnerTasked.cpp b/Engine/Plugins/Mutable/Source/MutableRuntime/Private/MuR/CodeRunnerTasked.cpp
--- a/Engine/Plugins/Mutable/Source/MutableRuntime/Private/MuR/CodeRunnerTasked.cpp
+++ b/Engine/Plugins/Mutable/Source/MutableRuntime/Private/MuR/CodeRunnerTasked.cpp
@@ -634,6 +634,7 @@ namespace UE::Mutable::Private
 		if (EBlendType(Args.blendType) != EBlendType::BT_NORMAL_COMBINE)
 		{
 			bool bApplyColorBlendToAlpha       = (Args.flags & OP::ImageLayerArgs::F_APPLY_TO_ALPHA        ) != 0;
+			bool bUseBaseSourceFromBaseAlpha   = (Args.flags & OP::ImageLayerArgs::F_BASE_RGB_FROM_ALPHA   ) != 0;
 			bool bUseBlendSourceFromBlendAlpha = (Args.flags & OP::ImageLayerArgs::F_BLENDED_RGB_FROM_ALPHA) != 0;
 			bool bUseMaskFromBlendAlpha        = (Args.flags & OP::ImageLayerArgs::F_USE_MASK_FROM_BLENDED ) != 0;
 
@@ -649,7 +650,11 @@ namespace UE::Mutable::Private
 					Result.Get(), Result.Get(), Blended.Get(), Mask.Get(), 
 					ColorBlendFunc, AlphaBlendFunc, 
 					LODBegin, LODEnd, 
-					bApplyColorBlendToAlpha, bUseMaskFromBlendAlpha, bUseBlendSourceFromBlendAlpha, Args.BlendAlphaSourceChannel);
+					bApplyColorBlendToAlpha, 
+					bUseMaskFromBlendAlpha,
+					bUseBaseSourceFromBaseAlpha,
+					bUseBlendSourceFromBlendAlpha, 
+					Args.BlendAlphaSourceChannel);
 		}
 		else
 		{
@@ -822,7 +827,10 @@ namespace UE::Mutable::Private
 			FBlendFuncType* ColorBlendFunc = SelectBlendFunc(EBlendType(Args.blendType));
 			FBlendFuncType* AlphaBlendFunc = ColorBlendFunc;
 			
-			bool bApplyColorBlendToAlpha = (Args.flags & OP::ImageLayerArgs::F_APPLY_TO_ALPHA) != 0;
+			bool bApplyColorBlendToAlpha       = (Args.flags & OP::ImageLayerArgs::F_APPLY_TO_ALPHA        ) != 0;
+			bool bUseBaseSourceFromBaseAlpha   = (Args.flags & OP::ImageLayerArgs::F_BASE_RGB_FROM_ALPHA   ) != 0;
+			bool bUseBlendSourceFromBlendAlpha = (Args.flags & OP::ImageLayerArgs::F_BLENDED_RGB_FROM_ALPHA) != 0;
+			bool bUseMaskFromBlendAlpha = false; 
 
 			if (!bApplyColorBlendToAlpha)
 			{
@@ -833,7 +841,11 @@ namespace UE::Mutable::Private
 					Result.Get(), Result.Get(), Mask.Get(), Color, 
 					ColorBlendFunc, AlphaBlendFunc, 
 					LODBegin, LODEnd, 
-					bApplyColorBlendToAlpha, false, false, Args.BlendAlphaSourceChannel);
+					bApplyColorBlendToAlpha, 
+					bUseMaskFromBlendAlpha, 
+					bUseBaseSourceFromBaseAlpha, 
+					bUseBlendSourceFromBlendAlpha,
+					Args.BlendAlphaSourceChannel);
 		}
 		else
 		{
diff --git a/Engine/Plugins/Mutable/Source/MutableRuntime/Private/MuR/ImageComputeUtils.cpp b/Engine/Plugins/Mutable/Source/MutableRuntime/Private/MuR/ImageComputeUtils.cpp
--- a/Engine/Plugins/Mutable/Source/MutableRuntime/Private/MuR/ImageComputeUtils.cpp
+++ b/Engine/Plugins/Mutable/Source/MutableRuntime/Private/MuR/ImageComputeUtils.cpp
@@ -362,7 +362,7 @@ namespace ImageComputeUtils
 		}
 	}
 
-	void Vec4SwizzleWZYXAndConvertToVec1_U8(uint8* Dest, const uint8* Src, int32 NumElems)
+	void Vec4SwizzleWWWWAndConvertToVec1_U8(uint8* Dest, const uint8* Src, int32 NumElems)
 	{
 		int32 ElemIndex = 0;
 
@@ -389,6 +389,36 @@ namespace ImageComputeUtils
 		}
 	}
 
+	void Vec4SwizzleWWWWAndConvertToVec3_U8(uint8* Dest, const uint8* Src, int32 NumElems)
+	{
+		int32 ElemIndex = 0;
+
+#if UE_MUTABLE_IMAGE_COMPUTE_UTILS_USE_VECTOR_IMPL
+		FConvertDescriptor ConvDesc
+		{
+			.SrcBytesPerIter = 16,	
+			.DstBytesPerIter = 12,
+
+			.Shuffle = MakeShuffle_Int8x16(3,3,3, 7,7,7, 11,11,11, 15,15,15),
+		};
+
+		constexpr int32 ElemsPerIter = 4;
+
+		int32 NumIters = FMath::Max<int32>(0, NumElems/ElemsPerIter - 1); 
+		VecConvertImpl(Dest, Src, NumIters, ConvDesc);
+
+		ElemIndex = NumIters*ElemsPerIter;
+#endif
+
+		for (; ElemIndex < NumElems; ++ElemIndex)
+		{
+			Dest[ElemIndex*3 + 0] = Src[ElemIndex*4 + 3];
+			Dest[ElemIndex*3 + 1] = Src[ElemIndex*4 + 3];
+			Dest[ElemIndex*3 + 2] = Src[ElemIndex*4 + 3];
+		}
+	}
+
+
 	void Vec4SwizzleXXXX_U8(uint8* Dest, const uint8* Src, int32 NumElems)
 	{
 		int32 ElemIndex = 0;
@@ -473,15 +503,10 @@ namespace ImageComputeUtils
 
 		for (; ElemIndex < NumElems; ++ElemIndex)
 		{
-			Dest[ElemIndex*4 + 0] = A[ElemIndex*4 + 0] & Mask.X;
-			Dest[ElemIndex*4 + 1] = A[ElemIndex*4 + 1] & Mask.Y;
-			Dest[ElemIndex*4 + 2] = A[ElemIndex*4 + 2] & Mask.Z;
-			Dest[ElemIndex*4 + 3] = A[ElemIndex*4 + 3] & Mask.W;
-			
-			Dest[ElemIndex*4 + 0] |= B[ElemIndex*4 + 0] & ~Mask.X;
-			Dest[ElemIndex*4 + 1] |= B[ElemIndex*4 + 1] & ~Mask.Y;
-			Dest[ElemIndex*4 + 2] |= B[ElemIndex*4 + 2] & ~Mask.Z;
-			Dest[ElemIndex*4 + 3] |= B[ElemIndex*4 + 3] & ~Mask.W;
+			Dest[ElemIndex*4 + 0] = (A[ElemIndex*4 + 0] & Mask.X) | (B[ElemIndex*4 + 0] & ~Mask.X);
+			Dest[ElemIndex*4 + 1] = (A[ElemIndex*4 + 1] & Mask.Y) | (B[ElemIndex*4 + 1] & ~Mask.Y);
+			Dest[ElemIndex*4 + 2] = (A[ElemIndex*4 + 2] & Mask.Z) | (B[ElemIndex*4 + 2] & ~Mask.Z);
+			Dest[ElemIndex*4 + 3] = (A[ElemIndex*4 + 3] & Mask.W) | (B[ElemIndex*4 + 3] & ~Mask.W);
 		}
 	}
 
diff --git a/Engine/Plugins/Mutable/Source/MutableRuntime/Private/MuR/ImageComputeUtils.h b/Engine/Plugins/Mutable/Source/MutableRuntime/Private/MuR/ImageComputeUtils.h
--- a/Engine/Plugins/Mutable/Source/MutableRuntime/Private/MuR/ImageComputeUtils.h
+++ b/Engine/Plugins/Mutable/Source/MutableRuntime/Private/MuR/ImageComputeUtils.h
@@ -22,10 +22,11 @@ namespace ImageComputeUtils
 	void Vec4SelectZ_U8(uint8* Dest, const uint8* A, const uint8* B, int32 NumElems);
 	void Vec4SelectW_U8(uint8* Dest, const uint8* A, const uint8* B, int32 NumElems);
 	
-	void Vec4SwizzleWZYXAndConvertToVec1_U8(uint8* Dest, const uint8* Src, int32 NumElems);
-	void Vec4SwizzleZYXW_U8(uint8* Dest, const uint8* Src, int32 NumElems);
+	void Vec4SwizzleWWWWAndConvertToVec1_U8(uint8* Dest, const uint8* Src, int32 NumElems);
+	void Vec4SwizzleWWWWAndConvertToVec3_U8(uint8* Dest, const uint8* Src, int32 NumElems);
 	void Vec4SwizzleXXXX_U8(uint8* Dest, const uint8* Src, int32 NumElems);
 	void Vec4SwizzleWWWW_U8(uint8* Dest, const uint8* Src, int32 NumElems);
+	void Vec4SwizzleZYXW_U8(uint8* Dest, const uint8* Src, int32 NumElems);
 
 	void Vec4Fill_U8(uint8* Dest, FColor Value, int32 NumElems);
 	void Vec3Fill_U8(uint8* Dest, FColor Value, int32 NumElems);	
diff --git a/Engine/Plugins/Mutable/Source/MutableRuntime/Private/MuR/OpImageLayer.cpp b/Engine/Plugins/Mutable/Source/MutableRuntime/Private/MuR/OpImageLayer.cpp
--- a/Engine/Plugins/Mutable/Source/MutableRuntime/Private/MuR/OpImageLayer.cpp
+++ b/Engine/Plugins/Mutable/Source/MutableRuntime/Private/MuR/OpImageLayer.cpp
@@ -11,8 +11,10 @@
 #include "Math/UnrealMathUtility.h"
 #include "Math/VectorRegister.h"
 
+
 namespace UE::Mutable::Private
 {
+
 namespace OpImageLayerBlendOps
 {
 
@@ -317,7 +319,7 @@ namespace OpImageLayerBlendOps
 		uint16x8_t AddHi = vmlal_u8(ScreenTimesAHi, ProdHi, vget_high_u8(OneMinusA));
 
 		uint8x8_t ResultLo = vqrshrn_n_u16(vrsraq_n_u16(AddLo, AddLo, 8), 8);
-		uint8x8_t ResultHi = vqrshrn_n_u16(vrsraq_n_u16(AddLo, AddHi, 8), 8);
+		uint8x8_t ResultHi = vqrshrn_n_u16(vrsraq_n_u16(AddHi, AddHi, 8), 8);
 
 		return vcombine_u8(ResultLo, ResultHi);
 	}
@@ -519,7 +521,7 @@ namespace OpImageLayerBlendOps
 	{
 		int32 ElemIndex = 0;
 
-#if 0//UE_MUTABLE_IMAGE_LAYER_USE_VECTOR_IMPL
+#if UE_MUTABLE_IMAGE_LAYER_USE_VECTOR_IMPL
 		int32 NumIters = NumElems / 16;
 
 		for (int32 Iter = 0; Iter < NumIters; ++Iter)
@@ -850,7 +852,20 @@ namespace OpImageLayerBlendOps
 		}
 	}
 
+	void Vec1FillOnes(uint8* Dest, const uint8*, int32 NumElems)
+	{
+		FMemory::Memset(Dest, 255, NumElems);
+	}
+
+	void Vec3FillOnes(uint8* Dest, const uint8*, int32 NumElems)
+	{
+		FMemory::Memset(Dest, 255, NumElems*3);
+	}
 
+	void Vec4FillOnes(uint8* Dest, const uint8*, int32 NumElems)
+	{
+		FMemory::Memset(Dest, 255, NumElems*4);
+	}
 } // namespace OpImageLayerBlendOps
 
 namespace OpImageLayerCombineOps
@@ -986,7 +1001,7 @@ namespace OpImageLayerInternal
 		TArrayView<const uint8> Blend;
 		TArrayView<const uint8> Mask;
 
-		uint64 DecodeTokenRLE;
+		uint64 DecodeTokenRLE = 0;
 	};
 
 	struct FOpBlendKernel
@@ -1003,6 +1018,7 @@ namespace OpImageLayerInternal
 		FMakeConstantU8FuncType* MakeConstantU8   = nullptr;
 		FMakeConstantU16FuncType* MakeConstantU16 = nullptr;
 
+		FConvertFuncType* BaseFormat  = nullptr;
 		FConvertFuncType* BlendFormat = nullptr;
 		FBlendFuncType* ColorBlend    = nullptr;
 		FBlendFuncType* AlphaBlend    = nullptr;
@@ -1029,6 +1045,7 @@ namespace OpImageLayerInternal
 		FBlendFuncType* AlphaBlendFunc = nullptr;
 
 		bool bApplyColorBlendToAlpha       = false;
+		bool bUseBaseSourceFromBaseAlpha   = false;
 		bool bUseBlendSourceFromBlendAlpha = false;
 		bool bUseMaskFromBlendAlpha        = false;
 		uint8 BlendAlphaSourceChannel      = 3;
@@ -1142,16 +1159,21 @@ namespace OpImageLayerInternal
 
 		if (Args.bBlendWithConstant)
 		{
-			Kernel.ConstantU8  = FLinearColor(Args.RealConstant).QuantizeRound();
+			FVector4f Constant = Args.RealConstant;
+			if (Args.bUseBlendSourceFromBlendAlpha)
+			{
+				Constant = FVector4f{Constant.W, Constant.W, Constant.W, Constant.W};
+			}
+
+			Kernel.ConstantU8  = FLinearColor(Constant).QuantizeRound();
 			Kernel.ConstantU16 = Math::TIntVector4<uint16>(
-					FMath::Clamp(Args.RealConstant.X * 255.0f + 0.5f, 0.0f, 65535.0f),
-					FMath::Clamp(Args.RealConstant.Y * 255.0f + 0.5f, 0.0f, 65535.0f),
-					FMath::Clamp(Args.RealConstant.Z * 255.0f + 0.5f, 0.0f, 65535.0f),
-					FMath::Clamp(Args.RealConstant.W * 255.0f + 0.5f, 0.0f, 65535.0f));
+					FMath::Clamp(Constant.X * 255.0f + 0.5f, 0.0f, 65535.0f),
+					FMath::Clamp(Constant.Y * 255.0f + 0.5f, 0.0f, 65535.0f),
+					FMath::Clamp(Constant.Z * 255.0f + 0.5f, 0.0f, 65535.0f),
+					FMath::Clamp(Constant.W * 255.0f + 0.5f, 0.0f, 65535.0f));
 
 			bool bAllLessThanEqOne = 
-				(Args.RealConstant.X <= 1.0f) && (Args.RealConstant.Y <= 1.0f) && 
-				(Args.RealConstant.Z <= 1.0f) && (Args.RealConstant.W <= 1.0f);
+					(Constant.X <= 1.0f) && (Constant.Y <= 1.0f) && (Constant.Z <= 1.0f) && (Constant.W <= 1.0f);
 
 			bool bUseU16Constant = !bAllLessThanEqOne && 
 					(Kernel.ColorBlend == OpImageLayerBlendOps::Multiply || Kernel.AlphaBlend == OpImageLayerBlendOps::Multiply);
@@ -1220,6 +1242,29 @@ namespace OpImageLayerInternal
 			}
 		}
 
+		if (Args.bUseBaseSourceFromBaseAlpha)
+		{
+			switch (Args.Base->GetFormat())
+			{
+			case EImageFormat::L_UByte:
+			{
+				Kernel.BaseFormat = OpImageLayerBlendOps::Vec1FillOnes;
+				break;
+			}
+			case EImageFormat::RGB_UByte:
+			{
+				Kernel.BaseFormat = OpImageLayerBlendOps::Vec3FillOnes;
+				break;
+			}
+			case EImageFormat::RGBA_UByte:
+			case EImageFormat::BGRA_UByte:
+			{
+				Kernel.BaseFormat = Vec4SwizzleWWWW_U8;
+				break;
+			}
+			}
+		}
+
 		if (!Args.bBlendWithConstant && Args.Base->GetFormat() != Args.Blend->GetFormat())
 		{
 			switch (Args.Base->GetFormat())
@@ -1230,7 +1275,14 @@ namespace OpImageLayerInternal
 				{
 				case EImageFormat::RGB_UByte:
 				{
-					Kernel.BlendFormat = Vec3ToVec1_U8;
+					if (Args.bUseBlendSourceFromBlendAlpha)
+					{
+						Kernel.BlendFormat = OpImageLayerBlendOps::Vec1FillOnes;
+					}
+					else
+					{
+						Kernel.BlendFormat = Vec3ToVec1_U8;
+					}
 					break;
 				}
 				case EImageFormat::RGBA_UByte:
@@ -1238,7 +1290,7 @@ namespace OpImageLayerInternal
 				{
 					if (Args.bUseBlendSourceFromBlendAlpha)
 					{
-						Kernel.BlendFormat = Vec4SwizzleWZYXAndConvertToVec1_U8;
+						Kernel.BlendFormat = Vec4SwizzleWWWWAndConvertToVec1_U8;
 					}
 					else
 					{
@@ -1257,13 +1309,27 @@ namespace OpImageLayerInternal
 				{
 				case EImageFormat::L_UByte:
 				{
-					Kernel.BlendFormat = Vec1ToVec3_U8;
+					if (Args.bUseBlendSourceFromBlendAlpha)
+					{
+						Kernel.BlendFormat = OpImageLayerBlendOps::Vec3FillOnes;
+					}
+					else
+					{
+						Kernel.BlendFormat = Vec1ToVec3_U8;
+					}
 					break;
 				}
 				case EImageFormat::RGBA_UByte:
 				case EImageFormat::BGRA_UByte:
 				{
-					Kernel.BlendFormat = Vec4ToVec3_U8;
+					if (Args.bUseBlendSourceFromBlendAlpha)
+					{
+						Kernel.BlendFormat = Vec4SwizzleWWWWAndConvertToVec3_U8;
+					}
+					else
+					{
+						Kernel.BlendFormat = Vec4ToVec3_U8;
+					}
 					break;
 				}
 				default: check(Args.Blend->GetFormat() == EImageFormat::RGB_UByte); break;
@@ -1277,12 +1343,26 @@ namespace OpImageLayerInternal
 				{
 				case EImageFormat::L_UByte:
 				{
-					Kernel.BlendFormat = Vec1ToVec4_U8;
+					if (Args.bUseBlendSourceFromBlendAlpha)
+					{
+						Kernel.BlendFormat = OpImageLayerBlendOps::Vec4FillOnes;
+					}
+					else
+					{
+						Kernel.BlendFormat = Vec1ToVec4_U8;
+					}
 					break;
 				}
 				case EImageFormat::RGB_UByte:
 				{
-					Kernel.BlendFormat = Vec3ToVec4_U8;
+					if (Args.bUseBlendSourceFromBlendAlpha)
+					{
+						Kernel.BlendFormat = OpImageLayerBlendOps::Vec4FillOnes;
+					}
+					else
+					{
+						Kernel.BlendFormat = Vec3ToVec4_U8;
+					}
 					break;
 				}
 				default: check(Args.Blend->GetFormat() == EImageFormat::RGBA_UByte || Args.Blend->GetFormat() == EImageFormat::BGRA_UByte); break;
@@ -1301,11 +1381,33 @@ namespace OpImageLayerInternal
 
 				if (Args.BlendAlphaSourceChannel == 3)
 				{
-					Kernel.ExtractMaskFromBlendAlpha = Vec4SwizzleWZYXAndConvertToVec1_U8;
+					Kernel.ExtractMaskFromBlendAlpha = Vec4SwizzleWWWWAndConvertToVec1_U8;
 				}
 				else
 				{
-					Kernel.ExtractMaskFromBlendAlpha = Vec4ToVec1_U8;
+					if (!Args.bUseBlendSourceFromBlendAlpha)
+					{
+						Kernel.ExtractMaskFromBlendAlpha = Vec4ToVec1_U8;
+					}
+					else
+					{
+						switch(Args.Blend->GetFormat())
+						{
+							case EImageFormat::L_UByte:
+							case EImageFormat::RGB_UByte:
+							{
+								Kernel.ExtractMaskFromBlendAlpha = OpImageLayerBlendOps::Vec1FillOnes;
+								break;
+							}
+							case EImageFormat::RGBA_UByte:
+							case EImageFormat::BGRA_UByte:
+							{
+								Kernel.ExtractMaskFromBlendAlpha = Vec4SwizzleWWWWAndConvertToVec1_U8;
+								break;
+							}
+							default: check(false);
+						}
+					}
 				}
 			}
 
@@ -1332,7 +1434,7 @@ namespace OpImageLayerInternal
 				case EImageFormat::BGRA_UByte:
 				case EImageFormat::RGBA_UByte:
 				{
-					if (Kernel.ExtractMaskFromBlendAlpha == Vec4SwizzleWZYXAndConvertToVec1_U8)
+					if (Kernel.ExtractMaskFromBlendAlpha == Vec4SwizzleWWWWAndConvertToVec1_U8)
 					{
 						Kernel.ExtractMaskFromBlendAlpha = Vec4SwizzleWWWW_U8;
 					}
@@ -1459,7 +1561,7 @@ namespace OpImageLayerInternal
 		//MUTABLE_CPUPROFILER_SCOPE(OpImageLayer_RunKernelOnBatch_Blend);
 
 		constexpr int32 MaxElemSize = 4;
-		constexpr int32 NumCacheLinesPerWave = 16;
+		constexpr int32 NumCacheLinesPerWave = 8;
 		constexpr int32 WaveSizeInByteElems = (PLATFORM_CACHE_LINE_SIZE * NumCacheLinesPerWave)/MaxElemSize;
 
 		struct alignas(PLATFORM_CACHE_LINE_SIZE) FTemporalRegister
@@ -1510,6 +1612,40 @@ namespace OpImageLayerInternal
 			const uint8* MaskDataBuffer = MaskBuffer + WaveBeginInElems;
 			const uint8* BlenDataBuffer = BlenBuffer + WaveBeginInElems*BlenBytesPerElem;
 
+			// Zero mask optimization
+			bool bSkipColorBlend = false;
+			if (Kernel.Mix)
+			{
+				if (Kernel.ExtractMaskFromBlendAlpha)
+				{
+					Kernel.ExtractMaskFromBlendAlpha(Register3.Data, BlenDataBuffer, NumWaveElems);
+					MaskDataBuffer = Register3.Data;
+				}
+				else if (Kernel.MaskDecode)
+				{
+					check(InOutDecodeToken != 0);
+					//NOTE: Decode operation, unlike the other operations, use the batch buffer and not the wave buffer 
+					//      because the data offset is carried by the DecodeToken.
+					Kernel.MaskDecode(Register3.Data, MaskBuffer, InOutDecodeToken, NumWaveElems);
+
+					MaskDataBuffer = Register3.Data;
+				}
+
+				if (Kernel.MaskFormat)
+				{
+					Kernel.MaskFormat(Register4.Data, MaskDataBuffer, NumWaveElems);
+					MaskDataBuffer = Register4.Data;
+				}
+
+				bSkipColorBlend = FMemory::MemIsZero(MaskDataBuffer, NumWaveElems*BaseBytesPerElem);
+			}
+			
+			if (Kernel.BaseFormat)
+			{
+				Kernel.BaseFormat(ResultDataBuffer, BaseDataBuffer, NumWaveElems);
+				BaseDataBuffer = ResultDataBuffer;
+			}
+
 			if (Kernel.MakeConstantU8 || Kernel.MakeConstantU16)
 			{
 				BlenDataBuffer = Register0.Data;
@@ -1521,7 +1657,7 @@ namespace OpImageLayerInternal
 			}
 
 			const uint8* BlendedResultBuffer = BaseDataBuffer; 
-			if (Kernel.ColorBlend)
+			if (Kernel.ColorBlend && !bSkipColorBlend)
 			{
 				uint8* MutableColorBlendResultBuffer = Kernel.Mix || Kernel.SelectAlpha ? Register1.Data : ResultDataBuffer;
 				BlendedResultBuffer = MutableColorBlendResultBuffer;
@@ -1560,32 +1696,10 @@ namespace OpImageLayerInternal
 				{
 					Kernel.AlphaBlend(MutableAlphaBlendedResult, BaseDataBuffer, BlenDataBuffer, BaseBytesPerElem*NumWaveElems);
 				}
-
-			}
-
-			if (Kernel.ExtractMaskFromBlendAlpha)
-			{
-				Kernel.ExtractMaskFromBlendAlpha(Register3.Data, BlenDataBuffer, NumWaveElems);
-				MaskDataBuffer = Register3.Data;
-			}
-			else if (Kernel.MaskDecode)
-			{
-				check(InOutDecodeToken != 0);
-				//NOTE: Decode operation, unlike the other operations, use the batch buffer and not the wave buffer 
-				//      because the data offset is carried by the DecodeToken.
-				Kernel.MaskDecode(Register3.Data, MaskBuffer, InOutDecodeToken, NumWaveElems);
-
-				MaskDataBuffer = Register3.Data;
-			}
-
-			if (Kernel.MaskFormat)
-			{
-				Kernel.MaskFormat(Register4.Data, MaskDataBuffer, NumWaveElems);
-				MaskDataBuffer = Register4.Data;
 			}
 
 			const uint8* MixResultBuffer = BlendedResultBuffer;
-			if (Kernel.Mix)
+			if (Kernel.Mix && !bSkipColorBlend)
 			{
 				uint8* MutableMixResultBuffer = Kernel.SelectAlpha ? Register3.Data : ResultDataBuffer;
 				MixResultBuffer = MutableMixResultBuffer;
@@ -1937,7 +2051,11 @@ namespace OpImageLayerInternal
 			FBlendFuncType* ColorBlendFunc, FBlendFuncType* AlphaBlendFunc,
 			bool bBlendWithConstant, const FVector4f& BlendConstant,
 			int32 LODBegin, int32 LODEnd, 
-			bool bApplyColorBlendToAlpha, bool bUseMaskFromBlendAlpha, bool bUseBlendSourceFromBlendAlpha, uint8 BlendAlphaSourceChannel)
+			bool bApplyColorBlendToAlpha, 
+			bool bUseMaskFromBlendAlpha, 
+			bool bUseBaseSourceFromBaseAlpha, 
+			bool bUseBlendSourceFromBlendAlpha, 
+			uint8 BlendAlphaSourceChannel)
 	{
 		MUTABLE_CPUPROFILER_SCOPE(ImageLayer)
 
@@ -2009,6 +2127,7 @@ namespace OpImageLayerInternal
 			.AlphaBlendFunc = AlphaBlendFunc,
 
 			.bApplyColorBlendToAlpha       = bApplyColorBlendToAlpha,
+			.bUseBaseSourceFromBaseAlpha   = bUseBaseSourceFromBaseAlpha,
 			.bUseBlendSourceFromBlendAlpha = bUseBlendSourceFromBlendAlpha,
 			.bUseMaskFromBlendAlpha        = bUseMaskFromBlendAlpha,
 			.BlendAlphaSourceChannel       = BlendAlphaSourceChannel,
@@ -2359,7 +2478,11 @@ namespace OpImageLayerInternal
 			FImage* DestImage, const FImage* BaseImage, const FImage* BlenImage, const FImage* MaskImage, 
 			FBlendFuncType* ColorBlendFunc, FBlendFuncType* AlphaBlendFunc,
 			int32 LODBegin, int32 LODEnd,
-			bool bApplyColorBlendToAlpha, bool bUseMaskFromBlendAlpha, bool bUseBlendSourceFromBlendAlpha, uint8 BlendAlphaSourceChannel)
+			bool bApplyColorBlendToAlpha, 
+			bool bUseMaskFromBlendAlpha, 
+			bool bUseBaseSourceFromBaseAlpha, 
+			bool bUseBlendSourceFromBlendAlpha, 
+			uint8 BlendAlphaSourceChannel)
 	{
 		constexpr bool bBlendWithConstant = false;
 		OpImageLayerInternal::ImageLayerImpl(
@@ -2367,14 +2490,22 @@ namespace OpImageLayerInternal
 				ColorBlendFunc, AlphaBlendFunc, 
 				bBlendWithConstant, FVector4f{}, 
 				LODBegin, LODEnd,  
-				bApplyColorBlendToAlpha, bUseMaskFromBlendAlpha, bUseBlendSourceFromBlendAlpha, BlendAlphaSourceChannel);
+				bApplyColorBlendToAlpha, 
+				bUseMaskFromBlendAlpha, 
+				bUseBaseSourceFromBaseAlpha,
+				bUseBlendSourceFromBlendAlpha, 
+				BlendAlphaSourceChannel);
 	}
 
 	void ImageLayerBlendConstant(
 			FImage* DestImage, const FImage* BaseImage, const FImage* MaskImage, const FVector4f& Constant,  
 			FBlendFuncType* ColorBlendFunc, FBlendFuncType* AlphaBlendFunc,
 			int32 LODBegin, int32 LODEnd, 
-			bool bApplyColorBlendToAlpha, bool bUseMaskFromBlendAlpha, bool bUseBlendSourceFromBlendAlpha, uint8 BlendAlphaSourceChannel)
+			bool bApplyColorBlendToAlpha, 
+			bool bUseMaskFromBlendAlpha, 
+			bool bUseBaseSourceFromBaseAlpha, 
+			bool bUseBlendSourceFromBlendAlpha, 
+			uint8 BlendAlphaSourceChannel)
 	{
 		constexpr bool bBlendWithConstant = true;
 		OpImageLayerInternal::ImageLayerImpl(
@@ -2382,7 +2513,11 @@ namespace OpImageLayerInternal
 				ColorBlendFunc, AlphaBlendFunc, 
 				bBlendWithConstant, Constant, 
 				LODBegin, LODEnd,
-				bApplyColorBlendToAlpha, bUseMaskFromBlendAlpha, bUseBlendSourceFromBlendAlpha, BlendAlphaSourceChannel);
+				bApplyColorBlendToAlpha, 
+				bUseMaskFromBlendAlpha, 
+				bUseBaseSourceFromBaseAlpha,
+				bUseBlendSourceFromBlendAlpha, 
+				BlendAlphaSourceChannel);
 	}
 
 	void ImageLayerCombine(
diff --git a/Engine/Plugins/Mutable/Source/MutableRuntime/Private/MuR/OpImageLayer.h b/Engine/Plugins/Mutable/Source/MutableRuntime/Private/MuR/OpImageLayer.h
--- a/Engine/Plugins/Mutable/Source/MutableRuntime/Private/MuR/OpImageLayer.h
+++ b/Engine/Plugins/Mutable/Source/MutableRuntime/Private/MuR/OpImageLayer.h
@@ -39,13 +39,21 @@ void ImageLayerBlend(
 		FImage* Dest, const FImage* Base, const FImage* Blend, const FImage* Mask, 
 		FBlendFuncType* ColorBlendFunc, FBlendFuncType* AlphaBlendFunc,
 		int32 LODBegin, int32 LODEnd, 
-		bool bApplyColorBlendToAlpha, bool bUseMaskFromBlendAlpha, bool bUseBlendSourceFromBlendAlpha, uint8 BlendAlphaSourceChannel);
+		bool bApplyColorBlendToAlpha, 
+		bool bUseMaskFromBlendAlpha, 
+		bool bUseBaseSourceFromBaseAlpha, 
+		bool bUseBlendSourceFromBlendAlpha, 
+		uint8 BlendAlphaSourceChannel);
 
 void ImageLayerBlendConstant(
 		FImage* Dest, const FImage* Base, const FImage* Mask, const FVector4f& Constant, 
 		FBlendFuncType* ColorBlendFunc, FBlendFuncType* AlphaBlendFunc,
 		int32 LODBegin, int32 LODEnd, 
-		bool bApplyColorBlendToAlpha, bool bUseMaskFromBlendAlpha, bool bUseBlendSourceFromBlendAlpha, uint8 BlendAlphaSourceChannel);
+		bool bApplyColorBlendToAlpha, 
+		bool bUseMaskFromBlendAlpha, 
+		bool bUseBaseSourceFromBaseAlpha, 
+		bool bUseBlendSourceFromBlendAlpha, 
+		uint8 BlendAlphaSourceChannel);
 
 void ImageLayerCombine(
 		FImage* Dest, const FImage* Base, const FImage* Blend, const FImage* Mask,
