[mutable] Removed Texture locking from the "Task_Game_LockCache" method as part of its removal.
- Moved the texture locking to "Task_Mutable_GetImages" delaying this way its operation and making it no longer part of the "Task_Game_LockCache"
- Added locks to the CO Texture Cache to prevent collisions when accessing the Texture Cache

#rnx

#Synced-CL 51368115

diff --git a/Engine/Plugins/Mutable/Source/CustomizableObject/Internal/MuCO/CustomizableObjectPrivate.h b/Engine/Plugins/Mutable/Source/CustomizableObject/Internal/MuCO/CustomizableObjectPrivate.h
--- a/Engine/Plugins/Mutable/Source/CustomizableObject/Internal/MuCO/CustomizableObjectPrivate.h
+++ b/Engine/Plugins/Mutable/Source/CustomizableObject/Internal/MuCO/CustomizableObjectPrivate.h
@@ -169,6 +169,8 @@ public:
 		friend uint32 GetTypeHash(const FId& Key);
 	};
 
+	UE::FMutex Lock;
+	
 	UTexture2D* Get(const FId& Key);
 	
 	void Add(const FId& Key, UTexture2D* Value);
diff --git a/Engine/Plugins/Mutable/Source/CustomizableObject/Internal/MuCO/CustomizableObjectSystemPrivate.h b/Engine/Plugins/Mutable/Source/CustomizableObject/Internal/MuCO/CustomizableObjectSystemPrivate.h
--- a/Engine/Plugins/Mutable/Source/CustomizableObject/Internal/MuCO/CustomizableObjectSystemPrivate.h
+++ b/Engine/Plugins/Mutable/Source/CustomizableObject/Internal/MuCO/CustomizableObjectSystemPrivate.h
@@ -427,7 +427,6 @@ public:
 	// between AdditionalAssetsAsyncLoaded() and their setting into the generated materials in BuildResources()
 	TArray<TStrongObjectPtr<const UTexture>> LoadedPassThroughTexturesPendingSetMaterial;
 	
-	TArray<UE::Mutable::Private::FImageId> CachedTextures;
 
 	TArray<UE::Mutable::Private::TManagedPtr<const UE::Mutable::Private::FExtensionData>> ExtensionData;
 	
diff --git a/Engine/Plugins/Mutable/Source/CustomizableObject/Private/MuCO/CustomizableObject.cpp b/Engine/Plugins/Mutable/Source/CustomizableObject/Private/MuCO/CustomizableObject.cpp
--- a/Engine/Plugins/Mutable/Source/CustomizableObject/Private/MuCO/CustomizableObject.cpp
+++ b/Engine/Plugins/Mutable/Source/CustomizableObject/Private/MuCO/CustomizableObject.cpp
@@ -1924,8 +1924,11 @@ void FSkeletonCache::Clear()
 	MergedSkeletons.Empty();
 }
 
+
 UTexture2D* FTextureCache::Get(const FId& Key)
 {
+	UE::TScopeLock<UE::FMutex> LockGuard(Lock);
+	
 	const TWeakObjectPtr<UTexture2D>* Result = GeneratedTextures.Find(Key);
 	return Result ? Result->Get() : nullptr;
 }
@@ -1938,14 +1941,18 @@ void FTextureCache::Add(const FId& Key, UTexture2D* Value)
 		return;
 	}
 
-	GeneratedTextures.Add(Key, Value);
-
-	// Remove invalid SkeletalMeshes from the cache.
-	for (auto SkeletonIterator = GeneratedTextures.CreateIterator(); SkeletonIterator; ++SkeletonIterator)
 	{
-		if (SkeletonIterator.Value().IsStale())
+		UE::TScopeLock<UE::FMutex> LockGuard(Lock);
+		
+		GeneratedTextures.Add(Key, Value);
+
+		// Remove invalid SkeletalMeshes from the cache.
+		for (auto SkeletonIterator = GeneratedTextures.CreateIterator(); SkeletonIterator; ++SkeletonIterator)
 		{
-			SkeletonIterator.RemoveCurrent();
+			if (SkeletonIterator.Value().IsStale())
+			{
+				SkeletonIterator.RemoveCurrent();
+			}
 		}
 	}
 }
@@ -1953,6 +1960,8 @@ void FTextureCache::Add(const FId& Key, UTexture2D* Value)
 
 void FTextureCache::Remove(const UTexture2D& Value)
 {
+	UE::TScopeLock<UE::FMutex> LockGuard(Lock);
+	
 	for (auto It = GeneratedTextures.CreateIterator(); It; ++It)
 	{
 		if (It->Value.Get() == &Value)
@@ -1966,6 +1975,8 @@ void FTextureCache::Remove(const UTexture2D& Value)
 
 void FTextureCache::Clear()
 {
+	UE::TScopeLock<UE::FMutex> LockGuard(Lock);
+	
 	GeneratedTextures.Empty();
 }
 
diff --git a/Engine/Plugins/Mutable/Source/CustomizableObject/Private/MuCO/CustomizableObjectSystem.cpp b/Engine/Plugins/Mutable/Source/CustomizableObject/Private/MuCO/CustomizableObjectSystem.cpp
--- a/Engine/Plugins/Mutable/Source/CustomizableObject/Private/MuCO/CustomizableObjectSystem.cpp
+++ b/Engine/Plugins/Mutable/Source/CustomizableObject/Private/MuCO/CustomizableObjectSystem.cpp
@@ -2375,12 +2375,35 @@ namespace impl
 			}
 			else
 			{
-				const bool bCached = RequestedGetImages->Contains(Image->ImageID) || // See if it is cached from this same instance (can happen with LODs)
-					OperationData->CachedTextures.Contains(Image->ImageID); // See if it is cached from another instance
-			
-				if (bCached)
+				// Access the CO TextureCache and try to grab the texture from there if found.
+				if (OperationData->bUseCaches)
+				{
+					TStrongObjectPtr<const UTexture> StrongCachedTexture = nullptr;
+					
+					FTextureCache& TextureCache = OperationData->Object->GetPrivate()->TextureCache;
+					const FTextureCache::FId TextureCacheKey = FTextureCache::FId(Image->ImageID, OperationData->MipsToSkip, OperationData->bBake);
+					if (const UTexture2D* CachedTexture = TextureCache.Get(TextureCacheKey))
+					{
+						FGCScopeTryGuard Guard;
+						if (Guard.LockSucceeded())
+						{
+							StrongCachedTexture = TStrongObjectPtr(CachedTexture);
+						}
+					}
+					
+					if (StrongCachedTexture)
+					{
+						UE_LOG(LogMutable, VeryVerbose, TEXT("Texture resource with id [%s] is cached."), *Image->ImageID.ToString());
+						
+						OperationData->Objects.Add(StrongCachedTexture);
+						
+						continue;
+					}
+				}
+				
+				// Check if we have already found this ImageID during the runtime of this method
+				if (RequestedGetImages->Contains(Image->ImageID))
 				{
-					UE_LOG(LogMutable, VeryVerbose, TEXT("Texture resource with id [%s] is cached."), *Image->ImageID.ToString());
 					continue;
 				}
 
@@ -2697,21 +2720,6 @@ namespace impl
 		}
 
 		UCustomizableInstancePrivate* ObjectInstancePrivateData = ObjectInstance->GetPrivate();
-		
-		if (OperationData->bUseCaches)
-		{
-			FTextureCache& TextureCache = CustomizableObject->GetPrivate()->TextureCache;
-		
-			for (const TSharedRef<FInstanceUpdateData::FImage>& Image : OperationData->InstanceUpdateData.Images)
-			{
-				const FTextureCache::FId TextureCacheKey = FTextureCache::FId(Image->ImageID, OperationData->MipsToSkip, OperationData->bBake);
-				if (const UTexture2D* Texture = TextureCache.Get(TextureCacheKey))
-				{
-					OperationData->Objects.Add(TStrongObjectPtr(Texture));
-					OperationData->CachedTextures.Add(Image->ImageID);
-				}
-			}
-		}
 
 		// Any external texture that may be needed for this update will be requested from Mutable Core's GetImage
 		// which will safely access the GlobalExternalImages map, and then just get the cached image or issue a disk read
