[Mutable] Remove legacy LOD management and legacy mesh streaming (hacky streaming).
- Implement high-priority and low-priority queues for manual instance updates.
- Refactor 'UpdateOnlyCloseDistanceInstances' code and optimize distance checks.
- Refactor 'DiscardInstances' code to avoid potential crashes due to RHI resources being freed while in use.
- Fix FLiveInstances not clearing cache on EndUpdate, which could lead to performance issues due to being out of budget permanently.
- Fix Customizable Object Instances not generating when editing the level.
- Replace TObjectIterators with the static TSets of instances and usages.
- Optimize CustomizableObjectSystem::Tick.

#rnx

#Synced-CL 51330672

diff --git a/Engine/Plugins/Mutable/Source/CustomizableObject/Internal/MuCO/CustomizableObjectInstancePrivate.h b/Engine/Plugins/Mutable/Source/CustomizableObject/Internal/MuCO/CustomizableObjectInstancePrivate.h
--- a/Engine/Plugins/Mutable/Source/CustomizableObject/Internal/MuCO/CustomizableObjectInstancePrivate.h
+++ b/Engine/Plugins/Mutable/Source/CustomizableObject/Internal/MuCO/CustomizableObjectInstancePrivate.h
@@ -57,17 +57,14 @@ enum ECOInstanceFlags
 	ECONone							= 0,  // Should not use the name None here.. it collides with other enum in global namespace
 
 	// Update process
-	ReplacePhysicsAssets			= 1 << 4,	// Merge active PhysicsAssets and replace the base physics asset
+	ReplacePhysicsAssets			= 1 << 1,	// Merge active PhysicsAssets and replace the base physics asset
 
 	// Update priorities
-	UsedByComponent					= 1 << 5,	// If any components are using this instance, they will set flag every frame
-	UsedByComponentInPlay			= 1 << 6,	// If any components are using this instance in play, they will set flag every frame
-	UsedByPlayerOrNearIt			= 1 << 7,	// The instance is used by the player or is near the player, used to give more priority to its updates
-	DiscardedByNumInstancesLimit	= 1 << 8,	// The instance is discarded because we exceeded the limit of instances generated 
-
-	// Types of updates
-	PendingLODsUpdate				= 1 << 9,	// Used to queue an update due to a change in LODs required by the instance
-	PendingLODsDowngrade			= 1 << 10,	// Used to queue a downgrade update to reduce the number of LODs. LOD update goes from a high res level to a low res one, ex: 0 to 1 or 1 to 2
+	UsedByPlayerOrNearIt			= 1 << 2,	// The instance is used by the player or is near the player, used to give more priority to its updates
+	Discarded						= 1 << 3,
+
+	HasPendingHighPriorityUpdate	= 1 << 4,
+	HasPendingLowPriorityUpdate		= 1 << 5,
 };
 
 ENUM_CLASS_FLAGS(ECOInstanceFlags);
@@ -206,9 +203,6 @@ public:
 	bool HasCOInstanceFlags(ECOInstanceFlags FlagsToCheck) const { return (InstanceFlagsPrivate & FlagsToCheck) != 0; }
 
 	UE_API void BuildResources(const TSharedRef<FUpdateContextPrivate>& Context, UCustomizableObjectInstance* Public);
-	
-	UE_API void TickUpdateCloseCustomizableObjects(UCustomizableObjectInstance& Public, FMutableInstanceUpdateMap& InOutRequestedUpdates);
-	UE_API void UpdateInstanceIfNotGenerated(UCustomizableObjectInstance& Public, FMutableInstanceUpdateMap& InOutRequestedUpdates);
 
 	// Returns true if success (?)
 	UE_API bool UpdateSkeletalMesh_PostBeginUpdate0(UCustomizableObjectInstance* Instance, const TSharedRef<FUpdateContextPrivate>& Context);
@@ -216,10 +210,6 @@ public:
 	// The following method is basically copied from PostEditChangeProperty and/or SkeletalMesh.cpp to be able to replicate PostEditChangeProperty without the editor
 	UE_API void PostEditChangePropertyWithoutEditor();
 	
-	/** Calls ReleaseResources on all SkeletalMeshes generated by this instance and invalidates the generated data.
-	  * It should not be called if the meshes are still in use or shared with other instances. */
-	UE_API void DiscardResources();
-	
 	/** Only use as an emergency fallback! Set the reference SkeletalMesh unconditionally to all actors using this instance. It will hitch. */
 	UE_API void ForceSetReferenceSkeletalMesh() const;
 
@@ -276,7 +266,6 @@ public:
 
 	// If any components are using this instance, they will store the min of their distances to the player here every frame for LOD purposes
 	float MinSquareDistFromComponentToPlayer = FLT_MAX;
-	float LastMinSquareDistFromComponentToPlayer = FLT_MAX; // The same as the previous dist for last frame
 	
 	UPROPERTY(Transient)
 	TArray<FCustomizableInstanceComponentData> ComponentsData;
@@ -307,11 +296,7 @@ public:
 
 	/** Automatic update required.
 	 * Set to true when a Customizable Object Instance Usage requires an automatic update (e.g., component reattached). */
-	bool bAutomaticUpdateRequired = false;
-	
-	// TEMP VARIABLE to check the Min desired LODs for this instance
-	TWeakObjectPtr<UCustomizableObjectInstanceUsage> NearestToActor;
-	TWeakObjectPtr<const AActor> NearestToViewCenter;
+	bool bAllowAutomaticUpdate = false;
 
 #if WITH_EDITOR
 	/** Profile index the instance parameters are in and if the profile needs to be refreshed */
diff --git a/Engine/Plugins/Mutable/Source/CustomizableObject/Internal/MuCO/CustomizableObjectInstanceUsagePrivate.h b/Engine/Plugins/Mutable/Source/CustomizableObject/Internal/MuCO/CustomizableObjectInstanceUsagePrivate.h
--- a/Engine/Plugins/Mutable/Source/CustomizableObject/Internal/MuCO/CustomizableObjectInstanceUsagePrivate.h
+++ b/Engine/Plugins/Mutable/Source/CustomizableObject/Internal/MuCO/CustomizableObjectInstanceUsagePrivate.h
@@ -26,17 +26,10 @@ public:
 	/** Common end point of all updates. Even those which failed. */
 	UE_API void Callbacks();
 
-#if WITH_EDITOR
-	// Used to generate instances outside the CustomizableObject editor and PIE
-	UE_API void UpdateDistFromComponentToLevelEditorCamera(const FVector& CameraPosition);
-#endif
-
 	UE_API USkeletalMesh* GetSkeletalMesh() const;
 
 	UE_API USkeletalMesh* GetAttachedSkeletalMesh() const;
 	
-	UE_API void UpdateDistFromComponentToPlayer(const AActor* const Pawn, bool bForceEvenIfNotBegunPlay = false);
-	
 	/** Return a valid pointer if this Usage is being used inside a UCustomizableSkeletalComponent. Nullptr if it is being used in Standalone mode. */
 	UE_API UCustomizableSkeletalComponent* GetCustomizableSkeletalComponent() const;
 	
@@ -44,11 +37,10 @@ public:
 
 	UE_API const UCustomizableObjectInstanceUsage* GetPublic() const;
 
-	// Returns true if the NetMode of the associated UCustomizableSkeletalComponent (or the associated SkeletalMeshComponent if the former does not exist) is equal to InNetMode
-	UE_API bool IsNetMode(ENetMode InNetMode) const;
+	/** Used to replace the Skeletal Mesh of the parent component by the Reference Skeletal Mesh. */ 
+	bool bPendingSetReferenceSkeletalMesh = true;
 
-	/** Used to replace the Skeletal Mesh of the parent component by the Reference Skeletal Mesh or the generated Skeletal Mesh. */ 
-	bool bPendingSetSkeletalMesh = true;
+	bool bIsNetModeDedicatedServer = false;
 };
 
 #undef UE_API
diff --git a/Engine/Plugins/Mutable/Source/CustomizableObject/Internal/MuCO/CustomizableObjectPrivate.h b/Engine/Plugins/Mutable/Source/CustomizableObject/Internal/MuCO/CustomizableObjectPrivate.h
--- a/Engine/Plugins/Mutable/Source/CustomizableObject/Internal/MuCO/CustomizableObjectPrivate.h
+++ b/Engine/Plugins/Mutable/Source/CustomizableObject/Internal/MuCO/CustomizableObjectPrivate.h
@@ -1588,9 +1588,6 @@ public:
 
 	CUSTOMIZABLEOBJECT_API void SetIsChildObject(bool bIsChildObject);
 	
-	/** Return the MinLOD index to generate based on the active LODSettings (PerPlatformMinLOD or PerQualityLevelMinLOD) */
-	CUSTOMIZABLEOBJECT_API uint8 GetMinLODIndex(const UE::Mutable::Private::FComponentId ComponentId) const;
-
 #if WITH_EDITOR
 	/** See ICustomizableObjectEditorModule::IsCompilationOutOfDate. */
 	CUSTOMIZABLEOBJECT_API bool IsCompilationOutOfDate(bool bSkipIndirectReferences, TArray<FName>& OutOfDatePackages, TArray<FName>& AddedPackages, TArray<FName>& RemovedPackages, bool& bReleaseVersionDiff) const;
diff --git a/Engine/Plugins/Mutable/Source/CustomizableObject/Internal/MuCO/CustomizableObjectSystemPrivate.h b/Engine/Plugins/Mutable/Source/CustomizableObject/Internal/MuCO/CustomizableObjectSystemPrivate.h
--- a/Engine/Plugins/Mutable/Source/CustomizableObject/Internal/MuCO/CustomizableObjectSystemPrivate.h
+++ b/Engine/Plugins/Mutable/Source/CustomizableObject/Internal/MuCO/CustomizableObjectSystemPrivate.h
@@ -8,8 +8,8 @@
 #include "MuCO/CustomizableObjectPrivate.h"
 #include "MuCO/CustomizableObjectInstance.h"
 #include "MuCO/CustomizableObjectExtension.h"
-#include "MuCO/CustomizableInstanceLODManagement.h"
 #include "Containers/Ticker.h"
+#include "Containers/Deque.h"
 
 #include "MuCO/CustomizableObjectInstanceDescriptor.h"
 #include "MuR/Mesh.h"
@@ -31,6 +31,7 @@
 
 #define UE_API CUSTOMIZABLEOBJECT_API
 
+class AActor;
 class FMutableUpdateImageContext;
 class FUnrealMutableResourceProvider;
 class FMutableStreamRequest;
@@ -72,147 +73,6 @@ void ReconstructMorphTargetsFromMeshCompressedData(const UE::Mutable::Private::F
 void FinishUpdateGlobal(const TSharedRef<FUpdateContextPrivate>& Context);
 
 
-class FMutableUpdateCandidate
-{
-public:
-	// The Instance to possibly update
-	UCustomizableObjectInstance* CustomizableObjectInstance;
-	EQueuePriorityType Priority = EQueuePriorityType::Med;
-
-	// These are the LODs that would be applied if this candidate is chosen
-	TMap<FName, uint8> MinLOD;
-
-	// These are the LODs that would be copied to the descriptor to trigger mesh updates on quality setting changes.
-	TMap<FName, uint8> QualitySettingMinLODs;
-
-	TMap<FName, uint8> FirstRequestedLOD;
-
-	FMutableUpdateCandidate(UCustomizableObjectInstance* InCustomizableObjectInstance);
-	
-	FMutableUpdateCandidate(const UCustomizableObjectInstance* InCustomizableObjectInstance, const TMap<FName, uint8>& InMinLOD,
-		const TMap<FName, uint8>& InFirstRequestedLOD) :
-		CustomizableObjectInstance(const_cast<UCustomizableObjectInstance*>(InCustomizableObjectInstance)),
-		MinLOD(InMinLOD),
-		FirstRequestedLOD(InFirstRequestedLOD) {}
-
-	bool HasBeenIssued() const;
-
-	void Issue();
-
-	void ApplyLODUpdateParamsToInstance(FUpdateContextPrivate& Context);
-
-private:
-	/** If true it means that EnqueueUpdateSkeletalMesh has decided this update should be performed, if false it should be ignored. Just used for consistency checks */
-	bool bHasBeenIssued = false;
-};
-
-
-struct FMutablePendingInstanceUpdate
-{
-	TSharedRef<FUpdateContextPrivate> Context;
-
-	FMutablePendingInstanceUpdate(const TSharedRef<FUpdateContextPrivate>& InContext);
-
-	bool operator==(const FMutablePendingInstanceUpdate& Other) const;
-
-	bool operator<(const FMutablePendingInstanceUpdate& Other) const;
-};
-
-
-inline uint32 GetTypeHash(const FMutablePendingInstanceUpdate& Update);
-
-
-struct FPendingInstanceUpdateKeyFuncs : BaseKeyFuncs<FMutablePendingInstanceUpdate, TWeakObjectPtr<const UCustomizableObjectInstance>>
-{
-	FORCEINLINE static TWeakObjectPtr<const UCustomizableObjectInstance> GetSetKey(const FMutablePendingInstanceUpdate& PendingUpdate);
-
-	FORCEINLINE static bool Matches(const TWeakObjectPtr<const UCustomizableObjectInstance>& A, const TWeakObjectPtr<const UCustomizableObjectInstance>& B);
-
-	FORCEINLINE static uint32 GetKeyHash(const TWeakObjectPtr<const UCustomizableObjectInstance>& Identifier);
-};
-
-
-struct FMutablePendingInstanceDiscard
-{
-	TWeakObjectPtr<UCustomizableObjectInstance> CustomizableObjectInstance;
-
-	FMutablePendingInstanceDiscard(UCustomizableObjectInstance* InCustomizableObjectInstance)
-	{
-		CustomizableObjectInstance = InCustomizableObjectInstance;
-	}
-
-	friend bool operator ==(const FMutablePendingInstanceDiscard& A, const FMutablePendingInstanceDiscard& B)
-	{
-		return A.CustomizableObjectInstance.HasSameIndexAndSerialNumber(B.CustomizableObjectInstance);
-	}
-};
-
-
-inline uint32 GetTypeHash(const FMutablePendingInstanceDiscard& Discard)
-{
-	return GetTypeHash(Discard.CustomizableObjectInstance.GetWeakPtrTypeHash());
-}
-
-
-struct FPendingInstanceDiscardKeyFuncs : BaseKeyFuncs<FMutablePendingInstanceUpdate, TWeakObjectPtr<UCustomizableObjectInstance>>
-{
-	FORCEINLINE static const TWeakObjectPtr<UCustomizableObjectInstance>& GetSetKey(const FMutablePendingInstanceDiscard& PendingDiscard)
-	{
-		return PendingDiscard.CustomizableObjectInstance;
-	}
-
-	FORCEINLINE static bool Matches(const TWeakObjectPtr<UCustomizableObjectInstance>& A, const TWeakObjectPtr<UCustomizableObjectInstance>& B)
-	{
-		return A.HasSameIndexAndSerialNumber(B);
-	}
-
-	FORCEINLINE static uint32 GetKeyHash(const TWeakObjectPtr<UCustomizableObjectInstance>& Identifier)
-	{
-		return GetTypeHash(Identifier.GetWeakPtrTypeHash());
-	}
-};
-
-
-/** Instance updates queue.
- *
- * The queues will only contain a single operation per UCustomizableObjectInstance.
- * If there is already an operation it will be replaced. */
-class FMutablePendingInstanceWork
-{
-	TSet<FMutablePendingInstanceUpdate, FPendingInstanceUpdateKeyFuncs> PendingInstanceUpdates;
-
-	TSet<FMutablePendingInstanceDiscard, FPendingInstanceDiscardKeyFuncs> PendingInstanceDiscards;
-
-public:
-	// Returns the number of pending instance updates, LOD Updates and discards.
-	int32 Num() const;
-	
-	// Adds a new instance update
-	void AddUpdate(const FMutablePendingInstanceUpdate& UpdateToAdd);
-
-	// Removes an instance update
-	void RemoveUpdate(const TWeakObjectPtr<UCustomizableObjectInstance>& Instance);
-
-#if WITH_EDITOR
-	void RemoveUpdatesForObject(const UCustomizableObject* InObject);
-#endif
-
-	const FMutablePendingInstanceUpdate* GetUpdate(const TWeakObjectPtr<const UCustomizableObjectInstance>& Instance) const;
-
-	TSet<FMutablePendingInstanceUpdate, FPendingInstanceUpdateKeyFuncs>::TIterator GetUpdateIterator()
-	{
-		return PendingInstanceUpdates.CreateIterator();
-	}
-
-	TSet<FMutablePendingInstanceDiscard, FPendingInstanceDiscardKeyFuncs>::TIterator GetDiscardIterator()
-	{
-		return PendingInstanceDiscards.CreateIterator();
-	}
-
-	void AddDiscard(const FMutablePendingInstanceDiscard& TaskToEnqueue);
-};
-
-
 /** Final data per component. */
 struct FSkeletalMeshMorphTargets
 {
@@ -384,15 +244,6 @@ public:
 	CUSTOMIZABLEOBJECT_API FUpdateContextPrivate(UCustomizableObjectInstance& InInstance);
 	
 	bool IsContextValid() const;
-	
-	void SetMinLOD(const TMap<FName, uint8>& MinLOD);
-
-	/** Return an array of LODs per object component. */
-	const TMap<FName, uint8>& GetFirstRequestedLOD() const;
-	
-	void SetFirstRequestedLOD(const TMap<FName, uint8>& RequestedLODs);
-
-	void SetQualitySettingMinLODs(const TMap<FName, uint8>& FirstLODs);
 
 	const FCustomizableObjectInstanceDescriptor& GetCapturedDescriptor() const;
 
@@ -511,19 +362,11 @@ public:
 	/** This option comes from the operation request. It is used to reduce the number of mipmaps that mutable must generate for images.  */
 	int32 MipsToSkip = 0;
 	
-	EQueuePriorityType PriorityType = EQueuePriorityType::Low;
-
 	EUpdateResult UpdateResult = EUpdateResult::Success;
 	
 	/** This option comes from the operation request */
 	uint8 bNeverStreamMips : 1 = false;
 	
-	uint8 bOnlyUpdateIfNotGenerated : 1 = false;
-
-	uint8 bIgnoreCloseDist : 1 = false;
-	
-	uint8 bForceHighPriority : 1 = false;
-	
 	/** When this option is enabled it will reuse the Mutable core instance and its temp data between updates.  */
 	uint8 bLiveUpdateMode : 1 = false;
 
@@ -539,6 +382,9 @@ public:
 	/** true if the Update has blocked Low Priority Tasks from launching. */
 	uint8 bLowPriorityTasksBlocked : 1 = false;
 
+	/** Add to the high priority queue or not */
+	uint8 bIsHighPriority : 1 = false;
+
 	/** The Context has been successfully created. */
 	uint8 bValid : 1 = false;
 
@@ -551,9 +397,6 @@ public:
 	/** User config. */
 	uint8 bProgressiveMipStreamingEnabled : 1 = false;
 	
-	/** User config. */
-	uint8 bIsOnlyGenerateRequestedLODsEnabled : 1 = false;
-	
 	uint8 bIsUpdateAborted : 1 = false;
 		
 	uint8 bKeepOwnershipOfGeneratedResources : 1 = false;
@@ -572,18 +415,7 @@ public:
 };
 
 
-USTRUCT()
-struct FPendingReleaseSkeletalMeshInfo
-{
-	GENERATED_USTRUCT_BODY()
-
-	UPROPERTY()
-	TObjectPtr<USkeletalMesh> SkeletalMesh = nullptr;
-
-	UPROPERTY()
-	double TimeStamp = 0.0f;
-};
-
+typedef TSharedPtr<FUpdateContextPrivate> FMutableInstanceUpdate;
 
 #if WITH_EDITORONLY_DATA
 
@@ -626,37 +458,16 @@ struct FEditorCompileSettings
  *               ^                    ^
  *               |                    |
  *               FirstLODAvailable    FirstResidentLOD
- *                                    FirstRequestedLOD (LODs generated by Core)
  * 
  * [- NumLODsToStream ----------------] = 5 (Compiled constant, ModelResources)
  *
- *
- * HACKY MUTABLE STREAMING:
- * 
- * [- NumLODsAvailable -----------------------------] = 8
- * [- Stripped --[- Packaged -----------------------]
- * 0      1      2      3      4      5      6      7    
- * |------|------|------|------|------|------|------|
- *               [- Residents ----------------------]
- *               [------] Data copied from FirstRequestedLODs. Hacky Mutable LOD Streaming.
- *               ^      ^         
- *               |      |
- *               |      FirstRequestedLOD (LODs generated by Core)
- *               |      
- *               FirstLODAvailable (Compilation constant)
- *               FirstResidentLOD
- *               MinLOD
- * 
  * NumLODsToStream = 0
  *
  *
  * DEFINITIONS:
  * 
- * QualitySettingMinLODs =	MinLOD based on the active quality settings.													COI Descriptor.
- * MinLOD =					From user. Artificial limit.											                     	Skeletal Mesh Component.
  * FirstLODAvailable =	    First available lod per platform.																Skeletal Mesh Component.
  * FirstResidentLOD =       First LOD generated with geometry.																Skeletal Mesh Component.
- * FirstRequestedLOD =		From user. Usually from USkeletalMeshComponent::GetPredictedLODLevel							Skeletal Mesh Component. 
  */
 UCLASS()
 class UCustomizableObjectSystemPrivate : public UObject, public IStreamingManager
@@ -675,17 +486,18 @@ public:
 
 	// This object is responsible for streaming data to the MutableSystem.
 	TSharedPtr<class FUnrealMutableModelBulkReader> Streamer;
-	
-	// The pending instance updates, discards or releases
-	FMutablePendingInstanceWork MutablePendingInstanceWork;
-	
+
+	TDeque<FMutableInstanceUpdate> HighPriorityInstanceUpdates;
+	TDeque<FMutableInstanceUpdate> LowPriorityInstanceUpdates;
+	int32 PendingInstanceUpdates = 0;
+
 	static int32 EnableMutableProgressiveMipStreaming;
 	static int32 EnableMutableLiveUpdate;
 	static int32 EnableMutableAnimInfoDebugging;
 	static int32 EnableSkipGenerateResidentMips;
-	static int32 EnableOnlyGenerateRequestedLODs;
+	static int32 MutableMeshesLODBias;
 	static int32 MaxTextureSizeToGenerate;
-	static int32 SkeletalMeshMinLodQualityLevel;
+	static bool bGenerateInstancesWithinRange;
 
 #if WITH_EDITOR
 	UE::Mutable::Private::FImageOperator::FImagePixelFormatFunc ImageFormatOverrideFunc;
@@ -708,24 +520,12 @@ public:
 	UE_API int32 GetRemainingWork() const;
 	
 	UCustomizableObjectSystem* GetPublic() const;
-	
-	EUpdateRequired IsUpdateRequired(const UCustomizableObjectInstance& Instance, bool bOnlyUpdateIfNotGenerated, bool bOnlyUpdateIfLOD, bool bIgnoreCloseDist) const;
-
-	EQueuePriorityType GetUpdatePriority(const UCustomizableObjectInstance& Instance, bool bForceHighPriority) const;
 
 	CUSTOMIZABLEOBJECT_API void EnqueueUpdateSkeletalMesh(const TSharedRef<FUpdateContextPrivate>& Context);
-		
-	// Init an async and safe release of the UE and Mutable resources used by the instance without actually destroying the instance, for example if it's very far away
-	void InitDiscardResourcesSkeletalMesh(UCustomizableObjectInstance* InCustomizableObjectInstance);
-
-	bool IsReplaceDiscardedWithReferenceMeshEnabled() const;
-	void SetReplaceDiscardedWithReferenceMeshEnabled(bool bIsEnabled);
 
 	/** Updated at the beginning of each tick. */
 	int32 GetNumSkeletalMeshes() const;
 
-	bool bReplaceDiscardedWithReferenceMesh = false;
-	
 	bool bSupport16BitBoneIndex = false;
 	
 	UPROPERTY(Transient)
@@ -743,8 +543,11 @@ public:
 	/** Update the last set amount of internal memory Mutable can use to build objects. */
 	void UpdateMemoryLimit();
 	
+	bool TryStartUpdate(const TSharedRef<FUpdateContextPrivate>& Context);
+	bool TryStartAutomaticUpdateForInstance(UCustomizableObjectInstance& Instance);
+
 	/** Start the actual work of Update Skeletal Mesh process (Update Skeletal Mesh without the queue). */
-	void StartUpdateSkeletalMesh(const TSharedRef<FUpdateContextPrivate>& Context);
+	void StartUpdate(const TSharedRef<FUpdateContextPrivate>& Context);
 
 	/** See UCustomizableObjectInstance::IsUpdating. */
 	bool IsUpdating(const UCustomizableObjectInstance& Instance) const;
@@ -759,16 +562,15 @@ public:
 	
 	bool bAutoCompileCommandletEnabled = false;
 
-	bool bPendingLODUpdate = false;
+	// mutable.GenerateInstancesWithinRange = true
+	void UpdateInstanceDistances();
+	bool IsDiscardedByDistance(const UCustomizableInstancePrivate* Instance) const;
 
-	UPROPERTY()
-	TArray<FPendingReleaseSkeletalMeshInfo> PendingReleaseSkeletalMesh;
-	
-	UPROPERTY(Transient)
-	TObjectPtr<UCustomizableInstanceLODManagementBase> DefaultInstanceLODManagement = nullptr;
+	// The list of actors that define a view radius (GenerationRange) around them.
+	TSet<TWeakObjectPtr<const AActor>> ViewCenters;
 
-	UPROPERTY(Transient)
-	TObjectPtr<UCustomizableInstanceLODManagementBase> CurrentInstanceLODManagement = nullptr;
+	// Instance generation range when mutable.GenerateInstancesWithinRange is true
+	float GenerationRangeSquare = 2000.f * 2000.f;
 	
 	TSharedRef<FMutableStreamableManager> StreamableManager = MakeShared<FMutableStreamableManager>();
 
@@ -783,13 +585,8 @@ public:
 
 	/** Weak pointer to the Uncompiled Customizable Objects notification */
 	TWeakPtr<SNotificationItem> UncompiledCustomizableObjectsNotificationPtr;
-
-	/** Map used to cache per platform MaxChunkSize. If MaxChunkSize > 0, streamed data will be split in multiple files */
-	TMap<FString, int64> PlatformMaxChunkSize;
 #endif
 
-	int32 NumLODUpdatesLastTick = 0;
-
 	/** Time when the "Started Update Skeletal Mesh Async" log will be unmuted (in seconds). */
 	float LogStartedUpdateUnmute = 0.0;
 
@@ -797,19 +594,19 @@ public:
 	float LogStartedUpdateLast = 0;
 
 public:
-	static const TSet<UCustomizableObjectInstanceUsage*>& GetCustomizableObjectInstanceUsages()
+	CUSTOMIZABLEOBJECT_API static const TArray<UCustomizableObjectInstanceUsage*>& GetCustomizableObjectInstanceUsages()
 	{
 		return CustomizableObjectInstanceUsages;
 	}
-	static void RegisterCustomizableObjectInstanceUsage(UCustomizableObjectInstanceUsage* Obj);
-	static void UnregisterCustomizableObjectInstanceUsage(UCustomizableObjectInstanceUsage* Obj);
+	static void RegisterCustomizableObjectInstanceUsage(UCustomizableObjectInstanceUsage& Obj);
+	static void UnregisterCustomizableObjectInstanceUsage(UCustomizableObjectInstanceUsage& Obj);
 
-	static const TSet<UCustomizableObjectInstance*>& GetCustomizableObjectInstances()
+	static const TArray<UCustomizableObjectInstance*>& GetCustomizableObjectInstances()
 	{
 		return CustomizableObjectInstances;
 	}
-	static void RegisterCustomizableObjectInstance(UCustomizableObjectInstance* Obj);
-	static void UnregisterCustomizableObjectInstance(UCustomizableObjectInstance* Obj);
+	static void RegisterCustomizableObjectInstance(UCustomizableObjectInstance& Obj);
+	static void UnregisterCustomizableObjectInstance(UCustomizableObjectInstance& Obj);
 
 	/**
 	 * Get to know if the settings used by the mutable syustem are optimzed for benchmarking operations or not
@@ -830,8 +627,8 @@ private:
 	inline static bool bUseBenchmarkingSettings = false;
 
 	// Added to in PostInitProperties, removed from in BeginDestroy. No CDOs
-	static TSet<UCustomizableObjectInstanceUsage*> CustomizableObjectInstanceUsages;
-	static TSet<UCustomizableObjectInstance*> CustomizableObjectInstances;
+	static TArray<UCustomizableObjectInstanceUsage*> CustomizableObjectInstanceUsages;
+	static TArray<UCustomizableObjectInstance*> CustomizableObjectInstances;
 };
 
 
diff --git a/Engine/Plugins/Mutable/Source/CustomizableObject/Internal/MuCO/DescriptorHash.h b/Engine/Plugins/Mutable/Source/CustomizableObject/Internal/MuCO/DescriptorHash.h
--- a/Engine/Plugins/Mutable/Source/CustomizableObject/Internal/MuCO/DescriptorHash.h
+++ b/Engine/Plugins/Mutable/Source/CustomizableObject/Internal/MuCO/DescriptorHash.h
@@ -16,26 +16,14 @@ public:
 
 	UE_API explicit FDescriptorHash(const FCustomizableObjectInstanceDescriptor& Descriptor);
 	
-	/** Return true if this Hash is a subset of the other Hash (i.e., this Descriptor is a subset of the other Descriptor). */
-	UE_API bool IsSubset(const FDescriptorHash& Other) const;
+	/** Return true if this Hashes are the same. */
+	UE_API bool operator==(const FDescriptorHash& Other) const = default;
 
 	UE_API FString ToString() const;
 
 private:
 	uint32 Hash = 0;
 
-public:
-
-	bool bStreamingEnabled = false;
-
-	// MinLOD based on user input and quality settings. First LOD to generate.
-	TMap<FName, uint8> MinLODs;
-
-	// MinLOD based on quality settings. Used to trigger updates after changing the active quality level.
-	TMap<FName, uint8> QualitySettingMinLODs;
-	
-	// Array of bitmasks that indicate which LODs of each component have been requested
-	TMap<FName, uint8> FirstRequestedLOD;
 };
 
 #undef UE_API
diff --git a/Engine/Plugins/Mutable/Source/CustomizableObject/Internal/MuCO/ICustomizableObjectEditorModule.h b/Engine/Plugins/Mutable/Source/CustomizableObject/Internal/MuCO/ICustomizableObjectEditorModule.h
--- a/Engine/Plugins/Mutable/Source/CustomizableObject/Internal/MuCO/ICustomizableObjectEditorModule.h
+++ b/Engine/Plugins/Mutable/Source/CustomizableObject/Internal/MuCO/ICustomizableObjectEditorModule.h
@@ -109,6 +109,8 @@ public:
 	 * If calling inside the Customizable Object Editor module, consider using ICustomizableObjectEditorModulePrivate::EnqueueCompileRequest. */
 	virtual void CompileCustomizableObject(UCustomizableObject& Object, const FCompileParams* Params, bool bSilent, bool bForce) = 0;
 
+	virtual void UpdateInstanceDistanceToLevelEditorCamera() = 0;
+
 	virtual void Tick(bool bBlocking) = 0;
 	
 	virtual int32 GetRemainingWork() const = 0;
diff --git a/Engine/Plugins/Mutable/Source/CustomizableObject/Private/MuCO/CustomizableInstanceLODManagement.cpp b/Engine/Plugins/Mutable/Source/CustomizableObject/Private/MuCO/CustomizableInstanceLODManagement.cpp
deleted file mode 100644
diff --git a/Engine/Plugins/Mutable/Source/CustomizableObject/Private/MuCO/CustomizableObject.cpp b/Engine/Plugins/Mutable/Source/CustomizableObject/Private/MuCO/CustomizableObject.cpp
--- a/Engine/Plugins/Mutable/Source/CustomizableObject/Private/MuCO/CustomizableObject.cpp
+++ b/Engine/Plugins/Mutable/Source/CustomizableObject/Private/MuCO/CustomizableObject.cpp
@@ -865,15 +865,15 @@ USkeletalMesh* UCustomizableObject::GetComponentMeshReferenceSkeletalMesh(const
 {
 	// Try to get the data from the CO graph.
 #if WITH_EDITORONLY_DATA
-	if (const ICustomizableObjectEditorModule* Module = ICustomizableObjectEditorModule::Get())
+	if (!IsCompiled() || GetPrivate()->IsLocked())
 	{
-		return Module->GetReferenceSkeletalMesh(*this, ComponentName);
+		if (const ICustomizableObjectEditorModule* Module = ICustomizableObjectEditorModule::Get())
+		{
+			return Module->GetReferenceSkeletalMesh(*this, ComponentName);
+		}
 	}
 #endif
 
-	// If WITH_EDITORONLY_DATA is false Look instead for the ReferenceSkeletalMesh in the compiled resources
-	// TODO UE-282479  Make this a separated option so we can access the compiled data even in the editor execution as
-	// currently the graph search will be always performed if the editor only data is accessible.
 	if (const UModelResources* ModelResources = Private->GetModelResources())
 	{
 		// TODO Due to the adding support of multiple components with the same name this is no longer correct. A single name may point to multiple FComponentId.
@@ -1816,54 +1816,6 @@ void UCustomizableObjectPrivate::ApplyStateForcedValuesToParameters(FCustomizabl
 }
 
 
-uint8 UCustomizableObjectPrivate::GetMinLODIndex(const UE::Mutable::Private::FComponentId ComponentId) const
-{
-	uint8 MinLODIdx = 0;
-
-	const UModelResources* LocalModelResources = GetModelResources();
-	if (LocalModelResources)
-	{
-		// Use the Scalability quality settings to determine what the MinLOD will be used
-		if (GEngine && GEngine->UseSkeletalMeshMinLODPerQualityLevels)
-		{
-			if (UCustomizableObjectSystem::GetInstance() != nullptr)
-			{
-				// Get the quality level for the skeletal meshes for the current scalability setting. The bigger the better.
-				// Value extracted from the Scalability.ini of the current platform (if found).
-				const int32 QualityLevel = UCustomizableObjectSystem::GetInstance()->GetSkeletalMeshMinLODQualityLevel();
-
-				const FPerQualityLevelInt* QualityLevelInt = LocalModelResources->MinQualityLevelLODs.Find(ComponentId);
-				if (QualityLevelInt)
-				{
-					MinLODIdx = QualityLevelInt->GetValue(QualityLevel);
-				}
-			}
-		}
-		else
-		{
-			const FPerPlatformInt* PlatformInt = LocalModelResources->MinLODs.Find(ComponentId);
-			if (PlatformInt)
-			{
-				MinLODIdx = PlatformInt->GetValue();
-			}
-		}
-	}
-	
-	// Get the first lod the current platform can get generated. It represents the Min LOD mutable can generate for the current platform.
-	uint8 FirstLODAvailable = 0;
-
-	if (LocalModelResources)
-	{
-		if (const uint8* Result = LocalModelResources->FirstLODs.Find(ComponentId))
-		{
-			FirstLODAvailable = *Result;
-		}
-	}
-
-	return FMath::Max(MinLODIdx, FirstLODAvailable);
-}
-
-
 //-------------------------------------------------------------------------------------------------
 //-------------------------------------------------------------------------------------------------
 //-------------------------------------------------------------------------------------------------
diff --git a/Engine/Plugins/Mutable/Source/CustomizableObject/Private/MuCO/CustomizableObjectInstance.cpp b/Engine/Plugins/Mutable/Source/CustomizableObject/Private/MuCO/CustomizableObjectInstance.cpp
--- a/Engine/Plugins/Mutable/Source/CustomizableObject/Private/MuCO/CustomizableObjectInstance.cpp
+++ b/Engine/Plugins/Mutable/Source/CustomizableObject/Private/MuCO/CustomizableObjectInstance.cpp
@@ -272,11 +272,10 @@ bool UCustomizableObjectInstance::IsEditorOnly() const
 void UCustomizableObjectInstance::PostInitProperties()
 {
 	Super::PostInitProperties();
+	UCustomizableObjectSystemPrivate::RegisterCustomizableObjectInstance(*this);
 
 	if (!HasAnyFlags(RF_ClassDefaultObject | RF_DefaultSubObject))
 	{
-		UCustomizableObjectSystemPrivate::RegisterCustomizableObjectInstance(this);
-
 		if (!PrivateData)
 		{
 			PrivateData = NewObject<UCustomizableInstancePrivate>(this, FName("Private"));
@@ -299,11 +298,7 @@ void UCustomizableObjectInstance::BeginDestroy()
 #endif
 	}
 
-	if (!HasAnyFlags(RF_ClassDefaultObject | RF_DefaultSubObject))
-	{
-		UCustomizableObjectSystemPrivate::UnregisterCustomizableObjectInstance(this);
-	}
-	
+	UCustomizableObjectSystemPrivate::UnregisterCustomizableObjectInstance(*this);
 	Super::BeginDestroy();
 }
 
@@ -775,156 +770,41 @@ bool UCustomizableInstancePrivate::CanUpdateInstance() const
 }
 
 
-void UCustomizableObjectInstance::UpdateSkeletalMeshAsync(bool bIgnoreCloseDist, bool bForceHighPriority)
+void UCustomizableObjectInstance::UpdateSkeletalMeshAsync(bool, bool bForceHighPriority)
 {
 	UCustomizableObjectSystemPrivate* SystemPrivate = UCustomizableObjectSystem::GetInstance()->GetPrivate();
 
 	const TSharedRef<FUpdateContextPrivate> Context = MakeShared<FUpdateContextPrivate>(*this);
-	Context->bIgnoreCloseDist = bIgnoreCloseDist;
-	Context->bForceHighPriority = bForceHighPriority;
+	Context->bIsHighPriority = bForceHighPriority;
 	
 	SystemPrivate->EnqueueUpdateSkeletalMesh(Context);
 }
 
 
-void UCustomizableObjectInstance::UpdateSkeletalMeshAsyncResult(FInstanceUpdateDelegate Callback, bool bIgnoreCloseDist, bool bForceHighPriority)
+void UCustomizableObjectInstance::UpdateSkeletalMeshAsyncResult(FInstanceUpdateDelegate Callback, bool, bool bForceHighPriority)
 {
 	UCustomizableObjectSystemPrivate* SystemPrivate = UCustomizableObjectSystem::GetInstance()->GetPrivate();
 
 	const TSharedRef<FUpdateContextPrivate> Context = MakeShared<FUpdateContextPrivate>(*this);
-	Context->bIgnoreCloseDist = bIgnoreCloseDist;
-	Context->bForceHighPriority = bForceHighPriority;
+	Context->bIsHighPriority = bForceHighPriority;
 	Context->UpdateCallback = Callback;
 	
 	SystemPrivate->EnqueueUpdateSkeletalMesh(Context);
 }
 
 
-void UCustomizableObjectInstance::UpdateSkeletalMeshAsyncResult(FInstanceUpdateNativeDelegate Callback, bool bIgnoreCloseDist, bool bForceHighPriority)
+void UCustomizableObjectInstance::UpdateSkeletalMeshAsyncResult(FInstanceUpdateNativeDelegate Callback, bool, bool bForceHighPriority)
 {
 	UCustomizableObjectSystemPrivate* SystemPrivate = UCustomizableObjectSystem::GetInstance()->GetPrivate();
 
 	const TSharedRef<FUpdateContextPrivate> Context = MakeShared<FUpdateContextPrivate>(*this);
-	Context->bIgnoreCloseDist = bIgnoreCloseDist;
-	Context->bForceHighPriority = bForceHighPriority;
+	Context->bIsHighPriority = bForceHighPriority;
 	Context->UpdateNativeCallback = Callback;
 	
 	SystemPrivate->EnqueueUpdateSkeletalMesh(Context);
 }
 
 
-void UCustomizableInstancePrivate::TickUpdateCloseCustomizableObjects(UCustomizableObjectInstance& Public, FMutableInstanceUpdateMap& InOutRequestedUpdates)
-{
-	UCustomizableObject* Object = Public.GetCustomizableObject();
-	if (!Object)
-	{
-		return;
-	}
-
-#if WITH_EDITOR
-	if (!Object->IsCompiled() &&
-		Object->GetPrivate()->CompilationResult != ECompilationResultPrivate::Errors) // Avoid constantly retry failed compilations.
-	{
-		if (ICustomizableObjectEditorModule* EditorModule = ICustomizableObjectEditorModule::Get())
-		{
-			FCompileParams CompileParams;
-			CompileParams.bSkipIfCompiled = true;
-			
-			EditorModule->CompileCustomizableObject(*Object, &CompileParams, true, false);
-		}
-	}
-#endif
-		
-	if (!CanUpdateInstance())
-	{
-		return;
-	}
-
-	if (SkeletalMeshStatus == ESkeletalMeshStatus::Error)
-	{
-		return;
-	}
-
-	const UCustomizableObjectSystemPrivate* SystemPrivate = UCustomizableObjectSystem::GetInstance()->GetPrivate();	
-
-	const EUpdateRequired UpdateRequired = SystemPrivate->IsUpdateRequired(Public, true, true, false);
-	if (UpdateRequired != EUpdateRequired::NoUpdate) // Since this is done in the tick, avoid starting an update that we know for sure that would not be performed. Once started it has some performance implications that we want to avoid.
-	{
-		if (UpdateRequired == EUpdateRequired::Discard)
-		{
-			UCustomizableObjectSystem::GetInstance()->GetPrivate()->InitDiscardResourcesSkeletalMesh(&Public);
-			InOutRequestedUpdates.Remove(&Public);
-		}
-		else if (UpdateRequired == EUpdateRequired::Update)
-		{
-			const EQueuePriorityType Priority = SystemPrivate->GetUpdatePriority(Public, false);
-
-			if (FMutableUpdateCandidate* UpdateCandidate = InOutRequestedUpdates.Find(&Public))
-			{
-				ensure(HasCOInstanceFlags(PendingLODsUpdate | PendingLODsDowngrade));
-
-				UpdateCandidate->Priority = Priority;
-				UpdateCandidate->Issue();
-			}
-			else
-			{
-				FMutableUpdateCandidate Candidate(&Public);
-				Candidate.Priority = Priority;
-				Candidate.Issue();
-				InOutRequestedUpdates.Add(&Public, Candidate);
-			}
-		}
-		else
-		{
-			check(false);
-		}
-	}
-	else
-	{
-		InOutRequestedUpdates.Remove(&Public);
-	}
-
-	ClearCOInstanceFlags(PendingLODsUpdate | PendingLODsDowngrade);
-}
-
-
-void UCustomizableInstancePrivate::UpdateInstanceIfNotGenerated(UCustomizableObjectInstance& Public, FMutableInstanceUpdateMap& InOutRequestedUpdates)
-{
-	if (SkeletalMeshStatus != ESkeletalMeshStatus::NotGenerated)
-	{
-		return;
-	}
-
-	if (!CanUpdateInstance())
-	{
-		return;
-	}
-
-	UCustomizableObjectSystemPrivate* SystemPrivate = UCustomizableObjectSystem::GetInstance()->GetPrivate();
-
-	const TSharedRef<FUpdateContextPrivate> Context = MakeShared<FUpdateContextPrivate>(Public);
-	Context->bOnlyUpdateIfNotGenerated = true;
-	
-	SystemPrivate->EnqueueUpdateSkeletalMesh(Context);
-
-	EQueuePriorityType Priority = SystemPrivate->GetUpdatePriority(Public, false);
-	FMutableUpdateCandidate* UpdateCandidate = InOutRequestedUpdates.Find(&Public);
-
-	if (UpdateCandidate)
-	{
-		UpdateCandidate->Priority = Priority;
-		UpdateCandidate->Issue();
-	}
-	else
-	{
-		FMutableUpdateCandidate Candidate(&Public);
-		Candidate.Priority = Priority;
-		Candidate.Issue();
-		InOutRequestedUpdates.Add(&Public, Candidate);
-	}
-}
-
-
 #if !UE_BUILD_SHIPPING
 bool AreSkeletonsCompatible(const TArray<TObjectPtr<USkeleton>>& InSkeletons)
 {
@@ -1448,12 +1328,8 @@ UPhysicsAsset* GetOrBuildMainPhysicsAsset(
 	const FName ComponentName = Context->ComponentNames[Component->Id];
 	
 	TArray<TObjectPtr<UPhysicsAsset>> ValidAssets;
-
-	const int32 FirstLOD = Context->bStreamMeshLODs ?
-		Context->FirstLODAvailable[ComponentName] :
-		Context->GetFirstRequestedLOD()[ComponentName];
-	
-	for (int32 LODIndex = FirstLOD; LODIndex < Component->LODCount; ++LODIndex)
+		
+	for (int32 LODIndex = Context->FirstLODAvailable[ComponentName]; LODIndex < Component->LODCount; ++LODIndex)
 	{
 		const TSharedRef<FInstanceUpdateData::FLOD>& LOD = Context->InstanceUpdateData.LODs[Component->FirstLOD + LODIndex];
 		check(LOD->Mesh);
@@ -1768,10 +1644,6 @@ bool DoComponentsNeedUpdate(const TSharedRef<FUpdateContextPrivate>& Context, bo
 		
 		bool bComponentWithMesh = false;
 
-		const int32 FirstLOD = Context->bStreamMeshLODs ?
-			Context->FirstLODAvailable[ComponentName] :
-			Context->GetFirstRequestedLOD()[ComponentName];
-
 		bool bMeshChanged = true;
 		
 		if (Component->SkeletalMesh->PassthroughObject)
@@ -1781,7 +1653,7 @@ bool DoComponentsNeedUpdate(const TSharedRef<FUpdateContextPrivate>& Context, bo
 		else
 		{
 			// Look for invalid meshes
-			for (int32 LODIndex = FirstLOD; LODIndex < Component->LODCount; ++LODIndex)
+			for (int32 LODIndex = Context->FirstLODAvailable[ComponentName]; LODIndex < Component->LODCount; ++LODIndex)
 			{
 				const TSharedRef<FInstanceUpdateData::FLOD>& LOD = Context->InstanceUpdateData.LODs[Component->FirstLOD + LODIndex];
 
@@ -1895,10 +1767,6 @@ bool UCustomizableInstancePrivate::UpdateSkeletalMesh_PostBeginUpdate0(UCustomiz
 			}
 		}
 
-		const int32 FirstLOD = Context->bStreamMeshLODs ?
-			Context->FirstLODAvailable[ComponentName] :
-			Context->GetFirstRequestedLOD()[ComponentName];
-		
 		if (Context->bUseMeshCache)
 		{
 			if (USkeletalMesh* CachedMesh = CustomizableObject->GetPrivate()->MeshCache.Get(ComponentDescriptor->MeshDescriptors))
@@ -1937,7 +1805,7 @@ bool UCustomizableInstancePrivate::UpdateSkeletalMesh_PostBeginUpdate0(UCustomiz
 		USkeletalMesh* SkeletalMesh = nullptr;
 		{
 			uint32 MeshHash = INDEX_NONE;
-			for (int32 LODIndex = FirstLOD; LODIndex < Component->LODCount && LODIndex < MAX_MESH_LOD_COUNT; ++LODIndex)
+			for (int32 LODIndex = Context->FirstLODAvailable[ComponentName]; LODIndex < Component->LODCount && LODIndex < MAX_MESH_LOD_COUNT; ++LODIndex)
 			{
 				const UE::Mutable::Private::FMeshId& MeshId = ComponentDescriptor->MeshDescriptors[LODIndex];
 				const uint32 MeshIDHash = UE::Mutable::Private::GetTypeHashPersistent(MeshId);
@@ -2796,39 +2664,6 @@ bool UCustomizableInstancePrivate::IsSelectedParameterProfileDirty() const
 }
 
 
-void UCustomizableInstancePrivate::DiscardResources()
-{
-	check(IsInGameThread());
-
-	UCustomizableObjectInstance* Instance = Cast<UCustomizableObjectInstance>(GetOuter());
-	if (!Instance)
-	{
-		return;
-	}
-
-	if (SkeletalMeshStatus == ESkeletalMeshStatus::Success)
-	{
-		if (CVarEnableReleaseMeshResources.GetValueOnGameThread())
-		{
-			for (const TTuple<FName, TObjectPtr<USkeletalMesh>>& Tuple : SkeletalMeshes)
-			{
-				USkeletalMesh* SkeletalMesh = Tuple.Get<1>();
-			
-				if (SkeletalMesh->IsValidLowLevel() && !SkeletalMesh->HasPendingInitOrStreaming())
-				{
-					SkeletalMesh->ReleaseResources();
-				}
-			}
-		}
-		
-		SkeletalMeshes.Empty();
-	}
-	
-	InvalidateGeneratedData();
-	
-}
-
-
 void UCustomizableInstancePrivate::ForceSetReferenceSkeletalMesh() const
 {
 	UCustomizableObjectInstance* Instance = Cast<UCustomizableObjectInstance>(GetOuter());
@@ -2849,16 +2684,15 @@ void UCustomizableInstancePrivate::ForceSetReferenceSkeletalMesh() const
 		return;
 	}
 	
-	for (TObjectIterator<UCustomizableObjectInstanceUsage> It; It; ++It)
+	for (UCustomizableObjectInstanceUsage* CustomizableObjectInstanceUsage : UCustomizableObjectSystemPrivate::GetCustomizableObjectInstanceUsages())
 	{
-		UCustomizableObjectInstanceUsage* CustomizableObjectInstanceUsage = *It;
 		if (!IsValid(CustomizableObjectInstanceUsage) || CustomizableObjectInstanceUsage->GetCustomizableObjectInstance() != Instance)
 		{
 			continue;
 		}
 
 #if WITH_EDITOR
-		if (CustomizableObjectInstanceUsage->GetPrivate()->IsNetMode(NM_DedicatedServer))
+		if (CustomizableObjectInstanceUsage->GetPrivate()->bIsNetModeDedicatedServer)
 		{
 			continue;
 		}
@@ -2941,39 +2775,6 @@ UCustomizableInstancePrivate* UCustomizableObjectInstance::GetPrivate() const
 }
 
 
-FMutableUpdateCandidate::FMutableUpdateCandidate(UCustomizableObjectInstance* InCustomizableObjectInstance): CustomizableObjectInstance(InCustomizableObjectInstance)
-{
-	const FCustomizableObjectInstanceDescriptor& Descriptor = InCustomizableObjectInstance->GetPrivate()->GetDescriptor();
-	MinLOD = Descriptor.MinLOD;
-	QualitySettingMinLODs = Descriptor.QualitySettingMinLODs;
-	FirstRequestedLOD = Descriptor.GetFirstRequestedLOD();
-}
-
-
-bool FMutableUpdateCandidate::HasBeenIssued() const
-{
-	return bHasBeenIssued;
-}
-
-
-void FMutableUpdateCandidate::Issue()
-{
-	bHasBeenIssued = true;
-}
-
-
-void FMutableUpdateCandidate::ApplyLODUpdateParamsToInstance(FUpdateContextPrivate& Context)
-{
-	CustomizableObjectInstance->Descriptor.MinLOD = MinLOD;
-	CustomizableObjectInstance->Descriptor.QualitySettingMinLODs = QualitySettingMinLODs;
-	CustomizableObjectInstance->Descriptor.FirstRequestedLOD = FirstRequestedLOD;
-
-	Context.SetMinLOD(MinLOD);
-	Context.SetQualitySettingMinLODs(QualitySettingMinLODs);
-	Context.SetFirstRequestedLOD(FirstRequestedLOD);
-}
-
-
 // The memory allocated in the function and pointed by the returned pointer is owned by the caller and must be freed. 
 // If assigned to a UTexture2D, it will be freed by that UTexture2D
 TUniquePtr<FTexturePlatformData> MutableCreateImagePlatformData(UE::Mutable::Private::TManagedPtr<const UE::Mutable::Private::FImage> MutableImage, int32 OnlyLOD, uint16 FullSizeX, uint16 FullSizeY)
@@ -3353,13 +3154,13 @@ void InitSkeletalMeshData(const TSharedRef<FUpdateContextPrivate>& Context, USke
 	
 	const FName ComponentName = Context->ComponentNames[Component->Id];
 
-	SkeletalMesh->NeverStream = !Context->bStreamMeshLODs;
+	SkeletalMesh->NeverStream = Context->FirstLODAvailable[ComponentName] == Context->FirstResidentLOD[ComponentName];
 
 	SkeletalMesh->SetImportedBounds(RefSkeletalMeshData.Bounds);
 	SkeletalMesh->SetPostProcessAnimBlueprint(RefSkeletalMeshData.PostProcessAnimInst.Get());
 	SkeletalMesh->SetShadowPhysicsAsset(RefSkeletalMeshData.ShadowPhysicsAsset.Get());
 	
-	const bool bEnableRayTracingFix = EnableRayTracingFix == 2 || (EnableRayTracingFix == 1 && Context->bStreamMeshLODs);
+	const bool bEnableRayTracingFix = EnableRayTracingFix == 2 || (EnableRayTracingFix == 1 && !SkeletalMesh->NeverStream);
 	if (bEnableRayTracingFix)
 	{
 		SkeletalMesh->SetSupportRayTracing(false);
@@ -3371,16 +3172,12 @@ void InitSkeletalMeshData(const TSharedRef<FUpdateContextPrivate>& Context, USke
 	SkeletalMesh->SetPhysicsAsset(RefSkeletalMeshData.PhysicsAsset.Get());
 	SkeletalMesh->SetEnablePerPolyCollision(RefSkeletalMeshData.Settings.bEnablePerPolyCollision);
 	
-	const int32 FirstLOD = Context->bStreamMeshLODs ?
-			Context->FirstLODAvailable[ComponentName] :
-			Context->GetFirstRequestedLOD()[ComponentName];
-	
 	// Asset User Data
 	{
 		TMap<const UClass*, TArray<UAssetUserData*>> AssetUserDataByType; 
 		
 		// Hack. Get the asset user data from the first LOD.
-		for (int32 LODIndex = FirstLOD; LODIndex < Component->LODCount; ++LODIndex)
+		for (int32 LODIndex = Context->FirstLODAvailable[ComponentName]; LODIndex < Component->LODCount; ++LODIndex)
 		{
 			const TSharedRef<FInstanceUpdateData::FLOD>& LOD = Context->InstanceUpdateData.LODs[Component->FirstLOD + LODIndex];
 
@@ -3431,7 +3228,7 @@ void InitSkeletalMeshData(const TSharedRef<FUpdateContextPrivate>& Context, USke
 	// Anim Slots
 	{
 		// Hack. Get the Anim Slot data from the first LOD.
-		for (int32 LODIndex = FirstLOD; LODIndex < Component->LODCount; ++LODIndex)
+		for (int32 LODIndex = Context->FirstLODAvailable[ComponentName]; LODIndex < Component->LODCount; ++LODIndex)
 		{
 			const TSharedRef<FInstanceUpdateData::FLOD>& LOD = Context->InstanceUpdateData.LODs[Component->FirstLOD + LODIndex];
 
@@ -3549,7 +3346,17 @@ void InitSkeletalMeshData(const TSharedRef<FUpdateContextPrivate>& Context, USke
 	// Set Min LOD (Override the Reference Skeletal Mesh LOD Settings)
 	const UModelResources& ModelResources = CustomizableObject.GetPrivate()->GetModelResourcesChecked();
 	SkeletalMesh->SetMinLod(FMath::Max(ModelResources.MinLODs.FindChecked(Component->Id).GetDefault(), static_cast<int32>(Context->FirstLODAvailable[ComponentName])));
-	SkeletalMesh->SetQualityLevelMinLod(ModelResources.MinQualityLevelLODs.FindChecked(Component->Id));
+	
+	// Clamp QualityLevelLODs to mutable.MeshLODBias.
+	const int32 FirstLODAvailable = Context->FirstLODAvailable[ComponentName];
+	FPerQualityLevelInt MinQualityLevelLODs = ModelResources.MinQualityLevelLODs.FindChecked(Component->Id);
+	MinQualityLevelLODs.Default = FMath::Max(MinQualityLevelLODs.Default, FirstLODAvailable);
+	for (auto& Pair : MinQualityLevelLODs.PerQuality)
+	{
+		Pair.Value = FMath::Max(Pair.Value, FirstLODAvailable);
+	}
+
+	SkeletalMesh->SetQualityLevelMinLod(MinQualityLevelLODs);
 
 	// Set up unreal's default material, will be replaced when building materials
 	{
@@ -3583,18 +3390,13 @@ bool BuildSkeletonData(const TSharedRef<FUpdateContextPrivate>& Context, USkelet
 
 	const FName ComponentName = Context->ComponentNames[Component->Id];
 
-	// Other data
-	const int32 FirstLOD = Context->bStreamMeshLODs ?
-		Context->FirstLODAvailable[ComponentName] :
-		Context->GetFirstRequestedLOD()[ComponentName];
-
 	TArray<TObjectPtr<USkeleton>> SkeletonsToMerge;
 
 	// Add Reference Skeleton to the list of SkeletonsToMerge. TODO: Make this optional
 	const FMutableRefSkeletalMeshData& RefSkeletalMeshData = CustomizableObject->GetPrivate()->GetModelResourcesChecked().ReferenceSkeletalMeshesData[Component->Id];
 	SkeletonsToMerge.Add(RefSkeletalMeshData.Skeleton);
 
-	for (int32 LODIndex = FirstLOD; LODIndex < Component->LODCount; ++LODIndex)
+	for (int32 LODIndex = Context->FirstLODAvailable[ComponentName]; LODIndex < Component->LODCount; ++LODIndex)
 	{
 		const TSharedRef<FInstanceUpdateData::FLOD>& LOD = Context->InstanceUpdateData.LODs[Component->FirstLOD + LODIndex];
 		check(LOD->Mesh)
@@ -3854,7 +3656,7 @@ void BuildOrCopyMorphTargetsData(const TSharedRef<FUpdateContextPrivate>& Contex
 
 		TArray<FMorphTargetLODModel>& MorphLODModels = NewMorphTarget->GetMorphLODModels();
 
-		if (Context->bStreamMeshLODs)
+		if (!SkeletalMesh->NeverStream)
 		{
 			const int32 FirstResidentLOD = Context->FirstResidentLOD[ComponentName];
 			MorphLODModels.SetNum(ComponentMorphTargets->RealTimeMorphsLODData[MorphTargetIndex].Num());
@@ -3873,7 +3675,7 @@ void BuildOrCopyMorphTargetsData(const TSharedRef<FUpdateContextPrivate>& Contex
 			}
 			
 			// Residents LODs
-			for (int32 LODIndex = Context->GetFirstRequestedLOD()[ComponentName]; LODIndex < Context->NumLODsAvailable[ComponentName]; ++LODIndex)
+			for (int32 LODIndex = Context->FirstResidentLOD[ComponentName]; LODIndex < Context->NumLODsAvailable[ComponentName]; ++LODIndex)
 			{
 				MorphLODModels[LODIndex] = ComponentMorphTargets->RealTimeMorphsLODData[MorphTargetIndex][LODIndex];
 			}
@@ -3886,23 +3688,7 @@ void BuildOrCopyMorphTargetsData(const TSharedRef<FUpdateContextPrivate>& Contex
 		MorphTargets.Add(NewMorphTarget);
 	}
 
-	// Mutable hacky LOD Streaming
-	if (!Context->bStreamMeshLODs)
-	{
-		// Copy MorphTargets from the FirstGeneratedLOD to the LODs below
-		const int32 FirstRequestedLOD = Context->GetFirstRequestedLOD()[ComponentName];
-		for (int32 LODIndex = Context->FirstLODAvailable[ComponentName]; LODIndex < FirstRequestedLOD; ++LODIndex)
-		{
-			MUTABLE_CPUPROFILER_SCOPE(CopyMorphTargetsData);
-
-			for (int32 MorphTargetIndex = 0; MorphTargetIndex < MorphTargets.Num(); ++MorphTargetIndex)
-			{
-				MorphTargets[MorphTargetIndex]->GetMorphLODModels()[LODIndex] = MorphTargets[MorphTargetIndex]->GetMorphLODModels()[FirstRequestedLOD];
-			}
-		}
-	}
-
-	const bool bInKeepEmptyMorphTargets = Context->bStreamMeshLODs;
+	const bool bInKeepEmptyMorphTargets = !SkeletalMesh->NeverStream;
 	SkeletalMesh->InitMorphTargets(bInKeepEmptyMorphTargets); // True to avoid removing streamed Morph Targets.
 }
 
@@ -3970,7 +3756,7 @@ void BuildOrCopyClothingData(
 	}
 	
 	// Residents
-	for (int32 LODIndex = Context->GetFirstRequestedLOD()[ComponentName]; LODIndex < Context->NumLODsAvailable[ComponentName]; ++LODIndex)
+	for (int32 LODIndex = Context->FirstResidentLOD[ComponentName]; LODIndex < Context->NumLODsAvailable[ComponentName]; ++LODIndex)
 	{
 		const TSharedRef<FInstanceUpdateData::FLOD>& LOD = Context->InstanceUpdateData.LODs[Component->FirstLOD + LODIndex];
 
@@ -3983,12 +3769,8 @@ void BuildOrCopyClothingData(
 		UnrealConversionUtils::ClothVertexBuffers(RenderLOD, *LOD->Mesh);
 	}
 	
-	// Other data
-	const int32 FirstLOD = Context->bStreamMeshLODs ?
-		Context->FirstLODAvailable[ComponentName] :
-		Context->GetFirstRequestedLOD()[ComponentName];
-
-	for (int32 LODIndex = FirstLOD; LODIndex < Component->LODCount; ++LODIndex)
+	// Other Data
+	for (int32 LODIndex = Context->FirstLODAvailable[ComponentName]; LODIndex < Component->LODCount; ++LODIndex)
 	{
 		const TSharedRef<FInstanceUpdateData::FLOD>& LOD = Context->InstanceUpdateData.LODs[Component->FirstLOD + LODIndex];
 		
@@ -4084,11 +3866,7 @@ bool BuildOrCopyRenderData(const TSharedRef<FUpdateContextPrivate>& Context, USk
 	const UModelResources& ModelResources = CustomizableObject->GetPrivate()->GetModelResourcesChecked();
 	const FName ComponentName = ModelResources.ComponentNamesPerObjectComponent[Component->Id];
 
-	const int32 FirstLOD = Context->bStreamMeshLODs ?
-		Context->FirstLODAvailable[ComponentName] :
-		Context->GetFirstRequestedLOD()[ComponentName];
-	
-	for (int32 LODIndex = FirstLOD; LODIndex < Component->LODCount; ++LODIndex)
+	for (int32 LODIndex = Context->FirstLODAvailable[ComponentName]; LODIndex < Component->LODCount; ++LODIndex)
 	{
 		MUTABLE_CPUPROFILER_SCOPE(BuildRenderData);
 		
@@ -4294,30 +4072,6 @@ bool BuildOrCopyRenderData(const TSharedRef<FUpdateContextPrivate>& Context, USk
 		}
 	}
 	
-
-	// Mutable hacky LOD Streaming
-	if (!Context->bStreamMeshLODs)
-	{
-		// Copy LODRenderData from the FirstRequestedLOD to the LODs below
-		const int32 FirstRequestedLOD = Context->GetFirstRequestedLOD()[ComponentName];
-		for (int32 LODIndex = Context->FirstLODAvailable[ComponentName]; LODIndex < FirstRequestedLOD; ++LODIndex)
-		{
-			TRACE_CPUPROFILER_EVENT_SCOPE_TEXT(*FString::Printf(TEXT("CopyRenderData: From LOD %d to LOD %d"), FirstRequestedLOD, LODIndex));
-
-			// Render Data will be reused from the previously generated component
-			FSkeletalMeshLODRenderData& SourceLODResource = RenderData->LODRenderData[FirstRequestedLOD];
-			FSkeletalMeshLODRenderData& LODResource = RenderData->LODRenderData[LODIndex];
-
-			UnrealConversionUtils::CopySkeletalMeshLODRenderData(
-				LODResource,
-				SourceLODResource,
-				*SkeletalMesh,
-				LODIndex,
-				SkeletalMesh->GetLODInfo(LODIndex)->bAllowCPUAccess
-			);
-		}
-	}
-
 	return true;
 }
 
@@ -4987,11 +4741,7 @@ void UCustomizableInstancePrivate::BuildResources(const TSharedRef<FUpdateContex
 
 		MUTABLE_CPUPROFILER_SCOPE(BuildResources_LODLoop);
 
-		const int32 FirstLOD = Context->bStreamMeshLODs ?
-			Context->FirstLODAvailable[ComponentName] :
-			Context->GetFirstRequestedLOD()[ComponentName];
-
-		for (int32 LODIndex = FirstLOD; LODIndex < Component->LODCount; LODIndex++)
+		for (int32 LODIndex = Context->FirstLODAvailable[ComponentName]; LODIndex < Component->LODCount; LODIndex++)
 		{
 			const TSharedRef<const FInstanceUpdateData::FLOD>& LOD = Context->InstanceUpdateData.LODs[Component->FirstLOD + LODIndex];
 			
@@ -5107,30 +4857,6 @@ void UCustomizableInstancePrivate::BuildResources(const TSharedRef<FUpdateContex
 
 		if (!bUseOverrideMaterialsOnly)
 		{
-			// Mutable hacky LOD Streaming
-			if (!Context->bStreamMeshLODs)
-			{
-				// Copy data from the FirstLODAvailable into the LODs below.
-				for (int32 LODIndex = Context->FirstLODAvailable[ComponentName]; LODIndex < FirstLOD; ++LODIndex)
-				{
-					SkeletalMesh->GetLODInfo(LODIndex)->LODMaterialMap = SkeletalMesh->GetLODInfo(FirstLOD)->LODMaterialMap;
-
-					TIndirectArray<FSkeletalMeshLODRenderData>& LODRenderData = SkeletalMesh->GetResourceForRendering()->LODRenderData;
-
-					const int32 NumRenderSections = LODRenderData[LODIndex].RenderSections.Num();
-					check(NumRenderSections == LODRenderData[FirstLOD].RenderSections.Num());
-
-					if (NumRenderSections == LODRenderData[FirstLOD].RenderSections.Num())
-					{
-						for (int32 RenderSectionIndex = 0; RenderSectionIndex < NumRenderSections; ++RenderSectionIndex)
-						{
-							const int32 MaterialIndex = LODRenderData[FirstLOD].RenderSections[RenderSectionIndex].MaterialIndex;
-							LODRenderData[LODIndex].RenderSections[RenderSectionIndex].MaterialIndex = MaterialIndex;
-						}
-					}
-				}
-			}
-
 			// Force recreate render state after replacing the materials to avoid a crash in the render pipeline if the old materials are GCed while in use.
 			RecreateRenderStateOnInstanceComponent[ComponentIndex] |= SkeletalMesh->GetResourceForRendering()->IsInitialized() && SkeletalMesh->GetMaterials() != SkeletalMaterials;
 
@@ -5155,17 +4881,15 @@ void UCustomizableInstancePrivate::BuildResources(const TSharedRef<FUpdateContex
 	{
 		MUTABLE_CPUPROFILER_SCOPE(BuildResources_RecreateRenderState);
 
-		for (TObjectIterator<UCustomizableObjectInstanceUsage> It; It; ++It)
+		for (UCustomizableObjectInstanceUsage* CustomizableObjectInstanceUsage : UCustomizableObjectSystemPrivate::GetCustomizableObjectInstanceUsages())
 		{
-			UCustomizableObjectInstanceUsage* CustomizableObjectInstanceUsage = *It;
-
 			if (!IsValid(CustomizableObjectInstanceUsage) || CustomizableObjectInstanceUsage->GetCustomizableObjectInstance() != Public)
 			{
 				continue;
 			}
 
 #if WITH_EDITOR
-			if (CustomizableObjectInstanceUsage->GetPrivate()->IsNetMode(NM_DedicatedServer))
+			if (CustomizableObjectInstanceUsage->GetPrivate()->bIsNetModeDedicatedServer)
 			{
 				continue;
 			}
@@ -5235,178 +4959,12 @@ void UCustomizableObjectInstance::SetKeepOwnershipOfGeneratedResources(bool bEna
 }
 
 
-void UCustomizableObjectInstance::SetIsBeingUsedByComponentInPlay(bool bIsUsedByComponentInPlay)
-{
-	bIsUsedByComponentInPlay ? GetPrivate()->SetCOInstanceFlags(UsedByComponentInPlay) : GetPrivate()->ClearCOInstanceFlags(UsedByComponentInPlay);
-}
-
-
-bool UCustomizableObjectInstance::GetIsBeingUsedByComponentInPlay() const
-{
-	return GetPrivate()->HasCOInstanceFlags(UsedByComponentInPlay);
-}
-
-
-void UCustomizableObjectInstance::SetIsDiscardedBecauseOfTooManyInstances(bool bIsDiscarded)
-{
-	bIsDiscarded ? GetPrivate()->SetCOInstanceFlags(DiscardedByNumInstancesLimit) : GetPrivate()->ClearCOInstanceFlags(DiscardedByNumInstancesLimit);
-}
-
-
-bool UCustomizableObjectInstance::GetIsDiscardedBecauseOfTooManyInstances() const
-{
-	return GetPrivate()->HasCOInstanceFlags(DiscardedByNumInstancesLimit);
-}
-
-
-void UCustomizableObjectInstance::SetIsPlayerOrNearIt(bool bIsPlayerorNearIt)
-{
-	bIsPlayerorNearIt ? GetPrivate()->SetCOInstanceFlags(UsedByPlayerOrNearIt) : GetPrivate()->ClearCOInstanceFlags(UsedByPlayerOrNearIt);
-}
-
-
-float UCustomizableObjectInstance::GetMinSquareDistToPlayer() const
-{
-	return GetPrivate()->MinSquareDistFromComponentToPlayer;
-}
-
-void UCustomizableObjectInstance::SetMinSquareDistToPlayer(float NewValue)
-{
-	GetPrivate()->MinSquareDistFromComponentToPlayer = NewValue;
-}
-
-
 int32 UCustomizableObjectInstance::GetNumComponents() const
 {
 	return GetCustomizableObject() ? GetCustomizableObject()->GetComponentCount() : 0;
 }
 
 
-void UCustomizableObjectInstance::SetRequestedLODs(const TMap<FName, uint8>& InMinLODs, const TMap<FName, uint8>& InFirstRequestedLOD, FMutableInstanceUpdateMap& InOutRequestedUpdates)
-{
-	check(PrivateData);
-	
-	if (!GetPrivate()->CanUpdateInstance())
-	{
-		return;
-	}
-
-	if (GetPrivate()->SkeletalMeshStatus == ESkeletalMeshStatus::Error)
-	{
-		return;
-	}
-
-	UCustomizableObject* CustomizableObject = GetCustomizableObject();
-
-	if (!CustomizableObject)
-	{
-		return;
-	}
-
-	if (IsStreamingEnabled(*CustomizableObject, Descriptor.State))
-	{
-		return;
-	}
-
-	if (CVarPreserveUserLODsOnFirstGeneration.GetValueOnGameThread() &&
-		CustomizableObject->bPreserveUserLODsOnFirstGeneration &&
-		GetPrivate()->SkeletalMeshStatus != ESkeletalMeshStatus::Success)
-	{
-		return;
-	}
-	
-	FMutableUpdateCandidate MutableUpdateCandidate(this);
-
-	// Clamp Min LOD
-	const UModelResources* ModelResources = CustomizableObject->GetPrivate()->GetModelResources();
-	if (!ModelResources)
-	{
-		return;
-	}
-
-	bool bMinLODChanged = false;
-	
-	// Save the new LODs
-	MutableUpdateCandidate.MinLOD = InMinLODs;
-	MutableUpdateCandidate.FirstRequestedLOD = Descriptor.GetFirstRequestedLOD();
-	
-	const TMap<FName, uint8>& FirstRequestedLOD = GetPrivate()->CommittedDescriptorHash.FirstRequestedLOD;
-	
-	for (UE::Mutable::Private::FComponentId ComponentId = 0; ComponentId < ModelResources->ComponentNamesPerObjectComponent.Num(); ++ComponentId)
-	{	
-		const FName& ComponentName = ModelResources->ComponentNamesPerObjectComponent[ComponentId];
-		
-		uint8& InMinLOD = MutableUpdateCandidate.MinLOD.FindOrAdd(ComponentName);
-		if (const uint8* Result = InMinLODs.Find(ComponentName))
-		{
-			InMinLOD = *Result;
-		}
-
-		const uint8 MinLODIdx = CustomizableObject->GetPrivate()->GetMinLODIndex(ComponentId);
-		MutableUpdateCandidate.QualitySettingMinLODs.Add(ComponentName, MinLODIdx);
-
-		int32 MaxLODIdx = 0;
-		if (const uint8* Found = ModelResources->NumLODs.Find(ComponentId))
-		{
-			MaxLODIdx = *Found - 1;
-		}
-
-		InMinLOD = FMath::Clamp(InMinLOD, MinLODIdx, MaxLODIdx);
-
-		uint8 DescriptorMinLOD = 0;
-		if (const uint8* Result = Descriptor.MinLOD.Find(ComponentName))
-		{
-			DescriptorMinLOD = *Result;
-		}
-
-		bMinLODChanged |= DescriptorMinLOD != InMinLOD;
-
-		if (UCustomizableObjectSystem::GetInstance()->IsOnlyGenerateRequestedLODsEnabled())
-		{
-			uint8 CurrentMinLOD = 0;
-			if (const uint8* Result = GetPrivate()->CommittedDescriptor.MinLOD.Find(ComponentName))
-			{
-				CurrentMinLOD = *Result;
-			}
-
-			PrivateData->SetCOInstanceFlags(InMinLOD > CurrentMinLOD ? PendingLODsDowngrade : ECONone);
-
-			uint8 FirstNonStreamedLODIndex = 0;
-			if (const uint8* Found = ModelResources->NumStreamingLODs.Find(ComponentId))
-			{
-				FirstNonStreamedLODIndex = *Found;
-			}
-
-			MutableUpdateCandidate.FirstRequestedLOD.Add(ComponentName, FirstNonStreamedLODIndex);
-
-			uint8 PredictedLOD = FirstNonStreamedLODIndex;
-			if (const uint8* Result = InFirstRequestedLOD.Find(ComponentName))
-			{
-				PredictedLOD = FMath::Min(*Result, PredictedLOD);
-			}
-
-			if (const uint8* Result = FirstRequestedLOD.Find(ComponentName))
-			{
-				PredictedLOD = FMath::Min(*Result, PredictedLOD);
-			}
-
-			PredictedLOD = FMath::Clamp(PredictedLOD, MinLODIdx, MaxLODIdx);
-
-			// Save new RequestedLODs
-			MutableUpdateCandidate.FirstRequestedLOD[ComponentName] = PredictedLOD;
-		}
-	}
-	
-	if (bMinLODChanged || !FirstRequestedLOD.OrderIndependentCompareEqual(MutableUpdateCandidate.FirstRequestedLOD))
-	{
-		// TODO: Remove this flag as it will become redundant with the new InOutRequestedUpdates system
-		PrivateData->SetCOInstanceFlags(PendingLODsUpdate);
-
-		InOutRequestedUpdates.Add(this, MutableUpdateCandidate);
-	}
-}
-
-
 #if WITH_EDITOR
 void UCustomizableObjectInstance::Bake(const FBakingConfiguration& InBakingConfiguration)
 {
@@ -5724,7 +5282,7 @@ void CalculateBonesToRemove(const FSkeletalMeshLODRenderData& LODResource, const
 
 void UCustomizableInstancePrivate::RegenerateImportedModels(const TSharedRef<FUpdateContextPrivate>& OperationData)
 {
-	MUTABLE_CPUPROFILER_SCOPE(RegenerateEditorImportedModels);
+	MUTABLE_EDITOR_CPUPROFILER_SCOPE(RegenerateEditorImportedModels);
 
 	struct FMeshDataConvertJob
 	{
diff --git a/Engine/Plugins/Mutable/Source/CustomizableObject/Private/MuCO/CustomizableObjectInstanceDescriptor.cpp b/Engine/Plugins/Mutable/Source/CustomizableObject/Private/MuCO/CustomizableObjectInstanceDescriptor.cpp
--- a/Engine/Plugins/Mutable/Source/CustomizableObject/Private/MuCO/CustomizableObjectInstanceDescriptor.cpp
+++ b/Engine/Plugins/Mutable/Source/CustomizableObject/Private/MuCO/CustomizableObjectInstanceDescriptor.cpp
@@ -952,8 +952,6 @@ void FCustomizableObjectInstanceDescriptor::ReloadParameters()
  
 	SetState(FMath::Clamp(GetState(), 0, CustomizableObject->GetStateCount() - 1));
  
-	FirstRequestedLOD.Empty();
- 
 	TArray<FCustomizableObjectBoolParameterValue> OldBoolParameters = BoolParameters;
 	TArray<FCustomizableObjectIntParameterValue> OldIntParameters = IntParameters;
 	TArray<FCustomizableObjectFloatParameterValue> OldFloatParameters = FloatParameters;
@@ -1554,18 +1552,6 @@ void FCustomizableObjectInstanceDescriptor::ReloadParameters()
 }
  
  
-void FCustomizableObjectInstanceDescriptor::SetFirstRequestedLOD(const TMap<FName, uint8>& InFirstRequestedLOD)
-{
-	FirstRequestedLOD = InFirstRequestedLOD;
-}
- 
- 
-const TMap<FName, uint8>& FCustomizableObjectInstanceDescriptor::GetFirstRequestedLOD() const
-{
-	return FirstRequestedLOD;
-}
- 
- 
 bool FCustomizableObjectInstanceDescriptor::HasAnyParameters() const
 {
 	return !BoolParameters.IsEmpty() ||
@@ -2108,7 +2094,7 @@ void FCustomizableObjectInstanceDescriptor::SetExternalTypeParameterSelectedOpti
 		{
 			const int32 InsertionIndex = ExternalTypeParameters[ParameterIndexInInstance].ParameterRangeValues.Num();
 			const int32 NumInsertedElements = RangeIndex + 1 - ExternalTypeParameters[ParameterIndexInInstance].ParameterRangeValues.Num();
-			FloatParameters[ParameterIndexInInstance].ParameterRangeValues.InsertDefaulted(InsertionIndex, NumInsertedElements);
+			ExternalTypeParameters[ParameterIndexInInstance].ParameterRangeValues.InsertDefaulted(InsertionIndex, NumInsertedElements);
 		}
  
 		check(ExternalTypeParameters[ParameterIndexInInstance].ParameterRangeValues.IsValidIndex(RangeIndex));
diff --git a/Engine/Plugins/Mutable/Source/CustomizableObject/Private/MuCO/CustomizableObjectInstanceUsage.cpp b/Engine/Plugins/Mutable/Source/CustomizableObject/Private/MuCO/CustomizableObjectInstanceUsage.cpp
--- a/Engine/Plugins/Mutable/Source/CustomizableObject/Private/MuCO/CustomizableObjectInstanceUsage.cpp
+++ b/Engine/Plugins/Mutable/Source/CustomizableObject/Private/MuCO/CustomizableObjectInstanceUsage.cpp
@@ -50,18 +50,12 @@ UCustomizableObjectInstanceUsage::UCustomizableObjectInstanceUsage()
 void UCustomizableObjectInstanceUsage::PostInitProperties()
 {
 	Super::PostInitProperties();
-	if (!HasAnyFlags(RF_ClassDefaultObject | RF_DefaultSubObject))
-	{
-		UCustomizableObjectSystemPrivate::RegisterCustomizableObjectInstanceUsage(this);
-	}
+	UCustomizableObjectSystemPrivate::RegisterCustomizableObjectInstanceUsage(*this);
 }
 
 void UCustomizableObjectInstanceUsage::BeginDestroy()
 {
-	if (!HasAnyFlags(RF_ClassDefaultObject | RF_DefaultSubObject))
-	{
-		UCustomizableObjectSystemPrivate::UnregisterCustomizableObjectInstanceUsage(this);
-	}
+	UCustomizableObjectSystemPrivate::UnregisterCustomizableObjectInstanceUsage(*this);
 	Super::BeginDestroy();
 }
 
@@ -91,19 +85,6 @@ UCustomizableObjectInstance* UCustomizableObjectInstanceUsage::GetCustomizableOb
 }
 
 
-int32 UCustomizableObjectInstanceUsage::GetComponentIndex() const
-{
-	if (UCustomizableSkeletalComponent* CustomizableSkeletalComponent = GetPrivate()->GetCustomizableSkeletalComponent())
-	{
-		return CustomizableSkeletalComponent->ComponentIndex;
-	}
-	else
-	{
-		return UsedComponentIndex;
-	}
-}
-
-
 void UCustomizableObjectInstanceUsage::SetComponentName(const FName& Name)
 {
 	if (UCustomizableSkeletalComponent* CustomizableSkeletalComponent = GetPrivate()->GetCustomizableSkeletalComponent())
@@ -224,15 +205,17 @@ void UCustomizableObjectInstanceUsage::AttachTo(USkeletalMeshComponent* Skeletal
 	if (IsValid(SkeletalMeshComponent))
 	{
 		UsedSkeletalMeshComponent = SkeletalMeshComponent;
+		GetPrivate()->bIsNetModeDedicatedServer = SkeletalMeshComponent->IsNetMode(NM_DedicatedServer);
 	}
 	else
 	{
 		UsedSkeletalMeshComponent = nullptr;
+		GetPrivate()->bIsNetModeDedicatedServer = false;
 	}
 
 	if (!GetSkipSetSkeletalMeshOnAttach())
 	{
-		GetPrivate()->bPendingSetSkeletalMesh = true;
+		GetPrivate()->bPendingSetReferenceSkeletalMesh = true;
 	}
 }
 
@@ -295,35 +278,6 @@ void UCustomizableObjectInstanceUsage::UpdateSkeletalMeshAsyncResult(FInstanceUp
 }
 
 
-#if WITH_EDITOR
-void UCustomizableObjectInstanceUsagePrivate::UpdateDistFromComponentToLevelEditorCamera(const FVector& CameraPosition)
-{
-	// We want instances in the editor to be generated
-	if (!GetWorld() || GetWorld()->WorldType != EWorldType::Editor)
-	{
-		return;
-	}
-
-	UCustomizableObjectInstance* CustomizableObjectInstance = GetPublic()->GetCustomizableObjectInstance();
-
-	if (CustomizableObjectInstance)
-	{
-		USkeletalMeshComponent* SkeletalMeshComponent = GetPublic()->GetAttachParent();
-		AActor* ParentActor = SkeletalMeshComponent ? SkeletalMeshComponent->GetAttachmentRootActor() : nullptr;
-		if (ParentActor && ParentActor->IsValidLowLevel())
-		{
-			// update distance to camera and set the instance as being used by a component
-			CustomizableObjectInstance->GetPrivate()->SetCOInstanceFlags(UsedByComponent);
-
-			float SquareDist = FVector::DistSquared(CameraPosition, ParentActor->GetActorLocation());
-			CustomizableObjectInstance->GetPrivate()->MinSquareDistFromComponentToPlayer = 
-				FMath::Min(SquareDist, CustomizableObjectInstance->GetPrivate()->MinSquareDistFromComponentToPlayer);
-		}
-	}
-}
-#endif
-
-
 UCustomizableObjectInstanceUsagePrivate* UCustomizableObjectInstanceUsage::GetPrivate()
 {
 	return Private;
@@ -333,77 +287,4 @@ UCustomizableObjectInstanceUsagePrivate* UCustomizableObjectInstanceUsage::GetPr
 const UCustomizableObjectInstanceUsagePrivate* UCustomizableObjectInstanceUsage::GetPrivate() const
 {
 	return Private;
-}
-
-
-void UCustomizableObjectInstanceUsagePrivate::UpdateDistFromComponentToPlayer(const AActor* ViewCenter, const bool bForceEvenIfNotBegunPlay)
-{
-	UCustomizableObjectInstance* CustomizableObjectInstance = GetPublic()->GetCustomizableObjectInstance();
-
-	if (CustomizableObjectInstance)
-	{
-		USkeletalMeshComponent* SkeletalMeshComponent = GetPublic()->GetAttachParent();
-		AActor* ParentActor = SkeletalMeshComponent ? SkeletalMeshComponent->GetAttachmentRootActor() : nullptr;
-		
-		CustomizableObjectInstance->SetIsPlayerOrNearIt(false);
-
-		if (ParentActor && ParentActor->IsValidLowLevel())
-		{
-			if (ParentActor->HasActorBegunPlay() || bForceEvenIfNotBegunPlay)
-			{
-				float SquareDist = FLT_MAX;
-
-				if (ViewCenter && ViewCenter->IsValidLowLevel())
-				{
-					APawn* Pawn = Cast<APawn>(ParentActor);
-					bool bIsPlayer = Pawn ? Pawn->IsPlayerControlled() : false;
-					CustomizableObjectInstance->SetIsPlayerOrNearIt(bIsPlayer);
-
-					if (bIsPlayer)
-					{
-						SquareDist = -0.01f; // Negative value to give the player character more priority than any other character
-					}
-					else
-					{
-						SquareDist = FVector::DistSquared(ViewCenter->GetActorLocation(), ParentActor->GetActorLocation());
-					}
-				}
-				else if (bForceEvenIfNotBegunPlay)
-				{
-					SquareDist = -0.01f; // This is a manual update before begin play and the creation of the pawn, so it should probably be high priority
-					CustomizableObjectInstance->GetPrivate()->LastMinSquareDistFromComponentToPlayer = FMath::Min(SquareDist, CustomizableObjectInstance->GetPrivate()->MinSquareDistFromComponentToPlayer);
-				}
-				else
-				{
-					SquareDist = 0.f; // This a mutable tick before begin play and the creation of the pawn, so it should have a definite and high priority but less than a manual update
-					CustomizableObjectInstance->GetPrivate()->LastMinSquareDistFromComponentToPlayer = FMath::Min(SquareDist, CustomizableObjectInstance->GetPrivate()->MinSquareDistFromComponentToPlayer);
-				}
-
-				CustomizableObjectInstance->GetPrivate()->MinSquareDistFromComponentToPlayer = FMath::Min(SquareDist, CustomizableObjectInstance->GetPrivate()->MinSquareDistFromComponentToPlayer);
-				CustomizableObjectInstance->SetIsBeingUsedByComponentInPlay(true);
-
-				if (CustomizableObjectInstance->GetPrivate()->MinSquareDistFromComponentToPlayer == SquareDist)
-				{
-					CustomizableObjectInstance->GetPrivate()->NearestToActor = GetPublic();
-					CustomizableObjectInstance->GetPrivate()->NearestToViewCenter = ViewCenter;
-				}
-			}
-		}
-	}
-}
-
-
-bool UCustomizableObjectInstanceUsagePrivate::IsNetMode(ENetMode InNetMode) const
-{
-	if (UCustomizableSkeletalComponent* CustomizableSkeletalComponent = GetCustomizableSkeletalComponent())
-	{
-		return CustomizableSkeletalComponent->IsNetMode(InNetMode);
-	}
-	else if (GetPublic()->UsedSkeletalMeshComponent.IsValid())
-	{
-		return GetPublic()->UsedSkeletalMeshComponent->IsNetMode(InNetMode);
-	}
-
-	return false;
-}
-
+}
\ No newline at end of file
diff --git a/Engine/Plugins/Mutable/Source/CustomizableObject/Private/MuCO/CustomizableObjectModule.cpp b/Engine/Plugins/Mutable/Source/CustomizableObject/Private/MuCO/CustomizableObjectModule.cpp
--- a/Engine/Plugins/Mutable/Source/CustomizableObject/Private/MuCO/CustomizableObjectModule.cpp
+++ b/Engine/Plugins/Mutable/Source/CustomizableObject/Private/MuCO/CustomizableObjectModule.cpp
@@ -357,7 +357,7 @@ UCustomizableObjectInstanceUsage* GetPlayerCustomizableObjectInstanceUsage(const
 		for (TObjectIterator<UCustomizableObjectInstanceUsage> CustomizableObjectInstanceUsage; CustomizableObjectInstanceUsage; ++CustomizableObjectInstanceUsage)
 		{
 #if WITH_EDITOR
-			if (IsValid(*CustomizableObjectInstanceUsage) && CustomizableObjectInstanceUsage->GetPrivate()->IsNetMode(NM_DedicatedServer))
+			if (IsValid(*CustomizableObjectInstanceUsage) && CustomizableObjectInstanceUsage->GetPrivate()->bIsNetModeDedicatedServer)
 			{
 				continue;
 			}
@@ -389,7 +389,7 @@ UCustomizableObjectInstanceUsage* GetPlayerCustomizableObjectInstanceUsage(const
 		for (TObjectIterator<UCustomizableObjectInstanceUsage> CustomizableObjectInstanceUsage; CustomizableObjectInstanceUsage; ++CustomizableObjectInstanceUsage)
 		{
 #if WITH_EDITOR
-			if (IsValid(*CustomizableObjectInstanceUsage) && CustomizableObjectInstanceUsage->GetPrivate()->IsNetMode(NM_DedicatedServer))
+			if (IsValid(*CustomizableObjectInstanceUsage) && CustomizableObjectInstanceUsage->GetPrivate()->bIsNetModeDedicatedServer)
 			{
 				continue;
 			}
diff --git a/Engine/Plugins/Mutable/Source/CustomizableObject/Private/MuCO/CustomizableObjectSkeletalMesh.cpp b/Engine/Plugins/Mutable/Source/CustomizableObject/Private/MuCO/CustomizableObjectSkeletalMesh.cpp
--- a/Engine/Plugins/Mutable/Source/CustomizableObject/Private/MuCO/CustomizableObjectSkeletalMesh.cpp
+++ b/Engine/Plugins/Mutable/Source/CustomizableObject/Private/MuCO/CustomizableObjectSkeletalMesh.cpp
@@ -28,11 +28,7 @@ void UCustomizableObjectSkeletalMesh::InitMutableStreamingData(const TSharedRef<
 	MeshIDs.Init({}, MAX_MESH_LOD_COUNT);
 	SurfaceIDs.SetNum(MAX_MESH_LOD_COUNT);
 
-	const int32 FirstLOD = InOperationData->bStreamMeshLODs ?
-		InOperationData->FirstLODAvailable[ComponentName] :
-		InOperationData->GetFirstRequestedLOD()[ComponentName];
-
-	for (int32 LODIndex = FirstLOD; LODIndex < LODCount; ++LODIndex)
+	for (int32 LODIndex = InOperationData->FirstLODAvailable[ComponentName]; LODIndex < LODCount; ++LODIndex)
 	{
 		const TSharedRef<const FInstanceUpdateData::FLOD>& LOD = InOperationData->InstanceUpdateData.LODs[InstanceUpdateFirstLOD + LODIndex];
 		MeshIDs[LODIndex] = LOD->MeshId;
diff --git a/Engine/Plugins/Mutable/Source/CustomizableObject/Private/MuCO/CustomizableObjectSystem.cpp b/Engine/Plugins/Mutable/Source/CustomizableObject/Private/MuCO/CustomizableObjectSystem.cpp
--- a/Engine/Plugins/Mutable/Source/CustomizableObject/Private/MuCO/CustomizableObjectSystem.cpp
+++ b/Engine/Plugins/Mutable/Source/CustomizableObject/Private/MuCO/CustomizableObjectSystem.cpp
@@ -6,7 +6,6 @@
 #include "Engine/SkeletalMesh.h"
 #include "Engine/SkinnedAsset.h"
 #include "Engine/SkeletalMeshLODSettings.h"
-#include "MuCO/CustomizableInstanceLODManagement.h"
 #include "MuCO/CustomizableObjectInstancePrivate.h"
 #include "MuCO/CustomizableObjectPrivate.h"
 #include "MuCO/CustomizableObjectInstanceUsagePrivate.h"
@@ -30,6 +29,8 @@
 #include "MuR/Parameters.h"
 #include "MuR/System.h"
 #include "MuR/PassthroughObject.h"
+#include "GameFramework/Pawn.h"
+#include "GameFramework/PlayerController.h"
 
 #if WITH_EDITOR
 #include "Editor.h"
@@ -38,7 +39,6 @@
 #include "Engine/World.h"
 #include "Interfaces/ITargetPlatform.h"
 #include "MuT/UnrealPixelFormatOverride.h"
-#include "MuR/Parameters.h"
 #include "IMessageLogListing.h"
 #else
 #include "Engine/Engine.h"
@@ -150,11 +150,6 @@ TAutoConsoleVariable<bool> CVarMutableLogObjectMemoryOnUpdate(
 	ECVF_Scalability);
 #endif
 
-TAutoConsoleVariable<bool> CVarEnableReleaseMeshResources(
-	TEXT("mutable.EnableReleaseMeshResources"),
-	true,
-	TEXT("Allow releasing resources when discarding instances."));
-
 TAutoConsoleVariable<bool> CVarFixLowPriorityTasksOverlap(
 	TEXT("mutable.rollback.FixLowPriorityTasksOverlap"),
 	true,
@@ -172,26 +167,6 @@ TAutoConsoleVariable<bool> CVarReuseLiveInstanceCacheBetweenUpdates (
 	ECVF_Default
 );
 
-int32 UCustomizableObjectSystemPrivate::SkeletalMeshMinLodQualityLevel = -1;
-
-
-static void CVarMutableSinkFunction()
-{
-	if (UCustomizableObjectSystem::IsCreated())
-	{
-		UCustomizableObjectSystemPrivate* PrivateSystem = UCustomizableObjectSystem::GetInstance()->GetPrivate();
-
-		// Store the quality level set in the scalability settings so we can later determine what the MinLOD should be used by mutable.
-		// Does not seem to be triggered when changing the visibility quality directly.
-		static const IConsoleVariable* CVarSkeletalMeshMinLodQualityLevelCVarName = IConsoleManager::Get().FindConsoleVariable(TEXT("r.SkeletalMesh.MinLodQualityLevel"));
-		PrivateSystem->SkeletalMeshMinLodQualityLevel = CVarSkeletalMeshMinLodQualityLevelCVarName ? CVarSkeletalMeshMinLodQualityLevelCVarName->GetInt() : INDEX_NONE;
-	}
-}
-
-
-static FAutoConsoleVariableSink CVarMutableSink(FConsoleCommandDelegate::CreateStatic(&CVarMutableSinkFunction));
-
-
 /** How often update the on screen warnings (seconds). */
 constexpr float OnScreenWarningsTickerTime = 5.0f;
 
@@ -258,8 +233,6 @@ FUpdateContextPrivate::FUpdateContextPrivate(UCustomizableObjectInstance& InInst
 
 	if (const UModelResources* ModelResources = InObject->GetPrivate()->GetModelResources())
 	{
-		CapturedDescriptor.bStreamingEnabled = IsStreamingEnabled(*InObject, Descriptor.State);
-
 		ComponentNames = ModelResources->ComponentNamesPerObjectComponent;
 
 		for (TTuple<int32, uint8> Pair : ModelResources->FirstLODs)
@@ -284,7 +257,7 @@ FUpdateContextPrivate::FUpdateContextPrivate(UCustomizableObjectInstance& InInst
 	check(MutableSystem);
 	
 	bProgressiveMipStreamingEnabled = System->IsProgressiveMipStreamingEnabled();
-	bIsOnlyGenerateRequestedLODsEnabled = System->IsOnlyGenerateRequestedLODsEnabled();
+
 #if WITH_EDITOR
 	PixelFormatOverride = SystemPrivate->ImageFormatOverrideFunc;
 #endif
@@ -415,13 +388,8 @@ bool TickWarnings(float DeltaTime)
 	
 	TSet<const UCustomizableObject*> Objects;
 	
-	for (TObjectIterator<UCustomizableObjectInstanceUsage> CustomizableObjectInstanceUsage; CustomizableObjectInstanceUsage; ++CustomizableObjectInstanceUsage)
+	for (UCustomizableObjectInstanceUsage* CustomizableObjectInstanceUsage : UCustomizableObjectSystemPrivate::GetCustomizableObjectInstanceUsages())
 	{
-		if (!IsValid(*CustomizableObjectInstanceUsage) || CustomizableObjectInstanceUsage->IsTemplate())
-		{
-			continue;
-		}
-		
 		const UCustomizableObjectInstance* Instance = CustomizableObjectInstanceUsage->GetCustomizableObjectInstance();
 		if (!Instance)
 		{
@@ -537,33 +505,6 @@ FUpdateContextPrivate::FUpdateContextPrivate(UCustomizableObjectInstance& InInst
 }
 
 
-void FUpdateContextPrivate::SetMinLOD(const TMap<FName, uint8>& MinLOD)
-{
-	CapturedDescriptor.MinLOD = MinLOD;
-	CapturedDescriptorHash.MinLODs = MinLOD;
-}
-
-
-void FUpdateContextPrivate::SetQualitySettingMinLODs(const TMap<FName, uint8>& QualitySettingsMinLODs)
-{
-	CapturedDescriptor.QualitySettingMinLODs = QualitySettingsMinLODs;
-	CapturedDescriptorHash.QualitySettingMinLODs = QualitySettingsMinLODs;
-}
-
-
-const TMap<FName, uint8>& FUpdateContextPrivate::GetFirstRequestedLOD() const
-{
-	return CapturedDescriptor.GetFirstRequestedLOD();
-}
-
-
-void FUpdateContextPrivate::SetFirstRequestedLOD(const TMap<FName, uint8>& FirstRequestedLOD)
-{
-	CapturedDescriptor.SetFirstRequestedLOD(FirstRequestedLOD);
-	CapturedDescriptorHash.FirstRequestedLOD = FirstRequestedLOD;
-}
-
-
 const FCustomizableObjectInstanceDescriptor& FUpdateContextPrivate::GetCapturedDescriptor() const
 {
 	return CapturedDescriptor;
@@ -600,138 +541,6 @@ FComponentDescriptor* FUpdateContextPrivate::GetComponentDescriptor(UE::Mutable:
 }
 
 
-FMutablePendingInstanceUpdate::FMutablePendingInstanceUpdate(const TSharedRef<FUpdateContextPrivate>& InContext) :
-	Context(InContext)
-{
-}
-
-
-bool FMutablePendingInstanceUpdate::operator==(const FMutablePendingInstanceUpdate& Other) const
-{
-	return Context->Instance.HasSameIndexAndSerialNumber(Other.Context->Instance);
-}
-
-
-bool FMutablePendingInstanceUpdate::operator<(const FMutablePendingInstanceUpdate& Other) const
-{
-	if (Context->PriorityType < Other.Context->PriorityType)
-	{
-		return true;
-	}
-	else if (Context->PriorityType > Other.Context->PriorityType)
-	{
-		return false;
-	}
-	else
-	{
-		return Context->StartQueueTime < Other.Context->StartQueueTime;
-	}
-}
-
-
-uint32 GetTypeHash(const FMutablePendingInstanceUpdate& Update)
-{
-	return GetTypeHash(Update.Context->Instance.GetWeakPtrTypeHash());
-}
-
-
-TWeakObjectPtr<const UCustomizableObjectInstance> FPendingInstanceUpdateKeyFuncs::GetSetKey(const FMutablePendingInstanceUpdate& PendingUpdate)
-{
-	return PendingUpdate.Context->Instance;
-}
-
-
-bool FPendingInstanceUpdateKeyFuncs::Matches(const TWeakObjectPtr<const UCustomizableObjectInstance>& A, const TWeakObjectPtr<const UCustomizableObjectInstance>& B)
-{
-	return A.HasSameIndexAndSerialNumber(B);
-}
-
-
-uint32 FPendingInstanceUpdateKeyFuncs::GetKeyHash(const TWeakObjectPtr<const UCustomizableObjectInstance>& Identifier)
-{
-	return GetTypeHash(Identifier.GetWeakPtrTypeHash());
-}
-
-
-int32 FMutablePendingInstanceWork::Num() const
-{
-	return PendingInstanceUpdates.Num() + PendingInstanceDiscards.Num();
-}
-
-
-void FMutablePendingInstanceWork::AddUpdate(const FMutablePendingInstanceUpdate& UpdateToAdd)
-{
-	UpdateToAdd.Context->StartQueueTime = FPlatformTime::Seconds();
-	
-	if (const FMutablePendingInstanceUpdate* ExistingUpdate = PendingInstanceUpdates.Find(UpdateToAdd.Context->Instance))
-	{
-		ExistingUpdate->Context->UpdateResult = EUpdateResult::ErrorReplaced;
-		FinishUpdateGlobal(ExistingUpdate->Context);
-
-		const FMutablePendingInstanceUpdate TaskToEnqueue = UpdateToAdd;
-		TaskToEnqueue.Context->PriorityType = FMath::Min(ExistingUpdate->Context->PriorityType, UpdateToAdd.Context->PriorityType);
-		TaskToEnqueue.Context->StartQueueTime = FMath::Min(ExistingUpdate->Context->StartQueueTime, UpdateToAdd.Context->StartQueueTime);
-		
-		RemoveUpdate(ExistingUpdate->Context->Instance);
-		PendingInstanceUpdates.Add(TaskToEnqueue);
-	}
-	else
-	{
-		PendingInstanceUpdates.Add(UpdateToAdd);
-	}
-
-	if (const FMutablePendingInstanceDiscard* ExistingDiscard = PendingInstanceDiscards.Find(UpdateToAdd.Context->Instance))
-	{
-		UpdateToAdd.Context->UpdateResult = EUpdateResult::ErrorReplaced;
-		FinishUpdateGlobal(UpdateToAdd.Context);
-
-		PendingInstanceDiscards.Remove(ExistingDiscard->CustomizableObjectInstance);
-	}
-}
-
-
-void FMutablePendingInstanceWork::RemoveUpdate(const TWeakObjectPtr<UCustomizableObjectInstance>& Instance)
-{
-	if (const FMutablePendingInstanceUpdate* Update = PendingInstanceUpdates.Find(Instance))
-	{
-		Update->Context->QueueTime = FPlatformTime::Seconds() - Update->Context->StartQueueTime;
-		PendingInstanceUpdates.Remove(Instance);
-	}	
-}
-
-#if WITH_EDITOR
-void FMutablePendingInstanceWork::RemoveUpdatesForObject(const UCustomizableObject* InObject)
-{
-	check(InObject);
-	for (auto Iterator = PendingInstanceUpdates.CreateIterator(); Iterator; ++Iterator)
-	{
-		if (Iterator->Context->Instance.IsValid() && Iterator->Context->Instance->GetCustomizableObject() == InObject)
-		{
-			Iterator.RemoveCurrent();
-		}
-	}
-}
-#endif
-
-const FMutablePendingInstanceUpdate* FMutablePendingInstanceWork::GetUpdate(const TWeakObjectPtr<const UCustomizableObjectInstance>& Instance) const
-{
-	return PendingInstanceUpdates.Find(Instance);
-}
-
-
-void FMutablePendingInstanceWork::AddDiscard(const FMutablePendingInstanceDiscard& TaskToEnqueue)
-{
-	if (const FMutablePendingInstanceUpdate* ExistingUpdate = PendingInstanceUpdates.Find(TaskToEnqueue.CustomizableObjectInstance.Get()))
-	{
-		ExistingUpdate->Context->UpdateResult = EUpdateResult::ErrorDiscarded;
-		FinishUpdateGlobal(ExistingUpdate->Context);
-		RemoveUpdate(ExistingUpdate->Context->Instance);
-	}
-
-	PendingInstanceDiscards.Add(TaskToEnqueue);
-}
-
-
 UCustomizableObjectSystem* UCustomizableObjectSystem::GetInstance()
 {
 	if (!UCustomizableObjectSystemPrivate::SSystem)
@@ -767,18 +576,6 @@ bool UCustomizableObjectSystem::IsUpdateResultValid(const EUpdateResult UpdateRe
 }
 
 
-UCustomizableInstanceLODManagementBase* UCustomizableObjectSystem::GetInstanceLODManagement() const
-{
-	return GetPrivate()->CurrentInstanceLODManagement.Get();
-}
-
-
-void UCustomizableObjectSystem::SetInstanceLODManagement(UCustomizableInstanceLODManagementBase* NewInstanceLODManagement)
-{
-	GetPrivate()->CurrentInstanceLODManagement = NewInstanceLODManagement ? NewInstanceLODManagement : ToRawPtr(GetPrivate()->DefaultInstanceLODManagement);
-}
-
-
 FString UCustomizableObjectSystem::GetPluginVersion() const
 {
 	// Bridge the call from the module. This implementation is available from blueprint.
@@ -817,8 +614,6 @@ void UCustomizableObjectSystem::InitSystem()
 	Private = NewObject<UCustomizableObjectSystemPrivate>(this, FName("Private"));
 	check(Private != nullptr);
 
-	Private->bReplaceDiscardedWithReferenceMesh = false;
-
 	Private->CurrentMutableOperation = nullptr;
 	Private->CurrentInstanceBeingUpdated = nullptr;
 
@@ -834,17 +629,10 @@ void UCustomizableObjectSystem::InitSystem()
 	check(Private->Streamer != nullptr);
 	Private->MutableSystem->SetStreamingInterface(Private->Streamer);
 
-	GetPrivate()->DefaultInstanceLODManagement = NewObject<UCustomizableInstanceLODManagement>();
-	check(GetPrivate()->DefaultInstanceLODManagement != nullptr);
-	GetPrivate()->CurrentInstanceLODManagement = GetPrivate()->DefaultInstanceLODManagement;
-
 	// This CVar is constant for the lifespan of the program. Read its value once. 
 	const IConsoleVariable* CVarSupport16BitBoneIndex = IConsoleManager::Get().FindConsoleVariable(TEXT("r.GPUSkin.Support16BitBoneIndex"));
 	Private->bSupport16BitBoneIndex = CVarSupport16BitBoneIndex ? CVarSupport16BitBoneIndex->GetBool() : false;
 
-	// Read non-constant CVars and do work if required.
-	CVarMutableSinkFunction();
-
 	Private->OnMutableEnabledChanged();
 }
 
@@ -867,7 +655,7 @@ void UCustomizableObjectSystem::BeginDestroy()
 		FTSTicker::RemoveTicker(Private->TickWarningsDelegateHandle);
 #endif // !UE_SERVER
 		
-		// Complete pending taskgraph tasks
+		// Complete pending task graph tasks
 		Private->MutableTaskGraph.AllowLaunchingMutableTaskLowPriority(false, false);
 		check(Private->Streamer);
 		Private->MutableTaskGraph.AddMutableThreadTask(TEXT("EndStream"), [Streamer = Private->Streamer]()
@@ -912,7 +700,6 @@ static FAutoConsoleVariableRef CVarMutableForceStreamMeshLODs(
 	TEXT("Experimental - true/false - If true, and bStreamMeshLODs is enabled, all COs will stream mesh LODs. "),
 	ECVF_Default);
 
-
 bool bStreamMeshLODs = true;
 
 static FAutoConsoleVariableRef CVarMutableStreamMeshLODsEnabled(
@@ -944,12 +731,25 @@ static TAutoConsoleVariable<bool> CVarEnableMutableReuseInstanceTextures(
 	TEXT("If set to true set in the corresponding setting in the current Mutable state, Mutable can reuse instance UTextures (only uncompressed and not streaming, so set the options in the state) and their resources between updates when they are modified. If geometry or state is changed they cannot be reused."),
 	ECVF_Default);
 
+bool UCustomizableObjectSystemPrivate::bGenerateInstancesWithinRange = false;
+
+static FAutoConsoleVariableRef CVarGenerateInstancesWithinRange(
+	TEXT("mutable.GenerateInstancesWithinRange"),
+	UCustomizableObjectSystemPrivate::bGenerateInstancesWithinRange,
+	TEXT("If true, only instances within a predefined distance to the view centers will be generated"),
+	ECVF_Default);
+
+TAutoConsoleVariable<int32> CVarMaxNumInstancesToDiscardPerTick(
+	TEXT("mutable.MaxNumInstancesToDiscardPerTick"),
+	30,
+	TEXT("The maximum number of stale instances that will be discarded per tick by Mutable."),
+	ECVF_Scalability);
 
-int32 UCustomizableObjectSystemPrivate::EnableOnlyGenerateRequestedLODs = 1;
+int32 UCustomizableObjectSystemPrivate::MutableMeshesLODBias = 0;
 
-static FAutoConsoleVariableRef CVarEnableOnlyGenerateRequestedLODs(
-	TEXT("mutable.EnableOnlyGenerateRequestedLODs"), UCustomizableObjectSystemPrivate::EnableOnlyGenerateRequestedLODs,
-	TEXT("If 1 or greater, Only the RequestedLOD will be generated. If 0, all LODs will be build."),
+static FAutoConsoleVariableRef CVarMutableMeshesLODBias(
+	TEXT("mutable.MeshLODBias"), UCustomizableObjectSystemPrivate::MutableMeshesLODBias,
+	TEXT("LOD Bias applied to Mutable-generated meshes when Mesh LOD streaming is turned off."),
 	ECVF_Default);
 
 int32 UCustomizableObjectSystemPrivate::MaxTextureSizeToGenerate = 0;
@@ -974,7 +774,6 @@ void FinishUpdateGlobal(const TSharedRef<FUpdateContextPrivate>& Context)
 
 	UCustomizableObjectInstance* Instance = Context->Instance.Get();
 	UCustomizableObject* Object = IsValid(Instance) ? Instance->GetCustomizableObject() : nullptr;
-	UModelResources* ModelResources = IsValid(Object) ? Object->GetPrivate()->GetModelResources() : nullptr;
 
 	UCustomizableObjectSystem* System = UCustomizableObjectSystem::GetInstance();
 	UCustomizableObjectSystemPrivate* SystemPrivate = System ? System->GetPrivate() : nullptr;
@@ -998,8 +797,6 @@ void FinishUpdateGlobal(const TSharedRef<FUpdateContextPrivate>& Context)
 			
 			PrivateInstance->CommittedDescriptor = Context->MoveCommittedDescriptor();
 			PrivateInstance->CommittedDescriptorHash = Context->GetCapturedDescriptorHash();
-
-			PrivateInstance->bAutomaticUpdateRequired = false;
 			
 			// Delegates must be called only after updating the Instance flags.
 			Instance->UpdatedDelegate.Broadcast(Instance);
@@ -1012,6 +809,9 @@ void FinishUpdateGlobal(const TSharedRef<FUpdateContextPrivate>& Context)
 		case EUpdateResult::ErrorDiscarded:
 			break; // Status will be updated once the discard is performed.
 
+		case EUpdateResult::ErrorUncompiled:
+			break; // Object not compiled.
+
 		case EUpdateResult::Error: 
 		case EUpdateResult::Error16BitBoneIndex:
 			{
@@ -1060,7 +860,7 @@ void FinishUpdateGlobal(const TSharedRef<FUpdateContextPrivate>& Context)
 				}
 
 #if WITH_EDITOR
-				if (InstanceUsage->GetPrivate()->IsNetMode(NM_DedicatedServer))
+				if (InstanceUsage->GetPrivate()->bIsNetModeDedicatedServer)
 				{
 					continue;
 				}
@@ -1076,16 +876,10 @@ void FinishUpdateGlobal(const TSharedRef<FUpdateContextPrivate>& Context)
 		}
 		else
 		{
-			for (TObjectIterator<UCustomizableObjectInstanceUsage> It; It; ++It) // Since iterating objects is expensive, for now CustomizableObjectInstanceUsage does not have a FinishUpdate function.
+			for (UCustomizableObjectInstanceUsage* InstanceUsage : UCustomizableObjectSystemPrivate::GetCustomizableObjectInstanceUsages())
 			{
-				UCustomizableObjectInstanceUsage* InstanceUsage = *It;
-				if (!IsValid(InstanceUsage))
-				{
-					continue;
-				}
-
 #if WITH_EDITOR
-				if (It->GetPrivate()->IsNetMode(NM_DedicatedServer))
+				if (InstanceUsage->GetPrivate()->bIsNetModeDedicatedServer)
 				{
 					continue;
 				}
@@ -1131,11 +925,11 @@ void FinishUpdateGlobal(const TSharedRef<FUpdateContextPrivate>& Context)
 		Context->UpdateTime = FPlatformTime::Seconds() - Context->StartUpdateTime;
 	}
 	
-	const FName ObjectName = GetFNameSafe(Object);
-	const FName InstanceName = GetFNameSafe(Instance);
-
 	if (SystemPrivate)
 	{
+		const FName ObjectName = GetFNameSafe(Object);
+		const FName InstanceName = GetFNameSafe(Instance);
+
 		FString Message = FString::Printf(TEXT("Finished Update Skeletal Mesh Async. CustomizableObject=%s Instance=%s, Frame=%d  QueueTime=%f, UpdateTime=%f"), *ObjectName.ToString(), *InstanceName.ToString(), GFrameNumber, Context->QueueTime, Context->UpdateTime);
 		SystemPrivate->LogUpdateMessage(Message, ELogVerbosity::Verbose, Instance);
 	}
@@ -1280,17 +1074,11 @@ void UpdateSkeletalMesh(const TSharedRef<FUpdateContextPrivate>& Context)
 	
 	UCustomizableInstancePrivate* CustomizableObjectInstancePrivateData = CustomizableObjectInstance->GetPrivate();
 	check(CustomizableObjectInstancePrivateData != nullptr);
-	for (TObjectIterator<UCustomizableObjectInstanceUsage> It; It; ++It)
-	{
-		UCustomizableObjectInstanceUsage* Usage = *It;
 
-		if (!IsValid(Usage))
-		{
-			continue;
-		}
-		
+	for (UCustomizableObjectInstanceUsage* Usage : UCustomizableObjectSystemPrivate::GetCustomizableObjectInstanceUsages())
+	{
 #if WITH_EDITOR
-		if (Usage->GetPrivate()->IsNetMode(NM_DedicatedServer))
+		if (Usage->GetPrivate()->bIsNetModeDedicatedServer)
 		{
 			continue;
 		}
@@ -1311,7 +1099,7 @@ void UpdateSkeletalMesh(const TSharedRef<FUpdateContextPrivate>& Context)
 
 		bool bAttachedParentUpdated = false;
 
-		Usage->GetPrivate()->bPendingSetSkeletalMesh = false;
+		Usage->GetPrivate()->bPendingSetReferenceSkeletalMesh = false;
 
 		USkeletalMesh* SkeletalMesh = CustomizableObjectInstance->GetComponentMeshSkeletalMesh(Usage->GetComponentName());
 		if (SkeletalMesh != Parent->GetSkeletalMeshAsset())
@@ -1402,100 +1190,9 @@ void UpdateSkeletalMesh(const TSharedRef<FUpdateContextPrivate>& Context)
 }
 
 
-bool UCustomizableObjectSystemPrivate::IsReplaceDiscardedWithReferenceMeshEnabled() const
-{
-	return bReplaceDiscardedWithReferenceMesh;
-}
-
-
-void UCustomizableObjectSystemPrivate::SetReplaceDiscardedWithReferenceMeshEnabled(bool bIsEnabled)
-{
-	bReplaceDiscardedWithReferenceMesh = bIsEnabled;
-}
-
-
-EUpdateRequired UCustomizableObjectSystemPrivate::IsUpdateRequired(const UCustomizableObjectInstance& Instance, bool bOnlyUpdateIfNotGenerated, bool bOnlyUpdateIfLODs, bool bIgnoreCloseDist) const
-{
-	UCustomizableObjectSystem* System = UCustomizableObjectSystem::GetInstance();
-	const UCustomizableInstancePrivate* const Private = Instance.GetPrivate();
-	
-	if (!Instance.GetPrivate()->CanUpdateInstance())
-	{
-		return EUpdateRequired::NoUpdate;
-	}
-
-	const bool bIsGenerated = Private->SkeletalMeshStatus != ESkeletalMeshStatus::NotGenerated;
-
-	const bool bDiscardByDistance = Private->LastMinSquareDistFromComponentToPlayer > FMath::Square(System->GetInstanceLODManagement()->GetOnlyUpdateCloseCustomizableObjectsDist());
-	const bool bLODManagementDiscard = System->GetInstanceLODManagement()->IsOnlyUpdateCloseCustomizableObjectsEnabled() &&
-			bDiscardByDistance &&
-			!bIgnoreCloseDist;
-	
-	if (Private->HasCOInstanceFlags(DiscardedByNumInstancesLimit) ||
-		bLODManagementDiscard)
-	{
-		if (bIsGenerated)
-		{
-			return EUpdateRequired::Discard;		
-		}
-		else
-		{
-			return EUpdateRequired::NoUpdate;
-		}
-	}
-
-	const bool bShouldUpdateLODs = Private->HasCOInstanceFlags(PendingLODsUpdate);
-
-	const bool bNoUpdateLODs = bOnlyUpdateIfLODs && !bShouldUpdateLODs;
-	const bool bNoInitialUpdate = bOnlyUpdateIfNotGenerated && bIsGenerated;
-
-	if (bNoUpdateLODs &&
-		bNoInitialUpdate &&
-		!Private->bAutomaticUpdateRequired)
-	{
-		return EUpdateRequired::NoUpdate;
-	}
-
-	return EUpdateRequired::Update;
-}
-
-
-EQueuePriorityType UCustomizableObjectSystemPrivate::GetUpdatePriority(const UCustomizableObjectInstance& Instance,	bool bForceHighPriority) const
+bool UCustomizableObjectSystemPrivate::IsDiscardedByDistance(const UCustomizableInstancePrivate* InstancePrivate) const
 {
-	const UCustomizableInstancePrivate* InstancePrivate = Instance.GetPrivate();
-		
-	const bool bNotGenerated = InstancePrivate->SkeletalMeshStatus == ESkeletalMeshStatus::NotGenerated;
-	const bool bShouldUpdateLODs = InstancePrivate->HasCOInstanceFlags(PendingLODsUpdate);
-	const bool bIsDowngradeLODUpdate = InstancePrivate->HasCOInstanceFlags(PendingLODsDowngrade);
-	const bool bIsPlayerOrNearIt = InstancePrivate->HasCOInstanceFlags(UsedByPlayerOrNearIt);
-
-	EQueuePriorityType Priority = EQueuePriorityType::Low;
-	if (bForceHighPriority)
-	{
-		Priority = EQueuePriorityType::High;
-	}
-	else if (bNotGenerated || !Instance.HasAnySkeletalMesh())
-	{
-		Priority = EQueuePriorityType::Med;
-	}
-	else if (bShouldUpdateLODs && bIsDowngradeLODUpdate)
-	{
-		Priority = EQueuePriorityType::Med_Low;
-	}
-	else if (bIsPlayerOrNearIt && bShouldUpdateLODs && !bIsDowngradeLODUpdate)
-	{
-		Priority = EQueuePriorityType::High;
-	}
-	else if (bShouldUpdateLODs && !bIsDowngradeLODUpdate)
-	{
-		Priority = EQueuePriorityType::Med;
-	}
-	else if (bIsPlayerOrNearIt)
-	{
-		Priority = EQueuePriorityType::High;
-	}
-
-	return Priority;
+	return bGenerateInstancesWithinRange && InstancePrivate->MinSquareDistFromComponentToPlayer > GenerationRangeSquare;
 }
 
 
@@ -1520,146 +1217,175 @@ void UCustomizableObjectSystemPrivate::EnqueueUpdateSkeletalMesh(const TSharedRe
 
 	check(Object);
 	check(Instance);
+	
+	UCustomizableInstancePrivate* InstancePrivate = Instance->GetPrivate();
 		
 	if (!bIsMutableEnabled)
 	{
 		// Mutable is disabled. Set the reference SkeletalMesh and finish the update with success to avoid breaking too many things.
 		Context->UpdateResult = EUpdateResult::Success;
-		Instance->GetPrivate()->ForceSetReferenceSkeletalMesh();
+		InstancePrivate->ForceSetReferenceSkeletalMesh();
+		FinishUpdateGlobal(Context);
+		return;
+	}
+
+#if WITH_EDITOR
+	if (!Object->IsCompiled() &&
+		Object->GetPrivate()->CompilationResult != ECompilationResultPrivate::Errors) // Avoid constantly retry failed compilations.
+	{
+		if (ICustomizableObjectEditorModule* EditorModule = ICustomizableObjectEditorModule::Get())
+		{
+			FCompileParams CompileParams;
+			CompileParams.bSkipIfCompiled = true;
+
+			EditorModule->CompileCustomizableObject(*Object, &CompileParams, true, false);
+		}
+
+		Context->UpdateResult = EUpdateResult::ErrorUncompiled;
 		FinishUpdateGlobal(Context);
 		return;
 	}
+#endif
 	
-	if (!Instance->GetPrivate()->CanUpdateInstance())
+	if (!InstancePrivate->CanUpdateInstance())
 	{
 		Context->UpdateResult = EUpdateResult::Error;
 		FinishUpdateGlobal(Context);
 		return;
 	}
 	
-	const EUpdateRequired UpdateRequired = IsUpdateRequired(*Instance, Context->bOnlyUpdateIfNotGenerated, false, Context->bIgnoreCloseDist);
-	switch (UpdateRequired)
+	if (IsDiscardedByDistance(InstancePrivate))
 	{
-	case EUpdateRequired::NoUpdate:
-	{	
-		Context->UpdateResult = EUpdateResult::Error;
+		Context->UpdateResult = EUpdateResult::ErrorDiscarded;
 		FinishUpdateGlobal(Context);
-		break;
-	}		
-	case EUpdateRequired::Update:
+		return;
+	}
+
+	if (!Context->bForce)
 	{
-		if (!Context->bForce)
+		if (CurrentMutableOperation &&
+			Instance == CurrentMutableOperation->Instance &&
+			Context->GetCapturedDescriptorHash() == CurrentMutableOperation->GetCapturedDescriptorHash())
 		{
-			if (const FMutablePendingInstanceUpdate* QueueElem = MutablePendingInstanceWork.GetUpdate(Instance))
-			{
-				if (Context->GetCapturedDescriptorHash().IsSubset(QueueElem->Context->GetCapturedDescriptorHash()))
-				{
-					Context->bOptimizedUpdate = true;
-					Context->UpdateResult = EUpdateResult::ErrorOptimized;
-					FinishUpdateGlobal(Context);			
-					return; // The requested update is equal to the last enqueued update.
-				}
-			}	
+			Context->bOptimizedUpdate = true;
+			Context->UpdateResult = EUpdateResult::ErrorOptimized;
+			FinishUpdateGlobal(Context);
+			return; // The requested update is equal to the running update.
+		}
 
-			if (CurrentMutableOperation &&
-				Instance == CurrentMutableOperation->Instance &&
-				Context->GetCapturedDescriptorHash().IsSubset(CurrentMutableOperation->GetCapturedDescriptorHash()))
+		if (Context->GetCapturedDescriptorHash() == InstancePrivate->CommittedDescriptorHash &&
+			!(CurrentMutableOperation &&
+				Instance == CurrentMutableOperation->Instance)) // This condition is necessary because even if the descriptor is a subset, it will be replaced by the CurrentMutableOperation
+		{
+			if (CVarEnableUpdateOptimization.GetValueOnGameThread()) // TODO Remove hotfix: UE-218957 
 			{
 				Context->bOptimizedUpdate = true;
-				Context->UpdateResult = EUpdateResult::ErrorOptimized;
+
+				// The user may have changed the AttachParent and we need to re-customize it.
+				// In case nothing need to be re-customized, the update will be considered ErrorOptimized.
+				UpdateSkeletalMesh(Context);
+				Context->UpdateResult = Context->AttachedParentUpdated.IsEmpty() ? EUpdateResult::ErrorOptimized : EUpdateResult::Success;
+
 				FinishUpdateGlobal(Context);
-				return; // The requested update is equal to the running update.
+				return;
 			}
-		
-			if (Context->GetCapturedDescriptorHash().IsSubset(Instance->GetPrivate()->CommittedDescriptorHash) &&
-				!(CurrentMutableOperation &&
-				Instance == CurrentMutableOperation->Instance)) // This condition is necessary because even if the descriptor is a subset, it will be replaced by the CurrentMutableOperation
+			else
 			{
-				if (CVarEnableUpdateOptimization.GetValueOnGameThread()) // TODO Remove hotfix: UE-218957 
-				{
-					Context->bOptimizedUpdate = true;
+				Context->bOptimizedUpdate = false;
 
-					// The user may have changed the AttachParent and we need to re-customize it.
-					// In case nothing need to be re-customized, the update will be considered ErrorOptimized.
-					UpdateSkeletalMesh(Context);
-					Context->UpdateResult = Context->AttachedParentUpdated.IsEmpty() ? EUpdateResult::ErrorOptimized : EUpdateResult::Success;
+				// The user may have changed the AttachParent and we need to re-customize it.
+				// In case nothing need to be re-customized, the update will be considered ErrorOptimized.
+				UpdateSkeletalMesh(Context);
+				Context->UpdateResult = EUpdateResult::Success;
 
-					FinishUpdateGlobal(Context);
-				}
-				else 
-				{
-					Context->bOptimizedUpdate = false;
-
-					// The user may have changed the AttachParent and we need to re-customize it.
-					// In case nothing need to be re-customized, the update will be considered ErrorOptimized.
-					UpdateSkeletalMesh(Context);
-					Context->UpdateResult = EUpdateResult::Success;
-
-					FinishUpdateGlobal(Context);
-				}
+				FinishUpdateGlobal(Context);
+				return;
 			}
 		}
-
-		if (CVarDescriptorDebugPrint->GetBool())
-		{
-			FString String = TEXT("DESCRIPTOR DEBUG PRINT\n");
-			String += "================================\n";				
-			String += FString::Printf(TEXT("=== DESCRIPTOR HASH ===\n%s\n"), *Context->GetCapturedDescriptorHash().ToString());
-			String += FString::Printf(TEXT("=== DESCRIPTOR ===\n%s"), *Instance->GetPrivate()->GetDescriptor().ToString());
-			String += "================================";
-			
-			UE_LOG(LogMutable, Log, TEXT("%s"), *String);
-		}
-
-		const FMutablePendingInstanceUpdate InstanceUpdate(Context);
-		MutablePendingInstanceWork.AddUpdate(InstanceUpdate);
-
-		break;
 	}
 
-	case EUpdateRequired::Discard:
+	if (CVarDescriptorDebugPrint->GetBool())
 	{
-		InitDiscardResourcesSkeletalMesh(Instance);
-
-		Context->UpdateResult = EUpdateResult::ErrorDiscarded;
-		FinishUpdateGlobal(Context);
-		break;
-	}
+		FString String = TEXT("DESCRIPTOR DEBUG PRINT\n");
+		String += "================================\n";
+		String += FString::Printf(TEXT("=== DESCRIPTOR HASH ===\n%s\n"), *Context->GetCapturedDescriptorHash().ToString());
+		String += FString::Printf(TEXT("=== DESCRIPTOR ===\n%s"), *Instance->GetPrivate()->GetDescriptor().ToString());
+		String += "================================";
 
-	default:
-		unimplemented();
+		UE_LOG(LogMutable, Log, TEXT("%s"), *String);
 	}
-}
 
+	Context->StartQueueTime = FPlatformTime::Seconds();
 
-void UCustomizableObjectSystemPrivate::InitDiscardResourcesSkeletalMesh(UCustomizableObjectInstance* InCustomizableObjectInstance)
-{
-	check(IsInGameThread());
+	auto EnqueueUpdate = [&]()
+		{
+			InstancePrivate->ClearCOInstanceFlags(ECOInstanceFlags::HasPendingHighPriorityUpdate | ECOInstanceFlags::HasPendingLowPriorityUpdate);
 
-	if (InCustomizableObjectInstance && InCustomizableObjectInstance->IsValidLowLevel())
-	{
-		check(InCustomizableObjectInstance->GetPrivate() != nullptr);
-		MutablePendingInstanceWork.AddDiscard(FMutablePendingInstanceDiscard(InCustomizableObjectInstance));
-	}
-}
+			if (Context->bIsHighPriority)
+			{
+				HighPriorityInstanceUpdates.PushLast(Context);
+				InstancePrivate->SetCOInstanceFlags(ECOInstanceFlags::HasPendingHighPriorityUpdate);
+			}
+			else
+			{
+				LowPriorityInstanceUpdates.PushLast(Context);
+				InstancePrivate->SetCOInstanceFlags(ECOInstanceFlags::HasPendingLowPriorityUpdate);
+			}
 
+			++PendingInstanceUpdates;
+		};
 
-bool UCustomizableObjectSystem::IsReplaceDiscardedWithReferenceMeshEnabled() const
-{
-	if (Private)
+	if (InstancePrivate->HasCOInstanceFlags(ECOInstanceFlags::HasPendingHighPriorityUpdate))
 	{
-		return Private->IsReplaceDiscardedWithReferenceMeshEnabled();
-	}
+		for (FMutableInstanceUpdate& HighPriorityUpdate : HighPriorityInstanceUpdates)
+		{
+			if (HighPriorityUpdate && HighPriorityUpdate->Instance == Context->Instance)
+			{
+				HighPriorityUpdate->bOptimizedUpdate = true;
+				HighPriorityUpdate->UpdateResult = EUpdateResult::ErrorOptimized;
+				FinishUpdateGlobal(HighPriorityUpdate.ToSharedRef());
 
-	return false;
-}
+				if (Context->bIsHighPriority)
+				{
+					HighPriorityUpdate = Context;
+				}
+				else
+				{
+					HighPriorityUpdate.Reset();
+					EnqueueUpdate();
+				}
 
+				break;
+			}
+		}
+	}
+	else if (InstancePrivate->HasCOInstanceFlags(ECOInstanceFlags::HasPendingLowPriorityUpdate))
+	{
+		for (FMutableInstanceUpdate& LowPriorityUpdate : LowPriorityInstanceUpdates)
+		{
+			if (LowPriorityUpdate && LowPriorityUpdate->Instance == Context->Instance)
+			{
+				LowPriorityUpdate->bOptimizedUpdate = true;
+				LowPriorityUpdate->UpdateResult = EUpdateResult::ErrorOptimized;
+				FinishUpdateGlobal(LowPriorityUpdate.ToSharedRef());
 
-void UCustomizableObjectSystem::SetReplaceDiscardedWithReferenceMeshEnabled(bool bIsEnabled)
-{
-	if (Private)
+				if (!Context->bIsHighPriority)
+				{
+					LowPriorityUpdate = Context;
+				}
+				else
+				{
+					LowPriorityUpdate.Reset();
+					EnqueueUpdate();
+				}
+
+				break;
+			}
+		}
+	}
+	else
 	{
-		Private->SetReplaceDiscardedWithReferenceMeshEnabled(bIsEnabled);
+		EnqueueUpdate();
 	}
 }
 
@@ -1731,8 +1457,6 @@ bool UCustomizableObjectSystem::LockObject(UCustomizableObject* InObject)
 			Task.Wait();
 		}
 
-		Private->MutablePendingInstanceWork.RemoveUpdatesForObject(InObject);
-		
 		check(InObject->GetPrivate()->bLocked);
 
 		return true;
@@ -1859,38 +1583,31 @@ namespace impl
 	
 	void FixLODs(const TSharedRef<FUpdateContextPrivate>& Operation)
 	{
-		TMap<FName, uint8> MinLODs = Operation->GetCapturedDescriptor().MinLOD;
-		TMap<FName, uint8> RequestedLODs = Operation->GetFirstRequestedLOD();
-		
 		for (int32 ComponentIndex = 0; ComponentIndex < Operation->MutableInstance->GetComponentCount(); ++ComponentIndex)
 		{
 			const UE::Mutable::Private::FComponentId ComponentId = Operation->MutableInstance->GetComponentId(ComponentIndex);
 			const FName& ComponentName = Operation->ComponentNames[ComponentId];
 
-			uint8& MinLOD = MinLODs.FindOrAdd(ComponentName, 0);
-			uint8& RequestedLOD = RequestedLODs.FindOrAdd(ComponentName, 0);
 			uint8& NumLODsAvailable = Operation->NumLODsAvailable.FindOrAdd(ComponentName, 0);
 			uint8& FirstResidentLOD = Operation->FirstResidentLOD.FindOrAdd(ComponentName, 0);
-			const uint8& FirstLODAvailable = Operation->FirstLODAvailable.FindOrAdd(ComponentName, 0);
+			uint8& FirstLODAvailable = Operation->FirstLODAvailable.FindOrAdd(ComponentName, 0);
 
 			NumLODsAvailable = Operation->MutableInstance->GetLODCount(ComponentIndex);
 			
 			if (Operation->bStreamMeshLODs)
 			{
 				FirstResidentLOD = FMath::Clamp(FirstResidentLOD, FirstLODAvailable, NumLODsAvailable - 1);
-				MinLOD = 0;
-				RequestedLOD = FirstResidentLOD;
 			}
 			else
 			{
+				if (!Operation->bBake)
+				{
+					FirstLODAvailable = FMath::Clamp(FirstLODAvailable, UCustomizableObjectSystemPrivate::MutableMeshesLODBias, NumLODsAvailable - 1);
+				}
+				
 				FirstResidentLOD = FirstLODAvailable;
-				MinLOD = FMath::Clamp(MinLOD, FirstLODAvailable, NumLODsAvailable - 1);
-				RequestedLOD = FMath::Clamp(RequestedLOD, MinLOD, static_cast<uint8>(NumLODsAvailable - 1));
 			}
 		}
-		
-		Operation->SetMinLOD(MinLODs);
-		Operation->SetFirstRequestedLOD(RequestedLODs);
 	}
 	
 
@@ -1970,7 +1687,7 @@ namespace impl
 			// Allocate Morph data for used morphs.
 			const int32 NumLODsAvailable = OperationData->NumLODsAvailable[ComponentName];
 	
-			for (int32 LODIndex = OperationData->GetFirstRequestedLOD()[ComponentName]; LODIndex < NumLODsAvailable; ++LODIndex)
+			for (int32 LODIndex = OperationData->FirstResidentLOD[ComponentName]; LODIndex < NumLODsAvailable; ++LODIndex)
 			{
 				const TSharedRef<FInstanceUpdateData::FLOD>& LOD = UpdateData.LODs[Component->FirstLOD + LODIndex];
 
@@ -2252,12 +1969,8 @@ namespace impl
 			{
 				OperationData->InstanceUpdateData.LODs.Add(MakeShared<FInstanceUpdateData::FLOD>());
 			}
-
-			const int32 FirstLOD = OperationData->bStreamMeshLODs ?
-				OperationData->FirstLODAvailable[ComponentName] :
-				OperationData->GetFirstRequestedLOD()[ComponentName];
 			
-			for (int32 LODIndex = FirstLOD; LODIndex < Component->LODCount; ++LODIndex)
+			for (int32 LODIndex = OperationData->FirstLODAvailable[ComponentName]; LODIndex < Component->LODCount; ++LODIndex)
 			{
 				MUTABLE_CPUPROFILER_SCOPE(GetMesh);
 
@@ -2836,33 +2549,6 @@ namespace impl
 		}
 		
 		TRACE_CPUPROFILER_EVENT_SCOPE_TEXT(*Object->GetName()); // So that we know what we are currently updating.
-
-		UCustomizableObjectInstance* Instance = Operation->Instance.Get();
-		if (!IsValid(Instance)) // Only start if it hasn't been already destroyed (i.e. GC after finish PIE)
-		{
-			Operation->UpdateResult = EUpdateResult::Error;
-			FinishUpdateGlobal(Operation);
-			return;
-		}
-		
-		UCustomizableInstancePrivate* InstancePrivate = Instance->GetPrivate();
-		check(InstancePrivate); // Already checked in GetPrivate. Make static analyzer happy.
-		
-		if (InstancePrivate->HasCOInstanceFlags(PendingLODsUpdate))
-		{
-			InstancePrivate->ClearCOInstanceFlags(PendingLODsUpdate);
-			// TODO: Is anything needed for this now?
-			//Operation->CustomizableObjectInstance->ReleaseMutableInstanceId(); // To make mutable regenerate the LODs even if the instance parameters have not changed
-		}
-
-		// Skip update, the requested update is equal to the running update.
-		if (Operation->GetCapturedDescriptorHash().IsSubset(InstancePrivate->CommittedDescriptorHash))
-		{
-			Operation->UpdateResult = EUpdateResult::Success;
-			UpdateSkeletalMesh(Operation);
-			FinishUpdateGlobal(Operation);
-			return;
-		}
 		
 		UCustomizableObjectPrivate* ObjectPrivate = Object->GetPrivate();
 		
@@ -2874,21 +2560,8 @@ namespace impl
 			return;
 		}
 
-		// Only update resources if the instance is in range (it could have got far from the player since the task was queued)
-		check(SystemPrivate->CurrentInstanceLODManagement != nullptr);
-		if (SystemPrivate->CurrentInstanceLODManagement->IsOnlyUpdateCloseCustomizableObjectsEnabled()
-			&& InstancePrivate
-			&& InstancePrivate->LastMinSquareDistFromComponentToPlayer > FMath::Square(SystemPrivate->CurrentInstanceLODManagement->GetOnlyUpdateCloseCustomizableObjectsDist())
-			&& InstancePrivate->LastMinSquareDistFromComponentToPlayer != FLT_MAX // This means it is the first frame so it has to be updated
-		   )
-		{
-			Operation->UpdateResult = EUpdateResult::Error;
-			FinishUpdateGlobal(Operation);
-			return;
-		}
-
-		Operation->Parameters = Operation->GetCapturedDescriptor().GetParameters();
-		if (!Operation->Parameters)
+		UCustomizableObjectInstance* Instance = Operation->Instance.Get();
+		if (!IsValid(Instance)) // Only start if it hasn't been already destroyed (i.e. GC after finish PIE)
 		{
 			Operation->UpdateResult = EUpdateResult::Error;
 			FinishUpdateGlobal(Operation);
@@ -2906,6 +2579,9 @@ namespace impl
 
 		const UModelResources& ModelResources = Instance->GetCustomizableObject()->GetPrivate()->GetModelResourcesChecked();
 
+		UCustomizableInstancePrivate* InstancePrivate = Instance->GetPrivate();
+		check(InstancePrivate); // Already checked in GetPrivate. Make static analyzer happy.
+
 		const int32 State = InstancePrivate->GetState();
 		const FString& StateName = ObjectPrivate->GetStateName(State);
 		const FMutableStateData* StateData = ModelResources.StateUIDataMap.Find(StateName);
@@ -2984,55 +2660,6 @@ namespace impl
 		Operation->PixelFormatOverride = SystemPrivate->ImageFormatOverrideFunc;
 #endif
 		
-		bool bRequestAllLODs = !Operation->bIsOnlyGenerateRequestedLODsEnabled ||
-			!SystemPrivate->CurrentInstanceLODManagement->IsOnlyGenerateRequestedLODLevelsEnabled();
-
-#if WITH_EDITOR
-		// In the editor LOD Management is disabled by default. Overwrite requested LODs when disabled.
-		bRequestAllLODs |= !bEnableLODManagmentInEditor;
-
-		for (TObjectIterator<UCustomizableObjectInstanceUsage> CustomizableObjectInstanceUsage; CustomizableObjectInstanceUsage && !bRequestAllLODs; ++CustomizableObjectInstanceUsage)
-		{
-			if (IsValid(*CustomizableObjectInstanceUsage) && CustomizableObjectInstanceUsage->GetPrivate()->IsNetMode(NM_DedicatedServer))
-			{
-				continue;
-			}
-
-			if (IsValid(*CustomizableObjectInstanceUsage) &&
-				CustomizableObjectInstanceUsage->GetCustomizableObjectInstance() == Instance)
-			{
-				EWorldType::Type WorldType = EWorldType::Type::None;
-
-				USkeletalMeshComponent* Parent = Cast<USkeletalMeshComponent>(CustomizableObjectInstanceUsage->GetAttachParent());
-
-				if (Parent && Parent->GetWorld())
-				{
-					WorldType = Parent->GetWorld()->WorldType;
-				}
-
-				switch (WorldType)
-				{
-					// Editor preview instances
-				case EWorldType::EditorPreview:
-				case EWorldType::None:
-					bRequestAllLODs = true;
-				default:;
-				}
-			}
-		}
-#endif // WITH_EDITOR
-
-		if (bRequestAllLODs)
-		{
-			TMap<FName, uint8> RequestedLODs = Operation->GetFirstRequestedLOD();
-			for (const FName& ComponentName : ModelResources.ComponentNamesPerObjectComponent)
-			{
-				RequestedLODs.Add(ComponentName, 0);
-			}
-
-			Operation->SetFirstRequestedLOD(RequestedLODs);
-		}
-		
 		// CreateMutableInstance
 		{
 			if (FLogBenchmarkUtil::IsBenchmarkingReportingEnabled())
@@ -3128,11 +2755,7 @@ namespace impl
 
 			const FName& ComponentName = Operation->ComponentNames[ComponentId];
 
-			const int32 FirstLOD = Operation->bStreamMeshLODs ?
-				Operation->FirstLODAvailable[ComponentName] :
-				Operation->GetFirstRequestedLOD()[ComponentName];
-
-			for (int32 LODIndex = FirstLOD; LODIndex < NumLODsAvailable[ComponentName]; ++LODIndex)
+			for (int32 LODIndex = Operation->FirstLODAvailable[ComponentName]; LODIndex < NumLODsAvailable[ComponentName]; ++LODIndex)
 			{
 				ComponentDescriptor->MeshDescriptors[LODIndex] = Operation->MutableInstance->GetMeshId(ComponentIndex, LODIndex);
 			}
@@ -3258,60 +2881,77 @@ namespace impl
 
 void UCustomizableObjectSystemPrivate::UpdateTick(const bool bBlocking)
 {
-	FMutableUpdateCandidate* LODUpdateCandidateFound = nullptr;
-
-	for (UCustomizableObjectInstanceUsage* CustomizableObjectInstanceUsage : UCustomizableObjectSystemPrivate::GetCustomizableObjectInstanceUsages())
 	{
-		if (!IsValid(CustomizableObjectInstanceUsage))
-		{
-			continue;
-		}
+		MUTABLE_CPUPROFILER_SCOPE(UpdateTick_Usage);
 
-		UCustomizableObjectInstanceUsagePrivate* UsagePrivate = CustomizableObjectInstanceUsage->GetPrivate();
-		
-		if (!UsagePrivate->bPendingSetSkeletalMesh)
+		for (UCustomizableObjectInstanceUsage* CustomizableObjectInstanceUsage : CustomizableObjectInstanceUsages)
 		{
-			continue;
-		}
+			UCustomizableObjectInstance* Instance = CustomizableObjectInstanceUsage->GetCustomizableObjectInstance();
+			if (!Instance)
+			{
+				continue;
+			}
 
-		UsagePrivate->bPendingSetSkeletalMesh = false;
+			UCustomizableObject* Object = Instance->GetCustomizableObject();
+			if (!Object)
+			{
+				continue;
+			}
 
-		UCustomizableObjectInstance* Instance = CustomizableObjectInstanceUsage->GetCustomizableObjectInstance();
-		if (!Instance)
-		{
-			continue;
-		}
-		
-		Instance->GetPrivate()->SetCOInstanceFlags(UsedByComponent);
+			USkeletalMeshComponent* Parent = Cast<USkeletalMeshComponent>(CustomizableObjectInstanceUsage->GetAttachParent());
+			if (!Parent)
+			{
+				continue;
+			}
 
-		UCustomizableObject* Object = Instance->GetCustomizableObject();
-		if (!Object)
-		{
-			continue;
-		}
+			Instance->GetPrivate()->bAllowAutomaticUpdate = true;
 
-		USkeletalMeshComponent* Parent = Cast<USkeletalMeshComponent>(CustomizableObjectInstanceUsage->GetAttachParent());
-		if (!Parent)
-		{
-			continue;
-		}
-		
-		if (!CustomizableObjectInstanceUsage->GetSkipSetReferenceSkeletalMesh() &&
-			 Object->bEnableUseRefSkeletalMeshAsPlaceholder)
-		{
-			const FName& ComponentName = CustomizableObjectInstanceUsage->GetComponentName();
-			
-			if (USkeletalMesh* ReferenceSkeletalMesh = Object->GetComponentMeshReferenceSkeletalMesh(ComponentName))
+			UCustomizableObjectInstanceUsagePrivate* UsagePrivate = CustomizableObjectInstanceUsage->GetPrivate();
+
+			if (Instance->GetPrivate()->HasCOInstanceFlags(ECOInstanceFlags::Discarded))
 			{
+				Instance->GetPrivate()->InvalidateGeneratedData();
+
 				Parent->EmptyOverrideMaterials();
-				Parent->SetSkeletalMesh(ReferenceSkeletalMesh);
+
+				USkeletalMesh* SkeletalMesh = nullptr;
+
+				if (!CustomizableObjectInstanceUsage->GetSkipSetReferenceSkeletalMesh() &&
+					Object->bEnableUseRefSkeletalMeshAsPlaceholder)
+				{
+					const FName& ComponentName = CustomizableObjectInstanceUsage->GetComponentName();
+					SkeletalMesh = Object->GetComponentMeshReferenceSkeletalMesh(ComponentName);
+				}
+
+				Parent->SetSkeletalMesh(SkeletalMesh);
+
+				for (const UCustomizableObjectExtension* Extension : ICustomizableObjectModule::Get().GetRegisteredExtensions())
+				{
+					Extension->OnCustomizableObjectInstanceUsageDiscarded(*CustomizableObjectInstanceUsage);
+				}
+
+				UsagePrivate->bPendingSetReferenceSkeletalMesh = false;
+			}
+
+			else if (UsagePrivate->bPendingSetReferenceSkeletalMesh)
+			{
+				if (!CustomizableObjectInstanceUsage->GetSkipSetReferenceSkeletalMesh() &&
+					Object->bEnableUseRefSkeletalMeshAsPlaceholder)
+				{
+					const FName& ComponentName = CustomizableObjectInstanceUsage->GetComponentName();
+
+					if (USkeletalMesh* ReferenceSkeletalMesh = Object->GetComponentMeshReferenceSkeletalMesh(ComponentName))
+					{
+						Parent->EmptyOverrideMaterials();
+						Parent->SetSkeletalMesh(ReferenceSkeletalMesh);
+					}
+				}
+
+				UsagePrivate->bPendingSetReferenceSkeletalMesh = false;
 			}
 		}
-		
-		Instance->GetPrivate()->bAutomaticUpdateRequired = true;
 	}
 	
-	
 	bool bPendingCompilation = false;
 #if WITH_EDITOR
 	ICustomizableObjectEditorModule* EditorModule = ICustomizableObjectEditorModule::Get();
@@ -3321,174 +2961,82 @@ void UCustomizableObjectSystemPrivate::UpdateTick(const bool bBlocking)
 	// Get a new operation if we aren't working on one
 	if (!CurrentMutableOperation && bIsMutableEnabled && !bPendingCompilation)
 	{
-		// Reset the instance relevancy
-		// The RequestedUpdates only refer to LOD changes. User Customization and discards are handled separately
-		FMutableInstanceUpdateMap RequestedLODUpdates;
-		
-		CurrentInstanceLODManagement->UpdateInstanceDistsAndLODs(RequestedLODUpdates);
-
-		for (UCustomizableObjectInstance* CustomizableObjectInstance : UCustomizableObjectSystemPrivate::GetCustomizableObjectInstances())
+		if (bGenerateInstancesWithinRange)
 		{
-			if (IsValid(CustomizableObjectInstance) && CustomizableObjectInstance->GetPrivate())
-			{
-				UCustomizableInstancePrivate* ObjectInstancePrivateData = CustomizableObjectInstance->GetPrivate();
+			UpdateInstanceDistances();
 
-				if (ObjectInstancePrivateData->HasCOInstanceFlags(UsedByComponentInPlay))
-				{
-					ObjectInstancePrivateData->TickUpdateCloseCustomizableObjects(*CustomizableObjectInstance, RequestedLODUpdates);
-				}
-				else if (ObjectInstancePrivateData->HasCOInstanceFlags(UsedByComponent))
-				{
-					ensure(!RequestedLODUpdates.Contains(CustomizableObjectInstance));
-					ObjectInstancePrivateData->UpdateInstanceIfNotGenerated(*CustomizableObjectInstance, RequestedLODUpdates);
-				}
-				else
-				{
-					ensure(!RequestedLODUpdates.Contains(CustomizableObjectInstance));
-				}
-
-				ObjectInstancePrivateData->ClearCOInstanceFlags((ECOInstanceFlags)(UsedByComponent | UsedByComponentInPlay | PendingLODsUpdate)); // TODO MTBL-391: Makes no sense to clear it here, what if an update is requested before we set it back to true
-			}
-			else
+			if (ICustomizableObjectEditorModule* Module = ICustomizableObjectEditorModule::Get())
 			{
-				ensure(!RequestedLODUpdates.Contains(CustomizableObjectInstance));
+				Module->UpdateInstanceDistanceToLevelEditorCamera();
 			}
 		}
 
+		bool bUpdatingInstance = false;
 		{
-			// Look for the highest priority update between the pending updates and the LOD Requested Updates
-			EQueuePriorityType MaxPriorityFound = EQueuePriorityType::Low;
-			double MaxSquareDistanceFound = TNumericLimits<double>::Max();
-			double MinTimeFound = TNumericLimits<double>::Max();
-			const FMutablePendingInstanceUpdate* PendingInstanceUpdateFound = nullptr;
-
-			// Look for the highest priority Pending Update
-			for (auto Iterator = MutablePendingInstanceWork.GetUpdateIterator(); Iterator; ++Iterator)
-			{
-				FMutablePendingInstanceUpdate& PendingUpdate = *Iterator;
+			MUTABLE_CPUPROFILER_SCOPE(UpdateTick_Queue);
 
-				if (PendingUpdate.Context->Instance.IsValid() && PendingUpdate.Context->Instance->GetCustomizableObject())
-				{
-					const EQueuePriorityType PriorityType = GetUpdatePriority(*PendingUpdate.Context->Instance, false);
-					
-					if (PendingUpdate.Context->PriorityType <= MaxPriorityFound)
-					{
-						const double MinSquareDistFromComponentToPlayer = PendingUpdate.Context->Instance->GetPrivate()->MinSquareDistFromComponentToPlayer;
-						
-						if (MinSquareDistFromComponentToPlayer < MaxSquareDistanceFound ||
-							(MinSquareDistFromComponentToPlayer == MaxSquareDistanceFound && PendingUpdate.Context->StartQueueTime < MinTimeFound))
-						{
-							MaxPriorityFound = PriorityType;
-							MaxSquareDistanceFound = MinSquareDistFromComponentToPlayer;
-							MinTimeFound = PendingUpdate.Context->StartQueueTime;
-							PendingInstanceUpdateFound = &PendingUpdate;
-							LODUpdateCandidateFound = nullptr;
-						}
-					}
-				}
-				else
+			while (!bUpdatingInstance)
+			{
+				FMutableInstanceUpdate InstanceUpdate;
+				if (!HighPriorityInstanceUpdates.TryPopFirst(InstanceUpdate))
 				{
-					Iterator.RemoveCurrent();
+					break;
 				}
-			}
 
-			// Look for a higher priority LOD update
-			for (TPair<const UCustomizableObjectInstance*, FMutableUpdateCandidate>& LODUpdateTuple : RequestedLODUpdates)
-			{
-				if (const UCustomizableObjectInstance* Instance = LODUpdateTuple.Key)
+				if (InstanceUpdate)
 				{
-					if (UCustomizableObject* Object = Instance->GetCustomizableObject())
-					{
-						FMutableUpdateCandidate& LODUpdateCandidate = LODUpdateTuple.Value;
-						ensure(LODUpdateCandidate.HasBeenIssued());
-
-						if (LODUpdateCandidate.Priority <= MaxPriorityFound)
-						{
-							UCustomizableInstancePrivate* CustomizableInstancePrivate = LODUpdateCandidate.CustomizableObjectInstance->GetPrivate();
-
-							FDescriptorHash LODUpdateDescriptorHash = CustomizableInstancePrivate->CommittedDescriptorHash;
-							LODUpdateDescriptorHash.MinLODs = LODUpdateCandidate.MinLOD;
-							LODUpdateDescriptorHash.QualitySettingMinLODs = LODUpdateCandidate.QualitySettingMinLODs;
-							LODUpdateDescriptorHash.FirstRequestedLOD = LODUpdateCandidate.FirstRequestedLOD;
-
-							if (CustomizableInstancePrivate->MinSquareDistFromComponentToPlayer < MaxSquareDistanceFound &&
-								(IsStreamingEnabled(*Object, CustomizableInstancePrivate->GetState()) || !LODUpdateDescriptorHash.IsSubset(CustomizableInstancePrivate->CommittedDescriptorHash) || CustomizableInstancePrivate->bAutomaticUpdateRequired))
-							{
-								MaxPriorityFound = LODUpdateCandidate.Priority;
-								MaxSquareDistanceFound = CustomizableInstancePrivate->MinSquareDistFromComponentToPlayer;
-								PendingInstanceUpdateFound = nullptr;
-								LODUpdateCandidateFound = &LODUpdateCandidate;
-							}
-						}
-					}
+					bUpdatingInstance = TryStartUpdate(InstanceUpdate.ToSharedRef());
 				}
+				--PendingInstanceUpdates;
 			}
 
-			NumLODUpdatesLastTick = RequestedLODUpdates.Num();
-
-			// If the chosen LODUpdate has the same instance as a PendingUpdate, choose the PendingUpdate to apply both the LOD update
-			// and customization change
-			if (LODUpdateCandidateFound)
+			while (!bUpdatingInstance)
 			{
-				if (const FMutablePendingInstanceUpdate *PendingUpdateWithSameInstance = MutablePendingInstanceWork.GetUpdate(LODUpdateCandidateFound->CustomizableObjectInstance))
+				FMutableInstanceUpdate InstanceUpdate;
+				if (!LowPriorityInstanceUpdates.TryPopFirst(InstanceUpdate))
 				{
-					PendingInstanceUpdateFound = PendingUpdateWithSameInstance;
-					LODUpdateCandidateFound = nullptr;
-
-					// In the processing of the PendingUpdate just below, it will add the LODUpdate's LOD params
+					break;
 				}
-			}
 
-			if (PendingInstanceUpdateFound)
-			{
-				check(!LODUpdateCandidateFound);
-
-				UCustomizableObjectInstance* PendingInstance = PendingInstanceUpdateFound->Context->Instance.Get();
-				check(PendingInstance);
-				
-				// Maybe there's a LODUpdate that has the same instance, merge both updates as an optimization
-				if (FMutableUpdateCandidate* LODUpdateWithSameInstance = RequestedLODUpdates.Find(PendingInstance))
+				if (InstanceUpdate)
 				{
-					LODUpdateWithSameInstance->ApplyLODUpdateParamsToInstance(PendingInstanceUpdateFound->Context.Get());
+					bUpdatingInstance = TryStartUpdate(InstanceUpdate.ToSharedRef());
 				}
-
-				StartUpdateSkeletalMesh(PendingInstanceUpdateFound->Context);
-				MutablePendingInstanceWork.RemoveUpdate(PendingInstanceUpdateFound->Context->Instance);
-			}
-			else if (LODUpdateCandidateFound)
-			{
-				UCustomizableObjectInstance* Instance = LODUpdateCandidateFound->CustomizableObjectInstance;
-				const bool bGenerated = Instance->GetPrivate()->SkeletalMeshStatus == ESkeletalMeshStatus::Success;
-				const FCustomizableObjectInstanceDescriptor& Descriptor = bGenerated ? Instance->GetPrivate()->CommittedDescriptor : Instance->GetPrivate()->GetDescriptor();
-
-				const TSharedRef<FUpdateContextPrivate> Context = MakeShared<FUpdateContextPrivate>(*Instance, Descriptor);
-				
-				// Commit the LOD changes
-				LODUpdateCandidateFound->ApplyLODUpdateParamsToInstance(*Context);
-
-				StartUpdateSkeletalMesh(Context);
+				--PendingInstanceUpdates;
 			}
 		}
 
+		int32 MaxNumInstancesToDiscard = CVarMaxNumInstancesToDiscardPerTick.GetValueOnGameThread();
+
 		{
-			for (UCustomizableObjectInstance* CustomizableObjectInstance : UCustomizableObjectSystemPrivate::GetCustomizableObjectInstances())
+			MUTABLE_CPUPROFILER_SCOPE(UpdateTick_AutomaticUpdates);
+			for (UCustomizableObjectInstance* CustomizableObjectInstance : CustomizableObjectInstances)
 			{
-				if (IsValid(CustomizableObjectInstance))
+				UCustomizableInstancePrivate* CustomizableObjectInstancePrivate = CustomizableObjectInstance->GetPrivate();
+				CustomizableObjectInstancePrivate->ClearCOInstanceFlags(ECOInstanceFlags::Discarded);
+
+				if (IsDiscardedByDistance(CustomizableObjectInstancePrivate))
 				{
-					if (UCustomizableInstancePrivate* InstancePrivate = CustomizableObjectInstance->GetPrivate())
+					if (MaxNumInstancesToDiscard > 0 &&
+						CustomizableObjectInstancePrivate->SkeletalMeshStatus == ESkeletalMeshStatus::Success)
 					{
-						InstancePrivate->LastMinSquareDistFromComponentToPlayer = InstancePrivate->MinSquareDistFromComponentToPlayer;
-						InstancePrivate->MinSquareDistFromComponentToPlayer = FLT_MAX;
+						CustomizableObjectInstancePrivate->SetCOInstanceFlags(ECOInstanceFlags::Discarded);
+						--MaxNumInstancesToDiscard;
 					}
 				}
+				else if (!bUpdatingInstance &&
+					(CustomizableObjectInstancePrivate->bAllowAutomaticUpdate && CustomizableObjectInstancePrivate->SkeletalMeshStatus == ESkeletalMeshStatus::NotGenerated))
+				{
+					bUpdatingInstance = TryStartAutomaticUpdateForInstance(*CustomizableObjectInstance);
+				}
+
+				CustomizableObjectInstancePrivate->MinSquareDistFromComponentToPlayer = FLT_MAX;
+				CustomizableObjectInstancePrivate->bAllowAutomaticUpdate = false;
 			}
 		}
 
 		// Update the streaming limit if it has changed. It is safe to do this now.
 		UpdateMemoryLimit();
-
-		// Free memory before starting the new update
-		GetPublic()->DiscardInstances();
 	}
 	
 	LogBenchmarkUtil.UpdateStats(); // Must be the last thing to perform
@@ -3497,8 +3045,6 @@ void UCustomizableObjectSystemPrivate::UpdateTick(const bool bBlocking)
 	{
 		FStreamingManagerCollection::Get().RemoveStreamingManager(this);
 	}
-	
-	bPendingLODUpdate = LODUpdateCandidateFound != nullptr; // Still a pending LOD update. We can not use the size of RequestedLODUpdates since not all requests valid in future ticks.
 }
 
 
@@ -3567,12 +3113,10 @@ int32 UCustomizableObjectSystemPrivate::GetRemainingWork() const
 	// Do not tick if the CookCommandlet is running.
 	if (!IsRunningCookCommandlet())
 	{
-		RemainingWork += CurrentMutableOperation.IsValid() + 
-			   MutablePendingInstanceWork.Num() +
-			   bPendingLODUpdate;
+		RemainingWork += CurrentMutableOperation.IsValid() + PendingInstanceUpdates;
 	}
 	
-	RemainingWork = MutableTaskGraph.GetRemainingWork();
+	RemainingWork += MutableTaskGraph.GetRemainingWork();
 
 	if (ICustomizableObjectEditorModule* Module = ICustomizableObjectEditorModule::Get())
 	{
@@ -3586,86 +3130,6 @@ int32 UCustomizableObjectSystemPrivate::GetRemainingWork() const
 }
 
 
-TAutoConsoleVariable<int32> CVarMaxNumInstancesToDiscardPerTick(
-	TEXT("mutable.MaxNumInstancesToDiscardPerTick"),
-	30,
-	TEXT("The maximum number of stale instances that will be discarded per tick by Mutable."),
-	ECVF_Scalability);
-
-
-void UCustomizableObjectSystem::DiscardInstances()
-{
-	MUTABLE_CPUPROFILER_SCOPE(DiscardInstances);
-
-	check(IsInGameThread());
-
-	// Handle instance discards
-	int32 NumInstancesDiscarded = 0;
-	const int32 DiscardLimitPerTick = CVarMaxNumInstancesToDiscardPerTick.GetValueOnGameThread();
-
-	for (TSet<FMutablePendingInstanceDiscard, FPendingInstanceDiscardKeyFuncs>::TIterator Iterator = Private->MutablePendingInstanceWork.GetDiscardIterator();
-		Iterator && NumInstancesDiscarded < DiscardLimitPerTick;
-		++Iterator)
-	{
-		UCustomizableObjectInstance* COI = Iterator->CustomizableObjectInstance.Get();
-
-		const bool bUpdating = Private->CurrentMutableOperation && Private->CurrentMutableOperation->Instance == Iterator->CustomizableObjectInstance;
-		if (COI && COI->GetPrivate() && !bUpdating)
-		{
-			UCustomizableInstancePrivate* COIPrivateData = COI ? COI->GetPrivate() : nullptr;
-
-			// Only discard resources if the instance is still out range (it could have got closer to the player since the task was queued)
-			if (!GetPrivate()->CurrentInstanceLODManagement->IsOnlyUpdateCloseCustomizableObjectsEnabled() ||
-				COIPrivateData->LastMinSquareDistFromComponentToPlayer > FMath::Square(GetPrivate()->CurrentInstanceLODManagement->GetOnlyUpdateCloseCustomizableObjectsDist()))
-			{
-				COIPrivateData->DiscardResources();
-
-				if (UCustomizableObject* CustomizableObject = COI->GetCustomizableObject())
-				{
-					if (const UModelResources* ModelResources = CustomizableObject->GetPrivate()->GetModelResources())
-					{
-						for (UCustomizableObjectInstanceUsage* InstanceUsage : UCustomizableObjectSystemPrivate::GetCustomizableObjectInstanceUsages())
-						{
-							if (!IsValid(InstanceUsage) || InstanceUsage->GetCustomizableObjectInstance() != COI)
-							{
-								continue;
-							}
-
-#if WITH_EDITOR
-							if (InstanceUsage->GetPrivate()->IsNetMode(NM_DedicatedServer))
-							{
-								continue;
-							}
-#endif
-
-							if (USkeletalMeshComponent* Parent = Cast<USkeletalMeshComponent>(InstanceUsage->GetAttachParent()))
-							{
-								Parent->EmptyOverrideMaterials();
-
-								if (IsReplaceDiscardedWithReferenceMeshEnabled())
-								{
-									const FName& ComponentName = InstanceUsage->GetComponentName();
-									USkeletalMesh* SkeletalMesh = CustomizableObject->GetSkeletalMeshComponentReferenceSkeletalMesh(ComponentName);
-									Parent->SetSkeletalMesh(SkeletalMesh);
-								}
-							}
-
-							for (const UCustomizableObjectExtension* Extension : ICustomizableObjectModule::Get().GetRegisteredExtensions())
-							{
-								Extension->OnCustomizableObjectInstanceUsageDiscarded(*InstanceUsage);
-							}
-						}
-					}
-				}
-
-				Iterator.RemoveCurrent();
-				NumInstancesDiscarded++;
-			}
-		}
-	}
-}
-
-
 bool UCustomizableObjectSystem::IsUpdating(const UCustomizableObjectInstance* Instance) const
 {
 	if (!Instance)
@@ -3676,27 +3140,41 @@ bool UCustomizableObjectSystem::IsUpdating(const UCustomizableObjectInstance* In
 	return GetPrivate()->IsUpdating(*Instance);
 }
 
-TSet<UCustomizableObjectInstanceUsage*> UCustomizableObjectSystemPrivate::CustomizableObjectInstanceUsages;
-TSet<UCustomizableObjectInstance*> UCustomizableObjectSystemPrivate::CustomizableObjectInstances;
+TArray<UCustomizableObjectInstanceUsage*> UCustomizableObjectSystemPrivate::CustomizableObjectInstanceUsages;
+TArray<UCustomizableObjectInstance*> UCustomizableObjectSystemPrivate::CustomizableObjectInstances;
 
-void UCustomizableObjectSystemPrivate::RegisterCustomizableObjectInstanceUsage(UCustomizableObjectInstanceUsage* Obj)
+void UCustomizableObjectSystemPrivate::RegisterCustomizableObjectInstanceUsage(UCustomizableObjectInstanceUsage& Obj)
 {
-	UCustomizableObjectSystemPrivate::CustomizableObjectInstanceUsages.Add(Obj);
+	if (!Obj.IsTemplate())
+	{
+		check(!UCustomizableObjectSystemPrivate::CustomizableObjectInstanceUsages.Contains(&Obj));
+		UCustomizableObjectSystemPrivate::CustomizableObjectInstanceUsages.Add(&Obj);
+	}
 }
 
-void UCustomizableObjectSystemPrivate::UnregisterCustomizableObjectInstanceUsage(UCustomizableObjectInstanceUsage* Obj)
+void UCustomizableObjectSystemPrivate::UnregisterCustomizableObjectInstanceUsage(UCustomizableObjectInstanceUsage& Obj)
 {
-	UCustomizableObjectSystemPrivate::CustomizableObjectInstanceUsages.Remove(Obj);
+	if (!Obj.IsTemplate())
+	{
+		UCustomizableObjectSystemPrivate::CustomizableObjectInstanceUsages.Remove(&Obj);
+	}
 }
 
-void UCustomizableObjectSystemPrivate::RegisterCustomizableObjectInstance(UCustomizableObjectInstance* Obj)
+void UCustomizableObjectSystemPrivate::RegisterCustomizableObjectInstance(UCustomizableObjectInstance& Obj)
 {
-	UCustomizableObjectSystemPrivate::CustomizableObjectInstances.Add(Obj);
+	if (!Obj.IsTemplate())
+	{
+		check(!UCustomizableObjectSystemPrivate::CustomizableObjectInstances.Contains(&Obj));
+		UCustomizableObjectSystemPrivate::CustomizableObjectInstances.Add(&Obj);
+	}
 }
 
-void UCustomizableObjectSystemPrivate::UnregisterCustomizableObjectInstance(UCustomizableObjectInstance* Obj)
+void UCustomizableObjectSystemPrivate::UnregisterCustomizableObjectInstance(UCustomizableObjectInstance& Obj)
 {
-	UCustomizableObjectSystemPrivate::CustomizableObjectInstances.Remove(Obj);
+	if (!Obj.IsTemplate())
+	{
+		UCustomizableObjectSystemPrivate::CustomizableObjectInstances.Remove(&Obj);
+	}
 }
 
 bool UCustomizableObjectSystemPrivate::IsUsingBenchmarkingSettings()
@@ -3725,25 +3203,13 @@ int32 UCustomizableObjectSystem::GetNumInstances() const
 
 int32 UCustomizableObjectSystem::GetNumPendingInstances() const
 {
-	return GetPrivate()->MutablePendingInstanceWork.Num() + GetPrivate()->NumLODUpdatesLastTick;
+	return GetPrivate()->PendingInstanceUpdates;
 }
 
 
 int32 UCustomizableObjectSystem::GetTotalInstances() const
 {
-	int32 NumInstances = 0;
-	
-	for (TObjectIterator<UCustomizableObjectInstance> Instance; Instance; ++Instance)
-	{
-		if (!IsValid(*Instance) ||
-			Instance->HasAnyFlags(RF_ClassDefaultObject))
-		{
-			continue;
-		}
-				
-		++NumInstances;
-	}
-	return NumInstances;
+	return UCustomizableObjectSystemPrivate::GetCustomizableObjectInstances().Num();
 }
 
 int64 UCustomizableObjectSystem::GetTextureMemoryUsed() const
@@ -3757,12 +3223,6 @@ int32 UCustomizableObjectSystem::GetAverageBuildTime() const
 }
 
 
-int32 UCustomizableObjectSystem::GetSkeletalMeshMinLODQualityLevel() const
-{
-	return GetPrivate()->SkeletalMeshMinLodQualityLevel;
-}
-
-
 bool UCustomizableObjectSystem::IsSupport16BitBoneIndexEnabled() const
 {
 	return GetPrivate()->bSupport16BitBoneIndex;
@@ -3781,15 +3241,133 @@ void UCustomizableObjectSystem::SetProgressiveMipStreamingEnabled(bool bIsEnable
 }
 
 
-bool UCustomizableObjectSystem::IsOnlyGenerateRequestedLODsEnabled() const
+void UCustomizableObjectSystem::SetGenerateInstancesWithinRange(bool bActive)
 {
-	return GetPrivate()->EnableOnlyGenerateRequestedLODs != 0;
+	GetPrivate()->bGenerateInstancesWithinRange = bActive;
 }
 
 
-void UCustomizableObjectSystem::SetOnlyGenerateRequestedLODsEnabled(bool bIsEnabled)
+void UCustomizableObjectSystem::SetInstanceGenerationRange(float GenerationRange)
 {
-	GetPrivate()->EnableOnlyGenerateRequestedLODs = bIsEnabled ? 1 : 0;
+	GetPrivate()->GenerationRangeSquare = GenerationRange * GenerationRange;
+}
+
+
+void UCustomizableObjectSystem::AddViewCenter(const AActor* const InViewCenter)
+{
+	if (IsValid(InViewCenter))
+	{
+		GetPrivate()->ViewCenters.Add(InViewCenter);
+	}
+}
+
+
+void UCustomizableObjectSystem::RemoveViewCenter(const AActor* const InViewCenter)
+{
+	if (IsValid(InViewCenter))
+	{
+		GetPrivate()->ViewCenters.Remove(InViewCenter);
+	}
+}
+
+
+void UCustomizableObjectSystem::ClearViewCenters()
+{
+	GetPrivate()->ViewCenters.Empty();
+}
+
+
+void UCustomizableObjectSystemPrivate::UpdateInstanceDistances()
+{
+	MUTABLE_CPUPROFILER_SCOPE(UpdateInstanceDistances);
+
+	if (CustomizableObjectInstanceUsages.IsEmpty())
+	{
+		return;
+	}
+
+	UCustomizableObjectSystem* CustomizableObjectSystem = UCustomizableObjectSystem::GetInstance();
+
+	TArray<const FVector> ViewLocations;
+	ViewLocations.Reset(ViewCenters.Num());
+
+	const int32 NumViewCenters = ViewCenters.Num();
+	for (const TWeakObjectPtr<const AActor> ViewCenter : ViewCenters)
+	{
+		if (ViewCenter.IsValid())
+		{
+			ViewLocations.Add(ViewCenter->GetActorLocation());
+		}
+	}
+
+	if (ViewLocations.IsEmpty())
+	{
+		for (UCustomizableObjectInstanceUsage* CustomizableObjectInstanceUsage : CustomizableObjectInstanceUsages)
+		{
+			USkeletalMeshComponent* Parent = Cast<USkeletalMeshComponent>(CustomizableObjectInstanceUsage->GetAttachParent());
+			if (!Parent || !Parent->GetWorld())
+			{
+				continue;
+			}
+
+			APlayerController* FirstLocalPlayerController = Parent->GetWorld()->GetFirstPlayerController();
+			if (FirstLocalPlayerController && FirstLocalPlayerController->GetPawn())
+			{
+				ViewLocations.Add(FirstLocalPlayerController->GetPawn()->GetActorLocation());
+				break;
+			}
+		}
+	}
+
+	if (ViewLocations.Num())
+	{
+		for (UCustomizableObjectInstanceUsage* CustomizableObjectInstanceUsage : CustomizableObjectInstanceUsages)
+		{
+#if WITH_EDITOR
+			if (CustomizableObjectInstanceUsage->GetPrivate()->bIsNetModeDedicatedServer)
+			{
+				continue;
+			}
+
+#endif // WITH_EDITOR
+
+			UCustomizableObjectInstance* CustomizableObjectInstance = CustomizableObjectInstanceUsage->GetCustomizableObjectInstance();
+			if (!CustomizableObjectInstance)
+			{
+				continue;
+			}
+
+			const USkeletalMeshComponent* SkeletalMeshComponent = CustomizableObjectInstanceUsage->GetAttachParent();
+			if (!SkeletalMeshComponent)
+			{
+				continue;
+			}
+
+			const AActor* ParentActor = SkeletalMeshComponent->GetAttachmentRootActor();
+			if (!IsValid(ParentActor))
+			{
+				continue;
+			}
+
+			UCustomizableInstancePrivate* InstancePrivate = CustomizableObjectInstance->GetPrivate();
+
+			const APawn* Pawn = Cast<const APawn>(ParentActor);
+			if (Pawn && Pawn->IsPlayerControlled())
+			{
+				InstancePrivate->SetCOInstanceFlags(ECOInstanceFlags::UsedByPlayerOrNearIt);
+				InstancePrivate->MinSquareDistFromComponentToPlayer = -1.0;
+				continue;
+			}
+
+			InstancePrivate->ClearCOInstanceFlags(ECOInstanceFlags::UsedByPlayerOrNearIt);
+
+			const FVector ActorLocation = ParentActor->GetActorLocation();
+			for (const FVector& ViewLocation : ViewLocations)
+			{
+				InstancePrivate->MinSquareDistFromComponentToPlayer = FMath::Min(FVector::DistSquared(ViewLocation, ActorLocation), InstancePrivate->MinSquareDistFromComponentToPlayer);
+			}
+		}
+	}
 }
 
 
@@ -3870,48 +3448,6 @@ int32 UCustomizableObjectSystem::GetWorkingMemory() const
 }
 
 
-#if WITH_EDITOR
-
-uint64 UCustomizableObjectSystem::GetMaxChunkSizeForPlatform(const ITargetPlatform* TargetPlatform)
-{
-	if (!TargetPlatform || !TargetPlatform->RequiresCookedData())
-	{
-		return MAX_uint64;
-	}
-
-	const FString& PlatformName = TargetPlatform ? TargetPlatform->IniPlatformName() : FPlatformProperties::IniPlatformName();
-
-	if (const int64* CachedMaxChunkSize = GetPrivate()->PlatformMaxChunkSize.Find(PlatformName))
-	{
-		return *CachedMaxChunkSize;
-	}
-
-	int64 MaxChunkSize = -1;
-
-	if (!FParse::Value(FCommandLine::Get(), TEXT("ExtraFlavorChunkSize="), MaxChunkSize) || MaxChunkSize < 0)
-	{
-		FConfigFile PlatformIniFile;
-		FConfigCacheIni::LoadLocalIniFile(PlatformIniFile, TEXT("Game"), true, *PlatformName);
-		FString ConfigString;
-		if (PlatformIniFile.GetString(TEXT("/Script/UnrealEd.ProjectPackagingSettings"), TEXT("MaxChunkSize"), ConfigString))
-		{
-			MaxChunkSize = FCString::Atoi64(*ConfigString);
-		}
-	}
-
-	// If no limit is specified default it to MUTABLE_STREAMED_DATA_MAXCHUNKSIZE 
-	if (MaxChunkSize <= 0)
-	{
-		MaxChunkSize = MUTABLE_STREAMED_DATA_MAXCHUNKSIZE;
-	}
-
-	GetPrivate()->PlatformMaxChunkSize.Add(PlatformName, MaxChunkSize);
-
-	return MaxChunkSize;
-}
-#endif // WITH_EDITOR
-
-
 void UCustomizableObjectSystemPrivate::OnMutableEnabledChanged(IConsoleVariable* MutableEnabled)
 {
 	if (!UCustomizableObjectSystem::IsCreated())
@@ -3937,40 +3473,127 @@ void UCustomizableObjectSystemPrivate::OnMutableEnabledChanged(IConsoleVariable*
 }
 
 
-void UCustomizableObjectSystemPrivate::StartUpdateSkeletalMesh(const TSharedRef<FUpdateContextPrivate>& Context)
+bool UCustomizableObjectSystemPrivate::TryStartUpdate(const TSharedRef<FUpdateContextPrivate>& Context)
 {
+	MUTABLE_CPUPROFILER_SCOPE(TryStartUpdate);
+
 	check(!CurrentMutableOperation); // Can not start an update if there is already another in progress
 	check(Context->bValid);
-	check(Context->Instance.IsValid()) // The COI has to be alive to start the update
 	check(Context->Object.IsValid()) // The CO has to be alive to start the update
+	check(Context->Object->IsCompiled()) // The CO has to be compiled to start the update
 
+	const UCustomizableObjectInstance* Instance = Context->Instance.Get();
 	const UCustomizableObject* Object = Context->Object.Get();
 
+	if (!IsValid(Instance))
+	{
+		Context->UpdateResult = EUpdateResult::Error;
+		FinishUpdateGlobal(Context);
+		return false;
+	}
+
+	UCustomizableInstancePrivate* InstancePrivate = Instance->GetPrivate();
+	InstancePrivate->ClearCOInstanceFlags(ECOInstanceFlags::HasPendingHighPriorityUpdate | ECOInstanceFlags::HasPendingLowPriorityUpdate);
+
+	// Only update resources if the instance is in range (it could have got far from the player since the task was queued)
+	if (IsDiscardedByDistance(InstancePrivate))
+	{
+		Context->UpdateResult = EUpdateResult::ErrorDiscarded;
+		FinishUpdateGlobal(Context);
+		return false;
+	}
+
+	// Skip update, the requested update is equal to the running update.
+	if (Context->GetCapturedDescriptorHash() == InstancePrivate->CommittedDescriptorHash)
+	{
+		Context->UpdateResult = EUpdateResult::Success;
+		UpdateSkeletalMesh(Context);
+		FinishUpdateGlobal(Context);
+		return false;
+	}
+
+	Context->Parameters = Context->GetCapturedDescriptor().GetParameters();
+
+	StartUpdate(Context);
+
+	return true;
+}
+
+
+bool UCustomizableObjectSystemPrivate::TryStartAutomaticUpdateForInstance(UCustomizableObjectInstance& Instance)
+{
+	MUTABLE_CPUPROFILER_SCOPE(TryStartAutomaticUpdateForInstance);
+
+	check(!CurrentMutableOperation); // Can not start an update if there is already another in progress
+
+	UCustomizableObject* Object = Instance.GetCustomizableObject();
+	if (!IsValid(Object))
+	{
+		return false;
+	}
+
+#if WITH_EDITOR
+	if (!Object->IsCompiled() &&
+		Object->GetPrivate()->CompilationResult != ECompilationResultPrivate::Errors) // Avoid constantly retry failed compilations.
+	{
+		if (ICustomizableObjectEditorModule* EditorModule = ICustomizableObjectEditorModule::Get())
+		{
+			FCompileParams CompileParams;
+			CompileParams.bSkipIfCompiled = true;
+
+			EditorModule->CompileCustomizableObject(*Object, &CompileParams, true, false);
+		}
+	}
+#endif
+
+	UCustomizableInstancePrivate* InstancePrivate = Instance.GetPrivate();
+	if (!InstancePrivate->CanUpdateInstance())
+	{
+		return false;
+	}
+
+	const TSharedRef<FUpdateContextPrivate> Context = MakeShared<FUpdateContextPrivate>(Instance, InstancePrivate->GetDescriptor());
+	if (!Context->IsContextValid())
+	{
+		return false;
+	}
+
+	Context->Parameters = Context->GetCapturedDescriptor().GetParameters();
+
+	StartUpdate(Context);
+
+	return true;
+}
+
+
+void UCustomizableObjectSystemPrivate::StartUpdate(const TSharedRef<FUpdateContextPrivate>& Context)
+{
 	if (FPlatformTime::Seconds() > LogStartedUpdateUnmute)
 	{
-		const FName ObjectName = Object ? Object->GetFName() : NAME_None;
-		const FName InstanceName = Context->Instance->GetFName();
-		FString Message = FString::Printf(TEXT("Started Update Skeletal Mesh Async. CustomizableObject=%s Instance=%s, Frame=%d"), *ObjectName.ToString(), *InstanceName.ToString(), GFrameNumber);
-		LogUpdateMessage(Message, ELogVerbosity::Log, Context->Instance.Get());
+		MUTABLE_CPUPROFILER_SCOPE(LogStartedUpdate);
 
-		const float CurrentTime = FPlatformTime::Seconds();
+		UCustomizableObjectInstance* Instance = Context->Instance.Get();
+		const UCustomizableObject* Object = Context->Object.Get();
+		const FName AssetName = !Instance->HasAnyFlags(RF_Transient) ? Instance->GetFName() : Object->GetFName();
+
+		FString Message = FString::Printf(TEXT("Started Update Skeletal Mesh Async. CO/COI = %s"), *AssetName.ToString());
+		LogUpdateMessage(Message, ELogVerbosity::Log, Instance);
 
+		const float CurrentTime = FPlatformTime::Seconds();
 		constexpr float LogInterval = 1.0 / 2.0; // Allow maximum 2 logs per second.
 		const bool bMute = CurrentTime - LogStartedUpdateLast < LogInterval;
 		LogStartedUpdateLast = CurrentTime;
-			
+
 		if (bMute)
 		{
-			constexpr float MuteTime = 5.0;
-			UE_LOG(LogMutable, Log, TEXT("Disabling \"Started Update Skeletal Mesh Async\" log during %f seconds due to spam"), MuteTime);
-			LogStartedUpdateUnmute = CurrentTime + MuteTime;
+			UE_LOGF(LogMutable, Log, "Disabling \"Started Update Skeletal Mesh Async\" log during 5 seconds due to spam");
+			LogStartedUpdateUnmute = CurrentTime + 5.0;
 		}
 	}
 
 	// It is safe to do this now.
 	UpdateMemoryLimit();
-	
-	check(!CurrentMutableOperation);
+
 	CurrentMutableOperation = Context;
 
 	MutableTaskGraph.AddGameThreadTask(
@@ -3990,11 +3613,6 @@ bool UCustomizableObjectSystemPrivate::IsUpdating(const UCustomizableObjectInsta
 	{
 		return true;
 	}
-
-	if (MutablePendingInstanceWork.GetUpdate(TWeakObjectPtr<const UCustomizableObjectInstance>(&Instance)))
-	{
-		return true;
-	}
 	
 	return false;
 }
@@ -4043,20 +3661,20 @@ void UCustomizableObjectSystemPrivate::LogUpdateMessage(const FString& Message,
 	switch (Verbosity)
 	{
 	case ELogVerbosity::Error:
-		UE_LOG(LogMutable, Error, TEXT("%s"), *Message);
+		UE_LOGF(LogMutable, Error, "%ls", *Message);
 		break;
 	case ELogVerbosity::Warning:
-		UE_LOG(LogMutable, Warning, TEXT("%s"), *Message);
+		UE_LOGF(LogMutable, Warning, "%ls", *Message);
 		break;
 	case ELogVerbosity::Verbose:
-		UE_LOG(LogMutable, Verbose, TEXT("%s"), *Message);
+		UE_LOGF(LogMutable, Verbose, "%ls", *Message);
 		break;
 	case ELogVerbosity::Display:
-		UE_LOG(LogMutable, Verbose, TEXT("%s"), *Message);
+		UE_LOGF(LogMutable, Verbose, "%ls", *Message);
 		break;
 	case ELogVerbosity::Log:
 	default:
-		UE_LOG(LogMutable, Log, TEXT("%s"), *Message);
+		UE_LOGF(LogMutable, Log, "%ls", *Message);
 		break;
 	}
 
@@ -4123,7 +3741,7 @@ void ReconstructMorphTargetsFromMeshCompressedData(const UE::Mutable::Private::F
 
 			if (MorphNumBatches > 0)
 			{
-				// Get the Offset value for the morph traget first batch.
+				// Get the Offset value for the morph target first batch.
 				MorphDeltaDataBeginInDwords = Mesh.MorphDataBuffer[BatchStartOffset * NumBatchHeaderDwords] / sizeof(uint32);
 
 				for (uint32 BatchHeaderIndex = 0; BatchHeaderIndex < MorphNumBatches; ++BatchHeaderIndex)
diff --git a/Engine/Plugins/Mutable/Source/CustomizableObject/Private/MuCO/DescriptorHash.cpp b/Engine/Plugins/Mutable/Source/CustomizableObject/Private/MuCO/DescriptorHash.cpp
--- a/Engine/Plugins/Mutable/Source/CustomizableObject/Private/MuCO/DescriptorHash.cpp
+++ b/Engine/Plugins/Mutable/Source/CustomizableObject/Private/MuCO/DescriptorHash.cpp
@@ -66,104 +66,10 @@ FDescriptorHash::FDescriptorHash(const FCustomizableObjectInstanceDescriptor& De
 
 	Hash = HashCombine(Hash, GetTypeHash(Descriptor.State));
 	Hash = HashCombine(Hash, GetTypeHash(Descriptor.GetBuildParameterRelevancy()));
-	
-	bStreamingEnabled = Descriptor.bStreamingEnabled;
-	MinLODs = Descriptor.MinLOD;
-	QualitySettingMinLODs = Descriptor.QualitySettingMinLODs;
-	FirstRequestedLOD = Descriptor.FirstRequestedLOD;
-}
-
-
-bool FDescriptorHash::IsSubset(const FDescriptorHash& Other) const
-{
-	if (Hash != Other.Hash)
-	{
-		return false;
-	}
-
-	if (bStreamingEnabled && Other.bStreamingEnabled)
-	{
-		return true;
-	}
-
-	for (const TTuple<FName, uint8>& Pair : MinLODs)
-	{
-		const uint8* Result = Other.MinLODs.Find(Pair.Key);
-		if (!Result)
-		{
-			return false;
-		}
-
-		if (Pair.Value < *Result)
-		{
-			return false;
-		}
-	}
-	
-	// Scalability quality change : Do not care if it is a subset, always update
-	for (const TTuple<FName, uint8>& Pair : QualitySettingMinLODs)
-	{
-		const uint8* Result = Other.QualitySettingMinLODs.Find(Pair.Key);
-		if (!Result)
-		{
-			return false;
-		}
-
-		if (Pair.Value != *Result)
-		{
-			return false;
-		}
-	}
-
-	if (FirstRequestedLOD.Num() != Other.FirstRequestedLOD.Num())
-	{
-		return false;
-	}
-	
-	for (const TTuple<FName, uint8>& Pair : FirstRequestedLOD)
-	{
-		const uint8* PairOther = Other.FirstRequestedLOD.Find(Pair.Key);
-		if (!PairOther)
-		{
-			return false;
-		}
-
-		if (Pair.Value < *PairOther)
-		{
-			return false;
-		}
-	}
-
-	return true;
 }
 
 
 FString FDescriptorHash::ToString() const
 {
-	TStringBuilder<150> Builder;
-
-	Builder.Appendf(TEXT("(Hash=%u,"), Hash);
-	
-	Builder.Appendf(TEXT("MinLODs=["));
-	for (const TTuple<FName, uint8>& MinLOD : MinLODs)
-	{
-		Builder.Appendf(TEXT("(%s, %i),"), *MinLOD.Key.ToString(), MinLOD.Value);
-	}
-	Builder.Appendf(TEXT("])"));
-
-	Builder.Appendf(TEXT("FirstLODAvailable=["));
-	for (const TTuple<FName, uint8>& QualitySetting : QualitySettingMinLODs)
-	{
-		Builder.Appendf(TEXT("(%s, %i),"), *QualitySetting.Key.ToString(), QualitySetting.Value);
-	}
-	Builder.Appendf(TEXT("])"));
-
-	Builder.Appendf(TEXT("FirstRequestedLOD=["));
-	for (const TTuple<FName, uint8>& RequestedLODs : FirstRequestedLOD)
-	{
-		Builder.Appendf(TEXT("(%s, %i),"), *RequestedLODs.Key.ToString(), RequestedLODs.Value);
-	}
-	Builder.Appendf(TEXT("])"));
-
-	return Builder.ToString();
+	return FString::Printf(TEXT("(Hash=%u,"), Hash);
 }
diff --git a/Engine/Plugins/Mutable/Source/CustomizableObject/Public/MuCO/CustomizableInstanceLODManagement.h b/Engine/Plugins/Mutable/Source/CustomizableObject/Public/MuCO/CustomizableInstanceLODManagement.h
deleted file mode 100644
diff --git a/Engine/Plugins/Mutable/Source/CustomizableObject/Public/MuCO/CustomizableObjectInstance.h b/Engine/Plugins/Mutable/Source/CustomizableObject/Public/MuCO/CustomizableObjectInstance.h
--- a/Engine/Plugins/Mutable/Source/CustomizableObject/Public/MuCO/CustomizableObjectInstance.h
+++ b/Engine/Plugins/Mutable/Source/CustomizableObject/Public/MuCO/CustomizableObjectInstance.h
@@ -220,9 +220,6 @@ struct FBakingConfiguration
 };
 
 
-// Priority for the mutable update queue, Low is the normal distance-based priority, High is normally used for discards and Mid for LOD downgrades
-enum class EQueuePriorityType : uint8 { High, Med, Med_Low, Low };
-
 /** Result of all the checks just before beginning an update. */
 enum class EUpdateRequired : uint8
 {
@@ -243,7 +240,9 @@ enum class EUpdateResult : uint8
 	ErrorOptimized, // The update was skipped since its result would have been the same as the current customization.
 	ErrorReplaced, // The update was replaced by a newer update request.
 	ErrorDiscarded, // The update was not finished since due to the LOD management discarding the data.
-	Error16BitBoneIndex // The update finish unsuccessfully due to Instance not supporting 16 Bit Bone Indexing required by the Engine.
+	Error16BitBoneIndex, // The update finish unsuccessfully due to Instance not supporting 16 Bit Bone Indexing required by the Engine.
+
+	ErrorUncompiled // The update failed due to the CO being uncompiled
 };
 
 
@@ -776,15 +775,6 @@ public:
 	/** If true, Mutable will keep the ownership of the generated resources. This means that Mutable may decide to modify or reuse them for performance considerations. */
 	UFUNCTION(BlueprintCallable, Category = CustomizableObjectInstance)
 	CUSTOMIZABLEOBJECT_API void SetKeepOwnershipOfGeneratedResources(bool bEnabled);
-	
-	// The following methods should only be used in an LOD management class
-	CUSTOMIZABLEOBJECT_API void SetIsBeingUsedByComponentInPlay(bool bIsUsedByComponent);
-	CUSTOMIZABLEOBJECT_API bool GetIsBeingUsedByComponentInPlay() const;
-	CUSTOMIZABLEOBJECT_API void SetIsDiscardedBecauseOfTooManyInstances(bool bIsDiscarded);
-	CUSTOMIZABLEOBJECT_API bool GetIsDiscardedBecauseOfTooManyInstances() const;
-	CUSTOMIZABLEOBJECT_API void SetIsPlayerOrNearIt(bool NewValue);
-	CUSTOMIZABLEOBJECT_API float GetMinSquareDistToPlayer() const;
-	CUSTOMIZABLEOBJECT_API void SetMinSquareDistToPlayer(float NewValue);
 
 	/** Return the total number of components that can be generated with the CustomizableObject (unrelated to this instance parameters). 
 	* DEPRECATED: Get the CO with GetCustomizableObect and get the components there with GetComponentCount.
@@ -793,13 +783,6 @@ public:
 	UFUNCTION(BlueprintCallable, Category = CustomizableObjectInstance)
 	CUSTOMIZABLEOBJECT_API int32 GetNumComponents() const;
 	
-	/** Sets an array of LODs to generate per component. Mutable will generate those plus the currently generated LODs (if any).
-	  * Requires mutable.EnableOnlyGenerateRequestedLODs and CurrentInstanceLODManagement->IsOnlyGenerateRequestedLODLevelsEnabled() to be true.
-	  * @param InMinLODs - MinLOD to generate per Component.
-	  * @param InFirstRequestedLOD - Requested LODs per component to generate. Key is the Component name. Value represents the first LOD to generate. Zero means all LODs.
-	  * @param InOutRequestedUpdates - Map from Instance to Update data that stores a request for the Instance to be updated, which will be either processed or discarded by priority (to be rerequested the next tick) */
-	CUSTOMIZABLEOBJECT_API void SetRequestedLODs(const TMap<FName, uint8>& InMinLODs, const TMap<FName, uint8>& InFirstRequestedLOD, FMutableInstanceUpdateMap& InOutRequestedUpdates);
-
 #if WITH_EDITOR
 	/**
 	 * Performs the baking of the instance resources in an async fashion. Bind yourself to the callback present in InBakingConfiguration to get notified in case it fails
diff --git a/Engine/Plugins/Mutable/Source/CustomizableObject/Public/MuCO/CustomizableObjectInstanceDescriptor.h b/Engine/Plugins/Mutable/Source/CustomizableObject/Public/MuCO/CustomizableObjectInstanceDescriptor.h
--- a/Engine/Plugins/Mutable/Source/CustomizableObject/Public/MuCO/CustomizableObjectInstanceDescriptor.h
+++ b/Engine/Plugins/Mutable/Source/CustomizableObject/Public/MuCO/CustomizableObjectInstanceDescriptor.h
@@ -74,10 +74,6 @@ struct FCustomizableObjectInstanceDescriptor
 	
 	/** Update all parameters to be up to date with the Mutable Core parameters. */
 	UE_API void ReloadParameters();
-	
-	UE_API void SetFirstRequestedLOD(const TMap<FName, uint8>& FirstRequestedLOD);
-
-	UE_API const TMap<FName, uint8>& GetFirstRequestedLOD() const;
 
 	// ------------------------------------------------------------
 	// FParameters
@@ -265,18 +261,6 @@ struct FCustomizableObjectInstanceDescriptor
 	/** If this is set to true, when updating the instance an additional step will be performed to calculate the list of instance parameters that are relevant for the current parameter values. */
 	bool bBuildParameterRelevancy = false;
 
-	bool bStreamingEnabled = false;
-
-	/** These are the LODs Mutable can generate based on user input and quality settings, they MUST NOT be used in an update (Mutable thread). */
-	TMap<FName, uint8> MinLOD;
-
-	/** These are the MinLODs based on the active quality setting */
-	TMap<FName, uint8> QualitySettingMinLODs;
-
-	/** RequestedLODs per component to generate.
-	 * Key is the Component name. Value represents the first LOD to generate. Zero means all LODs. */
-	TMap<FName, uint8> FirstRequestedLOD;
-	
 	// Friends
 	friend FDescriptorHash;
 	friend UCustomizableObjectInstance;
diff --git a/Engine/Plugins/Mutable/Source/CustomizableObject/Public/MuCO/CustomizableObjectSystem.h b/Engine/Plugins/Mutable/Source/CustomizableObject/Public/MuCO/CustomizableObjectSystem.h
--- a/Engine/Plugins/Mutable/Source/CustomizableObject/Public/MuCO/CustomizableObjectSystem.h
+++ b/Engine/Plugins/Mutable/Source/CustomizableObject/Public/MuCO/CustomizableObjectSystem.h
@@ -14,7 +14,6 @@
 #define UE_API CUSTOMIZABLEOBJECT_API
 
 class IConsoleVariable;
-class UCustomizableInstanceLODManagementBase;
 class ITargetPlatform;
 class UCustomizableObject;
 class USkeletalMesh;
@@ -151,9 +150,6 @@ public:
 	UE_API virtual FString GetDesc() override;
 	// End UObject interface.
 
-	UE_API bool IsReplaceDiscardedWithReferenceMeshEnabled() const;
-	UE_API void SetReplaceDiscardedWithReferenceMeshEnabled(bool bIsEnabled);
-
 #if WITH_EDITOR
 	// Lock a CustomizableObjects, preventing the generation or update of any of its instances
 	// Will return true if successful, false if it fails to lock because an update is already underway
@@ -181,17 +177,11 @@ public:
 	// If true, uncompiled Customizable Objects will be compiled synchronously
 	UE_API bool IsAutoCompilationSync() const;
 #endif
-	
-	// Return the current MinLodQualityLevel for skeletal meshes.
-	UE_API int32 GetSkeletalMeshMinLODQualityLevel() const;
 
 	UE_API bool IsSupport16BitBoneIndexEnabled() const;
 
 	UE_API bool IsProgressiveMipStreamingEnabled() const;
 	UE_API void SetProgressiveMipStreamingEnabled(bool bIsEnabled);
-
-	UE_API bool IsOnlyGenerateRequestedLODsEnabled() const;
-	UE_API void SetOnlyGenerateRequestedLODsEnabled(bool bIsEnabled);
 	
 	// Show a warning on-screen and via a notification (if in Editor) and log an error when a CustomizableObject is
 	// being used and it's not compiled.  Callers can add additional information to the error log.
@@ -201,11 +191,6 @@ public:
 	UE_API UCustomizableObjectSystemPrivate* GetPrivate();
 	UE_API const UCustomizableObjectSystemPrivate* GetPrivate() const;
 
-	UE_API UCustomizableInstanceLODManagementBase* GetInstanceLODManagement() const;
-
-	// Pass a null ptr to reset to the default InstanceLODManagement
-	UE_API void SetInstanceLODManagement(UCustomizableInstanceLODManagementBase* NewInstanceLODManagement);
-
 	// Find out the version of the plugin
 	UFUNCTION(BlueprintCallable, Category = Status)
 	UE_API FString GetPluginVersion() const;
@@ -230,15 +215,6 @@ public:
 	UFUNCTION(BlueprintCallable, Category = Status)
 	UE_API int32 GetAverageBuildTime() const;
 
-#if WITH_EDITOR
-	// Get the maximum size a chunk can have on a specific platform. If unspecified return MUTABLE_STREAMED_DATA_MAXCHUNKSIZE
-	UE_API uint64 GetMaxChunkSizeForPlatform(const ITargetPlatform* InTargetPlatform);
-#endif
-
-private:
-	UE_API void DiscardInstances();
-
-public:
 	/** Return true if the instance is being updated. */
 	UFUNCTION(BlueprintCallable, BlueprintPure, Category = CustomizableObjectSystem)
 	UE_API bool IsUpdating(const UCustomizableObjectInstance* Instance) const;
@@ -260,6 +236,27 @@ public:
 	 */
 	static UE_API bool ShouldClearWorkingMemoryOnUpdateEnd();
 
+	// Instances within the generation range to the ViewCenters will be generated, 
+	// When an instance falls outside the generation range the generated resources (meshes, textured, etc) will be dereferenced, and meshes on the components cleared.
+	UFUNCTION(BlueprintCallable, Category = CustomizableObjectSystem)
+	UE_API void SetGenerateInstancesWithinRange(bool bActive);
+
+	// Instance generation range around view centers when mutable.GenerateInstancesWithinRange is true
+	UFUNCTION(BlueprintCallable, Category = CustomizableObjectSystem)
+	UE_API void SetInstanceGenerationRange(float Distance);
+
+	// Add ViewCenters. Used to generate instances around when generating instances within a range.
+	UFUNCTION(BlueprintCallable, Category = CustomizableObjectSystem)
+	UE_API void AddViewCenter(const AActor* const InCentralActor);
+
+	// Remove ViewCenters. Used to generate instances around when generating instances within a range.
+	UFUNCTION(BlueprintCallable, Category = CustomizableObjectSystem)
+	UE_API void RemoveViewCenter(const AActor* const InCentralActor);
+	
+	// Clear ViewCenters. Used to generate instances around when generating instances within a range.
+	UFUNCTION(BlueprintCallable, Category = CustomizableObjectSystem)
+	UE_API void ClearViewCenters();
+
 private:
 	UPROPERTY(Transient)
 	TObjectPtr<UCustomizableObjectSystemPrivate> Private = nullptr;
diff --git a/Engine/Plugins/Mutable/Source/CustomizableObjectEditor/Private/MuCOE/CustomizableInstanceDetails.cpp b/Engine/Plugins/Mutable/Source/CustomizableObjectEditor/Private/MuCOE/CustomizableInstanceDetails.cpp
--- a/Engine/Plugins/Mutable/Source/CustomizableObjectEditor/Private/MuCOE/CustomizableInstanceDetails.cpp
+++ b/Engine/Plugins/Mutable/Source/CustomizableObjectEditor/Private/MuCOE/CustomizableInstanceDetails.cpp
@@ -2479,22 +2479,12 @@ FReply FCustomizableInstanceDetails::OnPasteAllParameters()
 	FCustomizableObjectInstanceDescriptor& Descriptor = CustomInstance->GetPrivate()->GetDescriptor();
 	const UScriptStruct* Struct = Descriptor.StaticStruct();
 
-	UCustomizableObject* CustomizableObject = Descriptor.GetCustomizableObject();
-	
-	const TMap<FName, uint8> MinLOD = Descriptor.MinLOD;
-	const TMap<FName, uint8> FirstRequestedLOD = Descriptor.GetFirstRequestedLOD();
-	
 	FScopedTransaction LocalTransaction(LOCTEXT("OnPasteAllParameters", "Paste All Parameters"));
 	CustomInstance->Modify();
 
 	if (Struct->ImportText(*ClipText, &Descriptor, nullptr, 0, GLog, GetPathNameSafe(Struct)))
 	{
-		Descriptor.SetCustomizableObject(CustomizableObject);
-		
-		// Keep current LOD
-		Descriptor.MinLOD = MinLOD;
-		Descriptor.SetFirstRequestedLOD(FirstRequestedLOD);
-		
+		Descriptor.SetCustomizableObject(Descriptor.GetCustomizableObject());
 		UpdateInstance();
 	}
 	
diff --git a/Engine/Plugins/Mutable/Source/CustomizableObjectEditor/Private/MuCOE/CustomizableObjectBenchmarkingUtils.cpp b/Engine/Plugins/Mutable/Source/CustomizableObjectEditor/Private/MuCOE/CustomizableObjectBenchmarkingUtils.cpp
--- a/Engine/Plugins/Mutable/Source/CustomizableObjectEditor/Private/MuCOE/CustomizableObjectBenchmarkingUtils.cpp
+++ b/Engine/Plugins/Mutable/Source/CustomizableObjectEditor/Private/MuCOE/CustomizableObjectBenchmarkingUtils.cpp
@@ -36,15 +36,7 @@ bool CustomizableObjectBenchmarkingUtils::GenerateDeterministicSetOfInstances(UC
 	for (uint32 CustomizableObjectInstanceIndex = 0; CustomizableObjectInstanceIndex < InstancesPerState; CustomizableObjectInstanceIndex++)
 	{
 		if (UCustomizableObjectInstance* GeneratedInstance = TargetCustomizableObject.CreateInstance())
-		{
-			// Force generation of all LODS
-			TMap<FName, uint8> FirstRequestedLOD;
-			for (const FName ComponentName : ModelResources->ComponentNamesPerObjectComponent)
-			{
-				FirstRequestedLOD.Add(ComponentName, 0);
-			}
-			GeneratedInstance->GetPrivate()->GetDescriptor().SetFirstRequestedLOD(FirstRequestedLOD);
-				
+		{				
 			// Randomize instance values
 			GeneratedInstance->SetRandomValuesFromStream(RandomizationStream);
 				
diff --git a/Engine/Plugins/Mutable/Source/CustomizableObjectEditor/Private/MuCOE/CustomizableObjectEditorModule.cpp b/Engine/Plugins/Mutable/Source/CustomizableObjectEditor/Private/MuCOE/CustomizableObjectEditorModule.cpp
--- a/Engine/Plugins/Mutable/Source/CustomizableObjectEditor/Private/MuCOE/CustomizableObjectEditorModule.cpp
+++ b/Engine/Plugins/Mutable/Source/CustomizableObjectEditor/Private/MuCOE/CustomizableObjectEditorModule.cpp
@@ -12,8 +12,10 @@
 #include "Misc/DateTime.h"
 #include "Misc/Timespan.h"
 #include "HAL/FileManager.h"
+#include "LevelEditorViewport.h"
 #include "MuCO/CustomizableObjectSystem.h"		// For defines related to memory function replacements.
 #include "MuCO/CustomizableObjectInstanceUsage.h"
+#include "MuCO/CustomizableObjectInstanceUsagePrivate.h"
 #include "MuCO/ICustomizableObjectModule.h"		// For instance editor command utility function
 #include "MuCO/UnrealPortabilityHelpers.h"
 #include "MuCOE/CustomizableInstanceDetails.h"
@@ -104,7 +106,7 @@
 #include "Nodes/CustomizableObjectNodeTextureParameter.h"
 #include "Nodes/ExternalTypeCustomization.h"
 #include "PhysicsEngine/PhysicsAsset.h"
-
+#include "Components/SkeletalMeshComponent.h"
 
 class AActor;
 class FString;
@@ -619,7 +621,7 @@ struct FCompilationOutOfDateContext
 /** Async because work is split in between ticks. */
 void IsCompilationOutOfDate_Async(const TSharedRef<FCompilationOutOfDateContext>& Context)
 {
-	MUTABLE_CPUPROFILER_SCOPE(IsCompilationOutOfDate_Async)
+	MUTABLE_EDITOR_CPUPROFILER_SCOPE(IsCompilationOutOfDate_Async)
 
 	check(IsInGameThread());
 	
@@ -678,7 +680,7 @@ void IsCompilationOutOfDate_Async(const TSharedRef<FCompilationOutOfDateContext>
 
 void FCustomizableObjectEditorModule::IsCompilationOutOfDate(const UCustomizableObject& Object, bool bSkipIndirectReferences, float MaxTime, const IsCompilationOutOfDateCallback& Callback) const
 {
-	MUTABLE_CPUPROFILER_SCOPE(FCustomizableObjectEditorModule::IsCompilationOutOfDate)
+	MUTABLE_EDITOR_CPUPROFILER_SCOPE(FCustomizableObjectEditorModule::IsCompilationOutOfDate)
 
 	check(IsInGameThread());
 	
@@ -892,7 +894,7 @@ USkeletalMesh* FCustomizableObjectEditorModule::GetReferenceSkeletalMesh(const U
 
 TMap<FName, FGuid> FCustomizableObjectEditorModule::GetParticipatingObjects(const UCustomizableObject* Object, const FCompilationOptions* InOptions) const
 {
-	MUTABLE_CPUPROFILER_SCOPE(FCustomizableObjectEditorModule::GetParticipatingObjects);
+	MUTABLE_EDITOR_CPUPROFILER_SCOPE(FCustomizableObjectEditorModule::GetParticipatingObjects);
 
 	FCompilationOptions Options = InOptions ? *InOptions : Object->GetPrivate()->GetCompileOptions();
 	TSharedPtr<FMutableCompilationContext> CompilationContext = MakeShared<FMutableCompilationContext>(Object, nullptr, Options);
@@ -1213,6 +1215,53 @@ void FCustomizableObjectEditorModule::CompileCustomizableObject(UCustomizableObj
 }
 
 
+void FCustomizableObjectEditorModule::UpdateInstanceDistanceToLevelEditorCamera()
+{
+	MUTABLE_EDITOR_CPUPROFILER_SCOPE(UpdateInstanceDistanceToLevelEditorCamera);
+
+	for (UCustomizableObjectInstanceUsage* CustomizableObjectInstanceUsage : UCustomizableObjectSystemPrivate::GetCustomizableObjectInstanceUsages())
+	{
+		if (CustomizableObjectInstanceUsage->GetPrivate()->bIsNetModeDedicatedServer)
+		{
+			continue;
+		}
+
+		UCustomizableObjectInstance* CustomizableObjectInstance = CustomizableObjectInstanceUsage->GetCustomizableObjectInstance();
+		if (!CustomizableObjectInstance)
+		{
+			continue;
+		}
+
+		USkeletalMeshComponent* Parent = Cast<USkeletalMeshComponent>(CustomizableObjectInstanceUsage->GetAttachParent());
+		if (!Parent)
+		{
+			continue;
+		}
+		
+		UWorld* LocalWorld = Parent->GetWorld();
+		if (!LocalWorld || LocalWorld->WorldType != EWorldType::Editor)
+		{
+			continue;
+		}
+
+		AActor* ParentActor = Parent->GetAttachmentRootActor();
+		if (ParentActor)
+		{
+			for (FLevelEditorViewportClient* LevelVC : GEditor->GetLevelViewportClients())
+			{
+				if (LevelVC && LevelVC->IsPerspective() && LevelVC->GetWorld() == LocalWorld)
+				{
+					double SquareDist = FVector::DistSquared(LevelVC->GetViewLocation(), ParentActor->GetActorLocation());
+					CustomizableObjectInstance->GetPrivate()->MinSquareDistFromComponentToPlayer =
+						FMath::Min(SquareDist, CustomizableObjectInstance->GetPrivate()->MinSquareDistFromComponentToPlayer);
+					break;
+				}
+			}
+		}
+	}
+}
+
+
 void FCustomizableObjectEditorModule::Tick(bool bBlocking)
 {
 	if (TSet<TStrongObjectPtr<UCustomizableObject>>::TIterator It = ObjectsPendingLoad.CreateIterator())
@@ -1434,7 +1483,7 @@ TStatId FCustomizableObjectEditorModule::GetStatId() const
 
 void FCustomizableObjectEditorModule::TickCook(float DeltaTime, bool bCookCompete)
 {
-	MUTABLE_CPUPROFILER_SCOPE(FCustomizableObjectEditorModule::TickCook);
+	MUTABLE_EDITOR_CPUPROFILER_SCOPE(FCustomizableObjectEditorModule::TickCook);
 	UCustomizableObjectSystem::GetInstance()->GetPrivate()->MainTick(false);
 }
 
diff --git a/Engine/Plugins/Mutable/Source/CustomizableObjectEditor/Private/MuCOE/CustomizableObjectEditorModule.h b/Engine/Plugins/Mutable/Source/CustomizableObjectEditor/Private/MuCOE/CustomizableObjectEditorModule.h
--- a/Engine/Plugins/Mutable/Source/CustomizableObjectEditor/Private/MuCOE/CustomizableObjectEditorModule.h
+++ b/Engine/Plugins/Mutable/Source/CustomizableObjectEditor/Private/MuCOE/CustomizableObjectEditorModule.h
@@ -67,6 +67,7 @@ public:
 	virtual int32 GetNumCompileRequests() override;
 	virtual bool IsCompiling(const UCustomizableObject& Object) const override;
 	virtual void CompileCustomizableObject(UCustomizableObject& Object, const FCompileParams* Params, bool bSilent, bool bForce) override;
+	virtual void UpdateInstanceDistanceToLevelEditorCamera() override;
 	virtual void Tick(bool bBlocking) override;
 	virtual int32 GetRemainingWork() const override;
 	virtual void BeginCacheForCookedPlatformData(UCustomizableObject& Object, const ITargetPlatform* TargetPlatform) override;
diff --git a/Engine/Plugins/Mutable/Source/CustomizableObjectEditor/Private/MuCOE/CustomizableObjectInstanceBakingUtils.cpp b/Engine/Plugins/Mutable/Source/CustomizableObjectEditor/Private/MuCOE/CustomizableObjectInstanceBakingUtils.cpp
--- a/Engine/Plugins/Mutable/Source/CustomizableObjectEditor/Private/MuCOE/CustomizableObjectInstanceBakingUtils.cpp
+++ b/Engine/Plugins/Mutable/Source/CustomizableObjectEditor/Private/MuCOE/CustomizableObjectInstanceBakingUtils.cpp
@@ -421,11 +421,9 @@ void ScheduleInstanceUpdateForBaking(UCustomizableObjectInstance& InInstance, FI
 	const TSharedRef<FUpdateContextPrivate> Context = MakeShared<FUpdateContextPrivate>(InInstance);
 	Context->bForce = true;
 	Context->bBake = true;
-	Context->bIgnoreCloseDist = true;
-	Context->bForceHighPriority = true;
+	Context->bIsHighPriority = true;
 	Context->UpdateNativeCallback = InInstanceUpdateDelegate;
 	Context->bProgressiveMipStreamingEnabled = false;
-	Context->bIsOnlyGenerateRequestedLODsEnabled = false;
 	Context->PixelFormatOverride = UnrealPixelFormatFunc;
 
 	UCustomizableObjectSystemPrivate* SystemPrivate = UCustomizableObjectSystem::GetInstance()->GetPrivate();
diff --git a/Engine/Plugins/Mutable/Source/MutableRuntime/Internal/MuR/MutableTrace.h b/Engine/Plugins/Mutable/Source/MutableRuntime/Internal/MuR/MutableTrace.h
--- a/Engine/Plugins/Mutable/Source/MutableRuntime/Internal/MuR/MutableTrace.h
+++ b/Engine/Plugins/Mutable/Source/MutableRuntime/Internal/MuR/MutableTrace.h
@@ -10,6 +10,10 @@
 #define MUTABLE_CPUPROFILER_SCOPE(Name) \
 	TRACE_CPUPROFILER_EVENT_SCOPE(Mutable_##Name)
 
+/** Custom Mutable profiler scope. */
+#define MUTABLE_EDITOR_CPUPROFILER_SCOPE(Name) \
+	TRACE_CPUPROFILER_EVENT_SCOPE(MutableEditor_##Name)
+
 
 /** Simple class that saves the time the scope is alive in the given location. */
 class FMutableScopeTimer
