[Mutable] Fix undefined symbols for UCustomizableObjectSystemPrivate::GetCustomizableObjectInstanceUsages()

#rnx

#Synced-CL 51330672

diff --git a/Engine/Plugins/Mutable/Source/CustomizableObjectEditor/Private/MuCOE/CustomizableObjectEditorModule.cpp b/Engine/Plugins/Mutable/Source/CustomizableObjectEditor/Private/MuCOE/CustomizableObjectEditorModule.cpp
--- a/Engine/Plugins/Mutable/Source/CustomizableObjectEditor/Private/MuCOE/CustomizableObjectEditorModule.cpp
+++ b/Engine/Plugins/Mutable/Source/CustomizableObjectEditor/Private/MuCOE/CustomizableObjectEditorModule.cpp
@@ -1219,8 +1219,13 @@ void FCustomizableObjectEditorModule::UpdateInstanceDistanceToLevelEditorCamera(
 {
 	MUTABLE_EDITOR_CPUPROFILER_SCOPE(UpdateInstanceDistanceToLevelEditorCamera);
 
-	for (UCustomizableObjectInstanceUsage* CustomizableObjectInstanceUsage : UCustomizableObjectSystemPrivate::GetCustomizableObjectInstanceUsages())
+	for (TObjectIterator<UCustomizableObjectInstanceUsage> CustomizableObjectInstanceUsage; CustomizableObjectInstanceUsage; ++CustomizableObjectInstanceUsage)
 	{
+		if (!IsValid(*CustomizableObjectInstanceUsage) || CustomizableObjectInstanceUsage->IsTemplate())
+		{
+			continue;
+		}
+
 		if (CustomizableObjectInstanceUsage->GetPrivate()->bIsNetModeDedicatedServer)
 		{
 			continue;
