[Mutable] Invalidate DDC

#Synced-CL 51334569

diff --git a/Engine/Plugins/Mutable/Source/CustomizableObject/Internal/MuCO/CustomizableObjectPrivate.h b/Engine/Plugins/Mutable/Source/CustomizableObject/Internal/MuCO/CustomizableObjectPrivate.h
--- a/Engine/Plugins/Mutable/Source/CustomizableObject/Internal/MuCO/CustomizableObjectPrivate.h
+++ b/Engine/Plugins/Mutable/Source/CustomizableObject/Internal/MuCO/CustomizableObjectPrivate.h
@@ -1721,7 +1721,7 @@ public:
 	/** Editor Log list to print all the messages related to this CO. */
 	TSharedPtr<class IMutableEditorLogger> CompilationLogger;
 	
-	static constexpr int32 DerivedDataVersion = 0x5854e531;
+	static constexpr int32 DerivedDataVersion = 0x2850e5a1;
 };
 
 #if WITH_EDITOR
