[Mutable] Fix memory leak when converting skeletal meshes to Mutable at runtime.

#rnx

#Synced-CL 51332317

diff --git a/Engine/Plugins/Mutable/Source/CustomizableObject/Private/MuCO/UnrealConversionUtils.cpp b/Engine/Plugins/Mutable/Source/CustomizableObject/Private/MuCO/UnrealConversionUtils.cpp
--- a/Engine/Plugins/Mutable/Source/CustomizableObject/Private/MuCO/UnrealConversionUtils.cpp
+++ b/Engine/Plugins/Mutable/Source/CustomizableObject/Private/MuCO/UnrealConversionUtils.cpp
@@ -1504,8 +1504,10 @@ namespace UnrealConversionUtils
 									FMemoryReaderView Ar(TArrayView<uint8>(BulkData, IORequest->GetSize()), true);
 
 									constexpr uint8 DummyStripFlags = 0;
-									const bool bForceKeepCPUResources = true;
-									const bool bNeedsCPUAccess = true;
+
+									// NeedsCPUaccess must be false. Otherwise, the FSkeletalMeshLODRenderData buffers wouldn't free their resources, leading to a memory leak.
+									const bool bForceKeepCPUResources = false;
+									const bool bNeedsCPUAccess = false;
 									Context->StreamedLOD.SerializeStreamedData(Ar, const_cast<USkeletalMesh*>(SkeletalMesh.Get()), LODIndex, DummyStripFlags, bNeedsCPUAccess, bForceKeepCPUResources);
 
 									Context->DataLOD = &Context->StreamedLOD;
