[Mutable] Remove cyclic reference preventing the deletion of the FMeshConversionContext.

#rnx

#Synced-CL 51332317

diff --git a/Engine/Plugins/Mutable/Source/CustomizableObject/Private/MuCO/UnrealConversionUtils.cpp b/Engine/Plugins/Mutable/Source/CustomizableObject/Private/MuCO/UnrealConversionUtils.cpp
--- a/Engine/Plugins/Mutable/Source/CustomizableObject/Private/MuCO/UnrealConversionUtils.cpp
+++ b/Engine/Plugins/Mutable/Source/CustomizableObject/Private/MuCO/UnrealConversionUtils.cpp
@@ -1491,13 +1491,15 @@ namespace UnrealConversionUtils
 					UE::Tasks::FTaskEvent LoadedEvent(TEXT("MutableMeshParamLoadCompleted"));
 					UE::Tasks::AddNested(LoadedEvent);
 
-					FBulkDataIORequestCallBack RequestCallback = [LoadedEvent, SkeletalMesh, Context, LODIndex](bool bWasCancelled, IBulkDataIORequest* IORequest) mutable
+					FBulkDataIORequestCallBack RequestCallback = [LoadedEvent, SkeletalMesh, WeakContext = Context.ToWeakPtr(), LODIndex](bool bWasCancelled, IBulkDataIORequest* IORequest) mutable
 						{
 							if (!bWasCancelled)
 							{
 								uint8* BulkData = IORequest->GetReadResults();
 								check(BulkData != nullptr);	
 
+								TSharedPtr<FMeshConversionContext> Context = WeakContext.Pin();
+								if (Context)
 								{
 									MUTABLE_CPUPROFILER_SCOPE(MeshStreamingSerialization);
 
@@ -1512,6 +1514,10 @@ namespace UnrealConversionUtils
 
 									Context->DataLOD = &Context->StreamedLOD;
 								}
+								else
+								{
+									check(false); // Context is owned by the conversion task. It should be always valid.
+								}
 							}
 							else
 							{
