[Mutable] Fix possible crash in CloneOrTakeOver for non resource types when a null with multiple references is passed as argument.

#rnx
#Synced-CL 51342373

diff --git a/Engine/Plugins/Mutable/Source/MutableRuntime/Internal/MuR/Image.h b/Engine/Plugins/Mutable/Source/MutableRuntime/Internal/MuR/Image.h
--- a/Engine/Plugins/Mutable/Source/MutableRuntime/Internal/MuR/Image.h
+++ b/Engine/Plugins/Mutable/Source/MutableRuntime/Internal/MuR/Image.h
@@ -395,7 +395,10 @@ namespace UE::Mutable::Private
 		}
 		else
 		{
-			Result = Source->Clone();
+			if (Source)
+			{
+				Result = Source->Clone();
+			}
 		}
 
 		Source = nullptr;
@@ -413,7 +416,10 @@ namespace UE::Mutable::Private
 		}
 		else
 		{
-			Result = Source->Clone();
+			if (Source)
+			{
+				Result = Source->Clone();
+			}
 		}
 
 		Source = nullptr;
