[mutable] Removed Task_LockCache.
- By removing this GT task we remove one synchronization with the main thread, which should improve mutable update performance.
- Changed the approach to pinning the cached textures by using TryPin.
- Removed all references to the deleted method (performance reporting)

#rnx

#Synced-CL 51377091

diff --git a/Engine/Plugins/Mutable/Source/CustomizableObject/Internal/MuCO/CustomizableObjectPrivate.h b/Engine/Plugins/Mutable/Source/CustomizableObject/Internal/MuCO/CustomizableObjectPrivate.h
--- a/Engine/Plugins/Mutable/Source/CustomizableObject/Internal/MuCO/CustomizableObjectPrivate.h
+++ b/Engine/Plugins/Mutable/Source/CustomizableObject/Internal/MuCO/CustomizableObjectPrivate.h
@@ -171,8 +171,8 @@ public:
 
 	UE::FMutex Lock;
 	
-	UTexture2D* Get(const FId& Key);
-	
+	TWeakObjectPtr<UTexture2D> Get(const FId& Key);
+
 	void Add(const FId& Key, UTexture2D* Value);
 
 	void Remove(const UTexture2D& Value);
diff --git a/Engine/Plugins/Mutable/Source/CustomizableObject/Internal/MuCO/CustomizableObjectSystemPrivate.h b/Engine/Plugins/Mutable/Source/CustomizableObject/Internal/MuCO/CustomizableObjectSystemPrivate.h
--- a/Engine/Plugins/Mutable/Source/CustomizableObject/Internal/MuCO/CustomizableObjectSystemPrivate.h
+++ b/Engine/Plugins/Mutable/Source/CustomizableObject/Internal/MuCO/CustomizableObjectSystemPrivate.h
@@ -346,7 +346,6 @@ public:
 
 	double TaskGetMeshesTimeStart = 0.0;
 	double TaskGetMeshesTime = 0.0;
-	double TaskLockCacheTime = 0.0;
 	double TaskGetImagesTimeStart = 0.0;
 	double TaskGetImagesTime = 0.0;
 	double TaskConvertResourcesTime = 0.0;
diff --git a/Engine/Plugins/Mutable/Source/CustomizableObject/Internal/MuCO/LogBenchmarkUtil.h b/Engine/Plugins/Mutable/Source/CustomizableObject/Internal/MuCO/LogBenchmarkUtil.h
--- a/Engine/Plugins/Mutable/Source/CustomizableObject/Internal/MuCO/LogBenchmarkUtil.h
+++ b/Engine/Plugins/Mutable/Source/CustomizableObject/Internal/MuCO/LogBenchmarkUtil.h
@@ -87,7 +87,6 @@ struct FInstanceUpdateStats
 	double QueueTime = 0.0;
 	double UpdateTime = 0.0;
 	double TaskGetMeshesTime = 0.0;
-	double TaskLockCacheTime = 0.0;
 
 	double TaskGetImagesTime = 0.0;
 	double TaskConvertResourcesTime = 0.0;
@@ -117,7 +116,6 @@ struct FInstanceUpdateStatistics
 	FStatistic QueueTime;
 	FStatistic UpdateTime;
 	FStatistic TaskGetMeshTime;
-	FStatistic TaskLockCacheTime;
 	FStatistic TaskGetImagesTime;
 	FStatistic TaskConvertResourcesTime;
 	FStatistic TaskCallbacksTime;
diff --git a/Engine/Plugins/Mutable/Source/CustomizableObject/Private/MuCO/CustomizableObject.cpp b/Engine/Plugins/Mutable/Source/CustomizableObject/Private/MuCO/CustomizableObject.cpp
--- a/Engine/Plugins/Mutable/Source/CustomizableObject/Private/MuCO/CustomizableObject.cpp
+++ b/Engine/Plugins/Mutable/Source/CustomizableObject/Private/MuCO/CustomizableObject.cpp
@@ -1877,12 +1877,12 @@ void FSkeletonCache::Clear()
 }
 
 
-UTexture2D* FTextureCache::Get(const FId& Key)
+TWeakObjectPtr<UTexture2D> FTextureCache::Get(const FId& Key)
 {
 	UE::TScopeLock<UE::FMutex> LockGuard(Lock);
-	
-	const TWeakObjectPtr<UTexture2D>* Result = GeneratedTextures.Find(Key);
-	return Result ? Result->Get() : nullptr;
+
+	TWeakObjectPtr<UTexture2D>* FoundTexture = GeneratedTextures.Find(Key);
+	return FoundTexture ? *FoundTexture : TWeakObjectPtr<UTexture2D>();
 }
 
 
diff --git a/Engine/Plugins/Mutable/Source/CustomizableObject/Private/MuCO/CustomizableObjectInstance.cpp b/Engine/Plugins/Mutable/Source/CustomizableObject/Private/MuCO/CustomizableObjectInstance.cpp
--- a/Engine/Plugins/Mutable/Source/CustomizableObject/Private/MuCO/CustomizableObjectInstance.cpp
+++ b/Engine/Plugins/Mutable/Source/CustomizableObject/Private/MuCO/CustomizableObjectInstance.cpp
@@ -4122,11 +4122,12 @@ UE::Tasks::FTask LoadAdditionalAssetsAndData(const TSharedRef<FUpdateContextPriv
 	UE::Tasks::FTask StreamingTask = StreamRequest.Stream();
 	Prerequisites.Add(StreamingTask);
 
-	return UE::Tasks::Launch(TEXT("CaptureContext"), [Context]() {}, // Keep a reference to make sure allocated memory is always alive. TODO Probably not necessary since AdditionalAssetsAsyncLoaded already captrues Context.
+	return UE::Tasks::Launch(TEXT("CaptureContext"), [Context]() {}, // Keep a reference to make sure allocated memory is always alive. TODO Probably not necessary since AdditionalAssetsAsyncLoaded already captures Context.
 	Prerequisites,
 	UE::Tasks::ETaskPriority::Inherit);
 }
 
+
 FCustomizableObjectInstanceDescriptor& UCustomizableInstancePrivate::GetDescriptor() const
 {
 	return GetPublic()->Descriptor;
@@ -4301,7 +4302,8 @@ UTexture* BuildTexture(const TSharedRef<FUpdateContextPrivate>& Context, const i
 	{
 		if (Context->bUseCaches)
 		{
-			Texture = TextureCache.Get(TextureCacheKey);
+			const TWeakObjectPtr<UTexture2D> WeakTexture = TextureCache.Get(TextureCacheKey);
+			Texture = WeakTexture.Get();
 			check(Texture);
 		}
 	}
diff --git a/Engine/Plugins/Mutable/Source/CustomizableObject/Private/MuCO/CustomizableObjectSystem.cpp b/Engine/Plugins/Mutable/Source/CustomizableObject/Private/MuCO/CustomizableObjectSystem.cpp
--- a/Engine/Plugins/Mutable/Source/CustomizableObject/Private/MuCO/CustomizableObjectSystem.cpp
+++ b/Engine/Plugins/Mutable/Source/CustomizableObject/Private/MuCO/CustomizableObjectSystem.cpp
@@ -2091,25 +2091,16 @@ namespace impl
 				// Access the CO TextureCache and try to grab the texture from there if found.
 				if (OperationData->bUseCaches)
 				{
-					TStrongObjectPtr<const UTexture> StrongCachedTexture = nullptr;
-					
 					FTextureCache& TextureCache = OperationData->Object->GetPrivate()->TextureCache;
 					const FTextureCache::FId TextureCacheKey = FTextureCache::FId(Image->ImageID, OperationData->MipsToSkip, OperationData->bBake);
-					if (const UTexture2D* CachedTexture = TextureCache.Get(TextureCacheKey))
-					{
-						FGCScopeTryGuard Guard;
-						if (Guard.LockSucceeded())
-						{
-							StrongCachedTexture = TStrongObjectPtr(CachedTexture);
-						}
-					}
-					
-					if (StrongCachedTexture)
+					const TWeakObjectPtr<UTexture2D> CachedWeakTexture = TextureCache.Get(TextureCacheKey);
+						
+					bool bWasPinSuccessful = false;
+					TStrongObjectPtr<const UTexture> CachedTexture = CachedWeakTexture.TryPin(bWasPinSuccessful);
+					if (CachedTexture.IsValid())
 					{
 						UE_LOG(LogMutable, VeryVerbose, TEXT("Texture resource with id [%s] is cached."), *Image->ImageID.ToString());
-						
-						OperationData->Objects.Add(StrongCachedTexture);
-						
+						OperationData->Objects.Add(CachedTexture);
 						continue;
 					}
 				}
@@ -2321,6 +2312,42 @@ namespace impl
 		{
 			UCustomizableInstancePrivate* CustomizableInstancePrivateData = CustomizableObjectInstance->GetPrivate();
 
+			// Prepare the component data
+			//-------------------------------------------------------------
+			{
+				UCustomizableObject* Object = OperationData->Object.Get();
+				check(Object);
+		
+				const UModelResources& ModelResources = Object->GetPrivate()->GetModelResourcesChecked();
+				
+				const int32 NumComponents = OperationData->MutableInstance->GetComponentCount();
+				CustomizableInstancePrivateData->ComponentsData.SetNum(NumComponents);
+
+				TArray<FName> ComponentNames;
+				ComponentNames.Reserve(NumComponents);
+
+				for (int32 ComponentIndex = 0; ComponentIndex < NumComponents; ++ComponentIndex)
+				{
+					const UE::Mutable::Private::FComponentId ComponentId = OperationData->MutableInstance->GetComponentId(ComponentIndex);
+					const FName& ComponentName = ModelResources.ComponentNamesPerObjectComponent[ComponentId];
+
+					if (ComponentNames.Contains(ComponentName)) // Components with the same name. Abort.
+					{
+						UE_LOG(LogMutable, Error, TEXT("Update error. Trying to generate multiple components with the same name: %s"),
+							*ComponentName.ToString());
+						
+						OperationData->UpdateResult = EUpdateResult::Error;
+						FinishUpdateGlobal(OperationData);
+						return;
+					}
+
+					ComponentNames.Add(ComponentName);
+
+					CustomizableInstancePrivateData->ComponentsData[ComponentIndex].ComponentId = ComponentId;
+					CustomizableInstancePrivateData->ComponentsData[ComponentIndex].ComponentName = ComponentName;
+				}
+			}
+			
 			// Convert Step
 			//-------------------------------------------------------------
 
@@ -2396,119 +2423,6 @@ namespace impl
 			},
 			false);
 	}
-
-
-	/** Lock Cached Resources. */
-	void Task_Game_LockCache(const TSharedRef<FUpdateContextPrivate>& OperationData)
-	{
-		MUTABLE_CPUPROFILER_SCOPE(Task_Game_LockCache)
-		FMutableScopeTimer Timer(OperationData->TaskLockCacheTime);
-
-		check(IsInGameThread());
-
-		UCustomizableObjectSystem* System = UCustomizableObjectSystem::GetInstance();
-		if (!IsValid(System))
-		{
-			OperationData->UpdateResult = EUpdateResult::Error;
-			FinishUpdateGlobal(OperationData);
-			return;
-		}
-		
-		UCustomizableObjectSystemPrivate* SystemPrivate = System->GetPrivate();
-
-		UCustomizableObject* CustomizableObject = OperationData->Object.Get(); 
-		if (!IsValid(CustomizableObject))
-		{
-			OperationData->UpdateResult = EUpdateResult::Error;
-			FinishUpdateGlobal(OperationData);
-			return;
-		}
-		
-		UCustomizableObjectInstance* ObjectInstance = OperationData->Instance.Get();
-		if (!IsValid(ObjectInstance))
-		{
-			OperationData->UpdateResult = EUpdateResult::Error;
-			FinishUpdateGlobal(OperationData);
-			return;
-		}
-
-		UCustomizableInstancePrivate* ObjectInstancePrivateData = ObjectInstance->GetPrivate();
-
-		// Any external texture that may be needed for this update will be requested from Mutable Core's GetImage
-		// which will safely access the GlobalExternalImages map, and then just get the cached image or issue a disk read
-
-		const UModelResources& ModelResources = CustomizableObject->GetPrivate()->GetModelResourcesChecked();
-
-		const int32 NumComponents = OperationData->MutableInstance->GetComponentCount();
-		ObjectInstancePrivateData->ComponentsData.SetNum(NumComponents);
-		
-		TArray<FName> ComponentNames;
-		ComponentNames.Reserve(NumComponents);
-
-		for (int32 ComponentIndex = 0; ComponentIndex < NumComponents; ++ComponentIndex)
-		{
-			const UE::Mutable::Private::FComponentId ComponentId = OperationData->MutableInstance->GetComponentId(ComponentIndex);
-			const FName& ComponentName = ModelResources.ComponentNamesPerObjectComponent[ComponentId];
-			
-			if (ComponentNames.Contains(ComponentName)) // Components with the same name. Abort.
-			{
-				UE_LOG(LogMutable, Error, TEXT("Update error. Trying to generate multiple components with the same name: %s"), *ComponentName.ToString());
-
-				OperationData->UpdateResult = EUpdateResult::Error;
-				FinishUpdateGlobal(OperationData);
-				return;
-			}
-			
-			ComponentNames.Add(ComponentName);
-			
-			ObjectInstancePrivateData->ComponentsData[ComponentIndex].ComponentId = ComponentId;
-			ObjectInstancePrivateData->ComponentsData[ComponentIndex].ComponentName = ComponentName;
-		}
-		
-		UE::Tasks::FTask Mutable_GetImageDescriptors_Task = SystemPrivate->MutableTaskGraph.AddMutableThreadTask(
-		TEXT("Task_Mutable_GetImageDescriptors"),
-		[OperationData]()
-		{
-			Task_Mutable_GetImageDescriptors(OperationData);
-		});
-		
-		UE::Tasks::FTask Mutable_GetImagesTask = SystemPrivate->MutableTaskGraph.AddMutableThreadTask(
-		TEXT("Task_Mutable_GetImages"),
-		[OperationData]()
-		{
-			Task_Mutable_GetImages(OperationData);
-		},
-		{ Mutable_GetImageDescriptors_Task });
-
-		UE::Tasks::FTask Mutable_PrepareTextures = SystemPrivate->MutableTaskGraph.AddMutableThreadTask(
-		TEXT("Task_Mutable_PrepareTextures"),
-		[OperationData]()
-		{
-			Task_Mutable_PrepareTextures(OperationData);
-		},
-		{ Mutable_GetImagesTask });
-		
-		// Next Task: Load Unreal Assets
-		//-------------------------------------------------------------
-		UE::Tasks::FTask Game_LoadUnrealAssets = LoadAdditionalAssetsAndData(OperationData);
-		if (!Game_LoadUnrealAssets.IsValid())
-		{
-			OperationData->UpdateResult = EUpdateResult::Error;
-			FinishUpdateGlobal(OperationData);
-			return;
-		}
-
-		// Next-next Task: Convert Resources
-		//-------------------------------------------------------------
-		SystemPrivate->MutableTaskGraph.AddGameThreadTask(
-			TEXT("Task_Game_ConvertResources"),
-			[OperationData](float RemainingTime)
-			{
-				Task_Game_ConvertResources(OperationData);
-			},
-			false,
-			{ Game_LoadUnrealAssets, Mutable_PrepareTextures });
-	}
 	
 	
 	/** "Start Update" */
@@ -2827,6 +2741,7 @@ namespace impl
 		}
 #endif
 		
+		
 		UE::Tasks::FTask Mutable_GetMeshTask = SystemPrivate->MutableTaskGraph.AddMutableThreadTask(
 		TEXT("Task_Mutable_GetMeshes"),
 		[Operation]()
@@ -2865,16 +2780,46 @@ namespace impl
 		{
 			Task_Mutable_UpdateParameterRelevancy(Operation);
 		},
-		{ Mutable_PrepareSkeletonDataTask });
+		{Mutable_PrepareSkeletonDataTask});
+		
+		UE::Tasks::FTask Mutable_GetImageDescriptors_Task = SystemPrivate->MutableTaskGraph.AddMutableThreadTask(
+		TEXT("Task_Mutable_GetImageDescriptors"),
+		[Operation]()
+		{
+			Task_Mutable_GetImageDescriptors(Operation);
+		},
+		{Mutable_UpdateParameterRelevancy});	
 
+		UE::Tasks::FTask Mutable_GetImagesTask = SystemPrivate->MutableTaskGraph.AddMutableThreadTask(
+		TEXT("Task_Mutable_GetImages"),
+		[Operation]()
+		{
+			Task_Mutable_GetImages(Operation);
+		},
+		{Mutable_GetImageDescriptors_Task});
+
+		UE::Tasks::FTask Mutable_PrepareTextures = SystemPrivate->MutableTaskGraph.AddMutableThreadTask(
+		TEXT("Task_Mutable_PrepareTextures"),
+		[Operation]()
+		{
+			Task_Mutable_PrepareTextures(Operation);
+		},
+		{Mutable_GetImagesTask});
+
+		// Next Task: Load Unreal Assets
+		//-------------------------------------------------------------
+		UE::Tasks::FTask Game_LoadUnrealAssets = LoadAdditionalAssetsAndData(Operation);
+
+		// Next-next Task: Convert Resources
+		//-------------------------------------------------------------
 		SystemPrivate->MutableTaskGraph.AddGameThreadTask(
-		TEXT("Task_Game_LockCache"),
+		TEXT("Task_Game_ConvertResources"),
 		[Operation](float RemainingTime)
 		{
-			impl::Task_Game_LockCache(Operation);
+			Task_Game_ConvertResources(Operation);
 		},
 		false,
-		{ Mutable_UpdateParameterRelevancy });
+		{Game_LoadUnrealAssets, Mutable_PrepareTextures});
 	}
 } // namespace impl
 
diff --git a/Engine/Plugins/Mutable/Source/CustomizableObject/Private/MuCO/LogBenchmarkUtil.cpp b/Engine/Plugins/Mutable/Source/CustomizableObject/Private/MuCO/LogBenchmarkUtil.cpp
--- a/Engine/Plugins/Mutable/Source/CustomizableObject/Private/MuCO/LogBenchmarkUtil.cpp
+++ b/Engine/Plugins/Mutable/Source/CustomizableObject/Private/MuCO/LogBenchmarkUtil.cpp
@@ -59,7 +59,7 @@ TSharedPtr<FArchive> CreateFile()
 	const FString HeaderRow = TEXT(
 		"ID_CO;ID_COI;ID_UpdateType;ID_Descriptor;ID_UpdateResult;"											// Basic identifying data
 		"Context_LevelBegunPlay;TriangleCount;"															// The context of the update
-		"Time_Queue;Time_Update;Time_TaskGetMesh;Time_TaskLockCache;"										// Time and memory data ...
+		"Time_Queue;Time_Update;Time_TaskGetMesh;"														// Time and memory data ...
 		"Time_TaskGetImages;Time_TaskConvertResources;Time_TaskCallbacks;Memory_Update;Memory_Update_Real;"
 		"Time_TaskUpdateImage;Memory_TaskUpdateImage;Memory_TaskUpdateImage_Real");
 
@@ -246,7 +246,6 @@ void FLogBenchmarkUtil::FinishUpdateMesh(const TSharedRef<FUpdateContextPrivate>
 	UpdateData.QueueTime = Context->QueueTime * 1000;
 	UpdateData.UpdateTime = Context->UpdateTime * 1000;
 	UpdateData.TaskGetMeshesTime = Context->TaskGetMeshesTime * 1000;
-	UpdateData.TaskLockCacheTime = Context->TaskLockCacheTime * 1000;
 	UpdateData.TaskGetImagesTime = Context->TaskGetImagesTime * 1000;
 	UpdateData.TaskConvertResourcesTime = Context->TaskConvertResourcesTime * 1000;
 	UpdateData.TaskCallbacksTime = Context->TaskCallbacksTime * 1000;
@@ -254,7 +253,7 @@ void FLogBenchmarkUtil::FinishUpdateMesh(const TSharedRef<FUpdateContextPrivate>
 	UpdateData.UpdateRealPeakMemory = (Context->UpdateEndRealPeakBytes / 1024.0) / 1024.0;
 	
 	// todo: Find a better way of handling the construction of this string so we know the symmetry with the header row will not get lost when adding new elements
-	const FString UpdateString = FString::Printf(TEXT("%s;%s;%s;%s;%s;%s;%u;%f;%f;%f;%f;%f;%f;%f;%f;%f"),
+	const FString UpdateString = FString::Printf(TEXT("%s;%s;%s;%s;%s;%s;%u;%f;%f;%f;%f;%f;%f;%f;%f"),
 		*UpdateData.CustomizableObjectPathName,
 		*UpdateData.CustomizableObjectInstancePathName,
 		*UpdateData.UpdateType,
@@ -265,7 +264,6 @@ void FLogBenchmarkUtil::FinishUpdateMesh(const TSharedRef<FUpdateContextPrivate>
 		UpdateData.QueueTime,
 		UpdateData.UpdateTime,
 		UpdateData.TaskGetMeshesTime,
-		UpdateData.TaskLockCacheTime,
 		UpdateData.TaskGetImagesTime,
 		UpdateData.TaskConvertResourcesTime,
 		UpdateData.TaskCallbacksTime,
@@ -304,7 +302,7 @@ void FLogBenchmarkUtil::FinishUpdateImage(const FString& CustomizableObjectPathN
 	
 
 	// todo: Find a better way of handling the construction of this string so we know the symmetry with the header row will not get lost when adding new elements
-	const FString UpdateString = FString::Printf(TEXT("%s;%s;%s;%s;;%s;;;;;;;;;;;%f;%f;%f"),
+	const FString UpdateString = FString::Printf(TEXT("%s;%s;%s;%s;;%s;;;;;;;;;;%f;%f;%f"),
 		*MipsUpdateData.CustomizableObjectPathName,
 		*MipsUpdateData.CustomizableObjectInstancePathName,
 		*MipsUpdateData.UpdateType,
diff --git a/Engine/Plugins/Mutable/Source/CustomizableObjectEditor/Private/MuCOE/CustomizableObjectEditorPerformanceAnalyzer.cpp b/Engine/Plugins/Mutable/Source/CustomizableObjectEditor/Private/MuCOE/CustomizableObjectEditorPerformanceAnalyzer.cpp
--- a/Engine/Plugins/Mutable/Source/CustomizableObjectEditor/Private/MuCOE/CustomizableObjectEditorPerformanceAnalyzer.cpp
+++ b/Engine/Plugins/Mutable/Source/CustomizableObjectEditor/Private/MuCOE/CustomizableObjectEditorPerformanceAnalyzer.cpp
@@ -41,7 +41,6 @@ namespace InstanceUpdatesMainDataColumns
 	static const FName QueueTimeColumnID (TEXT("QueueTime"));
 	static const FName UpdateTimeColumnID (TEXT("UpdateTime"));
 	static const FName GetMeshTimeColumnID (TEXT("GetMeshTime"));
-	static const FName LockCacheTimeColumnID (TEXT("LockCacheTime"));
 	static const FName GetImagesTimeColumnID (TEXT("GetImagesTime"));
 	static const FName ConvertResourcesTimeColumnID (TEXT("ConvertResourcesTime"));
 	static const FName CallbacksTimeColumnID (TEXT("CallbacksTime"));
@@ -181,12 +180,6 @@ TSharedRef<SWidget> SInstanceUpdateDataRow::GenerateWidgetForColumn(const FName&
 			.Text(FText::FromString(FString::Printf(TEXT("%.2f"),InstanceUpdateElement->UpdateStats.TaskGetMeshesTime)))
 			.ColorAndOpacity(TextColor(InstanceUpdateElement->UpdateStats.TaskGetMeshesTime, InstanceUpdateStatistics->TaskGetMeshTime));
 	}
-	else if (InColumnName == InstanceUpdatesMainDataColumns::LockCacheTimeColumnID)
-	{
-		return SNew(STextBlock)
-			.Text(FText::FromString(FString::Printf(TEXT("%.4f"),InstanceUpdateElement->UpdateStats.TaskLockCacheTime)))
-			.ColorAndOpacity(TextColor(InstanceUpdateElement->UpdateStats.TaskLockCacheTime, InstanceUpdateStatistics->TaskLockCacheTime));
-	}
 	else if (InColumnName == InstanceUpdatesMainDataColumns::GetImagesTimeColumnID)
 	{
 		return SNew(STextBlock)
@@ -399,14 +392,6 @@ void SCustomizableObjectEditorPerformanceAnalyzer::Construct(const FArguments& I
 				.HAlignCell(EHorizontalAlignment::HAlign_Right)
 				.FillWidth(1)
 
-				+ SHeaderRow::Column(InstanceUpdatesMainDataColumns::LockCacheTimeColumnID)
-				.DefaultLabel(FText(LOCTEXT("LockCacheTimeColumnLabel", "LockCache (ms)")))
-				.DefaultTooltip( FText(LOCTEXT("LockCacheTimeColumnToolTip", "Initial generation - Time spent protecting used UObjects from GC")))
-				.OnSort(this, &SCustomizableObjectEditorPerformanceAnalyzer::OnInstanceUpdateListViewSort)
-				.SortMode(this, &SCustomizableObjectEditorPerformanceAnalyzer::GetColumnSortMode, InstanceUpdatesMainDataColumns::LockCacheTimeColumnID)
-				.HAlignCell(EHorizontalAlignment::HAlign_Right)
-				.FillWidth(1)
-
 				+ SHeaderRow::Column(InstanceUpdatesMainDataColumns::GetImagesTimeColumnID)
 				.DefaultLabel(FText(LOCTEXT("GetImagesTimeColumnLabel", "GetImages (ms)")))
 				.DefaultTooltip(FText(LOCTEXT("GetImagesTimeColumnToolTip", "Initial generation - Image generation time")))
@@ -792,21 +777,6 @@ void SCustomizableObjectEditorPerformanceAnalyzer::OnInstanceUpdateListViewSort(
 			return ElementA.IsValid();
 		});
 	}
-	else if (ColumnId == InstanceUpdatesMainDataColumns::LockCacheTimeColumnID)
-	{
-		InstanceUpdateElements.StableSort([&](const TSharedPtr<FInstanceUpdateDataElement>& ElementA, const TSharedPtr<FInstanceUpdateDataElement>& ElementB)
-		{
-			if (ElementA.IsValid() && ElementB.IsValid())
-			{
-				const double ValA = ElementA->UpdateStats.TaskLockCacheTime;
-				const double ValB =  ElementB->UpdateStats.TaskLockCacheTime;
-
-				return NewSortMode == EColumnSortMode::Ascending ? ValA < ValB : ValA > ValB;
-			}
-
-			return ElementA.IsValid();
-		});
-	}
 	else if (ColumnId == InstanceUpdatesMainDataColumns::GetImagesTimeColumnID)
 	{
 		InstanceUpdateElements.StableSort([&](const TSharedPtr<FInstanceUpdateDataElement>& ElementA, const TSharedPtr<FInstanceUpdateDataElement>& ElementB)
@@ -1009,7 +979,6 @@ void SCustomizableObjectEditorPerformanceAnalyzer::ComputeStatistics()
 	Compute(InstanceUpdateStatistics->QueueTime, [&](int32 Index){	return InstanceUpdateElements[Index]->UpdateStats.QueueTime; });
 	Compute(InstanceUpdateStatistics->UpdateTime, [&](int32 Index) { return InstanceUpdateElements[Index]->UpdateStats.UpdateTime; });
 	Compute(InstanceUpdateStatistics->TaskGetMeshTime, [&](int32 Index) { return InstanceUpdateElements[Index]->UpdateStats.TaskGetMeshesTime; });
-	Compute(InstanceUpdateStatistics->TaskLockCacheTime, [&](int32 Index) { return InstanceUpdateElements[Index]->UpdateStats.TaskLockCacheTime; });
 	Compute(InstanceUpdateStatistics->TaskGetImagesTime, [&](int32 Index) { return InstanceUpdateElements[Index]->UpdateStats.TaskGetImagesTime; });
 	Compute(InstanceUpdateStatistics->TaskConvertResourcesTime, [&](int32 Index) { return InstanceUpdateElements[Index]->UpdateStats.TaskConvertResourcesTime; });
 	Compute(InstanceUpdateStatistics->TaskCallbacksTime, [&](int32 Index) { return InstanceUpdateElements[Index]->UpdateStats.TaskCallbacksTime; });
