[mutable] Added TransformWithMesh node (UCONodeTransformInMesh).
- Added the new node to the previewer system that allows it to show the bounding mesh connected to the node.
- Fixed previewer bug that prevented the update of the mesh from being displayed with the ShowGizmo... method (Clipping Mesh).
- Moved the "UCustomizableObjectNodeModifierTransformInMesh" node into the node creation "Modifiers" space.
- Added compilation error messages for various cases for the "UCONodeRemoveMesh" and "UCONodeTransformWithBone" nodes.
- Added logs to the code generation of the "NodeModifierMeshTransformInMesh" node.

#rnx

#Synced-CL 51453194

diff --git a/Engine/Plugins/Mutable/Source/CustomizableObjectEditor/Private/MuCOE/CustomizableObjectEditor.cpp b/Engine/Plugins/Mutable/Source/CustomizableObjectEditor/Private/MuCOE/CustomizableObjectEditor.cpp
--- a/Engine/Plugins/Mutable/Source/CustomizableObjectEditor/Private/MuCOE/CustomizableObjectEditor.cpp
+++ b/Engine/Plugins/Mutable/Source/CustomizableObjectEditor/Private/MuCOE/CustomizableObjectEditor.cpp
@@ -62,6 +62,7 @@
 #include "MuCOE/CompileRequest.h"
 #include "MuCOE/CustomizableObjectInstanceFactory.h"
 #include "MessageLogModule.h"
+#include "Nodes/CONodeTransformInMesh.h"
 
 #include UE_INLINE_GENERATED_CPP_BY_NAME(CustomizableObjectEditor)
 
@@ -943,6 +944,11 @@ void FCustomizableObjectEditor::ShowGizmoClipMesh(UCustomizableObjectNode& Node,
 
 		Viewport->ShowGizmoClipMesh(Node, Transform, *ClipMesh, LODIndex, SectionIndex, MaterialSlotIndex);
 	}
+	else
+	{
+		// If the mesh to set is not valid ensure the viewport shows no mesh.
+		Viewport->HideGizmoClipMesh();
+	}
 }
 
 
@@ -963,6 +969,7 @@ void FCustomizableObjectEditor::HideGizmoClipMesh()
 		const UObject* Node = *NodeIt;
 		if (Node->IsA<UCustomizableObjectNodeModifierClipWithMesh>() ||
 			Node->IsA<UCustomizableObjectNodeModifierTransformInMesh>() ||
+			Node->IsA<UCONodeTransformInMesh>() ||
 			Node->IsA<UCONodeClipMeshWithMesh>())
 		{
 			GraphEditor->ClearSelectionSet();
@@ -1642,6 +1649,13 @@ void FCustomizableObjectEditor::OnSelectedGraphNodesChanged(const FGraphPanelSel
 				ShowGizmoClipMesh(*NodeModifierTransformInMesh, &NodeModifierTransformInMesh->BoundingMeshTransform, *Pin);
 			}
 		}
+		else if (UCONodeTransformInMesh* NodeTransformInMesh = Cast<UCONodeTransformInMesh>(Objects[0]))
+		{
+			if (const UEdGraphPin* Pin = NodeTransformInMesh->BoundingMeshPin.Get())
+			{
+				ShowGizmoClipMesh(*NodeTransformInMesh, &NodeTransformInMesh->BoundingMeshTransform, *Pin);
+			}
+		}
 		else if (UCustomizableObjectNodeProjectorConstant* NodeProjectorConstant = Cast<UCustomizableObjectNodeProjectorConstant>(Objects[0]))
 		{
 			ShowGizmoProjectorNodeProjectorConstant(*NodeProjectorConstant);
diff --git a/Engine/Plugins/Mutable/Source/CustomizableObjectEditor/Private/MuCOE/CustomizableObjectEditorViewportClient.cpp b/Engine/Plugins/Mutable/Source/CustomizableObjectEditor/Private/MuCOE/CustomizableObjectEditorViewportClient.cpp
--- a/Engine/Plugins/Mutable/Source/CustomizableObjectEditor/Private/MuCOE/CustomizableObjectEditorViewportClient.cpp
+++ b/Engine/Plugins/Mutable/Source/CustomizableObjectEditor/Private/MuCOE/CustomizableObjectEditorViewportClient.cpp
@@ -45,6 +45,7 @@
 #include "MuCO/CustomizableObjectPrivate.h"
 #include "MuCO/LoadUtils.h"
 #include "Nodes/CONodeClipMeshWithMesh.h"
+#include "Nodes/CONodeTransformInMesh.h"
 #include "Widgets/Input/SButton.h"
 #include "Widgets/Input/SEditableTextBox.h"
 #include "Widgets/Layout/SExpandableArea.h"
@@ -158,8 +159,8 @@ FCustomizableObjectEditorViewportClient::FCustomizableObjectEditorViewportClient
 		ClipMeshMaterial = TStrongObjectPtr<UMaterial>(UMaterial::GetDefaultMaterial(MD_Surface));
 	}
 	
-	const float FOVMin = 5.f;
-	const float FOVMax = 170.f;
+	constexpr float FOVMin = 5.f;
+	constexpr float FOVMax = 170.f;
 	ViewFOV = FMath::Clamp<float>(53.43f, FOVMin, FOVMax);
 
 	SetRealtime(true);
@@ -173,7 +174,7 @@ FCustomizableObjectEditorViewportClient::FCustomizableObjectEditorViewportClient
 	
 	StateChangeShowGeometryDataFlag = false;
 
-	// Register delegate to update the show flags when the post processing is turned on or off
+	// Register delegate to update the show flags when the post-processing is turned on or off
 	UAssetViewerSettings::Get()->OnAssetViewerSettingsChanged().AddRaw(this, &FCustomizableObjectEditorViewportClient::OnAssetViewerSettingsChanged);
 	// Set correct flags according to current profile settings
 	SetAdvancedShowFlagsForScene(UAssetViewerSettings::Get()->Profiles[GetMutableDefault<UEditorPerProjectUserSettings>()->AssetViewerProfileIndex].bPostProcessingEnabled);
@@ -232,7 +233,7 @@ void FCustomizableObjectEditorViewportClient::Draw(const FSceneView* View, FPrim
 				// Draw outer light cone
 				DrawWireSphereCappedCone(PDI, TransformNoScale, SpotLightComp->AttenuationRadius, SpotLightComp->OuterConeAngle, 32, 8, 10, FColor(200, 255, 255), SDPG_World);
 
-				// Draw inner light cone (if non zero)
+				// Draw inner light cone (if non-zero)
 				if (SpotLightComp->InnerConeAngle > UE_KINDA_SMALL_NUMBER)
 				{
 					DrawWireSphereCappedCone(PDI, TransformNoScale, SpotLightComp->AttenuationRadius, SpotLightComp->InnerConeAngle, 32, 8, 10, FColor(150, 200, 255), SDPG_World);
@@ -647,6 +648,10 @@ void FCustomizableObjectEditorViewportClient::ShowGizmoClipMesh(UCustomizableObj
 	{
 		TransformExternallyChangedDelegateHandle = NodeClipMeshWithMesh->TransformChangedDelegate.AddSP(this, &FCustomizableObjectEditorViewportClient::UpdateGizmoClipMeshTransform);
 	}
+	else if (UCONodeTransformInMesh* NodeTransformInMesh = Cast<UCONodeTransformInMesh>(ClipMeshNode.Get()))
+	{
+		TransformExternallyChangedDelegateHandle = NodeTransformInMesh->TransformChangedDelegate.AddSP(this, &FCustomizableObjectEditorViewportClient::UpdateGizmoClipMeshTransform);
+	}
 	else
 	{
 		unimplemented();
@@ -713,6 +718,10 @@ void FCustomizableObjectEditorViewportClient::HideGizmoClipMesh()
 			{
 				NodeClipMeshWithMesh->TransformChangedDelegate.Remove(TransformExternallyChangedDelegateHandle);
 			}
+			else if (UCONodeTransformInMesh* NodeTransformInMesh = Cast<UCONodeTransformInMesh>(ClipMeshNode.Get()))
+			{
+				NodeTransformInMesh->TransformChangedDelegate.Remove(TransformExternallyChangedDelegateHandle);
+			}
 			else
 			{
 				unimplemented();
diff --git a/Engine/Plugins/Mutable/Source/CustomizableObjectEditor/Private/MuCOE/EdGraphSchema_CustomizableObject.cpp b/Engine/Plugins/Mutable/Source/CustomizableObjectEditor/Private/MuCOE/EdGraphSchema_CustomizableObject.cpp
--- a/Engine/Plugins/Mutable/Source/CustomizableObjectEditor/Private/MuCOE/EdGraphSchema_CustomizableObject.cpp
+++ b/Engine/Plugins/Mutable/Source/CustomizableObjectEditor/Private/MuCOE/EdGraphSchema_CustomizableObject.cpp
@@ -124,6 +124,7 @@
 #include "Nodes/CONodeRemoveMesh.h"
 #include "Nodes/CONodeSkeletalMeshBreak.h"
 #include "Nodes/CONodeTransformWithBone.h"
+#include "Nodes/CONodeTransformInMesh.h"
 #include "UObject/Class.h"
 #include "Widgets/SMutablePinTypeSelector.h"
 
@@ -506,6 +507,7 @@ void UEdGraphSchema_CustomizableObject::GetGraphContextActions(FGraphContextMenu
 			ContextMenuBuilder.CreateTemplateNode<UCustomizableObjectNodeModifierClipWithMesh>(),
 			ContextMenuBuilder.CreateTemplateNode<UCustomizableObjectNodeModifierClipWithUVMask>(),
 			ContextMenuBuilder.CreateTemplateNode<UCONodeModifierTransformWithBone>(),
+			ContextMenuBuilder.CreateTemplateNode<UCustomizableObjectNodeModifierTransformInMesh>(),
 		};
 
 		AddNewNodeCategoryActionsFiltered(ModifierTemplateNodes, ContextMenuBuilder, TEXT("Modifier"), GeneralGrouping, Filter);
@@ -537,6 +539,7 @@ void UEdGraphSchema_CustomizableObject::GetGraphContextActions(FGraphContextMenu
 			ContextMenuBuilder.CreateTemplateNode<UCONodeClipMeshWithMesh>(),
 			ContextMenuBuilder.CreateTemplateNode<UCONodeRemoveMesh>(),
 			ContextMenuBuilder.CreateTemplateNode<UCONodeTransformWithBone>(),
+			ContextMenuBuilder.CreateTemplateNode<UCONodeTransformInMesh>(),
 		};
 
 		AddNewNodeCategoryActionsFiltered(MeshTemplateNodes, ContextMenuBuilder, TEXT("Mesh"), GeneralGrouping, Filter);
@@ -652,7 +655,6 @@ void UEdGraphSchema_CustomizableObject::GetGraphContextActions(FGraphContextMenu
 			ContextMenuBuilder.CreateTemplateNode<UCustomizableObjectNodeComponentPassthroughMesh>(),
 			ContextMenuBuilder.CreateTemplateNode<UCustomizableObjectNodeMeshReshape>(),
 			ContextMenuBuilder.CreateTemplateNode<UCustomizableObjectNodeModifierClipDeform>(),
-			ContextMenuBuilder.CreateTemplateNode<UCustomizableObjectNodeModifierTransformInMesh>(),
 			ContextMenuBuilder.CreateTemplateNode<UCustomizableObjectNodeTextureParameter>(),
 			ContextMenuBuilder.CreateTemplateNode<UCustomizableObjectNodeSkeletalMeshParameter>(),
 			ContextMenuBuilder.CreateTemplateNode<UCustomizableObjectNodeMaterialParameter>(),
diff --git a/Engine/Plugins/Mutable/Source/CustomizableObjectEditor/Private/MuCOE/GenerateMutableSource/GenerateMutableSourceMesh.cpp b/Engine/Plugins/Mutable/Source/CustomizableObjectEditor/Private/MuCOE/GenerateMutableSource/GenerateMutableSourceMesh.cpp
--- a/Engine/Plugins/Mutable/Source/CustomizableObjectEditor/Private/MuCOE/GenerateMutableSource/GenerateMutableSourceMesh.cpp
+++ b/Engine/Plugins/Mutable/Source/CustomizableObjectEditor/Private/MuCOE/GenerateMutableSource/GenerateMutableSourceMesh.cpp
@@ -65,10 +65,12 @@
 #include "MuCO/CustomizableObjectInstanceAssetUserData.h"
 #include "MuCOE/Nodes/CONodeExternalOperation.h"
 #include "MuCOE/Nodes/CONodeRemoveMesh.h"
+#include "MuCOE/Nodes/CONodeTransformInMesh.h"
 #include "MuT/NodeMeshClipWithMesh.h"
 #include "MuT/NodeMeshExternal.h"
 #include "MuT/NodeMeshRemoveMesh.h"
 #include "MuT/NodeMeshTransform.h"
+#include "MuT/NodeMeshTransformInMesh.h"
 #include "MuT/NodeMeshTransformWithBone.h"
 
 #define LOCTEXT_NAMESPACE "CustomizableObjectEditor"
@@ -3124,7 +3126,7 @@ UE::Mutable::Private::Ptr<UE::Mutable::Private::NodeMesh> GenerateMutableSourceM
 		}
 		else
 		{
-			const FText ErrorMsg = LOCTEXT("Base mesh missing", "The clip mesh with mesh node requires an input base mesh.");
+			const FText ErrorMsg = LOCTEXT("ClipMeshWithMesh_Base_Mesh_missing", "The Clip Mesh with Mesh node requires an input base mesh.");
 			GenerationContext.Log(ErrorMsg, ClipMeshWithMeshNode, EMessageSeverity::Error);
 		}
 		
@@ -3159,7 +3161,7 @@ UE::Mutable::Private::Ptr<UE::Mutable::Private::NodeMesh> GenerateMutableSourceM
 		}
 		else
 		{
-			const FText ErrorMsg = LOCTEXT("Clipping mesh missing", "The clip mesh with mesh node requires an input clip mesh.");
+			const FText ErrorMsg = LOCTEXT("ClipMeshWithMesh_Clipping_Mesh_missing", "The Clip Mesh with Mesh node requires an input clipping mesh.");
 			GenerationContext.Log(ErrorMsg, ClipMeshWithMeshNode, EMessageSeverity::Error);
 		}
 		
@@ -3179,6 +3181,11 @@ UE::Mutable::Private::Ptr<UE::Mutable::Private::NodeMesh> GenerateMutableSourceM
 			UE::Mutable::Private::Ptr<UE::Mutable::Private::NodeMesh> BaseMeshCoreNode = GenerateMutableSourceMesh(ConnectedPin, GenerationContext, FMutableSourceMeshData(), false, true);
 			RemoveNode->Source = BaseMeshCoreNode;
 		}
+		else
+		{
+			const FText ErrorMsg = LOCTEXT("RemoveMesh_Remove_Mesh_missing", "The Remove Mesh node requires an input base mesh.");
+			GenerationContext.Log(ErrorMsg, RemoveMeshNode, EMessageSeverity::Error);
+		}
 		
 		// Remove Mesh
 		UEdGraphPin* RemoveMeshPin = RemoveMeshNode->RemoveMeshPin.Get();
@@ -3199,6 +3206,11 @@ UE::Mutable::Private::Ptr<UE::Mutable::Private::NodeMesh> GenerateMutableSourceM
 			
 			GenerationContext.MeshGenerationFlags.Pop();
 		}
+		else
+		{
+			const FText ErrorMsg = LOCTEXT("RemoveMesh_Remove_Mesh_missing", "The Remove Mesh node requires an input remove mesh.");
+			GenerationContext.Log(ErrorMsg, RemoveMeshNode, EMessageSeverity::Error);
+		}
 		
 		Result = RemoveNode;
 	}
@@ -3211,19 +3223,94 @@ UE::Mutable::Private::Ptr<UE::Mutable::Private::NodeMesh> GenerateMutableSourceM
 		TransformNode->BoneName = *TransformWithBoneNode->BoneName;
 		TransformNode->ThresholdFactor = TransformWithBoneNode->ThresholdFactor;
 
-		const UEdGraphPin* SourceMeshPin = TransformWithBoneNode->BaseMeshPin.Get();
-		check(SourceMeshPin);
-		if (const UEdGraphPin* ConnectedPin = FollowInputPin(*SourceMeshPin))
+		const UEdGraphPin* BaseMeshPin = TransformWithBoneNode->BaseMeshPin.Get();
+		if (const UEdGraphPin* ConnectedPin = FollowInputPin(*BaseMeshPin))
 		{
 			TransformNode->Source = GenerateMutableSourceMesh(ConnectedPin, GenerationContext, BaseMeshData, false, bOnlyConnectedLOD);
 		}
+		else
+		{
+			const FText ErrorMsg = LOCTEXT("TransformWithBone_Base_Mesh_missing", "The Transform With Bone node requires an input base mesh.");
+			GenerationContext.Log(ErrorMsg, TransformWithBoneNode, EMessageSeverity::Error);
+		}
 		
 		const UEdGraphPin* TransformPin = TransformWithBoneNode->TransformPin.Get();
-		check(TransformPin);
 		if (const UEdGraphPin* ConnectedPin = FollowInputPin(*TransformPin))
 		{
 			TransformNode->MatrixNode = GenerateMutableSourceTransform(ConnectedPin, GenerationContext);
 		}
+		else
+		{
+			const FText ErrorMsg = LOCTEXT("TransformWithBone_Transform_missing", "The Transform With Bone node requires an input transform.");
+			GenerationContext.Log(ErrorMsg, TransformWithBoneNode, EMessageSeverity::Error);
+		}
+	}
+	
+	
+	else if (const UCONodeTransformInMesh* TransformInMeshNode = Cast<UCONodeTransformInMesh>(Node))
+	{
+		UE::Mutable::Private::Ptr<UE::Mutable::Private::NodeMeshTransformInMesh> TransformNode = new UE::Mutable::Private::NodeMeshTransformInMesh();
+		Result = TransformNode;
+
+		// Input mesh
+		const UEdGraphPin* BaseMeshPin = TransformInMeshNode->BaseMeshPin.Get();
+		if (const UEdGraphPin* ConnectedPin = FollowInputPin(*BaseMeshPin))
+		{
+			TransformNode->SourceMesh = GenerateMutableSourceMesh(ConnectedPin, GenerationContext, BaseMeshData, false, bOnlyConnectedLOD);
+		}
+		else
+		{
+			const FText ErrorMsg = LOCTEXT("TransformInMesh_Base_Mesh_missing", "The Transform In mesh node requires an input base mesh.");
+			GenerationContext.Log(ErrorMsg, TransformInMeshNode, EMessageSeverity::Error);
+		}
+		
+		// transform to apply to the source mesh
+		const UEdGraphPin* TransformPin = TransformInMeshNode->TransformMeshPin.Get();
+		if (const UEdGraphPin* ConnectedPin = FollowInputPin(*TransformPin))
+		{
+			TransformNode->MatrixNode = GenerateMutableSourceTransform(ConnectedPin, GenerationContext);
+		}
+		else
+		{
+			const FText ErrorMsg = LOCTEXT("TransformInMesh_Transform_missing", "The Transform In mesh node requires an input transform.");
+			GenerationContext.Log(ErrorMsg, TransformInMeshNode, EMessageSeverity::Error);
+		}
+		
+		// Bounding Mesh
+		const UEdGraphPin* BoundingMeshPin = TransformInMeshNode->BoundingMeshPin.Get();
+		if (const UEdGraphPin* ConnectedPin = FollowInputPin(*BoundingMeshPin))
+		{
+			constexpr EMutableMeshConversionFlags BoundingMeshFlags =
+				EMutableMeshConversionFlags::IgnoreSkinning |
+				EMutableMeshConversionFlags::IgnorePhysics |
+				EMutableMeshConversionFlags::IgnoreMorphs |
+				EMutableMeshConversionFlags::IgnoreTexCoords |
+				EMutableMeshConversionFlags::IgnoreAUD |
+				EMutableMeshConversionFlags::DoNotCreateMeshMetadata;
+				
+			GenerationContext.MeshGenerationFlags.Push(BoundingMeshFlags);
+			{
+				UE::Mutable::Private::Ptr<UE::Mutable::Private::NodeMesh> BoundingMesh = GenerateMutableSourceMesh(ConnectedPin, GenerationContext, FMutableSourceMeshData(), false, true);
+				
+				// Apply the transform defined within the node to the bounding mesh
+				if (const FMatrix Matrix = TransformInMeshNode->BoundingMeshTransform.ToMatrixWithScale(); Matrix != FMatrix::Identity)
+				{
+					UE::Mutable::Private::Ptr<UE::Mutable::Private::NodeMeshTransform> TransformMesh = new UE::Mutable::Private::NodeMeshTransform();
+					TransformMesh->Source = BoundingMesh;
+
+					TransformMesh->Transform = FMatrix44f(Matrix);
+					BoundingMesh = TransformMesh;
+				}
+				
+				TransformNode->BoundingMesh = BoundingMesh;
+			}
+			GenerationContext.MeshGenerationFlags.Pop();
+		}
+		else
+		{
+			const FText ErrorMsg = LOCTEXT("TransformInMesh_Bounding_Mesh_missing", "The Transform In mesh node requires an input Bounding Mesh.");
+			GenerationContext.Log(ErrorMsg, TransformInMeshNode, EMessageSeverity::Error);
+		}
 	}
 	
 	else if (const UCustomizableObjectNodeAnimationPose* TypedNode = Cast<UCustomizableObjectNodeAnimationPose>(Node))
diff --git a/Engine/Plugins/Mutable/Source/CustomizableObjectEditor/Private/MuCOE/GraphTraversal.cpp b/Engine/Plugins/Mutable/Source/CustomizableObjectEditor/Private/MuCOE/GraphTraversal.cpp
--- a/Engine/Plugins/Mutable/Source/CustomizableObjectEditor/Private/MuCOE/GraphTraversal.cpp
+++ b/Engine/Plugins/Mutable/Source/CustomizableObjectEditor/Private/MuCOE/GraphTraversal.cpp
@@ -32,6 +32,7 @@
 #include "Nodes/CONodeClipMeshWithMesh.h"
 #include "Nodes/CONodeRemoveMesh.h"
 #include "Nodes/CONodeTransformWithBone.h"
+#include "Nodes/CONodeTransformInMesh.h"
 #include "UObject/UObjectIterator.h"
 
 
@@ -789,6 +790,14 @@ const UEdGraphPin* FindMeshBaseSource(const UEdGraphPin& Pin, const bool bOnlyLo
 			return FindMeshBaseSource(*ConnectedPin, bOnlyLookForStaticMesh, MacroContext);
 		}
 	}
+	
+	else if (const UCONodeTransformInMesh* NodeTransformWithMesh = Cast<UCONodeTransformInMesh>(Node))
+	{
+		if (const UEdGraphPin* ConnectedPin = FollowInputPin(*NodeTransformWithMesh->BaseMeshPin.Get()))
+		{
+			return FindMeshBaseSource(*ConnectedPin, bOnlyLookForStaticMesh, MacroContext);
+		}
+	}
 
 	else if (Cast<UCONodeExternalOperation>(Node))
 	{
diff --git a/Engine/Plugins/Mutable/Source/CustomizableObjectEditor/Private/MuCOE/Nodes/CONodeModifierTransformWithBone.h b/Engine/Plugins/Mutable/Source/CustomizableObjectEditor/Private/MuCOE/Nodes/CONodeModifierTransformWithBone.h
--- a/Engine/Plugins/Mutable/Source/CustomizableObjectEditor/Private/MuCOE/Nodes/CONodeModifierTransformWithBone.h
+++ b/Engine/Plugins/Mutable/Source/CustomizableObjectEditor/Private/MuCOE/Nodes/CONodeModifierTransformWithBone.h
@@ -27,6 +27,7 @@ public:
 
 	// UCustomizableObjectNode interface
 	virtual void AllocateDefaultPins(UCustomizableObjectNodeRemapPins* RemapPins) override;
+	virtual bool IsExperimental() const override {return true;};
 
 	// UCustomizableObjectNodeModifierBase interface
 	virtual UEdGraphPin* GetOutputPin() const override;
diff --git a/Engine/Plugins/Mutable/Source/CustomizableObjectEditor/Private/MuCOE/Nodes/CONodeTransformInMesh.cpp b/Engine/Plugins/Mutable/Source/CustomizableObjectEditor/Private/MuCOE/Nodes/CONodeTransformInMesh.cpp
new file mode 100644
--- /dev/null
+++ b/Engine/Plugins/Mutable/Source/CustomizableObjectEditor/Private/MuCOE/Nodes/CONodeTransformInMesh.cpp
@@ -0,0 +1,59 @@
+// Copyright Epic Games, Inc. All Rights Reserved.
+
+#include "MuCOE/Nodes/CONodeTransformInMesh.h"
+
+#include "MuCOE/EdGraphSchema_CustomizableObject.h"
+
+#include UE_INLINE_GENERATED_CPP_BY_NAME(CONodeTransformInMesh)
+
+#define LOCTEXT_NAMESPACE "CustomizableObjectEditor"
+
+
+FText UCONodeTransformInMesh::GetNodeTitle(ENodeTitleType::Type TitleType) const
+{
+	return LOCTEXT("Transform_Mesh_In_Mesh", "Transform Mesh In Mesh");
+}
+
+
+FText UCONodeTransformInMesh::GetTooltipText() const
+{
+	return LOCTEXT("Transform_Mesh_In_Mesh_Tooltip", "Applies a transform to the vertices of a mesh that are contained within the given bounding mesh");
+}
+
+
+FLinearColor UCONodeTransformInMesh::GetNodeTitleColor() const
+{
+	return UEdGraphSchema_CustomizableObject::GetPinTypeColor(UEdGraphSchema_CustomizableObject::PC_Mesh);
+}
+
+
+void UCONodeTransformInMesh::AllocateDefaultPins(UCustomizableObjectNodeRemapPins* RemapPins)
+{
+	BaseMeshPin = CustomCreatePin(EGPD_Input, UEdGraphSchema_CustomizableObject::PC_Mesh, TEXT("Base Mesh"), LOCTEXT("BaseMesh", "Base Mesh"));
+	BoundingMeshPin = CustomCreatePin(EGPD_Input, UEdGraphSchema_CustomizableObject::PC_Mesh, TEXT("Bounding Mesh"), LOCTEXT("BoundingMesh", "Bounding Mesh"));
+	TransformMeshPin = CustomCreatePin(EGPD_Input, UEdGraphSchema_CustomizableObject::PC_Transform, TEXT("Transform"), LOCTEXT("Transform", "Transform"));
+	
+	OutputMeshPin =  CustomCreatePin(EGPD_Output, UEdGraphSchema_CustomizableObject::PC_Mesh);
+}
+
+
+void UCONodeTransformInMesh::PostEditChangeProperty(struct FPropertyChangedEvent& PropertyChangedEvent)
+{
+	Super::PostEditChangeProperty(PropertyChangedEvent);
+
+	if (PropertyChangedEvent.Property == nullptr)
+	{
+		return;
+	}
+	
+	if (PropertyChangedEvent.Property->GetFName() == GET_MEMBER_NAME_CHECKED(UCONodeTransformInMesh, BoundingMeshTransform))
+	{
+		if (TransformChangedDelegate.IsBound())
+		{
+			TransformChangedDelegate.Broadcast(BoundingMeshTransform);
+		}
+	}
+}
+
+
+#undef LOCTEXT_NAMESPACE
diff --git a/Engine/Plugins/Mutable/Source/CustomizableObjectEditor/Private/MuCOE/Nodes/CONodeTransformInMesh.h b/Engine/Plugins/Mutable/Source/CustomizableObjectEditor/Private/MuCOE/Nodes/CONodeTransformInMesh.h
new file mode 100644
--- /dev/null
+++ b/Engine/Plugins/Mutable/Source/CustomizableObjectEditor/Private/MuCOE/Nodes/CONodeTransformInMesh.h
@@ -0,0 +1,55 @@
+// Copyright Epic Games, Inc. All Rights Reserved.
+#pragma once
+
+#include "MuCOE/Nodes/CustomizableObjectNode.h"
+
+#include "CONodeTransformInMesh.generated.h"
+
+
+namespace ENodeTitleType { enum Type : int; }
+
+class UCustomizableObject;
+class UCustomizableObjectNodeRemapPins;
+class UEdGraphPin;
+struct FPropertyChangedEvent;
+
+DECLARE_MULTICAST_DELEGATE_OneParam(OnTransformChanged, const FTransform&);
+
+UCLASS()
+class UCONodeTransformInMesh : public UCustomizableObjectNode
+{
+	GENERATED_BODY()
+
+public:
+	
+	// EdGraphNode interface
+	virtual FText GetNodeTitle(ENodeTitleType::Type TitleType) const override;
+	virtual FText GetTooltipText() const override;
+	virtual FLinearColor GetNodeTitleColor() const override;
+
+	// UCustomizableObjectNode interface
+	virtual void AllocateDefaultPins(UCustomizableObjectNodeRemapPins* RemapPins) override;
+	virtual void PostEditChangeProperty(FPropertyChangedEvent& PropertyChangedEvent) override;
+	virtual bool IsExperimental() const override { return true; }
+
+	/** Transform to apply to the bounding mesh before selecting for vertices to transform. */
+	UPROPERTY(EditAnywhere, Category = BoundingMesh)
+	FTransform BoundingMeshTransform = FTransform::Identity;
+	
+	UPROPERTY()
+	FEdGraphPinReference BaseMeshPin;
+	
+	UPROPERTY()
+	FEdGraphPinReference BoundingMeshPin;
+	
+	UPROPERTY()
+	FEdGraphPinReference TransformMeshPin;
+	
+	UPROPERTY()
+	FEdGraphPinReference OutputMeshPin;
+	
+	/**
+	 * Delegate invoked each time the BoundingMeshTransform value gets modified
+	 */
+	OnTransformChanged TransformChangedDelegate;
+};
diff --git a/Engine/Plugins/Mutable/Source/CustomizableObjectEditor/Private/MuCOE/SMutableGraphViewer.cpp b/Engine/Plugins/Mutable/Source/CustomizableObjectEditor/Private/MuCOE/SMutableGraphViewer.cpp
--- a/Engine/Plugins/Mutable/Source/CustomizableObjectEditor/Private/MuCOE/SMutableGraphViewer.cpp
+++ b/Engine/Plugins/Mutable/Source/CustomizableObjectEditor/Private/MuCOE/SMutableGraphViewer.cpp
@@ -55,6 +55,7 @@
 #include "MuT/NodeMeshRemoveMesh.h"
 #include "MuT/NodeMeshTable.h"
 #include "MuT/NodeMeshTransform.h"
+#include "MuT/NodeMeshTransformInMesh.h"
 #include "MuT/NodeMeshTransformWithBone.h"
 #include "MuT/NodeModifierMeshClipDeform.h"
 #include "MuT/NodeModifierMeshClipMorphPlane.h"
@@ -108,6 +109,7 @@ static const char* MutableNodeNames[] =
 	"MeshClipWithMesh",
 	"MeshRemoveMesh",
 	"MeshTransformWithBone",
+	"MeshTransformInMesh",
 	"MeshApplyPose",
 	"MeshVariation",
 	"MeshReshape",
@@ -939,7 +941,6 @@ void SMutableGraphViewer::GetChildrenForInfo(TSharedPtr<FMutableGraphTreeElement
 	else if (ParentNodeType == UE::Mutable::Private::NodeMeshTransformWithBone::GetStaticType())
 	{
 		const UE::Mutable::Private::NodeMeshTransformWithBone* NodeMeshTransformWithBone = StaticCast<UE::Mutable::Private::NodeMeshTransformWithBone*>(ParentNode);
-		
 		AddChildFunc(NodeMeshTransformWithBone->Source.get(), TEXT("SOURCE"));
 		AddChildFunc( NodeMeshTransformWithBone->MatrixNode.get(), TEXT("TRANSFORM"));
 	}
@@ -950,6 +951,14 @@ void SMutableGraphViewer::GetChildrenForInfo(TSharedPtr<FMutableGraphTreeElement
 		AddChildFunc(NodeMeshTransform->Source.get(), TEXT("SOURCE"));
 	}
 	
+	else if (ParentNodeType == UE::Mutable::Private::NodeMeshTransformInMesh::GetStaticType())
+	{
+		const UE::Mutable::Private::NodeMeshTransformInMesh* NodeMeshTransformInMesh = StaticCast<UE::Mutable::Private::NodeMeshTransformInMesh*>(ParentNode);
+		AddChildFunc(NodeMeshTransformInMesh->SourceMesh.get(), TEXT("SOURCE"));
+		AddChildFunc(NodeMeshTransformInMesh->BoundingMesh.get(), TEXT("BOUNDS"));
+		AddChildFunc(NodeMeshTransformInMesh->MatrixNode.get(), TEXT("TRANSFORM"));
+	}
+	
 	else
 	{
 		UE_LOG(LogMutable,Error,TEXT("The node of type %d has not been implemented, so its children won't be added to the tree."), int32(ParentNodeType->Type));
diff --git a/Engine/Plugins/Mutable/Source/MutableTools/Internal/MuT/Node.h b/Engine/Plugins/Mutable/Source/MutableTools/Internal/MuT/Node.h
--- a/Engine/Plugins/Mutable/Source/MutableTools/Internal/MuT/Node.h
+++ b/Engine/Plugins/Mutable/Source/MutableTools/Internal/MuT/Node.h
@@ -46,6 +46,7 @@ namespace UE::Mutable::Private
 			MeshClipWithMesh,
 			MeshRemoveMesh,
 			MeshTransformWithBone,
+			MeshTransformInMesh,
 			MeshApplyPose,
 			MeshVariation,
 			MeshReshape,
diff --git a/Engine/Plugins/Mutable/Source/MutableTools/Internal/MuT/NodeMeshTransformInMesh.h b/Engine/Plugins/Mutable/Source/MutableTools/Internal/MuT/NodeMeshTransformInMesh.h
new file mode 100644
--- /dev/null
+++ b/Engine/Plugins/Mutable/Source/MutableTools/Internal/MuT/NodeMeshTransformInMesh.h
@@ -0,0 +1,50 @@
+// Copyright Epic Games, Inc. All Rights Reserved.
+
+#pragma once
+
+#include "MuR/Ptr.h"
+#include "MuR/Types.h"
+#include "MuT/Node.h"
+#include "MuT/NodeModifier.h"
+#include "MuT/NodeMesh.h"
+#include "MuT/NodeMatrix.h"
+
+#define UE_API MUTABLETOOLS_API
+
+namespace UE::Mutable::Private
+{
+
+	/** */
+	class NodeMeshTransformInMesh : public NodeMesh
+	{
+	public:
+
+		/** */
+		Ptr<NodeMesh> SourceMesh;
+		
+		/** */
+		Ptr<NodeMesh> BoundingMesh;
+
+		/** */
+		Ptr<NodeMatrix> MatrixNode;
+		
+	public:
+
+		// Node interface
+		virtual const FNodeType* GetType() const override { return &StaticType; }
+		static const FNodeType* GetStaticType() { return &StaticType; }
+
+	protected:
+
+		/** Forbidden. Manage with the Ptr<> template. */
+		virtual ~NodeMeshTransformInMesh() override = default;
+
+	private:
+
+		static UE_API FNodeType StaticType;
+
+	};
+
+}
+
+#undef UE_API
\ No newline at end of file
diff --git a/Engine/Plugins/Mutable/Source/MutableTools/Private/MuT/CodeGenerator.cpp b/Engine/Plugins/Mutable/Source/MutableTools/Private/MuT/CodeGenerator.cpp
--- a/Engine/Plugins/Mutable/Source/MutableTools/Private/MuT/CodeGenerator.cpp
+++ b/Engine/Plugins/Mutable/Source/MutableTools/Private/MuT/CodeGenerator.cpp
@@ -2397,68 +2397,71 @@ namespace UE::Mutable::Private
 				const NodeModifierMeshTransformInMesh* TypedTransformNode = static_cast<const NodeModifierMeshTransformInMesh*>(Modifier.Node);
 
 				// If a matrix node is not connected, the op won't do anything, so let's not create it at all.
-				if (TypedTransformNode->MatrixNode)
+				if (!TypedTransformNode->MatrixNode)
+				{
+					ErrorLog->Add("Matrix has not been generated", ELMT_ERROR, ErrorContext);
+					continue;
+				}
+				else if (!TypedTransformNode->BoundingMesh)
 				{
-					if (!TypedTransformNode->MatrixNode || !TypedTransformNode->BoundingMesh)
+					ErrorLog->Add("Bounding mesh has not been generated", ELMT_ERROR, ErrorContext);
+					continue;
+				}
+				
+				// Transform matrix.
+				Ptr<ASTOp> MatrixOp;
+				FMatrixGenerationResult ChildResult;
+				GenerateMatrix(ChildResult, StaticOptions, TypedTransformNode->MatrixNode);
+				MatrixOp = ChildResult.op;
+
+				// Bounding mesh
+				FMeshGenerationStaticOptions BoundingMeshStaticOptions(StaticOptions);
+				BoundingMeshStaticOptions.ActiveTags.Empty();
+				BoundingMeshStaticOptions.ModifiersToIgnore.AddUnique(Modifier);
+				FMeshGenerationDynamicOptions BoundingMeshDynamicOptions;
+				BoundingMeshDynamicOptions.bLayouts = false;
+				FMeshTask BoundingMeshTask = GenerateMesh(BoundingMeshStaticOptions, UE::Tasks::MakeCompletedTask<FMeshGenerationDynamicOptions>(BoundingMeshDynamicOptions), TypedTransformNode->BoundingMesh);
+
+				LastMeshTask = UE::Tasks::Launch(TEXT("MutableMeshTransformWithBoundingMeshModifier"),
+					[BoundingMeshTask, LastMeshTask, Modifier, MatrixOp, ErrorLog = ErrorLog, ErrorContext]() mutable
 					{
-						ErrorLog->Add("Bounding mesh or matrix have not been generated", ELMT_ERROR, ErrorContext);
-						continue;
-					}
-
-					// Transform matrix.
-					Ptr<ASTOp> MatrixOp;
-					FMatrixGenerationResult ChildResult;
-					GenerateMatrix(ChildResult, StaticOptions, TypedTransformNode->MatrixNode);
-					MatrixOp = ChildResult.op;
+						FMeshGenerationResult BoundingResults = BoundingMeshTask.GetResult();
+						FMeshGenerationResult LastMeshResults = LastMeshTask.GetResult();
 
-					// Bounding mesh
-					FMeshGenerationStaticOptions ClipStaticOptions(StaticOptions);
-					ClipStaticOptions.ActiveTags.Empty();
-					ClipStaticOptions.ModifiersToIgnore.AddUnique(Modifier);
-					FMeshGenerationDynamicOptions BoundingMeshDynamicOptions;
-					BoundingMeshDynamicOptions.bLayouts = false;
-					FMeshTask BoundingMeshTask = GenerateMesh(ClipStaticOptions, UE::Tasks::MakeCompletedTask<FMeshGenerationDynamicOptions>(BoundingMeshDynamicOptions), TypedTransformNode->BoundingMesh);
+						Ptr<ASTOpMeshTransformWithBoundingMesh> TransformOp = new ASTOpMeshTransformWithBoundingMesh();
+						TransformOp->source = LastMeshResults.MeshOp;
+						TransformOp->matrix = MatrixOp;
+						TransformOp->boundingMesh = BoundingResults.MeshOp;
 
-					LastMeshTask = UE::Tasks::Launch(TEXT("MutableMeshTransformWithBoundingMeshModifier"),
-						[BoundingMeshTask, LastMeshTask, Modifier, MatrixOp, ErrorLog = ErrorLog, ErrorContext]() mutable
+						if (TransformOp->boundingMesh)
 						{
-							FMeshGenerationResult BoundingResults = BoundingMeshTask.GetResult();
-							FMeshGenerationResult LastMeshResults = LastMeshTask.GetResult();
-
-							Ptr<ASTOpMeshTransformWithBoundingMesh> TransformOp = new ASTOpMeshTransformWithBoundingMesh();
-							TransformOp->source = LastMeshResults.MeshOp;
-							TransformOp->matrix = MatrixOp;
-							TransformOp->boundingMesh = BoundingResults.MeshOp;
-
-							if (TransformOp->boundingMesh)
+							ASTOp::EClosedMeshTest IsClosed = TransformOp->boundingMesh->IsClosedMesh();
+							if (IsClosed == ASTOp::EClosedMeshTest::No)
 							{
-								ASTOp::EClosedMeshTest IsClosed = TransformOp->boundingMesh->IsClosedMesh();
-								if (IsClosed == ASTOp::EClosedMeshTest::No)
-								{
-									ErrorLog->Add("Mesh used for clipping is not closed.", ELMT_WARNING, ErrorContext);
-								}
+								ErrorLog->Add("Mesh used for the mesh transformation is not closed.", ELMT_WARNING, ErrorContext);
 							}
+						}
 
-							// Condition to apply the transform op
-							if (Modifier.FinalCondition)
-							{
-								Ptr<ASTOpConditional> ConditionalOp = new ASTOpConditional();
-								ConditionalOp->type = EOpType::ME_CONDITIONAL;
-								ConditionalOp->no = LastMeshResults.MeshOp;
-								ConditionalOp->yes = TransformOp;
-								ConditionalOp->condition = Modifier.FinalCondition;
-								LastMeshResults.MeshOp = ConditionalOp;
-							}
-							else
-							{
-								LastMeshResults.MeshOp = TransformOp;
-							}
+						// Condition to apply the transform op
+						if (Modifier.FinalCondition)
+						{
+							Ptr<ASTOpConditional> ConditionalOp = new ASTOpConditional();
+							ConditionalOp->type = EOpType::ME_CONDITIONAL;
+							ConditionalOp->no = LastMeshResults.MeshOp;
+							ConditionalOp->yes = TransformOp;
+							ConditionalOp->condition = Modifier.FinalCondition;
+							LastMeshResults.MeshOp = ConditionalOp;
+						}
+						else
+						{
+							LastMeshResults.MeshOp = TransformOp;
+						}
+
+						return LastMeshResults;
+					},
+					UE::Tasks::Prerequisites(BoundingMeshTask, LastMeshTask)
+				);
 
-							return LastMeshResults;
-						},
-						UE::Tasks::Prerequisites(BoundingMeshTask, LastMeshTask)
-					);
-				}
 			}
 			else if (Modifier.Node->GetType() == NodeModifierMeshTransformWithBone::GetStaticType())
 			{
diff --git a/Engine/Plugins/Mutable/Source/MutableTools/Private/MuT/CodeGenerator.h b/Engine/Plugins/Mutable/Source/MutableTools/Private/MuT/CodeGenerator.h
--- a/Engine/Plugins/Mutable/Source/MutableTools/Private/MuT/CodeGenerator.h
+++ b/Engine/Plugins/Mutable/Source/MutableTools/Private/MuT/CodeGenerator.h
@@ -134,6 +134,7 @@ namespace UE::Mutable::Private
 	class NodeMatrixParameter;
 	class NodeMaterialParameter;
 	class NodeMeshTransformWithBone;
+	class NodeMeshTransformInMesh;
 	struct FObjectState;
 	struct FProgram;
 
@@ -861,6 +862,7 @@ namespace UE::Mutable::Private
         FMeshTask GenerateMesh_ApplyPose(const FMeshGenerationStaticOptions&, FMeshOptionsTask, const NodeMeshApplyPose*);
         FMeshTask GenerateMesh_RemoveMesh(const FMeshGenerationStaticOptions&, FMeshOptionsTask, const NodeMeshRemoveMesh*);
 		FMeshTask GenerateMesh_TransformWithBone(const FMeshGenerationStaticOptions&, FMeshOptionsTask, const NodeMeshTransformWithBone*);
+    	FMeshTask GenerateMesh_TransformInMesh(const FMeshGenerationStaticOptions&, const FMeshOptionsTask, const NodeMeshTransformInMesh*);
         FMeshTask GenerateMesh_Variation(const FMeshGenerationStaticOptions&, FMeshOptionsTask, const NodeMeshVariation*);
 		FMeshTask GenerateMesh_Table(const FMeshGenerationStaticOptions&, FMeshOptionsTask, const NodeMeshTable*);
 		FMeshTask GenerateMesh_Reshape(const FMeshGenerationStaticOptions&, FMeshOptionsTask, const NodeMeshReshape*);
diff --git a/Engine/Plugins/Mutable/Source/MutableTools/Private/MuT/CodeGenerator_Mesh.cpp b/Engine/Plugins/Mutable/Source/MutableTools/Private/MuT/CodeGenerator_Mesh.cpp
--- a/Engine/Plugins/Mutable/Source/MutableTools/Private/MuT/CodeGenerator_Mesh.cpp
+++ b/Engine/Plugins/Mutable/Source/MutableTools/Private/MuT/CodeGenerator_Mesh.cpp
@@ -7,6 +7,7 @@
 #include "ASTOpMeshRemoveMask.h"
 #include "ASTOpMeshSkeletalMeshBreak.h"
 #include "ASTOpMeshTransformWithBone.h"
+#include "ASTOpMeshTransformWithBoundingMesh.h"
 #include "Containers/Array.h"
 #include "Misc/AssertionMacros.h"
 #include "MuR/Layout.h"
@@ -60,6 +61,7 @@
 #include "MuT/NodeMeshTransform.h"
 #include "MuT/NodeMeshVariation.h"
 #include "MuT/NodeMeshSkeletalMeshBreak.h"
+#include "MuT/NodeMeshTransformInMesh.h"
 #include "MuT/NodeMeshTransformWithBone.h"
 #include "MuT/Table.h"
 #include "MuT/TablePrivate.h"
@@ -102,6 +104,7 @@ namespace UE::Mutable::Private
 		case Node::EType::MeshClipWithMesh: Result = GenerateMesh_ClipWithMesh(StaticOptions, InOptions, static_cast<const NodeMeshClipWithMesh*>(Node)); break;
 		case Node::EType::MeshRemoveMesh: Result = GenerateMesh_RemoveMesh(StaticOptions, InOptions, static_cast<const NodeMeshRemoveMesh*>(Node)); break;
 		case Node::EType::MeshTransformWithBone: Result = GenerateMesh_TransformWithBone(StaticOptions, InOptions, static_cast<const NodeMeshTransformWithBone*>(Node)); break;
+		case Node::EType::MeshTransformInMesh: Result = GenerateMesh_TransformInMesh(StaticOptions, InOptions, static_cast<const NodeMeshTransformInMesh*>(Node)); break;
 		case Node::EType::MeshApplyPose: Result = GenerateMesh_ApplyPose(StaticOptions, InOptions, static_cast<const NodeMeshApplyPose*>(Node)); break;
 		case Node::EType::MeshVariation: Result = GenerateMesh_Variation(StaticOptions, InOptions, static_cast<const NodeMeshVariation*>(Node)); break;
 		case Node::EType::MeshTable: Result = GenerateMesh_Table(StaticOptions, InOptions, static_cast<const NodeMeshTable*>(Node)); break;
@@ -1216,33 +1219,32 @@ namespace UE::Mutable::Private
 	
 	FMeshTask CodeGenerator::GenerateMesh_TransformWithBone(const FMeshGenerationStaticOptions& InStaticOptions, FMeshOptionsTask InOptions, const NodeMeshTransformWithBone* InNode)
 	{
-		const NodeMeshTransformWithBone* TransformWithBoneNode = static_cast<const NodeMeshTransformWithBone*>(InNode);
-		const void* MessageContext = TransformWithBoneNode->GetMessageContext();
+		const void* MessageContext = InNode->GetMessageContext();
 		
-		if (!TransformWithBoneNode->Source)
+		if (!InNode->Source)
 		{
 			// This argument is required
 			ErrorLog->Add("Mesh Transform with Bone source node is not set.", ELMT_ERROR, MessageContext);
 			return UE::Tasks::MakeCompletedTask<FMeshGenerationResult>();
 		}
 		
-		if (!TransformWithBoneNode->MatrixNode)
+		// Generate the source mesh
+		FMeshTask MeshGenerationTask = GenerateMesh(InStaticOptions, InOptions, InNode->Source);
+		
+		if (!InNode->MatrixNode)
 		{
 			// This argument is required
 			ErrorLog->Add("Mesh Transform with Bone transform matrix node is not set.", ELMT_ERROR, MessageContext);
-			return UE::Tasks::MakeCompletedTask<FMeshGenerationResult>();
+			return MeshGenerationTask;
 		}
 		
-		// Generate the source mesh
-		FMeshTask MeshGenerationTask = GenerateMesh(InStaticOptions, InOptions, TransformWithBoneNode->Source);
-		
 		// Transform matrix.
 		FMatrixGenerationResult ChildResult;
-		GenerateMatrix(ChildResult, InStaticOptions, TransformWithBoneNode->MatrixNode);
+		GenerateMatrix(ChildResult, InStaticOptions, InNode->MatrixNode);
 		Ptr<ASTOp> MatrixOp = ChildResult.op;
 		
 		return UE::Tasks::Launch(TEXT("MutableMeshTransformWithBone"),
-			[MeshGenerationTask, MatrixOp, BoneName = TransformWithBoneNode->BoneName, ThresholdFactor = TransformWithBoneNode->ThresholdFactor]() mutable
+			[MeshGenerationTask, MatrixOp, BoneName = InNode->BoneName, ThresholdFactor = InNode->ThresholdFactor]() mutable
 			{
 				FMeshGenerationResult SourceMeshGenerationResults = MeshGenerationTask.GetResult();
 
@@ -1260,6 +1262,73 @@ namespace UE::Mutable::Private
 		);
 	}
 
+	
+	FMeshTask CodeGenerator::GenerateMesh_TransformInMesh(const FMeshGenerationStaticOptions& InStaticOptions, const FMeshOptionsTask InOptions, const NodeMeshTransformInMesh* InNode)
+	{
+		const void* MessageContext = InNode->GetMessageContext();
+		
+		if (!InNode->SourceMesh)
+		{
+			ErrorLog->Add("Mesh Transform In Mesh source node is not set.", ELMT_ERROR, MessageContext);
+			return UE::Tasks::MakeCompletedTask<FMeshGenerationResult>();
+		}
+		
+		// Generate the source mesh
+		FMeshTask SourceMeshTask = GenerateMesh(InStaticOptions, InOptions, InNode->SourceMesh);
+		
+		if (!InNode->MatrixNode)
+		{
+			ErrorLog->Add("Mesh Transform In Mesh transform matrix node is not set.", ELMT_ERROR, MessageContext);
+			return SourceMeshTask;
+		}
+		
+		// Transform matrix.
+		FMatrixGenerationResult ChildResult;
+		GenerateMatrix(ChildResult, InStaticOptions, InNode->MatrixNode);
+		Ptr<ASTOp> MatrixOp = ChildResult.op;
+		
+		if (!InNode->BoundingMesh)
+		{
+			ErrorLog->Add("Mesh Transform In Mesh bounding mesh node is not set.", ELMT_ERROR, MessageContext);
+			return SourceMeshTask;
+		}
+		
+		// Bounding mesh
+		FMeshGenerationStaticOptions BoundingMeshStaticOptions(InStaticOptions);
+		BoundingMeshStaticOptions.ActiveTags.Empty();
+		FMeshGenerationDynamicOptions BoundingMeshDynamicOptions;
+		BoundingMeshDynamicOptions.bLayouts = false;
+		FMeshTask BoundingMeshTask = GenerateMesh(BoundingMeshStaticOptions, UE::Tasks::MakeCompletedTask<FMeshGenerationDynamicOptions>(BoundingMeshDynamicOptions), InNode->BoundingMesh);
+		
+		return UE::Tasks::Launch(TEXT("MutableMeshTransformWithBoundingMeshModifier"),
+			[BoundingMeshTask, SourceMeshTask, MatrixOp, ErrorLog = ErrorLog, MessageContext]() mutable
+			{
+				FMeshGenerationResult BoundingResults = BoundingMeshTask.GetResult();
+				FMeshGenerationResult SourceMeshResults = SourceMeshTask.GetResult();
+
+				Ptr<ASTOpMeshTransformWithBoundingMesh> TransformOp = new ASTOpMeshTransformWithBoundingMesh();
+				TransformOp->source = SourceMeshResults.MeshOp;
+				TransformOp->matrix = MatrixOp;
+				TransformOp->boundingMesh = BoundingResults.MeshOp;
+
+				if (TransformOp->boundingMesh)
+				{
+					ASTOp::EClosedMeshTest IsClosed = TransformOp->boundingMesh->IsClosedMesh();
+					if (IsClosed == ASTOp::EClosedMeshTest::No)
+					{
+						ErrorLog->Add("Bounding mesh used is not closed.", ELMT_WARNING, MessageContext);
+					}
+				}
+
+				FMeshGenerationResult Result = SourceMeshResults;
+				Result.MeshOp = TransformOp;
+						
+				return Result;
+			},
+			UE::Tasks::Prerequisites(SourceMeshTask, BoundingMeshTask)
+		);
+	}
+
 
 	FMeshTask CodeGenerator::GenerateMesh_RemoveMesh(const FMeshGenerationStaticOptions& StaticOptions, FMeshOptionsTask InOptions, const NodeMeshRemoveMesh* RemoveMeshNode)
 	{
diff --git a/Engine/Plugins/Mutable/Source/MutableTools/Private/MuT/Node.cpp b/Engine/Plugins/Mutable/Source/MutableTools/Private/MuT/Node.cpp
--- a/Engine/Plugins/Mutable/Source/MutableTools/Private/MuT/Node.cpp
+++ b/Engine/Plugins/Mutable/Source/MutableTools/Private/MuT/Node.cpp
@@ -117,6 +117,7 @@
 #include "MuT/NodeSkeletalMeshParameter.h"
 #include "MuT/NodeMaterialModify.h"
 #include "MuT/NodeMeshRemoveMesh.h"
+#include "MuT/NodeMeshTransformInMesh.h"
 #include "MuT/NodeMeshTransformWithBone.h"
 
 
@@ -219,6 +220,7 @@ namespace UE::Mutable::Private
 	FNodeType NodeMeshClipWithMesh::StaticType = FNodeType(Node::EType::MeshClipWithMesh, NodeMesh::GetStaticType());
 	FNodeType NodeMeshRemoveMesh::StaticType = FNodeType(Node::EType::MeshRemoveMesh, NodeMesh::GetStaticType());
 	FNodeType NodeMeshTransformWithBone::StaticType = FNodeType(Node::EType::MeshTransformWithBone, NodeMesh::GetStaticType());
+	FNodeType NodeMeshTransformInMesh::StaticType = FNodeType(Node::EType::MeshTransformInMesh, NodeMesh::GetStaticType());
 	FNodeType NodeMeshSkeletalMeshBreak::StaticType = FNodeType(Node::EType::MeshSkeletalMeshBreak, NodeMesh::GetStaticType());
 	FNodeType NodeMeshMakeMorph::StaticType = FNodeType(Node::EType::MeshMakeMorph, NodeMesh::GetStaticType());
 	FNodeType NodeMeshApplyPose::StaticType = FNodeType(Node::EType::MeshApplyPose, NodeMesh::GetStaticType());
