[mutable] Regex-based CO discovery for the MutableEditortests : Engine code changes
- Added a new type of scoped section that defines the logs part of one plugin's contents
- Added new commandlet to handle the discovery and testing of COs part of multiple plugins
- Other small changes and tweaks.

#rnx

#Synced-CL 51453194

diff --git a/Engine/Plugins/Mutable/Source/MutableValidation/Private/MuV/COIBakingTestCommandlet.cpp b/Engine/Plugins/Mutable/Source/MutableValidation/Private/MuV/COIBakingTestCommandlet.cpp
--- a/Engine/Plugins/Mutable/Source/MutableValidation/Private/MuV/COIBakingTestCommandlet.cpp
+++ b/Engine/Plugins/Mutable/Source/MutableValidation/Private/MuV/COIBakingTestCommandlet.cpp
@@ -136,9 +136,20 @@ int32 UCOIBakingTestCommandlet::Main(const FString& Params)
 
 		UE_LOG(LogMutable, Display, TEXT("Scheduling Customizable Object compilation."));
 		ScheduleCOCompilationForBaking(*TargetInstance, CompilationEndDelegate, bPerformPartialCompilation);
-
-		WaitUntilConditionChanges(bIsCustomizableObjectBeingCompiled);
-
+		
+		// Exit the wait process once the compilation finishes
+		while (bIsCustomizableObjectBeingCompiled) 
+		{ 
+			// Tick the engine
+			CommandletHelpers::TickEngine();
+
+			// Stop if exit was requested
+			if (IsEngineExitRequested())
+			{
+				break;
+			}
+		}
+		
 		if (!bWasCustomizableObjectCompilationSuccessful)
 		{
 			UE_LOG(LogMutable, Error, TEXT("Failed to successfully compile the target COIs CO. Exiting commandlet."));
@@ -161,9 +172,19 @@ int32 UCOIBakingTestCommandlet::Main(const FString& Params)
 		bIsInstanceBeingUpdated = true;
 		UE_LOG(LogMutable, Display, TEXT("Scheduling instance update."));
 		ScheduleInstanceUpdateForBaking(*TargetInstance, InstanceUpdateDelegate);
-
+		
 		// Now tick the engine so the instance gets updated while running in the commandlet context
-		WaitUntilConditionChanges(bIsInstanceBeingUpdated);
+		while (bIsInstanceBeingUpdated) 
+		{ 
+			// Tick the engine
+			CommandletHelpers::TickEngine();
+
+			// Stop if exit was requested
+			if (IsEngineExitRequested())
+			{
+				break;
+			}
+		}
 		
 		// Check the end status of the instance update
 		if (!bWasInstanceUpdateSuccessful)
diff --git a/Engine/Plugins/Mutable/Source/MutableValidation/Private/MuV/COMultiPluginBulkValidationCommandlet.cpp b/Engine/Plugins/Mutable/Source/MutableValidation/Private/MuV/COMultiPluginBulkValidationCommandlet.cpp
new file mode 100644
--- /dev/null
+++ b/Engine/Plugins/Mutable/Source/MutableValidation/Private/MuV/COMultiPluginBulkValidationCommandlet.cpp
@@ -0,0 +1,176 @@
+﻿// Copyright Epic Games, Inc. All Rights Reserved.
+
+#include "MuV/COMultiPluginBulkValidationCommandlet.h"
+
+#include "ScopedLogSection.h"
+#include "ShaderCompiler.h"
+#include "ValidationUtils.h"
+#include "MuCO/CustomizableObject.h"
+#include "MuCO/CustomizableObjectSystemPrivate.h"
+#include "MuCO/ICustomizableObjectEditorModule.h"
+#include "MuCO/LoadUtils.h"
+#include "MuCO/LogBenchmarkUtil.h"
+
+#include UE_INLINE_GENERATED_CPP_BY_NAME(COMultiPluginBulkValidationCommandlet)
+
+int32 UCOMultiPluginBulkValidationCommandlet::Main(const FString& Params)
+{
+	LLM_SCOPE_BYNAME(TEXT("COMultiPluginBulkValidationCommandlet"));
+
+	// Prepare the environment for the testing -------------------------------------------------------------------------------------------------------
+	
+	// Ensure we have set the mutable system to the benchmarking mode and that we are reporting benchmarking data
+	FLogBenchmarkUtil::SetBenchmarkReportingStateOverride(true);
+	UCustomizableObjectSystemPrivate::SetUsageOfBenchmarkingSettings(true);
+	
+	// Ensure we do not show any OK dialog since we are not a user that can interact with them
+	GIsRunningUnattendedScript = true;
+	
+	// Parse the arguments provided ------------------------------------------------------------------------------------------------------------------
+	
+	// Parse the value for the platform to be used for the compilation of the COs
+	const ITargetPlatform* TargetCompilationPlatform = GetCompilationPlatform(Params);
+	if (!TargetCompilationPlatform)
+	{
+		return 1;
+	}
+
+	// Get the amount of instances to generate if parameter was provided (it will get multiplied by the amount of states later so this is a minimun value)
+	const uint32 InstancesToGenerate = GetTargetAmountOfInstances(Params);
+
+	// Work only for the root Customizable Objects found
+	const bool bOnlyTestRootObjects = GetIfOnlyRootCOsShouldBeUsed(Params);
+	
+	// Find the assets to be tested based on the plugin names provided -------------------------------------------------------------------------------
+
+	TArray<FName> PluginsToScanForCOs;
+	{
+		FString RawCollectionOfPluginNames;
+		if (FParse::Value(*Params, TEXT("PluginsToScanForCOs="), RawCollectionOfPluginNames, false))
+		{
+			TArray<FString> SeparatedPluginNames;
+			if (RawCollectionOfPluginNames.ParseIntoArray(SeparatedPluginNames, TEXT(","), true))
+			{
+				for (const FString& PluginNameString : SeparatedPluginNames)
+				{
+					const FName PluginName = FName(PluginNameString);
+					
+					if (PluginName.IsNone())
+					{
+						UE_LOG(LogMutable, Warning, TEXT("An empty PluginName was provided. Skipping."));
+						continue;
+					}
+					
+					PluginsToScanForCOs.Add(PluginName);
+				}
+			}
+		}
+	}
+	
+	// Perform a blocking search to ensure all assets used by mutable are reachable using the AssetRegistry
+	PrepareAssetRegistry();
+	
+	// Make sure there is nothing else that the engine needs to do before starting our test
+	// Wait(60);	// todo: UE-304050 Remove this wait as it may no longer be required due to us calling for GShaderCompilingManager->FinishAllCompilation()
+	
+	// Block until async shader compiling is finished before we try to use the shaders for exporting
+	// The code is structured to only block once for all materials, so that shader compiling is able to utilize many cores
+	if (GShaderCompilingManager && GShaderCompilingManager->GetNumRemainingJobs() > 0)
+	{
+		GShaderCompilingManager->FinishAllCompilation();
+	}
+	
+	LogGlobalSettings();
+	
+	for (int32 TargetPluginIndex = 0; TargetPluginIndex < PluginsToScanForCOs.Num(); TargetPluginIndex++)
+	{		
+		// All CO pointing FAssetData found at the to scan plugin path
+		TArray<FAssetData> FoundAssetData;
+		{
+			LLM_SCOPE_BYNAME(TEXT("COMultiPluginBulkValidationCommandlet/AssetsSearch"));
+		
+			const FString PluginSearchPath = FString::Printf(TEXT("/%s"), *PluginsToScanForCOs[TargetPluginIndex].ToString());
+			const FName CustomizableObjectsSearchPath = FName(PluginSearchPath);
+		
+			FoundAssetData = FindAllAssetsAtPath(CustomizableObjectsSearchPath, UCustomizableObject::StaticClass());
+			if (FoundAssetData.IsEmpty())
+			{
+				UE_LOG(LogMutable, Display, TEXT("Found no Customizable Objects to validate in the provided path : %s"), *CustomizableObjectsSearchPath.ToString());
+				continue;
+			}
+		
+			// Log all the Customizable Objects to be tested:
+			UE_LOG(LogMutable, Display, TEXT("Found a total of %u Customizable Objects to validate at '%s'. Some may be discarded based on the test settings."), FoundAssetData.Num(), *CustomizableObjectsSearchPath.ToString())
+			for (const FAssetData& MutableAssetData : FoundAssetData)
+			{
+				UE_LOG(LogMutable, Display, TEXT("\t - %s (%s)"), *MutableAssetData.AssetName.ToString(), *MutableAssetData.PackageName.ToString());
+			}
+		}
+		
+		// Try to collect the garbage
+		{
+			if (GIsInitialLoad)
+			{
+				UE_LOG(LogMutable, Warning, TEXT("GC will not run as GIsInitialLoad is currently set to true."));
+			}
+			
+			CollectGarbage(RF_NoFlags, true);
+		}
+		
+		// FoundAssetData is not empty so we can start to test things out
+		{
+			const FName PluginName = PluginsToScanForCOs[TargetPluginIndex];
+			check(!PluginName.IsNone());
+		
+			// Make all the logs that appear in this scope have a relation with the plugin name provided. This will be later used by the MTU to name the CO generated data
+			const FScopedLogSection PluginLogSection (EMutableLogSection::Plugin, PluginName);
+			
+			UE_LOG(LogMutable, Display, TEXT("Starting testing of COs from the '%s' plugin ."), *PluginName.ToString())
+					
+			// iterate over the FoundAsstData of each Plugin-Dir structure
+			const int32 AssetDataCount = FoundAssetData.Num();
+			for (int32 AssetDataIndex = 0; AssetDataIndex < AssetDataCount; AssetDataIndex++)
+			{
+				LLM_SCOPE_BYNAME(TEXT("COMultiPluginBulkValidationCommandlet/COTest"));
+				
+				const FAssetData& AssetData = FoundAssetData[AssetDataIndex];
+				TObjectPtr<UObject> FoundObject = UE::Mutable::Private::LoadObject(AssetData);
+				if (!FoundObject)
+				{
+					UE_LOG(LogMutable, Error, TEXT("Failed to load the asset with path : %s ."), *AssetData.GetSoftObjectPath().ToString())
+					continue;
+				}
+
+				TObjectPtr<UCustomizableObject> TargetCustomizableObject = CastChecked<UCustomizableObject>(FoundObject);
+				
+				// If we only want to test root objects and this is not one go to the next CO
+				if (bOnlyTestRootObjects && !ICustomizableObjectEditorModule::GetChecked().IsRootObject(*TargetCustomizableObject.Get()))
+				{
+					UE_LOG(LogMutable, Display, TEXT("Skipping CO \"%s\" as it is not a root CO."), *TargetCustomizableObject->GetName())
+					continue;
+				}
+			
+				// Body of the test ----------------------------------------------------------------------------------------------------------------------
+				TestCustomizableObject(*TargetCustomizableObject, *TargetCompilationPlatform, InstancesToGenerate);
+				// ---------------------------------------------------------------------------------------------------------------------------------------
+			
+				// Try to collect the garbage
+				{
+					if (GIsInitialLoad)
+					{
+						UE_LOG(LogMutable, Warning, TEXT("GC will not run as GIsInitialLoad is currently set to true."));
+					}
+				
+					CollectGarbage(RF_NoFlags, true);
+				}
+
+				UE_LOG(LogMutable, Display, TEXT("%s : Validated %u/%u assets."), *PluginName.ToString(), AssetDataIndex + 1, AssetDataCount);
+			}
+			
+			UE_LOG(LogMutable, Display, TEXT("Finished testing the COs from the '%s' plugin ."), *PluginName.ToString())
+		}
+	}
+	
+	UE_LOG(LogMutable, Display, TEXT("Mutable commandlet finished."));
+	return 0;
+}
diff --git a/Engine/Plugins/Mutable/Source/MutableValidation/Private/MuV/COMultiPluginBulkValidationCommandlet.h b/Engine/Plugins/Mutable/Source/MutableValidation/Private/MuV/COMultiPluginBulkValidationCommandlet.h
new file mode 100644
--- /dev/null
+++ b/Engine/Plugins/Mutable/Source/MutableValidation/Private/MuV/COMultiPluginBulkValidationCommandlet.h
@@ -0,0 +1,16 @@
+﻿// Copyright Epic Games, Inc. All Rights Reserved.
+
+#pragma once
+
+#include "Commandlets/Commandlet.h"
+#include "COMultiPluginBulkValidationCommandlet.generated.h"
+
+
+UCLASS()
+class UCOMultiPluginBulkValidationCommandlet : public UCommandlet
+{
+	GENERATED_BODY()
+	
+public:
+	virtual int32 Main(const FString& Params) override;
+};
diff --git a/Engine/Plugins/Mutable/Source/MutableValidation/Private/MuV/CustomizableObjectBulkValidationCommandlet.cpp b/Engine/Plugins/Mutable/Source/MutableValidation/Private/MuV/CustomizableObjectBulkValidationCommandlet.cpp
--- a/Engine/Plugins/Mutable/Source/MutableValidation/Private/MuV/CustomizableObjectBulkValidationCommandlet.cpp
+++ b/Engine/Plugins/Mutable/Source/MutableValidation/Private/MuV/CustomizableObjectBulkValidationCommandlet.cpp
@@ -39,12 +39,7 @@ int32 UCustomizableObjectBulkValidationCommandlet::Main(const FString& Params)
 	const uint32 InstancesToGenerate = GetTargetAmountOfInstances(Params);
 
 	// Work only for the root Customizable Objects found
-	bool bOnlyTestRootObjects = false;
-	FParse::Bool(*Params,TEXT("SkipNonRootObjects="),bOnlyTestRootObjects);
-	if (bOnlyTestRootObjects)
-	{
-		UE_LOG(LogMutable, Display, TEXT("Only the root COs will be tested")) ;
-	}
+	const bool bOnlyTestRootObjects = GetIfOnlyRootCOsShouldBeUsed(Params);
 
 	// Find the assets to be tested based on the path provided ---------------------------------------------------------------------------------------
 	
@@ -117,19 +112,12 @@ int32 UCustomizableObjectBulkValidationCommandlet::Main(const FString& Params)
 			continue;
 		}
 
-		TObjectPtr<UCustomizableObject> TargetCustomizableObject = Cast<UCustomizableObject>(FoundObject);
-		if (!TargetCustomizableObject)
-		{
-			UE_LOG(LogMutable, Error, TEXT("Failed to cast found UObject to UCustomizableObject."));
-			continue;
-		}		
-		
-		const FString CustomizableObjectName = TargetCustomizableObject->GetName();
-		
+		TObjectPtr<UCustomizableObject> TargetCustomizableObject = CastChecked<UCustomizableObject>(FoundObject);
+
 		// If we only want to test root objects and this is not one go to the next CO
-		if (bOnlyTestRootObjects && !ICustomizableObjectEditorModule::GetChecked().GetRootObject(TargetCustomizableObject.Get()))
+		if (bOnlyTestRootObjects && !ICustomizableObjectEditorModule::GetChecked().IsRootObject(*TargetCustomizableObject.Get()))
 		{
-			UE_LOG(LogMutable, Display, TEXT("Skipping CO \"%s\" as it is not a root CO."), *CustomizableObjectName)
+			UE_LOG(LogMutable, Display, TEXT("Skipping CO \"%s\" as it is not a root CO."), *TargetCustomizableObject->GetName())
 			continue;
 		}
 		
diff --git a/Engine/Plugins/Mutable/Source/MutableValidation/Private/MuV/ScopedLogSection.cpp b/Engine/Plugins/Mutable/Source/MutableValidation/Private/MuV/ScopedLogSection.cpp
--- a/Engine/Plugins/Mutable/Source/MutableValidation/Private/MuV/ScopedLogSection.cpp
+++ b/Engine/Plugins/Mutable/Source/MutableValidation/Private/MuV/ScopedLogSection.cpp
@@ -13,16 +13,23 @@ FScopedLogSection::FScopedLogSection(EMutableLogSection Section, const FName& In
 {
 	// Set the current section handled by this object
 	CurrentSection = Section;
-
-	// Cache the name of the current target. Will only be displayed for sections of type Object.
-	if (!InSectionTarget.IsNone())
+	
+	if (CurrentSection == EMutableLogSection::Object)
 	{
+		check(SectionObject.IsNone());		// A new object section can not be opened if another is already opened.
+		check(!InSectionTarget.IsNone());
 		SectionObject = InSectionTarget;
+		
+		UE_LOG(LogMutableValidation, Log, TEXT(" SECTION START : %s - [%s] "), *GetLogSectionName(CurrentSection), *SectionObject.ToString());
 	}
-
-	if (CurrentSection == EMutableLogSection::Object)
+	else if (CurrentSection == EMutableLogSection::Plugin)
 	{
-		UE_LOG(LogMutableValidation, Log, TEXT(" SECTION START : %s - [%s] "), *GetLogSectionName(CurrentSection), *SectionObject.ToString());
+		check(SectionPlugin.IsNone());		// A new plugin section can not be opened if another is already open.
+		check(SectionObject.IsNone());		// A new plugin section can not be opened when an object section is already open.
+		check(!InSectionTarget.IsNone());
+		SectionPlugin = InSectionTarget;
+		
+		UE_LOG(LogMutableValidation, Log, TEXT(" SECTION START : %s - [%s] "), *GetLogSectionName(CurrentSection), *SectionPlugin.ToString());
 	}
 	else
 	{
@@ -36,6 +43,12 @@ FScopedLogSection::~FScopedLogSection()
 	if (CurrentSection == EMutableLogSection::Object)
 	{
 		UE_LOG(LogMutableValidation, Log, TEXT(" SECTION END : %s - [%s] "), *GetLogSectionName(CurrentSection), *SectionObject.ToString());
+		SectionObject = NAME_None;
+	}
+	else if (CurrentSection == EMutableLogSection::Plugin)
+	{
+		UE_LOG(LogMutableValidation, Log, TEXT(" SECTION END : %s - [%s] "), *GetLogSectionName(CurrentSection), *SectionPlugin.ToString());
+		SectionPlugin = NAME_None;
 	}
 	else
 	{
@@ -46,7 +59,6 @@ FScopedLogSection::~FScopedLogSection()
 	CurrentSection = EMutableLogSection::Undefined;
 }
 
-
 	
 FString FScopedLogSection::GetLogSectionName(EMutableLogSection Section) const
 {
@@ -69,6 +81,9 @@ FString FScopedLogSection::GetLogSectionName(EMutableLogSection Section) const
 	 case EMutableLogSection::Object:
 	 	Output = TEXT("object");
 	 	break;
+	 case EMutableLogSection::Plugin:
+	 	Output = TEXT("plugin");
+	 	break;
 	 default:
 	 	checkNoEntry();
 	 	break;
diff --git a/Engine/Plugins/Mutable/Source/MutableValidation/Private/MuV/ScopedLogSection.h b/Engine/Plugins/Mutable/Source/MutableValidation/Private/MuV/ScopedLogSection.h
--- a/Engine/Plugins/Mutable/Source/MutableValidation/Private/MuV/ScopedLogSection.h
+++ b/Engine/Plugins/Mutable/Source/MutableValidation/Private/MuV/ScopedLogSection.h
@@ -20,6 +20,8 @@ enum class EMutableLogSection : uint8
 	Bake,
 	
 	Object,
+	
+	Plugin,
 };
 
 
@@ -42,4 +44,8 @@ private:
 
 	/** The object the current section is representing */
 	static inline FName SectionObject;
-};
\ No newline at end of file
+	
+	/** The object the current section is representing */
+	static inline FName SectionPlugin;
+};
+
diff --git a/Engine/Plugins/Mutable/Source/MutableValidation/Private/MuV/ValidationUtils.cpp b/Engine/Plugins/Mutable/Source/MutableValidation/Private/MuV/ValidationUtils.cpp
--- a/Engine/Plugins/Mutable/Source/MutableValidation/Private/MuV/ValidationUtils.cpp
+++ b/Engine/Plugins/Mutable/Source/MutableValidation/Private/MuV/ValidationUtils.cpp
@@ -69,24 +69,6 @@ void Wait(const double ToWaitSeconds)
 }
 
 
-void WaitUntilConditionChanges(bool& bInWaitCondition)
-{
-	const bool bStartingCondition = bInWaitCondition;
-	// Exit the wait process once the provided conditional value changes
-	while (bInWaitCondition == bStartingCondition) 
-	{ 
-		// Tick the engine
-		CommandletHelpers::TickEngine();
-
-		// Stop if exit was requested
-		if (IsEngineExitRequested())
-		{
-			break;
-		}
-	}
-}
-
-
 FCompilationOptions GetCompilationOptionsForBenchmarking (const UCustomizableObject& ReferenceCustomizableObject)
 {
 	// Override some configurations that may have been changed by the user
@@ -174,10 +156,30 @@ ITargetPlatform* GetCompilationPlatform(const FString& Params)
 }
 
 
+bool GetIfOnlyRootCOsShouldBeUsed(const FString& Params)
+{
+	bool bOnlyTestRootObjects = false;
+	FParse::Bool(*Params,TEXT("SkipNonRootObjects="),bOnlyTestRootObjects);
+	if (bOnlyTestRootObjects)
+	{
+		UE_LOG(LogMutable, Display, TEXT("Only the root COs will be tested")) ;
+	}
+	
+	return bOnlyTestRootObjects;
+}
+
+
 bool TestCustomizableObject(UCustomizableObject& InTargetCustomizableObject, const ITargetPlatform& TargetCompilationPlatform,
 	const uint32 InstancesToGenerate)
 {
-	const FScopedLogSection ObjectSection (EMutableLogSection::Object, FName( InTargetCustomizableObject.GetPathName()));
+	const FName ObjectName = FName(InTargetCustomizableObject.GetPathName());
+	if (ObjectName.IsNone())
+	{
+		UE_LOG(LogMutable, Error, TEXT("The compilation of the Customizable object was not successful : The name of the CO was found to be None."));
+		return false;
+	}
+	
+	const FScopedLogSection ObjectSection (EMutableLogSection::Object, ObjectName);
 	
 	// Keep a strong object pointer pointing at the CO to prevent it from being GCd during the test
 	const TStrongObjectPtr<UCustomizableObject> TargetCO = TStrongObjectPtr{&InTargetCustomizableObject};
diff --git a/Engine/Plugins/Mutable/Source/MutableValidation/Private/MuV/ValidationUtils.h b/Engine/Plugins/Mutable/Source/MutableValidation/Private/MuV/ValidationUtils.h
--- a/Engine/Plugins/Mutable/Source/MutableValidation/Private/MuV/ValidationUtils.h
+++ b/Engine/Plugins/Mutable/Source/MutableValidation/Private/MuV/ValidationUtils.h
@@ -26,6 +26,13 @@ uint32 GetTargetAmountOfInstances(const FString& Params);
  */
 ITargetPlatform* GetCompilationPlatform(const FString& Params);
 
+/**
+ * Determines, based on the provided commandline, if the non-root COs should be compiled or not.
+ * @param Params The arguments provided to this commandlet.
+ * @return True if only root objects should be compiled, false otherwise.
+ */
+bool GetIfOnlyRootCOsShouldBeUsed(const FString& Params);
+
 
 // SETUP ---------------------------------------------------------------------------------------------------------------------------------------------
 
@@ -42,12 +49,6 @@ void PrepareAssetRegistry();
 void Wait(const double ToWaitSeconds);
 
 
-/**
- * old the thread until the provided variable evaluates to the inverse of the original value while ticking the engine.
- * @param bInWaitCondition Externally updatable variable that will trigger the end of the wait once the value it has differs from the original one.
- */
-void WaitUntilConditionChanges(bool& bInWaitCondition);
-
 /**
  * Logs some configuration data related to how mutable will compile and then generate instances. We do this so we can later
  * Isolate tests using different configurations.
