[Mutable] Updated platform configs.

- ClearWorkingMemoryOnUpdateEnd is now true by default.

#Synced-CL 51453873

#Synced-CL 51453873

diff --git a/Engine/Plugins/Mutable/Config/DefaultMutable.ini b/Engine/Plugins/Mutable/Config/DefaultMutable.ini
--- a/Engine/Plugins/Mutable/Config/DefaultMutable.ini
+++ b/Engine/Plugins/Mutable/Config/DefaultMutable.ini
@@ -66,39 +66,4 @@ bExtraBoneInfluencesEnabled=False
 [EditorDefaults]
 ClipMorphMaterialName=/Engine/EditorMaterials/LevelGridMaterial.LevelGridMaterial
 ClipMeshMaterialName=/Mutable/Materials/ClipMeshMaterial.ClipMeshMaterial
-DynamicallyGenerated_DGGUI_Path=/Game/UI/DynamicallyGeneratedGUI_DGGUI/DynamicallyGeneratedGUI_DGGUI.DynamicallyGeneratedGUI_DGGUI_C
-
-;--------------------------------------------------------------------------------------------------
-; These are scalability settings that can be defined in any xxxScalability.ini file.
-[ScalabilitySettings]
-
-; If different than 0, limit the amount of memory (in KB) to use to cache streaming data when 
-; building characters. 0 means no limit, -1 means use default (10Mb in PC and 3Mb in consoles).
-+CVars=b.MutableStreamingMemory=10000
-
-; Override in platform scalability .ini files
-; b.MutableStreamingMemory=3000
-
-[ViewDistanceQuality@0]
-; If different than 0, limit the number of mutable instances with full LODs to have at a given time.
-+CVars=b.NumGeneratedInstancesLimit=7
-
-; If different than 0, limit the number of mutable instances with LOD 1 to have at a given time.
-+CVars=b.NumGeneratedInstancesLimitLOD1=0
-
-; If different than 0, limit the number of mutable instances with LOD 2 to have at a given time.
-+CVars=b.NumGeneratedInstancesLimitLOD2=0
-
-; If NumGeneratedInstancesLimit is different than 0, sets the distance at which the system will fix 
-; the LOD of an instance to the lowest res one (LOD2) to prevent unnecessary LOD changes and memory 
-; consumption
-+CVars=b.DistanceForFixedLOD2=0
-
-[ViewDistanceQuality@1]
-+CVars=b.NumGeneratedInstancesLimit=0
-[ViewDistanceQuality@2]
-+CVars=b.NumGeneratedInstancesLimit=0
-[ViewDistanceQuality@3]
-+CVars=b.NumGeneratedInstancesLimit=0
-[ViewDistanceQuality@Cine]
-+CVars=b.NumGeneratedInstancesLimit=0
+DynamicallyGenerated_DGGUI_Path=/Game/UI/DynamicallyGeneratedGUI_DGGUI/DynamicallyGeneratedGUI_DGGUI.DynamicallyGeneratedGUI_DGGUI_C
\ No newline at end of file
diff --git a/Engine/Plugins/Mutable/Source/CustomizableObject/Internal/MuCO/CustomizableObjectSystemPrivate.h b/Engine/Plugins/Mutable/Source/CustomizableObject/Internal/MuCO/CustomizableObjectSystemPrivate.h
--- a/Engine/Plugins/Mutable/Source/CustomizableObject/Internal/MuCO/CustomizableObjectSystemPrivate.h
+++ b/Engine/Plugins/Mutable/Source/CustomizableObject/Internal/MuCO/CustomizableObjectSystemPrivate.h
@@ -493,7 +493,6 @@ public:
 	static int32 EnableMutableProgressiveMipStreaming;
 	static int32 EnableMutableLiveUpdate;
 	static int32 EnableMutableAnimInfoDebugging;
-	static int32 EnableSkipGenerateResidentMips;
 	static int32 MutableMeshesLODBias;
 	static int32 MaxTextureSizeToGenerate;
 	static bool bGenerateInstancesWithinRange;
diff --git a/Engine/Plugins/Mutable/Source/CustomizableObject/Private/MuCO/CustomizableObjectSystem.cpp b/Engine/Plugins/Mutable/Source/CustomizableObject/Private/MuCO/CustomizableObjectSystem.cpp
--- a/Engine/Plugins/Mutable/Source/CustomizableObject/Private/MuCO/CustomizableObjectSystem.cpp
+++ b/Engine/Plugins/Mutable/Source/CustomizableObject/Private/MuCO/CustomizableObjectSystem.cpp
@@ -85,7 +85,7 @@ static FAutoConsoleVariableRef CVarWorkingMemoryKB(
 
 TAutoConsoleVariable<bool> CVarClearWorkingMemoryOnUpdateEnd(
 	TEXT("mutable.ClearWorkingMemoryOnUpdateEnd"),
-	false,
+	true,
 	TEXT("Clear the working memory and cache after every Mutable operation."),
 	ECVF_Scalability);
 
