[Mutable] Program cache rewrok.

- Changed Address format to incorporate type info.
- Changed cache data structrure. Now cache operations are thread safe and don't waste memory for not executed ops.

#rnx

#Synced-CL 51454919

diff --git a/Engine/Plugins/Mutable/Source/CustomizableObject/Internal/MuCO/CustomizableObjectInstancePrivate.h b/Engine/Plugins/Mutable/Source/CustomizableObject/Internal/MuCO/CustomizableObjectInstancePrivate.h
--- a/Engine/Plugins/Mutable/Source/CustomizableObject/Internal/MuCO/CustomizableObjectInstancePrivate.h
+++ b/Engine/Plugins/Mutable/Source/CustomizableObject/Internal/MuCO/CustomizableObjectInstancePrivate.h
@@ -125,7 +125,7 @@ struct FAnimBpGeneratedPhysicsAssets
 struct FMaterialReuseCacheKey
 {
 	TStrongObjectPtr<UMaterialInterface> MaterialTemplate;
-	UE::Mutable::Private::OP::ADDRESS Root;
+	UE::Mutable::Private::FOperation::ADDRESS Root;
 
 	bool operator==(const FMaterialReuseCacheKey&) const = default;
 };
@@ -136,7 +136,7 @@ uint32 GetTypeHash(const FMaterialReuseCacheKey& Key);
 
 struct FTextureReuseCacheKey
 {
-	UE::Mutable::Private::OP::ADDRESS Root;
+	UE::Mutable::Private::FOperation::ADDRESS Root;
 
 	bool operator==(const FTextureReuseCacheKey&) const = default;
 };
diff --git a/Engine/Plugins/Mutable/Source/CustomizableObject/Internal/MuCO/CustomizableObjectPrivate.h b/Engine/Plugins/Mutable/Source/CustomizableObject/Internal/MuCO/CustomizableObjectPrivate.h
--- a/Engine/Plugins/Mutable/Source/CustomizableObject/Internal/MuCO/CustomizableObjectPrivate.h
+++ b/Engine/Plugins/Mutable/Source/CustomizableObject/Internal/MuCO/CustomizableObjectPrivate.h
@@ -1428,7 +1428,7 @@ enum class ECustomizableObjectVersions : int32
 	BackoutGenerateBulkFileWithDataLimit,
 
 	ClothCore,
-	
+
 	ConvertBoneIdsToBoneNames,
 	
 	AddLayerIndexToMatParamTextures,
@@ -1442,7 +1442,7 @@ enum class ECustomizableObjectVersions : int32
 	ChangedLayerIndexDataType,
 
 	MutableSkeletonToPassthroughObject,
-	
+
 	UseCompressedMorphsForAllMorphData,
 
 	TestMoveSocketsToCore,
@@ -1450,7 +1450,9 @@ enum class ECustomizableObjectVersions : int32
 	RemoveMorphTarget,
 	
 	AssetUserDataPassthroughObject2,
-	
+
+	NewProgramAddressFormat,
+
 	// -----<new versions can be added above this line>--------
 	LastCustomizableObjectVersion
 };
diff --git a/Engine/Plugins/Mutable/Source/CustomizableObject/Private/MuCO/CustomizableObject.cpp b/Engine/Plugins/Mutable/Source/CustomizableObject/Private/MuCO/CustomizableObject.cpp
--- a/Engine/Plugins/Mutable/Source/CustomizableObject/Private/MuCO/CustomizableObject.cpp
+++ b/Engine/Plugins/Mutable/Source/CustomizableObject/Private/MuCO/CustomizableObject.cpp
@@ -3048,7 +3048,6 @@ void UCustomizableObjectPrivate::LogMemory()
 	{
 		const UE::Mutable::Private::FProgram& Program = MuModel->GetProgram();
 		int32 ByteCode = Program.ByteCode.GetAllocatedSize();
-		int32 OpAddress = Program.OpAddress.GetAllocatedSize();
 		int32 Roms = Program.Roms.GetAllocatedSize();
 		int32 ConstantImages = Program.ConstantImages.GetAllocatedSize();
 		int32 ConstantImageLODIndices = Program.ConstantImageLODIndices.GetAllocatedSize();
@@ -3114,7 +3113,7 @@ void UCustomizableObjectPrivate::LogMemory()
 			}
 		}
 
-		int32 ModelTotal = ByteCode + OpAddress + Roms 
+		int32 ModelTotal = ByteCode + Roms 
 			+ ImageTotal + MeshesTotal 
 			+ StringsTotal + ConstantNames
 			+ SkeletonsTotal + PhysicsTotal
@@ -3158,7 +3157,6 @@ void UCustomizableObjectPrivate::LogMemory()
 		UE_LOG(LogMutable, Log, TEXT("  ModelResources     : %8d"), ModeleResData);
 		UE_LOG(LogMutable, Log, TEXT("  Model              : %8d"), ModelTotal);
 		UE_LOG(LogMutable, Log, TEXT("    ByteCode         : %8d"), ByteCode);
-		UE_LOG(LogMutable, Log, TEXT("    OpAddress        : %8d"), OpAddress);
 		UE_LOG(LogMutable, Log, TEXT("    Roms             : %8d"), Roms);
 		UE_LOG(LogMutable, Log, TEXT("    Strings          : %8d"), StringsTotal);
 		UE_LOG(LogMutable, Log, TEXT("    Names			   : %8d"), ConstantNames);
diff --git a/Engine/Plugins/Mutable/Source/CustomizableObject/Private/MuCO/CustomizableObjectMeshUpdate.cpp b/Engine/Plugins/Mutable/Source/CustomizableObject/Private/MuCO/CustomizableObjectMeshUpdate.cpp
--- a/Engine/Plugins/Mutable/Source/CustomizableObject/Private/MuCO/CustomizableObjectMeshUpdate.cpp
+++ b/Engine/Plugins/Mutable/Source/CustomizableObject/Private/MuCO/CustomizableObjectMeshUpdate.cpp
@@ -17,6 +17,7 @@ CustomizableObjectMeshUpdate.cpp: Helpers to stream in CustomizableObject skelet
 
 #include "MuR/Model.h"
 #include "MuR/MeshBufferSet.h"
+#include "MuR/ProgramCache.h"
 
 #include "Engine/SkeletalMesh.h"
 #include "Streaming/RenderAssetUpdate.inl"
@@ -148,33 +149,10 @@ namespace impl
 		
 		// End update
 		const bool bClearRoms = UCustomizableObjectSystem::ShouldClearWorkingMemoryOnUpdateEnd();
-		OperationData->System->EndUpdate(OperationData->LiveInstance, bClearRoms);
+		const bool bFreeCache = !CVarReuseLiveInstanceCacheBetweenUpdates.GetValueOnAnyThread();
+		OperationData->System->EndUpdate(OperationData->LiveInstance, bClearRoms, bFreeCache);
 
-		// Stream data (morphs, clothing...).
-		UE::Tasks::TTask StreamDataTask = UE::Tasks::Launch(TEXT("StreamData"),
-		[OperationData]()
-		{
-			FMutableStreamRequest StreamRequest(OperationData->ModelStreamableBulkData);
-
-			// Stream data.
-			UE::Tasks::AddNested(StreamRequest.Stream());
-		},
-		UE::Tasks::ETaskPriority::Inherit);
-
-		// Stream Data should be executed as a child task for this one
-		UE::Tasks::AddNested(StreamDataTask);
-		
-		// Go to next step.
-		UE::Tasks::TTask<void> OnUpdateMeshFinishedTask = UE::Tasks::Launch(TEXT("OnUpdateMeshFinished"),
-		[StreamingTask]()
-		{
-			StreamingTask->OnUpdateMeshFinished();
-		},
-		StreamDataTask,
-		UE::Tasks::ETaskPriority::Inherit,
-		UE::Tasks::EExtendedTaskPriority::Inline);
-		
-		UE::Tasks::AddNested(OnUpdateMeshFinishedTask);
+		StreamingTask->OnUpdateMeshFinished();
 		
 		TRACE_END_REGION(UE_MUTABLE_UPDATE_MESH_REGION);
 	}
@@ -240,7 +218,7 @@ namespace impl
 		// Prevent the usage of the already cached elements by fully clearing the cache.
 		if (!CVarReuseLiveInstanceCacheBetweenUpdates.GetValueOnAnyThread())
 		{
-			LiveInstance->Cache->Clear();
+			LiveInstance->Cache->Clear(UE::Mutable::Private::FProgramCache::EClearFlags::Full);
 		}
 		
 		// Main instance generation step
@@ -339,7 +317,7 @@ void FCustomizableObjectMeshStreamIn::ConvertMesh(const FContext& Context, bool&
 
 		if (!MutableMesh)
 		{
-			check(false);
+			//check(false);
 			Abort();
 			return;
 		}
diff --git a/Engine/Plugins/Mutable/Source/CustomizableObject/Private/MuCO/CustomizableObjectMipDataProvider.cpp b/Engine/Plugins/Mutable/Source/CustomizableObject/Private/MuCO/CustomizableObjectMipDataProvider.cpp
--- a/Engine/Plugins/Mutable/Source/CustomizableObject/Private/MuCO/CustomizableObjectMipDataProvider.cpp
+++ b/Engine/Plugins/Mutable/Source/CustomizableObject/Private/MuCO/CustomizableObjectMipDataProvider.cpp
@@ -10,6 +10,7 @@
 #include "MuCO/UnrealMutableImageProvider.h"
 #include "MuR/Model.h"
 #include "MuR/Operations.h"
+#include "MuR/ProgramCache.h"
 #include "TextureResource.h"
 
 #include "Containers/Ticker.h"
@@ -143,7 +144,7 @@ namespace Impl
 			// Prevent the usage of the already cached elements by fully clearing the cache.
 			if (!CVarReuseLiveInstanceCacheBetweenUpdates.GetValueOnAnyThread())
 			{
-				LiveInstance->Cache->Clear();
+				LiveInstance->Cache->Clear(UE::Mutable::Private::FProgramCache::EClearFlags::Full);
 			}
 
 			Instance = System->BeginUpdate_MutableThread(LiveInstance, LODMask, false);
@@ -201,7 +202,8 @@ namespace Impl
 					{
 						MUTABLE_CPUPROFILER_SCOPE(EndUpdate);
 						const bool bClearRoms = UCustomizableObjectSystem::ShouldClearWorkingMemoryOnUpdateEnd();
-						System->EndUpdate(OperationData->UpdateContext->LiveInstance, bClearRoms);
+						const bool bFreeCache = !CVarReuseLiveInstanceCacheBetweenUpdates.GetValueOnAnyThread();
+						System->EndUpdate(OperationData->UpdateContext->LiveInstance, bClearRoms, bFreeCache);
 					}
 
 					if (FLogBenchmarkUtil::IsBenchmarkingReportingEnabled())
diff --git a/Engine/Plugins/Mutable/Source/CustomizableObject/Private/MuCO/CustomizableObjectSystem.cpp b/Engine/Plugins/Mutable/Source/CustomizableObject/Private/MuCO/CustomizableObjectSystem.cpp
--- a/Engine/Plugins/Mutable/Source/CustomizableObject/Private/MuCO/CustomizableObjectSystem.cpp
+++ b/Engine/Plugins/Mutable/Source/CustomizableObject/Private/MuCO/CustomizableObjectSystem.cpp
@@ -1824,7 +1824,7 @@ namespace impl
 				Material->ImageCount = Material->Material->ImageParameters.Num();
 				for (TTuple<UE::Mutable::Private::FParameterKey, UE::Mutable::Private::FMaterial::FImageParameterData> ImageParameter : Material->Material->ImageParameters)
 				{
-					const UE::Mutable::Private::OP::ADDRESS Address = ImageParameter.Value.ImageParameter.Get<UE::Mutable::Private::OP::ADDRESS>();
+					const UE::Mutable::Private::FOperation::ADDRESS Address = ImageParameter.Value.ImageParameter.Get<UE::Mutable::Private::FOperation::ADDRESS>();
 					
 					TSharedRef<FInstanceUpdateData::FImage>& Image = Context->InstanceUpdateData.Images.Add_GetRef(MakeShared<FInstanceUpdateData::FImage>());
 					Image->ParameterName = ImageParameter.Key.ParameterName;
@@ -1872,7 +1872,7 @@ namespace impl
 		// Prevent the usage of the already cached elements by fully clearing the cache.
 		if (!CVarReuseLiveInstanceCacheBetweenUpdates.GetValueOnAnyThread())
 		{
-			OperationData->LiveInstance->Cache->Clear();
+			OperationData->LiveInstance->Cache->Clear(UE::Mutable::Private::FProgramCache::EClearFlags::Full);
 		}
 		
 		UE::Mutable::Private::TManagedPtr<const UE::Mutable::Private::FInstance> Instance = OperationData->MutableSystem->BeginUpdate_MutableThread(
@@ -2225,7 +2225,8 @@ namespace impl
 		check(MutableSystem);
 
 		const bool bClearRoms = UCustomizableObjectSystem::ShouldClearWorkingMemoryOnUpdateEnd();
-		MutableSystem->EndUpdate(LiveInstance, bClearRoms);
+		const bool bFreeCache = !CVarReuseLiveInstanceCacheBetweenUpdates.GetValueOnAnyThread();
+		MutableSystem->EndUpdate(LiveInstance, bClearRoms, bFreeCache);
 		
 		MutableSystem->SetImagePixelConversionOverride(nullptr);
 
diff --git a/Engine/Plugins/Mutable/Source/CustomizableObject/Private/MuCO/ExternalOperationProvider.cpp b/Engine/Plugins/Mutable/Source/CustomizableObject/Private/MuCO/ExternalOperationProvider.cpp
--- a/Engine/Plugins/Mutable/Source/CustomizableObject/Private/MuCO/ExternalOperationProvider.cpp
+++ b/Engine/Plugins/Mutable/Source/CustomizableObject/Private/MuCO/ExternalOperationProvider.cpp
@@ -15,7 +15,7 @@ namespace UE::Mutable::Private
 
 	void FExternalOperationProvider::AddReferencedObjects(FReferenceCollector& Collector)
 	{
-		for (TPair<OP::ADDRESS, FInstancedStruct>& Pair : ExternalOperations)
+		for (TPair<FOperation::ADDRESS, FInstancedStruct>& Pair : ExternalOperations)
 		{
 			if (const UScriptStruct* ScriptStruct = Pair.Value.GetScriptStruct())
 			{
@@ -25,7 +25,7 @@ namespace UE::Mutable::Private
 	}
 
 
-	const FInstancedStruct& FExternalOperationProvider::Get(OP::ADDRESS Address) const
+	const FInstancedStruct& FExternalOperationProvider::Get(FOperation::ADDRESS Address) const
 	{
 		return ExternalOperations[Address];
 	}
diff --git a/Engine/Plugins/Mutable/Source/CustomizableObject/Private/MuCO/ExternalOperationProvider.h b/Engine/Plugins/Mutable/Source/CustomizableObject/Private/MuCO/ExternalOperationProvider.h
--- a/Engine/Plugins/Mutable/Source/CustomizableObject/Private/MuCO/ExternalOperationProvider.h
+++ b/Engine/Plugins/Mutable/Source/CustomizableObject/Private/MuCO/ExternalOperationProvider.h
@@ -23,13 +23,13 @@ namespace UE::Mutable::Private
 		}
 
 		// IExternalOperationProvider interface
-		virtual const FInstancedStruct& Get(OP::ADDRESS Address) const override;
+		virtual const FInstancedStruct& Get(FOperation::ADDRESS Address) const override;
 		
 		// Own interface
 		FExternalOperationProvider(UModelResources& InModelResources);
 	
 	private:
-		TMap<OP::ADDRESS, FInstancedStruct> ExternalOperations;
+		TMap<FOperation::ADDRESS, FInstancedStruct> ExternalOperations;
 	};
 }
 
diff --git a/Engine/Plugins/Mutable/Source/CustomizableObjectEditor/Private/MuCOE/CustomizableObjectCompileRunnable.h b/Engine/Plugins/Mutable/Source/CustomizableObjectEditor/Private/MuCOE/CustomizableObjectCompileRunnable.h
--- a/Engine/Plugins/Mutable/Source/CustomizableObjectEditor/Private/MuCOE/CustomizableObjectCompileRunnable.h
+++ b/Engine/Plugins/Mutable/Source/CustomizableObjectEditor/Private/MuCOE/CustomizableObjectCompileRunnable.h
@@ -69,7 +69,7 @@ public:
 
 	TSharedPtr<UE::Mutable::Private::FModel, ESPMode::ThreadSafe> Model;
 
-	TMap<UE::Mutable::Private::OP::ADDRESS, FInstancedStruct> ExternalOperations;
+	TMap<UE::Mutable::Private::FOperation::ADDRESS, FInstancedStruct> ExternalOperations;
 	
 	FCompilationOptions Options;
 
diff --git a/Engine/Plugins/Mutable/Source/CustomizableObjectEditor/Private/MuCOE/CustomizableObjectCompiler.cpp b/Engine/Plugins/Mutable/Source/CustomizableObjectEditor/Private/MuCOE/CustomizableObjectCompiler.cpp
--- a/Engine/Plugins/Mutable/Source/CustomizableObjectEditor/Private/MuCOE/CustomizableObjectCompiler.cpp
+++ b/Engine/Plugins/Mutable/Source/CustomizableObjectEditor/Private/MuCOE/CustomizableObjectCompiler.cpp
@@ -1654,7 +1654,7 @@ void FCustomizableObjectCompiler::FinishCompilationTask()
 
 		ModelResources->ExternalOperations = MoveTemp(CompileTask->ExternalOperations);
 
-		for (TPair<UE::Mutable::Private::OP::ADDRESS, FInstancedStruct> ExternalOperation : ModelResources->ExternalOperations)
+		for (TPair<UE::Mutable::Private::FOperation::ADDRESS, FInstancedStruct> ExternalOperation : ModelResources->ExternalOperations)
 		{
 			const uint32 Version = ExternalOperation.Value.GetPtr<UE::Mutable::FExternalOperation>()->GetVersion();
 			ModelResources->ExternalOperationVersions.Add(ExternalOperation.Value.GetScriptStruct(), Version);
diff --git a/Engine/Plugins/Mutable/Source/CustomizableObjectEditor/Private/MuCOE/SMutableCodeViewer.cpp b/Engine/Plugins/Mutable/Source/CustomizableObjectEditor/Private/MuCOE/SMutableCodeViewer.cpp
--- a/Engine/Plugins/Mutable/Source/CustomizableObjectEditor/Private/MuCOE/SMutableCodeViewer.cpp
+++ b/Engine/Plugins/Mutable/Source/CustomizableObjectEditor/Private/MuCOE/SMutableCodeViewer.cpp
@@ -881,7 +881,7 @@ void SMutableCodeViewer::GetOperationsMatchingStringPattern(const FString& InStr
 		}
 		
 		// Get one of the previous run "children" and treat as a parent to get it's children and process them
-		const UE::Mutable::Private::OP::ADDRESS& ParentAddress = SearchPayload.BatchData[ParentIndex].Child;
+		const UE::Mutable::Private::FOperation::ADDRESS& ParentAddress = SearchPayload.BatchData[ParentIndex].Child;
 		
 		if (bMatchesPattern)
 		{
@@ -1126,7 +1126,7 @@ void SMutableCodeViewer::GetOperationsOfType(const UE::Mutable::Private::EOpType
 	for	(int32 ParentIndex = 0 ; ParentIndex < InSearchPayload.BatchData.Num(); ParentIndex++)
 	{
 		// Get one of the previous run "children" and treat as a parent to get it's children and process them
-		const UE::Mutable::Private::OP::ADDRESS& CurrentAddress = InSearchPayload.BatchData[ParentIndex].Child;
+		const UE::Mutable::Private::FOperation::ADDRESS& CurrentAddress = InSearchPayload.BatchData[ParentIndex].Child;
 		
 		// Cache if same data type and we share the same address (means this op is pointing at the provided resource)
 		// It will cache duplicated entries
@@ -1154,40 +1154,7 @@ void SMutableCodeViewer::GetOperationsOfType(const UE::Mutable::Private::EOpType
 
 void SMutableCodeViewer::CacheOperationTypesPresentOnModel()
 {
-	check(MutableModel)
-
-	// Initialize NodeOperationTypes with empty TPair for each possible mutable operation type
-	{
-		constexpr uint32 OperationTypesCount = static_cast<uint16>(UE::Mutable::Private::EOpType::COUNT);
-		ModelOperationTypes.Empty(OperationTypesCount);
-		for (uint32 Index = 0; Index < OperationTypesCount ; Index++)
-		{
-			UE::Mutable::Private::EOpType TargetType = StaticCast<UE::Mutable::Private::EOpType>(Index);
-			ModelOperationTypes.Add(TPair<UE::Mutable::Private::EOpType,uint32>{TargetType,0});
-		}
-	}
-	
-	// Locate all operation types found on the provided model program data structure and count how many instances of each
-	// there are
-	{
-		// Get the types and the amount of instances of each unique operation on the mutable model
-		const UE::Mutable::Private::FProgram& Program = MutableModel->GetProgram();
-		const uint32 ProgramAddressesCount = Program.OpAddress.Num();
-		
-		// Ensure first operation type is NONE since we are skipping it due to it having to have that type
-		check (Program.GetOpType(Program.OpAddress[0]) == UE::Mutable::Private::EOpType::NONE);
-
-		// Iterate over the addresses of the program and count how many instances each type has.
-		for (uint32 ProgramAddressesIndex = 1 ; ProgramAddressesIndex < ProgramAddressesCount; ProgramAddressesIndex++)
-		{
-			// Locate what is the position (index) of the operation type of the address on our collection of types found until now
-			const UE::Mutable::Private::EOpType OperationType = Program.GetOpType(ProgramAddressesIndex);
-			
-			// Increase the counter for this operation type
-			const uint16 TypeAsInteger = StaticCast<uint16>(OperationType);
-			ModelOperationTypes[TypeAsInteger].Value++;
-		}
-	}
+	// ModelOperationTypes is filled when the tree is generating.
 	
 	// Remove all operation types that do have no operations present on the model
 	{
@@ -1348,13 +1315,15 @@ void SMutableCodeViewer::GenerateAllTreeElements()
 	// Generate all root nodes
 	const UE::Mutable::Private::FProgram& Program = MutableModel->GetProgram();
 	const uint32 StateCount = Program.States.Num();
+
+	ModelOperationTypes.Init({ UE::Mutable::Private::EOpType::NONE, 0 }, (int32)UE::Mutable::Private::EOpType::COUNT);
+
 	for ( uint32 StateIndex = 0; StateIndex < StateCount; ++StateIndex )
 	{
 		const UE::Mutable::Private::FProgram::FState& State = Program.States[StateIndex];
 		const FString Caption = FString::Printf( TEXT("state [%s]"), *State.Name );
 
-		const FSlateColor LabelColor = ColorPerComputationalCost[StaticCast<uint8>(GetOperationTypeComputationalCost(
-			Program.GetOpType(State.Root)))];
+		const FSlateColor LabelColor = ColorPerComputationalCost[StaticCast<uint8>(GetOperationTypeComputationalCost(Program.GetOpType(State.Root)))];
 
 		// Locate the "original" tree element :
 		// This may happen if for some reason the state is duplicated (should never happen)
@@ -1365,27 +1334,32 @@ void SMutableCodeViewer::GenerateAllTreeElements()
 		RootNodes.Add(RootNodeElement);
 		
 		// Add the element to the cache so we keep the indices straight.
-		constexpr UE::Mutable::Private::OP::ADDRESS CommonParent = 0;
+		constexpr UE::Mutable::Private::FOperation::ADDRESS CommonParent = 0;
 		const FItemCacheKey Key = { CommonParent, State.Root, StateIndex };
 		ItemCache.Add(Key, RootNodeElement);
-		
+
 		if (!MainItemPtr)
 		{
 			// Cache this node as it may be duplicated of another state. Check the "MainItemPtr" initialization for more info
 			MainItemPerOp.Add(State.Root, RootNodeElement);
 
+			// Populate OperatationTypes with the root.
+			UE::Mutable::Private::EOpType RootType = Program.GetOpType(State.Root);
+			ModelOperationTypes[(int32)RootType].Key = RootType;
+			++ModelOperationTypes[(int32)RootType].Value;
+
 			// Iterate over each root node and generate all the elements in a human-readable pattern (Z Pattern)
-			GenerateElementRecursive(StateIndex,State.Root,Program);
+			GenerateElementRecursive(StateIndex, State.Root, Program);
 		}
 	}
 	
 }
 
-void SMutableCodeViewer::GenerateElementRecursive(const int32& InStateIndex, UE::Mutable::Private::OP::ADDRESS InParentAddress,  const UE::Mutable::Private::FProgram& InProgram)
+void SMutableCodeViewer::GenerateElementRecursive(const int32& InStateIndex, UE::Mutable::Private::FOperation::ADDRESS InParentAddress,  const UE::Mutable::Private::FProgram& InProgram)
 {
 	// This will be used to add operations
 	uint32 ChildIndex = 0;
-	auto AddOpFunc = [this, InParentAddress, &InProgram, &ChildIndex, &InStateIndex](UE::Mutable::Private::OP::ADDRESS ChildAddress, const FString& Caption)
+	auto AddOpFunc = [this, InParentAddress, &InProgram, &ChildIndex, &InStateIndex](UE::Mutable::Private::FOperation::ADDRESS ChildAddress, const FString& Caption)
 		{
 			{
 				const FItemCacheKey Key = { InParentAddress, ChildAddress, ChildIndex };
@@ -1410,6 +1384,11 @@ void SMutableCodeViewer::GenerateElementRecursive(const int32& InStateIndex, UE:
 					{
 						MainItemPerOp.Add(ChildAddress, Item);
 
+						// Populate OperatationTypes with the root.
+						UE::Mutable::Private::EOpType ChildType = InProgram.GetOpType(ChildAddress);
+						ModelOperationTypes[(int32)ChildType].Key = ChildType;
+						++ModelOperationTypes[(int32)ChildType].Value;
+
 						GenerateElementRecursive(InStateIndex,ChildAddress, InProgram);
 					}
 				}
@@ -1438,7 +1417,7 @@ void SMutableCodeViewer::GenerateElementRecursive(const int32& InStateIndex, UE:
 	case UE::Mutable::Private::EOpType::IN_CONDITIONAL:
 	case UE::Mutable::Private::EOpType::ED_CONDITIONAL:
 	{
-		UE::Mutable::Private::OP::ConditionalArgs Args = InProgram.GetOpArgs<UE::Mutable::Private::OP::ConditionalArgs>(InParentAddress);
+		UE::Mutable::Private::FOperation::ConditionalArgs Args = InProgram.GetOpArgs<UE::Mutable::Private::FOperation::ConditionalArgs>(InParentAddress);
 		AddOpFunc(Args.condition, TEXT("cond "));
 		AddOpFunc(Args.yes, TEXT("true "));
 		AddOpFunc(Args.no, TEXT("false "));
@@ -1456,19 +1435,19 @@ void SMutableCodeViewer::GenerateElementRecursive(const int32& InStateIndex, UE:
 	{
 		const uint8* OpData = InProgram.GetOpArgsPointer(InParentAddress);
 
-		UE::Mutable::Private::OP::ADDRESS VarAddress;
-		FMemory::Memcpy(&VarAddress, OpData, sizeof(UE::Mutable::Private::OP::ADDRESS));
-		OpData += sizeof(UE::Mutable::Private::OP::ADDRESS);
+		UE::Mutable::Private::FOperation::ADDRESS VarAddress;
+		FMemory::Memcpy(&VarAddress, OpData, sizeof(UE::Mutable::Private::FOperation::ADDRESS));
+		OpData += sizeof(UE::Mutable::Private::FOperation::ADDRESS);
 		AddOpFunc(VarAddress, TEXT("var "));
 
-		UE::Mutable::Private::OP::ADDRESS DefAddress;
-		FMemory::Memcpy(&DefAddress, OpData, sizeof(UE::Mutable::Private::OP::ADDRESS));
-		OpData += sizeof(UE::Mutable::Private::OP::ADDRESS);
+		UE::Mutable::Private::FOperation::ADDRESS DefAddress;
+		FMemory::Memcpy(&DefAddress, OpData, sizeof(UE::Mutable::Private::FOperation::ADDRESS));
+		OpData += sizeof(UE::Mutable::Private::FOperation::ADDRESS);
 		AddOpFunc(DefAddress, TEXT("def "));
 
-		UE::Mutable::Private::OP::FSwitchCaseDescriptor CaseDesc;
-		FMemory::Memcpy(&CaseDesc, OpData, sizeof(UE::Mutable::Private::OP::FSwitchCaseDescriptor));
-		OpData += sizeof(UE::Mutable::Private::OP::FSwitchCaseDescriptor);
+		UE::Mutable::Private::FOperation::FSwitchCaseDescriptor CaseDesc;
+		FMemory::Memcpy(&CaseDesc, OpData, sizeof(UE::Mutable::Private::FOperation::FSwitchCaseDescriptor));
+		OpData += sizeof(UE::Mutable::Private::FOperation::FSwitchCaseDescriptor);
 
 		if (!CaseDesc.bUseRanges)
 		{
@@ -1478,9 +1457,9 @@ void SMutableCodeViewer::GenerateElementRecursive(const int32& InStateIndex, UE:
 				FMemory::Memcpy(&Condition, OpData, sizeof(int32));
 				OpData += sizeof(int32);
 
-				UE::Mutable::Private::OP::ADDRESS At;
-				FMemory::Memcpy(&At, OpData, sizeof(UE::Mutable::Private::OP::ADDRESS));
-				OpData += sizeof(UE::Mutable::Private::OP::ADDRESS);
+				UE::Mutable::Private::FOperation::ADDRESS At;
+				FMemory::Memcpy(&At, OpData, sizeof(UE::Mutable::Private::FOperation::ADDRESS));
+				OpData += sizeof(UE::Mutable::Private::FOperation::ADDRESS);
 
 				FString Caption = FString::Printf(TEXT("case %d "), Condition);
 				AddOpFunc(At, Caption);
@@ -1498,9 +1477,9 @@ void SMutableCodeViewer::GenerateElementRecursive(const int32& InStateIndex, UE:
 				FMemory::Memcpy(&RangeSize, OpData, sizeof(uint32));
 				OpData += sizeof(uint32);
 
-				UE::Mutable::Private::OP::ADDRESS At;
-				FMemory::Memcpy(&At, OpData, sizeof(UE::Mutable::Private::OP::ADDRESS));
-				OpData += sizeof(UE::Mutable::Private::OP::ADDRESS);
+				UE::Mutable::Private::FOperation::ADDRESS At;
+				FMemory::Memcpy(&At, OpData, sizeof(UE::Mutable::Private::FOperation::ADDRESS));
+				OpData += sizeof(UE::Mutable::Private::FOperation::ADDRESS);
 
 				FString Caption = FString::Printf(TEXT("case %d to %d "), ConditionStart, ConditionStart+RangeSize-1);
 				AddOpFunc(At, Caption);
@@ -1512,7 +1491,7 @@ void SMutableCodeViewer::GenerateElementRecursive(const int32& InStateIndex, UE:
 
 	case UE::Mutable::Private::EOpType::IM_SWIZZLE:
 	{
-		UE::Mutable::Private::OP::ImageSwizzleArgs Args = InProgram.GetOpArgs<UE::Mutable::Private::OP::ImageSwizzleArgs>(InParentAddress);
+		UE::Mutable::Private::FOperation::ImageSwizzleArgs Args = InProgram.GetOpArgs<UE::Mutable::Private::FOperation::ImageSwizzleArgs>(InParentAddress);
 		for (int32 Channel = 0; Channel < 4; ++Channel)
 		{
 			FString Caption = FString::Printf(TEXT("%d is %d from "), Channel, Args.sourceChannels[Channel]);
@@ -1523,7 +1502,7 @@ void SMutableCodeViewer::GenerateElementRecursive(const int32& InStateIndex, UE:
 
 	case UE::Mutable::Private::EOpType::CO_SWIZZLE:
 	{
-		UE::Mutable::Private::OP::ColourSwizzleArgs Args = InProgram.GetOpArgs<UE::Mutable::Private::OP::ColourSwizzleArgs>(InParentAddress);
+		UE::Mutable::Private::FOperation::ColourSwizzleArgs Args = InProgram.GetOpArgs<UE::Mutable::Private::FOperation::ColourSwizzleArgs>(InParentAddress);
 		for (int32 Channel = 0; Channel < 4; ++Channel)
 		{
 			FString Caption = FString::Printf(TEXT("%d is %d from "), Channel, Args.sourceChannels[Channel]);
@@ -1534,7 +1513,7 @@ void SMutableCodeViewer::GenerateElementRecursive(const int32& InStateIndex, UE:
 
 	case UE::Mutable::Private::EOpType::IM_LAYER:
 	{
-		UE::Mutable::Private::OP::ImageLayerArgs Args = InProgram.GetOpArgs<UE::Mutable::Private::OP::ImageLayerArgs>(InParentAddress);
+		UE::Mutable::Private::FOperation::ImageLayerArgs Args = InProgram.GetOpArgs<UE::Mutable::Private::FOperation::ImageLayerArgs>(InParentAddress);
 		AddOpFunc(Args.base, TEXT("base "));
 		AddOpFunc(Args.mask, TEXT("mask "));
 		AddOpFunc(Args.blended, TEXT("blended "));
@@ -1543,7 +1522,7 @@ void SMutableCodeViewer::GenerateElementRecursive(const int32& InStateIndex, UE:
 
 	case UE::Mutable::Private::EOpType::IM_LAYERCOLOUR:
 	{
-		UE::Mutable::Private::OP::ImageLayerColourArgs Args = InProgram.GetOpArgs<UE::Mutable::Private::OP::ImageLayerColourArgs>(InParentAddress);
+		UE::Mutable::Private::FOperation::ImageLayerColourArgs Args = InProgram.GetOpArgs<UE::Mutable::Private::FOperation::ImageLayerColourArgs>(InParentAddress);
 		AddOpFunc(Args.base, TEXT("base "));
 		AddOpFunc(Args.mask, TEXT("mask "));
 		AddOpFunc(Args.colour, TEXT("colour "));
@@ -1552,7 +1531,7 @@ void SMutableCodeViewer::GenerateElementRecursive(const int32& InStateIndex, UE:
 
 	case UE::Mutable::Private::EOpType::IM_MULTILAYER:
 	{
-		UE::Mutable::Private::OP::ImageMultiLayerArgs Args = InProgram.GetOpArgs<UE::Mutable::Private::OP::ImageMultiLayerArgs>(InParentAddress);
+		UE::Mutable::Private::FOperation::ImageMultiLayerArgs Args = InProgram.GetOpArgs<UE::Mutable::Private::FOperation::ImageMultiLayerArgs>(InParentAddress);
 		AddOpFunc(Args.rangeSize, TEXT("range "));
 		AddOpFunc(Args.base, TEXT("base "));
 		AddOpFunc(Args.mask, TEXT("mask "));
@@ -1562,9 +1541,9 @@ void SMutableCodeViewer::GenerateElementRecursive(const int32& InStateIndex, UE:
 
 	case UE::Mutable::Private::EOpType::ME_ADDMETADATA:
 	{
-		UE::Mutable::Private::OP::MeshAddMetadataArgs Args = InProgram.GetOpArgs<UE::Mutable::Private::OP::MeshAddMetadataArgs>(InParentAddress);
+		UE::Mutable::Private::FOperation::MeshAddMetadataArgs Args = InProgram.GetOpArgs<UE::Mutable::Private::FOperation::MeshAddMetadataArgs>(InParentAddress);
 
-		using OpEnumFlags = UE::Mutable::Private::OP::MeshAddMetadataArgs::EnumFlags;
+		using OpEnumFlags = UE::Mutable::Private::FOperation::MeshAddMetadataArgs::EnumFlags;
 
 		int32 TagCount = EnumHasAnyFlags(Args.Flags, OpEnumFlags::HasGameplayTags) ? 1 : 0;
 		if (EnumHasAnyFlags(Args.Flags, OpEnumFlags::IsGameplayTagList) &&
@@ -1619,7 +1598,7 @@ void SMutableCodeViewer::GenerateElementRecursive(const int32& InStateIndex, UE:
 
 	case UE::Mutable::Private::EOpType::ME_APPLYLAYOUT:
 	{
-		UE::Mutable::Private::OP::MeshApplyLayoutArgs Args = InProgram.GetOpArgs<UE::Mutable::Private::OP::MeshApplyLayoutArgs>(InParentAddress);
+		UE::Mutable::Private::FOperation::MeshApplyLayoutArgs Args = InProgram.GetOpArgs<UE::Mutable::Private::FOperation::MeshApplyLayoutArgs>(InParentAddress);
 		AddOpFunc(Args.Layout, TEXT("layout "));
 		AddOpFunc(Args.Mesh, TEXT("mesh "));
 		break;
@@ -1627,7 +1606,7 @@ void SMutableCodeViewer::GenerateElementRecursive(const int32& InStateIndex, UE:
 
 	case UE::Mutable::Private::EOpType::ME_PREPARELAYOUT:
 	{
-		UE::Mutable::Private::OP::MeshPrepareLayoutArgs Args = InProgram.GetOpArgs<UE::Mutable::Private::OP::MeshPrepareLayoutArgs>(InParentAddress);
+		UE::Mutable::Private::FOperation::MeshPrepareLayoutArgs Args = InProgram.GetOpArgs<UE::Mutable::Private::FOperation::MeshPrepareLayoutArgs>(InParentAddress);
 		AddOpFunc(Args.Layout, TEXT("layout "));
 		AddOpFunc(Args.Mesh, TEXT("mesh "));
 		break;
@@ -1637,22 +1616,22 @@ void SMutableCodeViewer::GenerateElementRecursive(const int32& InStateIndex, UE:
 	{
 		const uint8* data = InProgram.GetOpArgsPointer(InParentAddress);
 
-		UE::Mutable::Private::OP::ADDRESS BaseAt = 0;
-		FMemory::Memcpy(&BaseAt, data, sizeof(UE::Mutable::Private::OP::ADDRESS));
+		UE::Mutable::Private::FOperation::ADDRESS BaseAt = 0;
+		FMemory::Memcpy(&BaseAt, data, sizeof(UE::Mutable::Private::FOperation::ADDRESS));
 		AddOpFunc(BaseAt, TEXT("base "));
-		data += sizeof(UE::Mutable::Private::OP::ADDRESS);
+		data += sizeof(UE::Mutable::Private::FOperation::ADDRESS);
 
-		UE::Mutable::Private::OP::ADDRESS TargetAt = 0;
-		FMemory::Memcpy(&TargetAt, data, sizeof(UE::Mutable::Private::OP::ADDRESS));
+		UE::Mutable::Private::FOperation::ADDRESS TargetAt = 0;
+		FMemory::Memcpy(&TargetAt, data, sizeof(UE::Mutable::Private::FOperation::ADDRESS));
 		AddOpFunc(TargetAt, TEXT("target "));
-		data += sizeof(UE::Mutable::Private::OP::ADDRESS);
+		data += sizeof(UE::Mutable::Private::FOperation::ADDRESS);
 		
 		break;
 	}
 		
 	case UE::Mutable::Private::EOpType::ME_MORPH:
 	{
-		const UE::Mutable::Private::OP::MeshMorphArgs Args = InProgram.GetOpArgs<UE::Mutable::Private::OP::MeshMorphArgs>(InParentAddress);
+		const UE::Mutable::Private::FOperation::MeshMorphArgs Args = InProgram.GetOpArgs<UE::Mutable::Private::FOperation::MeshMorphArgs>(InParentAddress);
 
 		FString MorphName = InProgram.ConstantStrings.IsValidIndex(Args.Name) 
 			? InProgram.ConstantStrings[Args.Name]
@@ -1666,15 +1645,14 @@ void SMutableCodeViewer::GenerateElementRecursive(const int32& InStateIndex, UE:
 		
 	case UE::Mutable::Private::EOpType::SC_CURVE:
 	{
-		UE::Mutable::Private::OP::ScalarCurveArgs Args = InProgram.GetOpArgs<UE::Mutable::Private::OP::ScalarCurveArgs>(InParentAddress);
-		AddOpFunc(Args.time, TEXT("time "));
-		AddOpFunc(Args.curve, TEXT("curve "));
+		UE::Mutable::Private::FOperation::ScalarCurveArgs Args = InProgram.GetOpArgs<UE::Mutable::Private::FOperation::ScalarCurveArgs>(InParentAddress);
+		AddOpFunc(Args.time, TEXT("Time "));
 		break;
 	}
 		
 	case UE::Mutable::Private::EOpType::CO_SAMPLEIMAGE:
 	{
-		UE::Mutable::Private::OP::ColourSampleImageArgs Args = InProgram.GetOpArgs<UE::Mutable::Private::OP::ColourSampleImageArgs>(InParentAddress);
+		UE::Mutable::Private::FOperation::ColourSampleImageArgs Args = InProgram.GetOpArgs<UE::Mutable::Private::FOperation::ColourSampleImageArgs>(InParentAddress);
 		AddOpFunc(Args.Image, TEXT("image "));
 		AddOpFunc(Args.X, TEXT("x "));
 		AddOpFunc(Args.Y, TEXT("y "));
@@ -1683,14 +1661,14 @@ void SMutableCodeViewer::GenerateElementRecursive(const int32& InStateIndex, UE:
 
 	case UE::Mutable::Private::EOpType::CO_LINEARTOSRGB:
 	{
-		UE::Mutable::Private::OP::ColorArgs Args = InProgram.GetOpArgs<UE::Mutable::Private::OP::ColorArgs>(InParentAddress);
+		UE::Mutable::Private::FOperation::ColorArgs Args = InProgram.GetOpArgs<UE::Mutable::Private::FOperation::ColorArgs>(InParentAddress);
 		AddOpFunc(Args.Color, TEXT("Color "));
 		break;
 	}
 		
 	case UE::Mutable::Private::EOpType::IM_RESIZELIKE:
 	{
-		UE::Mutable::Private::OP::ImageResizeLikeArgs Args = InProgram.GetOpArgs<UE::Mutable::Private::OP::ImageResizeLikeArgs>(InParentAddress);
+		UE::Mutable::Private::FOperation::ImageResizeLikeArgs Args = InProgram.GetOpArgs<UE::Mutable::Private::FOperation::ImageResizeLikeArgs>(InParentAddress);
 		AddOpFunc(Args.Source, TEXT("src "));
 		AddOpFunc(Args.SizeSource, TEXT("sizeSrc "));
 		break;
@@ -1701,7 +1679,7 @@ void SMutableCodeViewer::GenerateElementRecursive(const int32& InStateIndex, UE:
 	case UE::Mutable::Private::EOpType::IN_ADDSURFACE:
 	case UE::Mutable::Private::EOpType::IN_ADDSKELETALMESH:
 	{
-		UE::Mutable::Private::OP::InstanceAddArgs Args = InProgram.GetOpArgs<UE::Mutable::Private::OP::InstanceAddArgs>(InParentAddress);
+		UE::Mutable::Private::FOperation::InstanceAddArgs Args = InProgram.GetOpArgs<UE::Mutable::Private::FOperation::InstanceAddArgs>(InParentAddress);
 		AddOpFunc(Args.instance, TEXT("instance "));
 		AddOpFunc(Args.value, TEXT("value "));
 		break;
@@ -1709,7 +1687,7 @@ void SMutableCodeViewer::GenerateElementRecursive(const int32& InStateIndex, UE:
 		
 	case UE::Mutable::Private::EOpType::IN_ADDCOMPONENT:
 	{
-		UE::Mutable::Private::OP::InstanceAddComponentArgs Args = InProgram.GetOpArgs<UE::Mutable::Private::OP::InstanceAddComponentArgs>(InParentAddress);
+		UE::Mutable::Private::FOperation::InstanceAddComponentArgs Args = InProgram.GetOpArgs<UE::Mutable::Private::FOperation::InstanceAddComponentArgs>(InParentAddress);
 		AddOpFunc(Args.Instance, TEXT("instance "));
 		AddOpFunc(Args.Value, TEXT("value "));
 		break;
@@ -1717,7 +1695,7 @@ void SMutableCodeViewer::GenerateElementRecursive(const int32& InStateIndex, UE:
 		
 	case UE::Mutable::Private::EOpType::IN_ADDOVERLAYMATERIAL:
 	{
-		UE::Mutable::Private::OP::InstanceAddOverlayMaterialArgs Args = InProgram.GetOpArgs<UE::Mutable::Private::OP::InstanceAddOverlayMaterialArgs>(InParentAddress);
+		UE::Mutable::Private::FOperation::InstanceAddOverlayMaterialArgs Args = InProgram.GetOpArgs<UE::Mutable::Private::FOperation::InstanceAddOverlayMaterialArgs>(InParentAddress);
 		AddOpFunc(Args.Instance, TEXT("instance "));
 		AddOpFunc(Args.Material, TEXT("value "));
 		break;
@@ -1725,7 +1703,7 @@ void SMutableCodeViewer::GenerateElementRecursive(const int32& InStateIndex, UE:
 
 	case UE::Mutable::Private::EOpType::IN_ADDOVERRIDEMATERIAL:
 	{
-		UE::Mutable::Private::OP::InstanceAddOverrideMaterialArgs Args = InProgram.GetOpArgs<UE::Mutable::Private::OP::InstanceAddOverrideMaterialArgs>(InParentAddress);
+		UE::Mutable::Private::FOperation::InstanceAddOverrideMaterialArgs Args = InProgram.GetOpArgs<UE::Mutable::Private::FOperation::InstanceAddOverrideMaterialArgs>(InParentAddress);
 		AddOpFunc(Args.Instance, TEXT("instance "));
 		AddOpFunc(Args.Material, TEXT("value "));
 		break;
@@ -1733,7 +1711,7 @@ void SMutableCodeViewer::GenerateElementRecursive(const int32& InStateIndex, UE:
 	
 	case UE::Mutable::Private::EOpType::IN_ADDMATERIAL:
 	{
-		UE::Mutable::Private::OP::InstanceAddMaterialArgs Args = InProgram.GetOpArgs<UE::Mutable::Private::OP::InstanceAddMaterialArgs>(InParentAddress);
+		UE::Mutable::Private::FOperation::InstanceAddMaterialArgs Args = InProgram.GetOpArgs<UE::Mutable::Private::FOperation::InstanceAddMaterialArgs>(InParentAddress);
 		AddOpFunc(Args.Instance, TEXT("instance "));
 		AddOpFunc(Args.Material, TEXT("material "));
 		break;
@@ -1741,7 +1719,7 @@ void SMutableCodeViewer::GenerateElementRecursive(const int32& InStateIndex, UE:
 	
 	case UE::Mutable::Private::EOpType::IM_COMPOSE:
 	{
-		UE::Mutable::Private::OP::ImageComposeArgs Args = InProgram.GetOpArgs<UE::Mutable::Private::OP::ImageComposeArgs>(InParentAddress);
+		UE::Mutable::Private::FOperation::ImageComposeArgs Args = InProgram.GetOpArgs<UE::Mutable::Private::FOperation::ImageComposeArgs>(InParentAddress);
 		AddOpFunc(Args.layout, TEXT("Layout "));
 		AddOpFunc(Args.base, TEXT("Base "));
 		AddOpFunc(Args.blockImage, TEXT("BlockImage "));
@@ -1751,7 +1729,7 @@ void SMutableCodeViewer::GenerateElementRecursive(const int32& InStateIndex, UE:
 
 	case UE::Mutable::Private::EOpType::IM_MULTICOMPOSE:
 	{
-		UE::Mutable::Private::OP::ImageMultiComposeArgs Args = InProgram.GetOpArgs<UE::Mutable::Private::OP::ImageMultiComposeArgs>(InParentAddress);
+		UE::Mutable::Private::FOperation::ImageMultiComposeArgs Args = InProgram.GetOpArgs<UE::Mutable::Private::FOperation::ImageMultiComposeArgs>(InParentAddress);
 		AddOpFunc(Args.Layout, TEXT("Layout "));
 		AddOpFunc(Args.Base, TEXT("Base "));
 		AddOpFunc(Args.SourceLayout, TEXT("SourceLayout"));
@@ -1762,11 +1740,11 @@ void SMutableCodeViewer::GenerateElementRecursive(const int32& InStateIndex, UE:
 		
 	case UE::Mutable::Private::EOpType::IM_INTERPOLATE:
 	{
-		UE::Mutable::Private::OP::ImageInterpolateArgs Args = InProgram.GetOpArgs<UE::Mutable::Private::OP::ImageInterpolateArgs>(InParentAddress);
+		UE::Mutable::Private::FOperation::ImageInterpolateArgs Args = InProgram.GetOpArgs<UE::Mutable::Private::FOperation::ImageInterpolateArgs>(InParentAddress);
 		AddOpFunc(Args.Factor, TEXT("factor "));
 			
 		int32 TargetIndex = 0;
-		for (const UE::Mutable::Private::OP::ADDRESS& Operation : Args.Targets)
+		for (const UE::Mutable::Private::FOperation::ADDRESS& Operation : Args.Targets)
 		{
 			AddOpFunc(Operation, FString::Printf(TEXT("target %i "), TargetIndex++));
 		}
@@ -1775,7 +1753,7 @@ void SMutableCodeViewer::GenerateElementRecursive(const int32& InStateIndex, UE:
 		
 	case UE::Mutable::Private::EOpType::IM_SATURATE:
 	{
-		UE::Mutable::Private::OP::ImageSaturateArgs Args = InProgram.GetOpArgs<UE::Mutable::Private::OP::ImageSaturateArgs>(InParentAddress);
+		UE::Mutable::Private::FOperation::ImageSaturateArgs Args = InProgram.GetOpArgs<UE::Mutable::Private::FOperation::ImageSaturateArgs>(InParentAddress);
 		AddOpFunc(Args.Base, TEXT("base "));
 		AddOpFunc(Args.Factor, TEXT("factor "));
 		break;
@@ -1783,7 +1761,7 @@ void SMutableCodeViewer::GenerateElementRecursive(const int32& InStateIndex, UE:
 		
 	case UE::Mutable::Private::EOpType::IM_COLOURMAP:
 	{
-		UE::Mutable::Private::OP::ImageColourMapArgs Args = InProgram.GetOpArgs<UE::Mutable::Private::OP::ImageColourMapArgs>(InParentAddress);
+		UE::Mutable::Private::FOperation::ImageColourMapArgs Args = InProgram.GetOpArgs<UE::Mutable::Private::FOperation::ImageColourMapArgs>(InParentAddress);
 		AddOpFunc(Args.Base, TEXT("base "));
 		AddOpFunc(Args.Mask, TEXT("mask "));
 		AddOpFunc(Args.Map, TEXT("map "));
@@ -1792,7 +1770,7 @@ void SMutableCodeViewer::GenerateElementRecursive(const int32& InStateIndex, UE:
 		
 	case UE::Mutable::Private::EOpType::IM_BINARISE:
 	{
-		UE::Mutable::Private::OP::ImageBinariseArgs Args = InProgram.GetOpArgs<UE::Mutable::Private::OP::ImageBinariseArgs>(InParentAddress);
+		UE::Mutable::Private::FOperation::ImageBinariseArgs Args = InProgram.GetOpArgs<UE::Mutable::Private::FOperation::ImageBinariseArgs>(InParentAddress);
 		AddOpFunc(Args.Base, TEXT("base "));
 		AddOpFunc(Args.Threshold, TEXT("threshold "));
 		break;
@@ -1800,7 +1778,7 @@ void SMutableCodeViewer::GenerateElementRecursive(const int32& InStateIndex, UE:
 		
 	case UE::Mutable::Private::EOpType::IM_PATCH:
 	{
-		UE::Mutable::Private::OP::ImagePatchArgs Args = InProgram.GetOpArgs<UE::Mutable::Private::OP::ImagePatchArgs>(InParentAddress);
+		UE::Mutable::Private::FOperation::ImagePatchArgs Args = InProgram.GetOpArgs<UE::Mutable::Private::FOperation::ImagePatchArgs>(InParentAddress);
 		AddOpFunc(Args.base, TEXT("base "));
 		AddOpFunc(Args.patch, TEXT("patch "));
 		break;
@@ -1808,7 +1786,7 @@ void SMutableCodeViewer::GenerateElementRecursive(const int32& InStateIndex, UE:
 		
 	case UE::Mutable::Private::EOpType::IM_RASTERMESH:
 	{
-		UE::Mutable::Private::OP::ImageRasterMeshArgs Args = InProgram.GetOpArgs<UE::Mutable::Private::OP::ImageRasterMeshArgs>(InParentAddress);
+		UE::Mutable::Private::FOperation::ImageRasterMeshArgs Args = InProgram.GetOpArgs<UE::Mutable::Private::FOperation::ImageRasterMeshArgs>(InParentAddress);
 		AddOpFunc(Args.mesh, TEXT("mesh "));
 		AddOpFunc(Args.image, TEXT("image "));
 		AddOpFunc(Args.mask, TEXT("mask "));
@@ -1819,7 +1797,7 @@ void SMutableCodeViewer::GenerateElementRecursive(const int32& InStateIndex, UE:
 		
 	case UE::Mutable::Private::EOpType::IM_DISPLACE:
 	{
-		UE::Mutable::Private::OP::ImageDisplaceArgs Args = InProgram.GetOpArgs<UE::Mutable::Private::OP::ImageDisplaceArgs>(InParentAddress);
+		UE::Mutable::Private::FOperation::ImageDisplaceArgs Args = InProgram.GetOpArgs<UE::Mutable::Private::FOperation::ImageDisplaceArgs>(InParentAddress);
 		AddOpFunc(Args.Source, TEXT("src "));
 		AddOpFunc(Args.DisplacementMap, TEXT("displacementMap "));
 		break;
@@ -1827,7 +1805,7 @@ void SMutableCodeViewer::GenerateElementRecursive(const int32& InStateIndex, UE:
 		
 	case UE::Mutable::Private::EOpType::IM_NORMALCOMPOSITE:
 	{
-		UE::Mutable::Private::OP::ImageNormalCompositeArgs Args = InProgram.GetOpArgs<UE::Mutable::Private::OP::ImageNormalCompositeArgs>(InParentAddress);
+		UE::Mutable::Private::FOperation::ImageNormalCompositeArgs Args = InProgram.GetOpArgs<UE::Mutable::Private::FOperation::ImageNormalCompositeArgs>(InParentAddress);
 		AddOpFunc(Args.base, TEXT("base "));
 		AddOpFunc(Args.normal, TEXT("normal "));
 		break;
@@ -1835,7 +1813,7 @@ void SMutableCodeViewer::GenerateElementRecursive(const int32& InStateIndex, UE:
 		
 	case UE::Mutable::Private::EOpType::IM_TRANSFORM:
 	{
-		UE::Mutable::Private::OP::ImageTransformArgs Args = InProgram.GetOpArgs<UE::Mutable::Private::OP::ImageTransformArgs>(InParentAddress);
+		UE::Mutable::Private::FOperation::ImageTransformArgs Args = InProgram.GetOpArgs<UE::Mutable::Private::FOperation::ImageTransformArgs>(InParentAddress);
 		AddOpFunc(Args.Base, TEXT("base "));
 		AddOpFunc(Args.OffsetX, TEXT("offsetX "));
 		AddOpFunc(Args.OffsetY, TEXT("offsetY "));
@@ -1849,7 +1827,7 @@ void SMutableCodeViewer::GenerateElementRecursive(const int32& InStateIndex, UE:
 	case UE::Mutable::Private::EOpType::CO_MATERIAL_BREAK:
 	case UE::Mutable::Private::EOpType::IM_MATERIAL_BREAK:
 	{
-		UE::Mutable::Private::OP::MaterialBreakArgs Args = InProgram.GetOpArgs<UE::Mutable::Private::OP::MaterialBreakArgs>(InParentAddress);
+		UE::Mutable::Private::FOperation::MaterialBreakArgs Args = InProgram.GetOpArgs<UE::Mutable::Private::FOperation::MaterialBreakArgs>(InParentAddress);
 		AddOpFunc(Args.Material, TEXT("Material "));
 		break;
 	}
@@ -1869,7 +1847,7 @@ void SMutableCodeViewer::GenerateElementRecursive(const int32& InStateIndex, UE:
 	if (bUseGeneric)
 	{
 		// Find children of the provided element without adding any extra string to better identify what each of them represents
-		UE::Mutable::Private::ForEachReference(InProgram, InParentAddress, [this, &InProgram, AddOpFunc](UE::Mutable::Private::OP::ADDRESS ChildAddress)
+		UE::Mutable::Private::ForEachReference(InProgram, InParentAddress, [this, &InProgram, AddOpFunc](UE::Mutable::Private::FOperation::ADDRESS ChildAddress)
 			{
 				AddOpFunc(ChildAddress,TEXT(""));
 			});
@@ -1879,7 +1857,7 @@ void SMutableCodeViewer::GenerateElementRecursive(const int32& InStateIndex, UE:
 		// Validate in case there is a mismatch in the custom processing of children and the generic one, which would cause problems.
 		ChildIndex = 0;
 
-		auto ValidateOpFunc = [this, InParentAddress, ParentOperationType, &ChildIndex](UE::Mutable::Private::OP::ADDRESS ChildAddress)
+		auto ValidateOpFunc = [this, InParentAddress, ParentOperationType, &ChildIndex](UE::Mutable::Private::FOperation::ADDRESS ChildAddress)
 		{
 			const FItemCacheKey Key = { InParentAddress, ChildAddress, ChildIndex };
 			const TSharedPtr<FMutableCodeTreeElement>* CachedItem = ItemCache.Find(Key);
@@ -1891,7 +1869,7 @@ void SMutableCodeViewer::GenerateElementRecursive(const int32& InStateIndex, UE:
 			++ChildIndex;
 		};
 
-		UE::Mutable::Private::ForEachReference(InProgram, InParentAddress, [this, ValidateOpFunc](UE::Mutable::Private::OP::ADDRESS ChildAddress)
+		UE::Mutable::Private::ForEachReference(InProgram, InParentAddress, [this, ValidateOpFunc](UE::Mutable::Private::FOperation::ADDRESS ChildAddress)
 		{
 			ValidateOpFunc(ChildAddress);
 		});
@@ -1955,11 +1933,11 @@ void SMutableCodeViewer::GetChildrenForInfo(TSharedPtr<FMutableCodeTreeElement>
 	check(MutableModel);
 	const UE::Mutable::Private::FProgram& Program = MutableModel->GetProgram();
 
-	UE::Mutable::Private::OP::ADDRESS ParentAddress = InInfo->MutableOperation;
+	UE::Mutable::Private::FOperation::ADDRESS ParentAddress = InInfo->MutableOperation;
 
 	// Generic case for unnamed children traversal.
 	uint32 ChildIndex = 0;
-	UE::Mutable::Private::ForEachReference(Program, InInfo->MutableOperation, [this, ParentAddress, &ChildIndex, &OutChildren](UE::Mutable::Private::OP::ADDRESS ChildAddress)
+	UE::Mutable::Private::ForEachReference(Program, InInfo->MutableOperation, [this, ParentAddress, &ChildIndex, &OutChildren](UE::Mutable::Private::FOperation::ADDRESS ChildAddress)
 	{
 		{
 			const FItemCacheKey Key = { ParentAddress, ChildAddress, ChildIndex };
@@ -1987,7 +1965,7 @@ void SMutableCodeViewer::OnExpansionChanged(TSharedPtr<FMutableCodeTreeElement>
 	// If an element gets expanded then contract (if found) the other element that uses the same address
 	if (bInExpanded)
 	{
-		const UE::Mutable::Private::OP::ADDRESS MutableOperation = InItem->MutableOperation;
+		const UE::Mutable::Private::FOperation::ADDRESS MutableOperation = InItem->MutableOperation;
 		const TSharedPtr<FMutableCodeTreeElement>* PreviouslyExpandedElement = ExpandedElements.Find(MutableOperation);
 		if (PreviouslyExpandedElement)
 		{
@@ -2030,11 +2008,11 @@ void SMutableCodeViewer::GetVisibleChildren(TSharedPtr<FMutableCodeTreeElement>
 		const TSharedPtr<FMutableCodeTreeElement> ToCheck = ToSearchForChildren[0];
 		ToSearchForChildren.RemoveAt(0);
 		
-		const UE::Mutable::Private::OP::ADDRESS ParentAddress = ToCheck->MutableOperation;
+		const UE::Mutable::Private::FOperation::ADDRESS ParentAddress = ToCheck->MutableOperation;
 	
 		// Generic case for unnamed children traversal.
 		uint32 ChildIndex = 0;
-		UE::Mutable::Private::ForEachReference(MutableProgram, ParentAddress, [this, ParentAddress, &ChildIndex, &OutChildren, &ToSearchForChildren ](UE::Mutable::Private::OP::ADDRESS ChildAddress)
+		UE::Mutable::Private::ForEachReference(MutableProgram, ParentAddress, [this, ParentAddress, &ChildIndex, &OutChildren, &ToSearchForChildren ](UE::Mutable::Private::FOperation::ADDRESS ChildAddress)
 		{
 			{
 				const FItemCacheKey Key = { ParentAddress, ChildAddress, ChildIndex };
@@ -2351,8 +2329,8 @@ void SMutableCodeViewer::TreeExpandElements(TArray<TSharedPtr<FMutableCodeTreeEl
 
 	// Load references to the arrays containing all the operations already worked on during another recursive call to this
 	// method
-	TArray<UE::Mutable::Private::OP::ADDRESS>& AlreadyExpandedOriginalOperations = InExpandedOperationsBuffer->ExpandedOriginalOperations;
-	TArray<UE::Mutable::Private::OP::ADDRESS>& AlreadyExpandedDuplicatedOperations = InExpandedOperationsBuffer->ExpandedDuplicatedOperations;
+	TArray<UE::Mutable::Private::FOperation::ADDRESS>& AlreadyExpandedOriginalOperations = InExpandedOperationsBuffer->ExpandedOriginalOperations;
+	TArray<UE::Mutable::Private::FOperation::ADDRESS>& AlreadyExpandedDuplicatedOperations = InExpandedOperationsBuffer->ExpandedDuplicatedOperations;
 	
 	// Array containing the children object found on Item.
 	TArray<TSharedPtr<FMutableCodeTreeElement>> Children;
@@ -2366,7 +2344,7 @@ void SMutableCodeViewer::TreeExpandElements(TArray<TSharedPtr<FMutableCodeTreeEl
 		check(Item);
 
 		// Identifier of the element. May be repeated if there are elements duplicating another element
-		const UE::Mutable::Private::OP::ADDRESS OperationAddress = Item->MutableOperation;
+		const UE::Mutable::Private::FOperation::ADDRESS OperationAddress = Item->MutableOperation;
 		
 		// Filter the elements being expanded if the user has defined a desired EDataType
 		if (FilteringDataType != UE::Mutable::Private::EDataType::None)
@@ -2510,7 +2488,7 @@ void SMutableCodeViewer::CacheRootNodeAddresses()
 	check (MutableModel);
 	check (RootNodeAddresses.IsEmpty())
 
-	TArray<UE::Mutable::Private::OP::ADDRESS> FoundRootNodeAddresses;
+	TArray<UE::Mutable::Private::FOperation::ADDRESS> FoundRootNodeAddresses;
 	
 	const int32 StateCount = MutableModel->GetProgram().States.Num();
 	for ( int32 StateIndex=0; StateIndex < StateCount; ++StateIndex )
@@ -2579,7 +2557,7 @@ void SMutableCodeViewer::GetOperationsReferencingConstantResource(
 	for	(int32 ParentIndex = 0 ; ParentIndex < InSearchPayload.BatchData.Num(); ParentIndex++)
 	{
 		// Get one of the previous run "children" and treat as a parent to get it's children and process them
-		const UE::Mutable::Private::OP::ADDRESS& ParentAddress = InSearchPayload.BatchData[ParentIndex].Child;
+		const UE::Mutable::Private::FOperation::ADDRESS& ParentAddress = InSearchPayload.BatchData[ParentIndex].Child;
 		
 		// Cache if same data type and we share the same address (means this op is pointing at the provided resource)
 		// It will cache duplicated entries
@@ -2607,7 +2585,7 @@ void SMutableCodeViewer::GetOperationsReferencingConstantResource(
 }
 
 bool SMutableCodeViewer::IsConstantResourceUsedByOperation(const int32 IndexOnConstantsArray,
-	const UE::Mutable::Private::EDataType ConstantDataType, const UE::Mutable::Private::OP::ADDRESS OperationAddress, const UE::Mutable::Private::FProgram& InProgram) const
+	const UE::Mutable::Private::EDataType ConstantDataType, const UE::Mutable::Private::FOperation::ADDRESS OperationAddress, const UE::Mutable::Private::FProgram& InProgram) const
 {
 	// Cache the current operation type to know where to look and what to check
 	const UE::Mutable::Private::EOpType OperationType = InProgram.GetOpType(OperationAddress);
@@ -2631,23 +2609,23 @@ bool SMutableCodeViewer::IsConstantResourceUsedByOperation(const int32 IndexOnCo
 				// TIP: To know if they represent a constant value check the code on code runner to see if they read from the constants array
 				if (OperationType == UE::Mutable::Private::EOpType::ST_CONSTANT)
 				{
-					bResourceLocated = IndexOnConstantsArray == InProgram.GetOpArgs<UE::Mutable::Private::OP::ResourceConstantArgs>(OperationAddress).value;
+					bResourceLocated = IndexOnConstantsArray == InProgram.GetOpArgs<UE::Mutable::Private::FOperation::ResourceConstantArgs>(OperationAddress).value;
 				}
 				else if (OperationType == UE::Mutable::Private::EOpType::IN_ADDSTRING)
 				{
-					bResourceLocated = IndexOnConstantsArray == InProgram.GetOpArgs<UE::Mutable::Private::OP::InstanceAddArgs>(OperationAddress).name;
+					bResourceLocated = IndexOnConstantsArray == InProgram.GetOpArgs<UE::Mutable::Private::FOperation::InstanceAddArgs>(OperationAddress).name;
 				}
 				else if (OperationType == UE::Mutable::Private::EOpType::IN_ADDMESH)
 				{
-					bResourceLocated = IndexOnConstantsArray == InProgram.GetOpArgs<UE::Mutable::Private::OP::InstanceAddArgs>(OperationAddress).name;
+					bResourceLocated = IndexOnConstantsArray == InProgram.GetOpArgs<UE::Mutable::Private::FOperation::InstanceAddArgs>(OperationAddress).name;
 				}
 				else if (OperationType == UE::Mutable::Private::EOpType::IN_ADDSURFACE)
 				{
-					bResourceLocated = IndexOnConstantsArray == InProgram.GetOpArgs<UE::Mutable::Private::OP::InstanceAddArgs>(OperationAddress).name;
+					bResourceLocated = IndexOnConstantsArray == InProgram.GetOpArgs<UE::Mutable::Private::FOperation::InstanceAddArgs>(OperationAddress).name;
 				}
 				else if (OperationType == UE::Mutable::Private::EOpType::IN_CONDITIONAL)
 				{
-					UE::Mutable::Private::OP::ConditionalArgs Arguments = InProgram.GetOpArgs<UE::Mutable::Private::OP::ConditionalArgs>(OperationAddress);
+					UE::Mutable::Private::FOperation::ConditionalArgs Arguments = InProgram.GetOpArgs<UE::Mutable::Private::FOperation::ConditionalArgs>(OperationAddress);
 					
 					if (IndexOnConstantsArray == Arguments.condition)
 					{
@@ -2675,7 +2653,7 @@ bool SMutableCodeViewer::IsConstantResourceUsedByOperation(const int32 IndexOnCo
 			{
 				if (OperationType == UE::Mutable::Private::EOpType::IM_CONSTANT)
 				{
-					bResourceLocated = IndexOnConstantsArray == InProgram.GetOpArgs<UE::Mutable::Private::OP::ResourceConstantArgs>(OperationAddress).value;
+					bResourceLocated = IndexOnConstantsArray == InProgram.GetOpArgs<UE::Mutable::Private::FOperation::ResourceConstantArgs>(OperationAddress).value;
 				} 
 				
 				break;
@@ -2685,7 +2663,7 @@ bool SMutableCodeViewer::IsConstantResourceUsedByOperation(const int32 IndexOnCo
 			{
 				if (OperationType == UE::Mutable::Private::EOpType::ME_CONSTANT)
 				{
-					bResourceLocated = IndexOnConstantsArray == InProgram.GetOpArgs<UE::Mutable::Private::OP::MeshConstantArgs>(OperationAddress).Value;
+					bResourceLocated = IndexOnConstantsArray == InProgram.GetOpArgs<UE::Mutable::Private::FOperation::MeshConstantArgs>(OperationAddress).Value;
 				}
 				break;
 			}
@@ -2694,7 +2672,7 @@ bool SMutableCodeViewer::IsConstantResourceUsedByOperation(const int32 IndexOnCo
 			{
 				if (OperationType == UE::Mutable::Private::EOpType::LA_CONSTANT)
 				{
-					bResourceLocated = IndexOnConstantsArray == InProgram.GetOpArgs<UE::Mutable::Private::OP::ResourceConstantArgs>(OperationAddress).value;
+					bResourceLocated = IndexOnConstantsArray == InProgram.GetOpArgs<UE::Mutable::Private::FOperation::ResourceConstantArgs>(OperationAddress).value;
 				}
 				break;
 			}
@@ -2703,7 +2681,7 @@ bool SMutableCodeViewer::IsConstantResourceUsedByOperation(const int32 IndexOnCo
 			{
 				if (OperationType == UE::Mutable::Private::EOpType::PR_CONSTANT)
 				{
-					bResourceLocated = IndexOnConstantsArray == InProgram.GetOpArgs<UE::Mutable::Private::OP::ResourceConstantArgs>(OperationAddress).value;
+					bResourceLocated = IndexOnConstantsArray == InProgram.GetOpArgs<UE::Mutable::Private::FOperation::ResourceConstantArgs>(OperationAddress).value;
 				}
 				break;
 			}
@@ -2712,19 +2690,19 @@ bool SMutableCodeViewer::IsConstantResourceUsedByOperation(const int32 IndexOnCo
 			{
 				if (OperationType == UE::Mutable::Private::EOpType::ME_TRANSFORM)
 				{
-					bResourceLocated = IndexOnConstantsArray == InProgram.GetOpArgs<UE::Mutable::Private::OP::MeshTransformArgs>(OperationAddress).matrix;
+					bResourceLocated = IndexOnConstantsArray == InProgram.GetOpArgs<UE::Mutable::Private::FOperation::MeshTransformArgs>(OperationAddress).matrix;
 				}
 				else if (OperationType == UE::Mutable::Private::EOpType::ME_TRANSFORMWITHMESH)
 				{
-					bResourceLocated = IndexOnConstantsArray == InProgram.GetOpArgs<UE::Mutable::Private::OP::MeshTransformWithinMeshArgs>(OperationAddress).matrix;
+					bResourceLocated = IndexOnConstantsArray == InProgram.GetOpArgs<UE::Mutable::Private::FOperation::MeshTransformWithinMeshArgs>(OperationAddress).matrix;
 				}
 				else if (OperationType == UE::Mutable::Private::EOpType::ME_TRANSFORMWITHBONE)
 				{
-					bResourceLocated = IndexOnConstantsArray == InProgram.GetOpArgs<UE::Mutable::Private::OP::MeshTransformWithBoneArgs>(OperationAddress).Matrix;
+					bResourceLocated = IndexOnConstantsArray == InProgram.GetOpArgs<UE::Mutable::Private::FOperation::MeshTransformWithBoneArgs>(OperationAddress).Matrix;
 				}
 				else if (OperationType == UE::Mutable::Private::EOpType::MA_CONSTANT)
 				{
-					bResourceLocated = IndexOnConstantsArray == InProgram.GetOpArgs<UE::Mutable::Private::OP::MatrixConstantArgs>(OperationAddress).value;
+					bResourceLocated = IndexOnConstantsArray == InProgram.GetOpArgs<UE::Mutable::Private::FOperation::MatrixConstantArgs>(OperationAddress).value;
 				}
 				break;
 			}
@@ -2733,7 +2711,7 @@ bool SMutableCodeViewer::IsConstantResourceUsedByOperation(const int32 IndexOnCo
 			{
 				if (OperationType == UE::Mutable::Private::EOpType::ME_CLIPMORPHPLANE)
 				{
-					const UE::Mutable::Private::OP::MeshClipMorphPlaneArgs Arguments =  InProgram.GetOpArgs<UE::Mutable::Private::OP::MeshClipMorphPlaneArgs>(OperationAddress);
+					const UE::Mutable::Private::FOperation::MeshClipMorphPlaneArgs Arguments =  InProgram.GetOpArgs<UE::Mutable::Private::FOperation::MeshClipMorphPlaneArgs>(OperationAddress);
 					
 					// Morph shape
 					bResourceLocated = IndexOnConstantsArray == Arguments.MorphShape;
@@ -2745,7 +2723,7 @@ bool SMutableCodeViewer::IsConstantResourceUsedByOperation(const int32 IndexOnCo
 					if (Arguments.VertexSelectionType == EClipVertexSelectionType::Shape)
 					{
 						// Selection Shape
-						bResourceLocated = IndexOnConstantsArray == InProgram.GetOpArgs<UE::Mutable::Private::OP::MeshClipMorphPlaneArgs>(OperationAddress).VertexSelectionShapeOrBone;
+						bResourceLocated = IndexOnConstantsArray == InProgram.GetOpArgs<UE::Mutable::Private::FOperation::MeshClipMorphPlaneArgs>(OperationAddress).VertexSelectionShapeOrBone;
 					}
 				}
 				break;
@@ -2755,7 +2733,7 @@ bool SMutableCodeViewer::IsConstantResourceUsedByOperation(const int32 IndexOnCo
 			{
 				if (OperationType == UE::Mutable::Private::EOpType::SC_CURVE)
 				{
-					bResourceLocated = IndexOnConstantsArray == InProgram.GetOpArgs<UE::Mutable::Private::OP::ScalarCurveArgs>(OperationAddress).curve;
+					bResourceLocated = IndexOnConstantsArray == InProgram.GetOpArgs<UE::Mutable::Private::FOperation::ScalarCurveArgs>(OperationAddress).curve;
 				}
 				break;
 			}
@@ -2764,7 +2742,7 @@ bool SMutableCodeViewer::IsConstantResourceUsedByOperation(const int32 IndexOnCo
 			{
 				if (OperationType == UE::Mutable::Private::EOpType::ME_CONSTANT)
 				{
-					bResourceLocated = IndexOnConstantsArray == InProgram.GetOpArgs<UE::Mutable::Private::OP::MeshConstantArgs>(OperationAddress).Skeleton;
+					bResourceLocated = IndexOnConstantsArray == InProgram.GetOpArgs<UE::Mutable::Private::FOperation::MeshConstantArgs>(OperationAddress).Skeleton;
 				}
 				break;
 			}
@@ -2772,7 +2750,7 @@ bool SMutableCodeViewer::IsConstantResourceUsedByOperation(const int32 IndexOnCo
 			{
 				if (OperationType == UE::Mutable::Private::EOpType::ME_CONSTANT)
 				{
-					bResourceLocated = IndexOnConstantsArray == InProgram.GetOpArgs<UE::Mutable::Private::OP::MeshConstantArgs>(OperationAddress).Value;
+					bResourceLocated = IndexOnConstantsArray == InProgram.GetOpArgs<UE::Mutable::Private::FOperation::MeshConstantArgs>(OperationAddress).Value;
 				}
 				break;
 			}
@@ -2781,7 +2759,7 @@ bool SMutableCodeViewer::IsConstantResourceUsedByOperation(const int32 IndexOnCo
 			{
 				if (OperationType == UE::Mutable::Private::EOpType::MI_CONSTANT)
 				{
-					bResourceLocated = IndexOnConstantsArray == InProgram.GetOpArgs<UE::Mutable::Private::OP::MaterialConstantArgs>(OperationAddress).Value;
+					bResourceLocated = IndexOnConstantsArray == InProgram.GetOpArgs<UE::Mutable::Private::FOperation::MaterialConstantArgs>(OperationAddress).Value;
 				}
 				break;
 			}
@@ -2934,7 +2912,7 @@ void SMutableCodeViewer::Tick(const FGeometry& AllottedGeometry, const double In
 }
 
 
-UE::Mutable::Private::TManagedPtr<const UE::Mutable::Private::FImage> SMutableCodeViewer::BuildImage(UE::Mutable::Private::OP::ADDRESS TargetOpAddress) const
+UE::Mutable::Private::TManagedPtr<const UE::Mutable::Private::FImage> SMutableCodeViewer::BuildImage(UE::Mutable::Private::FOperation::ADDRESS TargetOpAddress) const
 {
 	UE::Mutable::Private::FSettings Settings;
 	const TSharedPtr<UE::Mutable::Private::FSystem> System = MakeShared<UE::Mutable::Private::FSystem>(Settings);
@@ -3130,7 +3108,7 @@ void SMutableCodeViewer::PreviewMutableShape(const UE::Mutable::Private::FShape*
 }
 
 
-FMutableCodeTreeElement::FMutableCodeTreeElement(int32 InIndexOnTree, const int32& InMutableStateIndex, const TSharedPtr<UE::Mutable::Private::FModel, ESPMode::ThreadSafe>& InModel, UE::Mutable::Private::OP::ADDRESS InOperation, const FString& InCaption, const FSlateColor InLabelColor, const TSharedPtr<FMutableCodeTreeElement>* InDuplicatedOf)
+FMutableCodeTreeElement::FMutableCodeTreeElement(int32 InIndexOnTree, const int32& InMutableStateIndex, const TSharedPtr<UE::Mutable::Private::FModel, ESPMode::ThreadSafe>& InModel, UE::Mutable::Private::FOperation::ADDRESS InOperation, const FString& InCaption, const FSlateColor InLabelColor, const TSharedPtr<FMutableCodeTreeElement>* InDuplicatedOf)
 {
 	MutableModel = InModel;
 	MutableOperation = InOperation;
@@ -3211,7 +3189,7 @@ void FMutableCodeTreeElement::GenerateLabelText()
 	case UE::Mutable::Private::EOpType::PR_PARAMETER:
 	case UE::Mutable::Private::EOpType::ST_PARAMETER:
 	{
-		UE::Mutable::Private::OP::ParameterArgs Args = Program.GetOpArgs<UE::Mutable::Private::OP::ParameterArgs>(MutableOperation);
+		UE::Mutable::Private::FOperation::ParameterArgs Args = Program.GetOpArgs<UE::Mutable::Private::FOperation::ParameterArgs>(MutableOperation);
 		OpName += TEXT(" ");
 		OpName += Program.Parameters[int32(Args.variable)].Name;
 		break;
@@ -3219,14 +3197,14 @@ void FMutableCodeTreeElement::GenerateLabelText()
 
 	case UE::Mutable::Private::EOpType::ME_SKELETALMESH_BREAK:
 	{
-		UE::Mutable::Private::OP::MeshSkeletalMeshBreakArgs Args = Program.GetOpArgs<UE::Mutable::Private::OP::MeshSkeletalMeshBreakArgs>(MutableOperation);
+		UE::Mutable::Private::FOperation::MeshSkeletalMeshBreakArgs Args = Program.GetOpArgs<UE::Mutable::Private::FOperation::MeshSkeletalMeshBreakArgs>(MutableOperation);
 		OpName += FString::Printf( TEXT(" LOD %d Section %d Flags %d"), Args.LOD, Args.Section, Args.Flags );
 		break;
 	}
 
 	case UE::Mutable::Private::EOpType::IM_SWIZZLE:
 	{
-		UE::Mutable::Private::OP::ImageSwizzleArgs Args = Program.GetOpArgs<UE::Mutable::Private::OP::ImageSwizzleArgs>(MutableOperation);
+		UE::Mutable::Private::FOperation::ImageSwizzleArgs Args = Program.GetOpArgs<UE::Mutable::Private::FOperation::ImageSwizzleArgs>(MutableOperation);
 		OpName += TEXT(" ");
 		OpName += StringCast<TCHAR>(UE::Mutable::Private::TypeInfo::s_imageFormatName[int32(Args.format)]).Get();
 		break;
@@ -3234,7 +3212,7 @@ void FMutableCodeTreeElement::GenerateLabelText()
 
 	case UE::Mutable::Private::EOpType::IM_PIXELFORMAT:
 	{
-		UE::Mutable::Private::OP::ImagePixelFormatArgs Args = Program.GetOpArgs<UE::Mutable::Private::OP::ImagePixelFormatArgs>(MutableOperation);
+		UE::Mutable::Private::FOperation::ImagePixelFormatArgs Args = Program.GetOpArgs<UE::Mutable::Private::FOperation::ImagePixelFormatArgs>(MutableOperation);
 		OpName += TEXT(" ");
 		OpName += StringCast<TCHAR>(UE::Mutable::Private::TypeInfo::s_imageFormatName[int32(Args.format)]).Get();
 		OpName += TEXT(" or ");
@@ -3244,28 +3222,28 @@ void FMutableCodeTreeElement::GenerateLabelText()
 
 	case UE::Mutable::Private::EOpType::IM_MIPMAP:
 	{
-		UE::Mutable::Private::OP::ImageMipmapArgs Args = Program.GetOpArgs<UE::Mutable::Private::OP::ImageMipmapArgs>(MutableOperation);
+		UE::Mutable::Private::FOperation::ImageMipmapArgs Args = Program.GetOpArgs<UE::Mutable::Private::FOperation::ImageMipmapArgs>(MutableOperation);
 		OpName += FString::Printf(TEXT(" levels: %d-%d tail: %d"), Args.levels, Args.blockLevels, int32(Args.onlyTail));
 		break;
 	}
 
 	case UE::Mutable::Private::EOpType::IM_RESIZE:
 	{
-		UE::Mutable::Private::OP::ImageResizeArgs Args = Program.GetOpArgs<UE::Mutable::Private::OP::ImageResizeArgs>(MutableOperation);
+		UE::Mutable::Private::FOperation::ImageResizeArgs Args = Program.GetOpArgs<UE::Mutable::Private::FOperation::ImageResizeArgs>(MutableOperation);
 		OpName += FString::Printf(TEXT(" %d x %d"), int32(Args.Size[0]), int32(Args.Size[1]));
 		break;
 	}
 
 	case UE::Mutable::Private::EOpType::IM_RESIZEREL:
 	{
-		UE::Mutable::Private::OP::ImageResizeRelArgs Args = Program.GetOpArgs<UE::Mutable::Private::OP::ImageResizeRelArgs>(MutableOperation);
+		UE::Mutable::Private::FOperation::ImageResizeRelArgs Args = Program.GetOpArgs<UE::Mutable::Private::FOperation::ImageResizeRelArgs>(MutableOperation);
 		OpName += FString::Printf(TEXT(" %.3f x %.3f"), Args.Factor[0], Args.Factor[1]);
 		break;
 	}
 
 	case UE::Mutable::Private::EOpType::IM_MULTILAYER:
 	{
-		UE::Mutable::Private::OP::ImageMultiLayerArgs Args = Program.GetOpArgs<UE::Mutable::Private::OP::ImageMultiLayerArgs>(MutableOperation);
+		UE::Mutable::Private::FOperation::ImageMultiLayerArgs Args = Program.GetOpArgs<UE::Mutable::Private::FOperation::ImageMultiLayerArgs>(MutableOperation);
 		OpName += TEXT(" rgb: ");
 		OpName += UE::Mutable::Private::TypeInfo::s_blendModeName[int32(Args.blendType)];
 		OpName += TEXT(", a: ");
@@ -3278,7 +3256,7 @@ void FMutableCodeTreeElement::GenerateLabelText()
 
 	case UE::Mutable::Private::EOpType::IM_LAYER:
 	{
-		UE::Mutable::Private::OP::ImageLayerArgs Args = Program.GetOpArgs<UE::Mutable::Private::OP::ImageLayerArgs>(MutableOperation);
+		UE::Mutable::Private::FOperation::ImageLayerArgs Args = Program.GetOpArgs<UE::Mutable::Private::FOperation::ImageLayerArgs>(MutableOperation);
 		OpName += TEXT(" rgb: ");
 		OpName += UE::Mutable::Private::TypeInfo::s_blendModeName[int32(Args.blendType)];
 		OpName += TEXT(", a: ");
@@ -3290,7 +3268,7 @@ void FMutableCodeTreeElement::GenerateLabelText()
 
 	case UE::Mutable::Private::EOpType::IM_LAYERCOLOUR:
 	{
-		UE::Mutable::Private::OP::ImageLayerColourArgs Args = Program.GetOpArgs<UE::Mutable::Private::OP::ImageLayerColourArgs>(MutableOperation);
+		UE::Mutable::Private::FOperation::ImageLayerColourArgs Args = Program.GetOpArgs<UE::Mutable::Private::FOperation::ImageLayerColourArgs>(MutableOperation);
 		OpName += TEXT(" rgb: ");
 		OpName += UE::Mutable::Private::TypeInfo::s_blendModeName[int32(Args.blendType)];
 		OpName += TEXT(" a: ");
@@ -3303,7 +3281,7 @@ void FMutableCodeTreeElement::GenerateLabelText()
 
 	case UE::Mutable::Private::EOpType::IM_PLAINCOLOUR:
 	{
-		UE::Mutable::Private::OP::ImagePlainColorArgs Args = Program.GetOpArgs<UE::Mutable::Private::OP::ImagePlainColorArgs>(MutableOperation);
+		UE::Mutable::Private::FOperation::ImagePlainColorArgs Args = Program.GetOpArgs<UE::Mutable::Private::FOperation::ImagePlainColorArgs>(MutableOperation);
 		OpName += TEXT(" format: ");
 		OpName += StringCast<TCHAR>(UE::Mutable::Private::TypeInfo::s_imageFormatName[int32(Args.Format)]).Get();
 		OpName += FString::Printf(TEXT(" size %d x %d"), Args.Size[0], Args.Size[1]);
diff --git a/Engine/Plugins/Mutable/Source/CustomizableObjectEditor/Private/MuCOE/SMutableCodeViewer.h b/Engine/Plugins/Mutable/Source/CustomizableObjectEditor/Private/MuCOE/SMutableCodeViewer.h
--- a/Engine/Plugins/Mutable/Source/CustomizableObjectEditor/Private/MuCOE/SMutableCodeViewer.h
+++ b/Engine/Plugins/Mutable/Source/CustomizableObjectEditor/Private/MuCOE/SMutableCodeViewer.h
@@ -73,7 +73,7 @@ public:
 	 * @return A pointer to the FImage generated. May be null if the operation does not evaluate to an FImage
 	 * @note This method is experimental and should not be used as example
 	 */
-	UE::Mutable::Private::TManagedPtr<const UE::Mutable::Private::FImage>  BuildImage(UE::Mutable::Private::OP::ADDRESS TargetOpAddress) const;
+	UE::Mutable::Private::TManagedPtr<const UE::Mutable::Private::FImage>  BuildImage(UE::Mutable::Private::FOperation::ADDRESS TargetOpAddress) const;
 	
 private:
 	TStrongObjectPtr<UCustomizableObject> Object;
@@ -89,7 +89,7 @@ private:
 	TArray<TSoftObjectPtr<const UStreamableRenderAsset>> ReferencedMeshes;
 
 	/** Selected model operation for preview. */
-	UE::Mutable::Private::OP::ADDRESS SelectedOperationAddress = 0;
+	UE::Mutable::Private::FOperation::ADDRESS SelectedOperationAddress = 0;
 
 	/** Mutable parameters used in the preview. */
 	TSharedPtr<UE::Mutable::Private::FParameters> PreviewParameters;
@@ -160,8 +160,8 @@ private:
 	*/
 	struct FItemCacheKey
 	{
-		UE::Mutable::Private::OP::ADDRESS Parent=0;
-		UE::Mutable::Private::OP::ADDRESS Child=0;
+		UE::Mutable::Private::FOperation::ADDRESS Parent=0;
+		UE::Mutable::Private::FOperation::ADDRESS Child=0;
 		uint32 ChildIndexInParent=0;
 
 		friend FORCEINLINE bool operator == (const FItemCacheKey& A, const FItemCacheKey& B)
@@ -188,13 +188,13 @@ private:
 	TMap< FItemCacheKey, TSharedPtr<FMutableCodeTreeElement>> ItemCache;
 
 	/** Main tree item for each op. An op can be represented with multiple tree nodes if it is reachable from different paths. */
-	TMap< UE::Mutable::Private::OP::ADDRESS, TSharedPtr<FMutableCodeTreeElement>> MainItemPerOp;
+	TMap< UE::Mutable::Private::FOperation::ADDRESS, TSharedPtr<FMutableCodeTreeElement>> MainItemPerOp;
 
 	/** List with all the elements related to nodes displayed on the tree. */
 	TArray<TSharedPtr<FMutableCodeTreeElement>> TreeElements;
 
 	/** Array with all the elements that have been manually expanded by the user */
-	TMap< UE::Mutable::Private::OP::ADDRESS, TSharedPtr<FMutableCodeTreeElement>> ExpandedElements;
+	TMap< UE::Mutable::Private::FOperation::ADDRESS, TSharedPtr<FMutableCodeTreeElement>> ExpandedElements;
 
 	/** Before any UI operation generate all the elements that may be navigable over the tree. No children of duplicated 
 	 * addresses will be generated.
@@ -202,11 +202,11 @@ private:
 	void GenerateAllTreeElements();
 
 	/** Support GenerateAllTreeElements by generating the elements recursively */
-	void GenerateElementRecursive( const int32& InStateIndex, UE::Mutable::Private::OP::ADDRESS InParentAddress, const UE::Mutable::Private::FProgram& InProgram);
+	void GenerateElementRecursive( const int32& InStateIndex, UE::Mutable::Private::FOperation::ADDRESS InParentAddress, const UE::Mutable::Private::FProgram& InProgram);
 
 	
 	/** The addresses of the root operations. Cached when this object gets loaded on the Construct method of this slate */
-	TArray<UE::Mutable::Private::OP::ADDRESS> RootNodeAddresses;
+	TArray<UE::Mutable::Private::FOperation::ADDRESS> RootNodeAddresses;
 	
 	/** Utility method that provides you with the addresses of the root nodes of the program */
 	void CacheRootNodeAddresses();
@@ -215,7 +215,7 @@ private:
 	bool bIsElementHighlighted = false;
 
 	/** Operation that is currently being highlighted */
-	UE::Mutable::Private::OP::ADDRESS HighlightedOperation{};
+	UE::Mutable::Private::FOperation::ADDRESS HighlightedOperation{};
 
 	/*
 	 * Operations computational cost reference collections
@@ -400,7 +400,7 @@ private:
 	struct FElementsSearchCache
 	{
 		/** Set of addresses that have already been searched for relevant data */
-		TSet<UE::Mutable::Private::OP::ADDRESS> ProcessedAddresses;
+		TSet<UE::Mutable::Private::FOperation::ADDRESS> ProcessedAddresses;
 
 		/** Collection of elements that have been found during the search. They may be related by OP_Type or used constant resource*/
 		TArray<TSharedPtr<FMutableCodeTreeElement>> FoundElements;
@@ -417,7 +417,7 @@ private:
 		 * start of the search operation.
 		 * @param InRootNodeAddresses The addresses of the root operations of the operations tree
 		 */
-		void SetupRootBatch(const TArray<UE::Mutable::Private::OP::ADDRESS>& InRootNodeAddresses)
+		void SetupRootBatch(const TArray<UE::Mutable::Private::FOperation::ADDRESS>& InRootNodeAddresses)
 		{
 			check (InRootNodeAddresses.Num());
 			// This method should only be called once when no data is present on this cache
@@ -449,7 +449,7 @@ private:
 		* @param InItemCache Cache with all the elements of the tree. Here is where we search for the element based on
 		* the OpAddress and the parent and ChildIndexInParent provided thanks to IndexAsChildOfInputAddress
 		*/
-		void AddToFoundElements(const UE::Mutable::Private::OP::ADDRESS OpAddress,const int32 IndexAsChildOfInputAddress,
+		void AddToFoundElements(const UE::Mutable::Private::FOperation::ADDRESS OpAddress,const int32 IndexAsChildOfInputAddress,
 		                        const TMap<FItemCacheKey, TSharedPtr<FMutableCodeTreeElement>>& InItemCache)
 		{
 			FItemCacheKey Key;
@@ -480,7 +480,7 @@ private:
 		 * @note It will not add the children of any provided ParentAddress for the next batch if the parent address provided
 		 * has already been processed and therefore whose children have already been searched or prepared for searching.
 		 */
-		void CacheChildrenOfAddressIfNotProcessed(UE::Mutable::Private::OP::ADDRESS InParentAddress,
+		void CacheChildrenOfAddressIfNotProcessed(UE::Mutable::Private::FOperation::ADDRESS InParentAddress,
 		                                          const UE::Mutable::Private::FProgram& InProgram,
 		                                          TArray<FItemCacheKey>& OutFoundChildrenData)
 		{
@@ -491,7 +491,7 @@ private:
 	
 				// Generic case for unnamed children traversal.
 				uint32 ChildIndex = 0;
-				UE::Mutable::Private::ForEachReference(InProgram, InParentAddress, [this, &InParentAddress, &ChildIndex,&OutFoundChildrenData]( UE::Mutable::Private::OP::ADDRESS ChildAddress)
+				UE::Mutable::Private::ForEachReference(InProgram, InParentAddress, [this, &InParentAddress, &ChildIndex,&OutFoundChildrenData]( UE::Mutable::Private::FOperation::ADDRESS ChildAddress)
 				{
 					// If the parent does have a child then process it
 					// Ignoring child address 0 since those can not have children
@@ -581,7 +581,7 @@ private:
 	 * @return True if the operation at OperationAddress does make use of the provided constant resource. False otherwise
 	 */
 	bool IsConstantResourceUsedByOperation(const int32 IndexOnConstantsArray,
-	                                       const UE::Mutable::Private::EDataType ConstantDataType, const UE::Mutable::Private::OP::ADDRESS OperationAddress,
+	                                       const UE::Mutable::Private::EDataType ConstantDataType, const UE::Mutable::Private::FOperation::ADDRESS OperationAddress,
 	                                       const UE::Mutable::Private::FProgram& InProgram) const;
 	
 
@@ -710,10 +710,10 @@ private:
 	struct FProcessedOperationsBuffer
 	{
 		/** Array with the operation addresses of all original (not duplicates) expanded operations. */
-		TArray<UE::Mutable::Private::OP::ADDRESS> ExpandedOriginalOperations;
+		TArray<UE::Mutable::Private::FOperation::ADDRESS> ExpandedOriginalOperations;
 
 		/** Array with the operation addresses of all duplicated expanded operations. */
-		TArray<UE::Mutable::Private::OP::ADDRESS> ExpandedDuplicatedOperations;
+		TArray<UE::Mutable::Private::FOperation::ADDRESS> ExpandedDuplicatedOperations;
 	};
 	
 	/** 
@@ -863,7 +863,7 @@ public:
 class FMutableCodeTreeElement : public TSharedFromThis<FMutableCodeTreeElement>
 {
 public:
-	FMutableCodeTreeElement(int32 InIndexOnTree, const int32& InMutableStateIndex, const TSharedPtr<UE::Mutable::Private::FModel, ESPMode::ThreadSafe>& InModel, UE::Mutable::Private::OP::ADDRESS InOperation, const FString& InCaption, const FSlateColor InLabelColor, const TSharedPtr<FMutableCodeTreeElement>* InDuplicatedOf = nullptr);
+	FMutableCodeTreeElement(int32 InIndexOnTree, const int32& InMutableStateIndex, const TSharedPtr<UE::Mutable::Private::FModel, ESPMode::ThreadSafe>& InModel, UE::Mutable::Private::FOperation::ADDRESS InOperation, const FString& InCaption, const FSlateColor InLabelColor, const TSharedPtr<FMutableCodeTreeElement>* InDuplicatedOf = nullptr);
 	
 	void SetElementCurrentState(const int32& InStateIndex);
 
@@ -877,7 +877,7 @@ public:
 	TSharedPtr<UE::Mutable::Private::FModel, ESPMode::ThreadSafe> MutableModel;
 	
 	/** Mutable Graph Node represented in this tree row. */
-	UE::Mutable::Private::OP::ADDRESS MutableOperation;
+	UE::Mutable::Private::FOperation::ADDRESS MutableOperation;
 
 	/** If this tree element is a duplicated of another op, this is the op. */
 	TSharedPtr<FMutableCodeTreeElement> DuplicatedOf;
diff --git a/Engine/Plugins/Mutable/Source/CustomizableObjectEditor/Private/MuCOE/SMutableInstanceViewer.cpp b/Engine/Plugins/Mutable/Source/CustomizableObjectEditor/Private/MuCOE/SMutableInstanceViewer.cpp
--- a/Engine/Plugins/Mutable/Source/CustomizableObjectEditor/Private/MuCOE/SMutableInstanceViewer.cpp
+++ b/Engine/Plugins/Mutable/Source/CustomizableObjectEditor/Private/MuCOE/SMutableInstanceViewer.cpp
@@ -320,7 +320,7 @@ void SMutableInstanceViewer::SetInstance(const TSharedRef<UE::Mutable::Private::
 					{
 						check(InLiveInstance->ExternalResourceProvider)
 						
-						UE::Mutable::Private::OP::ADDRESS MeshRootAddress = MeshId.GetKey()->Address;
+						UE::Mutable::Private::FOperation::ADDRESS MeshRootAddress = MeshId.GetKey()->Address;
 						UE::Mutable::Private::TManagedPtr<const UE::Mutable::Private::FMesh> MutableMesh = System.BuildMesh(
 								InLiveInstance,
 								MeshRootAddress, 
diff --git a/Engine/Plugins/Mutable/Source/CustomizableObjectEditor/Private/MuCOE/SMutableMaterialViewer.cpp b/Engine/Plugins/Mutable/Source/CustomizableObjectEditor/Private/MuCOE/SMutableMaterialViewer.cpp
--- a/Engine/Plugins/Mutable/Source/CustomizableObjectEditor/Private/MuCOE/SMutableMaterialViewer.cpp
+++ b/Engine/Plugins/Mutable/Source/CustomizableObjectEditor/Private/MuCOE/SMutableMaterialViewer.cpp
@@ -90,7 +90,7 @@ public:
 		// OP ADDRESS
 		if (InColumnName == MutableMaterialParameterColumns::ParameterImageAddressID)
 		{
-			if (const UE::Mutable::Private::OP::ADDRESS* ImageAddress = RowItem->Value.ImageParameter.TryGet<UE::Mutable::Private::OP::ADDRESS>())
+			if (const UE::Mutable::Private::FOperation::ADDRESS* ImageAddress = RowItem->Value.ImageParameter.TryGet<UE::Mutable::Private::FOperation::ADDRESS>())
 			{
 				return 
 					SNew(STextBlock)
@@ -108,7 +108,7 @@ public:
 		{
 			if (TSharedPtr<SMutableCodeViewer> PinnedCodeViewer = MutableCodeViewer.Pin())
 			{
-				if (const UE::Mutable::Private::OP::ADDRESS* ImageAddress = RowItem->Value.ImageParameter.TryGet<UE::Mutable::Private::OP::ADDRESS>())
+				if (const UE::Mutable::Private::FOperation::ADDRESS* ImageAddress = RowItem->Value.ImageParameter.TryGet<UE::Mutable::Private::FOperation::ADDRESS>())
 				{
 					UE::Mutable::Private::TManagedPtr<const UE::Mutable::Private::FImage> GeneratedImage = PinnedCodeViewer->BuildImage(*ImageAddress);
 					return 
diff --git a/Engine/Plugins/Mutable/Source/MutableRuntime/Internal/MuR/CodeVisitor.h b/Engine/Plugins/Mutable/Source/MutableRuntime/Internal/MuR/CodeVisitor.h
--- a/Engine/Plugins/Mutable/Source/MutableRuntime/Internal/MuR/CodeVisitor.h
+++ b/Engine/Plugins/Mutable/Source/MutableRuntime/Internal/MuR/CodeVisitor.h
@@ -66,7 +66,7 @@ namespace UE::Mutable::Private
         }
 
         //! For manual recursion that changes the state for a specific path.
-        void RecurseWithState(OP::ADDRESS Address, const STATE& NewState)
+        void RecurseWithState(FOperation::ADDRESS Address, const STATE& NewState)
         {
 			int32 StateIndex = States.Find(NewState);
             if (StateIndex==INDEX_NONE)
@@ -77,7 +77,7 @@ namespace UE::Mutable::Private
         }
 
         //! For manual recursion that doesn't change the state for a specific path.
-        void RecurseWithCurrentState(OP::ADDRESS Address)
+        void RecurseWithCurrentState(FOperation::ADDRESS Address)
         {
             Pending.Emplace( Address, CurrentState );
         }
@@ -93,22 +93,21 @@ namespace UE::Mutable::Private
         }
 
 
-        void Traverse( OP::ADDRESS Root, const FProgram& Program )
+        void Traverse(FOperation::ADDRESS Root, const FProgram& Program )
         {
-            Pending.Reserve( Program.OpAddress.Num() );
-
+            Pending.Reserve(Program.NumOps);
             // Visit the given root
             Pending.Emplace(Root, 0);
-            Recurse( Program );
+            Recurse(Program);
          }
 
-        void FullTraverse( const FProgram& Program )
+        void FullTraverse(const FProgram& Program)
         {
             // Visit all the state roots
-            for ( int32 StateIndex=0; StateIndex<Program.States.Num(); ++StateIndex)
+            for (int32 StateIndex = 0; StateIndex<Program.States.Num(); ++StateIndex)
             {
-                Pending.Add(FPending(Program.States[StateIndex].Root,0) );
-                Recurse( Program );
+                Pending.Add(FPending(Program.States[StateIndex].Root, 0));
+                Recurse(Program);
             }
         }
 
@@ -118,7 +117,7 @@ namespace UE::Mutable::Private
         /** Do the actual work by overriding this in the derived classes.
          * Return true if the traverse has to continue with the children of the given address.
 		 */
-        virtual bool Visit(OP::ADDRESS, const FProgram&) = 0;
+        virtual bool Visit(FOperation::ADDRESS, const FProgram&) = 0;
 
         //! Operations to be processed
         struct FPending
@@ -129,13 +128,13 @@ namespace UE::Mutable::Private
 				StateIndex = 0;
 			}
 
-			FPending(OP::ADDRESS InAddress, int32 InStateIndex)
+			FPending(FOperation::ADDRESS InAddress, int32 InStateIndex)
 			{
 				Address = InAddress;
 				StateIndex = InStateIndex;
 			}
 			
-			OP::ADDRESS Address=0;
+			FOperation::ADDRESS Address=0;
             int32 StateIndex=0;
         };
 		TArray<FPending> Pending;
@@ -152,34 +151,34 @@ namespace UE::Mutable::Private
 
         //! Array of states visited for each operation.
         //! Empty array means operation not visited at all.
-		TArray<TArray<int32>> Visited;
+		TMap<FOperation::ADDRESS, TArray<int32>> Visited;
 
         //! Process all the pending operations and visit all children if necessary
         void Recurse(const FProgram& Program)
         {
 			Visited.Empty();
-			Visited.SetNum(Program.OpAddress.Num());
+			Visited.Reserve(Program.NumOps);
 
             while ( Pending.Num() )
             {
-                OP::ADDRESS Address = Pending.Last().Address;
+                FOperation::ADDRESS Address = Pending.Last().Address;
                 CurrentState = Pending.Last().StateIndex;
                 Pending.Pop();
 
                 bool bRecurse = false;
 
-                bool bVisitedInThisState = Visited[Address].Contains(CurrentState);
+                bool bVisitedInThisState = Visited.FindOrAdd(Address).Contains(CurrentState);
                 if (!bVisitedInThisState)
                 {
                     Visited[Address].Add(CurrentState);
 
                     // Visit may change current state
-                    bRecurse = Visit(Address, Program );
+                    bRecurse = Visit(Address, Program);
                 }
 
                 if (bRecurse)
                 {
-                    ForEachReference( Program, Address, [&](OP::ADDRESS Ref)
+                    ForEachReference( Program, Address, [&](FOperation::ADDRESS Ref)
                     {
                         if (Ref)
                         {
@@ -234,7 +233,7 @@ namespace UE::Mutable::Private
             PARENT::SetDefaultState( STATE() );
         }
 
-        void Run( OP::ADDRESS at  )
+        void Run( FOperation::ADDRESS at  )
         {
             PARENT::SetDefaultState( STATE() );
 
@@ -243,7 +242,7 @@ namespace UE::Mutable::Private
 
     protected:
 
-        virtual bool Visit(OP::ADDRESS Address, const FProgram& Program)
+        virtual bool Visit(FOperation::ADDRESS Address, const FProgram& Program)
         {
             bool bRecurse = true;
 
@@ -260,7 +259,7 @@ namespace UE::Mutable::Private
             case EOpType::IN_CONDITIONAL:
 			case EOpType::ED_CONDITIONAL:
             {
-				OP::ConditionalArgs Args = Program.GetOpArgs<OP::ConditionalArgs>(Address);
+				FOperation::ConditionalArgs Args = Program.GetOpArgs<FOperation::ConditionalArgs>(Address);
 
                 bRecurse = false;
 
@@ -298,25 +297,25 @@ namespace UE::Mutable::Private
 
 				const uint8* Data = Program.GetOpArgsPointer(Address);
 				
-				OP::ADDRESS VarAddress;
-				FMemory::Memcpy(&VarAddress, Data, sizeof(OP::ADDRESS));
-				Data += sizeof(OP::ADDRESS);
+				FOperation::ADDRESS VarAddress;
+				FMemory::Memcpy(&VarAddress, Data, sizeof(FOperation::ADDRESS));
+				Data += sizeof(FOperation::ADDRESS);
 
                 if (VarAddress)
                 {
-					OP::ADDRESS DefAddress;
-					FMemory::Memcpy(&DefAddress, Data, sizeof(OP::ADDRESS));
-					Data += sizeof(OP::ADDRESS);
+					FOperation::ADDRESS DefAddress;
+					FMemory::Memcpy(&DefAddress, Data, sizeof(FOperation::ADDRESS));
+					Data += sizeof(FOperation::ADDRESS);
 
-					OP::FSwitchCaseDescriptor CaseDesc;
-					FMemory::Memcpy(&CaseDesc, Data, sizeof(OP::FSwitchCaseDescriptor));
-					Data += sizeof(OP::FSwitchCaseDescriptor);
+					FOperation::FSwitchCaseDescriptor CaseDesc;
+					FMemory::Memcpy(&CaseDesc, Data, sizeof(FOperation::FSwitchCaseDescriptor));
+					Data += sizeof(FOperation::FSwitchCaseDescriptor);
 
                     PARENT::RecurseWithCurrentState( VarAddress );
 
                     int32 Var = System->BuildInt(LiveInstance, VarAddress );
 
-					OP::ADDRESS ValueAt = DefAddress;
+					FOperation::ADDRESS ValueAt = DefAddress;
 
 					if (!CaseDesc.bUseRanges)
 					{
@@ -326,9 +325,9 @@ namespace UE::Mutable::Private
 							FMemory::Memcpy(&Condition, Data, sizeof(int32));
 							Data += sizeof(int32);
 
-							OP::ADDRESS CaseAt;
-							FMemory::Memcpy(&CaseAt, Data, sizeof(OP::ADDRESS));
-							Data += sizeof(OP::ADDRESS);
+							FOperation::ADDRESS CaseAt;
+							FMemory::Memcpy(&CaseAt, Data, sizeof(FOperation::ADDRESS));
+							Data += sizeof(FOperation::ADDRESS);
 
 							if (CaseAt && Var == (int32)Condition)
 							{
@@ -349,9 +348,9 @@ namespace UE::Mutable::Private
 							FMemory::Memcpy(&RangeSize, Data, sizeof(uint32));
 							Data += sizeof(uint32);
 
-							OP::ADDRESS CaseAt;
-							FMemory::Memcpy(&CaseAt, Data, sizeof(OP::ADDRESS));
-							Data += sizeof(OP::ADDRESS);
+							FOperation::ADDRESS CaseAt;
+							FMemory::Memcpy(&CaseAt, Data, sizeof(FOperation::ADDRESS));
+							Data += sizeof(FOperation::ADDRESS);
 
 							if (CaseAt && Var >= ConditionStart && Var < int32(ConditionStart + RangeSize))
 							{
@@ -381,9 +380,9 @@ namespace UE::Mutable::Private
                 STATE NewState = PARENT::GetCurrentState();
                 for (int8 LODIndex=0; LODIndex < LODCount;++LODIndex)
                 {
-					OP::ADDRESS LODAddress;
-					FMemory::Memcpy(&LODAddress, Data, sizeof(OP::ADDRESS));
-					Data += sizeof(OP::ADDRESS);
+					FOperation::ADDRESS LODAddress;
+					FMemory::Memcpy(&LODAddress, Data, sizeof(FOperation::ADDRESS));
+					Data += sizeof(FOperation::ADDRESS);
 					
 					if (LODAddress)
                     {
@@ -400,7 +399,7 @@ namespace UE::Mutable::Private
 
             case EOpType::IN_ADDMESH:
             {
-				OP::InstanceAddArgs Args = Program.GetOpArgs<OP::InstanceAddArgs>(Address);
+				FOperation::InstanceAddArgs Args = Program.GetOpArgs<FOperation::InstanceAddArgs>(Address);
 
                 bRecurse = false;
 
@@ -409,7 +408,7 @@ namespace UE::Mutable::Private
                 STATE NewState = PARENT::GetCurrentState();
                 NewState.m_underResourceCount=1;
 
-                OP::ADDRESS MeshAddress = Args.value;
+                FOperation::ADDRESS MeshAddress = Args.value;
                 if (MeshAddress)
                 {
                     PARENT::RecurseWithState(MeshAddress, NewState);
@@ -419,14 +418,14 @@ namespace UE::Mutable::Private
 
 			case EOpType::IN_ADDMATERIAL:
 			{
-				OP::InstanceAddMaterialArgs Args = Program.GetOpArgs<OP::InstanceAddMaterialArgs>(Address);
+				FOperation::InstanceAddMaterialArgs Args = Program.GetOpArgs<FOperation::InstanceAddMaterialArgs>(Address);
 				bRecurse = false;
 				const uint8* Data = Program.GetOpArgsPointer(Address);
 
 				PARENT::RecurseWithCurrentState(Args.Instance);
 				STATE NewState = PARENT::GetCurrentState();
 				NewState.m_underResourceCount = 1;
-				OP::ADDRESS MaterialAddress = Args.Material;
+				FOperation::ADDRESS MaterialAddress = Args.Material;
 				
 				if (MaterialAddress)
 				{
diff --git a/Engine/Plugins/Mutable/Source/MutableRuntime/Internal/MuR/ExternalOperationProvider.h b/Engine/Plugins/Mutable/Source/MutableRuntime/Internal/MuR/ExternalOperationProvider.h
--- a/Engine/Plugins/Mutable/Source/MutableRuntime/Internal/MuR/ExternalOperationProvider.h
+++ b/Engine/Plugins/Mutable/Source/MutableRuntime/Internal/MuR/ExternalOperationProvider.h
@@ -13,7 +13,7 @@ namespace UE::Mutable::Private
 	public:
 		virtual ~IExternalOperationProvider() = default;
 		
-		virtual const FInstancedStruct& Get(OP::ADDRESS Address) const = 0;
+		virtual const FInstancedStruct& Get(FOperation::ADDRESS Address) const = 0;
 	};
 }
 
diff --git a/Engine/Plugins/Mutable/Source/MutableRuntime/Internal/MuR/Material.h b/Engine/Plugins/Mutable/Source/MutableRuntime/Internal/MuR/Material.h
--- a/Engine/Plugins/Mutable/Source/MutableRuntime/Internal/MuR/Material.h
+++ b/Engine/Plugins/Mutable/Source/MutableRuntime/Internal/MuR/Material.h
@@ -71,7 +71,7 @@ namespace UE::Mutable::Private
 
 		struct FImageParameterData 
 		{
-			TVariant<OP::ADDRESS, TManagedPtr<const FImage>> ImageParameter;
+			TVariant<FOperation::ADDRESS, TManagedPtr<const FImage>> ImageParameter;
 			int32 ImagePropertyIndex;
 
 			inline void Serialise(FOutputArchive& arch) const;
diff --git a/Engine/Plugins/Mutable/Source/MutableRuntime/Internal/MuR/Model.h b/Engine/Plugins/Mutable/Source/MutableRuntime/Internal/MuR/Model.h
--- a/Engine/Plugins/Mutable/Source/MutableRuntime/Internal/MuR/Model.h
+++ b/Engine/Plugins/Mutable/Source/MutableRuntime/Internal/MuR/Model.h
@@ -89,11 +89,11 @@ namespace UE::Mutable::Private
 
     //!
     template<typename DATA>
-    inline void AppendCode(TArray<uint8>& Code, const DATA& Data )
+    inline void AppendCode(TArray<uint8>& ByteCode, const DATA& Data)
     {
-        int32 Pos = Code.Num();
-        Code.SetNum( Pos+sizeof(DATA), EAllowShrinking::No);
-		FMemory::Memcpy (&Code[Pos], &Data, sizeof(DATA));
+        int32 Pos = ByteCode.Num();
+        ByteCode.SetNum(Pos + sizeof(DATA), EAllowShrinking::No);
+		FMemory::Memcpy(&ByteCode[Pos], &Data, sizeof(DATA));
     }
 
 	//!
@@ -167,7 +167,7 @@ namespace UE::Mutable::Private
             FString Name;
 
             /** First instruction of the full build of an instance in this state */
-            OP::ADDRESS Root = 0;
+            FOperation::ADDRESS Root = 0;
 
             /** 
 			 * List of parameters index (to FProgram::Parameters) of the runtime parameters of
@@ -176,7 +176,7 @@ namespace UE::Mutable::Private
 			TArray<int> m_runtimeParameters;
 
             /** List of instructions that need to be cached to efficiently update this state */
-			TArray<OP::ADDRESS> m_updateCache;
+			TArray<FOperation::ADDRESS> m_updateCache;
 
             /** 
 			 * List of root instructions for the dynamic resources that depend on the runtime
@@ -184,7 +184,7 @@ namespace UE::Mutable::Private
              * The mask has a bit on for every runtime parameter in the m_runtimeParameters array.
 			 * The uint64 is linked to MUTABLE_MAX_RUNTIME_PARAMETERS_PER_STATE
 			 */
-			TArray< TPair<OP::ADDRESS,uint64> > m_dynamicResources;
+			TArray< TPair<FOperation::ADDRESS,uint64> > m_dynamicResources;
 
             /** */
             void Serialise(FOutputArchive& Arch) const;
@@ -195,17 +195,17 @@ namespace UE::Mutable::Private
             /** Returns the mask of parameters (from the runtime parameter list of this state) including the parameters that 
 			* are relevant for the dynamic resource at the given address.
 			*/
-            uint64 IsDynamic(OP::ADDRESS At) const;
+            uint64 IsDynamic(FOperation::ADDRESS At) const;
 
             /** */
-            bool IsUpdateCache(OP::ADDRESS at) const;
+            bool IsUpdateCache(FOperation::ADDRESS at) const;
 
             /** */ 
-            void AddUpdateCache(OP::ADDRESS At);
+            void AddUpdateCache(FOperation::ADDRESS At);
         };
 
-        /** Location in the ByteCode of the beginning of each operation */
-		TArray<uint32> OpAddress;
+		int32 NumOps = 0;
+		TArray<FOperation::ADDRESS> OperationDebugData;
 
         /** Byte-coded representation of the program, using variable-sized op data. */
 		TArray<uint8> ByteCode;
@@ -286,7 +286,7 @@ namespace UE::Mutable::Private
 		TArray<TArray<uint16>> ParameterLists;
 
     	/** Relevant parameters for the given instruction. */
-	   	TMap<OP::ADDRESS, int32> RelevantParameterList;
+	   	TMap<FOperation::ADDRESS, int32> RelevantParameterList;
 
 
 		/** Constant Material Data */
@@ -337,64 +337,64 @@ namespace UE::Mutable::Private
     	UE_API int32 GetConstantImageRomId(int32 ConstantImageIndex, int32 LODIndex) const;
 
 #if WITH_EDITOR
-        UE_API OP::CONSTANT_EXTENSIONDATA_INDEX AddConstant(const TManagedPtr<const FExtensionData>& Data);
+        UE_API FOperation::CONSTANT_EXTENSIONDATA_INDEX AddConstant(const TManagedPtr<const FExtensionData>& Data);
     	
-        UE_API OP::ADDRESS AddConstant(const TManagedPtr<const FLayout>& Layout);
+        UE_API FOperation::ADDRESS AddConstant(const TManagedPtr<const FLayout>& Layout);
     	
-        UE_API OP::ADDRESS AddConstant(const TManagedPtr<const FSkeleton>& Skeleton);
+        UE_API FOperation::ADDRESS AddConstant(const TManagedPtr<const FSkeleton>& Skeleton);
     	
-        UE_API OP::ADDRESS AddConstant(const TManagedPtr<const FPhysicsBody>& PhysicsBody);
+        UE_API FOperation::ADDRESS AddConstant(const TManagedPtr<const FPhysicsBody>& PhysicsBody);
 
         UE_API uint32 AddConstant(const FMeshSocket& Socket);
 
-        UE_API OP::CONSTANT_LIST_UINT32_INDEX AddConstant(const TArray<FMeshSocket>& Sockets);
+        UE_API FOperation::CONSTANT_LIST_UINT32_INDEX AddConstant(const TArray<FMeshSocket>& Sockets);
     	
-        UE_API OP::CONSTANT_STRING_INDEX AddConstant(const FString& Str);
+        UE_API FOperation::CONSTANT_STRING_INDEX AddConstant(const FString& Str);
     	
-        UE_API OP::ADDRESS AddConstant(const TArray<uint32>& UInt32List);
+        UE_API FOperation::ADDRESS AddConstant(const TArray<uint32>& UInt32List);
     	
-        UE_API  OP::ADDRESS AddConstant(const TArray<uint64>& UInt64List);
+        UE_API  FOperation::ADDRESS AddConstant(const TArray<uint64>& UInt64List);
     	
-        UE_API OP::ADDRESS AddConstant(const TArray<FString>& StringList);
+        UE_API FOperation::ADDRESS AddConstant(const TArray<FString>& StringList);
 
 		UE_API uint32 AddConstant(const FName& InName);
 
-		UE_API OP::ADDRESS AddConstant(const TArray<FName>& NameList);
+		UE_API FOperation::ADDRESS AddConstant(const TArray<FName>& NameList);
 
-        UE_API OP::CONSTANT_MATRIX_INDEX AddConstant(const FMatrix44f& Constant);
+        UE_API FOperation::CONSTANT_MATRIX_INDEX AddConstant(const FMatrix44f& Constant);
     	
-        UE_API OP::CONSTANT_SHAPE_INDEX AddConstant(const FShape& Constant);
+        UE_API FOperation::CONSTANT_SHAPE_INDEX AddConstant(const FShape& Constant);
     	
-        UE_API OP::ADDRESS AddConstant(const FProjector& Constant);
+        UE_API FOperation::ADDRESS AddConstant(const FProjector& Constant);
     	
-        UE_API OP::ADDRESS AddConstant(const FRichCurve& Constant);
+        UE_API FOperation::ADDRESS AddConstant(const FRichCurve& Constant);
     	
-        UE_API OP::ADDRESS AddConstant(const TManagedPtr<const FMaterial>& Material);
+        UE_API FOperation::ADDRESS AddConstant(const TManagedPtr<const FMaterial>& Material);
 #endif //WITH_EDITOR
     	
-        UE_API EOpType GetOpType(OP::ADDRESS At) const;
+        UE_API EOpType GetOpType(FOperation::ADDRESS At) const;
     	
         template<typename ARGS>
-		ARGS GetOpArgs(OP::ADDRESS at) const
+		ARGS GetOpArgs(FOperation::ADDRESS Address) const
 		{
-			ARGS result;
-			uint64 byteCodeAddress = OpAddress[at];
-			byteCodeAddress += sizeof(EOpType);
-			FMemory::Memcpy(&result, &ByteCode[byteCodeAddress], sizeof(ARGS));
-			return result;
+			ARGS Result;
+			
+			uint32 ArgsByteCodeOffset = GetAddressByteCodeOffset(Address) + sizeof(EOpType);
+			FMemory::Memcpy(&Result, &ByteCode[ArgsByteCodeOffset], sizeof(ARGS));
+			
+			return Result;
 		}
 
 		template<typename ARGS>
-		void SetOpArgs(OP::ADDRESS At, const ARGS& Args)
+		void SetOpArgs(FOperation::ADDRESS Address, const ARGS& Args)
 		{
-			uint64 byteCodeAddress = OpAddress[At];
-			byteCodeAddress += sizeof(EOpType);
-			FMemory::Memcpy(&ByteCode[byteCodeAddress], &Args, sizeof(ARGS));
+			uint32 ArgsByteCodeOffset = GetAddressByteCodeOffset(Address) + sizeof(EOpType);
+			FMemory::Memcpy(&ByteCode[ArgsByteCodeOffset], &Args, sizeof(ARGS));
 		}
 
-        UE_API const uint8* GetOpArgsPointer(OP::ADDRESS At) const;
+        UE_API const uint8* GetOpArgsPointer(FOperation::ADDRESS Address) const;
     	
-        uint8* GetOpArgsPointer(OP::ADDRESS At);
+        uint8* GetOpArgsPointer(FOperation::ADDRESS Address);
     };
 
 	
