[Mutable] Program cache rewrok.

diff --git a/Engine/Plugins/Mutable/Source/MutableRuntime/Internal/MuR/Operations.h b/Engine/Plugins/Mutable/Source/MutableRuntime/Internal/MuR/Operations.h
--- a/Engine/Plugins/Mutable/Source/MutableRuntime/Internal/MuR/Operations.h
+++ b/Engine/Plugins/Mutable/Source/MutableRuntime/Internal/MuR/Operations.h
@@ -21,393 +21,626 @@ namespace UE::Mutable::Private
 {
 	enum class EImageFormat : uint8;
 
-    //!
-    enum class EOpType : uint16
+    /** 
+     * Operation type enum. 
+     * All operations gerenrating a specific type **MUST** be contiguous and be between 
+     * its TYPE_RANGE_BEGIN and TYPE_RANGE_END
+     */
+    enum class EOpType : uint8
     {
         /** No operation. */
         NONE,
 
         //-----------------------------------------------------------------------------------------
-        // Generic operations
+        // Boolean operations
         //-----------------------------------------------------------------------------------------
-
-        //! Constant value
+        
+        /** Boolean constant */
         BO_CONSTANT,
-        NU_CONSTANT,
-        SC_CONSTANT,
-        CO_CONSTANT,
-        IM_CONSTANT,
-        ME_CONSTANT,
-        LA_CONSTANT,
-        PR_CONSTANT,
-        ST_CONSTANT,
-		ED_CONSTANT,
-    	MA_CONSTANT,
-		MI_CONSTANT,
+        
+        BO_RANGE_BEGIN = BO_CONSTANT,
 
-        //! User parameter
+        /** Boolean parameter */
         BO_PARAMETER,
-        NU_PARAMETER,
-        SC_PARAMETER,
-        CO_PARAMETER,
-        PR_PARAMETER,
-		IM_PARAMETER,
-		SK_PARAMETER,
-		ST_PARAMETER,
-    	MA_PARAMETER,
-    	MI_PARAMETER,
-    	IS_PARAMETER,
 
-		//! A referenced, but opaque engine resource
-		IM_REFERENCE,
-		ME_REFERENCE,
-
-        //! Select one value or the other depending on a boolean input
-        NU_CONDITIONAL,
-        SC_CONDITIONAL,
-        CO_CONDITIONAL,
-        IM_CONDITIONAL,
-        ME_CONDITIONAL,
-        LA_CONDITIONAL,
-        IN_CONDITIONAL,
-		ED_CONDITIONAL,
-		MI_CONDITIONAL,
-
-        //! Select one of several values depending on an int input
-        NU_SWITCH,
-        SC_SWITCH,
-        CO_SWITCH,
-        IM_SWITCH,
-        ME_SWITCH,
-        LA_SWITCH,
-        IN_SWITCH,
-		ED_SWITCH,
-		MI_SWITCH,
-
-		//-----------------------------------------------------------------------------------------
-	   // Material operations
-	   //-----------------------------------------------------------------------------------------
-
-		//! Selects a parameter value from a Material
-		SC_MATERIAL_BREAK,
-		CO_MATERIAL_BREAK,
-		IM_MATERIAL_BREAK,
-
-		//! Converts a texture of a material parameter into a texture parameter to process it at runtime.
-		IM_PARAMETER_FROM_MATERIAL,
-
-    	//! Modifies the parameter values of a material
-		MI_MODIFY,
-
-        //-----------------------------------------------------------------------------------------
-        // Boolean operations
-        //-----------------------------------------------------------------------------------------
-
-        //! Compare an integerexpression with an integer constant
+        /** Compare an integerexpression with an integer constant */
         BO_EQUAL_INT_CONST,
 
-        //! Logical and
+        /** Logical and */
         BO_AND,
 
-        //! Logical or
+        /** Logical or */
         BO_OR,
 
-        //! Left as an exercise to the reader to find out what this op does.
+        /** Logical not */
         BO_NOT,
+        
+
+        //-----------------------------------------------------------------------------------------
+        // Number operations
+        //-----------------------------------------------------------------------------------------
+        
+        /** Constant value */
+        NU_CONSTANT,
+        
+        BO_RANGE_END   = NU_CONSTANT,
+        NU_RANGE_BEGIN = BO_RANGE_END,
+        
+        /** User parameter */
+        NU_PARAMETER,
+
+        /** Select one value or the other depending on a boolean input */
+        NU_CONDITIONAL,
+        
+        /** Select one of several values depending on an int input */
+        NU_SWITCH,
+       
 
         //-----------------------------------------------------------------------------------------
         // Scalar operations
         //-----------------------------------------------------------------------------------------
+         
+        /** Constant value */
+        SC_CONSTANT,
+        
+        NU_RANGE_END   = SC_CONSTANT,
+        SC_RANGE_BEGIN = NU_RANGE_END,
 
-        //! Apply an arithmetic operation to two scalars
+        /** User parameter */
+        SC_PARAMETER,
+        
+        /** Select one value or the other depending on a boolean input */
+        SC_CONDITIONAL,
+        
+        /** Select one of several values depending on an int input */
+        SC_SWITCH,
+        
+        SC_MATERIAL_BREAK,
+        
+        /** Apply an arithmetic operation to two scalars */
         SC_ARITHMETIC,
 
-        //! Get a scalar value from a curve
+        /** Get a scalar value from a curve */
         SC_CURVE,
-
-    	//! External operation
+		
+    	/** External operation */
 		SC_EXTERNAL,
 
+        
         //-----------------------------------------------------------------------------------------
         // Colour operations. Colours are sometimes used as generic vectors.
         //-----------------------------------------------------------------------------------------
+        
+        /** Constant value */
+        CO_CONSTANT,
+
+        SC_RANGE_END   = CO_CONSTANT,
+        CO_RANGE_BEGIN = SC_RANGE_END,
 
-        //! Sample an image to get its colour.
+        /** User parameter */
+        CO_PARAMETER,
+        
+        /** Select one value or the other depending on a boolean input */
+        CO_CONDITIONAL,
+        
+        /** Select one of several values depending on an int input */
+        CO_SWITCH,
+        
+        CO_MATERIAL_BREAK,
+        
+        /** Sample an image to get its color */
         CO_SAMPLEIMAGE,
 
-        //! Make a color by shuffling channels from other colours.
+        /** Make a color by shuffling channels from other color */
         CO_SWIZZLE,
 
-        //! Compose a vector from 4 scalars
+        /** Compose a vector from 4 scalars */
         CO_FROMSCALARS,
 
-        //! Apply component-wise arithmetic operations to two colours
+        /** Apply component-wise arithmetic operations to two colours */
         CO_ARITHMETIC,
 
-		//! Apply a Linear to sRGB color transformation on a given color vector.
+		/** Apply a Linear to sRGB color transformation on a given color vector */
 		CO_LINEARTOSRGB,
-
-    	//! External operation
+		
+    	/** External operation */
 		CO_EXTERNAL,
-    	
+
+        //-----------------------------------------------------------------------------------------
+        // String operations
+        //-----------------------------------------------------------------------------------------
+
+        /** Constant value */
+        ST_CONSTANT,
+        
+        CO_RANGE_END   = ST_CONSTANT,
+        ST_RANGE_BEGIN = CO_RANGE_END,
+        
+		ST_PARAMETER,
+      
+
+        //-----------------------------------------------------------------------------------------
+        // Projector operations
+        //-----------------------------------------------------------------------------------------
+        
+        /** Constant value */
+        PR_CONSTANT, 
+
+        ST_RANGE_END   = PR_CONSTANT,
+        PR_RANGE_BEGIN = ST_RANGE_END,
+        
+		PR_PARAMETER,
+
+
         //-----------------------------------------------------------------------------------------
         // Image operations
         //-----------------------------------------------------------------------------------------
 
-        //! Combine an image on top of another one using a specific effect (Blend, SoftLight, 
-		//! Hardlight, Burn...). And optionally a mask.
+        IM_CONSTANT,
+        
+        PR_RANGE_END   = IM_CONSTANT,
+        IM_RANGE_BEGIN = PR_RANGE_END,
+        
+        /** User parameter */
+		IM_PARAMETER,
+        
+		/** A referenced, but opaque engine image */
+		IM_REFERENCE,
+        
+        /** Select one value or the other depending on a boolean input */
+        IM_CONDITIONAL,
+
+        IM_SWITCH,
+        
+        IM_MATERIAL_BREAK,
+
+		/** Converts a texture of a material parameter into a texture parameter to process it at runtime. */
+		IM_PARAMETER_FROM_MATERIAL,
+        
+        /** Combine an image on top of another one using a specific effect (Blend, SoftLight,
+	    Hardlight, Burn...). And optionally a mask. */
         IM_LAYER,
 
-        //! Apply a colour on top of an image using a specific effect (Blend, SoftLight, 
-		//! Hardlight, Burn...), optionally using a mask.
+        /** Apply a colour on top of an image using a specific effect (Blend, SoftLight,
+		Hardlight, Burn...), optionally using a mask. */
         IM_LAYERCOLOUR,        
 
-        //! Convert between pixel formats
+        /** Convert between pixel formats */
         IM_PIXELFORMAT,
 
-        //! Generate mipmaps up to a provided level
+        /** Generate mipmaps up to a provided level */
         IM_MIPMAP,
 
-        //! Resize the image to a constant size
+        /** Resize the image to a constant size */
         IM_RESIZE,
 
-        //! Resize the image to the size of another image
+        /** Resize the image to the size of another image */
         IM_RESIZELIKE,
 
-        //! Resize the image by a relative factor
+        /** Resize the image by a relative factor */
         IM_RESIZEREL,
 
-        //! Create an empty image to hold a particular layout.
+        /** Create an empty image to hold a particular layout. */
         IM_BLANKLAYOUT,
 
-        //! Copy an image into a rect of another one.
+        /** Copy an image into a rect of another one. */
         IM_COMPOSE,
 
-		//! Copy blocks from an image into a another one.
+		/** Copy blocks from an image into a another one. */
 		IM_MULTICOMPOSE,
 
-        //! Interpolate between 2 images taken from a row of targets (2 consecutive targets).
+        /** Interpolate between 2 images taken from a row of targets (2 consecutive targets). */
         IM_INTERPOLATE,
 
-        //! Change the saturation of the image.
+        /** Change the saturation of the image. */
         IM_SATURATE,
 
-        //! Generate a one-channel image with the luminance of the source image.
+        /** Generate a one-channel image with the luminance of the source image. */
         IM_LUMINANCE,
 
-        //! Recombine the channels of several images into one.
+        /** Recombine the channels of several images into one. */
         IM_SWIZZLE,
 
-        //! Convert the source image colours using a "palette" image sampled with the source
-        //! grey-level.
+        /** Convert the source image colours using a "palette" image sampled with the source */
+        /** grey-level. */
         IM_COLOURMAP,
 
-        //! Generate a black and white image from an image and a threshold.
+        /** Generate a black and white image from an image and a threshold. */
         IM_BINARISE,
 
-        //! Generate a plain colour image
+        /** Generate a plain colour image */
         IM_PLAINCOLOUR,
 
-        //! Cut a rect from an image
+        /** Cut a rect from an image */
         IM_CROP,
 
-        //! Replace a subrect of an image with another one
+        /** Replace a subrect of an image with another one */
         IM_PATCH,
 
-        //! Render a mesh texture layout into a mask
+        /** Render a mesh texture layout into a mask */
         IM_RASTERMESH,
 
-        //! Create an image displacement encoding the grow operation for a mask
+        /** Create an image displacement encoding the grow operation for a mask */
         IM_MAKEGROWMAP,
 
-        //! Apply an image displacement on another image.
+        /** Apply an image displacement on another image. */
         IM_DISPLACE,
 
-        //! Repeately apply
+        /** Repeately apply */
         IM_MULTILAYER,
 
-        //! Inverts the colors of an image
+        /** Inverts the colors of an image */
         IM_INVERT,
 
-        //! Modifiy roughness channel of an image based on normal variance.
+        /** Modifiy roughness channel of an image based on normal variance. */
         IM_NORMALCOMPOSITE,
 
-		//! Apply linear transform to Image content. Resulting samples outside the original image are tiled.
+		/** Apply linear transform to Image content. Resulting samples outside the original image are tiled. */
 		IM_TRANSFORM,
 
         /** Convert an FIamge in Passthrough Parameter mode to an FImage with data. */
     	IM_PARAMETER_CONVERT,
-
-    	//! External operation
+			
+        /** External operation */
 		IM_EXTERNAL,
-    	
+    
         //-----------------------------------------------------------------------------------------
         // Mesh operations
         //-----------------------------------------------------------------------------------------
 
-        //! Apply a layout to a mesh texture coordinates channel
+        /** Constant value */
+        ME_CONSTANT,
+       
+        IM_RANGE_END   = ME_CONSTANT,
+        ME_RANGE_BEGIN = IM_RANGE_END, 
+
+        /** User parameter */
+        ME_PARAMETER,
+
+		ME_REFERENCE,
+        
+        /** Select one value or the other depending on a boolean input */
+        ME_CONDITIONAL,
+        
+        /** Select one of several values depending on an int input */
+        ME_SWITCH,
+        
+        /** Apply a layout to a mesh texture coordinates channel */
         ME_APPLYLAYOUT,
 
 		/** */
 		ME_PREPARELAYOUT,
 
-        //! Compare two meshes and extract a morph from the first to the second
-        //! The meshes must have the same topology, etc.
+        /** Compare two meshes and extract a morph from the first to the second
+        The meshes must have the same topology, etc. */
         ME_DIFFERENCE,
 
-        //! Apply a one morphs on a base. 
+        /** Apply a one morphs on a base.  */
         ME_MORPH,
 
-        //! Merge a mesh to a mesh
+        /** Merge a mesh to a mesh */
         ME_MERGE,
 
-        //! Create a new mask mesh selecting all the faces of a source that are inside a given
-        //! clip mesh.
+        /** Create a new mask mesh selecting all the faces of a source that are inside a given
+        clip mesh. */
         ME_MASKCLIPMESH,
 
-        /** Create a new mask mesh selecting the faces of a source that have UVs inside the region marked in an image mask. */
+        /** Create a new mask mesh selecting the faces of a source that have UVs inside the 
+        region marked in an image mask. */
 		ME_MASKCLIPUVMASK,
 
-        //! Create a new mask mesh selecting all the faces of a source that match another mesh.
+        /** Create a new mask mesh selecting all the faces of a source that match another mesh. */
         ME_MASKDIFF,
 
-        //! Remove all the geometry selected by a mask.
+        /** Remove all the geometry selected by a mask. */
         ME_REMOVEMASK,
 
-        //! Change the mesh format to match the format of another one.
+        /** Change the mesh format to match the format of another one. */
         ME_FORMAT,
 
-        //! Extract a fragment of a mesh containing specific layout blocks.
+        /** Extract a fragment of a mesh containing specific layout blocks. */
         ME_EXTRACTLAYOUTBLOCK,
 
-        //! Apply a transform in a 4x4 matrix to the geometry channels of the mesh
+        /** Apply a transform in a 4x4 matrix to the geometry channels of the mesh */
         ME_TRANSFORM,
 
-        //! Clip the mesh with a plane and morph it when it is near until it becomes an ellipse on
-        //! the plane.
+        /** Clip the mesh with a plane and morph it when it is near until it becomes an ellipse on the plane. */
         ME_CLIPMORPHPLANE,
 
-        //! Clip the mesh with another mesh.
+        /** Clip the mesh with another mesh. */
         ME_CLIPWITHMESH,
 
-        //! Project a mesh using a projector and clipping the irrelevant faces
+        /** Project a mesh using a projector and clipping the irrelevant faces */
         ME_PROJECT,
 
-        //! Deform a skinned mesh applying a skeletal pose
+        /** Deform a skinned mesh applying a skeletal pose */
         ME_APPLYPOSE,
 
-		//! Calculate the binding of a mesh on a shape
+		/** Calculate the binding of a mesh on a shape */
 		ME_BINDSHAPE,
 
-		//! Apply a shape on a (previously bound) mesh
+		/** Apply a shape on a (previously bound) mesh */
 		ME_APPLYSHAPE,
 
-		//! Clip Deform using bind data.
+		/** Clip Deform using bind data. */
 		ME_CLIPDEFORM,
 	
-        //! Mesh morph with Skeleton Reshape based on the morphed mesh.
+        /** Mesh morph with Skeleton Reshape based on the morphed mesh. */
         ME_MORPHRESHAPE,
 
-		//! Optimize skinning before adding a mesh to the component
+		/** Optimize skinning before adding a mesh to the component */
 		ME_OPTIMIZESKINNING,
 
-		//! Add a metadata to a mesh
+		/** Add a metadata to a mesh */
 		ME_ADDMETADATA,
 
-    	//! Transform with a 4x4 matrix the geometry channels of a mesh that are bounded by another mesh
+    	/** Transform with a 4x4 matrix the geometry channels of a mesh that are bounded by another mesh */
     	ME_TRANSFORMWITHMESH,
 
-		//! Transform with a 4x4 matrix the geometry channels of a mesh that are skinned to a bone or hierarchy
+		/** Transform with a 4x4 matrix the geometry channels of a mesh that are skinned to a bone or hierarchy */
 		ME_TRANSFORMWITHBONE,
 
-    	//! External operation
+    	/** External operation */
     	ME_EXTERNAL,
 
         /** Select a Mesh Section from a Skeletal Mesh. */
     	ME_SKELETALMESH_BREAK,
 
+        //-----------------------------------------------------------------------------------------
+        // SkeletalMesh operations
+        //-----------------------------------------------------------------------------------------
+        
+
+        /** User parameter */
+		SK_PARAMETER,
+        
+        ME_RANGE_END   = SK_PARAMETER,
+        SK_RANGE_BEGIN = ME_RANGE_END,
+
+
         //-----------------------------------------------------------------------------------------
         // Instance operations
         //-----------------------------------------------------------------------------------------
 
-        //! Add a mesh to an instance
+        /** Select one value or the other depending on a boolean input */
+        IN_CONDITIONAL,
+
+        SK_RANGE_END   = IN_CONDITIONAL,
+        IN_RANGE_BEGIN = SK_RANGE_END,
+        
+        /** Select one of several values depending on an int input */
+        IN_SWITCH,
+        
+        /** Add a mesh to an instance */
         IN_ADDMESH,
 
-        //! Add a string to an instance
+        /** Add a string to an instance */
         IN_ADDSTRING,
 
-        //! Add a surface to an instance component
+        /** Add a surface to an instance component */
         IN_ADDSURFACE,
 
-        //! Add a component to an instance LOD
+        /** Add a component to an instance LOD */
         IN_ADDCOMPONENT,
 
-        //! Add all LODs to an instance. This operation can only appear once in a model.
+        /** Add all LODs to an instance. This operation can only appear once in a model. */
         IN_ADDLOD,
 
-    	//! Add a Skeletal Mesh to an instance.
+    	/** Add a Skeletal Mesh to an instance. */
 		IN_ADDSKELETALMESH,
 
-		//! Add extension data to an instance
+		/** Add extension data to an instance */
 		IN_ADDEXTENSIONDATA,
 
-		//! Add overlay material to an instance
+		/** Add overlay material to an instance */
 		IN_ADDOVERLAYMATERIAL,
 
-		//! Add override material to an instance
+		/** Add override material to an instance */
 		IN_ADDOVERRIDEMATERIAL,
 
-		//! Add a material to an instance
+		/** Add a material to an instance */
 		IN_ADDMATERIAL,
 
+
         //-----------------------------------------------------------------------------------------
         // Layout operations
         //-----------------------------------------------------------------------------------------
 
-        //! Pack all the layout blocks from the source in the grid without overlapping
+        /** Constant value */
+        LA_CONSTANT,
+
+        IN_RANGE_END   = LA_CONSTANT,
+        LA_RANGE_BEGIN = IN_RANGE_END,
+        
+        /** Select one value or the other depending on a boolean input */
+        LA_CONDITIONAL,
+        
+        /** Select one of several values depending on an int input */
+        LA_SWITCH,
+        
+        /** Pack all the layout blocks from the source in the grid without overlapping */
         LA_PACK,
 
-        //! Merge two layouts
+        /** Merge two layouts */
         LA_MERGE,
 
-        //! Remove all layout blocks not used by any vertex of the mesh.
-        //! This operation is for the new way of managing layout blocks.
+        /** 
+         * Remove all layout blocks not used by any vertex of the mesh. 
+         * This operation is for the new way of managing layout blocks. 
+         */
         LA_REMOVEBLOCKS,
 
-		//! Extract a layout from a mesh
+		/** Extract a layout from a mesh */
 		LA_FROMMESH,
 
     	//-----------------------------------------------------------------------------------------
 	    // Material operations
 	    //-----------------------------------------------------------------------------------------
 
-    	//! External operation
-		MI_EXTERNAL,
+        /** Constant value */
+		MI_CONSTANT,
+
+        LA_RANGE_END   = MI_CONSTANT,
+        MI_RANGE_BEGIN = LA_RANGE_END,
 
-    	// Get a material from a given skeletal mesh
+        /** User parameter */
+    	MI_PARAMETER,
+        
+        /** Select one value or the other depending on a boolean input */
+		MI_CONDITIONAL,
+        
+        /** Select one of several values depending on an int input */
+		MI_SWITCH,
+    	
+        /** Get a material from a given skeletal mesh */
 		MI_SKELETALMESH_BREAK,
+
+        /** Get a material from a given skeletal mesh slot */
+    	MI_FROM_SKELETALMESH_SLOT,
+
+        /** Modifies the parameter values of a material */
+		MI_MODIFY,
+
+        /** External operation */
+		MI_EXTERNAL,
+
     	
     	//-----------------------------------------------------------------------------------------
+	    // Matrix operations
+	    //-----------------------------------------------------------------------------------------
+
+        /** User parameter */
+    	MA_CONSTANT,
+
+        MI_RANGE_END   = MA_CONSTANT,
+        MA_RANGE_BEGIN = MI_RANGE_END,
+
+        /** User parameter */
+    	MA_PARAMETER,
+   
+    	//-----------------------------------------------------------------------------------------
+	    // ExtensionData operations
+	    //-----------------------------------------------------------------------------------------
+
+		ED_CONSTANT,
+
+	    MA_RANGE_END   = ED_CONSTANT,
+        ED_RANGE_BEGIN = MA_RANGE_END,
+
+        /** Select one value or the other depending on a boolean input */
+		ED_CONDITIONAL,
+        
+        /** Select one of several values depending on an int input */
+        ED_SWITCH,
+
+        //-----------------------------------------------------------------------------------------
 		// FInstancedStruct operations
 		//-----------------------------------------------------------------------------------------
 
-		//! External operation
+        /** User parameter */
+    	IS_PARAMETER,
+
+        ED_RANGE_END   = IS_PARAMETER,
+        IS_RANGE_BEGIN = ED_RANGE_END,
+
+        /** External operation */
 		IS_EXTERNAL,
-    	
+    
+        IS_RANGE_END,
+
         //-----------------------------------------------------------------------------------------
         // Utility values
         //-----------------------------------------------------------------------------------------
 
-        //!
-        COUNT
+        /** */
+        COUNT = IS_RANGE_END, 
     };
-	
-	
+
 	inline uint32 GetTypeHash(const EOpType& OpType)
 	{
 		return static_cast<uint32>(OpType);
 	}
 
+	/** Types of data handled by the Mutable runtime. */
+	enum class EDataType : uint8
+	{
+		None = 0,
+		Bool,
+		Int,
+		Scalar,
+		Color,
+		Image,
+		Layout,
+		Mesh,
+		Instance,
+		Projector,
+		String,
+		ExtensionData,
+		Matrix,
+		Material,
+		InstancedStruct,
+		SkeletalMesh,
+
+		// Supporting data types : Never returned as an actual data type for any operation.
+		Shape,
+		Curve,
+		Skeleton,
+		PhysicsAsset,
+
+		Count
+	};
+
+    /** Computes the data type of a given op type */
+    FORCEINLINE EDataType GetOpDataType(EOpType OpType)
+    { 
+        uint8 OpBo = static_cast<uint8>(((OpType >= EOpType::BO_RANGE_BEGIN) & (OpType < EOpType::BO_RANGE_END)) 
+                ? EDataType::Bool : EDataType::None); 
+                                                                                                                            
+        uint8 OpNu = static_cast<uint8>(((OpType >= EOpType::NU_RANGE_BEGIN) & (OpType < EOpType::NU_RANGE_END)) 
+                ? EDataType::Int : EDataType::None); 
+        
+        uint8 OpSc = static_cast<uint8>(((OpType >= EOpType::SC_RANGE_BEGIN) & (OpType < EOpType::SC_RANGE_END)) 
+                ? EDataType::Scalar : EDataType::None);
+                                                                                                                            
+        uint8 OpCo = static_cast<uint8>(((OpType >= EOpType::CO_RANGE_BEGIN) & (OpType < EOpType::CO_RANGE_END)) 
+                ? EDataType::Color : EDataType::None);
+                                                                                                                            
+        uint8 OpLa = static_cast<uint8>(((OpType >= EOpType::LA_RANGE_BEGIN) & (OpType < EOpType::LA_RANGE_END)) 
+                ? EDataType::Layout : EDataType::None);
+        
+        uint8 OpIm = static_cast<uint8>(((OpType >= EOpType::IM_RANGE_BEGIN) & (OpType < EOpType::IM_RANGE_END)) 
+                ? EDataType::Image : EDataType::None);
+                                                                                                                            
+        uint8 OpMe = static_cast<uint8>(((OpType >= EOpType::ME_RANGE_BEGIN) & (OpType < EOpType::ME_RANGE_END)) 
+                ? EDataType::Mesh : EDataType::None);
+
+        uint8 OpIn = static_cast<uint8>(((OpType >= EOpType::IN_RANGE_BEGIN) & (OpType < EOpType::IN_RANGE_END)) 
+                ? EDataType::Instance : EDataType::None);
+
+        uint8 OpPr = static_cast<uint8>(((OpType >= EOpType::PR_RANGE_BEGIN) & (OpType < EOpType::PR_RANGE_END)) 
+                ? EDataType::Projector : EDataType::None);
+
+        uint8 OpSt = static_cast<uint8>(((OpType >= EOpType::ST_RANGE_BEGIN) & (OpType < EOpType::ST_RANGE_END)) 
+                ? EDataType::String : EDataType::None);
+
+        uint8 OpEd = static_cast<uint8>(((OpType >= EOpType::ED_RANGE_BEGIN) & (OpType < EOpType::ED_RANGE_END)) 
+                ? EDataType::ExtensionData : EDataType::None);
+
+        uint8 OpMa = static_cast<uint8>(((OpType >= EOpType::MA_RANGE_BEGIN) & (OpType < EOpType::MA_RANGE_END)) 
+                ? EDataType::Matrix : EDataType::None);
+        
+        uint8 OpMi = static_cast<uint8>(((OpType >= EOpType::MI_RANGE_BEGIN) & (OpType < EOpType::MI_RANGE_END)) 
+                ? EDataType::Material : EDataType::None);
+
+        uint8 OpIs = static_cast<uint8>(((OpType >= EOpType::IS_RANGE_BEGIN) & (OpType < EOpType::IS_RANGE_END)) 
+                ? EDataType::InstancedStruct : EDataType::None);
+
+        uint8 OpSk = static_cast<uint8>(((OpType >= EOpType::SK_RANGE_BEGIN) & (OpType < EOpType::SK_RANGE_END)) 
+                ? EDataType::SkeletalMesh : EDataType::None);
+
+        return static_cast<EDataType>(
+                OpBo | OpNu | OpSc | OpCo | OpLa | OpIm | OpMe | OpIn | OpPr | OpSt | OpEd | OpMa | OpMi | OpIs | OpSk);
+    }
 
 	enum class EMeshBindShapeFlags : uint32
 	{
@@ -438,11 +671,11 @@ namespace UE::Mutable::Private
 
 	static_assert(sizeof(FMeshBindColorChannelUsages) == sizeof(uint32));
 
-
     /** */
-    template<typename ADDRESS_TYPE>
-    struct TOperation
+    struct FOperation
     {
+        using ADDRESS_TYPE = uint32;
+
         typedef ADDRESS_TYPE ADDRESS;
         typedef ADDRESS_TYPE PARAMETER;
         typedef ADDRESS_TYPE CONSTANT_LIST_PASSTHROUGH_ID;
@@ -453,6 +686,9 @@ namespace UE::Mutable::Private
     	typedef ADDRESS_TYPE CONSTANT_STRING_INDEX;
 		typedef ADDRESS_TYPE CONSTANT_LIST_UINT32_INDEX;
 
+        static constexpr uint32 AddressByteCodeOffsetBits = 26;
+        static constexpr uint32 MaxAddressByteCodeOffset = (1 << AddressByteCodeOffsetBits) - 1;
+
         /** Arguments for every operation type. */
         struct BoolConstantArgs
         {
@@ -1279,61 +1515,37 @@ namespace UE::Mutable::Private
     };
 
 
-    typedef TOperation<uint32> OP;
-
-    ENUM_CLASS_FLAGS(OP::MeshAddMetadataArgs::EnumFlags);
-
-	static_assert(sizeof(OP::FSwitchCaseDescriptor) == sizeof(uint32));
+    ENUM_CLASS_FLAGS(FOperation::MeshAddMetadataArgs::EnumFlags);
 
+	static_assert(sizeof(FOperation::FSwitchCaseDescriptor) == sizeof(uint32));
 
-	//! Types of data handled by the Mutable runtime.
-	enum class EDataType : uint8
-	{
-		None,
-		Bool,
-		Int,
-		Scalar,
-		Color,
-		Image,
-		Layout,
-		Mesh,
-		Instance,
-		Projector,
-		String,
-		ExtensionData,
-		Matrix,
-		Material,
-		InstancedStruct,
-		SkeletalMesh,
-
-		// Supporting data types : Never returned as an actual data type for any operation.
-		Shape,
-		Curve,
-		Skeleton,
-		PhysicsAsset,
+    FORCEINLINE FOperation::ADDRESS MakeProgramAddress(EOpType OpType, uint32 ByteCodeOffset)
+    {
+        check(ByteCodeOffset <= FOperation::MaxAddressByteCodeOffset);
+        static_assert(1 << (32 - FOperation::AddressByteCodeOffsetBits) > (int32)EDataType::Count);
 
-		Count
-	};
+        return 
+            (uint32(GetOpDataType(OpType)) << FOperation::AddressByteCodeOffsetBits) | 
+            (ByteCodeOffset & FOperation::MaxAddressByteCodeOffset);
+    } 
 
+   // FORCEINLINE EOpType GetAddressOpType(FOperation::ADDRESS At)
+   // {
+   //     return static_cast<EOpType>(At >> FOperation::AddressByteCodeOffsetBits);
+   // }
 
-	// Generic data about a Mutable operation that is needed at runtime.
-    struct FOpDesc
+    FORCEINLINE EDataType GetAddressDataType(FOperation::ADDRESS At)
     {
-		//! Type of data generated by the instruction
-		EDataType type;
-    };
-
-	MUTABLERUNTIME_API extern const FOpDesc& GetOpDesc( EOpType type );
-
+        return static_cast<EDataType>(At >> FOperation::AddressByteCodeOffsetBits);
+    }
 
-    //!
-    inline static EDataType GetOpDataType( EOpType type )
+    FORCEINLINE uint32 GetAddressByteCodeOffset(FOperation::ADDRESS At)
     {
-		return GetOpDesc(type).type;
+        return At & FOperation::MaxAddressByteCodeOffset;
     }
 
     //! Utility function to apply a function to all operation references to other operations.
-	MUTABLERUNTIME_API extern void ForEachReference( const struct FProgram& program, OP::ADDRESS at, const TFunctionRef<void(OP::ADDRESS)> );
+	MUTABLERUNTIME_API extern void ForEachReference( const struct FProgram& Program, FOperation::ADDRESS At, const TFunctionRef<void(FOperation::ADDRESS)> );
 
 	//!
 	inline EOpType GetSwitchForType( EDataType d )
diff --git a/Engine/Plugins/Mutable/Source/MutableRuntime/Internal/MuR/ResourceID.h b/Engine/Plugins/Mutable/Source/MutableRuntime/Internal/MuR/ResourceID.h
--- a/Engine/Plugins/Mutable/Source/MutableRuntime/Internal/MuR/ResourceID.h
+++ b/Engine/Plugins/Mutable/Source/MutableRuntime/Internal/MuR/ResourceID.h
@@ -12,7 +12,7 @@ namespace UE::Mutable::Private
 	// Mesh
 	struct FGeneratedMeshKey
 	{
-		OP::ADDRESS Address = 0;
+		FOperation::ADDRESS Address = 0;
 		
 		TMemoryTrackedArray<uint8> ParameterValuesBlob;
 		
@@ -34,7 +34,7 @@ namespace UE::Mutable::Private
 	// Image
 	struct FGeneratedImageKey
 	{
-		OP::ADDRESS Address = 0;
+		FOperation::ADDRESS Address = 0;
 
 		TMemoryTrackedArray<uint8> ParameterValuesBlob;
 
@@ -56,7 +56,7 @@ namespace UE::Mutable::Private
 	// Material
 	struct FGeneratedMaterialKey
 	{
-		OP::ADDRESS Address = 0;
+		FOperation::ADDRESS Address = 0;
 
 		TMemoryTrackedArray<uint8> ParameterValuesBlob;
 
@@ -79,7 +79,7 @@ namespace UE::Mutable::Private
 	// SkeletalMesh
 	struct FGeneratedSkeletalMeshKey
 	{
-		OP::ADDRESS Address = 0;
+		FOperation::ADDRESS Address = 0;
 
 		TMemoryTrackedArray<uint8> ParameterValuesBlob;
 
diff --git a/Engine/Plugins/Mutable/Source/MutableRuntime/Internal/MuR/System.h b/Engine/Plugins/Mutable/Source/MutableRuntime/Internal/MuR/System.h
--- a/Engine/Plugins/Mutable/Source/MutableRuntime/Internal/MuR/System.h
+++ b/Engine/Plugins/Mutable/Source/MutableRuntime/Internal/MuR/System.h
@@ -17,6 +17,7 @@
 #include "MuR/MutableTrace.h"
 #include "MuR/ManagedPointer.h"
 #include "MuR/MutableEditorLogger.h"
+#include "MuR/ProgramCache.h"
 
 #include "SkeletalMesh.h"
 #include "Tasks/Task.h"
@@ -128,1477 +129,6 @@ namespace UE::Mutable::Private
 	* array of ranges, and the second integer is the value inside that range.
 	* The vector order is undefined.
 	*/
-	class ExecutionIndex : public  TArray<TPair<int32, int32>>
-	{
-	public:
-		//! Set or add a value to the index
-		void SetFromModelRangeIndex(uint16 FRangeIndex, int32 RangeValue)
-		{
-			SizeType Index = IndexOfByPredicate([=](const ElementType& v) { return v.Key >= FRangeIndex; });
-			if (Index != INDEX_NONE && (*this)[Index].Key == FRangeIndex)
-			{
-				// Update
-				(*this)[Index].Value = RangeValue;
-			}
-			else
-			{
-				// Add new
-				Push(ElementType(FRangeIndex, RangeValue));
-			}
-		}
-
-		//! Get the value of the index from the range index in the model.
-		int32 GetFromModelRangeIndex(int32 ModelRangeIndex) const
-		{
-			for (const TPair<int32, int32>& e : *this)
-			{
-				if (e.Key == ModelRangeIndex)
-				{
-					return e.Value;
-				}
-			}
-			return 0;
-		}
-	};
-
-	//! Container that stores data per executable code operation (indexed by address and execution
-	//! index).
-	template<class DATA>
-	class CodeContainer
-	{
-		using MemoryCounter = MemoryCounters::FInternalMemoryCounter;
-		using ArrayDataContainerType = TArray<DATA, FDefaultMemoryTrackingAllocator<MemoryCounter>>;
-		using MapDataContainerType = TMap<FCacheAddress, DATA, FDefaultMemoryTrackingSetAllocator<MemoryCounter>>;
-
-	public:
-		void resize(size_t s)
-		{
-			m_index0.SetNumZeroed(s);
-		}
-
-		uint32 size_code() const
-		{
-			return uint32(m_index0.Num());
-		}
-
-		void clear()
-		{
-			m_index0.Empty();
-			m_otherIndex.Empty();
-		}
-
-		inline void erase(const FCacheAddress& at)
-		{
-			if (at.ExecutionIndex == 0 && at.ExecutionOptions == 0 && at.Type == FScheduledOp::EType::Full)
-			{
-				m_index0[at.At] = nullptr;
-			}
-			else
-			{
-				m_otherIndex.Remove(at);
-			}
-		}
-
-		inline DATA get(const FCacheAddress& at) const
-		{
-			if (at.ExecutionIndex == 0 && at.ExecutionOptions == 0 && at.Type == FScheduledOp::EType::Full)
-			{
-				if (at.At < uint32(m_index0.Num()))
-				{
-					return m_index0[at.At];
-				}
-				else
-				{
-					return 0;
-				}
-			}
-			else
-			{
-				const DATA* it = m_otherIndex.Find(at);
-				if (it)
-				{
-					return *it;
-				}
-			}
-			return 0;
-		}
-
-		inline DATA* get_ptr(const FCacheAddress& at)
-		{
-			if (at.ExecutionIndex == 0 && at.ExecutionOptions == 0 && at.Type == FScheduledOp::EType::Full)
-			{
-				if (at.At < uint32(m_index0.Num()))
-				{
-					return &m_index0[at.At];
-				}
-				else
-				{
-					return nullptr;
-				}
-			}
-			else
-			{
-				DATA* it = m_otherIndex.Find(at);
-				if (it)
-				{
-					return it;
-				}
-			}
-			return nullptr;
-		}
-
-		inline const DATA* get_ptr(const FCacheAddress& at) const
-		{
-			if (at.ExecutionIndex == 0 && at.ExecutionOptions == 0 && at.Type == FScheduledOp::EType::Full)
-			{
-				if (at.At < uint32(m_index0.Num()))
-				{
-					return &m_index0[at.At];
-				}
-				else
-				{
-					return nullptr;
-				}
-			}
-			else
-			{
-				const DATA* it = m_otherIndex.Find(at);
-				if (it)
-				{
-					return it;
-				}
-			}
-			return nullptr;
-		}
-
-		inline DATA& operator[](const FCacheAddress& at)
-		{
-			if (at.ExecutionIndex == 0 && at.ExecutionOptions == 0 && at.Type == FScheduledOp::EType::Full)
-			{
-				return m_index0[at.At];
-			}
-			else
-			{
-				return m_otherIndex.FindOrAdd(at);
-			}
-		}
-
-		inline const DATA& operator[](const FCacheAddress& at) const
-		{
-			if (at.ExecutionIndex == 0 && at.ExecutionOptions == 0 && at.Type == FScheduledOp::EType::Full)
-			{
-				return m_index0[at.At];
-			}
-			else
-			{
-				return m_otherIndex[at];
-			}
-		}
-
-		struct iterator
-		{
-		private:
-			friend class CodeContainer<DATA>;
-
-			const CodeContainer<DATA>* container;
-			typename CodeContainer<DATA>::ArrayDataContainerType::TIterator it0;
-			typename CodeContainer<DATA>::MapDataContainerType::TIterator it1;
-
-			iterator(CodeContainer<DATA>* InContainer)
-				: container(InContainer)
-				, it0(InContainer->m_index0.CreateIterator())
-				, it1(InContainer->m_otherIndex.CreateIterator())
-			{
-			}
-
-		public:
-
-			inline void operator++(int)
-			{
-				if (it0)
-				{
-					++it0;
-				}
-				else
-				{
-					++it1;
-				}
-			}
-
-			iterator operator++()
-			{
-				if (it0)
-				{
-					++it0;
-				}
-				else
-				{
-					++it1;
-				}
-				return *this;
-			}
-
-			inline bool operator!=(const iterator& o) const
-			{
-				return it0 != o.it0 || it1 != o.it1;
-			}
-
-			inline DATA& operator*()
-			{
-				if (it0)
-				{
-					return *it0;
-				}
-				else
-				{
-					return it1->Value;
-				}
-			}
-
-			inline FCacheAddress get_address() const
-			{
-				if (it0)
-				{
-					return { uint32_t(it0.GetIndex()), 0, 0};
-				}
-				else
-				{
-					return it1->Key;
-				}
-			}
-
-			inline bool IsValid() const
-			{
-				return bool(it0) || bool(it1);
-			}
-		};
-
-		inline iterator begin()
-		{
-			iterator it(this);
-			return it;
-		}
-
-		inline int32 GetAllocatedSize() const
-		{
-			return m_index0.GetAllocatedSize()
-				+ m_otherIndex.GetAllocatedSize();
-		}
-	private:
-		// For index 0
-		ArrayDataContainerType m_index0;
-
-		// For index>0
-		MapDataContainerType m_otherIndex;
-	};
-
-
-	/** Interface for storage of data while Mutable code is being executed. */
-	class FProgramCache
-	{
-	public:
-		TMemoryTrackedArray<ExecutionIndex, TInlineAllocator<4, AllocType>> m_usedRangeIndices;
-
-		/** Runtime data for each program op. */
-		struct FOpExecutionData
-		{
-			uint16 OpHitCount;
-
-			/** Enabled if the op descriptor has been calculated. */
-			uint8 IsDescCacheValid : 1;
-			uint8 IsValueValid : 1;
-			uint8 IsCacheLocked : 1;
-
-			/** The operation DATA_TYPE. */
-			EDataType DataType;
-
-			/** The position in the type-specific array where the result data is stored. 
-			* 0 means index not valid.
-			* For small types like bool, int and float, the actual value is the index, instead of having their own array.
-			*/
-			union
-			{
-				int32 DataTypeIndex;
-				float ScalarResult;
-			};
-		};
-
-		/** Cached resources while the program is executing. 
-		* first value of the pair:
-		* 0 : value not valid (not set).
-		* 1 : valid, not worth freeing for memory
-		* 2 : valid, worth freeing
-		*/
-		CodeContainer<FOpExecutionData> OpExecutionData;
-
-
-		/** */
-
-		TMemoryTrackedArray<TManagedWeakPtr<const FLayout>> LayoutResults;
-		TMemoryTrackedArray<TManagedWeakPtr<const FInstance>> InstanceResults;
-		TMemoryTrackedArray<TManagedWeakPtr<const String>> StringResults;
-		TMemoryTrackedArray<TManagedWeakPtr<const FExtensionData>> ExtensionDataResults;
-		TMemoryTrackedArray<TManagedWeakPtr<const FInstancedStruct>> InstancedStructResults;
-		TMemoryTrackedArray<TManagedWeakPtr<const FMaterial>> MaterialResults;
-		TMemoryTrackedArray<TManagedWeakPtr<const FSkeletalMesh>> SkeletalMeshResults;
-
-		TMemoryTrackedArray<TManagedWeakPtr<const FImage>> ImageResults;
-		TMemoryTrackedArray<TManagedWeakPtr<const FMesh>> MeshResults;
-		
-		/** Storage for strong references. Size of Strong Results is the same as the weak. */
-		TMemoryTrackedArray<TManagedPtr<const FImage>> StrongImageResults;
-		TMemoryTrackedArray<TManagedPtr<const FMesh>> StrongMeshResults;
-
-		TMemoryTrackedArray<FVector4f> ColorResults;
-		TMemoryTrackedArray<FProjector> ProjectorResults;
-		TMemoryTrackedArray<FMatrix44f> MatrixResults;
-
-		TManagedPtr<const FImage> NullImage = MakeManagedNull<const FImage>();
-		TManagedPtr<const FMesh> NullMesh = MakeManagedNull<const FMesh>();
-
-		/** */
-		inline const ExecutionIndex& GetRangeIndex(uint32_t i)
-		{
-			// Make sure we have the default element.
-			if (m_usedRangeIndices.IsEmpty())
-			{
-				m_usedRangeIndices.Push(ExecutionIndex());
-			}
-
-			check(i < uint32_t(m_usedRangeIndices.Num()));
-			return m_usedRangeIndices[i];
-		}
-
-		//!
-		inline uint32 GetRangeIndexIndex(const ExecutionIndex& rangeIndex)
-		{
-			if (rangeIndex.IsEmpty())
-			{
-				return 0;
-			}
-
-			// Make sure we have the default element.
-			if (m_usedRangeIndices.IsEmpty())
-			{
-				m_usedRangeIndices.Push(ExecutionIndex());
-			}
-
-			// Look for or add the new element.
-			int32 ElemIndex = m_usedRangeIndices.Find(rangeIndex);
-			if (ElemIndex != INDEX_NONE)
-			{
-				return ElemIndex;
-			}
-
-			m_usedRangeIndices.Push(rangeIndex);
-			return uint32_t(m_usedRangeIndices.Num()) - 1;
-		}
-
-
-		/**
-		 * Sets up the Cache so it can be used by the CodeRunner
-		 * @param InProgram The mutable program
-		 * @param InState The state you want the cache to be "bound" to. It is used to determine which operations should be locked and which not.
-		 * @note When determining which Ops gets locked the state is used as the ones that are part of the current state will be left unlocked while
-		 * the ones from other states will get marked as locked.
-		 */
-		MUTABLERUNTIME_API void Prepare(const UE::Mutable::Private::FProgram& InProgram, int32 InState);
-
-
-		void Init(uint32 Size)
-		{
-			// This clear would prevent live update cache reuse
-			//OpExecutionData.clear();
-
-			OpExecutionData.resize(Size);
-
-			if (ColorResults.IsEmpty())
-			{
-				// Insert dafault/null values
-
-				LayoutResults.Emplace();
-				InstanceResults.Emplace();
-				StringResults.Emplace();
-				ExtensionDataResults.Emplace();
-				MaterialResults.Emplace();
-				SkeletalMeshResults.Emplace();
-				InstancedStructResults.Add({});
-
-				ImageResults.Emplace();
-				MeshResults.Emplace();
-				
-				StrongImageResults.Emplace();
-				StrongMeshResults.Emplace();
-	
-				ColorResults.Emplace();
-				MatrixResults.Emplace();
-				ProjectorResults.Emplace();
-			}
-
-			for (TManagedPtr<const FImage>& StrongImagePtr : StrongImageResults)
-			{
-				StrongImagePtr = nullptr;
-			}
-
-			for (TManagedPtr<const FMesh>& StrongMeshPtr : StrongMeshResults)
-			{
-				StrongMeshPtr = nullptr;
-			}
-		}
-
-		bool IsValid(const FCacheAddress& Address) const
-		{
-			if (Address.At == 0 || Address.At>=OpExecutionData.size_code())
-			{
-				return false;
-			}
-
-			const FOpExecutionData* Data = OpExecutionData.get_ptr(Address);
-			if (!Data)
-			{
-				return false;
-			}
-
-			// Is it a desc data query?
-			switch (Address.Type)
-			{
-			case FScheduledOp::EType::Full:
-			{
-				// Resource validity cannot be checked with this method. See 
-				// AquireCachedResult and AquireCachedResultSet.
-				//check(Data->DataType != EDataType::Mesh && Data->DataType != EDataType::Image);
-				return Data->IsValueValid;
-			}
-			case FScheduledOp::EType::ImageDesc:
-			{
-				return Data->IsDescCacheValid;
-			}
-			default:
-				unimplemented();
-				return false;
-			}
-		}
-		
-
-		void Clear()
-		{
-			MUTABLE_CPUPROFILER_SCOPE(ProgramCacheClear);
-
-			uint32 CodeSize = OpExecutionData.size_code();
-			OpExecutionData.clear();
-			OpExecutionData.resize(CodeSize);
-			
-			LayoutResults.Empty();
-			InstanceResults.Empty();
-			StringResults.Empty();
-			ExtensionDataResults.Empty();
-			InstancedStructResults.Empty();
-			MaterialResults.Empty();
-			SkeletalMeshResults.Empty();
-
-			ImageResults.Empty();
-			MeshResults.Empty();
-		
-			/** Storage for strong references. Size of Strong Results is the same as the weak. */
-			StrongImageResults.Empty();
-			StrongMeshResults.Empty();
-
-			ColorResults.Empty();
-			ProjectorResults.Empty();
-			MatrixResults.Empty();
-		}
-
-
-		void ClearDescCache()
-		{
-			MUTABLE_CPUPROFILER_SCOPE(ProgramDescCacheClear);
-
-			CodeContainer<FProgramCache::FOpExecutionData>::iterator it = OpExecutionData.begin();
-			for (; it.IsValid(); ++it)
-			{
-				(*it).IsDescCacheValid = false;
-			}
-		}
-
-
-		/** 
-		 * Remove all intermediate data (big and small) from the memory except for the one that has been explicitly
-		 * marked as state cache.
-		 */
-		void ClearLayer0()
-		{
-			MUTABLE_CPUPROFILER_SCOPE(ClearLayer0);
-
-			CodeContainer<FProgramCache::FOpExecutionData>::iterator it = OpExecutionData.begin();
-			for (; it.IsValid(); ++it)
-			{
-				FProgramCache::FOpExecutionData& Data = *it;
-				if (Data.IsCacheLocked)
-				{
-					continue;
-				}
-
-				switch (Data.DataType)
-				{
-				case EDataType::Image:
-					{
-						if (Data.DataTypeIndex)
-						{
-							ImageResults[Data.DataTypeIndex] = nullptr;
-							StrongImageResults[Data.DataTypeIndex] = nullptr;
-						}
-						break;
-					}
-
-				case EDataType::Mesh:
-					{
-						if (Data.DataTypeIndex)
-						{
-							MeshResults[Data.DataTypeIndex] = nullptr;
-							StrongMeshResults[Data.DataTypeIndex] = nullptr;
-						}
-						break;
-					}
-
-				case EDataType::SkeletalMesh:
-					{
-						if (Data.DataTypeIndex)
-						{
-							SkeletalMeshResults[Data.DataTypeIndex] = nullptr;
-						}
-						break;
-					}
-
-				case EDataType::Layout:
-					{
-						if (Data.DataTypeIndex)
-						{
-							LayoutResults[Data.DataTypeIndex] = nullptr;
-						}
-						break;
-					}
-
-				case EDataType::Instance:
-					{
-						if (Data.DataTypeIndex)
-						{
-							InstanceResults[Data.DataTypeIndex] = nullptr;
-						}
-						break;
-					}
-
-				default:
-					break;
-				}
-
-				Data.OpHitCount = 0;
-				Data.IsValueValid = false;
-			}
-		}
-
-
-		/** Remove all intermediate data (big and small) from the memory including the one that has been explicitly
-		* marked as state cache.
-		*/
-		void ClearLayer1()
-		{
-			MUTABLE_CPUPROFILER_SCOPE(ClearLayer1);
-
-			CodeContainer<FProgramCache::FOpExecutionData>::iterator it = OpExecutionData.begin();
-			for (; it.IsValid(); ++it)
-			{
-				FProgramCache::FOpExecutionData& Data = *it;
-				
-				Data.OpHitCount = 0;
-				Data.IsValueValid = false;
-			}
-
-			for (TManagedWeakPtr<const FImage>& Value : ImageResults)
-			{
-				Value = nullptr;
-			}
-
-			for (TManagedPtr<const FImage>& Value : StrongImageResults)
-			{
-				Value = nullptr;
-			}
-			
-			
-			for (TManagedWeakPtr<const FSkeletalMesh>& Value : SkeletalMeshResults)
-			{
-				Value = nullptr;
-			}
-
-			for (TManagedPtr<const FMesh>& Value : StrongMeshResults)
-			{
-				Value = nullptr;
-			}
-			
-
-			for (TManagedWeakPtr<const FLayout>& Value : LayoutResults)
-			{
-				Value = nullptr;
-			}
-
-			for (TManagedWeakPtr<const FInstance>& Value : InstanceResults)
-			{
-				Value = nullptr;
-			}
-		}
-		
-		
-		bool GetBool(FCacheAddress Address)
-		{
-			if (!Address.At) 
-			{
-				return false;
-			}
-
-			const FOpExecutionData* Data = OpExecutionData.get_ptr(Address);
-			if (!Data)
-			{
-				return false;
-			}
-
-			check(Data->DataType == EDataType::Bool);
-			return Data->DataTypeIndex != 0;
-		}
-
-		
-		float GetScalar(FCacheAddress Address)
-		{
-			if (!Address.At)
-			{
-				return 0.0f;
-			}
-
-			const FOpExecutionData* Data = OpExecutionData.get_ptr(Address);
-			if (!Data)
-			{
-				return 0.0f;
-			}
-
-			check(Data->DataType == EDataType::Scalar);
-			return Data->ScalarResult;
-		}
-
-		
-		int32 GetInt(FCacheAddress Address)
-		{
-			if (!Address.At)
-			{
-				return 0;
-			}
-
-
-			const FOpExecutionData* Data = OpExecutionData.get_ptr(Address);
-			if (!Data) 
-			{
-				return 0;
-			}
-
-			check(Data->DataType == EDataType::Int);
-			return Data->DataTypeIndex;
-		}
-
-		
-		FVector4f GetColor(FCacheAddress Address)
-		{
-			if (!Address.At) 
-			{
-				return FVector4f();
-			}
-
-			const FOpExecutionData* Data = OpExecutionData.get_ptr(Address);
-			if (!Data)
-			{
-				return FVector4f();
-			}
-
-			check(Data->DataType == EDataType::Color);
-			return ColorResults[Data->DataTypeIndex];
-		}
-
-		
-		FMatrix44f GetMatrix(FCacheAddress Address)
-		{
-			if (!Address.At)
-			{
-				return FMatrix44f::Identity;
-			}
-
-			const FOpExecutionData* Data = OpExecutionData.get_ptr(Address);
-			if (!Data) 
-			{
-				return FMatrix44f::Identity;
-			}
-
-			check(Data->DataType == EDataType::Matrix);
-			return MatrixResults[Data->DataTypeIndex];
-		}
-
-		
-		TManagedPtr<const FMaterial> GetMaterial(FCacheAddress at)
-		{
-			if (!at.At)
-			{
-				return MakeManaged<FMaterial>();
-			}
-			const FOpExecutionData* Data = OpExecutionData.get_ptr(at);
-			if (!Data)
-			{
-				return MakeManaged<FMaterial>();
-			}
-
-			check(Data->DataType == EDataType::Material);
-			TManagedPtr<const FMaterial> FoundMaterial = MaterialResults[Data->DataTypeIndex].Pin();
-			check(FoundMaterial);
-
-			if (!FoundMaterial)
-			{
-				return MakeManaged<FMaterial>();
-			}
-			
-			return FoundMaterial;
-		}
-
-		
-		FProjector GetProjector(FCacheAddress at)
-		{
-			if (!at.At) return FProjector();
-			const FOpExecutionData* Data = OpExecutionData.get_ptr(at);
-			if (!Data) return FProjector();
-
-			check(Data->DataType == EDataType::Projector);
-			return ProjectorResults[Data->DataTypeIndex];
-		}
-
-
-		TManagedPtr<const FInstance> GetInstance(FCacheAddress Address)
-		{
-			if (!Address.At)
-			{
-				return nullptr;
-			}
-
-			FOpExecutionData* Data = OpExecutionData.get_ptr(Address);
-			if (!Data)
-			{
-				return nullptr;
-			}
-
-			check(Data->DataType == EDataType::Instance)
-			TManagedPtr<const FInstance> Result = InstanceResults[Data->DataTypeIndex].Pin();
-
-			// We need to decrease the hit-count even if the result is null.
-			check(Data->OpHitCount > 0);
-
-			--Data->OpHitCount;
-			if (Data->OpHitCount == 0 && !Data->IsCacheLocked)
-			{
-				Data->IsValueValid = false;
-				InstanceResults[Data->DataTypeIndex] = nullptr;
-			}
-
-			return Result;
-		}
-
-
-		TManagedPtr<const FImage> GetImage(FCacheAddress Address)
-		{
-			if (!Address.At || Address.At >= OpExecutionData.size_code())
-			{
-				return nullptr;
-			}
-			
-			FOpExecutionData* Data = OpExecutionData.get_ptr(Address);
-			if (!Data) 
-			{
-				return nullptr;
-			}
-
-			check(Data->DataType == EDataType::Image);
-			TManagedPtr<const FImage> Result = StrongImageResults[Data->DataTypeIndex];
-
-			// We need to decrease the hit-count even if the result is null.
-			check(Data->OpHitCount > 0);
-
-			--Data->OpHitCount;
-			if (Data->OpHitCount == 0)
-			{
-				// If this is the last cached reference, now Result is the only reference, which is given 
-				// to the calling code. 
-				if (!Data->IsCacheLocked)
-				{
-					ImageResults[Data->DataTypeIndex] = nullptr;
-				}
-
-				StrongImageResults[Data->DataTypeIndex] = nullptr;
-			}
-
-			return Result;
-		}
-
-		
-		TManagedPtr<const FMesh> GetMesh(FCacheAddress Address)
-		{
-			if (!Address.At || Address.At >= OpExecutionData.size_code()) 
-			{
-				return nullptr;
-			}
-
-			FOpExecutionData* Data = OpExecutionData.get_ptr(Address);
-			if (!Data) 
-			{
-				return nullptr;
-			}
-
-			check(Data->DataType == EDataType::Mesh);
-			TManagedPtr<const FMesh> Result = StrongMeshResults[Data->DataTypeIndex];
-
-			// We need to decrease the hit-count even if the result is null.
-			check(Data->OpHitCount > 0);
-
-			--Data->OpHitCount;
-			if (Data->OpHitCount == 0)
-			{
-				// If this is the last cached reference, now Result is the only reference, which is given 
-				// to the calling code. 
-				if (!Data->IsCacheLocked)
-				{
-					MeshResults[Data->DataTypeIndex] = nullptr;
-				}
-
-				StrongMeshResults[Data->DataTypeIndex] = nullptr;
-			}
-
-			return Result;
-		}
-
-		TManagedPtr<const FSkeletalMesh> GetSkeletalMesh(FCacheAddress Address)
-		{
-			if (!Address.At)
-			{
-				return nullptr;
-			}
-
-			if (Address.At >= OpExecutionData.size_code())
-			{
-				return nullptr;
-			}
-
-			FOpExecutionData* Data = OpExecutionData.get_ptr(Address);
-			if (!Data) 
-			{
-				return nullptr;
-			}
-
-			check(Data->DataType == EDataType::SkeletalMesh);
-			TManagedPtr<const FSkeletalMesh> Result = SkeletalMeshResults[Data->DataTypeIndex].Pin();
-
-			// We need to decrease the hit-count even if the result is null.
-			check(Data->OpHitCount > 0);
-
-			--Data->OpHitCount;
-			if (Data->OpHitCount == 0 && !Data->IsCacheLocked)
-			{
-				Data->IsValueValid = false;
-				SkeletalMeshResults[Data->DataTypeIndex] = nullptr;
-			}
-
-			return Result;
-		}
-		
-		TManagedPtr<const FLayout> GetLayout(FCacheAddress Address)
-		{
-			if (!Address.At)
-			{
-				return nullptr;
-			}
-
-			FOpExecutionData* Data = OpExecutionData.get_ptr(Address);
-			if (!Data)
-			{
-				return nullptr;
-			}
-
-			check(Data->DataType == EDataType::Layout);
-			return LayoutResults[Data->DataTypeIndex].Pin();
-		}
-
-
-		TManagedPtr<const String> GetString(FCacheAddress Address)
-		{
-			if (!Address.At)
-			{
-				return nullptr;
-			}
-			FOpExecutionData* Data = OpExecutionData.get_ptr(Address);
-			if (!Data) 
-			{
-				return nullptr;
-			}
-
-			check(Data->DataType == EDataType::String);
-			return StringResults[Data->DataTypeIndex].Pin();
-		}
-
-
-		TManagedPtr<const FExtensionData> GetExtensionData(FCacheAddress Address)
-		{
-			if (!Address.At) 
-			{
-				return nullptr;
-			}
-
-			FOpExecutionData* Data = OpExecutionData.get_ptr(Address);
-			if (!Data)
-			{
-				return nullptr;
-			}
-
-			check(Data->DataType == EDataType::ExtensionData);
-			return ExtensionDataResults[Data->DataTypeIndex].Pin();
-		}
-		
-		
-		TManagedPtr<const FInstancedStruct> GetInstancedStruct(FCacheAddress Address)
-		{
-			if (!Address.At) 
-			{
-				return nullptr;
-			}
-
-			FOpExecutionData* Data = OpExecutionData.get_ptr(Address);
-			if (!Data) 
-			{
-				return nullptr;
-			}
-
-			check(Data->DataType == EDataType::InstancedStruct);
-			TManagedPtr<const FInstancedStruct> Result = InstancedStructResults[Data->DataTypeIndex].Pin();
-
-			check(Result);
-			check(Result->IsValid());
-
-			return Result;
-		}
-		
-		
-		void SetValidDesc(FCacheAddress at)
-		{
-			check(at.Type == FScheduledOp::EType::ImageDesc);
-			check(at.At < OpExecutionData.size_code());
-			FOpExecutionData& Data = OpExecutionData[at];
-			Data.IsDescCacheValid = true;
-		}
-
-		
-		void SetBool(FCacheAddress at, bool v)
-		{
-			check(at.At < OpExecutionData.size_code());
-			FOpExecutionData* Data = OpExecutionData.get_ptr(at);
-			check(Data->DataType == EDataType::Bool || Data->DataType == EDataType::None);
-			Data->DataType = EDataType::Bool;
-			Data->DataTypeIndex = v;
-			Data->IsValueValid = true;
-		}
-
-		
-		void SetInt(FCacheAddress at, int32 v)
-		{
-			check(at.At < OpExecutionData.size_code());
-			FOpExecutionData* Data = OpExecutionData.get_ptr(at);
-			check(Data->DataType == EDataType::Int || Data->DataType == EDataType::None);
-			Data->DataType = EDataType::Int;
-			Data->DataTypeIndex = v;
-			Data->IsValueValid = true;
-		}
-
-		
-		void SetScalar(FCacheAddress Address, float v)
-		{
-			check(Address.At < OpExecutionData.size_code());
-			FOpExecutionData* Data = OpExecutionData.get_ptr(Address);
-			check(Data->DataType == EDataType::Scalar || Data->DataType == EDataType::None);
-			Data->DataType = EDataType::Scalar;
-			Data->ScalarResult = v;
-			Data->IsValueValid = true;
-		}
-
-		
-		void SetColour(FCacheAddress Address, const FVector4f& v)
-		{
-			check(Address.At < OpExecutionData.size_code());
-			FOpExecutionData* Data = OpExecutionData.get_ptr(Address);
-			check(Data->DataType == EDataType::Color || Data->DataType == EDataType::None);
-			Data->DataType = EDataType::Color;
-			Data->IsValueValid = true;
-
-			if (!Data->DataTypeIndex)
-			{
-				Data->DataTypeIndex = ColorResults.Num();
-				ColorResults.Add(v);
-			}
-			else
-			{
-				ColorResults[Data->DataTypeIndex] = v;
-			}
-			check(Data->DataTypeIndex != 0);
-		}
-
-		
-		void SetMatrix(FCacheAddress Address, const FMatrix44f& v)
-		{
-			check(Address.At < OpExecutionData.size_code());
-			FOpExecutionData* Data = OpExecutionData.get_ptr(Address);
-			check(Data->DataType == EDataType::Matrix || Data->DataType == EDataType::None);
-			Data->DataType = EDataType::Matrix;
-			Data->IsValueValid = true;
-
-			if (!Data->DataTypeIndex)
-			{
-				Data->DataTypeIndex = MatrixResults.Num();
-				MatrixResults.Add(v);
-			}
-			else
-			{
-				MatrixResults[Data->DataTypeIndex] = v;
-			}
-			check(Data->DataTypeIndex != 0);
-		}
-
-		
-		void SetMaterial(FCacheAddress Address, const TManagedPtr<const FMaterial>& Value)
-		{
-			check(Value);		// All material pointers provided must be valid
-			check(Address.At < OpExecutionData.size_code());
-			FOpExecutionData* Data = OpExecutionData.get_ptr(Address);
-			check(Data->DataType == EDataType::Material || Data->DataType == EDataType::None);
-			Data->DataType = EDataType::Material;
-			Data->IsValueValid = true;
-
-			if (!Data->DataTypeIndex)
-			{
-				Data->DataTypeIndex = MaterialResults.Num();
-				MaterialResults.Add(Value.ToWeakPtr());
-			}
-			else
-			{
-				MaterialResults[Data->DataTypeIndex] = Value.ToWeakPtr();
-			}
-			check(Data->DataTypeIndex != 0);
-		}
-		
-		
-		void SetProjector(FCacheAddress Address, const FProjector& v)
-		{
-			check(Address.At < OpExecutionData.size_code());
-			FOpExecutionData* Data = OpExecutionData.get_ptr(Address);
-			check(Data->DataType == EDataType::Projector || Data->DataType == EDataType::None);
-			Data->DataType = EDataType::Projector;
-			Data->IsValueValid = true;
-
-			if (!Data->DataTypeIndex)
-			{
-				Data->DataTypeIndex = ProjectorResults.Num();
-				ProjectorResults.Add(v);
-			}
-			else
-			{
-				ProjectorResults[Data->DataTypeIndex] = v;
-			}
-			check(Data->DataTypeIndex != 0);
-		}
-
-		
-		void SetInstance(FCacheAddress Address, const TManagedPtr<const FInstance>& Value)
-		{
-			check(Address.At < OpExecutionData.size_code());
-			FOpExecutionData* Data = OpExecutionData.get_ptr(Address);
-			check(Data->DataType == EDataType::Instance || Data->DataType == EDataType::None);
-			Data->DataType = EDataType::Instance;
-			Data->IsValueValid = true;
-
-			if (!Data->DataTypeIndex)
-			{
-				Data->DataTypeIndex = InstanceResults.Num();
-				InstanceResults.Add(Value.ToWeakPtr());
-			}
-			else
-			{
-				InstanceResults[Data->DataTypeIndex] = Value.ToWeakPtr();
-			}
-
-			check(Data->DataTypeIndex != 0);
-		}
-
-		
-		void SetExtensionData(FCacheAddress Address, const TManagedPtr<const FExtensionData>& Value)
-		{
-			check(Address.At < OpExecutionData.size_code());
-			FOpExecutionData* Data = OpExecutionData.get_ptr(Address);
-			check(Data->DataType == EDataType::ExtensionData || Data->DataType == EDataType::None);
-			Data->DataType = EDataType::ExtensionData;
-			Data->IsValueValid = true;
-
-			if (!Data->DataTypeIndex)
-			{
-				Data->DataTypeIndex = ExtensionDataResults.Num();
-				ExtensionDataResults.Add(Value.ToWeakPtr());
-			}
-			else
-			{
-				ExtensionDataResults[Data->DataTypeIndex] = Value.ToWeakPtr();
-			}
-			check(Data->DataTypeIndex != 0);
-		}
-
-		
-		void SetInstancedStruct(FCacheAddress Address, const TManagedPtr<const FInstancedStruct>& Value)
-		{
-			check(Address.At < OpExecutionData.size_code());
-			FOpExecutionData* Data = OpExecutionData.get_ptr(Address);
-			check(Data->DataType == EDataType::InstancedStruct || Data->DataType == EDataType::None);
-			Data->DataType = EDataType::InstancedStruct;
-			Data->IsValueValid = true;
-
-			if (!Data->DataTypeIndex)
-			{
-				Data->DataTypeIndex = InstancedStructResults.Num();
-				InstancedStructResults.Add(Value.ToWeakPtr());
-			}
-			else
-			{
-				InstancedStructResults[Data->DataTypeIndex] = Value.ToWeakPtr();
-			}
-			check(Data->DataTypeIndex != 0);
-		}
-		
-		
-		void SetImage(FCacheAddress Address, const TManagedPtr<const FImage>& Value)
-		{
-			check(Address.At < OpExecutionData.size_code());
-			FOpExecutionData* Data = OpExecutionData.get_ptr(Address);
-			check(Data->DataType == EDataType::Image || Data->DataType == EDataType::None);
-			Data->DataType = EDataType::Image;
-			Data->IsValueValid = true;
-
-			const TManagedPtr<const FImage>* ValueToStorePtr = Value.IsValid() ? &Value : &NullImage; 
-
-			check(ValueToStorePtr->IsSet());
-
-			if (!Data->DataTypeIndex)
-			{
-				Data->DataTypeIndex = ImageResults.Num();
-				ImageResults.Emplace(ValueToStorePtr->ToWeakPtr());
-				StrongImageResults.Emplace(*ValueToStorePtr);
-			}
-			else
-			{
-				ImageResults[Data->DataTypeIndex] = ValueToStorePtr->ToWeakPtr();
-				StrongImageResults[Data->DataTypeIndex] = *ValueToStorePtr;
-			}
-			check(Data->DataTypeIndex != 0);
-
-			UE::Mutable::Private::UpdateLLMStats();
-		}
-
-		
-		void SetMesh(FCacheAddress Address, const TManagedPtr<const FMesh>& Value)
-		{
-			check(Address.At < OpExecutionData.size_code());
-
-			FOpExecutionData* Data = OpExecutionData.get_ptr(Address);
-
-			check(Data->DataType == EDataType::Mesh || Data->DataType == EDataType::None);
-			Data->DataType = EDataType::Mesh;
-			Data->IsValueValid = true;
-
-			const TManagedPtr<const FMesh>* ValueToStorePtr = Value.IsValid() ? &Value : &NullMesh; 
-
-			check(ValueToStorePtr->IsSet());
-
-			if (!Data->DataTypeIndex)
-			{
-				Data->DataTypeIndex = MeshResults.Num();
-				MeshResults.Emplace(ValueToStorePtr->ToWeakPtr());
-				StrongMeshResults.Emplace(*ValueToStorePtr);
-			}
-			else
-			{
-				MeshResults[Data->DataTypeIndex] = ValueToStorePtr->ToWeakPtr();
-				StrongMeshResults[Data->DataTypeIndex] = *ValueToStorePtr;
-			}
-			check(Data->DataTypeIndex != 0);
-
-			UE::Mutable::Private::UpdateLLMStats();
-		}
-
-		
-		void SetSkeletalMesh(FCacheAddress Address, const TManagedPtr<const FSkeletalMesh>& Value)
-		{
-			check(Address.At < OpExecutionData.size_code());
-			FOpExecutionData* Data = OpExecutionData.get_ptr(Address);
-			check(Data->DataType == EDataType::SkeletalMesh || Data->DataType == EDataType::None);
-			Data->DataType = EDataType::SkeletalMesh;
-			Data->IsValueValid = true;
-
-			if (!Data->DataTypeIndex)
-			{
-				Data->DataTypeIndex = SkeletalMeshResults.Num();
-				SkeletalMeshResults.Add(Value);
-			}
-			else
-			{
-				SkeletalMeshResults[Data->DataTypeIndex] = Value;
-			}
-			check(Data->DataTypeIndex != 0);
-
-			UE::Mutable::Private::UpdateLLMStats();
-		}
-
-		
-		void SetLayout(FCacheAddress Address, const TManagedPtr<const FLayout>& Value)
-		{
-			check(Address.At < OpExecutionData.size_code());
-			FOpExecutionData* Data = OpExecutionData.get_ptr(Address);
-			check(Data->DataType == EDataType::Layout || Data->DataType == EDataType::None);
-			Data->DataType = EDataType::Layout;
-			Data->IsValueValid = true;
-
-			if (!Data->DataTypeIndex)
-			{
-				Data->DataTypeIndex = LayoutResults.Num();
-				LayoutResults.Add(Value);
-			}
-			else
-			{
-				LayoutResults[Data->DataTypeIndex] = Value;
-			}
-			check(Data->DataTypeIndex != 0);
-
-			UE::Mutable::Private::UpdateLLMStats();
-		}
-
-		
-		void SetString(FCacheAddress Address, const TManagedPtr<const String>& Value)
-		{
-			check(Address.At < OpExecutionData.size_code());
-			FOpExecutionData* Data = OpExecutionData.get_ptr(Address);
-			check(Data->DataType == EDataType::String || Data->DataType == EDataType::None);
-			Data->DataType = EDataType::String;
-			Data->IsValueValid = true;
-
-			if (!Data->DataTypeIndex)
-			{
-				Data->DataTypeIndex = StringResults.Num();
-				StringResults.Add(Value.ToWeakPtr());
-			}
-			else
-			{
-				StringResults[Data->DataTypeIndex] = Value.ToWeakPtr();
-			}
-			check(Data->DataTypeIndex!=0);
-		}
-
-
-		inline void IncreaseHitCount(FCacheAddress Address)
-		{
-			// Don't count hits for instruction 0, which is always null. It is usually already
-			// check that At is not 0, and then it is not requested, generating a stray non-zero count
-			// at its position.
-			if (Address.At)
-			{
-				check(Address.At < OpExecutionData.size_code());
-				FOpExecutionData& Data = OpExecutionData[Address];
-				Data.OpHitCount++;
-			}
-		}
-
-		inline void SetForceCached(OP::ADDRESS Address)
-		{
-			// \TODO: It only locks Address,0,0
-			if (Address)
-			{
-				check(Address < OpExecutionData.size_code());
-				FOpExecutionData& Data = OpExecutionData[FCacheAddress(Address, 0, 0)];
-				Data.IsCacheLocked = true;
-			}
-		}
-
-		bool TryAquireCachedResultSet(TConstArrayView<FCacheAddress> AddressSet)
-		{
-			TArray<TManagedPtr<const FImage>, TInlineAllocator<16>> LocalImages;
-			TArray<TManagedPtr<const FMesh> , TInlineAllocator<16>> LocalMeshes;
-			TArray<FOpExecutionData*        , TInlineAllocator<16>> LocalDataImages;
-			TArray<FOpExecutionData*        , TInlineAllocator<16>> LocalDataMeshes;
-
-			if (AddressSet.IsEmpty())
-			{
-				return true;
-			}
-
-			for (const FCacheAddress& Address : AddressSet)
-			{
-				if (!Address.At)
-				{
-					continue;
-				}
-
-				check(Address.At < OpExecutionData.size_code());
-
-				FOpExecutionData* Data = OpExecutionData.get_ptr(Address);
-			
-				if (!Data)
-				{
-					return false;
-				}
-			
-				if (Address.Type == FScheduledOp::EType::ImageDesc)
-				{
-					if (!Data->IsDescCacheValid)
-					{
-						return false;
-					}
-				}
-				else
-				{
-					switch (Data->DataType)
-					{
-						case EDataType::Image:
-						{
-							if (!LocalImages.Add_GetRef(ImageResults[Data->DataTypeIndex].Pin()).IsSet())
-							{
-								return false;
-							}
-							LocalDataImages.Add(Data);
-							break;
-						}
-						case EDataType::Mesh:
-						{
-							if (!LocalMeshes.Add_GetRef(MeshResults[Data->DataTypeIndex].Pin()).IsSet())
-							{
-								return false;
-							}
-							LocalDataMeshes.Add(Data);
-							break;
-						}
-
-						default:
-						{
-							if (!Data->IsValueValid)
-							{
-								return false;
-							}
-							break;
-						}
-					}
-				}
-			}
-
-			// We have been successful acquiring all addresses in the set. 
-			// Commit resource acquisition.
-			for (int32 Index = LocalDataImages.Num() - 1; Index >= 0; --Index)
-			{
-				FOpExecutionData* Data = LocalDataImages[Index];
-				check(Data->DataType == EDataType::Image);
-
-				StrongImageResults[Data->DataTypeIndex] = LocalImages[Index];
-			}
-
-			for (int32 Index = LocalDataMeshes.Num() - 1; Index >= 0; --Index)
-			{
-				FOpExecutionData* Data = LocalDataMeshes[Index];
-				check(Data->DataType == EDataType::Mesh);
-
-				StrongMeshResults[Data->DataTypeIndex] = LocalMeshes[Index];
-			}
-
-			return true;
-		}
-
-		bool AquireCachedResult(FCacheAddress Address)
-		{
-			if (!Address.At)
-			{
-				return false;
-			}
-
-			check(Address.At < OpExecutionData.size_code());
-
-			FOpExecutionData* Data = OpExecutionData.get_ptr(Address);
-		
-			if (!Data)
-			{
-				return false;
-			}
-
-			if (Address.Type == FScheduledOp::EType::ImageDesc)
-			{
-				return Data->IsDescCacheValid;
-			}
-	
-			switch (Data->DataType)
-			{
-				case EDataType::Image:
-				{
-					if (StrongImageResults[Data->DataTypeIndex].IsSet())
-					{
-						return true;
-					}
-					else
-					{
-						TManagedPtr<const FImage> PinnedPtr = ImageResults[Data->DataTypeIndex].Pin();
-						StrongImageResults[Data->DataTypeIndex] = PinnedPtr;
-						return PinnedPtr.IsSet();
-					}
-				}
-				case EDataType::Mesh:
-				{
-					if (StrongMeshResults[Data->DataTypeIndex].IsSet())
-					{
-						return true;
-					}
-					else
-					{
-
-						TManagedPtr<const FMesh> PinnedPtr = MeshResults[Data->DataTypeIndex].Pin();
-						StrongMeshResults[Data->DataTypeIndex] = PinnedPtr;
-						return PinnedPtr.IsSet();
-					}
-				}
-
-				default:
-				{
-					return Data->IsValueValid;
-				}
-			}
-		}
-	};
-
-
-	inline bool operator==(const FCacheAddress& a, const FCacheAddress& b)
-	{
-		return a.At == b.At
-			&&
-			a.ExecutionIndex == b.ExecutionIndex
-			&&
-			a.ExecutionOptions == b.ExecutionOptions
-			&&
-			a.Type == b.Type;
-	}
-
-	inline bool operator<(const FCacheAddress& a, const FCacheAddress& b)
-	{
-		if (a.At < b.At) return true;
-		if (a.At > b.At) return false;
-		if (a.ExecutionIndex < b.ExecutionIndex) return true;
-		if (a.ExecutionIndex > b.ExecutionIndex) return false;
-		if (a.ExecutionOptions < b.ExecutionOptions) return true;
-		if (a.ExecutionOptions > b.ExecutionOptions) return false;
-		return a.Type < b.Type;
-	}
 
 	/** Class to log update messages */
 	class FLiveInstanceLogger
@@ -1664,6 +194,7 @@ namespace UE::Mutable::Private
 		int64 BudgetExcessBytes = 0;
     	
 		/** Data for each mutable instance that is being updated. */
+		mutable FMutex LiveInstancesMutex;
 		TMemoryTrackedArray<TWeakPtr<FLiveInstance>> LiveInstances;
 
         /** Make sure the working memory is below the internal budget, even counting with the passed additional memory. 
@@ -1683,8 +214,9 @@ namespace UE::Mutable::Private
     	int32 GetCacheBytes() const
     	{
     		int32 Result = 0;
-    		TSet<TManagedPtr<const FResource>> Cache0Unique;
-    		TSet<TManagedPtr<const FResource>> Cache1Unique;
+    		TSet<FProgramCache*> VisitedCaches;
+
+			TScopeLock<FMutex> LockGuard(LiveInstancesMutex);
 
     		for (const TWeakPtr<FLiveInstance>& WeakInstance : LiveInstances)
     		{
@@ -1694,58 +226,21 @@ namespace UE::Mutable::Private
     				continue;
     			}
 
-    			CodeContainer<FProgramCache::FOpExecutionData>::iterator it = Instance->Cache->OpExecutionData.begin();
-    			for (; it.IsValid(); ++it)
-    			{
-    				if (!(*it).DataTypeIndex)
-    				{
-    					continue;
-    				}
-
-    				TSet<TManagedPtr<const FResource>>* TargetSet = &Cache0Unique;
-    				if ((*it).IsCacheLocked)
-    				{
-    					TargetSet = &Cache1Unique;
-    				}
-
-    				switch ((*it).DataType)
-    				{
-    				case EDataType::Image:
-    					{
-    						TManagedPtr<const FImage> Value = Instance->Cache->ImageResults[(*it).DataTypeIndex].Pin();
-    						if (Value)
-    						{
-    							TargetSet->Add(MoveTemp(Value));
-    						}
-    						break;
-    					}
-    				case EDataType::Mesh:
-    					{
-    						TManagedPtr<const FMesh> Value = Instance->Cache->MeshResults[(*it).DataTypeIndex].Pin();
-    						if (Value)
-    						{
-    							TargetSet->Add(MoveTemp(Value));
-    						}
-    						break;
-    					}
-    				default:
-    					break;
-    				}
-    			}
-    		}
+				TSharedPtr<FProgramCache> CurrentCache = Instance->Cache;
 
-    		// Merge
-    		Cache0Unique.Append(Cache1Unique);
+				if (VisitedCaches.Contains(CurrentCache.Get()))
+				{
+					continue;
+				}
+	
+				VisitedCaches.Add(CurrentCache.Get());
 
-    		// Count
-    		for (const TManagedPtr<const FResource>& Value : Cache0Unique)
-    		{
-    			Result += Value->GetDataSize();
+				Result += CurrentCache->CountBytes();
     		}
 
-    		return Result;
+			return Result;
     	}
-    	
+
 		void LogWorkingMemory(const class CodeRunner* CurrentRunner) const;
 	};
 	
@@ -1854,20 +349,20 @@ namespace UE::Mutable::Private
         //! the same instance index. After a call to this method, that Instance cannot be used any
         //! more and its content is undefined.
         //! \param instance The index of the instance whose last data will be invalidated.
-        UE_API void EndUpdate(const TSharedRef<FLiveInstance>& LiveInstance, bool bClearRoms);
+        UE_API void EndUpdate(const TSharedRef<FLiveInstance>& LiveInstance, bool bClearRoms, bool bFreeCache);
 		
 		//! Calculate the relevancy of every parameter. Some parameters may be unused depending on
 		//! the values of other parameters. This is useful to hide irrelevant parameters in
 		//! dynamic user interfaces.
         UE_API void GetParameterRelevancy(const TSharedRef<FLiveInstance>& LiveInstance, TArray<uint32>& OutRelevantParameters);
 
-		UE_API static FMeshId GetMeshId(const TSharedPtr<FMeshIdRegistry>& Registry, const FModel& Model, const FParameters& Params, OP::ADDRESS RootAt);
+		UE_API static FMeshId GetMeshId(const TSharedPtr<FMeshIdRegistry>& Registry, const FModel& Model, const FParameters& Params, FOperation::ADDRESS RootAt);
     	
-		UE_API static FImageId GetImageId(const TSharedPtr<FImageIdRegistry>& Registry, const FModel& Model, const FParameters& Params, OP::ADDRESS RootAt);
+		UE_API static FImageId GetImageId(const TSharedPtr<FImageIdRegistry>& Registry, const FModel& Model, const FParameters& Params, FOperation::ADDRESS RootAt);
 
-		UE_API static FMaterialId GetMaterialId(const TSharedPtr<FMaterialIdRegistry>& Registry, const FModel& Model, const FParameters& Params, OP::ADDRESS RootAt);
+		UE_API static FMaterialId GetMaterialId(const TSharedPtr<FMaterialIdRegistry>& Registry, const FModel& Model, const FParameters& Params, FOperation::ADDRESS RootAt);
 
-		UE_API static FSkeletalMeshId GetSkeletalMeshId(const TSharedPtr<FSkeletalMeshIdRegistry>& Registry, const FModel& Model, const FParameters& Params, OP::ADDRESS RootAt);
+		UE_API static FSkeletalMeshId GetSkeletalMeshId(const TSharedPtr<FSkeletalMeshIdRegistry>& Registry, const FModel& Model, const FParameters& Params, FOperation::ADDRESS RootAt);
 		
         // Prevent copy, move and assignment.
         FSystem( const FSystem& ) = delete;
@@ -1875,16 +370,16 @@ namespace UE::Mutable::Private
         FSystem( FSystem&& ) = delete;
         FSystem& operator=( FSystem&& ) = delete;
 		
-		MUTABLERUNTIME_API bool BuildBool(const TSharedRef<FLiveInstance>& InLiveInstance, OP::ADDRESS InAddress) ;
-		MUTABLERUNTIME_API int32 BuildInt(const TSharedRef<FLiveInstance>& InLiveInstance, OP::ADDRESS InAddress) ;
-		MUTABLERUNTIME_API float BuildScalar(const TSharedRef<FLiveInstance>& InLiveInstance, OP::ADDRESS InAddress) ;
-		MUTABLERUNTIME_API FVector4f BuildColor(const TSharedRef<FLiveInstance>& InLiveInstance, OP::ADDRESS InAddress) ;
-		MUTABLERUNTIME_API TManagedPtr<const String> BuildString(const TSharedRef<FLiveInstance>& InLiveInstance, OP::ADDRESS InAddress) ;
-		MUTABLERUNTIME_API TManagedPtr<const FImage> BuildImage(const TSharedRef<FLiveInstance>& InLiveInstance, OP::ADDRESS InAddress, int32 MipsToSkip);
-		MUTABLERUNTIME_API TManagedPtr<const FMesh> BuildMesh(const TSharedRef<FLiveInstance>& InLiveInstance, OP::ADDRESS InRootAddress, EMeshContentFlags MeshContentFilter);
-		MUTABLERUNTIME_API TManagedPtr<const FInstance> BuildInstance(const TSharedRef<FLiveInstance>& InLiveInstance, OP::ADDRESS InAddress);
-		MUTABLERUNTIME_API TManagedPtr<const FLayout> BuildLayout(const TSharedRef<FLiveInstance>& InLiveInstance, OP::ADDRESS InAddress) ;
-		MUTABLERUNTIME_API FProjector BuildProjector(const TSharedRef<FLiveInstance>& InLiveInstance, OP::ADDRESS InAddress) ;
+		MUTABLERUNTIME_API bool BuildBool(const TSharedRef<FLiveInstance>& InLiveInstance, FOperation::ADDRESS InAddress) ;
+		MUTABLERUNTIME_API int32 BuildInt(const TSharedRef<FLiveInstance>& InLiveInstance, FOperation::ADDRESS InAddress) ;
+		MUTABLERUNTIME_API float BuildScalar(const TSharedRef<FLiveInstance>& InLiveInstance, FOperation::ADDRESS InAddress) ;
+		MUTABLERUNTIME_API FVector4f BuildColor(const TSharedRef<FLiveInstance>& InLiveInstance, FOperation::ADDRESS InAddress) ;
+		MUTABLERUNTIME_API TManagedPtr<const String> BuildString(const TSharedRef<FLiveInstance>& InLiveInstance, FOperation::ADDRESS InAddress) ;
+		MUTABLERUNTIME_API TManagedPtr<const FImage> BuildImage(const TSharedRef<FLiveInstance>& InLiveInstance, FOperation::ADDRESS InAddress, int32 MipsToSkip);
+		MUTABLERUNTIME_API TManagedPtr<const FMesh> BuildMesh(const TSharedRef<FLiveInstance>& InLiveInstance, FOperation::ADDRESS InRootAddress, EMeshContentFlags MeshContentFilter);
+		MUTABLERUNTIME_API TManagedPtr<const FInstance> BuildInstance(const TSharedRef<FLiveInstance>& InLiveInstance, FOperation::ADDRESS InAddress);
+		MUTABLERUNTIME_API TManagedPtr<const FLayout> BuildLayout(const TSharedRef<FLiveInstance>& InLiveInstance, FOperation::ADDRESS InAddress) ;
+		MUTABLERUNTIME_API FProjector BuildProjector(const TSharedRef<FLiveInstance>& InLiveInstance, FOperation::ADDRESS InAddress) ;
 		
 		
         FSettings Settings;
@@ -1909,7 +404,7 @@ namespace UE::Mutable::Private
 		
 		void RunCode(
 			const TSharedRef<FLiveInstance>& InLiveInstance,
-			OP::ADDRESS,
+			FOperation::ADDRESS,
 			uint32 LODs = AllLODs,
 			uint8 ExecutionOptions = 0);
 
diff --git a/Engine/Plugins/Mutable/Source/MutableRuntime/Internal/MuR/SystemTypes.h b/Engine/Plugins/Mutable/Source/MutableRuntime/Internal/MuR/SystemTypes.h
--- a/Engine/Plugins/Mutable/Source/MutableRuntime/Internal/MuR/SystemTypes.h
+++ b/Engine/Plugins/Mutable/Source/MutableRuntime/Internal/MuR/SystemTypes.h
@@ -18,7 +18,7 @@ namespace UE::Mutable::Private
 
 		inline FScheduledOp() = default;
 
-		inline FScheduledOp(OP::ADDRESS InAt, const FScheduledOp& InOpTemplate, uint8 InStage = 0, uint32 InCustomState = 0)
+		inline FScheduledOp(FOperation::ADDRESS InAt, const FScheduledOp& InOpTemplate, uint8 InStage = 0, uint32 InCustomState = 0)
 			: At(InAt)
 			, CustomState(InCustomState)
 			, ExecutionIndex(InOpTemplate.ExecutionIndex)
@@ -29,7 +29,7 @@ namespace UE::Mutable::Private
 			check(InStage <= MaxNumStages);
 		}
 
-		static inline FScheduledOp FromOpAndOptions(OP::ADDRESS InAt, const FScheduledOp& InOpTemplate, uint8 InExecutionOptions)
+		static inline FScheduledOp FromOpAndOptions(FOperation::ADDRESS InAt, const FScheduledOp& InOpTemplate, uint8 InExecutionOptions)
 		{
 			FScheduledOp Result;
 
@@ -44,7 +44,7 @@ namespace UE::Mutable::Private
 		}
 
 		//! Address of the operation
-		OP::ADDRESS At = 0;
+		FOperation::ADDRESS At = 0;
 
 		//! Additional custom state data that the operation can store. This is usually used to pass information
 		//! between execution stages of an operation.
@@ -89,45 +89,4 @@ namespace UE::Mutable::Private
 	{
 		return HashCombineFast(::GetTypeHash(Op.At), HashCombineFast((uint32)Op.Stage), (uint32)Op.ExecutionIndex);
 	}
-
-	/** A cache address is the operation plus the context of execution(iteration indices, etc...). */
-	struct FCacheAddress
-	{
-		/** The meaning of all these fields is the same than the FScheduledOp struct. */
-		OP::ADDRESS At = 0;
-		uint16 ExecutionIndex = 0;
-		uint8 ExecutionOptions = 0;
-		FScheduledOp::EType Type = FScheduledOp::EType::Full;
-
-		FCacheAddress() 
-		{
-		}
-
-		FCacheAddress(OP::ADDRESS InAt, uint16 InExecutionIndex, uint8 InExecutionOptions, FScheduledOp::EType InType = FScheduledOp::EType::Full)
-			: At(InAt), ExecutionIndex(InExecutionIndex), ExecutionOptions(InExecutionOptions), Type(InType)
-		{
-		}
-
-		FCacheAddress(OP::ADDRESS InAt, const FScheduledOp& Item)
-			: At(InAt), ExecutionIndex(Item.ExecutionIndex), ExecutionOptions(Item.ExecutionOptions), Type(Item.Type)
-		{
-		}
-		
-		FCacheAddress(const FScheduledOp& Item)
-			: At(Item.At), ExecutionIndex(Item.ExecutionIndex), ExecutionOptions(Item.ExecutionOptions), Type(Item.Type)
-		{
-		}
-	};
-
-	inline uint32 GetTypeHash(const FCacheAddress& Address)
-	{
-		uint32 Hash = Address.At;
-
-		Hash = HashCombineFast(Hash, static_cast<uint32>(Address.ExecutionIndex));
-		Hash = HashCombineFast(Hash, static_cast<uint32>(Address.ExecutionOptions));
-		Hash = HashCombineFast(Hash, static_cast<uint32>(Address.Type));
-
-		return Hash;
-	}
-
 } // namespace UE::Mutable::Private
