[Mutable] Program cache rewrok.

diff --git a/Engine/Plugins/Mutable/Source/MutableRuntime/Private/MuR/CodeRunner.cpp b/Engine/Plugins/Mutable/Source/MutableRuntime/Private/MuR/CodeRunner.cpp
--- a/Engine/Plugins/Mutable/Source/MutableRuntime/Private/MuR/CodeRunner.cpp
+++ b/Engine/Plugins/Mutable/Source/MutableRuntime/Private/MuR/CodeRunner.cpp
@@ -150,7 +150,7 @@ namespace UE::Mutable::Private
 		const FSettings& InSettings,
 		FSystem& InSystem,
 		EExecutionStrategy InExecutionStrategy,
-		OP::ADDRESS At,
+		FOperation::ADDRESS At,
 		uint32 InLODMask,
 		uint8 ExecutionOptions,
 		bool bForceLazy,
@@ -167,7 +167,7 @@ namespace UE::Mutable::Private
 		const FSettings& InSettings,
 		FSystem& InSystem,
 		EExecutionStrategy InExecutionStrategy,
-		OP::ADDRESS At,
+		FOperation::ADDRESS At,
 		uint32 InLodMask,
 		uint8 ExecutionOptions,
 		bool bForceLazy,
@@ -183,19 +183,14 @@ namespace UE::Mutable::Private
 		, ExternalResourceProvider(InLiveInstance->ExternalResourceProvider)
 	{
 		MUTABLE_CPUPROFILER_SCOPE(CodeRunner_Create)
-		ScheduledStagePerOp.resize(Program.OpAddress.Num());
 		
 		if (Type == FScheduledOp::EType::ImageDesc)
 		{
-			ImageDescResults.Reserve(64);
 			ImageDescConstantImages.Reserve(32);
 		}
 	
 		// We will read this in the end, so make sure we keep it.
-   		if (Type == FScheduledOp::EType::Full)
-   		{
-			GetMemory().IncreaseHitCount(FCacheAddress(At, 0, ExecutionOptions));
-		}
+		GetMemory().UpdateHitCount(FCacheAddress(At, 0, ExecutionOptions, Type), 1);
 		
 		// Push the root operation
 		FScheduledOp RootOp;
@@ -292,7 +287,7 @@ namespace UE::Mutable::Private
         MUTABLE_CPUPROFILER_SCOPE(RunCode_Conditional);
 		
 		EOpType type = Program.GetOpType(Item.At);
-		OP::ConditionalArgs Args = Program.GetOpArgs<OP::ConditionalArgs>(Item.At);
+		FOperation::ConditionalArgs Args = Program.GetOpArgs<FOperation::ConditionalArgs>(Item.At);
 
         // Conditionals have the following execution stages:
         // 0: we need to run the condition
@@ -315,7 +310,7 @@ namespace UE::Mutable::Private
             bool value = true;
             value = LoadBool( FCacheAddress(Args.condition, Item.ExecutionIndex, Item.ExecutionOptions) );
 
-            OP::ADDRESS resultAt = value ? Args.yes : Args.no;
+            FOperation::ADDRESS resultAt = value ? Args.yes : Args.no;
 
             // Schedule the end of this instruction if necessary
             AddOp( FScheduledOp( Item.At, Item, 2, (uint32)value),
@@ -326,7 +321,7 @@ namespace UE::Mutable::Private
 
         case 2:
         {
-            OP::ADDRESS resultAt = Item.CustomState ? Args.yes : Args.no;
+            FOperation::ADDRESS resultAt = Item.CustomState ? Args.yes : Args.no;
 
             // Store the final result
             FCacheAddress cat( Item );
@@ -366,17 +361,17 @@ namespace UE::Mutable::Private
 
 		const uint8* Data = Program.GetOpArgsPointer(Item.At);
 
-		OP::ADDRESS VarAddress;
-		FMemory::Memcpy(&VarAddress, Data, sizeof(OP::ADDRESS));
-		Data += sizeof(OP::ADDRESS);
+		FOperation::ADDRESS VarAddress;
+		FMemory::Memcpy(&VarAddress, Data, sizeof(FOperation::ADDRESS));
+		Data += sizeof(FOperation::ADDRESS);
 
-		OP::ADDRESS DefAddress;
-		FMemory::Memcpy(&DefAddress, Data, sizeof(OP::ADDRESS));
-		Data += sizeof(OP::ADDRESS);
+		FOperation::ADDRESS DefAddress;
+		FMemory::Memcpy(&DefAddress, Data, sizeof(FOperation::ADDRESS));
+		Data += sizeof(FOperation::ADDRESS);
 
-		OP::FSwitchCaseDescriptor CaseDesc;
-		FMemory::Memcpy(&CaseDesc, Data, sizeof(OP::FSwitchCaseDescriptor));
-		Data += sizeof(OP::FSwitchCaseDescriptor);
+		FOperation::FSwitchCaseDescriptor CaseDesc;
+		FMemory::Memcpy(&CaseDesc, Data, sizeof(FOperation::FSwitchCaseDescriptor));
+		Data += sizeof(FOperation::FSwitchCaseDescriptor);
 
 		switch (Item.Stage)
 		{
@@ -415,7 +410,7 @@ namespace UE::Mutable::Private
 			// Get the variable result
 			int32 Var = LoadInt(FCacheAddress(VarAddress, Item));
 
-			OP::ADDRESS ValueAt = DefAddress;
+			FOperation::ADDRESS ValueAt = DefAddress;
 
 			if (!CaseDesc.bUseRanges)
 			{
@@ -425,9 +420,9 @@ namespace UE::Mutable::Private
 					FMemory::Memcpy(&Condition, Data, sizeof(int32));
 					Data += sizeof(int32);
 
-					OP::ADDRESS CaseAt;
-					FMemory::Memcpy(&CaseAt, Data, sizeof(OP::ADDRESS));
-					Data += sizeof(OP::ADDRESS);
+					FOperation::ADDRESS CaseAt;
+					FMemory::Memcpy(&CaseAt, Data, sizeof(FOperation::ADDRESS));
+					Data += sizeof(FOperation::ADDRESS);
 
 					if (CaseAt && Var == (int32)Condition)
 					{
@@ -448,9 +443,9 @@ namespace UE::Mutable::Private
 					FMemory::Memcpy(&RangeSize, Data, sizeof(uint32));
 					Data += sizeof(uint32);
 
-					OP::ADDRESS CaseAt;
-					FMemory::Memcpy(&CaseAt, Data, sizeof(OP::ADDRESS));
-					Data += sizeof(OP::ADDRESS);
+					FOperation::ADDRESS CaseAt;
+					FMemory::Memcpy(&CaseAt, Data, sizeof(FOperation::ADDRESS));
+					Data += sizeof(FOperation::ADDRESS);
 
 					if (CaseAt && Var >= ConditionStart && Var < int32(ConditionStart + RangeSize))
 					{
@@ -470,7 +465,7 @@ namespace UE::Mutable::Private
 
         case 2:
         {
-			OP::ADDRESS ResultAt = OP::ADDRESS(Item.CustomState);
+			FOperation::ADDRESS ResultAt = FOperation::ADDRESS(Item.CustomState);
 
             // Store the final result
             FCacheAddress ItemAddress(Item);
@@ -515,7 +510,7 @@ namespace UE::Mutable::Private
         {
         	MUTABLE_CPUPROFILER_SCOPE(IN_ADDSTRING);
         		
-			OP::InstanceAddArgs Args = Program.GetOpArgs<OP::InstanceAddArgs>( Item.At );
+			FOperation::InstanceAddArgs Args = Program.GetOpArgs<FOperation::InstanceAddArgs>( Item.At );
             switch ( Item.Stage )
             {
             case 0:
@@ -541,7 +536,7 @@ namespace UE::Mutable::Private
                 {
                     TManagedPtr<const String> value = LoadString( FCacheAddress( Args.value, Item ) );
 
-                    OP::ADDRESS nameAd = Args.name;
+                    FOperation::ADDRESS nameAd = Args.name;
                     check( nameAd < (uint32)Program.ConstantStrings.Num() );
                     const FString& Name = Program.ConstantStrings[nameAd];
 
@@ -562,7 +557,7 @@ namespace UE::Mutable::Private
         {
         	MUTABLE_CPUPROFILER_SCOPE(IN_ADDCOMPONENT);
         		
-			OP::InstanceAddComponentArgs Args = Program.GetOpArgs<OP::InstanceAddComponentArgs>(Item.At);
+			FOperation::InstanceAddComponentArgs Args = Program.GetOpArgs<FOperation::InstanceAddComponentArgs>(Item.At);
             switch (Item.Stage)
             {
             case 0:
@@ -614,7 +609,7 @@ namespace UE::Mutable::Private
         {
         	MUTABLE_CPUPROFILER_SCOPE(IN_ADDSURFACE);
         		
-			OP::InstanceAddArgs Args = Program.GetOpArgs<OP::InstanceAddArgs>(Item.At);
+			FOperation::InstanceAddArgs Args = Program.GetOpArgs<FOperation::InstanceAddArgs>(Item.At);
             switch (Item.Stage)
             {
             case 0:
@@ -659,7 +654,7 @@ namespace UE::Mutable::Private
                     }
 
                     // Name
-                    OP::ADDRESS nameAd = Args.name;
+                    FOperation::ADDRESS nameAd = Args.name;
                     check( nameAd < (uint32)Program.ConstantStrings.Num() );
                     const FString& Name = Program.ConstantStrings[ nameAd ];
                     Result->SetSurfaceName( 0, 0, sindex, FName(Name) );
@@ -696,9 +691,9 @@ namespace UE::Mutable::Private
                 TArray<FScheduledOp> deps;
                 for (uint8 LODIndex=0; LODIndex <LODCount; ++LODIndex)
                 {
-					OP::ADDRESS LODAddress;
-					FMemory::Memcpy(&LODAddress, Data, sizeof(OP::ADDRESS));
-					Data += sizeof(OP::ADDRESS);
+					FOperation::ADDRESS LODAddress;
+					FMemory::Memcpy(&LODAddress, Data, sizeof(FOperation::ADDRESS));
+					Data += sizeof(FOperation::ADDRESS);
 
                     if (LODAddress)
                     {
@@ -724,9 +719,9 @@ namespace UE::Mutable::Private
 
                 for (uint8 LODIndex = 0; LODIndex < LODCount; ++LODIndex)
                 {
-					OP::ADDRESS LODAddress;
-					FMemory::Memcpy(&LODAddress, Data, sizeof(OP::ADDRESS));
-					Data += sizeof(OP::ADDRESS);
+					FOperation::ADDRESS LODAddress;
+					FMemory::Memcpy(&LODAddress, Data, sizeof(FOperation::ADDRESS));
+					Data += sizeof(FOperation::ADDRESS);
 					
 					if ( LODAddress )
                     {						
@@ -763,7 +758,7 @@ namespace UE::Mutable::Private
 		{
         	MUTABLE_CPUPROFILER_SCOPE(IN_ADDEXTENSIONDATA);
         		
-			OP::InstanceAddExtensionDataArgs Args = Program.GetOpArgs<OP::InstanceAddExtensionDataArgs>(Item.At);
+			FOperation::InstanceAddExtensionDataArgs Args = Program.GetOpArgs<FOperation::InstanceAddExtensionDataArgs>(Item.At);
 			switch (Item.Stage)
 			{
 				case 0:
@@ -826,7 +821,7 @@ namespace UE::Mutable::Private
         {
         	MUTABLE_CPUPROFILER_SCOPE(IN_ADDMESH);
         		
-			OP::InstanceAddArgs Args = Program.GetOpArgs<OP::InstanceAddArgs>(Item.At);
+			FOperation::InstanceAddArgs Args = Program.GetOpArgs<FOperation::InstanceAddArgs>(Item.At);
             switch (Item.Stage)
             {
             case 0:
@@ -850,7 +845,7 @@ namespace UE::Mutable::Private
                 if ( Args.value )
                 {
 					FMeshId MeshId = FSystem::GetMeshId(LiveInstance->MeshIdRegistry, Model, Parameters, Args.value);
-					OP::ADDRESS NameAd = Args.name;
+					FOperation::ADDRESS NameAd = Args.name;
 					check(NameAd < (uint32)Program.ConstantStrings.Num());
 					const FString& Name = Program.ConstantStrings[NameAd];
 					Result->SetMesh(0, 0, MeshId, FName(Name));
@@ -869,7 +864,7 @@ namespace UE::Mutable::Private
         	{
         		MUTABLE_CPUPROFILER_SCOPE(IN_ADDOVERLAYMATERIAL);
         		
-        		OP::InstanceAddOverlayMaterialArgs Args = Program.GetOpArgs<OP::InstanceAddOverlayMaterialArgs>(Item.At);
+        		FOperation::InstanceAddOverlayMaterialArgs Args = Program.GetOpArgs<FOperation::InstanceAddOverlayMaterialArgs>(Item.At);
         		switch (Item.Stage)
         		{
         		case 0:
@@ -924,7 +919,7 @@ namespace UE::Mutable::Private
         	{
         		MUTABLE_CPUPROFILER_SCOPE(IN_ADDOVERRIDEMATERIAL);
         		
-        		OP::InstanceAddOverrideMaterialArgs Args = Program.GetOpArgs<OP::InstanceAddOverrideMaterialArgs>(Item.At);
+        		FOperation::InstanceAddOverrideMaterialArgs Args = Program.GetOpArgs<FOperation::InstanceAddOverrideMaterialArgs>(Item.At);
         		switch (Item.Stage)
         		{
         		case 0:
@@ -971,7 +966,7 @@ namespace UE::Mutable::Private
         	{
         		MUTABLE_CPUPROFILER_SCOPE(IN_ADDMATERIAL);
         		
-        		OP::InstanceAddMaterialArgs Args = Program.GetOpArgs<OP::InstanceAddMaterialArgs>(Item.At);
+        		FOperation::InstanceAddMaterialArgs Args = Program.GetOpArgs<FOperation::InstanceAddMaterialArgs>(Item.At);
         		switch (Item.Stage)
         		{
         		case 0:
@@ -1016,7 +1011,7 @@ namespace UE::Mutable::Private
         	{
         		MUTABLE_CPUPROFILER_SCOPE(IN_ADDSKELETALMESH);
         		
-        		OP::InstanceAddArgs Args = Program.GetOpArgs<OP::InstanceAddArgs>(Item.At);
+        		FOperation::InstanceAddArgs Args = Program.GetOpArgs<FOperation::InstanceAddArgs>(Item.At);
         		switch (Item.Stage)
         		{
         		case 0:
@@ -1082,7 +1077,7 @@ namespace UE::Mutable::Private
 		{
         	MUTABLE_CPUPROFILER_SCOPE(ED_CONSTANT);
         		
-			OP::ExternalDataConstantArgs Args = Program.GetOpArgs<OP::ExternalDataConstantArgs>(Item.At);
+			FOperation::ExternalDataConstantArgs Args = Program.GetOpArgs<FOperation::ExternalDataConstantArgs>(Item.At);
 
 			// Assume the ROM has been loaded previously
 			TManagedPtr<UE::Mutable::Private::FExtensionData> ExtensionData = MakeManaged<UE::Mutable::Private::FExtensionData>();
@@ -1096,7 +1091,7 @@ namespace UE::Mutable::Private
 		{
         	MUTABLE_CPUPROFILER_SCOPE(MI_CONSTANT);
         		
-			OP::MaterialConstantArgs Args = Program.GetOpArgs<OP::MaterialConstantArgs>(Item.At);
+			FOperation::MaterialConstantArgs Args = Program.GetOpArgs<FOperation::MaterialConstantArgs>(Item.At);
 
         	TManagedPtr<const FMaterial> SourceConst = Program.ConstantMaterials[Args.Value];
        		check(SourceConst);
@@ -1135,7 +1130,7 @@ namespace UE::Mutable::Private
         
         case EOpType::ME_APPLYLAYOUT:
         {
-			OP::MeshApplyLayoutArgs Args = Program.GetOpArgs<OP::MeshApplyLayoutArgs>(Item.At);
+			FOperation::MeshApplyLayoutArgs Args = Program.GetOpArgs<FOperation::MeshApplyLayoutArgs>(Item.At);
 			switch (Item.Stage)
 			{
 			case 0:
@@ -1196,7 +1191,7 @@ namespace UE::Mutable::Private
 
 		case EOpType::ME_PREPARELAYOUT:
 		{
-			OP::MeshPrepareLayoutArgs Args = Program.GetOpArgs<OP::MeshPrepareLayoutArgs>(Item.At);
+			FOperation::MeshPrepareLayoutArgs Args = Program.GetOpArgs<FOperation::MeshPrepareLayoutArgs>(Item.At);
 			switch (Item.Stage)
 			{
 			case 0:
@@ -1256,13 +1251,13 @@ namespace UE::Mutable::Private
         {
 			const uint8* data = Program.GetOpArgsPointer(Item.At);
 
-			OP::ADDRESS BaseAt = 0;
-			FMemory::Memcpy(&BaseAt, data, sizeof(OP::ADDRESS)); 
-			data += sizeof(OP::ADDRESS);
+			FOperation::ADDRESS BaseAt = 0;
+			FMemory::Memcpy(&BaseAt, data, sizeof(FOperation::ADDRESS)); 
+			data += sizeof(FOperation::ADDRESS);
 
-			OP::ADDRESS TargetAt = 0;
-			FMemory::Memcpy(&TargetAt, data, sizeof(OP::ADDRESS)); 
-			data += sizeof(OP::ADDRESS);
+			FOperation::ADDRESS TargetAt = 0;
+			FMemory::Memcpy(&TargetAt, data, sizeof(FOperation::ADDRESS)); 
+			data += sizeof(FOperation::ADDRESS);
 
             switch (Item.Stage)
             {
@@ -1342,7 +1337,7 @@ namespace UE::Mutable::Private
 
         case EOpType::ME_MORPH:
         {
-			const OP::MeshMorphArgs Args = Program.GetOpArgs<OP::MeshMorphArgs>(Item.At);
+			const FOperation::MeshMorphArgs Args = Program.GetOpArgs<FOperation::MeshMorphArgs>(Item.At);
 
 			switch (Item.Stage)
             {
@@ -1413,7 +1408,7 @@ namespace UE::Mutable::Private
 
         case EOpType::ME_MERGE:
         {
-			OP::MeshMergeArgs Args = Program.GetOpArgs<OP::MeshMergeArgs>(Item.At);
+			FOperation::MeshMergeArgs Args = Program.GetOpArgs<FOperation::MeshMergeArgs>(Item.At);
             switch (Item.Stage)
             {
             case 0:
@@ -1486,7 +1481,7 @@ namespace UE::Mutable::Private
 
         case EOpType::ME_MASKCLIPMESH:
         {
-			OP::MeshMaskClipMeshArgs Args = Program.GetOpArgs<OP::MeshMaskClipMeshArgs>(Item.At);
+			FOperation::MeshMaskClipMeshArgs Args = Program.GetOpArgs<FOperation::MeshMaskClipMeshArgs>(Item.At);
             switch (Item.Stage)
             {
             case 0:
@@ -1552,7 +1547,7 @@ namespace UE::Mutable::Private
 
 		case EOpType::ME_MASKCLIPUVMASK:
 		{
-			OP::MeshMaskClipUVMaskArgs Args = Program.GetOpArgs<OP::MeshMaskClipUVMaskArgs>(Item.At);
+			FOperation::MeshMaskClipUVMaskArgs Args = Program.GetOpArgs<FOperation::MeshMaskClipUVMaskArgs>(Item.At);
 			switch (Item.Stage)
 			{
 			case 0:
@@ -1638,7 +1633,7 @@ namespace UE::Mutable::Private
 
 		case EOpType::ME_MASKDIFF:
 		{
-			OP::MeshMaskDiffArgs Args = Program.GetOpArgs<OP::MeshMaskDiffArgs>(Item.At);
+			FOperation::MeshMaskDiffArgs Args = Program.GetOpArgs<FOperation::MeshMaskDiffArgs>(Item.At);
 			switch (Item.Stage)
 			{
 			case 0:
@@ -1704,7 +1699,7 @@ namespace UE::Mutable::Private
 
         case EOpType::ME_FORMAT:
         {
-			OP::MeshFormatArgs Args = Program.GetOpArgs<OP::MeshFormatArgs>(Item.At);
+			FOperation::MeshFormatArgs Args = Program.GetOpArgs<FOperation::MeshFormatArgs>(Item.At);
             switch (Item.Stage)
             {
             case 0:
@@ -1734,7 +1729,7 @@ namespace UE::Mutable::Private
 				else if (Source)
 				{
 					uint8 Flags = Args.Flags;
-					if (!Format && !(Flags & OP::MeshFormatArgs::ResetBufferIndices))
+					if (!Format && !(Flags & FOperation::MeshFormatArgs::ResetBufferIndices))
 					{
 						StoreMesh(Item, Source);
 					}
@@ -1742,7 +1737,7 @@ namespace UE::Mutable::Private
 					{
 						TManagedPtr<FMesh> Result = CloneOrTakeOver(MoveTemp(Source));
 
-						if (Flags & OP::MeshFormatArgs::ResetBufferIndices)
+						if (Flags & FOperation::MeshFormatArgs::ResetBufferIndices)
 						{
 							Result->ResetBufferIndices();
 						}
@@ -1756,18 +1751,18 @@ namespace UE::Mutable::Private
 						bool bOutSuccess = false;
 						MeshFormat(Result.Get(), Source.Get(), Format.Get(),
 							true,
-							(Flags & OP::MeshFormatArgs::Vertex) != 0,
-							(Flags & OP::MeshFormatArgs::Index) != 0,
-							(Flags & OP::MeshFormatArgs::IgnoreMissing) != 0,
+							(Flags & FOperation::MeshFormatArgs::Vertex) != 0,
+							(Flags & FOperation::MeshFormatArgs::Index) != 0,
+							(Flags & FOperation::MeshFormatArgs::IgnoreMissing) != 0,
 							bOutSuccess);
 						check(bOutSuccess);
 
-						if (Flags & OP::MeshFormatArgs::ResetBufferIndices)
+						if (Flags & FOperation::MeshFormatArgs::ResetBufferIndices)
 						{
 							Result->ResetBufferIndices();
 						}
 
-						if (Flags & OP::MeshFormatArgs::OptimizeBuffers)
+						if (Flags & FOperation::MeshFormatArgs::OptimizeBuffers)
 						{
 							MUTABLE_CPUPROFILER_SCOPE(MeshOptimizeBuffers)
 							MeshOptimizeBuffers(Result.Get());
@@ -1794,9 +1789,9 @@ namespace UE::Mutable::Private
         {
             const uint8* Data = Program.GetOpArgsPointer(Item.At);
 
-            OP::ADDRESS SourceAddress;
-            FMemory::Memcpy(&SourceAddress, Data, sizeof(OP::ADDRESS));
-            Data += sizeof(OP::ADDRESS);
+            FOperation::ADDRESS SourceAddress;
+            FMemory::Memcpy(&SourceAddress, Data, sizeof(FOperation::ADDRESS));
+            Data += sizeof(FOperation::ADDRESS);
 
             uint16 LayoutIndex;
 			FMemory::Memcpy(&LayoutIndex, Data, sizeof(uint16));
@@ -1874,7 +1869,7 @@ namespace UE::Mutable::Private
 
         case EOpType::ME_TRANSFORM:
         {
-			OP::MeshTransformArgs Args = Program.GetOpArgs<OP::MeshTransformArgs>(Item.At);
+			FOperation::MeshTransformArgs Args = Program.GetOpArgs<FOperation::MeshTransformArgs>(Item.At);
             switch (Item.Stage)
             {
             case 0:
@@ -1931,7 +1926,7 @@ namespace UE::Mutable::Private
 
         case EOpType::ME_CLIPMORPHPLANE:
         {
-			OP::MeshClipMorphPlaneArgs Args = Program.GetOpArgs<OP::MeshClipMorphPlaneArgs>(Item.At);
+			FOperation::MeshClipMorphPlaneArgs Args = Program.GetOpArgs<FOperation::MeshClipMorphPlaneArgs>(Item.At);
             switch (Item.Stage)
             {
             case 0:
@@ -2047,7 +2042,7 @@ namespace UE::Mutable::Private
 
         case EOpType::ME_CLIPWITHMESH:
         {
-			OP::MeshClipWithMeshArgs Args = Program.GetOpArgs<OP::MeshClipWithMeshArgs>(Item.At);
+			FOperation::MeshClipWithMeshArgs Args = Program.GetOpArgs<FOperation::MeshClipWithMeshArgs>(Item.At);
             switch (Item.Stage)
             {
             case 0:
@@ -2120,7 +2115,7 @@ namespace UE::Mutable::Private
         }
 		case EOpType::ME_CLIPDEFORM:
 		{
-			OP::MeshClipDeformArgs Args = Program.GetOpArgs<OP::MeshClipDeformArgs>(Item.At);
+			FOperation::MeshClipDeformArgs Args = Program.GetOpArgs<FOperation::MeshClipDeformArgs>(Item.At);
 			switch (Item.Stage)
 			{
 			case 0:
@@ -2193,7 +2188,7 @@ namespace UE::Mutable::Private
 
         case EOpType::ME_APPLYPOSE:
         {
-			OP::MeshApplyPoseArgs Args = Program.GetOpArgs<OP::MeshApplyPoseArgs>(Item.At);
+			FOperation::MeshApplyPoseArgs Args = Program.GetOpArgs<FOperation::MeshApplyPoseArgs>(Item.At);
             switch (Item.Stage)
             {
             case 0:
@@ -2251,7 +2246,7 @@ namespace UE::Mutable::Private
 
 		case EOpType::ME_BINDSHAPE:
 		{
-			OP::MeshBindShapeArgs Args = Program.GetOpArgs<OP::MeshBindShapeArgs>(Item.At);
+			FOperation::MeshBindShapeArgs Args = Program.GetOpArgs<FOperation::MeshBindShapeArgs>(Item.At);
 			const uint8* Data = Program.GetOpArgsPointer(Item.At);
 
 			constexpr uint8 ShapeContentFilter = (uint8)(EMeshContentFlags::GeometryData | EMeshContentFlags::PoseData);
@@ -2431,7 +2426,7 @@ namespace UE::Mutable::Private
 
 		case EOpType::ME_APPLYSHAPE:
 		{
-			OP::MeshApplyShapeArgs Args = Program.GetOpArgs<OP::MeshApplyShapeArgs>(Item.At);
+			FOperation::MeshApplyShapeArgs Args = Program.GetOpArgs<FOperation::MeshApplyShapeArgs>(Item.At);
 
 			EMeshBindShapeFlags ReshapeFlags = static_cast<EMeshBindShapeFlags>(Args.flags);
 
@@ -2537,7 +2532,7 @@ namespace UE::Mutable::Private
 
 		case EOpType::ME_MORPHRESHAPE:
 		{
-			OP::MeshMorphReshapeArgs Args = Program.GetOpArgs<OP::MeshMorphReshapeArgs>(Item.At);
+			FOperation::MeshMorphReshapeArgs Args = Program.GetOpArgs<FOperation::MeshMorphReshapeArgs>(Item.At);
 			switch(Item.Stage)
 			{
 			case 0:
@@ -2601,16 +2596,16 @@ namespace UE::Mutable::Private
             // TODO: Partial decode for each stage
             const uint8* Data = Program.GetOpArgsPointer(Item.At);
 
-            OP::ADDRESS SourceAddress;
-            FMemory::Memcpy(&SourceAddress, Data, sizeof(OP::ADDRESS)); 
-			Data += sizeof(OP::ADDRESS);
+            FOperation::ADDRESS SourceAddress;
+            FMemory::Memcpy(&SourceAddress, Data, sizeof(FOperation::ADDRESS)); 
+			Data += sizeof(FOperation::ADDRESS);
 
 			EFaceCullStrategy FaceCullStrategy;
 			FMemory::Memcpy(&FaceCullStrategy, Data, sizeof(EFaceCullStrategy));
 			Data += sizeof(EFaceCullStrategy);
 
             TArray<FScheduledOp, TInlineAllocator<4>> Conditions;
-			TArray<OP::ADDRESS, TInlineAllocator<4>> Masks;
+			TArray<FOperation::ADDRESS, TInlineAllocator<4>> Masks;
 
             uint16 NumRemoves;
 			FMemory::Memcpy(&NumRemoves, Data, sizeof(uint16)); 
@@ -2618,15 +2613,15 @@ namespace UE::Mutable::Private
 
             for(uint16 RemoveIndex = 0; RemoveIndex < NumRemoves; ++RemoveIndex)
             {
-                OP::ADDRESS Condition;
-				FMemory::Memcpy(&Condition, Data, sizeof(OP::ADDRESS)); 
-				Data += sizeof(OP::ADDRESS);
+                FOperation::ADDRESS Condition;
+				FMemory::Memcpy(&Condition, Data, sizeof(FOperation::ADDRESS)); 
+				Data += sizeof(FOperation::ADDRESS);
                 
 				Conditions.Emplace(Condition, Item);
 
-                OP::ADDRESS Mask;
-				FMemory::Memcpy(&Mask, Data, sizeof(OP::ADDRESS)); 
-				Data += sizeof(OP::ADDRESS);
+                FOperation::ADDRESS Mask;
+				FMemory::Memcpy(&Mask, Data, sizeof(FOperation::ADDRESS)); 
+				Data += sizeof(FOperation::ADDRESS);
 
                 Masks.Add(Mask);
             }
@@ -2679,7 +2674,8 @@ namespace UE::Mutable::Private
 						bool bValue = true;
 						if (Conditions[RemoveIndex].At)
 						{
-							bValue = LoadBool(FCacheAddress(Conditions[RemoveIndex].At, Item));
+							bool bSuccess = GetMemory().GetBoolIfSet(FCacheAddress(Conditions[RemoveIndex].At, Item), bValue);
+							check(bSuccess);
 						}
 
 						if (bValue)
@@ -2747,7 +2743,7 @@ namespace UE::Mutable::Private
 		{
 			MUTABLE_CPUPROFILER_SCOPE(ME_ADDMETADATA)
 
-			const OP::MeshAddMetadataArgs OpArgs = Program.GetOpArgs<OP::MeshAddMetadataArgs>(Item.At);
+			const FOperation::MeshAddMetadataArgs OpArgs = Program.GetOpArgs<FOperation::MeshAddMetadataArgs>(Item.At);
 
 			// Schedule next stages
 			switch (Item.Stage)
@@ -2781,7 +2777,7 @@ namespace UE::Mutable::Private
 						Result->SkeletonObjects.Add(FindPassthrough<USkeleton>(OpArgs.SkeletonId));
 					}
 
-					using OpEnumFlags = OP::MeshAddMetadataArgs::EnumFlags;
+					using OpEnumFlags = FOperation::MeshAddMetadataArgs::EnumFlags;
 
 					// Append Tags
 					if (EnumHasAnyFlags(OpArgs.Flags, OpEnumFlags::IsGameplayTagList))
@@ -2945,7 +2941,7 @@ namespace UE::Mutable::Private
 
         case EOpType::ME_PROJECT:
         {
-			OP::MeshProjectArgs Args = Program.GetOpArgs<OP::MeshProjectArgs>(Item.At);
+			FOperation::MeshProjectArgs Args = Program.GetOpArgs<FOperation::MeshProjectArgs>(Item.At);
             switch (Item.Stage)
             {
             case 0:
@@ -3017,7 +3013,7 @@ namespace UE::Mutable::Private
 
 		case EOpType::ME_OPTIMIZESKINNING:
 		{
-			OP::MeshOptimizeSkinningArgs Args = Program.GetOpArgs<OP::MeshOptimizeSkinningArgs>(Item.At);
+			FOperation::MeshOptimizeSkinningArgs Args = Program.GetOpArgs<FOperation::MeshOptimizeSkinningArgs>(Item.At);
 			switch (Item.Stage)
 			{
 			case 0:
@@ -3070,7 +3066,7 @@ namespace UE::Mutable::Private
 
         case EOpType::ME_TRANSFORMWITHMESH:
 		{
-			OP::MeshTransformWithinMeshArgs Args = Program.GetOpArgs<OP::MeshTransformWithinMeshArgs>(Item.At);
+			FOperation::MeshTransformWithinMeshArgs Args = Program.GetOpArgs<FOperation::MeshTransformWithinMeshArgs>(Item.At);
 			switch (Item.Stage)
 			{
 			case 0:
@@ -3142,7 +3138,7 @@ namespace UE::Mutable::Private
 
 		case EOpType::ME_TRANSFORMWITHBONE:
 		{
-			OP::MeshTransformWithBoneArgs Args = Program.GetOpArgs<OP::MeshTransformWithBoneArgs>(Item.At);
+			FOperation::MeshTransformWithBoneArgs Args = Program.GetOpArgs<FOperation::MeshTransformWithBoneArgs>(Item.At);
 			switch (Item.Stage)
 			{
 			case 0:
@@ -3205,7 +3201,7 @@ namespace UE::Mutable::Private
 
         case EOpType::ME_SKELETALMESH_BREAK:
 		{
-			OP::MeshSkeletalMeshBreakArgs Args = Program.GetOpArgs<OP::MeshSkeletalMeshBreakArgs>(Item.At);
+			FOperation::MeshSkeletalMeshBreakArgs Args = Program.GetOpArgs<FOperation::MeshSkeletalMeshBreakArgs>(Item.At);
 			switch (Item.Stage)
 			{
 			case 0:
@@ -3247,7 +3243,7 @@ namespace UE::Mutable::Private
 
         case EOpType::IM_LAYERCOLOUR:
         {
-			OP::ImageLayerColourArgs Args = Program.GetOpArgs<OP::ImageLayerColourArgs>(Item.At);
+			FOperation::ImageLayerColourArgs Args = Program.GetOpArgs<FOperation::ImageLayerColourArgs>(Item.At);
             switch (Item.Stage)
             {
             case 0:
@@ -3271,7 +3267,7 @@ namespace UE::Mutable::Private
 
         case EOpType::IM_LAYER:
         {
-			OP::ImageLayerArgs Args = Program.GetOpArgs<OP::ImageLayerArgs>(Item.At);
+			FOperation::ImageLayerArgs Args = Program.GetOpArgs<FOperation::ImageLayerArgs>(Item.At);
 
 			if (ExecutionStrategy == EExecutionStrategy::MinimizeMemory)
 			{
@@ -3324,7 +3320,7 @@ namespace UE::Mutable::Private
 
         case EOpType::IM_MULTILAYER:
         {
-			OP::ImageMultiLayerArgs Args = Program.GetOpArgs<OP::ImageMultiLayerArgs>(Item.At);
+			FOperation::ImageMultiLayerArgs Args = Program.GetOpArgs<FOperation::ImageMultiLayerArgs>(Item.At);
             switch (Item.Stage)
             {
             case 0:
@@ -3397,9 +3393,8 @@ namespace UE::Mutable::Private
 					int32 CurrentIteration = 0;
 					FScheduledOp ItemForFirstLayerBlend = Item;
 
-					ExecutionIndex Index = GetMemory().GetRangeIndex(Item.ExecutionIndex);
-					Index.SetFromModelRangeIndex(Args.rangeId, CurrentIteration);
-					ItemForFirstLayerBlend.ExecutionIndex = GetMemory().GetRangeIndexIndex(Index);
+					ItemForFirstLayerBlend.ExecutionIndex = 
+							GetMemory().GetOrAddModifiedExecutionIndex(Item.ExecutionIndex, Args.rangeId, CurrentIteration);
 
 					AddOp(FScheduledOp(Item.At, Item, 2, HeapDataIndex), 
 							FScheduledOp(Args.base, Item), 
@@ -3426,12 +3421,13 @@ namespace UE::Mutable::Private
 
 				TManagedPtr<FImage> Base = StaticCastManagedPtr<FImage>(ConstCastManagedPtr<FResource>(Data.Resource));
 				 
-                FScheduledOp ItemForLayerBlend = Item;
-                ExecutionIndex Index = GetMemory().GetRangeIndex(Item.ExecutionIndex);
+  
 				
+                FScheduledOp ItemForLayerBlend = Item;
                 {
-                    Index.SetFromModelRangeIndex(Args.rangeId, CurrentIteration);
-                    ItemForLayerBlend.ExecutionIndex = GetMemory().GetRangeIndexIndex(Index);
+                    ItemForLayerBlend.ExecutionIndex = 
+							GetMemory().GetOrAddModifiedExecutionIndex(Item.ExecutionIndex, Args.rangeId, CurrentIteration);
+
 					ItemForLayerBlend.CustomState = 0;
 
                     TManagedPtr<const FImage> Blended = LoadImage(FCacheAddress(Args.blended, ItemForLayerBlend));
@@ -3561,9 +3557,8 @@ namespace UE::Mutable::Private
 					// Request a new layer
 					++CurrentIteration;
 					FScheduledOp ItemForNextLayerBlend = Item;
-					ExecutionIndex NextIndex = GetMemory().GetRangeIndex(Item.ExecutionIndex);
-					NextIndex.SetFromModelRangeIndex(Args.rangeId, CurrentIteration);
-					ItemForNextLayerBlend.ExecutionIndex = GetMemory().GetRangeIndexIndex(NextIndex);
+					ItemForNextLayerBlend.ExecutionIndex = 
+							GetMemory().GetOrAddModifiedExecutionIndex(Item.ExecutionIndex, Args.rangeId, CurrentIteration);
 
 					AddOp(FScheduledOp(Item.At, Item, 2+CurrentIteration, Item.CustomState), 
 							FScheduledOp(Args.blended, ItemForNextLayerBlend), 
@@ -3581,7 +3576,7 @@ namespace UE::Mutable::Private
 
 		case EOpType::IM_NORMALCOMPOSITE:
 		{
-			OP::ImageNormalCompositeArgs Args = Program.GetOpArgs<OP::ImageNormalCompositeArgs>(Item.At);
+			FOperation::ImageNormalCompositeArgs Args = Program.GetOpArgs<FOperation::ImageNormalCompositeArgs>(Item.At);
 			switch (Item.Stage)
 			{
 			case 0:
@@ -3663,7 +3658,7 @@ namespace UE::Mutable::Private
 
         case EOpType::IM_PIXELFORMAT:
         {
-			OP::ImagePixelFormatArgs Args = Program.GetOpArgs<OP::ImagePixelFormatArgs>(Item.At);
+			FOperation::ImagePixelFormatArgs Args = Program.GetOpArgs<FOperation::ImagePixelFormatArgs>(Item.At);
             switch (Item.Stage)
             {
             case 0:
@@ -3686,7 +3681,7 @@ namespace UE::Mutable::Private
 
         case EOpType::IM_MIPMAP:
         {
-			OP::ImageMipmapArgs Args = Program.GetOpArgs<OP::ImageMipmapArgs>(Item.At);
+			FOperation::ImageMipmapArgs Args = Program.GetOpArgs<FOperation::ImageMipmapArgs>(Item.At);
             switch (Item.Stage)
             {
             case 0:
@@ -3707,7 +3702,7 @@ namespace UE::Mutable::Private
 
         case EOpType::IM_RESIZE:
         {
-			OP::ImageResizeArgs Args = Program.GetOpArgs<OP::ImageResizeArgs>(Item.At);
+			FOperation::ImageResizeArgs Args = Program.GetOpArgs<FOperation::ImageResizeArgs>(Item.At);
             switch (Item.Stage)
             {
             case 0:
@@ -3729,7 +3724,7 @@ namespace UE::Mutable::Private
 
         case EOpType::IM_RESIZELIKE:
         {
-			OP::ImageResizeLikeArgs Args = Program.GetOpArgs<OP::ImageResizeLikeArgs>(Item.At);
+			FOperation::ImageResizeLikeArgs Args = Program.GetOpArgs<FOperation::ImageResizeLikeArgs>(Item.At);
             switch (Item.Stage)
             {
             case 0:
@@ -3788,7 +3783,7 @@ namespace UE::Mutable::Private
 
         case EOpType::IM_RESIZEREL:
         {
-			OP::ImageResizeRelArgs Args = Program.GetOpArgs<OP::ImageResizeRelArgs>(Item.At);
+			FOperation::ImageResizeRelArgs Args = Program.GetOpArgs<FOperation::ImageResizeRelArgs>(Item.At);
             switch (Item.Stage)
             {
             case 0:
@@ -3810,7 +3805,7 @@ namespace UE::Mutable::Private
         }
         case EOpType::IM_BLANKLAYOUT:
         {
-			OP::ImageBlankLayoutArgs Args = Program.GetOpArgs<OP::ImageBlankLayoutArgs>(Item.At);
+			FOperation::ImageBlankLayoutArgs Args = Program.GetOpArgs<FOperation::ImageBlankLayoutArgs>(Item.At);
             switch (Item.Stage)
             {
             case 0:
@@ -3886,7 +3881,7 @@ namespace UE::Mutable::Private
 
         case EOpType::IM_COMPOSE:
         {
-			OP::ImageComposeArgs Args = Program.GetOpArgs<OP::ImageComposeArgs>(Item.At);
+			FOperation::ImageComposeArgs Args = Program.GetOpArgs<FOperation::ImageComposeArgs>(Item.At);
 
 			if (ExecutionStrategy == EExecutionStrategy::MinimizeMemory)
 			{
@@ -3982,7 +3977,7 @@ namespace UE::Mutable::Private
 		
 		case EOpType::IM_MULTICOMPOSE:
 		{
-			OP::ImageMultiComposeArgs Args = Program.GetOpArgs<OP::ImageMultiComposeArgs>(Item.At);
+			FOperation::ImageMultiComposeArgs Args = Program.GetOpArgs<FOperation::ImageMultiComposeArgs>(Item.At);
 			constexpr float LODRoundTolerance = 0.1f;
 			
 			switch (Item.Stage)
@@ -4354,7 +4349,7 @@ namespace UE::Mutable::Private
 
         case EOpType::IM_INTERPOLATE:
         {
-			OP::ImageInterpolateArgs Args = Program.GetOpArgs<OP::ImageInterpolateArgs>(Item.At);
+			FOperation::ImageInterpolateArgs Args = Program.GetOpArgs<FOperation::ImageInterpolateArgs>(Item.At);
             switch (Item.Stage)
             {
             case 0:
@@ -4526,7 +4521,7 @@ namespace UE::Mutable::Private
 
         case EOpType::IM_SATURATE:
         {
-			OP::ImageSaturateArgs Args = Program.GetOpArgs<OP::ImageSaturateArgs>(Item.At);
+			FOperation::ImageSaturateArgs Args = Program.GetOpArgs<FOperation::ImageSaturateArgs>(Item.At);
             switch (Item.Stage)
             {
             case 0:
@@ -4550,7 +4545,7 @@ namespace UE::Mutable::Private
 
         case EOpType::IM_LUMINANCE:
         {
-			OP::ImageLuminanceArgs Args = Program.GetOpArgs<OP::ImageLuminanceArgs>(Item.At);
+			FOperation::ImageLuminanceArgs Args = Program.GetOpArgs<FOperation::ImageLuminanceArgs>(Item.At);
             switch (Item.Stage)
             {
             case 0:
@@ -4580,7 +4575,7 @@ namespace UE::Mutable::Private
 
         case EOpType::IM_SWIZZLE:
         {
-			OP::ImageSwizzleArgs Args = Program.GetOpArgs<OP::ImageSwizzleArgs>(Item.At);
+			FOperation::ImageSwizzleArgs Args = Program.GetOpArgs<FOperation::ImageSwizzleArgs>(Item.At);
             switch (Item.Stage)
             {
             case 0:
@@ -4605,7 +4600,7 @@ namespace UE::Mutable::Private
 
         case EOpType::IM_COLOURMAP:
         {
-			OP::ImageColourMapArgs Args = Program.GetOpArgs<OP::ImageColourMapArgs>(Item.At);
+			FOperation::ImageColourMapArgs Args = Program.GetOpArgs<FOperation::ImageColourMapArgs>(Item.At);
             switch (Item.Stage)
             {
             case 0:
@@ -4657,7 +4652,7 @@ namespace UE::Mutable::Private
 
         case EOpType::IM_BINARISE:
         {
-			OP::ImageBinariseArgs Args = Program.GetOpArgs<OP::ImageBinariseArgs>(Item.At);
+			FOperation::ImageBinariseArgs Args = Program.GetOpArgs<FOperation::ImageBinariseArgs>(Item.At);
             switch (Item.Stage)
             {
             case 0:
@@ -4690,7 +4685,7 @@ namespace UE::Mutable::Private
 
 		case EOpType::IM_INVERT:
 		{
-			OP::ImageInvertArgs Args = Program.GetOpArgs<OP::ImageInvertArgs>(Item.At);
+			FOperation::ImageInvertArgs Args = Program.GetOpArgs<FOperation::ImageInvertArgs>(Item.At);
 			switch (Item.Stage)
 			{
 			case 0:
@@ -4712,7 +4707,7 @@ namespace UE::Mutable::Private
 
         case EOpType::IM_PLAINCOLOUR:
         {
-			OP::ImagePlainColorArgs Args = Program.GetOpArgs<OP::ImagePlainColorArgs>(Item.At);
+			FOperation::ImagePlainColorArgs Args = Program.GetOpArgs<FOperation::ImagePlainColorArgs>(Item.At);
             switch (Item.Stage)
             {
             case 0:
@@ -4759,7 +4754,7 @@ namespace UE::Mutable::Private
 
 		case EOpType::IM_REFERENCE:
 		{
-			OP::ResourceReferenceArgs Args = Program.GetOpArgs<OP::ResourceReferenceArgs>(Item.At);
+			FOperation::ResourceReferenceArgs Args = Program.GetOpArgs<FOperation::ResourceReferenceArgs>(Item.At);
 			switch (Item.Stage)
 			{
 			case 0:
@@ -4788,7 +4783,7 @@ namespace UE::Mutable::Private
 
 		case EOpType::IM_PARAMETER:
         {
-			OP::ParameterArgs Args = Program.GetOpArgs<OP::ParameterArgs>(Item.At);
+			FOperation::ParameterArgs Args = Program.GetOpArgs<FOperation::ParameterArgs>(Item.At);
 			TSharedPtr<FRangeIndex> Index = BuildCurrentOpRangeIndex(Item, Args.variable);
 			UTexture* Image = Parameters.GetTextureValue(Args.variable, Index.Get());
 
@@ -4801,7 +4796,7 @@ namespace UE::Mutable::Private
 
         case EOpType::IM_CROP:
         {
-			OP::ImageCropArgs Args = Program.GetOpArgs<OP::ImageCropArgs>(Item.At);
+			FOperation::ImageCropArgs Args = Program.GetOpArgs<FOperation::ImageCropArgs>(Item.At);
             switch (Item.Stage)
             {
             case 0:
@@ -4852,7 +4847,7 @@ namespace UE::Mutable::Private
         case EOpType::IM_PATCH:
         {
 			// TODO: This is optimized for memory-usage but base and patch could be requested at the same time
-			OP::ImagePatchArgs Args = Program.GetOpArgs<OP::ImagePatchArgs>(Item.At);
+			FOperation::ImagePatchArgs Args = Program.GetOpArgs<FOperation::ImagePatchArgs>(Item.At);
             switch (Item.Stage)
             {
 			case 0:
@@ -4959,7 +4954,7 @@ namespace UE::Mutable::Private
 
         case EOpType::IM_RASTERMESH:
         {
-			OP::ImageRasterMeshArgs Args = Program.GetOpArgs<OP::ImageRasterMeshArgs>(Item.At);
+			FOperation::ImageRasterMeshArgs Args = Program.GetOpArgs<FOperation::ImageRasterMeshArgs>(Item.At);
 
 			constexpr uint8 MeshContentFilter = (uint8)(EMeshContentFlags::GeometryData | EMeshContentFlags::PoseData);
             switch (Item.Stage)
@@ -5294,7 +5289,7 @@ namespace UE::Mutable::Private
 
         case EOpType::IM_MAKEGROWMAP:
         {
-			OP::ImageMakeGrowMapArgs Args = Program.GetOpArgs<OP::ImageMakeGrowMapArgs>(Item.At);
+			FOperation::ImageMakeGrowMapArgs Args = Program.GetOpArgs<FOperation::ImageMakeGrowMapArgs>(Item.At);
             switch (Item.Stage)
             {
             case 0:
@@ -5327,7 +5322,7 @@ namespace UE::Mutable::Private
 
         case EOpType::IM_DISPLACE:
         {
-			OP::ImageDisplaceArgs Args = Program.GetOpArgs<OP::ImageDisplaceArgs>(Item.At);
+			FOperation::ImageDisplaceArgs Args = Program.GetOpArgs<FOperation::ImageDisplaceArgs>(Item.At);
             switch (Item.Stage)
             {
             case 0:
@@ -5392,7 +5387,7 @@ namespace UE::Mutable::Private
 
         case EOpType::IM_TRANSFORM:
         {
-            const OP::ImageTransformArgs Args = Program.GetOpArgs<OP::ImageTransformArgs>(Item.At);
+            const FOperation::ImageTransformArgs Args = Program.GetOpArgs<FOperation::ImageTransformArgs>(Item.At);
 
             switch (Item.Stage)
             {
@@ -5602,7 +5597,7 @@ namespace UE::Mutable::Private
 
 		case EOpType::IM_MATERIAL_BREAK:
 		{
-			const OP::MaterialBreakArgs Args = Program.GetOpArgs<OP::MaterialBreakArgs>(Item.At);
+			const FOperation::MaterialBreakArgs Args = Program.GetOpArgs<FOperation::MaterialBreakArgs>(Item.At);
 
 			switch (Item.Stage)
 			{
@@ -5623,16 +5618,16 @@ namespace UE::Mutable::Private
 
 				if (ImageParameterData)
 				{
+					const TVariant<FOperation::ADDRESS, TManagedPtr<const FImage>> ImageParameter = ImageParameterData->ImageParameter;
 					bParameterFound = true;
-					const TVariant<OP::ADDRESS, TManagedPtr<const FImage>> ImageParameter = ImageParameterData->ImageParameter;
 
-					if (ImageParameter.IsType<OP::ADDRESS>())
+					if (ImageParameter.IsType<FOperation::ADDRESS>())
 					{
 						FScheduledOp ItemNextStage(Item.At, Item, 2);
-						ItemNextStage.CustomState = ImageParameter.Get<OP::ADDRESS>();
+						ItemNextStage.CustomState = ImageParameter.Get<FOperation::ADDRESS>();
 
 						AddOp(ItemNextStage,
-							FScheduledOp(ImageParameter.Get<OP::ADDRESS>(), Item));
+							FScheduledOp(ImageParameter.Get<FOperation::ADDRESS>(), Item));
 					}
 					else
 					{
@@ -5659,11 +5654,11 @@ namespace UE::Mutable::Private
 					}
 				}
 				
-				if(!bParameterFound)
+				if (!bParameterFound)
 				{
 					LiveInstance->UpdateLogger->LogUpdateMessage(FString::Printf(TEXT("Expected Texture parameter [%s] not found in Break Material node."), *ParameterName.ToString()));
 
-					AbortRun();
+					GetMemory().SetAborted(FCacheAddress(Item));
 				}
 
 				break;
@@ -5681,8 +5676,8 @@ namespace UE::Mutable::Private
 				}
 				else
 				{
-					const TVariant<OP::ADDRESS, TManagedPtr<const FImage>>& ImageParameter = ImageParameterData->ImageParameter;
-					TManagedPtr<const FImage> pImage = LoadImage(FCacheAddress(ImageParameter.Get<OP::ADDRESS>(), Item));
+					const TVariant<FOperation::ADDRESS, TManagedPtr<const FImage>>& ImageParameter = ImageParameterData->ImageParameter;
+					TManagedPtr<const FImage> pImage = LoadImage(FCacheAddress(ImageParameter.Get<FOperation::ADDRESS>(), Item));
 
 					StoreImage(Item, pImage);
 				}
@@ -5699,7 +5694,7 @@ namespace UE::Mutable::Private
 
 		case EOpType::IM_PARAMETER_CONVERT:
         {
-        	OP::ImageParameterConvertArgs Args = Program.GetOpArgs<OP::ImageParameterConvertArgs>(Item.At);
+        	FOperation::ImageParameterConvertArgs Args = Program.GetOpArgs<FOperation::ImageParameterConvertArgs>(Item.At);
         	switch (Item.Stage)
         	{
         	case 0:
@@ -5729,7 +5724,6 @@ namespace UE::Mutable::Private
         }
     }	
 	
-	//---------------------------------------------------------------------------------------------
     TSharedPtr<FRangeIndex> BuildCurrentOpRangeIndex(const FScheduledOp& Item, const FParameters& Params, const FModel& InModel, FProgramCache& ProgramCache, int32 ParameterIndex)
     {
         if (!Item.ExecutionIndex)
@@ -5737,24 +5731,23 @@ namespace UE::Mutable::Private
             return nullptr;
         }
 
-        // \todo: optimise to avoid allocating the index here, we could access internal
+        // TODO: optimise to avoid allocating the index here, we could access internal
         // data directly.
-		TSharedPtr<FRangeIndex> Index = Params.NewRangeIndex( ParameterIndex );
+		TSharedPtr<FRangeIndex> Index = Params.NewRangeIndex(ParameterIndex);
         if (!Index)
         {
             return nullptr;
         }
 
 		const FProgram& Program = InModel.GetProgram();
-		const FParameterDesc& paramDesc = Program.Parameters[ ParameterIndex ];
-        for( int32 rangeIndexInParam = 0;
-             rangeIndexInParam<paramDesc.Ranges.Num();
-             ++rangeIndexInParam )
+		const FParameterDesc& ParamDesc = Program.Parameters[ParameterIndex];
+        for (int32 RangeIndexInParam = 0; RangeIndexInParam < ParamDesc.Ranges.Num(); ++RangeIndexInParam)
         {
-            uint32 rangeIndexInModel = paramDesc.Ranges[rangeIndexInParam];
-            const ExecutionIndex& currentIndex = ProgramCache.GetRangeIndex( Item.ExecutionIndex );
-            int32 Position = currentIndex.GetFromModelRangeIndex(rangeIndexInModel);
-			Index->Values[rangeIndexInParam] = Position;
+            uint32 RangeIndexInModel = ParamDesc.Ranges[RangeIndexInParam];
+            
+            int32 Position = ProgramCache.GetExecutionIndexValue(Item.ExecutionIndex, RangeIndexInModel);
+
+			Index->Values[RangeIndexInParam] = Position;
         }
 
         return Index;
@@ -5765,7 +5758,6 @@ namespace UE::Mutable::Private
 		return UE::Mutable::Private::BuildCurrentOpRangeIndex(Item, Parameters, Model, GetMemory(), ParameterIndex);
 	}
 
-    //---------------------------------------------------------------------------------------------
     void CodeRunner::RunCode_Bool(const FScheduledOp& Item)
     {
         MUTABLE_CPUPROFILER_SCOPE(RunCode_Bool);
@@ -5776,7 +5768,7 @@ namespace UE::Mutable::Private
 
         case EOpType::BO_CONSTANT:
         {
-			OP::BoolConstantArgs Args = Program.GetOpArgs<OP::BoolConstantArgs>(Item.At);
+			FOperation::BoolConstantArgs Args = Program.GetOpArgs<FOperation::BoolConstantArgs>(Item.At);
             bool result = Args.bValue;
             StoreBool( Item, result );
             break;
@@ -5784,7 +5776,7 @@ namespace UE::Mutable::Private
 
         case EOpType::BO_PARAMETER:
         {
-			OP::ParameterArgs Args = Program.GetOpArgs<OP::ParameterArgs>(Item.At);
+			FOperation::ParameterArgs Args = Program.GetOpArgs<FOperation::ParameterArgs>(Item.At);
             bool result = false;
 			TSharedPtr<FRangeIndex> Index = BuildCurrentOpRangeIndex( Item, Args.variable );
             result = Parameters.GetBoolValue( Args.variable, Index.Get() );
@@ -5794,52 +5786,52 @@ namespace UE::Mutable::Private
 
         case EOpType::BO_AND:
         {
-			OP::BoolBinaryArgs Args = Program.GetOpArgs<OP::BoolBinaryArgs>(Item.At);
+			FOperation::BoolBinaryArgs Args = Program.GetOpArgs<FOperation::BoolBinaryArgs>(Item.At);
             switch (Item.Stage)
             {
             case 0:
-                {
-                    // Try to avoid the op entirely if we have some children cached
-                    bool skip = false;
-                    if ( Args.A && GetMemory().IsValid( FCacheAddress(Args.A,Item) ) )
-                    {
-                         bool a = LoadBool( FCacheAddress(Args.A,Item) );
-                         if (!a)
-                         {
-                            StoreBool( Item, false );
-                            skip=true;
-                         }
-                    }
+			{
+				// Try to avoid the op entirely if we have some children cached
+				bool bSkip = false;
+				bool bAValue = false;
+				if (Args.A && GetMemory().GetBoolIfSet(FCacheAddress(Args.A, Item), bAValue))
+				{
+					 if (!bAValue)
+					 {
+						StoreBool(Item, false);
+						bSkip = true;
+					 }
+				}
 
-                    if ( !skip && Args.B && GetMemory().IsValid( FCacheAddress(Args.B,Item) ) )
-                    {
-                         bool b = LoadBool( FCacheAddress(Args.B,Item) );
-                         if (!b)
-                         {
-                            StoreBool( Item, false );
-                            skip=true;
-                         }
-                    }
+				bool bBValue = false;
+				if (!bSkip && Args.B && GetMemory().GetBoolIfSet(FCacheAddress(Args.B, Item), bBValue))
+				{
+					 if (!bBValue)
+					 {
+						StoreBool(Item, false);
+						bSkip = true;
+					 }
+				}
 
-                    if (!skip)
-                    {
-                        AddOp( FScheduledOp( Item.At, Item, 1),
-                               FScheduledOp( Args.A, Item));
-                    }
+				if (!bSkip)
+				{
+					AddOp(FScheduledOp(Item.At, Item, 1),
+						   FScheduledOp(Args.A, Item));
+				}
 				break;
-                }
+			}
 
             case 1:
             {
-                bool a = LoadBool(FCacheAddress(Args.A, Item));
-                if (!a)
+                bool bAValue = LoadBool(FCacheAddress(Args.A, Item));
+                if (!bAValue)
                 {
-                    StoreBool( Item, false );
+                    StoreBool(Item, false);
                 }
                 else
                 {
-                    AddOp( FScheduledOp( Item.At, Item, 2),
-                           FScheduledOp( Args.B, Item));
+                    AddOp(FScheduledOp(Item.At, Item, 2),
+                           FScheduledOp(Args.B, Item));
                 }
                 break;
             }
@@ -5847,8 +5839,8 @@ namespace UE::Mutable::Private
             case 2:
             {
                 // We arrived here because a is true
-                bool b = LoadBool(FCacheAddress(Args.B, Item));
-                StoreBool( Item, b );
+                bool bBValue = LoadBool(FCacheAddress(Args.B, Item));
+                StoreBool(Item, bBValue);
                 break;
             }
 
@@ -5860,52 +5852,53 @@ namespace UE::Mutable::Private
 
         case EOpType::BO_OR:
         {
-			OP::BoolBinaryArgs Args = Program.GetOpArgs<OP::BoolBinaryArgs>(Item.At);
+			FOperation::BoolBinaryArgs Args = Program.GetOpArgs<FOperation::BoolBinaryArgs>(Item.At);
             switch (Item.Stage)
             {
             case 0:
-                {
-                    // Try to avoid the op entirely if we have some children cached
-                    bool skip = false;
-                    if ( Args.A && GetMemory().IsValid( FCacheAddress(Args.A,Item) ) )
-                    {
-                         bool a = LoadBool( FCacheAddress(Args.A,Item) );
-                         if (a)
-                         {
-                            StoreBool( Item, true );
-                            skip=true;
-                         }
-                    }
+			{
+				// Try to avoid the op entirely if we have some children cached
+				bool bSkip = false;
 
-                    if ( !skip && Args.B && GetMemory().IsValid( FCacheAddress(Args.B,Item) ) )
-                    {
-                         bool b = LoadBool( FCacheAddress(Args.B,Item) );
-                         if (b)
-                         {
-                            StoreBool( Item, true );
-                            skip=true;
-                         }
-                    }
+				bool bAValue = false;
+				if (Args.A && GetMemory().GetBoolIfSet(FCacheAddress(Args.A, Item), bAValue))
+				{
+					 if (bAValue)
+					 {
+						StoreBool(Item, true);
+						bSkip = true;
+					 }
+				}
 
-                    if (!skip)
-                    {
-                        AddOp( FScheduledOp( Item.At, Item, 1),
-                               FScheduledOp( Args.A, Item));
-                    }
+				bool bBValue = false;
+				if (!bSkip && Args.B && GetMemory().GetBoolIfSet(FCacheAddress(Args.B, Item), bBValue))
+				{
+					 if (bBValue)
+					 {
+						StoreBool(Item, true);
+						bSkip = true;
+					 }
+				}
+
+				if (!bSkip)
+				{
+					AddOp(FScheduledOp(Item.At, Item, 1),
+						   FScheduledOp(Args.A, Item));
+				}
 				break;
-                }
+			}
 
             case 1:
             {
-                bool a = Args.A ? LoadBool( FCacheAddress(Args.A,Item) ) : false;
-                if (a)
+                bool bAValue = Args.A ? LoadBool(FCacheAddress(Args.A, Item)) : false;
+                if (bAValue)
                 {
-                    StoreBool( Item, true );
+                    StoreBool(Item, true);
                 }
                 else
                 {
-                    AddOp( FScheduledOp( Item.At, Item, 2),
-                           FScheduledOp( Args.B, Item));
+                    AddOp(FScheduledOp(Item.At, Item, 2),
+                           FScheduledOp(Args.B, Item));
                 }
                 break;
             }
@@ -5913,8 +5906,8 @@ namespace UE::Mutable::Private
             case 2:
             {
                 // We arrived here because a is false
-                bool b = Args.B ? LoadBool( FCacheAddress(Args.B,Item) ) : false;
-                StoreBool( Item, b );
+                bool bBValue = Args.B ? LoadBool(FCacheAddress(Args.B, Item)) : false;
+                StoreBool(Item, bBValue);
                 break;
             }
 
@@ -5926,18 +5919,18 @@ namespace UE::Mutable::Private
 
         case EOpType::BO_NOT:
         {
-			OP::BoolNotArgs Args = Program.GetOpArgs<OP::BoolNotArgs>(Item.At);
+			FOperation::BoolNotArgs Args = Program.GetOpArgs<FOperation::BoolNotArgs>(Item.At);
             switch (Item.Stage)
             {
             case 0:
-                    AddOp( FScheduledOp( Item.At, Item, 1),
-                           FScheduledOp( Args.A, Item) );
+				AddOp(FScheduledOp(Item.At, Item, 1),
+						FScheduledOp(Args.A, Item));
                 break;
 
             case 1:
             {
-                bool result = !LoadBool( FCacheAddress(Args.A,Item) );
-                StoreBool( Item, result );
+                bool bResult = !LoadBool(FCacheAddress(Args.A, Item));
+                StoreBool(Item, bResult);
                 break;
             }
 
@@ -5949,17 +5942,18 @@ namespace UE::Mutable::Private
 
         case EOpType::BO_EQUAL_INT_CONST:
         {
-			OP::BoolEqualScalarConstArgs Args = Program.GetOpArgs<OP::BoolEqualScalarConstArgs>(Item.At);
-            switch (Item.Stage)
-            {
-            case 0:
-                    AddOp( FScheduledOp( Item.At, Item, 1),
-                           FScheduledOp( Args.Value, Item) );
-                break;
-
-            case 1:
+			FOperation::BoolEqualScalarConstArgs Args = Program.GetOpArgs<FOperation::BoolEqualScalarConstArgs>(Item.At);
+			switch (Item.Stage)
+			{
+			case 0:
+			{
+				AddOp(FScheduledOp(Item.At, Item, 1), 
+						FScheduledOp(Args.Value, Item));
+				break;
+			}
+			case 1:
             {
-                int32 a = LoadInt( FCacheAddress(Args.Value,Item) );
+                int32 a = LoadInt(FCacheAddress(Args.Value, Item));
                 bool result = a == Args.Constant;
                 StoreBool( Item, result );
                 break;
@@ -5988,7 +5982,7 @@ namespace UE::Mutable::Private
 
         case EOpType::NU_CONSTANT:
         {
-			OP::IntConstantArgs Args = Program.GetOpArgs<OP::IntConstantArgs>(Item.At);
+			FOperation::IntConstantArgs Args = Program.GetOpArgs<FOperation::IntConstantArgs>(Item.At);
             int32 result = Args.Value;
             StoreInt( Item, result );
             break;
@@ -5996,7 +5990,7 @@ namespace UE::Mutable::Private
 
         case EOpType::NU_PARAMETER:
         {
-			OP::ParameterArgs Args = Program.GetOpArgs<OP::ParameterArgs>(Item.At);
+			FOperation::ParameterArgs Args = Program.GetOpArgs<FOperation::ParameterArgs>(Item.At);
 			TSharedPtr<FRangeIndex> Index = BuildCurrentOpRangeIndex( Item, Args.variable );
             int32 result = Parameters.GetIntValue( Args.variable, Index.Get());
 
@@ -6039,7 +6033,7 @@ namespace UE::Mutable::Private
 
         case EOpType::SC_CONSTANT:
         {
-			OP::ScalarConstantArgs Args = Program.GetOpArgs<OP::ScalarConstantArgs>(Item.At);
+			FOperation::ScalarConstantArgs Args = Program.GetOpArgs<FOperation::ScalarConstantArgs>(Item.At);
 			float result = Args.Value;
             StoreScalar( Item, result );
             break;
@@ -6047,7 +6041,7 @@ namespace UE::Mutable::Private
 
         case EOpType::SC_PARAMETER:
         {
-			OP::ParameterArgs Args = Program.GetOpArgs<OP::ParameterArgs>(Item.At);
+			FOperation::ParameterArgs Args = Program.GetOpArgs<FOperation::ParameterArgs>(Item.At);
 			TSharedPtr<FRangeIndex> Index = BuildCurrentOpRangeIndex( Item, Args.variable );
             float result = Parameters.GetFloatValue( Args.variable, Index.Get());
             StoreScalar( Item, result );
@@ -6056,7 +6050,7 @@ namespace UE::Mutable::Private
 
         case EOpType::SC_CURVE:
         {
-			OP::ScalarCurveArgs Args = Program.GetOpArgs<OP::ScalarCurveArgs>(Item.At);
+			FOperation::ScalarCurveArgs Args = Program.GetOpArgs<FOperation::ScalarCurveArgs>(Item.At);
             switch (Item.Stage)
             {
             case 0:
@@ -6083,7 +6077,7 @@ namespace UE::Mutable::Private
 
         case EOpType::SC_ARITHMETIC:
         {
-			OP::ArithmeticArgs Args = Program.GetOpArgs<OP::ArithmeticArgs>(Item.At);
+			FOperation::ArithmeticArgs Args = Program.GetOpArgs<FOperation::ArithmeticArgs>(Item.At);
             switch (Item.Stage)
             {
             case 0:
@@ -6100,19 +6094,19 @@ namespace UE::Mutable::Private
                 float result = 1.0f;
                 switch (Args.Operation)
                 {
-                case OP::ArithmeticArgs::ADD:
+                case FOperation::ArithmeticArgs::ADD:
                     result = a + b;
                     break;
 
-                case OP::ArithmeticArgs::MULTIPLY:
+                case FOperation::ArithmeticArgs::MULTIPLY:
                     result = a * b;
                     break;
 
-                case OP::ArithmeticArgs::SUBTRACT:
+                case FOperation::ArithmeticArgs::SUBTRACT:
                     result = a - b;
                     break;
 
-                case OP::ArithmeticArgs::DIVIDE:
+                case FOperation::ArithmeticArgs::DIVIDE:
                     result = a / b;
                     break;
 
@@ -6133,7 +6127,7 @@ namespace UE::Mutable::Private
 
 		case EOpType::SC_MATERIAL_BREAK:
 		{
-			const OP::MaterialBreakArgs Args = Program.GetOpArgs<OP::MaterialBreakArgs>(Item.At);
+			const FOperation::MaterialBreakArgs Args = Program.GetOpArgs<FOperation::MaterialBreakArgs>(Item.At);
 
 			switch (Item.Stage)
 			{
@@ -6178,7 +6172,7 @@ namespace UE::Mutable::Private
 				if(!bParameterFound)
 				{
 					LiveInstance->UpdateLogger->LogUpdateMessage(FString::Printf(TEXT("Expected Scalar parameter [%s] not found in Break Material node."), *ParameterName.ToString()));
-					AbortRun();
+					GetMemory().SetAborted(FCacheAddress(Item));
 				}
 
 				break;
@@ -6208,7 +6202,7 @@ namespace UE::Mutable::Private
 
         case EOpType::ST_CONSTANT:
         {
-			OP::ResourceConstantArgs Args = Program.GetOpArgs<OP::ResourceConstantArgs>( Item.At );
+			FOperation::ResourceConstantArgs Args = Program.GetOpArgs<FOperation::ResourceConstantArgs>( Item.At );
             check( Args.value < (uint32)Program.ConstantStrings.Num() );
 
             const FString& Result = Program.ConstantStrings[Args.value];
@@ -6219,7 +6213,7 @@ namespace UE::Mutable::Private
 
         case EOpType::ST_PARAMETER:
         {
-			OP::ParameterArgs Args = Program.GetOpArgs<OP::ParameterArgs>( Item.At );
+			FOperation::ParameterArgs Args = Program.GetOpArgs<FOperation::ParameterArgs>( Item.At );
 			TSharedPtr<FRangeIndex> Index = BuildCurrentOpRangeIndex( Item, Args.variable );
 			FString Result;
 			Parameters.GetStringValue(Args.variable, Result, Index.Get());
@@ -6246,14 +6240,14 @@ namespace UE::Mutable::Private
 
         case EOpType::CO_CONSTANT:
         {
-			OP::ColorConstantArgs Args = Program.GetOpArgs<OP::ColorConstantArgs>(Item.At);
+			FOperation::ColorConstantArgs Args = Program.GetOpArgs<FOperation::ColorConstantArgs>(Item.At);
             StoreColor( Item, Args.Value );
             break;
         }
 
         case EOpType::CO_PARAMETER:
         {
-			OP::ParameterArgs Args = Program.GetOpArgs<OP::ParameterArgs>(Item.At);
+			FOperation::ParameterArgs Args = Program.GetOpArgs<FOperation::ParameterArgs>(Item.At);
 			TSharedPtr<FRangeIndex> Index = BuildCurrentOpRangeIndex( Item, Args.variable );
             FVector4f V;
             Parameters.GetColourValue( Args.variable, V, Index.Get());
@@ -6263,7 +6257,7 @@ namespace UE::Mutable::Private
 
         case EOpType::CO_SAMPLEIMAGE:
         {
-			OP::ColourSampleImageArgs Args = Program.GetOpArgs<OP::ColourSampleImageArgs>(Item.At);
+			FOperation::ColourSampleImageArgs Args = Program.GetOpArgs<FOperation::ColourSampleImageArgs>(Item.At);
             switch (Item.Stage)
             {
             case 0:
@@ -6312,7 +6306,7 @@ namespace UE::Mutable::Private
 
         case EOpType::CO_SWIZZLE:
         {
-			OP::ColourSwizzleArgs Args = Program.GetOpArgs<OP::ColourSwizzleArgs>(Item.At);
+			FOperation::ColourSwizzleArgs Args = Program.GetOpArgs<FOperation::ColourSwizzleArgs>(Item.At);
             switch (Item.Stage)
             {
             case 0:
@@ -6349,7 +6343,7 @@ namespace UE::Mutable::Private
 
         case EOpType::CO_FROMSCALARS:
         {
-			OP::ColourFromScalarsArgs Args = Program.GetOpArgs<OP::ColourFromScalarsArgs>(Item.At);
+			FOperation::ColourFromScalarsArgs Args = Program.GetOpArgs<FOperation::ColourFromScalarsArgs>(Item.At);
             switch (Item.Stage)
             {
             case 0:
@@ -6385,7 +6379,7 @@ namespace UE::Mutable::Private
 
         case EOpType::CO_ARITHMETIC:
         {
-			OP::ArithmeticArgs Args = Program.GetOpArgs<OP::ArithmeticArgs>(Item.At);
+			FOperation::ArithmeticArgs Args = Program.GetOpArgs<FOperation::ArithmeticArgs>(Item.At);
             switch (Item.Stage)
             {
             case 0:
@@ -6410,19 +6404,19 @@ namespace UE::Mutable::Private
 				FVector4f result = FVector4f(0,0,0,0);
                 switch (Args.Operation)
                 {
-                case OP::ArithmeticArgs::ADD:
+                case FOperation::ArithmeticArgs::ADD:
                     result = a + b;
                     break;
 
-                case OP::ArithmeticArgs::MULTIPLY:
+                case FOperation::ArithmeticArgs::MULTIPLY:
                     result = a * b;
                     break;
 
-                case OP::ArithmeticArgs::SUBTRACT:
+                case FOperation::ArithmeticArgs::SUBTRACT:
                     result = a - b;
                     break;
 
-                case OP::ArithmeticArgs::DIVIDE:
+                case FOperation::ArithmeticArgs::DIVIDE:
                     result = a / b;
                     break;
 
@@ -6444,7 +6438,7 @@ namespace UE::Mutable::Private
 
 		case EOpType::CO_LINEARTOSRGB:
 		{
-			OP::ColorArgs Args = Program.GetOpArgs<OP::ColorArgs>(Item.At);
+			FOperation::ColorArgs Args = Program.GetOpArgs<FOperation::ColorArgs>(Item.At);
 			switch (Item.Stage)
 			{
 			case 0:
@@ -6471,7 +6465,7 @@ namespace UE::Mutable::Private
 
 		case EOpType::CO_MATERIAL_BREAK:
 		{
-			const OP::MaterialBreakArgs Args = Program.GetOpArgs<OP::MaterialBreakArgs>(Item.At);
+			const FOperation::MaterialBreakArgs Args = Program.GetOpArgs<FOperation::MaterialBreakArgs>(Item.At);
 
 			switch (Item.Stage)
 			{
@@ -6516,7 +6510,7 @@ namespace UE::Mutable::Private
 				if(!bParameterFound)
 				{
 					LiveInstance->UpdateLogger->LogUpdateMessage(FString::Printf(TEXT("Expected Color parameter [%s] not found in Break Material node."), *ParameterName.ToString()));
-					AbortRun();
+					GetMemory().SetAborted(FCacheAddress(Item));
 				}
 
 				break;
@@ -6547,7 +6541,7 @@ namespace UE::Mutable::Private
 
         case EOpType::PR_CONSTANT:
         {
-			OP::ResourceConstantArgs Args = Program.GetOpArgs<OP::ResourceConstantArgs>(Item.At);
+			FOperation::ResourceConstantArgs Args = Program.GetOpArgs<FOperation::ResourceConstantArgs>(Item.At);
             FProjector Result = Program.ConstantProjectors[Args.value];
             StoreProjector( Item, Result );
             break;
@@ -6555,7 +6549,7 @@ namespace UE::Mutable::Private
 
         case EOpType::PR_PARAMETER:
         {
-			OP::ParameterArgs Args = Program.GetOpArgs<OP::ParameterArgs>(Item.At);
+			FOperation::ParameterArgs Args = Program.GetOpArgs<FOperation::ParameterArgs>(Item.At);
 			TSharedPtr<FRangeIndex> Index = BuildCurrentOpRangeIndex( Item, Args.variable );
             FProjector Result = Parameters.GetProjectorValue(Args.variable,Index.Get());
 
@@ -6584,14 +6578,14 @@ namespace UE::Mutable::Private
 		{
 		case EOpType::MA_CONSTANT:
 			{
-				OP::MatrixConstantArgs Args = Program.GetOpArgs<OP::MatrixConstantArgs>(Item.At);
+				FOperation::MatrixConstantArgs Args = Program.GetOpArgs<FOperation::MatrixConstantArgs>(Item.At);
 				StoreMatrix( Item, Program.ConstantMatrices[Args.value] );
 				break;
 			}
 
 		case EOpType::MA_PARAMETER:
 			{
-				OP::ParameterArgs Args = Program.GetOpArgs<OP::ParameterArgs>(Item.At);
+				FOperation::ParameterArgs Args = Program.GetOpArgs<FOperation::ParameterArgs>(Item.At);
 				TSharedPtr<FRangeIndex> Index = BuildCurrentOpRangeIndex( Item, Args.variable );
 				FMatrix44f Value;
 				Parameters.GetMatrixValue( Args.variable, Value, Index.Get());
@@ -6615,7 +6609,7 @@ namespace UE::Mutable::Private
 			MUTABLE_CPUPROFILER_SCOPE(RunCodeImageDesc_MI_PARAMETER);
 			check(Item.Stage == 0);
 
-			OP::MaterialParameterArgs Args = Program.GetOpArgs<OP::MaterialParameterArgs >(Item.At);
+			FOperation::MaterialParameterArgs Args = Program.GetOpArgs<FOperation::MaterialParameterArgs >(Item.At);
 			TSharedPtr<FRangeIndex> Index = BuildCurrentOpRangeIndex(Item, Args.Variable);
 
 			UMaterialInterface* MaterialInterface = Parameters.GetMaterialValue(Args.Variable, Index.Get());
@@ -6632,7 +6626,7 @@ namespace UE::Mutable::Private
 
 				// Get Parameter layers from program bytecode
 				const uint8* Data = Program.GetOpArgsPointer(Item.At);
-				Data += sizeof(OP::MaterialParameterArgs);
+				Data += sizeof(FOperation::MaterialParameterArgs);
 
 				// Textures
 				for (int32 ParameterIndex = 0; ParameterIndex < ImageParameterNames.Num(); ++ParameterIndex)
@@ -6658,8 +6652,8 @@ namespace UE::Mutable::Private
 					// Store Texture into the FMaterial
 					if (bParameterFound)
 					{
-						TVariant<OP::ADDRESS, TManagedPtr<const FImage>> NewParameterImage;
-						NewParameterImage.Set<OP::ADDRESS>(ImageParameterOperations[ParameterIndex]);
+						TVariant<FOperation::ADDRESS, TManagedPtr<const FImage>> NewParameterImage;
+						NewParameterImage.Set<FOperation::ADDRESS>(ImageParameterOperations[ParameterIndex]);
 
 						FMaterial::FImageParameterData NewImageParameter;
 						NewImageParameter.ImageParameter = NewParameterImage;
@@ -6678,12 +6672,12 @@ namespace UE::Mutable::Private
 		{
 			MUTABLE_CPUPROFILER_SCOPE(RunCodeMaterial_MI_MODIFY);
 				
-			OP::MaterialModifyArgs Args = Program.GetOpArgs<OP::MaterialModifyArgs>(Item.At);
+			FOperation::MaterialModifyArgs Args = Program.GetOpArgs<FOperation::MaterialModifyArgs>(Item.At);
 			int32 ParameterCount = Args.NumParameters;
 
 			//Dynamic data can't be stored into arg structures
 			const uint8* Data = Program.GetOpArgsPointer(Item.At);
-			Data += sizeof(OP::MaterialModifyArgs);
+			Data += sizeof(FOperation::MaterialModifyArgs);
 
 			switch (Item.Stage)
 			{
@@ -6697,17 +6691,17 @@ namespace UE::Mutable::Private
 				// Then, Schedule parameter operations
 				for (int32 ParameterIndex = 0; ParameterIndex < ParameterCount; ++ParameterIndex)
 				{
-					OP::ADDRESS NewName;
-					FMemory::Memcpy(&NewName, Data, sizeof(OP::ADDRESS));
-					Data += sizeof(OP::ADDRESS);
+					FOperation::ADDRESS NewName;
+					FMemory::Memcpy(&NewName, Data, sizeof(FOperation::ADDRESS));
+					Data += sizeof(FOperation::ADDRESS);
 
 					int8 LayerIndex;
 					FMemory::Memcpy(&LayerIndex, Data, sizeof(int8));
 					Data += sizeof(int8);
 
-					OP::ADDRESS NewParameterOperation;
-					FMemory::Memcpy(&NewParameterOperation, Data, sizeof(OP::ADDRESS));
-					Data += sizeof(OP::ADDRESS);
+					FOperation::ADDRESS NewParameterOperation;
+					FMemory::Memcpy(&NewParameterOperation, Data, sizeof(FOperation::ADDRESS));
+					Data += sizeof(FOperation::ADDRESS);
 
 					// Evaluate non lazy branches here
 					if (GetOpDataType(Program.GetOpType(NewParameterOperation)) != EDataType::Image)
@@ -6726,27 +6720,27 @@ namespace UE::Mutable::Private
 			}
 			case 1:
 			{
-				TArray<OP::ADDRESS> ParameterNames;
+				TArray<FOperation::ADDRESS> ParameterNames;
 				TArray<int8> ParameterLayerIndex;
-				TArray<OP::ADDRESS> ParameterOperations;
+				TArray<FOperation::ADDRESS> ParameterOperations;
 				TArray<int32> ImageParameterIndexArray;
 
 				for (int32 ParameterIndex = 0; ParameterIndex < ParameterCount; ++ParameterIndex)
 				{
-					OP::ADDRESS NewName;
-					FMemory::Memcpy(&NewName, Data, sizeof(OP::ADDRESS));
+					FOperation::ADDRESS NewName;
+					FMemory::Memcpy(&NewName, Data, sizeof(FOperation::ADDRESS));
 					ParameterNames.Add(NewName);
-					Data += sizeof(OP::ADDRESS);
+					Data += sizeof(FOperation::ADDRESS);
 
 					int8 LayerIndex;
 					FMemory::Memcpy(&LayerIndex, Data, sizeof(int8));
 					ParameterLayerIndex.Add(LayerIndex);
 					Data += sizeof(int8);
 
-					OP::ADDRESS NewParameterOperation;
-					FMemory::Memcpy(&NewParameterOperation, Data, sizeof(OP::ADDRESS));
+					FOperation::ADDRESS NewParameterOperation;
+					FMemory::Memcpy(&NewParameterOperation, Data, sizeof(FOperation::ADDRESS));
 					ParameterOperations.Add(NewParameterOperation);
-					Data += sizeof(OP::ADDRESS);
+					Data += sizeof(FOperation::ADDRESS);
 
 					// Evaluate lazy branches here
 					if (GetOpDataType(Program.GetOpType(NewParameterOperation)) == EDataType::Image)
@@ -6781,8 +6775,8 @@ namespace UE::Mutable::Private
 						case EDataType::Image:
 						{
 							//TODO: Evaluate the non-lazy textures
-							TVariant<OP::ADDRESS, TManagedPtr<const FImage>> NewParameterImage;
-							NewParameterImage.Set<OP::ADDRESS>(ParameterOperations[ParameterIndex]);
+							TVariant<FOperation::ADDRESS, TManagedPtr<const FImage>> NewParameterImage;
+							NewParameterImage.Set<FOperation::ADDRESS>(ParameterOperations[ParameterIndex]);
 
 							FMaterial::FImageParameterData ImageParameterData;
 							ImageParameterData.ImageParameter = NewParameterImage;
@@ -6851,7 +6845,7 @@ namespace UE::Mutable::Private
 		{
 			MUTABLE_CPUPROFILER_SCOPE(RunCodeMaterial_MI_SKELETALMESH_BREAK);
 
-			OP::MaterialSkeletalMeshBreakArgs Args = Program.GetOpArgs<OP::MaterialSkeletalMeshBreakArgs >(Item.At);
+			FOperation::MaterialSkeletalMeshBreakArgs Args = Program.GetOpArgs<FOperation::MaterialSkeletalMeshBreakArgs >(Item.At);
 			
 			switch (Item.Stage)
 			{
@@ -6876,17 +6870,17 @@ namespace UE::Mutable::Private
 					TObjectPtr<UMaterialInterface> FoundMaterialInterface = nullptr;
 					if (USkeletalMesh* EngineSkeletalMesh = PassthroughMesh.Get())
 					{
+						// From the Skeletal mesh get the material in the override slot with the provided name.
 						const FString& SlotName = Program.ConstantStrings[Args.SlotNameStringIndex];
-						
-						// From the Skeletal mesh get the material in the override slot with the provided name
-						const TArray<FSkeletalMaterial> Materials = EngineSkeletalMesh->GetMaterials();
-						for (const FSkeletalMaterial& Material : Materials)
+						FSkeletalMaterial* FoundSlotNameMaterial = EngineSkeletalMesh->GetMaterials().FindByPredicate(
+						[SlotName](const FSkeletalMaterial& Material)
 						{
-							if (Material.MaterialSlotName == SlotName)
-							{
-								FoundMaterialInterface = Material.MaterialInterface;
-								break;
-							}
+							return Material.MaterialSlotName == SlotName;
+						});
+
+						if (FoundSlotNameMaterial)
+						{
+							FoundMaterialInterface = FoundSlotNameMaterial->MaterialInterface;
 						}
 					}
 
@@ -6921,7 +6915,7 @@ namespace UE::Mutable::Private
 
         case EOpType::LA_CONSTANT:
         {
-			OP::ResourceConstantArgs Args = Program.GetOpArgs<OP::ResourceConstantArgs>(Item.At);
+			FOperation::ResourceConstantArgs Args = Program.GetOpArgs<FOperation::ResourceConstantArgs>(Item.At);
             check( Args.value < (uint32)Program.ConstantLayouts.Num() );
 
             TManagedPtr<const FLayout> Result = Program.ConstantLayouts[ Args.value ];
@@ -6931,7 +6925,7 @@ namespace UE::Mutable::Private
 
         case EOpType::LA_MERGE:
         {
-			OP::LayoutMergeArgs Args = Program.GetOpArgs<OP::LayoutMergeArgs>(Item.At);
+			FOperation::LayoutMergeArgs Args = Program.GetOpArgs<FOperation::LayoutMergeArgs>(Item.At);
             switch (Item.Stage)
             {
             case 0:
@@ -6974,7 +6968,7 @@ namespace UE::Mutable::Private
 
         case EOpType::LA_PACK:
         {
-			OP::LayoutPackArgs Args = Program.GetOpArgs<OP::LayoutPackArgs>(Item.At);
+			FOperation::LayoutPackArgs Args = Program.GetOpArgs<FOperation::LayoutPackArgs>(Item.At);
             switch (Item.Stage)
             {
             case 0:
@@ -7007,7 +7001,7 @@ namespace UE::Mutable::Private
 
 		case EOpType::LA_FROMMESH:
 		{
-			OP::LayoutFromMeshArgs Args = Program.GetOpArgs<OP::LayoutFromMeshArgs>(Item.At);
+			FOperation::LayoutFromMeshArgs Args = Program.GetOpArgs<FOperation::LayoutFromMeshArgs>(Item.At);
 
 			constexpr uint8 MeshContentFilter = (uint8)(EMeshContentFlags::AllFlags);
 			switch (Item.Stage)
@@ -7036,7 +7030,7 @@ namespace UE::Mutable::Private
 
 		case EOpType::LA_REMOVEBLOCKS:
 		{
-			OP::LayoutRemoveBlocksArgs Args = Program.GetOpArgs<OP::LayoutRemoveBlocksArgs>(Item.At);
+			FOperation::LayoutRemoveBlocksArgs Args = Program.GetOpArgs<FOperation::LayoutRemoveBlocksArgs>(Item.At);
 			switch (Item.Stage)
 			{
 			case 0:
@@ -7146,14 +7140,14 @@ namespace UE::Mutable::Private
 			{
 				const uint8* Data = Program.GetOpArgsPointer(Item.At);
 
-				OP::ExternalArgs Args = Program.GetOpArgs<OP::ExternalArgs>(Item.At);
-				Data += sizeof(OP::ExternalArgs);
+				FOperation::ExternalArgs Args = Program.GetOpArgs<FOperation::ExternalArgs>(Item.At);
+				Data += sizeof(FOperation::ExternalArgs);
 				
 				const FInstancedStruct& OperationInstancedStruct = Program.ExternalOperationProvider->Get(Item.At);
 				if (!OperationInstancedStruct.IsValid()) // Plugin not loaded
 				{
 					LiveInstance->UpdateLogger->LogUpdateMessage(FString::Printf(TEXT("Trying to evaluate an External Operation which has not been loaded.")));
-					AbortRun();
+					GetMemory().SetAborted(FCacheAddress(Item));
 
 					break;
 				}
@@ -7164,7 +7158,7 @@ namespace UE::Mutable::Private
 				if (Inputs.Num() != Args.NumOperants)
 				{
 					LiveInstance->UpdateLogger->LogUpdateMessage(FString::Printf(TEXT("External Operation inputs do not match compiled data: %s"), *OperationInstancedStruct.GetScriptStruct()->GetName()));
-					AbortRun();
+					GetMemory().SetAborted(FCacheAddress(Item));
 
 					return;
 				}
@@ -7177,9 +7171,9 @@ namespace UE::Mutable::Private
 
 						for (int32 Index = 0; Index < Args.NumOperants; ++Index)
 						{
-							OP::ADDRESS OperantAt;
-							FMemory::Memcpy(&OperantAt, Data, sizeof(OP::ADDRESS));
-							Data += sizeof(OP::ADDRESS);
+							FOperation::ADDRESS OperantAt;
+							FMemory::Memcpy(&OperantAt, Data, sizeof(FOperation::ADDRESS));
+							Data += sizeof(FOperation::ADDRESS);
 
 							Dependencies.Add(FScheduledOp(OperantAt, Item));
 						}
@@ -7196,9 +7190,9 @@ namespace UE::Mutable::Private
 						// Inputs
 						for (int32 Index = 0; Index < Args.NumOperants; ++Index)
 						{
-							OP::ADDRESS OperantAt;
-							FMemory::Memcpy(&OperantAt, Data, sizeof(OP::ADDRESS));
-							Data += sizeof(OP::ADDRESS);
+							FOperation::ADDRESS OperantAt;
+							FMemory::Memcpy(&OperantAt, Data, sizeof(FOperation::ADDRESS));
+							Data += sizeof(FOperation::ADDRESS);
 
 							const TPair<FName, const UScriptStruct*>& Input = Inputs[Index];
 
@@ -7289,8 +7283,7 @@ namespace UE::Mutable::Private
 						if (!Context.Output.Ptr)
 						{
 							LiveInstance->UpdateLogger->LogUpdateMessage(FString::Printf(TEXT("External Operation output not set: %s"), *OperationInstancedStruct.GetScriptStruct()->GetName()));
-							AbortRun();
-
+							GetMemory().SetAborted(FCacheAddress(Item));
 							return;
 						}
 						
@@ -7341,7 +7334,7 @@ namespace UE::Mutable::Private
 
 		case EOpType::IS_PARAMETER:
 			{
-				OP::ParameterArgs Args = Program.GetOpArgs<OP::ParameterArgs>(Item.At);
+				FOperation::ParameterArgs Args = Program.GetOpArgs<FOperation::ParameterArgs>(Item.At);
 				TSharedPtr<FRangeIndex> Index = BuildCurrentOpRangeIndex(Item, Args.variable);
 				TManagedPtr<const FInstancedStruct> Result = Parameters.GetInstancedStructValue(Args.variable, Index.Get());
 				StoreInstancedStruct(Item, Result);
@@ -7351,7 +7344,7 @@ namespace UE::Mutable::Private
 
 		case EOpType::SK_PARAMETER:
 			{
-				OP::ParameterArgs Args = Program.GetOpArgs<OP::ParameterArgs>(Item.At);
+				FOperation::ParameterArgs Args = Program.GetOpArgs<FOperation::ParameterArgs>(Item.At);
 				TSharedPtr<FRangeIndex> Index = BuildCurrentOpRangeIndex(Item, Args.variable);
 				USkeletalMesh* SkeletalMesh = Parameters.GetSkeletalMeshValue(Args.variable, Index.Get());
 
@@ -7443,10 +7436,10 @@ namespace UE::Mutable::Private
 			MUTABLE_CPUPROFILER_SCOPE(RunCodeImageDesc_IM_CONSTANT);
 
 			check(Item.Stage == 0);
-			OP::ResourceConstantArgs Args = Program.GetOpArgs<OP::ResourceConstantArgs>(Item.At);
+			FOperation::ResourceConstantArgs Args = Program.GetOpArgs<FOperation::ResourceConstantArgs>(Item.At);
 			int32 ImageIndex = Args.value;
 
-			FExtendedImageDesc& Result = ImageDescResults.FindOrAdd(Item.At);
+			FExtendedImageDesc Result;
 
 			Result.m_format = Program.ConstantImages[ImageIndex].ImageFormat;
 			Result.m_size[0] = Program.ConstantImages[ImageIndex].ImageSizeX;
@@ -7480,9 +7473,10 @@ namespace UE::Mutable::Private
 
 				Result.FirstLODAvailable = LODIndex;
 			}
+
 			ImageDescConstantImages.Add(ImageIndex);
 
-			StoreValidDesc(Item);
+			StoreImageDesc(Item, Result);
 			break;
 		}
 
@@ -7490,31 +7484,27 @@ namespace UE::Mutable::Private
 		{
 			MUTABLE_CPUPROFILER_SCOPE(RunCodeImageDesc_IM_PARAMETER_CONVERT);
 
-			OP::ImageParameterConvertArgs Args = Program.GetOpArgs<OP::ImageParameterConvertArgs>(Item.At);
+			FOperation::ImageParameterConvertArgs Args = Program.GetOpArgs<FOperation::ImageParameterConvertArgs>(Item.At);
 			switch (Item.Stage)
 			{
 			case 0:
-				{
-					AddOp(FScheduledOp(Item.At, Item, 1), 
+			{
+				AddOp(FScheduledOp(Item.At, Item, 1), 
 						FScheduledOp(Args.ImageParameter, Item));
 					
-					break;
-				}
+				break;
+			}
 
 			case 1:
-				{
-					FExtendedImageDesc& Result = ImageDescResults.FindOrAdd(Item.At);
-					Result = ImageDescResults.FindChecked(Args.ImageParameter);
-
-					StoreValidDesc(Item); 
-					break;
-				}
+			{
+				StoreImageDesc(Item, LoadImageDesc(FCacheAddress(Args.ImageParameter, Item)));
+				break;
+			}
 
 			default:
 				unimplemented();
 			}
 
-			StoreValidDesc(Item);
 			break;
 		}
 
@@ -7523,12 +7513,9 @@ namespace UE::Mutable::Private
 			MUTABLE_CPUPROFILER_SCOPE(RunCodeImageDesc_IM_REFERENCE);
 			check(Item.Stage == 0);
 
-			OP::ResourceReferenceArgs Args = Program.GetOpArgs<OP::ResourceReferenceArgs>(Item.At);
-			FExtendedImageDesc& Result = ImageDescResults.FindOrAdd(Item.At);
-			
-			Result = FExtendedImageDesc{ Args.ImageDesc, 0 };
+			FOperation::ResourceReferenceArgs Args = Program.GetOpArgs<FOperation::ResourceReferenceArgs>(Item.At);	
 
-			StoreValidDesc(Item);
+			StoreImageDesc(Item, FExtendedImageDesc{ Args.ImageDesc, 0 });
 			break;
 		}
 
@@ -7537,20 +7524,17 @@ namespace UE::Mutable::Private
 			MUTABLE_CPUPROFILER_SCOPE(RunCodeImageDesc_IM_PARAMETER);
 			check(Item.Stage == 0);
 
-			OP::ParameterArgs Args = Program.GetOpArgs<OP::ParameterArgs>(Item.At);
+			FOperation::ParameterArgs Args = Program.GetOpArgs<FOperation::ParameterArgs>(Item.At);
 			UTexture* Image = Parameters.GetTextureValue(Args.variable);
 
-			FExtendedImageDesc& Result = ImageDescResults.FindOrAdd(Item.At);
-			Result = GetExternalImageDesc(Image);
-
-			StoreValidDesc(Item);
+			StoreImageDesc(Item, GetExternalImageDesc(Image));
 			break;
 		}
 
 		case EOpType::IM_CONDITIONAL:
 		{
 			MUTABLE_CPUPROFILER_SCOPE(RunCodeImageDesc_IM_CONDITIONAL);
-			OP::ConditionalArgs Args = Program.GetOpArgs<OP::ConditionalArgs>(Item.At);
+			FOperation::ConditionalArgs Args = Program.GetOpArgs<FOperation::ConditionalArgs>(Item.At);
 			switch (Item.Stage)
 			{
 			case 0:
@@ -7566,7 +7550,7 @@ namespace UE::Mutable::Private
 			case 1:
 			{
 				bool bValue = LoadBool(FCacheAddress(Args.condition, Item.ExecutionIndex, Item.ExecutionOptions));
-				OP::ADDRESS ResultAt = bValue ? Args.yes : Args.no;
+				FOperation::ADDRESS ResultAt = bValue ? Args.yes : Args.no;
 
 				AddOp(FScheduledOp(Item.At, Item, 2, ResultAt), 
 						FScheduledOp(ResultAt, Item, 0));
@@ -7575,10 +7559,7 @@ namespace UE::Mutable::Private
 
 			case 2: 
 			{
-				FExtendedImageDesc& Result =  ImageDescResults.FindOrAdd(Item.At);
-				Result = ImageDescResults.FindChecked(Item.CustomState);
-
-				StoreValidDesc(Item); 
+				StoreImageDesc(Item, LoadImageDesc(FCacheAddress(Item.CustomState, Item))); 
 				break;
 			}
 			default: check(false);
@@ -7591,17 +7572,17 @@ namespace UE::Mutable::Private
 			MUTABLE_CPUPROFILER_SCOPE(RunCodeImageDesc_IM_SWITCH);
 			const uint8* Data = Program.GetOpArgsPointer(Item.At);
 		
-			OP::ADDRESS VarAddress;
-			FMemory::Memcpy(&VarAddress, Data, sizeof(OP::ADDRESS));
-			Data += sizeof(OP::ADDRESS);
+			FOperation::ADDRESS VarAddress;
+			FMemory::Memcpy(&VarAddress, Data, sizeof(FOperation::ADDRESS));
+			Data += sizeof(FOperation::ADDRESS);
 
-			OP::ADDRESS DefAddress;
-			FMemory::Memcpy(&DefAddress, Data, sizeof(OP::ADDRESS));
-			Data += sizeof(OP::ADDRESS);
+			FOperation::ADDRESS DefAddress;
+			FMemory::Memcpy(&DefAddress, Data, sizeof(FOperation::ADDRESS));
+			Data += sizeof(FOperation::ADDRESS);
 
-			OP::FSwitchCaseDescriptor CaseDesc;
-			FMemory::Memcpy(&CaseDesc, Data, sizeof(OP::FSwitchCaseDescriptor));
-			Data += sizeof(OP::FSwitchCaseDescriptor);
+			FOperation::FSwitchCaseDescriptor CaseDesc;
+			FMemory::Memcpy(&CaseDesc, Data, sizeof(FOperation::FSwitchCaseDescriptor));
+			Data += sizeof(FOperation::FSwitchCaseDescriptor);
 
 			switch (Item.Stage)
 			{
@@ -7616,8 +7597,7 @@ namespace UE::Mutable::Private
 				}
 				else
 				{
-					ImageDescResults.FindOrAdd(Item.At);
-					StoreValidDesc(Item);
+					StoreImageDesc(Item, FExtendedImageDesc{});
 				}
 				break;
 			}
@@ -7627,7 +7607,7 @@ namespace UE::Mutable::Private
 				// Get the variable result
 				int32 Var = LoadInt(FCacheAddress(VarAddress, Item.ExecutionIndex, Item.ExecutionOptions, FScheduledOp::EType::Full));
 
-				OP::ADDRESS ValueAt = DefAddress;
+				FOperation::ADDRESS ValueAt = DefAddress;
 
 				if (!CaseDesc.bUseRanges)
 				{
@@ -7637,9 +7617,9 @@ namespace UE::Mutable::Private
 						FMemory::Memcpy(&Condition, Data, sizeof(int32));
 						Data += sizeof(int32);
 
-						OP::ADDRESS CaseAt;
-						FMemory::Memcpy(&CaseAt, Data, sizeof(OP::ADDRESS));
-						Data += sizeof(OP::ADDRESS);
+						FOperation::ADDRESS CaseAt;
+						FMemory::Memcpy(&CaseAt, Data, sizeof(FOperation::ADDRESS));
+						Data += sizeof(FOperation::ADDRESS);
 
 						if (CaseAt && Var == (int32)Condition)
 						{
@@ -7660,9 +7640,9 @@ namespace UE::Mutable::Private
 						FMemory::Memcpy(&RangeSize, Data, sizeof(uint32));
 						Data += sizeof(int32);
 
-						OP::ADDRESS CaseAt;
-						FMemory::Memcpy(&CaseAt, Data, sizeof(OP::ADDRESS));
-						Data += sizeof(OP::ADDRESS);
+						FOperation::ADDRESS CaseAt;
+						FMemory::Memcpy(&CaseAt, Data, sizeof(FOperation::ADDRESS));
+						Data += sizeof(FOperation::ADDRESS);
 
 						if (CaseAt && Var >= ConditionStart && Var < int32(ConditionStart + RangeSize))
 						{
@@ -7680,8 +7660,7 @@ namespace UE::Mutable::Private
 				}
 				else
 				{
-					ImageDescResults.FindOrAdd(Item.At);
-					StoreValidDesc(Item); 
+					StoreImageDesc(Item, FExtendedImageDesc{}); 
 				}
 				break;
 			}
@@ -7690,10 +7669,8 @@ namespace UE::Mutable::Private
 			{
 				check(Item.CustomState);
 
-				FExtendedImageDesc& Result = ImageDescResults.FindOrAdd(Item.At);
-				Result = ImageDescResults.FindChecked(Item.CustomState);
+				StoreImageDesc(Item, LoadImageDesc(FCacheAddress(Item.CustomState, Item))); 
 
-				StoreValidDesc(Item); 
 				break;
 			}
 			default: check(false); break;
@@ -7704,7 +7681,7 @@ namespace UE::Mutable::Private
 		case EOpType::IM_LAYERCOLOUR:
 		{
 			MUTABLE_CPUPROFILER_SCOPE(RunCodeImageDesc_IM_LAYERCOLOUR);
-			OP::ImageLayerColourArgs Args = Program.GetOpArgs<OP::ImageLayerColourArgs>(Item.At);
+			FOperation::ImageLayerColourArgs Args = Program.GetOpArgs<FOperation::ImageLayerColourArgs>(Item.At);
 
 			switch (Item.Stage)
 			{
@@ -7718,16 +7695,15 @@ namespace UE::Mutable::Private
 			}
 			case 1:
 			{
-				FExtendedImageDesc& Result = ImageDescResults.FindOrAdd(Item.At); 
-				Result = ImageDescResults.FindChecked(Args.base);
+				FExtendedImageDesc Result = LoadImageDesc(FCacheAddress(Args.base, Item));
 
 				if (Args.mask)
 				{
-					const FExtendedImageDesc& MaskResult = ImageDescResults.FindChecked(Args.mask);
+					FExtendedImageDesc MaskResult = LoadImageDesc(FCacheAddress(Args.mask, Item));
 					Result.FirstLODAvailable = FMath::Max(Result.FirstLODAvailable, MaskResult.FirstLODAvailable);
 				}
 
-				StoreValidDesc(Item); 
+				StoreImageDesc(Item, Result); 
 
 				break;
 			}
@@ -7739,7 +7715,7 @@ namespace UE::Mutable::Private
 		case EOpType::IM_LAYER:
 		{
 			MUTABLE_CPUPROFILER_SCOPE(RunCodeImageDesc_IM_LAYER);
-			OP::ImageLayerArgs Args = Program.GetOpArgs<OP::ImageLayerArgs>(Item.At);	
+			FOperation::ImageLayerArgs Args = Program.GetOpArgs<FOperation::ImageLayerArgs>(Item.At);	
 			switch (Item.Stage)
 			{
 			case 0:
@@ -7753,22 +7729,21 @@ namespace UE::Mutable::Private
 			}
 			case 1:
 			{
-				FExtendedImageDesc& Result = ImageDescResults.FindOrAdd(Item.At); 
-				Result = ImageDescResults.FindChecked(Args.base);
+				FExtendedImageDesc Result = LoadImageDesc(FCacheAddress(Args.base, Item));
 
 				if (Args.mask)
 				{
-					const FExtendedImageDesc& MaskResult = ImageDescResults.FindChecked(Args.mask);
+					FExtendedImageDesc MaskResult = LoadImageDesc(FCacheAddress(Args.mask, Item));
 					Result.FirstLODAvailable = FMath::Max(Result.FirstLODAvailable, MaskResult.FirstLODAvailable);
 				}
 
 				if (Args.blended)
 				{
-					const FExtendedImageDesc& BlendedResult = ImageDescResults.FindChecked(Args.blended);
-					Result.FirstLODAvailable = FMath::Max(Result.FirstLODAvailable, BlendedResult.FirstLODAvailable);
+					FExtendedImageDesc BlenResult = LoadImageDesc(FCacheAddress(Args.blended, Item));
+					Result.FirstLODAvailable = FMath::Max(Result.FirstLODAvailable, BlenResult.FirstLODAvailable);
 				}
 		
-				StoreValidDesc(Item); 
+				StoreImageDesc(Item, Result); 
 				break;
 			}
 			default: check(false);
@@ -7779,7 +7754,7 @@ namespace UE::Mutable::Private
 		case EOpType::IM_MULTILAYER:
 		{
 			MUTABLE_CPUPROFILER_SCOPE(RunCodeImageDesc_IM_MULTILAYER);
-			OP::ImageMultiLayerArgs Args = Program.GetOpArgs<OP::ImageMultiLayerArgs>(Item.At);
+			FOperation::ImageMultiLayerArgs Args = Program.GetOpArgs<FOperation::ImageMultiLayerArgs>(Item.At);
 			switch (Item.Stage)
 			{
 			case 0:
@@ -7792,10 +7767,7 @@ namespace UE::Mutable::Private
 			}
 			case 1: 
 			{
-				FExtendedImageDesc& Result = ImageDescResults.FindOrAdd(Item.At); 
-				Result = ImageDescResults.FindChecked(Args.base);
-
-				StoreValidDesc(Item); 
+				StoreImageDesc(Item, LoadImageDesc(FCacheAddress(Args.base, Item))); 
 				break;
 			}
 			default: check(false);
@@ -7806,7 +7778,7 @@ namespace UE::Mutable::Private
 		case EOpType::IM_NORMALCOMPOSITE:
 		{
 			MUTABLE_CPUPROFILER_SCOPE(RunCodeImageDesc_IM_NORMALCOMPOSITE);
-			OP::ImageNormalCompositeArgs Args = Program.GetOpArgs<OP::ImageNormalCompositeArgs>(Item.At);
+			FOperation::ImageNormalCompositeArgs Args = Program.GetOpArgs<FOperation::ImageNormalCompositeArgs>(Item.At);
 			switch (Item.Stage)
 			{
 			case 0: 
@@ -7817,10 +7789,7 @@ namespace UE::Mutable::Private
 			}
 			case 1: 
 			{
-				FExtendedImageDesc& Result = ImageDescResults.FindOrAdd(Item.At); 
-				Result = ImageDescResults.FindChecked(Args.base);
-				
-				StoreValidDesc(Item); 
+				StoreImageDesc(Item, LoadImageDesc(FCacheAddress(Args.base, Item))); 
 				break;
 			}
 			default: check(false);
@@ -7831,7 +7800,7 @@ namespace UE::Mutable::Private
 		case EOpType::IM_PIXELFORMAT:
 		{
 			MUTABLE_CPUPROFILER_SCOPE(RunCodeImageDesc_IM_PIXELFORMAT);
-			OP::ImagePixelFormatArgs Args = Program.GetOpArgs<OP::ImagePixelFormatArgs>(Item.At);
+			FOperation::ImagePixelFormatArgs Args = Program.GetOpArgs<FOperation::ImagePixelFormatArgs>(Item.At);
 			switch (Item.Stage)
 			{
 			case 0:
@@ -7842,8 +7811,7 @@ namespace UE::Mutable::Private
 			case 1:
 			{
 
-				FExtendedImageDesc& Result = ImageDescResults.FindOrAdd(Item.At); 
-				Result = ImageDescResults.FindChecked(Args.source);
+				FExtendedImageDesc Result = LoadImageDesc(FCacheAddress(Args.source, Item));
 
 				EImageFormat OldFormat = Result.m_format;
 				EImageFormat NewFormat = Args.format;
@@ -7854,7 +7822,7 @@ namespace UE::Mutable::Private
 
 				Result.m_format = NewFormat;
 
-				StoreValidDesc(Item);
+				StoreImageDesc(Item, Result);
 				break;
 			}
 
@@ -7868,7 +7836,7 @@ namespace UE::Mutable::Private
 		case EOpType::IM_MIPMAP:
 		{
 			MUTABLE_CPUPROFILER_SCOPE(RunCodeImageDesc_IM_MIPMAP);
-			OP::ImageMipmapArgs Args = Program.GetOpArgs<OP::ImageMipmapArgs>(Item.At);
+			FOperation::ImageMipmapArgs Args = Program.GetOpArgs<FOperation::ImageMipmapArgs>(Item.At);
 			switch (Item.Stage)
 			{
 			case 0:
@@ -7878,9 +7846,8 @@ namespace UE::Mutable::Private
 
 			case 1:
 			{
-				// Somewhat synched with Full op execution code.
-				FExtendedImageDesc& Result = ImageDescResults.FindOrAdd(Item.At);
-				Result = ImageDescResults.FindChecked(Args.source);
+				// Somewhat synched with Full op execution code.	
+				FExtendedImageDesc Result = LoadImageDesc(FCacheAddress(Args.source, Item));
 
 				int32 LevelCount = Args.levels;
 				int32 MaxLevelCount = FImage::GetMipmapCount(Result.m_size[0], Result.m_size[1]);
@@ -7902,7 +7869,7 @@ namespace UE::Mutable::Private
 				// Update result.
 				Result.m_lods = LevelCount;
 
-				StoreValidDesc(Item);
+				StoreImageDesc(Item, Result);
 				break;
 			}
 
@@ -7915,7 +7882,7 @@ namespace UE::Mutable::Private
 		case EOpType::IM_RESIZE:
 		{
 			MUTABLE_CPUPROFILER_SCOPE(RunCodeImageDesc_IM_RESIZE);
-			OP::ImageResizeArgs Args = Program.GetOpArgs<OP::ImageResizeArgs>(Item.At);
+			FOperation::ImageResizeArgs Args = Program.GetOpArgs<FOperation::ImageResizeArgs>(Item.At);
 			switch (Item.Stage)
 			{
 			case 0:
@@ -7926,13 +7893,12 @@ namespace UE::Mutable::Private
 			}
 			case 1:
 			{
-				FExtendedImageDesc& Result = ImageDescResults.FindOrAdd(Item.At);
-				Result = ImageDescResults.FindChecked(Args.Source);
+				FExtendedImageDesc Result = LoadImageDesc(FCacheAddress(Args.Source, Item));
 
 				Result.m_size[0] = Args.Size[0];
 				Result.m_size[1] = Args.Size[1];
 
-				StoreValidDesc(Item);
+				StoreImageDesc(Item, Result);
 
 				break;
 			}
@@ -7945,7 +7911,7 @@ namespace UE::Mutable::Private
 		case EOpType::IM_RESIZELIKE:
 		{
 			MUTABLE_CPUPROFILER_SCOPE(RunCodeImageDesc_IM_RESIZELIKE);
-			OP::ImageResizeLikeArgs Args = Program.GetOpArgs<OP::ImageResizeLikeArgs>(Item.At);
+			FOperation::ImageResizeLikeArgs Args = Program.GetOpArgs<FOperation::ImageResizeLikeArgs>(Item.At);
 			switch (Item.Stage)
 			{
 			case 0:
@@ -7959,16 +7925,15 @@ namespace UE::Mutable::Private
 
 			case 1:
 			{
-				FExtendedImageDesc& Result = ImageDescResults.FindOrAdd(Item.At);
-				Result = ImageDescResults.FindChecked(Args.Source);
+				FExtendedImageDesc Result = LoadImageDesc(FCacheAddress(Args.Source, Item));
 	
 				if (Args.SizeSource)
 				{
-					const FExtendedImageDesc& SizeSourceResult = ImageDescResults.FindChecked(Args.SizeSource);
+					FExtendedImageDesc SizeSourceResult = LoadImageDesc(FCacheAddress(Args.SizeSource, Item));
 					Result.m_size = SizeSourceResult.m_size;
 				}
 
-				StoreValidDesc(Item);
+				StoreImageDesc(Item, Result);
 				break;
 			}
 
@@ -7982,7 +7947,7 @@ namespace UE::Mutable::Private
 		case EOpType::IM_RESIZEREL:
 		{
 			MUTABLE_CPUPROFILER_SCOPE(RunCodeImageDesc_IM_RESIZEREL);
-			OP::ImageResizeRelArgs Args = Program.GetOpArgs<OP::ImageResizeRelArgs>(Item.At);
+			FOperation::ImageResizeRelArgs Args = Program.GetOpArgs<FOperation::ImageResizeRelArgs>(Item.At);
 			switch (Item.Stage)
 			{
 			case 0:
@@ -7992,8 +7957,7 @@ namespace UE::Mutable::Private
 
 			case 1:
 			{
-				FExtendedImageDesc& Result = ImageDescResults.FindOrAdd(Item.At);
-				Result = ImageDescResults.FindChecked(Args.Source);
+				FExtendedImageDesc Result = LoadImageDesc(FCacheAddress(Args.Source, Item));
 				
 				FImageSize DestSize(
 					uint16(Result.m_size[0] * Args.Factor[0] + 0.5f),
@@ -8001,7 +7965,7 @@ namespace UE::Mutable::Private
 
 				Result.m_size = DestSize;
 
-				StoreValidDesc(Item);
+				StoreImageDesc(Item, Result);
 				break;
 			}
 
@@ -8015,7 +7979,7 @@ namespace UE::Mutable::Private
 		case EOpType::IM_BLANKLAYOUT:
 		{
 			MUTABLE_CPUPROFILER_SCOPE(RunCodeImageDesc_IM_BLANKLAYOUT);
-			OP::ImageBlankLayoutArgs Args = Program.GetOpArgs<OP::ImageBlankLayoutArgs>(Item.At);
+			FOperation::ImageBlankLayoutArgs Args = Program.GetOpArgs<FOperation::ImageBlankLayoutArgs>(Item.At);
 			switch (Item.Stage)
 			{
 			case 0:
@@ -8029,7 +7993,6 @@ namespace UE::Mutable::Private
 
 			case 1:
 			{
-				FExtendedImageDesc& Result = ImageDescResults.FindOrAdd(Item.At);
 				
 				TManagedPtr<const FLayout> Layout = LoadLayout(FCacheAddress(Args.Layout, Item.ExecutionIndex, Item.ExecutionOptions, FScheduledOp::EType::Full));
 
@@ -8037,8 +8000,10 @@ namespace UE::Mutable::Private
 				FIntPoint BlockSizeInPixels(Args.BlockSize[0], Args.BlockSize[1]);
 				FIntPoint ImageSizeInPixels = SizeInBlocks * BlockSizeInPixels;
 
-				
 				FImageSize DestSize(uint16(ImageSizeInPixels.X), uint16(ImageSizeInPixels.Y));
+			
+				FExtendedImageDesc Result;
+			
 				Result.m_size = DestSize;
 				Result.m_format = Args.Format;
 				
@@ -8054,7 +8019,7 @@ namespace UE::Mutable::Private
 					}
 				}
 
-				StoreValidDesc(Item);
+				StoreImageDesc(Item, Result);
 				break;
 			}
 
@@ -8068,7 +8033,7 @@ namespace UE::Mutable::Private
 		case EOpType::IM_COMPOSE:
 		{
 			MUTABLE_CPUPROFILER_SCOPE(RunCodeImageDesc_IM_COMPOSE);
-			OP::ImageComposeArgs Args = Program.GetOpArgs<OP::ImageComposeArgs>(Item.At);
+			FOperation::ImageComposeArgs Args = Program.GetOpArgs<FOperation::ImageComposeArgs>(Item.At);
 			switch (Item.Stage)
 			{
 			case 0: 
@@ -8079,17 +8044,16 @@ namespace UE::Mutable::Private
 				break;
 			} 
 			case 1:
-			{
-				FExtendedImageDesc& Result = ImageDescResults.FindOrAdd(Item.At); 
-				Result = ImageDescResults.FindChecked(Args.base);
+			{	
+				FExtendedImageDesc Result = LoadImageDesc(FCacheAddress(Args.base, Item));
 
 				if (Args.blockImage)
 				{
-					const FExtendedImageDesc& BlockResult = ImageDescResults.FindChecked(Args.blockImage);
+					FExtendedImageDesc BlockResult = LoadImageDesc(FCacheAddress(Args.blockImage, Item));
 					Result.FirstLODAvailable = FMath::Max(Result.FirstLODAvailable, BlockResult.FirstLODAvailable);
 				}
 
-				StoreValidDesc(Item); 
+				StoreImageDesc(Item, Result); 
 				break;
 			}
 			default: check(false);
@@ -8100,7 +8064,7 @@ namespace UE::Mutable::Private
 		case EOpType::IM_MULTICOMPOSE:
 		{
 			MUTABLE_CPUPROFILER_SCOPE(RunCodeImageDesc_IM_MULTICOMPOSE);
-			OP::ImageMultiComposeArgs Args = Program.GetOpArgs<OP::ImageMultiComposeArgs>(Item.At);
+			FOperation::ImageMultiComposeArgs Args = Program.GetOpArgs<FOperation::ImageMultiComposeArgs>(Item.At);
 			switch (Item.Stage)
 			{
 			case 0:
@@ -8112,16 +8076,15 @@ namespace UE::Mutable::Private
 			}
 			case 1:
 			{
-				FExtendedImageDesc& Result = ImageDescResults.FindOrAdd(Item.At);
-				Result = ImageDescResults.FindChecked(Args.Base);
+				FExtendedImageDesc Result = LoadImageDesc(FCacheAddress(Args.Base, Item));
 
 				if (Args.SourceImage)
 				{
-					const FExtendedImageDesc& BlockResult = ImageDescResults.FindChecked(Args.SourceImage);
+					FExtendedImageDesc BlockResult = LoadImageDesc(FCacheAddress(Args.SourceImage, Item));
 					Result.FirstLODAvailable = FMath::Max(Result.FirstLODAvailable, BlockResult.FirstLODAvailable);
 				}
 
-				StoreValidDesc(Item);
+				StoreImageDesc(Item, Result);
 				break;
 			}
 			default: check(false);
@@ -8132,7 +8095,7 @@ namespace UE::Mutable::Private
 		case EOpType::IM_INTERPOLATE:
 		{
 			MUTABLE_CPUPROFILER_SCOPE(RunCodeImageDesc_IM_INTERPOLATE);
-			OP::ImageInterpolateArgs Args = Program.GetOpArgs<OP::ImageInterpolateArgs>(Item.At);
+			FOperation::ImageInterpolateArgs Args = Program.GetOpArgs<FOperation::ImageInterpolateArgs>(Item.At);
 
 			int32 NumImages = 0;
 			for (; NumImages < MUTABLE_OP_MAX_INTERPOLATE_COUNT; ++NumImages)
@@ -8158,19 +8121,17 @@ namespace UE::Mutable::Private
 			}
 			case 1: 
 			{
-				FExtendedImageDesc& Result = ImageDescResults.FindOrAdd(Item.At);
-
 				check(Args.Targets[0]);
-
-				Result = ImageDescResults.FindChecked(Args.Targets[0]);
+				FExtendedImageDesc Result = LoadImageDesc(FCacheAddress(Args.Targets[0], Item));
 				
 				for (int32 ImageIndex = 1; ImageIndex < NumImages; ++ImageIndex)
 				{
-					const FExtendedImageDesc& TargetResult = ImageDescResults.FindChecked(Args.Targets[ImageIndex]);
+					FExtendedImageDesc TargetResult = LoadImageDesc(FCacheAddress(Args.Targets[ImageIndex], Item));
 					Result.FirstLODAvailable = FMath::Max(Result.FirstLODAvailable, TargetResult.FirstLODAvailable);
 				}
 
-				StoreValidDesc(Item); break;
+				StoreImageDesc(Item, Result); 
+				break;
 			}
 			default: check(false);
 			}
@@ -8180,7 +8141,7 @@ namespace UE::Mutable::Private
 		case EOpType::IM_SATURATE:
 		{
 			MUTABLE_CPUPROFILER_SCOPE(RunCodeImageDesc_IM_SATURATE);
-			OP::ImageSaturateArgs Args = Program.GetOpArgs<OP::ImageSaturateArgs>(Item.At);
+			FOperation::ImageSaturateArgs Args = Program.GetOpArgs<FOperation::ImageSaturateArgs>(Item.At);
 			switch (Item.Stage)
 			{
 			case 0: 
@@ -8191,10 +8152,7 @@ namespace UE::Mutable::Private
 			}
 			case 1: 
 			{
-				FExtendedImageDesc& Result = ImageDescResults.FindOrAdd(Item.At);
-				Result = ImageDescResults.FindChecked(Args.Base);
-
-				StoreValidDesc(Item); 
+				StoreImageDesc(Item, LoadImageDesc(FCacheAddress(Args.Base, Item))); 
 				break;
 			}
 			default: check(false);
@@ -8205,7 +8163,8 @@ namespace UE::Mutable::Private
 		case EOpType::IM_LUMINANCE:
 		{
 			MUTABLE_CPUPROFILER_SCOPE(RunCodeImageDesc_IM_LUMINANCE);
-			OP::ImageLuminanceArgs Args = Program.GetOpArgs<OP::ImageLuminanceArgs>(Item.At);
+			FOperation::ImageLuminanceArgs Args = Program.GetOpArgs<FOperation::ImageLuminanceArgs>(Item.At);
+
 			switch (Item.Stage)
 			{
 			case 0:
@@ -8216,11 +8175,10 @@ namespace UE::Mutable::Private
 			}
 			case 1:
 			{
-				FExtendedImageDesc& Result = ImageDescResults.FindOrAdd(Item.At);
-				Result = ImageDescResults.FindChecked(Args.Base);
+				FExtendedImageDesc Result = LoadImageDesc(FCacheAddress(Args.Base, Item));
 				Result.m_format = EImageFormat::L_UByte;
 				
-				StoreValidDesc(Item);
+				StoreImageDesc(Item, Result);
 				break;
 			}
 			default:
@@ -8233,9 +8191,9 @@ namespace UE::Mutable::Private
 		case EOpType::IM_SWIZZLE:
 		{
 			MUTABLE_CPUPROFILER_SCOPE(RunCodeImageDesc_IM_SWIZZLE);
-			OP::ImageSwizzleArgs Args = Program.GetOpArgs<OP::ImageSwizzleArgs>(Item.At);
+			FOperation::ImageSwizzleArgs Args = Program.GetOpArgs<FOperation::ImageSwizzleArgs>(Item.At);
 			
-			TArray<OP::ADDRESS, TFixedAllocator<4>> ValidArgs;
+			TArray<FOperation::ADDRESS, TFixedAllocator<4>> ValidArgs;
 			for (int32 SourceIndex = 0; SourceIndex < 4; ++SourceIndex)
 			{
 				if (Args.sources[SourceIndex])
@@ -8260,20 +8218,18 @@ namespace UE::Mutable::Private
 			}
 			case 1:
 			{	
-				FExtendedImageDesc& Result = ImageDescResults.FindOrAdd(Item.At);
-
 				check(ValidArgs.Num() > 0);
 
-				Result = ImageDescResults.FindChecked(ValidArgs[0]);
+				FExtendedImageDesc Result = LoadImageDesc(FCacheAddress(ValidArgs[0], Item));
 				Result.m_format = Args.format;
 				
 				for (int32 ArgIndex = 1; ArgIndex < ValidArgs.Num(); ++ArgIndex)
 				{
-					const FExtendedImageDesc& SourceResult = ImageDescResults.FindOrAdd(ValidArgs[ArgIndex]); 
+					FExtendedImageDesc SourceResult = LoadImageDesc(FCacheAddress(ValidArgs[ArgIndex], Item));
 					Result.FirstLODAvailable = FMath::Max(Result.FirstLODAvailable, SourceResult.FirstLODAvailable);
 				}
 
-				StoreValidDesc(Item);
+				StoreImageDesc(Item, Result);
 				break;
 			}
 			default:
@@ -8286,7 +8242,7 @@ namespace UE::Mutable::Private
 		case EOpType::IM_COLOURMAP:
 		{
 			MUTABLE_CPUPROFILER_SCOPE(RunCodeImageDesc_IM_COLOURMAP);
-			OP::ImageColourMapArgs Args = Program.GetOpArgs<OP::ImageColourMapArgs>(Item.At);
+			FOperation::ImageColourMapArgs Args = Program.GetOpArgs<FOperation::ImageColourMapArgs>(Item.At);
 			switch (Item.Stage)
 			{
 			case 0: 
@@ -8297,10 +8253,7 @@ namespace UE::Mutable::Private
 			}
 			case 1: 
 			{	
-				FExtendedImageDesc& Result = ImageDescResults.FindOrAdd(Item.At);
-				Result = ImageDescResults.FindChecked(Args.Base);
-
-				StoreValidDesc(Item); 
+				StoreImageDesc(Item, LoadImageDesc(FCacheAddress(Args.Base, Item))); 
 				break;
 			}
 			default: check(false);
@@ -8311,7 +8264,7 @@ namespace UE::Mutable::Private
 		case EOpType::IM_BINARISE:
 		{
 			MUTABLE_CPUPROFILER_SCOPE(RunCodeImageDesc_IM_BINARIZE);
-			OP::ImageBinariseArgs Args = Program.GetOpArgs<OP::ImageBinariseArgs>(Item.At);
+			FOperation::ImageBinariseArgs Args = Program.GetOpArgs<FOperation::ImageBinariseArgs>(Item.At);
 			switch (Item.Stage)
 			{
 			case 0:
@@ -8322,10 +8275,10 @@ namespace UE::Mutable::Private
 			}
 			case 1:
 			{
-				FExtendedImageDesc& Result = ImageDescResults.FindOrAdd(Item.At);
+				FExtendedImageDesc Result = LoadImageDesc(FCacheAddress(Args.Base, Item));
 				Result.m_format = EImageFormat::L_UByte;
 
-				StoreValidDesc(Item);
+				StoreImageDesc(Item, Result);
 				break;
 			}
 			default:
@@ -8337,7 +8290,7 @@ namespace UE::Mutable::Private
 		case EOpType::IM_INVERT:
 		{
 			MUTABLE_CPUPROFILER_SCOPE(RunCodeImageDesc_IM_INVERT);
-			OP::ImageInvertArgs Args = Program.GetOpArgs<OP::ImageInvertArgs>(Item.At);
+			FOperation::ImageInvertArgs Args = Program.GetOpArgs<FOperation::ImageInvertArgs>(Item.At);
 			switch (Item.Stage)
 			{
 			case 0: 
@@ -8348,10 +8301,8 @@ namespace UE::Mutable::Private
 			}
 			case 1: 
 			{
-				FExtendedImageDesc& Result = ImageDescResults.FindOrAdd(Item.At);
-				Result = ImageDescResults.FindChecked(Args.Base);
+				StoreImageDesc(Item, LoadImageDesc(FCacheAddress(Args.Base, Item)));
 
-				StoreValidDesc(Item);
 				break;
 			}
 			default: check(false);
@@ -8362,22 +8313,23 @@ namespace UE::Mutable::Private
 		case EOpType::IM_PLAINCOLOUR:
 		{
 			MUTABLE_CPUPROFILER_SCOPE(RunCodeImageDesc_IM_PLAINCOLOUR);
-			OP::ImagePlainColorArgs Args = Program.GetOpArgs<OP::ImagePlainColorArgs>(Item.At);
-			FExtendedImageDesc& Result = ImageDescResults.FindOrAdd(Item.At);
+			FOperation::ImagePlainColorArgs Args = Program.GetOpArgs<FOperation::ImagePlainColorArgs>(Item.At);
+			
+			FExtendedImageDesc Result;
 			
 			Result.m_size[0] = Args.Size[0];
 			Result.m_size[1] = Args.Size[1];
 			Result.m_lods = Args.LODs;
 			Result.m_format = Args.Format;
 
-			StoreValidDesc(Item);
+			StoreImageDesc(Item, Result);
 			break;
 		}
 
 		case EOpType::IM_CROP:
 		{
 			MUTABLE_CPUPROFILER_SCOPE(RunCodeImageDesc_IM_CROP);
-			OP::ImageCropArgs Args = Program.GetOpArgs<OP::ImageCropArgs>(Item.At);
+			FOperation::ImageCropArgs Args = Program.GetOpArgs<FOperation::ImageCropArgs>(Item.At);
 			switch (Item.Stage)
 			{
 			case 0:
@@ -8388,14 +8340,13 @@ namespace UE::Mutable::Private
 			}
 			case 1:
 			{
-				FExtendedImageDesc& Result = ImageDescResults.FindOrAdd(Item.At);
-				Result = ImageDescResults.FindChecked(Args.source);
+				FExtendedImageDesc Result = LoadImageDesc(FCacheAddress(Args.source, Item));
 				
 				Result.m_size[0] = Args.sizeX;
 				Result.m_size[1] = Args.sizeY;
 				Result.m_lods = 1;
 
-				StoreValidDesc(Item);
+				StoreImageDesc(Item, Result);
 				break;
 			}
 			default:
@@ -8407,7 +8358,8 @@ namespace UE::Mutable::Private
 		case EOpType::IM_PATCH:
 		{
 			MUTABLE_CPUPROFILER_SCOPE(RunCodeImageDesc_IM_PATCH);
-			OP::ImagePatchArgs Args = Program.GetOpArgs<OP::ImagePatchArgs>(Item.At);
+
+			FOperation::ImagePatchArgs Args = Program.GetOpArgs<FOperation::ImagePatchArgs>(Item.At);
 			switch (Item.Stage)
 			{
 			case 0: 
@@ -8419,16 +8371,15 @@ namespace UE::Mutable::Private
 			}
 			case 1: 
 			{
-				FExtendedImageDesc& Result = ImageDescResults.FindOrAdd(Item.At);
-				Result = ImageDescResults.FindChecked(Args.base);
+				FExtendedImageDesc Result = LoadImageDesc(FCacheAddress(Args.base, Item));
 
 				if (Args.patch)
 				{
-					FExtendedImageDesc& PatchImageDesc = ImageDescResults.FindChecked(Args.patch);
+					FExtendedImageDesc PatchImageDesc = LoadImageDesc(FCacheAddress(Args.patch, Item));
 					Result.FirstLODAvailable = FMath::Max(Result.FirstLODAvailable, PatchImageDesc.FirstLODAvailable);
 				}
 
-				StoreValidDesc(Item); 
+				StoreImageDesc(Item, Result); 
 				break;
 			}
 			default: check(false);
@@ -8439,22 +8390,24 @@ namespace UE::Mutable::Private
 		case EOpType::IM_RASTERMESH:
 		{
 			MUTABLE_CPUPROFILER_SCOPE(RunCodeImageDesc_IM_RASTERMESH);
-			OP::ImageRasterMeshArgs Args = Program.GetOpArgs<OP::ImageRasterMeshArgs>(Item.At);
-			FExtendedImageDesc& Result = ImageDescResults.FindOrAdd(Item.At);
+
+			FOperation::ImageRasterMeshArgs Args = Program.GetOpArgs<FOperation::ImageRasterMeshArgs>(Item.At);
+			FExtendedImageDesc Result;
 			
 			Result.m_size[0] = Args.sizeX;
 			Result.m_size[1] = Args.sizeY;
 			Result.m_lods = 1;
 			Result.m_format = EImageFormat::L_UByte;
 
-			StoreValidDesc(Item);
+			StoreImageDesc(Item, Result);
 			break;
 		}
 
 		case EOpType::IM_MAKEGROWMAP:
 		{
 			MUTABLE_CPUPROFILER_SCOPE(RunCodeImageDesc_IM_MAKEGROWMAP);
-			OP::ImageMakeGrowMapArgs Args = Program.GetOpArgs<OP::ImageMakeGrowMapArgs>(Item.At);
+
+			FOperation::ImageMakeGrowMapArgs Args = Program.GetOpArgs<FOperation::ImageMakeGrowMapArgs>(Item.At);
 			switch (Item.Stage)
 			{
 			case 0:
@@ -8465,13 +8418,12 @@ namespace UE::Mutable::Private
 			}
 			case 1:
 			{
-				FExtendedImageDesc& Result = ImageDescResults.FindOrAdd(Item.At);
-				Result = ImageDescResults.FindChecked(Args.mask);
+				FExtendedImageDesc Result = LoadImageDesc(FCacheAddress(Args.mask, Item));
 
 				Result.m_format = EImageFormat::L_UByte;
 				Result.m_lods = 1;
 
-				StoreValidDesc(Item);
+				StoreImageDesc(Item, Result);
 				break;
 			}
 			default:
@@ -8484,7 +8436,8 @@ namespace UE::Mutable::Private
 		case EOpType::IM_DISPLACE:
 		{
 			MUTABLE_CPUPROFILER_SCOPE(RunCodeImageDesc_IM_DISPLACE);
-			OP::ImageDisplaceArgs Args = Program.GetOpArgs<OP::ImageDisplaceArgs>(Item.At);
+
+			FOperation::ImageDisplaceArgs Args = Program.GetOpArgs<FOperation::ImageDisplaceArgs>(Item.At);
 			switch (Item.Stage)
 			{
 			case 0: 
@@ -8495,10 +8448,7 @@ namespace UE::Mutable::Private
 			}
 			case 1: 
 			{
-				FExtendedImageDesc& Result = ImageDescResults.FindOrAdd(Item.At);
-				Result = ImageDescResults.FindChecked(Args.Source);
-
-				StoreValidDesc(Item); 
+				StoreImageDesc(Item, LoadImageDesc(FCacheAddress(Args.Source, Item))); 
 				break;
 			}
 			default: check(false);
@@ -8509,7 +8459,7 @@ namespace UE::Mutable::Private
         case EOpType::IM_TRANSFORM:
         {
 			MUTABLE_CPUPROFILER_SCOPE(RunCodeImageDesc_IM_TRANSFORM);
-			OP::ImageTransformArgs Args = Program.GetOpArgs<OP::ImageTransformArgs>(Item.At);
+			FOperation::ImageTransformArgs Args = Program.GetOpArgs<FOperation::ImageTransformArgs>(Item.At);
 
             switch (Item.Stage)
             {
@@ -8521,8 +8471,7 @@ namespace UE::Mutable::Private
 			}
             case 1:
             {
-				FExtendedImageDesc& Result = ImageDescResults.FindOrAdd(Item.At);
-				Result = ImageDescResults.FindChecked(Args.Base);
+				FExtendedImageDesc Result = LoadImageDesc(FCacheAddress(Args.Base, Item));
 			
 				Result.m_lods = 1;
 				Result.m_format = GetUncompressedFormat(Result.m_format);
@@ -8533,7 +8482,7 @@ namespace UE::Mutable::Private
 					Result.m_size[1] = Args.SizeY;
 				}
 
-				StoreValidDesc(Item);
+				StoreImageDesc(Item, Result);
                 break;
             }
 
@@ -8547,7 +8496,7 @@ namespace UE::Mutable::Private
 		case EOpType::IM_MATERIAL_BREAK:
 		{
 			MUTABLE_CPUPROFILER_SCOPE(RunCodeImageDesc_IM_MATERIAL_BREAK);
-			const OP::MaterialBreakArgs Args = Program.GetOpArgs<OP::MaterialBreakArgs>(Item.At);
+			const FOperation::MaterialBreakArgs Args = Program.GetOpArgs<FOperation::MaterialBreakArgs>(Item.At);
 
 			switch (Item.Stage)
 			{
@@ -8556,8 +8505,7 @@ namespace UE::Mutable::Private
 				FScheduledOp MaterialOp(Args.Material, Item);
 				MaterialOp.Type = FScheduledOp::EType::Full;
 
-				AddOp(FScheduledOp(Item.At, Item, 1),
-					MaterialOp);
+				AddOp(FScheduledOp(Item.At, Item, 1), MaterialOp);
 				break;
 			}
 			case 1:
@@ -8566,6 +8514,12 @@ namespace UE::Mutable::Private
 				CacheMaterialOp.Type = FScheduledOp::EType::Full;
 
 				TManagedPtr<const FMaterial> Material = LoadMaterial(CacheMaterialOp);
+
+				int32 HeapDataIndex = HeapData.Emplace();	
+				FScheduledOpData& OpHeapData = HeapData[HeapDataIndex];
+
+				OpHeapData.Resource = Material;
+
 				const FName& ParameterName = Program.ConstantNames[Args.ParameterName];
 				const FParameterKey ParameterKey = { ParameterName, Args.LayerIndex };
 				const FMaterial::FImageParameterData* ImageParameterData = Material->ImageParameters.Find(ParameterKey);
@@ -8575,19 +8529,17 @@ namespace UE::Mutable::Private
 				{
 					bParameterFound = true;
 
-					const TVariant<OP::ADDRESS, TManagedPtr<const FImage>> ImageParameter = ImageParameterData->ImageParameter;
-					if (ImageParameter.IsType<OP::ADDRESS>())
+					const TVariant<FOperation::ADDRESS, TManagedPtr<const FImage>> ImageParameter = ImageParameterData->ImageParameter;
+					if (ImageParameter.IsType<FOperation::ADDRESS>())
 					{
-						FScheduledOp ItemNextStage(Item.At, Item, 2);
-						ItemNextStage.CustomState = ImageParameter.Get<OP::ADDRESS>();
+						OpHeapData.MaterialBreak.ImageAddress = ImageParameter.Get<FOperation::ADDRESS>();
 
-						AddOp(ItemNextStage, 
-							FScheduledOp(ImageParameter.Get<OP::ADDRESS>(), Item));
+						AddOp(FScheduledOp(Item.At, Item, 2, HeapDataIndex), 
+							FScheduledOp(OpHeapData.MaterialBreak.ImageAddress, Item));
 					}
 					else
 					{
-						ImageDescResults.FindOrAdd(Item.At);
-						StoreValidDesc(Item);
+						StoreImageDesc(Item, FExtendedImageDesc{});
 					}
 				}
 				else if (const UMaterialInterface* MaterialInterface = Material->PassthroughObject.Get())
@@ -8603,17 +8555,13 @@ namespace UE::Mutable::Private
 
 					if (bParameterFound)
 					{
-						FExtendedImageDesc& Result = ImageDescResults.FindOrAdd(Item.At);
-						Result = GetExternalImageDesc(ParameterValue);
-
-						StoreValidDesc(Item);
+						StoreImageDesc(Item, GetExternalImageDesc(ParameterValue));
 					}
 				}
 
 				if (!bParameterFound)
 				{
-					ImageDescResults.FindOrAdd(Item.At);
-					StoreValidDesc(Item);
+					StoreImageDesc(Item, FExtendedImageDesc{});
 				}
 
 				break;
@@ -8623,20 +8571,16 @@ namespace UE::Mutable::Private
 				FCacheAddress CacheMaterialOp(Args.Material, Item);
 				CacheMaterialOp.Type = FScheduledOp::EType::Full;
 
-				TManagedPtr<const FMaterial> Material = LoadMaterial(CacheMaterialOp);
+				FScheduledOpData& OpHeapData = HeapData[Item.CustomState];
+
+				TManagedPtr<const FMaterial> Material = StaticCastManagedPtr<const FMaterial>(OpHeapData.Resource);
+				FOperation::ADDRESS ImageAddress = OpHeapData.MaterialBreak.ImageAddress;
+
 				const FName& ParameterName = Program.ConstantNames[Args.ParameterName];
 				const FParameterKey ParameterKey = { ParameterName, Args.LayerIndex };
 				const FMaterial::FImageParameterData* ImageParameterData = Material->ImageParameters.Find(ParameterKey);
 					
-				if (ImageParameterData)
-				{
-					const TVariant<OP::ADDRESS, TManagedPtr<const FImage>> ImageParameter = ImageParameterData->ImageParameter;
-
-					FExtendedImageDesc& Result = ImageDescResults.FindOrAdd(Item.At);
-					Result = ImageDescResults.FindChecked(ImageParameter.Get<OP::ADDRESS>());
-				}
-
-				StoreValidDesc(Item);
+				StoreImageDesc(Item, LoadImageDesc(FCacheAddress(ImageAddress, Item)));
 				break;
 			}
 
@@ -8651,7 +8595,7 @@ namespace UE::Mutable::Private
 		{
 			check(Item.Stage == 0);
 
-			const OP::MaterialBreakImageParameterArgs Args = Program.GetOpArgs<OP::MaterialBreakImageParameterArgs>(Item.At);
+			const FOperation::MaterialBreakImageParameterArgs Args = Program.GetOpArgs<FOperation::MaterialBreakImageParameterArgs>(Item.At);
 
 			TSharedPtr<FRangeIndex> Index = BuildCurrentOpRangeIndex(Item, Args.MaterialParameter);
 
@@ -8674,11 +8618,14 @@ namespace UE::Mutable::Private
 				MaterialInterface->GetTextureParameterValue(ParameterInfo, Image);
 			}
 
-			FExtendedImageDesc& Result = ImageDescResults.FindOrAdd(Item.At);
-			Result = GetExternalImageDesc(Image);
-
-			StoreValidDesc(Item);
-
+			if (Image)
+			{
+				StoreImageDesc(Item, GetExternalImageDesc(Image));
+			}
+			else
+			{
+				StoreImageDesc(Item, FExtendedImageDesc{});
+			}
 			break;
 		}
 
@@ -8688,7 +8635,7 @@ namespace UE::Mutable::Private
 				// Operation not implemented
 				check(false);
 
-				ImageDescResults.FindOrAdd(Item.At);
+				StoreImageDesc(Item, FExtendedImageDesc{});
 			}
 			break;
 		}
