[Mutable] Program cache rewrok.

diff --git a/Engine/Plugins/Mutable/Source/MutableRuntime/Private/MuR/CodeRunner.h b/Engine/Plugins/Mutable/Source/MutableRuntime/Private/MuR/CodeRunner.h
--- a/Engine/Plugins/Mutable/Source/MutableRuntime/Private/MuR/CodeRunner.h
+++ b/Engine/Plugins/Mutable/Source/MutableRuntime/Private/MuR/CodeRunner.h
@@ -17,6 +17,7 @@
 #include "Concepts/SameAs.h"
 
 
+
 namespace  UE::Mutable::Private
 {
 	class FModel;
@@ -35,7 +36,7 @@ namespace  UE::Mutable::Private
 				const FSettings& InSettings, 
 				FSystem& InSystem, 
 				EExecutionStrategy,
-				OP::ADDRESS At,
+				FOperation::ADDRESS At,
 				uint32 LodMask,
 				uint8 ExecutionOptions,
 				bool bForceLazy,
@@ -49,7 +50,7 @@ namespace  UE::Mutable::Private
 			const FSettings& InSettings, 
 			FSystem& InSystem, 
 			EExecutionStrategy,
-			OP::ADDRESS At,
+			FOperation::ADDRESS At,
 			uint32 LodMask,
 			uint8 ExecutionOptions,
 			bool bForceLazy,
@@ -58,6 +59,7 @@ namespace  UE::Mutable::Private
     protected:
 		struct FProfileContext
 		{
+			FMutex Mutex;
 			uint32 NumRunOps = 0;
 			uint32 RunOpsPerType[int32(EOpType::COUNT)] = {};
 		};
@@ -106,6 +108,11 @@ namespace  UE::Mutable::Private
 					uint8 bSplitInMultipleIterations : 1;
 				} MultiCompose;
 
+				struct
+				{
+					FOperation::ADDRESS ImageAddress;
+				} MaterialBreak;
+
 			};
 
 			TManagedPtr<const FResource> Resource;
@@ -139,7 +146,9 @@ namespace  UE::Mutable::Private
     	void RunCode_Material(const FScheduledOp&);
 
     	void RunCodeImageDesc(const FScheduledOp&);
-    
+   
+		void LaunchWaitForOthersTask(Tasks::FTaskEvent Event, uint64 ReentryCounter = 0);
+
 	public:
 		struct FExternalResourceId
 		{
@@ -173,17 +182,16 @@ namespace  UE::Mutable::Private
     protected:
 		
 		/** Image descriptor intermediate results. */
-		TMap<OP::ADDRESS, FExtendedImageDesc> ImageDescResults;
+		TMap<FOperation::ADDRESS, FExtendedImageDesc> ImageDescResults;
 		TArray<int32> ImageDescConstantImages; 
 
 		void Run(TUniquePtr<FProfileContext>&& ProfileContext, bool bForceInlineExecution);
-		void AbortRun();
 	public:
 
 		void StartRun(bool bForceInlineExecution);
 
 		/** */
-		const FExtendedImageDesc& GetImageDescResult(OP::ADDRESS ResultAddres);
+		FExtendedImageDesc GetImageDescResult(FOperation::ADDRESS ResultAddres);
 
 		//!
 		FProgramCache& GetMemory();
@@ -295,14 +303,12 @@ namespace  UE::Mutable::Private
 			if constexpr (NumDeps == 0)
 			{
 				OpenTasks.Add(Op);
-				ScheduledStagePerOp[Op] = Op.Stage + 1;
 			}
 			else
 			{
 				FScheduledOp DepsStorage[NumDeps] = {InDeps...};
 
 				ClosedTasks.Add(FTask(Op, TConstArrayView<FScheduledOp>(DepsStorage, NumDeps)));
-				ScheduledStagePerOp[Op] = Op.Stage + 1;
 				AddChildren(DepsStorage);
 			}
 			
@@ -313,12 +319,10 @@ namespace  UE::Mutable::Private
 			if (!InDeps.Num())
 			{
 				OpenTasks.Add(Op);
-				ScheduledStagePerOp[Op] = Op.Stage + 1;
 			}
 			else
 			{
 				ClosedTasks.Emplace(Op, InDeps);
-				ScheduledStagePerOp[Op] = Op.Stage + 1;
 				AddChildren(InDeps);
 			}
 		}
@@ -328,11 +332,12 @@ namespace  UE::Mutable::Private
 		{
 			return sizeof(CodeRunner) 
 				+ HeapData.GetAllocatedSize() + ImageDescResults.GetAllocatedSize()
-				+ ClosedTasks.GetAllocatedSize() + OpenTasks.GetAllocatedSize() + ScheduledStagePerOp.GetAllocatedSize()
+				+ ClosedTasks.GetAllocatedSize() + OpenTasks.GetAllocatedSize()
 				// this contains smart pointers, approximate size like this:
 				+ IssuedTasks.Max() * ( sizeof(FIssuedTask) + 16);
 		}
 
+		FString PrintDependencyGraph();
 	protected:
 
 		/** Strategy to choose the order of execution of operations. */
@@ -342,175 +347,177 @@ namespace  UE::Mutable::Private
 		bool bForceSerialTaskExecution = false;
 
 		/** List of pending operations that we don't know if they cannot be run yet because of dependencies. */
-		TArray< FTask > ClosedTasks;
+		TArray<FTask> ClosedTasks;
 
 		/** List of tasks that can be run because they don't have any unmet dependency. */
-		TArray< FScheduledOp > OpenTasks;
-
-		/** For every op, up to what stage it has been scheduled to run. */
-		CodeContainer<uint8> ScheduledStagePerOp;
+		TArray<FScheduledOp> OpenTasks;
 
 		/** List of tasks that are ready to run concurrently. */
-		TArray< TSharedPtr<FIssuedTask> > IssuedTasksOnHold;
+		TArray<TSharedPtr<FIssuedTask>> IssuedTasksOnHold;
 
 		/** List of tasks that have been set to run concurrently and their completion is unknown. */
-		TArray< TSharedPtr<FIssuedTask> > IssuedTasks;
+		TArray<TSharedPtr<FIssuedTask>> IssuedTasks;
 
 	public:
 
+		inline FExtendedImageDesc LoadImageDesc(const FCacheAddress& From)
+		{
+			return LiveInstance->Cache->LoadImageDesc(From);
+		}
+
 		inline bool LoadBool(const FCacheAddress& From) const
 		{
-			return LiveInstance->Cache->GetBool(From);
+			return LiveInstance->Cache->LoadBool(From);
 		}
 
 		inline float LoadInt(const FCacheAddress& From) const
 		{
-			return LiveInstance->Cache->GetInt(From);
+			return LiveInstance->Cache->LoadInt(From);
 		}
 
 		inline float LoadScalar(const FCacheAddress& From) const
 		{
-			return LiveInstance->Cache->GetScalar(From);
+			return LiveInstance->Cache->LoadScalar(From);
 		}
 
 		inline FVector4f LoadColor(const FCacheAddress& From) const
 		{
-			return LiveInstance->Cache->GetColor(From);
+			return LiveInstance->Cache->LoadColor(From);
 		}
 
     	inline FMatrix44f LoadMatrix(const FCacheAddress& From) const
 	    {
-			return LiveInstance->Cache->GetMatrix(From);
+			return LiveInstance->Cache->LoadMatrix(From);
 		}
 
 		inline TManagedPtr<const FMaterial> LoadMaterial(const FCacheAddress& From) const
 		{
-			return LiveInstance->Cache->GetMaterial(From);
+			return LiveInstance->Cache->LoadMaterial(From);
 		}
 
 		inline TManagedPtr<const String> LoadString(const FCacheAddress& From) const
 		{
-			return LiveInstance->Cache->GetString(From);
+			return LiveInstance->Cache->LoadString(From);
 		}
 
 		inline FProjector LoadProjector(const FCacheAddress& From) const
 		{
-			return LiveInstance->Cache->GetProjector(From);
+			return LiveInstance->Cache->LoadProjector(From);
 		}
 
 		inline TManagedPtr<const FMesh> LoadMesh(const FCacheAddress& From) const
 		{
-			return LiveInstance->Cache->GetMesh(From);
+			return LiveInstance->Cache->LoadMesh(From);
 		}
 
     	inline TManagedPtr<const FSkeletalMesh> LoadSkeletalMesh(const FCacheAddress& From) const
 	    {
-			return LiveInstance->Cache->GetSkeletalMesh(From);
+			return LiveInstance->Cache->LoadSkeletalMesh(From);
 		}
     	
 		inline TManagedPtr<const FImage> LoadImage(const FCacheAddress& From) const
 		{
-			return LiveInstance->Cache->GetImage(From);
+			return LiveInstance->Cache->LoadImage(From);
 		}
 
 		inline TManagedPtr<const FLayout> LoadLayout(const FCacheAddress& From) const
 		{
-			return LiveInstance->Cache->GetLayout(From);
+			return LiveInstance->Cache->LoadLayout(From);
 		}
 
 		inline TManagedPtr<const FInstance> LoadInstance(const FCacheAddress& From) const
 		{
-			return LiveInstance->Cache->GetInstance(From);
+			return LiveInstance->Cache->LoadInstance(From);
 		}
 
 		inline TManagedPtr<const FExtensionData> LoadExtensionData(const FCacheAddress& From) const
 		{
-			return LiveInstance->Cache->GetExtensionData(From);
+			return LiveInstance->Cache->LoadExtensionData(From);
 		}
     	
     	inline TManagedPtr<const FInstancedStruct> LoadInstancedStruct(const FCacheAddress& From) const
 	    {
-			return LiveInstance->Cache->GetInstancedStruct(From);
+			return LiveInstance->Cache->LoadInstancedStruct(From);
 		}
 
-		inline void StoreValidDesc(const FCacheAddress& To)
+		inline void StoreImageDesc(const FCacheAddress& To, const FExtendedImageDesc& ImageDesc)
 		{
-			LiveInstance->Cache->SetValidDesc(To);
+			LiveInstance->Cache->StoreImageDesc(To, ImageDesc);
 		}
 
 		inline void StoreBool(const FCacheAddress& To, bool Value)
 		{
-			LiveInstance->Cache->SetBool(To, Value);
+			LiveInstance->Cache->StoreBool(To, Value);
 		}
 
 		inline void StoreInt(const FCacheAddress& To, int32 Value)
 		{
-			LiveInstance->Cache->SetInt(To, Value);
+			LiveInstance->Cache->StoreInt(To, Value);
 		}
 
 		inline void StoreScalar(const FCacheAddress& To, float Value)
 		{
-			LiveInstance->Cache->SetScalar(To, Value);
+			LiveInstance->Cache->StoreScalar(To, Value);
 		}
 
 		inline void StoreString(const FCacheAddress& To, TManagedPtr<const String> Value)
 		{
-			LiveInstance->Cache->SetString(To, Value);
+			LiveInstance->Cache->StoreString(To, Value);
 		}
 
 		inline void StoreColor(const FCacheAddress& To, const FVector4f& Value)
 		{
-			LiveInstance->Cache->SetColour(To, Value);
+			LiveInstance->Cache->StoreColor(To, Value);
 		}
 
     	inline void StoreMatrix(const FCacheAddress& To, const FMatrix44f& Value)
 		{
-			LiveInstance->Cache->SetMatrix(To, Value);
+			LiveInstance->Cache->StoreMatrix(To, Value);
 		}
 
 		inline void StoreMaterial(const FCacheAddress& To, TManagedPtr<const FMaterial> Value)
 		{
-			LiveInstance->Cache->SetMaterial(To, Value);
+			LiveInstance->Cache->StoreMaterial(To, Value);
 		}
     	
 		inline void StoreProjector(const FCacheAddress& To, const FProjector& Value)
 		{
-			LiveInstance->Cache->SetProjector(To, Value);
+			LiveInstance->Cache->StoreProjector(To, Value);
 		}
 
 		inline void StoreMesh(const FCacheAddress& To, const TManagedPtr<const FMesh>& Resource)
 		{
-			LiveInstance->Cache->SetMesh(To, Resource);
+			LiveInstance->Cache->StoreMesh(To, Resource);
 		}
 
     	inline void StoreSkeletalMesh(const FCacheAddress& To, TManagedPtr<const FSkeletalMesh> Resource)
 		{
-			LiveInstance->Cache->SetSkeletalMesh(To, Resource);
+			LiveInstance->Cache->StoreSkeletalMesh(To, Resource);
 		}
     	
 		inline void StoreImage(const FCacheAddress& To, const TManagedPtr<const FImage>& Resource)
 		{
-			LiveInstance->Cache->SetImage(To, Resource);
+			LiveInstance->Cache->StoreImage(To, Resource);
 		}
 
 		inline void StoreLayout(const FCacheAddress& To, const TManagedPtr<const FLayout>& Resource)
 		{
-			LiveInstance->Cache->SetLayout(To, Resource);
+			LiveInstance->Cache->StoreLayout(To, Resource);
 		}
 
 		inline void StoreInstance(const FCacheAddress& To, const TManagedPtr<const FInstance>& Resource)
 		{
-			LiveInstance->Cache->SetInstance(To, Resource);
+			LiveInstance->Cache->StoreInstance(To, Resource);
 		}
 
 		inline void StoreExtensionData(const FCacheAddress& To, const TManagedPtr<const FExtensionData>& Resource)
 		{
-			LiveInstance->Cache->SetExtensionData(To, Resource.ToStrongRef());
+			LiveInstance->Cache->StoreExtensionData(To, Resource.ToStrongRef());
 		}
 
 		inline void StoreInstancedStruct(const FCacheAddress& To, const TManagedPtr<const FInstancedStruct>& Resource)
 		{
-			LiveInstance->Cache->SetInstancedStruct(To, Resource.ToStrongRef());
+			LiveInstance->Cache->StoreInstancedStruct(To, Resource.ToStrongRef());
 		}
     	
 		inline TManagedPtr<FImage> CreateImage(uint32 SizeX, uint32 SizeY, uint32 Lods, EImageFormat Format, EInitializationType Init)
diff --git a/Engine/Plugins/Mutable/Source/MutableRuntime/Private/MuR/CodeRunnerBegin.cpp b/Engine/Plugins/Mutable/Source/MutableRuntime/Private/MuR/CodeRunnerBegin.cpp
--- a/Engine/Plugins/Mutable/Source/MutableRuntime/Private/MuR/CodeRunnerBegin.cpp
+++ b/Engine/Plugins/Mutable/Source/MutableRuntime/Private/MuR/CodeRunnerBegin.cpp
@@ -19,7 +19,7 @@ namespace UE::Mutable::Private
 		Params(InLiveInstance->Parameters),
 		Program(InLiveInstance->Model->GetProgram()),
 		LodMask(InLodMask),
-		Executed(InLiveInstance->Model->GetProgram().OpAddress.Num()),
+		Executed(InLiveInstance->Model->GetProgram().NumOps),
 		ProgramCache(InLiveInstance->Cache),
 		MeshIdRegistry(InLiveInstance->MeshIdRegistry),
 		ImageIdRegistry(InLiveInstance->ImageIdRegistry),
@@ -55,33 +55,20 @@ namespace UE::Mutable::Private
 
 	FOpSet::FOpSet(int32 NumElements)
 	{
-		Index0.SetNum(NumElements, false);
+		//Index0.SetNum(NumElements, false);
+		IndexOther.Reserve(NumElements / 16);
 	}
 
 
 	bool FOpSet::Contains(const FCacheAddressInline& Item)
 	{
-		if (Item.ExecutionIndex == 0)
-		{
-			return Index0[Item.At];
-		}
-		else
-		{
-			return IndexOther.Contains(Item);
-		}
+		return IndexOther.Contains(Item);
 	}
 
 
 	void FOpSet::Add(const FCacheAddressInline& Item)
 	{
-		if (Item.ExecutionIndex == 0)
-		{
-			Index0[Item.At] = true;
-		}
-		else
-		{
-			IndexOther.Add(Item);
-		}
+		IndexOther.Add(Item);
 	}
 	
 	
@@ -99,7 +86,7 @@ namespace UE::Mutable::Private
 
 	void CodeRunnerBegin::PushOp(const FScheduledOpInline& Item)
 	{
-		check(Item.At < static_cast<uint32>(Program.OpAddress.Num()));
+		check(GetAddressByteCodeOffset(Item.At) < FOperation::ADDRESS(Program.ByteCode.Num()));
 
 		if (!Item.At)
 		{
@@ -474,7 +461,7 @@ namespace UE::Mutable::Private
 			case EOpType::MI_CONDITIONAL:
 			case EOpType::NU_CONDITIONAL:
 				{
-					OP::ConditionalArgs Args = Program.GetOpArgs<OP::ConditionalArgs>(Item.At);
+					FOperation::ConditionalArgs Args = Program.GetOpArgs<FOperation::ConditionalArgs>(Item.At);
 
 					switch(Item.Stage)
 					{
@@ -490,7 +477,7 @@ namespace UE::Mutable::Private
 						{
 							const bool Value = LoadBool(Args.condition);
 
-							OP::ADDRESS ResultAt = Value ? Args.yes : Args.no;
+							FOperation::ADDRESS ResultAt = Value ? Args.yes : Args.no;
 
 							// Schedule the end of this instruction if necessary
 							FScheduledOpInline NextStage(Item, 2);
@@ -571,17 +558,17 @@ namespace UE::Mutable::Private
 				{
 					const uint8* Data = Program.GetOpArgsPointer(Item.At);
 
-					OP::ADDRESS VarAddress;
-					FMemory::Memcpy(&VarAddress, Data, sizeof(OP::ADDRESS));
-					Data += sizeof(OP::ADDRESS);
+					FOperation::ADDRESS VarAddress;
+					FMemory::Memcpy(&VarAddress, Data, sizeof(FOperation::ADDRESS));
+					Data += sizeof(FOperation::ADDRESS);
 
-					OP::ADDRESS DefAddress;
-					FMemory::Memcpy(&DefAddress, Data, sizeof(OP::ADDRESS));
-					Data += sizeof(OP::ADDRESS);
+					FOperation::ADDRESS DefAddress;
+					FMemory::Memcpy(&DefAddress, Data, sizeof(FOperation::ADDRESS));
+					Data += sizeof(FOperation::ADDRESS);
 
-					OP::FSwitchCaseDescriptor CaseDesc;
-					FMemory::Memcpy(&CaseDesc, Data, sizeof(OP::FSwitchCaseDescriptor));
-					Data += sizeof(OP::FSwitchCaseDescriptor);
+					FOperation::FSwitchCaseDescriptor CaseDesc;
+					FMemory::Memcpy(&CaseDesc, Data, sizeof(FOperation::FSwitchCaseDescriptor));
+					Data += sizeof(FOperation::FSwitchCaseDescriptor);
 
 					switch (Item.Stage)
 					{
@@ -648,7 +635,7 @@ namespace UE::Mutable::Private
 							// Get the variable result
 							int32 Var = LoadInt(VarAddress);
 
-							OP::ADDRESS ValueAt = DefAddress;
+							FOperation::ADDRESS ValueAt = DefAddress;
 
 							if (!CaseDesc.bUseRanges)
 							{
@@ -658,9 +645,9 @@ namespace UE::Mutable::Private
 									FMemory::Memcpy(&Condition, Data, sizeof(int32));
 									Data += sizeof(int32);
 
-									OP::ADDRESS CaseAt;
-									FMemory::Memcpy(&CaseAt, Data, sizeof(OP::ADDRESS));
-									Data += sizeof(OP::ADDRESS);
+									FOperation::ADDRESS CaseAt;
+									FMemory::Memcpy(&CaseAt, Data, sizeof(FOperation::ADDRESS));
+									Data += sizeof(FOperation::ADDRESS);
 
 									if (CaseAt && Var == (int32)Condition)
 									{
@@ -681,9 +668,9 @@ namespace UE::Mutable::Private
 									FMemory::Memcpy(&RangeSize, Data, sizeof(uint32));
 									Data += sizeof(uint32);
 
-									OP::ADDRESS CaseAt;
-									FMemory::Memcpy(&CaseAt, Data, sizeof(OP::ADDRESS));
-									Data += sizeof(OP::ADDRESS);
+									FOperation::ADDRESS CaseAt;
+									FMemory::Memcpy(&CaseAt, Data, sizeof(FOperation::ADDRESS));
+									Data += sizeof(FOperation::ADDRESS);
 
 									if (CaseAt && Var >= ConditionStart && Var < int32(ConditionStart + RangeSize))
 									{
@@ -764,7 +751,7 @@ namespace UE::Mutable::Private
 				{
 					check(Item.bEvaluate);
 
-					OP::InstanceAddArgs Args = Program.GetOpArgs<OP::InstanceAddArgs>(Item.At);
+					FOperation::InstanceAddArgs Args = Program.GetOpArgs<FOperation::InstanceAddArgs>(Item.At);
 					switch (Item.Stage)
 					{
 					case 0:
@@ -797,7 +784,7 @@ namespace UE::Mutable::Private
 							if (Args.value)
 							{
 								FMeshId MeshId = FSystem::GetMeshId(MeshIdRegistry, *Model.Get(), *Params.Get(), Args.value);
-								OP::ADDRESS NameAd = Args.name;
+								FOperation::ADDRESS NameAd = Args.name;
 								check(NameAd < (uint32)Program.ConstantStrings.Num());
 								const FString& Name = Program.ConstantStrings[NameAd];
 								Instance->SetMesh(0, 0, MeshId, FName(Name));
@@ -818,7 +805,7 @@ namespace UE::Mutable::Private
 				{
 					check(Item.bEvaluate);
 
-					OP::InstanceAddMaterialArgs Args = Program.GetOpArgs<OP::InstanceAddMaterialArgs>(Item.At);
+					FOperation::InstanceAddMaterialArgs Args = Program.GetOpArgs<FOperation::InstanceAddMaterialArgs>(Item.At);
 					switch (Item.Stage)
 					{
 					case 0:
@@ -872,7 +859,7 @@ namespace UE::Mutable::Private
 				{
 					check(Item.bEvaluate);
 
-					OP::InstanceAddOverlayMaterialArgs Args = Program.GetOpArgs<OP::InstanceAddOverlayMaterialArgs>(Item.At);
+					FOperation::InstanceAddOverlayMaterialArgs Args = Program.GetOpArgs<FOperation::InstanceAddOverlayMaterialArgs>(Item.At);
 					switch (Item.Stage)
 					{
 					case 0:
@@ -933,7 +920,7 @@ namespace UE::Mutable::Private
 
 			case EOpType::IN_ADDOVERRIDEMATERIAL:
 				{
-					OP::InstanceAddOverrideMaterialArgs Args = Program.GetOpArgs<OP::InstanceAddOverrideMaterialArgs>(Item.At);
+					FOperation::InstanceAddOverrideMaterialArgs Args = Program.GetOpArgs<FOperation::InstanceAddOverrideMaterialArgs>(Item.At);
 					switch (Item.Stage)
 					{
 					case 0:
@@ -988,7 +975,7 @@ namespace UE::Mutable::Private
 				{
 					check(Item.bEvaluate);
 
-					OP::InstanceAddComponentArgs Args = Program.GetOpArgs<OP::InstanceAddComponentArgs>(Item.At);
+					FOperation::InstanceAddComponentArgs Args = Program.GetOpArgs<FOperation::InstanceAddComponentArgs>(Item.At);
 					switch (Item.Stage)
 					{
 					case 0:
@@ -1048,7 +1035,7 @@ namespace UE::Mutable::Private
 				{
 					check(Item.bEvaluate);
 
-					OP::InstanceAddArgs Args = Program.GetOpArgs<OP::InstanceAddArgs>( Item.At );
+					FOperation::InstanceAddArgs Args = Program.GetOpArgs<FOperation::InstanceAddArgs>( Item.At );
 					switch (Item.Stage)
 					{
 					case 0:
@@ -1093,7 +1080,7 @@ namespace UE::Mutable::Private
 				{
 					check(Item.bEvaluate);
 
-					OP::InstanceAddArgs Args = Program.GetOpArgs<OP::InstanceAddArgs>(Item.At);
+					FOperation::InstanceAddArgs Args = Program.GetOpArgs<FOperation::InstanceAddArgs>(Item.At);
 		            switch (Item.Stage)
 		            {
 		            case 0:
@@ -1143,7 +1130,7 @@ namespace UE::Mutable::Private
 	                    }
 
 	                    // Name
-	                    OP::ADDRESS nameAd = Args.name;
+	                    FOperation::ADDRESS nameAd = Args.name;
 	                    check( nameAd < (uint32)Program.ConstantStrings.Num() );
 	                    const FString& Name = Program.ConstantStrings[ nameAd ];
 	                    Instance->SetSurfaceName( 0, 0, SurfaceIndex, FName(Name) );
@@ -1181,9 +1168,9 @@ namespace UE::Mutable::Private
 
 		                for (uint8 LODIndex = 0; LODIndex < LODCount; ++LODIndex)
 		                {
-							OP::ADDRESS LODAddress;
-							FMemory::Memcpy(&LODAddress, Data, sizeof(OP::ADDRESS));
-							Data += sizeof(OP::ADDRESS);
+							FOperation::ADDRESS LODAddress;
+							FMemory::Memcpy(&LODAddress, Data, sizeof(FOperation::ADDRESS));
+							Data += sizeof(FOperation::ADDRESS);
 
 		                    if (LODAddress)
 		                    {
@@ -1207,9 +1194,9 @@ namespace UE::Mutable::Private
 
 		                for (uint8 LODIndex = 0; LODIndex < LODCount; ++LODIndex)
 		                {
-							OP::ADDRESS LODAddress;
-							FMemory::Memcpy(&LODAddress, Data, sizeof(OP::ADDRESS));
-							Data += sizeof(OP::ADDRESS);
+							FOperation::ADDRESS LODAddress;
+							FMemory::Memcpy(&LODAddress, Data, sizeof(FOperation::ADDRESS));
+							Data += sizeof(FOperation::ADDRESS);
 							
 							if (LODAddress)
 		                    {						
@@ -1247,7 +1234,7 @@ namespace UE::Mutable::Private
 				{
 					check(Item.bEvaluate);
 
-					OP::InstanceAddArgs Args = Program.GetOpArgs<OP::InstanceAddArgs>(Item.At);
+					FOperation::InstanceAddArgs Args = Program.GetOpArgs<FOperation::InstanceAddArgs>(Item.At);
 					switch (Item.Stage)
 					{
 					case 0:
@@ -1298,7 +1285,7 @@ namespace UE::Mutable::Private
 				{
 					check(Item.bEvaluate);
 
-					OP::InstanceAddExtensionDataArgs Args = Program.GetOpArgs<OP::InstanceAddExtensionDataArgs>(Item.At);
+					FOperation::InstanceAddExtensionDataArgs Args = Program.GetOpArgs<FOperation::InstanceAddExtensionDataArgs>(Item.At);
 
 					switch (Item.Stage)
 					{
@@ -1344,7 +1331,7 @@ namespace UE::Mutable::Private
 				{
 					if (Item.bEvaluate)
 					{
-						OP::BoolConstantArgs Args = Program.GetOpArgs<OP::BoolConstantArgs>(Item.At);
+						FOperation::BoolConstantArgs Args = Program.GetOpArgs<FOperation::BoolConstantArgs>(Item.At);
 
 						StoreBool(Item, Args.bValue);
 					}
@@ -1358,7 +1345,7 @@ namespace UE::Mutable::Private
 
 			case EOpType::BO_PARAMETER:
 				{
-					OP::ParameterArgs Args = Program.GetOpArgs<OP::ParameterArgs>(Item.At);
+					FOperation::ParameterArgs Args = Program.GetOpArgs<FOperation::ParameterArgs>(Item.At);
 
 					if (Item.bEvaluate)
 					{
@@ -1381,7 +1368,7 @@ namespace UE::Mutable::Private
 
 			case EOpType::BO_AND:
 				{
-					OP::BoolBinaryArgs Args = Program.GetOpArgs<OP::BoolBinaryArgs>(Item.At);
+					FOperation::BoolBinaryArgs Args = Program.GetOpArgs<FOperation::BoolBinaryArgs>(Item.At);
 
 					if (Item.bEvaluate)
 					{
@@ -1423,7 +1410,7 @@ namespace UE::Mutable::Private
 
 			case EOpType::BO_OR:
 				{
-					OP::BoolBinaryArgs Args = Program.GetOpArgs<OP::BoolBinaryArgs>(Item.At);
+					FOperation::BoolBinaryArgs Args = Program.GetOpArgs<FOperation::BoolBinaryArgs>(Item.At);
 
 					if (Item.bEvaluate)
 					{
@@ -1465,7 +1452,7 @@ namespace UE::Mutable::Private
 
 			case EOpType::BO_NOT:
 				{
-					OP::BoolNotArgs Args = Program.GetOpArgs<OP::BoolNotArgs>(Item.At);
+					FOperation::BoolNotArgs Args = Program.GetOpArgs<FOperation::BoolNotArgs>(Item.At);
 
 					if (Item.bEvaluate)
 					{
@@ -1504,7 +1491,7 @@ namespace UE::Mutable::Private
 
 			case EOpType::BO_EQUAL_INT_CONST:
 				{
-					OP::BoolEqualScalarConstArgs Args = Program.GetOpArgs<OP::BoolEqualScalarConstArgs>(Item.At);
+					FOperation::BoolEqualScalarConstArgs Args = Program.GetOpArgs<FOperation::BoolEqualScalarConstArgs>(Item.At);
 
 					if (Item.bEvaluate)
 					{
@@ -1545,7 +1532,7 @@ namespace UE::Mutable::Private
 				{
 					if (Item.bEvaluate)
 					{
-						OP::ScalarConstantArgs Args = Program.GetOpArgs<OP::ScalarConstantArgs>(Item.At);
+						FOperation::ScalarConstantArgs Args = Program.GetOpArgs<FOperation::ScalarConstantArgs>(Item.At);
 						
 						StoreFloat(Item, Args.Value);
 					}
@@ -1560,7 +1547,7 @@ namespace UE::Mutable::Private
 				
 			case EOpType::SC_PARAMETER:
 				{
-					OP::ParameterArgs Args = Program.GetOpArgs<OP::ParameterArgs>(Item.At);
+					FOperation::ParameterArgs Args = Program.GetOpArgs<FOperation::ParameterArgs>(Item.At);
 
 					if (Item.bEvaluate)
 					{
@@ -1583,7 +1570,7 @@ namespace UE::Mutable::Private
 				
 			case EOpType::SC_ARITHMETIC:
 				{
-					OP::ArithmeticArgs Args = Program.GetOpArgs<OP::ArithmeticArgs>(Item.At);
+					FOperation::ArithmeticArgs Args = Program.GetOpArgs<FOperation::ArithmeticArgs>(Item.At);
 
 					if (Item.bEvaluate)
 					{
@@ -1606,19 +1593,19 @@ namespace UE::Mutable::Private
 								float Result = 1.0f;
 								switch (Args.Operation)
 								{
-								case OP::ArithmeticArgs::ADD:
+								case FOperation::ArithmeticArgs::ADD:
 									Result = ValueA + ValueB;
 									break;
 
-								case OP::ArithmeticArgs::MULTIPLY:
+								case FOperation::ArithmeticArgs::MULTIPLY:
 									Result = ValueA * ValueB;
 									break;
 
-								case OP::ArithmeticArgs::SUBTRACT:
+								case FOperation::ArithmeticArgs::SUBTRACT:
 									Result = ValueA - ValueB;
 									break;
 
-								case OP::ArithmeticArgs::DIVIDE:
+								case FOperation::ArithmeticArgs::DIVIDE:
 									Result = ValueA / ValueB;
 									break;
 
@@ -1646,7 +1633,7 @@ namespace UE::Mutable::Private
 			
 			case EOpType::SC_CURVE:
 				{
-					OP::ScalarCurveArgs Args = Program.GetOpArgs<OP::ScalarCurveArgs>(Item.At);
+					FOperation::ScalarCurveArgs Args = Program.GetOpArgs<FOperation::ScalarCurveArgs>(Item.At);
 
 					if (Item.bEvaluate)
 					{
@@ -1688,7 +1675,7 @@ namespace UE::Mutable::Private
 				{
 					check(Item.bEvaluate);
 
-					OP::ExternalDataConstantArgs Args = Program.GetOpArgs<OP::ExternalDataConstantArgs>(Item.At);
+					FOperation::ExternalDataConstantArgs Args = Program.GetOpArgs<FOperation::ExternalDataConstantArgs>(Item.At);
 					PassthroughObjectLoader->Add<UObject>(Args.ExternalObjectId);
 
 					break;
@@ -1698,7 +1685,7 @@ namespace UE::Mutable::Private
 				{
 					if (Item.bEvaluate)
 					{
-						OP::ResourceConstantArgs Args = Program.GetOpArgs<OP::ResourceConstantArgs>(Item.At);
+						FOperation::ResourceConstantArgs Args = Program.GetOpArgs<FOperation::ResourceConstantArgs>(Item.At);
 				
 						check(Args.value < static_cast<uint32>(Model->GetProgram().ConstantStrings.Num()));
 
@@ -1718,7 +1705,7 @@ namespace UE::Mutable::Private
 
 			case EOpType::ST_PARAMETER:
 				{
-					OP::ParameterArgs Args = Program.GetOpArgs<OP::ParameterArgs>(Item.At);
+					FOperation::ParameterArgs Args = Program.GetOpArgs<FOperation::ParameterArgs>(Item.At);
 
 					if (Item.bEvaluate)
 					{
@@ -1747,7 +1734,7 @@ namespace UE::Mutable::Private
 				{
 					if (Item.bEvaluate)
 					{
-						OP::IntConstantArgs Args = Program.GetOpArgs<OP::IntConstantArgs>(Item.At);
+						FOperation::IntConstantArgs Args = Program.GetOpArgs<FOperation::IntConstantArgs>(Item.At);
 
 						const int32 Value = Args.Value;
 					
@@ -1763,7 +1750,7 @@ namespace UE::Mutable::Private
 
 			case EOpType::NU_PARAMETER:
 				{
-					OP::ParameterArgs Args = Program.GetOpArgs<OP::ParameterArgs>(Item.At);
+					FOperation::ParameterArgs Args = Program.GetOpArgs<FOperation::ParameterArgs>(Item.At);
 
 					if (Item.bEvaluate)
 					{
@@ -1783,7 +1770,7 @@ namespace UE::Mutable::Private
 				{
 					check(!Item.bEvaluate);
 
-					OP::ArithmeticArgs Args = Program.GetOpArgs<OP::ArithmeticArgs>(Item.At);
+					FOperation::ArithmeticArgs Args = Program.GetOpArgs<FOperation::ArithmeticArgs>(Item.At);
 					PushOp(FScheduledOpInline(Args.A, Item));
 					PushOp(FScheduledOpInline(Args.B, Item));
 
@@ -1803,7 +1790,7 @@ namespace UE::Mutable::Private
 				{
 					check(!Item.bEvaluate);
 
-					OP::ColourSampleImageArgs Args = Program.GetOpArgs<OP::ColourSampleImageArgs>(Item.At);
+					FOperation::ColourSampleImageArgs Args = Program.GetOpArgs<FOperation::ColourSampleImageArgs>(Item.At);
 					PushOp(FScheduledOpInline(Args.Image, Item));
 
 					StoreNone(Item);
@@ -1814,7 +1801,7 @@ namespace UE::Mutable::Private
 				{
 					check(!Item.bEvaluate);
 
-					OP::ColourSwizzleArgs Args = Program.GetOpArgs<OP::ColourSwizzleArgs>(Item.At);
+					FOperation::ColourSwizzleArgs Args = Program.GetOpArgs<FOperation::ColourSwizzleArgs>(Item.At);
 					PushOp(FScheduledOpInline(Args.sources[0], Item));
 					PushOp(FScheduledOpInline(Args.sources[1], Item));
 					PushOp(FScheduledOpInline(Args.sources[2], Item));
@@ -1828,7 +1815,7 @@ namespace UE::Mutable::Private
 				{
 					check(!Item.bEvaluate);
 
-					OP::MeshConstantArgs Args = Program.GetOpArgs<OP::MeshConstantArgs>(Item.At);
+					FOperation::MeshConstantArgs Args = Program.GetOpArgs<FOperation::MeshConstantArgs>(Item.At);
 					PassthroughObjectLoader->Add<UMaterialInterface>(Args.ClothID);
 					
 					StoreNone(Item);
@@ -1839,7 +1826,7 @@ namespace UE::Mutable::Private
 				{
 					check(!Item.bEvaluate);
 
-					OP::MeshApplyLayoutArgs Args = Program.GetOpArgs<OP::MeshApplyLayoutArgs>(Item.At);
+					FOperation::MeshApplyLayoutArgs Args = Program.GetOpArgs<FOperation::MeshApplyLayoutArgs>(Item.At);
 					PushOp(FScheduledOpInline(Args.Mesh, Item));
 					//PushOp(FScheduledOpInline(Args.Layout, Item));
 
@@ -1852,7 +1839,7 @@ namespace UE::Mutable::Private
 				{
 					check(!Item.bEvaluate);
 
-					OP::MeshPrepareLayoutArgs Args = Program.GetOpArgs<OP::MeshPrepareLayoutArgs>(Item.At);
+					FOperation::MeshPrepareLayoutArgs Args = Program.GetOpArgs<FOperation::MeshPrepareLayoutArgs>(Item.At);
 					PushOp(FScheduledOpInline(Args.Mesh, Item));
 					//PushOp(FScheduledOpInline(Args.Layout, Item));
 
@@ -1866,13 +1853,13 @@ namespace UE::Mutable::Private
 
 					const uint8* data = Program.GetOpArgsPointer(Item.At);
 
-					OP::ADDRESS BaseAt = 0;
-					FMemory::Memcpy(&BaseAt, data, sizeof(OP::ADDRESS)); 
-					data += sizeof(OP::ADDRESS);
+					FOperation::ADDRESS BaseAt = 0;
+					FMemory::Memcpy(&BaseAt, data, sizeof(FOperation::ADDRESS)); 
+					data += sizeof(FOperation::ADDRESS);
 
-					OP::ADDRESS TargetAt = 0;
-					FMemory::Memcpy(&TargetAt, data, sizeof(OP::ADDRESS)); 
-					data += sizeof(OP::ADDRESS);
+					FOperation::ADDRESS TargetAt = 0;
+					FMemory::Memcpy(&TargetAt, data, sizeof(FOperation::ADDRESS)); 
+					data += sizeof(FOperation::ADDRESS);
 
 					PushOp(FScheduledOpInline(BaseAt, Item));
 					PushOp(FScheduledOpInline(TargetAt, Item));
@@ -1885,7 +1872,7 @@ namespace UE::Mutable::Private
 				{
 					check(!Item.bEvaluate);
 
-					OP::MeshMorphArgs Args = Program.GetOpArgs<OP::MeshMorphArgs>(Item.At);
+					FOperation::MeshMorphArgs Args = Program.GetOpArgs<FOperation::MeshMorphArgs>(Item.At);
 
 					//PushOp(FScheduledOpInline(Args.Factor, Item));
 					PushOp(FScheduledOpInline(Args.Base, Item));
@@ -1898,7 +1885,7 @@ namespace UE::Mutable::Private
 				{
 					check(!Item.bEvaluate);
 
-					OP::MeshMergeArgs Args = Program.GetOpArgs<OP::MeshMergeArgs>(Item.At);
+					FOperation::MeshMergeArgs Args = Program.GetOpArgs<FOperation::MeshMergeArgs>(Item.At);
 					PushOp(FScheduledOpInline(Args.Base, Item));
 					PushOp(FScheduledOpInline(Args.Added, Item));
 
@@ -1910,7 +1897,7 @@ namespace UE::Mutable::Private
 				{
 					check(!Item.bEvaluate);
 
-					OP::MeshMaskClipMeshArgs Args = Program.GetOpArgs<OP::MeshMaskClipMeshArgs>(Item.At);
+					FOperation::MeshMaskClipMeshArgs Args = Program.GetOpArgs<FOperation::MeshMaskClipMeshArgs>(Item.At);
 					PushOp(FScheduledOpInline(Args.source, Item));
 					PushOp(FScheduledOpInline(Args.clip, Item));
 
@@ -1922,7 +1909,7 @@ namespace UE::Mutable::Private
 				{
 					check(!Item.bEvaluate);
 
-					OP::MeshMaskClipUVMaskArgs Args = Program.GetOpArgs<OP::MeshMaskClipUVMaskArgs>(Item.At);
+					FOperation::MeshMaskClipUVMaskArgs Args = Program.GetOpArgs<FOperation::MeshMaskClipUVMaskArgs>(Item.At);
 					PushOp(FScheduledOpInline(Args.Source, Item));
 					PushOp(FScheduledOpInline(Args.UVSource, Item));
 					PushOp(FScheduledOpInline(Args.MaskImage, Item));
@@ -1936,7 +1923,7 @@ namespace UE::Mutable::Private
 				{
 					check(!Item.bEvaluate);
 
-					OP::MeshMaskDiffArgs Args = Program.GetOpArgs<OP::MeshMaskDiffArgs>(Item.At);
+					FOperation::MeshMaskDiffArgs Args = Program.GetOpArgs<FOperation::MeshMaskDiffArgs>(Item.At);
 					PushOp(FScheduledOpInline(Args.Source, Item));
 					PushOp(FScheduledOpInline(Args.Fragment, Item));
 
@@ -1948,7 +1935,7 @@ namespace UE::Mutable::Private
 				{
 					check(!Item.bEvaluate);
 
-					OP::MeshFormatArgs Args = Program.GetOpArgs<OP::MeshFormatArgs>(Item.At);
+					FOperation::MeshFormatArgs Args = Program.GetOpArgs<FOperation::MeshFormatArgs>(Item.At);
 					PushOp(FScheduledOpInline(Args.source, Item));
 					PushOp(FScheduledOpInline(Args.format, Item));
 
@@ -1962,9 +1949,9 @@ namespace UE::Mutable::Private
 
 					const uint8* Data = Program.GetOpArgsPointer(Item.At);
 
-					OP::ADDRESS Source;
-					FMemory::Memcpy( &Source, Data, sizeof(OP::ADDRESS) );
-					Data += sizeof(OP::ADDRESS);
+					FOperation::ADDRESS Source;
+					FMemory::Memcpy( &Source, Data, sizeof(FOperation::ADDRESS) );
+					Data += sizeof(FOperation::ADDRESS);
 					
 					PushOp(FScheduledOpInline(Source, Item));
 
@@ -1976,7 +1963,7 @@ namespace UE::Mutable::Private
 				{
 					check(!Item.bEvaluate);
 
-					OP::MeshTransformArgs Args = Program.GetOpArgs<OP::MeshTransformArgs>(Item.At);
+					FOperation::MeshTransformArgs Args = Program.GetOpArgs<FOperation::MeshTransformArgs>(Item.At);
 					PushOp(FScheduledOpInline(Args.source, Item));
 					//PushOp(FScheduledOpInline(Args.matrix, Item));
 
@@ -1988,7 +1975,7 @@ namespace UE::Mutable::Private
 				{
 					check(!Item.bEvaluate);
 
-					OP::MeshClipMorphPlaneArgs Args = Program.GetOpArgs<OP::MeshClipMorphPlaneArgs>(Item.At);
+					FOperation::MeshClipMorphPlaneArgs Args = Program.GetOpArgs<FOperation::MeshClipMorphPlaneArgs>(Item.At);
 					PushOp(FScheduledOpInline(Args.Source, Item));
 
 					StoreNone(Item);
@@ -1999,7 +1986,7 @@ namespace UE::Mutable::Private
 				{
 					check(!Item.bEvaluate);
 
-					OP::MeshClipWithMeshArgs Args = Program.GetOpArgs<OP::MeshClipWithMeshArgs>(Item.At);
+					FOperation::MeshClipWithMeshArgs Args = Program.GetOpArgs<FOperation::MeshClipWithMeshArgs>(Item.At);
 					PushOp(FScheduledOpInline(Args.Source, Item));
 					PushOp(FScheduledOpInline(Args.ClipMesh, Item));
 
@@ -2011,7 +1998,7 @@ namespace UE::Mutable::Private
 				{
 					check(!Item.bEvaluate);
 
-					OP::MeshClipDeformArgs Args = Program.GetOpArgs<OP::MeshClipDeformArgs>(Item.At);
+					FOperation::MeshClipDeformArgs Args = Program.GetOpArgs<FOperation::MeshClipDeformArgs>(Item.At);
 					PushOp(FScheduledOpInline(Args.mesh, Item));
 					PushOp(FScheduledOpInline(Args.clipShape, Item));
 
@@ -2023,7 +2010,7 @@ namespace UE::Mutable::Private
 				{
 					check(!Item.bEvaluate);
 
-					OP::MeshApplyPoseArgs Args = Program.GetOpArgs<OP::MeshApplyPoseArgs>(Item.At);
+					FOperation::MeshApplyPoseArgs Args = Program.GetOpArgs<FOperation::MeshApplyPoseArgs>(Item.At);
 					PushOp(FScheduledOpInline(Args.base, Item));
 					PushOp(FScheduledOpInline(Args.pose, Item));
 
@@ -2035,7 +2022,7 @@ namespace UE::Mutable::Private
 				{
 					check(!Item.bEvaluate);
 
-					OP::MeshBindShapeArgs Args = Program.GetOpArgs<OP::MeshBindShapeArgs>(Item.At);
+					FOperation::MeshBindShapeArgs Args = Program.GetOpArgs<FOperation::MeshBindShapeArgs>(Item.At);
 					PushOp(FScheduledOpInline(Args.mesh, Item));
 					PushOp(FScheduledOpInline(Args.shape, Item));
 
@@ -2047,7 +2034,7 @@ namespace UE::Mutable::Private
 				{
 					check(!Item.bEvaluate);
 
-					OP::MeshApplyShapeArgs Args = Program.GetOpArgs<OP::MeshApplyShapeArgs>(Item.At);
+					FOperation::MeshApplyShapeArgs Args = Program.GetOpArgs<FOperation::MeshApplyShapeArgs>(Item.At);
 					PushOp(FScheduledOpInline(Args.mesh, Item));
 					PushOp(FScheduledOpInline(Args.shape, Item));
 
@@ -2059,7 +2046,7 @@ namespace UE::Mutable::Private
 				{
 					check(!Item.bEvaluate);
 
-					OP::MeshMorphReshapeArgs Args = Program.GetOpArgs<OP::MeshMorphReshapeArgs>(Item.At);
+					FOperation::MeshMorphReshapeArgs Args = Program.GetOpArgs<FOperation::MeshMorphReshapeArgs>(Item.At);
 					PushOp(FScheduledOpInline(Args.Morph, Item));
 					PushOp(FScheduledOpInline(Args.Reshape, Item));
 
@@ -2073,9 +2060,9 @@ namespace UE::Mutable::Private
 
 					const uint8* Data = Program.GetOpArgsPointer(Item.At);
 
-					OP::ADDRESS Source;
-					FMemory::Memcpy(&Source,Data,sizeof(OP::ADDRESS)); 
-					Data += sizeof(OP::ADDRESS);
+					FOperation::ADDRESS Source;
+					FMemory::Memcpy(&Source,Data,sizeof(FOperation::ADDRESS)); 
+					Data += sizeof(FOperation::ADDRESS);
 
 					PushOp(FScheduledOpInline(Source, Item));
 
@@ -2095,18 +2082,18 @@ namespace UE::Mutable::Private
 
 							for (uint16 Index = 0; Index < NumRemoves; ++Index)
 							{
-								OP::ADDRESS Condition;
-								FMemory::Memcpy(&Condition, Data, sizeof(OP::ADDRESS)); 
-								Data += sizeof(OP::ADDRESS);
+								FOperation::ADDRESS Condition;
+								FMemory::Memcpy(&Condition, Data, sizeof(FOperation::ADDRESS)); 
+								Data += sizeof(FOperation::ADDRESS);
 
 								if (Condition)
 								{
 									PushOp(FScheduledOpInline(Condition, Item, true));								
 								}
 
-								OP::ADDRESS Mask;
-								FMemory::Memcpy(&Mask, Data,sizeof(OP::ADDRESS)); 
-								Data += sizeof(OP::ADDRESS);
+								FOperation::ADDRESS Mask;
+								FMemory::Memcpy(&Mask, Data,sizeof(FOperation::ADDRESS)); 
+								Data += sizeof(FOperation::ADDRESS);
 							}
 
 							break;
@@ -2116,13 +2103,13 @@ namespace UE::Mutable::Private
 						{
 							for (uint16 Index = 0; Index < NumRemoves; ++Index)
 							{
-								OP::ADDRESS Condition;
-								FMemory::Memcpy(&Condition,Data,sizeof(OP::ADDRESS)); 
-								Data += sizeof(OP::ADDRESS);
+								FOperation::ADDRESS Condition;
+								FMemory::Memcpy(&Condition,Data,sizeof(FOperation::ADDRESS)); 
+								Data += sizeof(FOperation::ADDRESS);
 								
-								OP::ADDRESS Mask;
-								FMemory::Memcpy(&Mask,Data,sizeof(OP::ADDRESS)); 
-								Data += sizeof(OP::ADDRESS);
+								FOperation::ADDRESS Mask;
+								FMemory::Memcpy(&Mask,Data,sizeof(FOperation::ADDRESS)); 
+								Data += sizeof(FOperation::ADDRESS);
 
 								bool bValue = true;
 								if (Condition)
@@ -2151,14 +2138,14 @@ namespace UE::Mutable::Private
 				{
 					check(!Item.bEvaluate);
 
-					const OP::MeshAddMetadataArgs OpArgs = Program.GetOpArgs<OP::MeshAddMetadataArgs>(Item.At);
+					const FOperation::MeshAddMetadataArgs OpArgs = Program.GetOpArgs<FOperation::MeshAddMetadataArgs>(Item.At);
 					PushOp(FScheduledOpInline(OpArgs.Source, Item));
 
 					if (OpArgs.SkeletonId != PASSTHROUGH_ID_INVALID)
 					{
 						PassthroughObjectLoader->Add<USkeleton>(OpArgs.SkeletonId);
 					}
-					using OpEnumFlags = OP::MeshAddMetadataArgs::EnumFlags;
+					using OpEnumFlags = FOperation::MeshAddMetadataArgs::EnumFlags;
 
 					if (OpArgs.PhysicsAssetId != PASSTHROUGH_ID_INVALID)
 					{
@@ -2215,7 +2202,7 @@ namespace UE::Mutable::Private
 				{
 					check(!Item.bEvaluate);
 
-					OP::MeshProjectArgs Args = Program.GetOpArgs<OP::MeshProjectArgs>(Item.At);
+					FOperation::MeshProjectArgs Args = Program.GetOpArgs<FOperation::MeshProjectArgs>(Item.At);
 					PushOp(FScheduledOpInline(Args.Mesh, Item));
 					PushOp(FScheduledOpInline(Args.Projector, Item));
 
@@ -2227,7 +2214,7 @@ namespace UE::Mutable::Private
 				{
 					check(!Item.bEvaluate);
 					
-					OP::MeshOptimizeSkinningArgs Args = Program.GetOpArgs<OP::MeshOptimizeSkinningArgs>(Item.At);
+					FOperation::MeshOptimizeSkinningArgs Args = Program.GetOpArgs<FOperation::MeshOptimizeSkinningArgs>(Item.At);
 					PushOp(FScheduledOpInline(Args.source, Item));
 
 					StoreNone(Item);
@@ -2238,7 +2225,7 @@ namespace UE::Mutable::Private
 				{
 					check(!Item.bEvaluate);
 					
-					OP::MeshTransformWithinMeshArgs Args = Program.GetOpArgs<OP::MeshTransformWithinMeshArgs>(Item.At);
+					FOperation::MeshTransformWithinMeshArgs Args = Program.GetOpArgs<FOperation::MeshTransformWithinMeshArgs>(Item.At);
 					PushOp(FScheduledOpInline(Args.sourceMesh, Item));
 					PushOp(FScheduledOpInline(Args.boundingMesh, Item));
 					//PushOp(FScheduledOpInline(Args.matrix, Item));
@@ -2251,7 +2238,7 @@ namespace UE::Mutable::Private
 				{
 					check(!Item.bEvaluate);
 					
-					OP::MeshTransformWithBoneArgs Args = Program.GetOpArgs<OP::MeshTransformWithBoneArgs>(Item.At);
+					FOperation::MeshTransformWithBoneArgs Args = Program.GetOpArgs<FOperation::MeshTransformWithBoneArgs>(Item.At);
 					PushOp(FScheduledOpInline(Args.SourceMesh, Item));
 					//PushOp(FScheduledOpInline(Args.Matrix, Item));
 
@@ -2263,7 +2250,7 @@ namespace UE::Mutable::Private
 				{
 					check(!Item.bEvaluate);
 					
-					OP::MeshSkeletalMeshBreakArgs Args = Program.GetOpArgs<OP::MeshSkeletalMeshBreakArgs>(Item.At);
+					FOperation::MeshSkeletalMeshBreakArgs Args = Program.GetOpArgs<FOperation::MeshSkeletalMeshBreakArgs>(Item.At);
 					PushOp(FScheduledOpInline(Args.SkeletalMeshParameter, Item));
 
 					StoreNone(Item);
@@ -2275,7 +2262,7 @@ namespace UE::Mutable::Private
 				{
 					check(!Item.bEvaluate);
 					
-					OP::ResourceReferenceArgs Args = Program.GetOpArgs<OP::ResourceReferenceArgs>(Item.At);
+					FOperation::ResourceReferenceArgs Args = Program.GetOpArgs<FOperation::ResourceReferenceArgs>(Item.At);
 
 					if (!Args.ForceLoad) // Not a passthrough. Only happens with None optimization.
 					{
@@ -2290,7 +2277,7 @@ namespace UE::Mutable::Private
 				{
 					check(!Item.bEvaluate);
 					
-					OP::ImageLayerColourArgs Args = Program.GetOpArgs<OP::ImageLayerColourArgs>(Item.At);
+					FOperation::ImageLayerColourArgs Args = Program.GetOpArgs<FOperation::ImageLayerColourArgs>(Item.At);
 					PushOp(FScheduledOpInline(Args.base, Item));
 					PushOp(FScheduledOpInline(Args.mask, Item));
 					//PushOp(FScheduledOpInline(Args.colour, Item));
@@ -2303,7 +2290,7 @@ namespace UE::Mutable::Private
 				{
 					check(!Item.bEvaluate);
 					
-					OP::ImageLayerArgs Args = Program.GetOpArgs<OP::ImageLayerArgs>(Item.At);	
+					FOperation::ImageLayerArgs Args = Program.GetOpArgs<FOperation::ImageLayerArgs>(Item.At);	
 					PushOp(FScheduledOpInline(Args.base, Item));
 					PushOp(FScheduledOpInline(Args.mask, Item));
 					PushOp(FScheduledOpInline(Args.blended, Item));
@@ -2316,7 +2303,7 @@ namespace UE::Mutable::Private
 				{
 					check(!Item.bEvaluate);
 					
-					OP::ImageMultiLayerArgs Args = Program.GetOpArgs<OP::ImageMultiLayerArgs>(Item.At);
+					FOperation::ImageMultiLayerArgs Args = Program.GetOpArgs<FOperation::ImageMultiLayerArgs>(Item.At);
 
 					switch (Item.Stage)
 					{
@@ -2347,11 +2334,9 @@ namespace UE::Mutable::Private
 
 							for (int32 IterationIndex = 0; IterationIndex < NumIterations; ++IterationIndex)
 							{
-								ExecutionIndex Index = ProgramCache->GetRangeIndex(Item.ExecutionIndex);
-								Index.SetFromModelRangeIndex(Args.rangeId, IterationIndex);
-
 								FScheduledOpInline ItemCopy = Item;
-								ItemCopy.ExecutionIndex = ProgramCache->GetRangeIndexIndex(Index);
+								ItemCopy.ExecutionIndex = ProgramCache->GetOrAddModifiedExecutionIndex(
+										Item.ExecutionIndex, Args.rangeId, IterationIndex);
 
 								PushOp(FScheduledOpInline(Args.mask, Item));
 								PushOp(FScheduledOpInline(Args.blended, Item));
@@ -2372,7 +2357,7 @@ namespace UE::Mutable::Private
 				{
 					check(!Item.bEvaluate);
 					
-					OP::ImageNormalCompositeArgs Args = Program.GetOpArgs<OP::ImageNormalCompositeArgs>(Item.At);
+					FOperation::ImageNormalCompositeArgs Args = Program.GetOpArgs<FOperation::ImageNormalCompositeArgs>(Item.At);
 					PushOp(FScheduledOpInline(Args.base, Item));
 					PushOp(FScheduledOpInline(Args.normal, Item));
 
@@ -2384,7 +2369,7 @@ namespace UE::Mutable::Private
 				{
 					check(!Item.bEvaluate);
 					
-					OP::ImagePixelFormatArgs Args = Program.GetOpArgs<OP::ImagePixelFormatArgs>(Item.At);
+					FOperation::ImagePixelFormatArgs Args = Program.GetOpArgs<FOperation::ImagePixelFormatArgs>(Item.At);
 					PushOp(FScheduledOpInline(Args.source, Item));
 
 					StoreNone(Item);
@@ -2395,7 +2380,7 @@ namespace UE::Mutable::Private
 				{
 					check(!Item.bEvaluate);
 					
-					OP::ImageMipmapArgs Args = Program.GetOpArgs<OP::ImageMipmapArgs>(Item.At);
+					FOperation::ImageMipmapArgs Args = Program.GetOpArgs<FOperation::ImageMipmapArgs>(Item.At);
 					PushOp(FScheduledOpInline(Args.source, Item));
 
 					StoreNone(Item);
@@ -2406,7 +2391,7 @@ namespace UE::Mutable::Private
 				{
 					check(!Item.bEvaluate);
 					
-					OP::ImageResizeArgs Args = Program.GetOpArgs<OP::ImageResizeArgs>(Item.At);
+					FOperation::ImageResizeArgs Args = Program.GetOpArgs<FOperation::ImageResizeArgs>(Item.At);
 					PushOp(FScheduledOpInline(Args.Source, Item));
 
 					StoreNone(Item);
@@ -2417,7 +2402,7 @@ namespace UE::Mutable::Private
 				{
 					check(!Item.bEvaluate);
 					
-					OP::ImageResizeLikeArgs Args = Program.GetOpArgs<OP::ImageResizeLikeArgs>(Item.At);
+					FOperation::ImageResizeLikeArgs Args = Program.GetOpArgs<FOperation::ImageResizeLikeArgs>(Item.At);
 					PushOp(FScheduledOpInline(Args.Source, Item));
 					PushOp(FScheduledOpInline(Args.SizeSource, Item));
 
@@ -2429,7 +2414,7 @@ namespace UE::Mutable::Private
 				{
 					check(!Item.bEvaluate);
 					
-					OP::ImageResizeRelArgs Args = Program.GetOpArgs<OP::ImageResizeRelArgs>(Item.At);
+					FOperation::ImageResizeRelArgs Args = Program.GetOpArgs<FOperation::ImageResizeRelArgs>(Item.At);
 					PushOp(FScheduledOpInline(Args.Source, Item));
 
 					StoreNone(Item);
@@ -2440,7 +2425,7 @@ namespace UE::Mutable::Private
 				{
 					check(!Item.bEvaluate);
 					
-					OP::ImageBlankLayoutArgs Args = Program.GetOpArgs<OP::ImageBlankLayoutArgs>(Item.At);
+					FOperation::ImageBlankLayoutArgs Args = Program.GetOpArgs<FOperation::ImageBlankLayoutArgs>(Item.At);
 					PushOp(FScheduledOpInline(Args.Layout, Item));
 
 					StoreNone(Item);
@@ -2451,7 +2436,7 @@ namespace UE::Mutable::Private
 				{
 					check(!Item.bEvaluate);
 					
-					OP::ImageComposeArgs Args = Program.GetOpArgs<OP::ImageComposeArgs>(Item.At);
+					FOperation::ImageComposeArgs Args = Program.GetOpArgs<FOperation::ImageComposeArgs>(Item.At);
 					//PushOp(FScheduledOpInline(Args.layout, Item));
 					PushOp(FScheduledOpInline(Args.base, Item));
 					PushOp(FScheduledOpInline(Args.blockImage, Item)); 
@@ -2463,7 +2448,7 @@ namespace UE::Mutable::Private
 
 			case EOpType::IM_MULTICOMPOSE:
 			{
-				OP::ImageMultiComposeArgs Args = Program.GetOpArgs<OP::ImageMultiComposeArgs>(Item.At);
+				FOperation::ImageMultiComposeArgs Args = Program.GetOpArgs<FOperation::ImageMultiComposeArgs>(Item.At);
 				PushOp(FScheduledOpInline(Args.Base, Item));
 				PushOp(FScheduledOpInline(Args.SourceImage, Item));
 
@@ -2475,7 +2460,7 @@ namespace UE::Mutable::Private
 				{
 					check(!Item.bEvaluate);
 					
-					OP::ImageInterpolateArgs Args = Program.GetOpArgs<OP::ImageInterpolateArgs>(Item.At);
+					FOperation::ImageInterpolateArgs Args = Program.GetOpArgs<FOperation::ImageInterpolateArgs>(Item.At);
 
 					//PushOp(FScheduledOpInline(Args.Factor, Item));
 					
@@ -2497,7 +2482,7 @@ namespace UE::Mutable::Private
 				{
 					check(!Item.bEvaluate);
 					
-					OP::ImageSaturateArgs Args = Program.GetOpArgs<OP::ImageSaturateArgs>(Item.At);
+					FOperation::ImageSaturateArgs Args = Program.GetOpArgs<FOperation::ImageSaturateArgs>(Item.At);
 					//PushOp(FScheduledOpInline(Args.Factor, Item, false));
 					PushOp(FScheduledOpInline(Args.Base, Item));
 
@@ -2509,7 +2494,7 @@ namespace UE::Mutable::Private
 				{
 					check(!Item.bEvaluate);
 					
-					OP::ImageLuminanceArgs Args = Program.GetOpArgs<OP::ImageLuminanceArgs>(Item.At);
+					FOperation::ImageLuminanceArgs Args = Program.GetOpArgs<FOperation::ImageLuminanceArgs>(Item.At);
 					PushOp(FScheduledOpInline(Args.Base, Item));
 
 					StoreNone(Item);
@@ -2520,9 +2505,9 @@ namespace UE::Mutable::Private
 				{
 					check(!Item.bEvaluate);
 					
-					OP::ImageSwizzleArgs Args = Program.GetOpArgs<OP::ImageSwizzleArgs>(Item.At);
+					FOperation::ImageSwizzleArgs Args = Program.GetOpArgs<FOperation::ImageSwizzleArgs>(Item.At);
 				
-					TArray<OP::ADDRESS, TFixedAllocator<4>> ValidArgs;
+					TArray<FOperation::ADDRESS, TFixedAllocator<4>> ValidArgs;
 					for (int32 SourceIndex = 0; SourceIndex < 4; ++SourceIndex)
 					{
 						if (Args.sources[SourceIndex])
@@ -2539,7 +2524,7 @@ namespace UE::Mutable::Private
 				{
 					check(!Item.bEvaluate);
 					
-					OP::ImageColourMapArgs Args = Program.GetOpArgs<OP::ImageColourMapArgs>(Item.At);
+					FOperation::ImageColourMapArgs Args = Program.GetOpArgs<FOperation::ImageColourMapArgs>(Item.At);
 					PushOp(FScheduledOpInline(Args.Base, Item));
 					//PushOp(FScheduledOpInline(Args.Mask, Item));
 					//PushOp(FScheduledOpInline(Args.Map, Item));
@@ -2552,7 +2537,7 @@ namespace UE::Mutable::Private
 				{
 					check(!Item.bEvaluate);
 					
-					OP::ImageBinariseArgs Args = Program.GetOpArgs<OP::ImageBinariseArgs>(Item.At);
+					FOperation::ImageBinariseArgs Args = Program.GetOpArgs<FOperation::ImageBinariseArgs>(Item.At);
 					PushOp(FScheduledOpInline(Args.Base, Item));
 					//PushOp(FScheduledOpInline(Args.Threshold, Item));
 
@@ -2564,7 +2549,7 @@ namespace UE::Mutable::Private
 				{
 					check(!Item.bEvaluate);
 					
-					OP::ImageInvertArgs Args = Program.GetOpArgs<OP::ImageInvertArgs>(Item.At);
+					FOperation::ImageInvertArgs Args = Program.GetOpArgs<FOperation::ImageInvertArgs>(Item.At);
 					PushOp(FScheduledOpInline(Args.Base, Item));
 
 					StoreNone(Item);
@@ -2583,7 +2568,7 @@ namespace UE::Mutable::Private
 				{
 					check(!Item.bEvaluate);
 					
-					OP::ImageCropArgs Args = Program.GetOpArgs<OP::ImageCropArgs>(Item.At);
+					FOperation::ImageCropArgs Args = Program.GetOpArgs<FOperation::ImageCropArgs>(Item.At);
 					PushOp(FScheduledOpInline(Args.source, Item));
 
 					StoreNone(Item);
@@ -2594,7 +2579,7 @@ namespace UE::Mutable::Private
 				{
 					check(!Item.bEvaluate);
 					
-					OP::ImagePatchArgs Args = Program.GetOpArgs<OP::ImagePatchArgs>(Item.At);
+					FOperation::ImagePatchArgs Args = Program.GetOpArgs<FOperation::ImagePatchArgs>(Item.At);
 					PushOp(FScheduledOpInline(Args.base, Item));
 					PushOp(FScheduledOpInline(Args.patch, Item));
 
@@ -2606,7 +2591,7 @@ namespace UE::Mutable::Private
 				{
 					check(!Item.bEvaluate);
 					
-					OP::ImageRasterMeshArgs Args = Program.GetOpArgs<OP::ImageRasterMeshArgs>(Item.At);
+					FOperation::ImageRasterMeshArgs Args = Program.GetOpArgs<FOperation::ImageRasterMeshArgs>(Item.At);
 					PushOp(FScheduledOpInline(Args.mesh, Item));
 					//PushOp(FScheduledOpInline(Args.image, Item));
 					//PushOp(FScheduledOpInline(Args.angleFadeProperties, Item));
@@ -2621,7 +2606,7 @@ namespace UE::Mutable::Private
 				{
 					check(!Item.bEvaluate);
 					
-					OP::ImageMakeGrowMapArgs Args = Program.GetOpArgs<OP::ImageMakeGrowMapArgs>(Item.At);
+					FOperation::ImageMakeGrowMapArgs Args = Program.GetOpArgs<FOperation::ImageMakeGrowMapArgs>(Item.At);
 					PushOp(FScheduledOpInline(Args.mask, Item));
 
 					StoreNone(Item);
@@ -2632,7 +2617,7 @@ namespace UE::Mutable::Private
 				{
 					check(!Item.bEvaluate);
 					
-					OP::ImageDisplaceArgs Args = Program.GetOpArgs<OP::ImageDisplaceArgs>(Item.At);
+					FOperation::ImageDisplaceArgs Args = Program.GetOpArgs<FOperation::ImageDisplaceArgs>(Item.At);
 					PushOp(FScheduledOpInline(Args.Source, Item));
 					//PushOp(FScheduledOpInline(Args.DisplacementMap, Item));
 
@@ -2644,7 +2629,7 @@ namespace UE::Mutable::Private
 				{
 					check(!Item.bEvaluate);
 					
-					OP::ImageTransformArgs Args = Program.GetOpArgs<OP::ImageTransformArgs>(Item.At);
+					FOperation::ImageTransformArgs Args = Program.GetOpArgs<FOperation::ImageTransformArgs>(Item.At);
 					PushOp(FScheduledOpInline(Args.Base, Item));
 					//PushOp(FScheduledOpInline(Args.OffsetX, Item));
 					//PushOp(FScheduledOpInline(Args.OffsetY, Item));
@@ -2660,7 +2645,7 @@ namespace UE::Mutable::Private
 				{
 					check(!Item.bEvaluate);
 					
-					OP::ImageParameterConvertArgs Args = Program.GetOpArgs<OP::ImageParameterConvertArgs>(Item.At);
+					FOperation::ImageParameterConvertArgs Args = Program.GetOpArgs<FOperation::ImageParameterConvertArgs>(Item.At);
 					PushOp(FScheduledOpInline(Args.ImageParameter, Item));
 
 					StoreNone(Item);
@@ -2671,7 +2656,7 @@ namespace UE::Mutable::Private
 				{
 					check(!Item.bEvaluate);
 					
-					OP::LayoutMergeArgs Args = Program.GetOpArgs<OP::LayoutMergeArgs>(Item.At);
+					FOperation::LayoutMergeArgs Args = Program.GetOpArgs<FOperation::LayoutMergeArgs>(Item.At);
 					PushOp(FScheduledOpInline(Args.Base, Item));
 					PushOp(FScheduledOpInline(Args.Added, Item));
 
@@ -2683,7 +2668,7 @@ namespace UE::Mutable::Private
 				{
 					check(!Item.bEvaluate);
 					
-					OP::LayoutPackArgs Args = Program.GetOpArgs<OP::LayoutPackArgs>(Item.At);
+					FOperation::LayoutPackArgs Args = Program.GetOpArgs<FOperation::LayoutPackArgs>(Item.At);
 					PushOp(FScheduledOpInline( Args.Source, Item));
 
 					StoreNone(Item);
@@ -2694,7 +2679,7 @@ namespace UE::Mutable::Private
 				{
 					check(!Item.bEvaluate);
 					
-					OP::LayoutFromMeshArgs Args = Program.GetOpArgs<OP::LayoutFromMeshArgs>(Item.At);
+					FOperation::LayoutFromMeshArgs Args = Program.GetOpArgs<FOperation::LayoutFromMeshArgs>(Item.At);
 					PushOp(FScheduledOpInline(Args.Mesh, Item));
 
 					StoreNone(Item);
@@ -2705,7 +2690,7 @@ namespace UE::Mutable::Private
 				{
 					check(!Item.bEvaluate);
 					
-					OP::LayoutRemoveBlocksArgs Args = Program.GetOpArgs<OP::LayoutRemoveBlocksArgs>(Item.At);
+					FOperation::LayoutRemoveBlocksArgs Args = Program.GetOpArgs<FOperation::LayoutRemoveBlocksArgs>(Item.At);
 					PushOp(FScheduledOpInline(Args.Source, Item));
 					PushOp(FScheduledOpInline(Args.ReferenceLayout, Item));
 
@@ -2725,7 +2710,7 @@ namespace UE::Mutable::Private
 				{
 					check(!Item.bEvaluate);
 					
-					OP::MaterialConstantArgs Args = Program.GetOpArgs<OP::MaterialConstantArgs>(Item.At);
+					FOperation::MaterialConstantArgs Args = Program.GetOpArgs<FOperation::MaterialConstantArgs>(Item.At);
 					PassthroughObjectLoader->Add<UMaterialInterface>(Args.ID);
 
 					TManagedPtr<const FMaterial> Material = Program.ConstantMaterials[Args.Value];
@@ -2739,7 +2724,7 @@ namespace UE::Mutable::Private
 			case EOpType::IM_MATERIAL_BREAK:
 			case EOpType::CO_MATERIAL_BREAK:
 				{
-					OP::MaterialBreakArgs Args = Program.GetOpArgs<OP::MaterialBreakArgs>(Item.At);
+					FOperation::MaterialBreakArgs Args = Program.GetOpArgs<FOperation::MaterialBreakArgs>(Item.At);
 
 					switch (Item.Stage)
 					{
@@ -2769,11 +2754,11 @@ namespace UE::Mutable::Private
 
 									if (ImageParameterData)
 									{
-										const TVariant<OP::ADDRESS, TManagedPtr<const FImage>> ImageParameter = ImageParameterData->ImageParameter;
+										const TVariant<FOperation::ADDRESS, TManagedPtr<const FImage>> ImageParameter = ImageParameterData->ImageParameter;
 
-										if (ImageParameter.IsType<OP::ADDRESS>())
+										if (ImageParameter.IsType<FOperation::ADDRESS>())
 										{
-											PushOp(FScheduledOpInline(ImageParameter.Get<OP::ADDRESS>(), Item));
+											PushOp(FScheduledOpInline(ImageParameter.Get<FOperation::ADDRESS>(), Item));
 										}
 									}
 
@@ -2810,8 +2795,8 @@ namespace UE::Mutable::Private
 					
 					const uint8* Data = Program.GetOpArgsPointer(Item.At);
 
-					OP::ExternalArgs Args = Program.GetOpArgs<OP::ExternalArgs>(Item.At);
-					Data += sizeof(OP::ExternalArgs);
+					FOperation::ExternalArgs Args = Program.GetOpArgs<FOperation::ExternalArgs>(Item.At);
+					Data += sizeof(FOperation::ExternalArgs);
 
 					const FInstancedStruct& OperationInstancedStruct = Program.ExternalOperationProvider->Get(Item.At);
 					if (!OperationInstancedStruct.IsValid()) // Plugin not loaded
@@ -2837,9 +2822,9 @@ namespace UE::Mutable::Private
 
 							for (int32 Index = 0; Index < Args.NumOperants; ++Index)
 							{
-								OP::ADDRESS OperantAt;
-								FMemory::Memcpy(&OperantAt, Data, sizeof(OP::ADDRESS));
-								Data += sizeof(OP::ADDRESS);
+								FOperation::ADDRESS OperantAt;
+								FMemory::Memcpy(&OperantAt, Data, sizeof(FOperation::ADDRESS));
+								Data += sizeof(FOperation::ADDRESS);
 
 								PushOp(FScheduledOpInline(OperantAt, Item));
 							}
@@ -2851,9 +2836,9 @@ namespace UE::Mutable::Private
 						{
 							for (int32 Index = 0; Index < Args.NumOperants; ++Index)
 							{
-								OP::ADDRESS OperantAt;
-								FMemory::Memcpy(&OperantAt, Data, sizeof(OP::ADDRESS));
-								Data += sizeof(OP::ADDRESS);
+								FOperation::ADDRESS OperantAt;
+								FMemory::Memcpy(&OperantAt, Data, sizeof(FOperation::ADDRESS));
+								Data += sizeof(FOperation::ADDRESS);
 
 								const EOpType InputOpType = Program.GetOpType(OperantAt);
 								const EDataType InputDataType = GetOpDataType(InputOpType);
@@ -2889,7 +2874,7 @@ namespace UE::Mutable::Private
 
 			case EOpType::MI_SKELETALMESH_BREAK:
 				{
-					const OP::MaterialSkeletalMeshBreakArgs Args = Program.GetOpArgs<OP::MaterialSkeletalMeshBreakArgs>(Item.At);
+					const FOperation::MaterialSkeletalMeshBreakArgs Args = Program.GetOpArgs<FOperation::MaterialSkeletalMeshBreakArgs>(Item.At);
 					PushOp(FScheduledOpInline(Args.SkeletalMesh, Item));
 					
 					StoreMaterial(Item, MakeManaged<const FMaterial>());
@@ -2899,7 +2884,7 @@ namespace UE::Mutable::Private
 			
 			case EOpType::MI_MODIFY:
 			{
-				OP::MaterialModifyArgs Args = Program.GetOpArgs<OP::MaterialModifyArgs>(Item.At);
+				FOperation::MaterialModifyArgs Args = Program.GetOpArgs<FOperation::MaterialModifyArgs>(Item.At);
 				int32 ParameterCount = Args.NumParameters;
 
 				switch (Item.Stage)
@@ -2911,7 +2896,7 @@ namespace UE::Mutable::Private
 
 						//Dynamic data can't be stored into arg structures
 						const uint8* Data = Program.GetOpArgsPointer(Item.At);
-						Data += sizeof(OP::MaterialModifyArgs);
+						Data += sizeof(FOperation::MaterialModifyArgs);
 
 						// First, Schedule the material
 						PushOp(FScheduledOpInline(Args.Material, Item));
@@ -2919,17 +2904,17 @@ namespace UE::Mutable::Private
 						// Store parameter names and schedule the parameter operations
 						for (int32 ParameterIndex = 0; ParameterIndex < ParameterCount; ++ParameterIndex)
 						{
-							OP::ADDRESS NewName;
-							FMemory::Memcpy(&NewName, Data, sizeof(OP::ADDRESS));
-							Data += sizeof(OP::ADDRESS);
+							FOperation::ADDRESS NewName;
+							FMemory::Memcpy(&NewName, Data, sizeof(FOperation::ADDRESS));
+							Data += sizeof(FOperation::ADDRESS);
 
 							int8 LayerIndex;
 							FMemory::Memcpy(&LayerIndex, Data, sizeof(int8));
 							Data += sizeof(int8);
 
-							OP::ADDRESS NewParameterOperation;
-							FMemory::Memcpy(&NewParameterOperation, Data, sizeof(OP::ADDRESS));
-							Data += sizeof(OP::ADDRESS);
+							FOperation::ADDRESS NewParameterOperation;
+							FMemory::Memcpy(&NewParameterOperation, Data, sizeof(FOperation::ADDRESS));
+							Data += sizeof(FOperation::ADDRESS);
 
 							// Skip 0 addresses since do not need to load anything
 							if (NewParameterOperation == 0)
@@ -2955,7 +2940,7 @@ namespace UE::Mutable::Private
 						const uint8* Data = Program.GetOpArgsPointer(Item.At);
 						
 						//Material
-						Data += sizeof(OP::MaterialModifyArgs);
+						Data += sizeof(FOperation::MaterialModifyArgs);
 
 						// Pop Material Value
 						TManagedPtr<const FMaterial> Material = LoadMaterial(Args.Material);
@@ -2963,17 +2948,17 @@ namespace UE::Mutable::Private
 						// Store parameter names and schedule the parameter operations
 						for (int32 ParameterIndex = 0; ParameterIndex < ParameterCount; ++ParameterIndex)
 						{
-							OP::ADDRESS NewName;
-							FMemory::Memcpy(&NewName, Data, sizeof(OP::ADDRESS));
-							Data += sizeof(OP::ADDRESS);
+							FOperation::ADDRESS NewName;
+							FMemory::Memcpy(&NewName, Data, sizeof(FOperation::ADDRESS));
+							Data += sizeof(FOperation::ADDRESS);
 
 							int8 LayerIndex;
 							FMemory::Memcpy(&LayerIndex, Data, sizeof(int8));
 							Data += sizeof(int8);
 
-							OP::ADDRESS ParameterOperation;
-							FMemory::Memcpy(&ParameterOperation, Data, sizeof(OP::ADDRESS));
-							Data += sizeof(OP::ADDRESS);
+							FOperation::ADDRESS ParameterOperation;
+							FMemory::Memcpy(&ParameterOperation, Data, sizeof(FOperation::ADDRESS));
+							Data += sizeof(FOperation::ADDRESS);
 
 							// Skip 0 addresses since do not need to load anything
 							if (ParameterOperation == 0)
diff --git a/Engine/Plugins/Mutable/Source/MutableRuntime/Private/MuR/CodeRunnerBegin.h b/Engine/Plugins/Mutable/Source/MutableRuntime/Private/MuR/CodeRunnerBegin.h
--- a/Engine/Plugins/Mutable/Source/MutableRuntime/Private/MuR/CodeRunnerBegin.h
+++ b/Engine/Plugins/Mutable/Source/MutableRuntime/Private/MuR/CodeRunnerBegin.h
@@ -18,7 +18,7 @@ namespace UE::Mutable::Private
 
 	struct FScheduledOpInline
 	{
-		explicit FScheduledOpInline(OP::ADDRESS InAt = 0, const uint16 InStage = 0, const uint16 InExecutionIndex = 0, bool bInEvaluate = true) :
+		explicit FScheduledOpInline(FOperation::ADDRESS InAt = 0, const uint16 InStage = 0, const uint16 InExecutionIndex = 0, bool bInEvaluate = true) :
 			At(InAt),
 			Stage(InStage),
 			bEvaluate(bInEvaluate),
@@ -26,14 +26,14 @@ namespace UE::Mutable::Private
 		{
 		}
 
-		explicit FScheduledOpInline(OP::ADDRESS InAt, const FScheduledOpInline& InOpTemplate) :
+		explicit FScheduledOpInline(FOperation::ADDRESS InAt, const FScheduledOpInline& InOpTemplate) :
 			At(InAt),
 			bEvaluate(InOpTemplate.bEvaluate),
 			ExecutionIndex(InOpTemplate.ExecutionIndex)
 		{
 		}
 		
-		explicit FScheduledOpInline(OP::ADDRESS InAt, const FScheduledOpInline& InOpTemplate, bool bInEvaluate) :
+		explicit FScheduledOpInline(FOperation::ADDRESS InAt, const FScheduledOpInline& InOpTemplate, bool bInEvaluate) :
 			At(InAt),
 			bEvaluate(bInEvaluate),
 			ExecutionIndex(InOpTemplate.ExecutionIndex)
@@ -50,11 +50,11 @@ namespace UE::Mutable::Private
 
 		bool operator==(const FScheduledOpInline& Other) const;
 
-		OP::ADDRESS At = 0;
+		FOperation::ADDRESS At = 0;
 		uint8 Stage = 0;
 		uint8 bEvaluate = 0;
 		uint16 ExecutionIndex = 0;
-		OP::ADDRESS CustomState = 0;
+		FOperation::ADDRESS CustomState = 0;
 	};
 
 
@@ -63,7 +63,7 @@ namespace UE::Mutable::Private
 
 	struct FCacheAddressInline
 	{
-		FCacheAddressInline(const OP::ADDRESS& InAt)
+		FCacheAddressInline(const FOperation::ADDRESS& InAt)
 		{
 			At = InAt;
 		}
@@ -74,7 +74,7 @@ namespace UE::Mutable::Private
 			ExecutionIndex = Item.ExecutionIndex;
 		}
 		
-		OP::ADDRESS At = 0;
+		FOperation::ADDRESS At = 0;
 		uint16 ExecutionIndex = 0;
 
 		bool operator==(const FCacheAddressInline&) const = default;
@@ -94,7 +94,7 @@ namespace UE::Mutable::Private
 		void Add(const FCacheAddressInline& Item);
 
 	private:
-		TBitArray<> Index0;
+		//TBitArray<> Index0;
 		TSet<FCacheAddressInline> IndexOther;
 	};
 	
diff --git a/Engine/Plugins/Mutable/Source/MutableRuntime/Private/MuR/CodeRunnerTasked.cpp b/Engine/Plugins/Mutable/Source/MutableRuntime/Private/MuR/CodeRunnerTasked.cpp
--- a/Engine/Plugins/Mutable/Source/MutableRuntime/Private/MuR/CodeRunnerTasked.cpp
+++ b/Engine/Plugins/Mutable/Source/MutableRuntime/Private/MuR/CodeRunnerTasked.cpp
@@ -53,12 +53,16 @@ static TAutoConsoleVariable<bool> CVarCodeRunnerForceInline(
 	TEXT("If true, force all Code Runners to be inline (do not split them into tasks)."),
 	ECVF_Default);
 
+static TAutoConsoleVariable<float> CVarRomLoadFailureRate(
+	TEXT("mutable.RomLoadFailureRate"),
+	0.0f,
+	TEXT("Rate of induced failures for rom load testing."),
+	ECVF_Default);
 
 DECLARE_CYCLE_STAT(TEXT("MutableCoreTask"), STAT_MutableCoreTask, STATGROUP_Game);
 
 namespace UE::Mutable::Private
 {
-
 	TRACE_DECLARE_INT_COUNTER(MutableRuntime_OpenTask, TEXT("MutableRuntime/OpenTask"));
 	TRACE_DECLARE_INT_COUNTER(MutableRuntime_ClosedTasks, TEXT("MutableRuntime/ClosedTasks"));
 	TRACE_DECLARE_INT_COUNTER(MutableRuntime_IssuedTasks, TEXT("MutableRuntime/IssuedTasks"));
@@ -68,33 +72,16 @@ namespace UE::Mutable::Private
 	{
 		for (const FScheduledOp& Dep : Deps)
 		{
-			FCacheAddress At(Dep);
-	
-			// We need to increase hit count unconditionally for Full execution. 
-			if (Dep.Type == FScheduledOp::EType::Full)
-			{
-				LiveInstance->Cache->IncreaseHitCount(At);
-			}
-
-			if (!Dep.At)
-			{
-				continue;
-			}
+			bool bRequiresScheduling = GetMemory().AcquireCachedResultForScheduling(FCacheAddress(Dep));
 
-			bool bSuccess = GetMemory().AquireCachedResult(At);
-
-			if (!bSuccess)
+			if (bRequiresScheduling)
 			{
-				if (ScheduledStagePerOp.get(At) <= Dep.Stage)
-				{
 #if MUTABLE_DEBUG_CODERUNNER_TASK_SCHEDULE_CALLSTACK
-					FScheduledOp& AddedOp = OpenTasks.Add_GetRef(Dep);
-					AddedOp.StackDepth = FPlatformStackWalk::CaptureStackBackTrace(&AddedOp.ScheduleCallstack[0], FScheduledOp::CallstackMaxDepth);
+				FScheduledOp& AddedOp = OpenTasks.Add_GetRef(Dep);
+				AddedOp.StackDepth = FPlatformStackWalk::CaptureStackBackTrace(&AddedOp.ScheduleCallstack[0], FScheduledOp::CallstackMaxDepth);
 #else
-					OpenTasks.Add(Dep);
+				OpenTasks.Add(Dep);
 #endif
-					ScheduledStagePerOp[At] = Dep.Stage + 1;
-				}
 			}
 		}
 	}
@@ -133,10 +120,9 @@ namespace UE::Mutable::Private
 
 
 	void CodeRunner::LaunchIssuedTask( const TSharedPtr<FIssuedTask>& TaskToIssue, bool& bOutFailed )
-	{
-		bool bFailed = false;
-		bool bHasWork = TaskToIssue->Prepare(this, bFailed);
-		if (bFailed)
+	{	
+		bool bHasWork = TaskToIssue->Prepare(this, bOutFailed);
+		if (bOutFailed)
 		{
 			bUnrecoverableError = true;
 			return;
@@ -181,9 +167,40 @@ namespace UE::Mutable::Private
 		Run(MoveTemp(ProfileContext), bForceInlineExecution);
 	}
 
-	void CodeRunner::AbortRun()
+	void CodeRunner::LaunchWaitForOthersTask(Tasks::FTaskEvent RunnerRestartEvent, uint64 ReentryCounter)
 	{
-		bUnrecoverableError = true;
+		UE::Tasks::Launch(TEXT("CheckIfRunnerCanRunTask"),
+		[Runner = AsShared(), RunnerRestartEvent, ReentryCounter]() mutable
+		{
+			if (ReentryCounter > 128)
+			{
+				UE_LOGF(LogMutableCore, Warning, "LaunchWaitForOtherTask ReentryCounter is higher than exepcetd [%llu].", ReentryCounter);
+			}
+
+			int32 Index = Runner->ClosedTasks.Num() - 1;
+			for (; Index >= 0; --Index)
+			{
+				FProgramCache::EAcquireSetResult Result = Runner->GetMemory().TryAcquireCachedResultSet(Runner->ClosedTasks[Index].Deps);
+
+				if (Result != FProgramCache::EAcquireSetResult::Failure)
+				{
+					break;
+				}
+			}
+		
+			// If the loop exited early then we have work to do, trigger the event that will restart the runner.
+			// Otherwise, relaunch this task registering a new wait event.
+			if (Index >= 0)
+			{
+				RunnerRestartEvent.Trigger();
+			}
+			else
+			{
+				Runner->LaunchWaitForOthersTask(RunnerRestartEvent, ReentryCounter + 1);
+			}
+		},
+		UE::Tasks::Prerequisites(GetMemory().RegisterWaitEvent(TEXT("LaunchWaitForOthersTask"))),
+		UE::Tasks::ETaskPriority::Inherit);
 	}
 
     void CodeRunner::Run(TUniquePtr<FProfileContext>&& ProfileContext, bool bForceInlineExecution)
@@ -226,15 +243,6 @@ namespace UE::Mutable::Private
 				{
 					const FScheduledOp& item = IssuedTasks[Index]->Op;
 					bSuccess = IssuedTasks[Index]->Complete(this);
-
-					if (ScheduledStagePerOp[item] == item.Stage + 1)
-					{
-						// We completed everything that was requested, clear it otherwise if needed
-						// again it is not going to be rebuilt.
-						// \TODO: track rebuilds.
-						ScheduledStagePerOp[item] = 0;
-					}
-
 					IssuedTasks.RemoveAt(Index, EAllowShrinking::No); // with swap? changes order of execution.
 				}
 				else
@@ -243,11 +251,6 @@ namespace UE::Mutable::Private
 				}
 			}
 
-			if (!bSuccess)
-			{
-				return AbortRun();
-			}
-
 			while (!OpenTasks.IsEmpty())
 			{
 				// Get a new task to run
@@ -270,15 +273,9 @@ namespace UE::Mutable::Private
 
 				}
 
-				// Special processing in case it is an ImageDesc operation
-				if (Item.Type == FScheduledOp::EType::ImageDesc)
-				{
-					RunCodeImageDesc(Item);
-					continue;
-				}
-
 				FCacheAddress CachedAddress(Item);
-				if (GetMemory().AquireCachedResult(CachedAddress))
+				
+				if (GetMemory().TryAcquireCachedResult(CachedAddress))
 				{
 					continue;
 				}
@@ -293,7 +290,7 @@ namespace UE::Mutable::Private
 						LaunchIssuedTask(IssuedTask, bFailed);
 						if (bFailed)
 						{
-							return AbortRun();
+							GetMemory().SetAborted(FCacheAddress(IssuedTask->Op));
 						}
 					}
 					else
@@ -304,22 +301,24 @@ namespace UE::Mutable::Private
 				else
 				{
 					// Run immediately
-					RunCode(Item, LODMask);
-					if (bUnrecoverableError)
+					if (Item.Type == FScheduledOp::EType::Full)
 					{
-						return;
+						RunCode(Item, LODMask);
 					}
-
-					if (ScheduledStagePerOp[Item] == Item.Stage + 1)
+					else if (Item.Type == FScheduledOp::EType::ImageDesc)
 					{
-						// We completed everything that was requested, clear it. Otherwise if needed again it is not going to be rebuilt.
-						// \TODO: track operations that are run more than once?.
-						ScheduledStagePerOp[Item] = 0;
+						RunCodeImageDesc(Item);
+					}
+					else
+					{
+						unimplemented();
 					}
 				}
 
 				if (ProfileContext)
 				{
+					TScopeLock<FMutex> LockGuard(ProfileContext->Mutex);
+
 					++ProfileContext->NumRunOps;
 					++ProfileContext->RunOpsPerType[int32(LiveInstance->Model->GetProgram().GetOpType(Item.At))];
 				}
@@ -336,51 +335,62 @@ namespace UE::Mutable::Private
 				LaunchIssuedTask(TaskToIssue, bFailed);
 				if (bFailed)
 				{
-					return AbortRun();
+					GetMemory().SetAborted(FCacheAddress(TaskToIssue->Op));
 				}
 			}
 
 			// Look for a closed task with dependencies satisfied and move them to the open task list.
 			bool bSomeWasReady = false;
 			
-			for (int32 Index = 0; Index < ClosedTasks.Num();)
+			for (int32 Index = ClosedTasks.Num() - 1; Index >= 0; --Index)
 			{
-				if (GetMemory().TryAquireCachedResultSet(ClosedTasks[Index].Deps))
+				FProgramCache::EAcquireSetResult Result = GetMemory().TryAcquireCachedResultSet(ClosedTasks[Index].Deps);
+
+				if (Result == FProgramCache::EAcquireSetResult::Success)
 				{
 					OpenTasks.Push(ClosedTasks[Index].Op);
-					ClosedTasks.RemoveAt(Index, EAllowShrinking::No); // with swap? would change order of execution.
-					bSomeWasReady = true;
 				}
-				else
+				
+				if (Result == FProgramCache::EAcquireSetResult::Abort)
 				{
-					++Index;
+					GetMemory().SetAborted(ClosedTasks[Index].Op);
+					
+					// The parent task is aborted, make sure the 
+					for (const FCacheAddress& Dep : ClosedTasks[Index].Deps)
+					{
+						GetMemory().UpdateHitCount(Dep, -1);
+					}
 				}
+
+				if (Result != FProgramCache::EAcquireSetResult::Failure)
+				{
+					ClosedTasks.RemoveAtSwap(Index, EAllowShrinking::No);
+				}
+
+				bSomeWasReady |= Result != FProgramCache::EAcquireSetResult::Failure;
 			}
 
 			UpdateTraces();
 
 			// Debug: Did we dead-lock?
-			bool bDeadLock = !(OpenTasks.Num() || IssuedTasks.Num() || !ClosedTasks.Num() || bSomeWasReady);
-			if (bDeadLock)
+			bool bWaitingForOtherRunners = !(OpenTasks.Num() || IssuedTasks.Num() || !ClosedTasks.Num() || bSomeWasReady);
+			if (bWaitingForOtherRunners)
 			{
-				// Log the task graph
-				for (int32 Index = 0; Index < ClosedTasks.Num(); ++Index)
-				{	
-					FString TaskDesc = FString::Printf(TEXT("Closed task %d-%d-%d depends on : "), ClosedTasks[Index].Op.At, ClosedTasks[Index].Op.ExecutionIndex, ClosedTasks[Index].Op.Stage );
-					for (const FCacheAddress& Dep : ClosedTasks[Index].Deps)
-					{
-						if (Dep.At && !GetMemory().IsValid(Dep))
-						{
-							TaskDesc += FString::Printf(TEXT("%d-%d, "), Dep.At, Dep.ExecutionIndex);
-						}
-					}
+				GetMemory().TriggerWaitEvents();
 
-					UE_LOG(LogMutableCore, Log, TEXT("%s"), *TaskDesc);
-				}
-				check(false);
+				UE::Tasks::FTaskEvent WaitForOtherScheduleEvent(TEXT("WaitForOtherRunnersReScheduleEvent"));
+				LaunchWaitForOthersTask(WaitForOtherScheduleEvent);
+
+				UE::Tasks::AddNested(UE::Tasks::Launch(TEXT("CoderRunnerFromWaitForOtherRunners"),
+				[Runner = AsShared(), ProfileContext = MoveTemp(ProfileContext)]() mutable
+				{
+					constexpr bool bForceInlineExecution = false;
+					Runner->Run(MoveTemp(ProfileContext), bForceInlineExecution);
+				},
+				UE::Tasks::Prerequisites(WaitForOtherScheduleEvent),
+				UE::Tasks::ETaskPriority::Inherit));
 
-				// This should never happen but if it does, abort the code execution.
-				return AbortRun();
+				return;
 			}
 
 			// If at this point there is no open op and we haven't finished, we need to wait for an issued op to complete.
@@ -399,6 +409,8 @@ namespace UE::Mutable::Private
 						}
 					}
 
+					GetMemory().TriggerWaitEvents();
+
 					UE::Tasks::AddNested(UE::Tasks::Launch(TEXT("CodeRunnerFromIssuedTasksTask"),
 						[Runner = AsShared(), ProfileContext = MoveTemp(ProfileContext)]() mutable
 						{
@@ -429,6 +441,8 @@ namespace UE::Mutable::Private
 			{
 				if (FTimespan::FromSeconds(FPlatformTime::Seconds()) > TimeOut)
 				{
+					GetMemory().TriggerWaitEvents();
+
 					UE::Tasks::AddNested(UE::Tasks::Launch(TEXT("CodeRunnerFromTimeoutTask"),
 						[Runner = AsShared(), ProfileContext = MoveTemp(ProfileContext)]() mutable
 						{
@@ -442,6 +456,8 @@ namespace UE::Mutable::Private
 			}
 		}
 
+		GetMemory().TriggerWaitEvents();
+		
 		if (ProfileContext)
 		{
 			UE_LOG(LogMutableCore, Log, TEXT("Mutable Heap Bytes: %d"), HeapData.Num()* HeapData.GetTypeSize());
@@ -477,9 +493,9 @@ namespace UE::Mutable::Private
 
 	
 	/** */
-	const FExtendedImageDesc& CodeRunner::GetImageDescResult(OP::ADDRESS ResultAddress)
+	FExtendedImageDesc CodeRunner::GetImageDescResult(FOperation::ADDRESS ResultAddress)
 	{
-		FExtendedImageDesc& Result = ImageDescResults.FindChecked(ResultAddress);
+		FExtendedImageDesc Result = LiveInstance->Cache->LoadImageDesc(FCacheAddress(ResultAddress, 0, 0, FScheduledOp::EType::ImageDesc));
 		Result.ConstantImagesNeededToGenerate = ImageDescConstantImages;
 
 		return Result;
@@ -491,7 +507,7 @@ namespace UE::Mutable::Private
 	class FImageLayerTask : public CodeRunner::FIssuedTask
 	{
 	public:
-		FImageLayerTask(const FScheduledOp& InOp, const OP::ImageLayerArgs& InArgs)
+		FImageLayerTask(const FScheduledOp& InOp, const FOperation::ImageLayerArgs& InArgs)
 			: FIssuedTask(InOp), Args(InArgs)
 		{}
 
@@ -502,7 +518,7 @@ namespace UE::Mutable::Private
 
 	private:
 		int32 ImageCompressionQuality = 0;
-		OP::ImageLayerArgs Args;
+		FOperation::ImageLayerArgs Args;
 		TManagedPtr<const FImage> Blended;
 		TManagedPtr<const FImage> Mask;
 		TManagedPtr<FImage> Result;
@@ -563,7 +579,7 @@ namespace UE::Mutable::Private
 			Base = Formatted;
 		}
 
-		bool bMustHaveSameFormat = !(Args.flags & (OP::ImageLayerArgs::F_BASE_RGB_FROM_ALPHA | OP::ImageLayerArgs::F_BLENDED_RGB_FROM_ALPHA));
+		bool bMustHaveSameFormat = !(Args.flags & (FOperation::ImageLayerArgs::F_BASE_RGB_FROM_ALPHA | FOperation::ImageLayerArgs::F_BLENDED_RGB_FROM_ALPHA));
 		if (Blended && InitialFormat != Blended->GetFormat() && bMustHaveSameFormat)
 		{
 			TManagedPtr<FImage> Formatted = Runner->CreateImage(Blended->GetSizeX(), Blended->GetSizeY(), Blended->GetLODCount(), Base->GetFormat(), EInitializationType::NotInitialized);
@@ -633,10 +649,10 @@ namespace UE::Mutable::Private
 		
 		if (EBlendType(Args.blendType) != EBlendType::BT_NORMAL_COMBINE)
 		{
-			bool bApplyColorBlendToAlpha       = (Args.flags & OP::ImageLayerArgs::F_APPLY_TO_ALPHA        ) != 0;
-			bool bUseBaseSourceFromBaseAlpha   = (Args.flags & OP::ImageLayerArgs::F_BASE_RGB_FROM_ALPHA   ) != 0;
-			bool bUseBlendSourceFromBlendAlpha = (Args.flags & OP::ImageLayerArgs::F_BLENDED_RGB_FROM_ALPHA) != 0;
-			bool bUseMaskFromBlendAlpha        = (Args.flags & OP::ImageLayerArgs::F_USE_MASK_FROM_BLENDED ) != 0;
+			bool bApplyColorBlendToAlpha       = (Args.flags & FOperation::ImageLayerArgs::F_APPLY_TO_ALPHA        ) != 0;
+			bool bUseBaseSourceFromBaseAlpha   = (Args.flags & FOperation::ImageLayerArgs::F_BASE_RGB_FROM_ALPHA   ) != 0;
+			bool bUseBlendSourceFromBlendAlpha = (Args.flags & FOperation::ImageLayerArgs::F_BLENDED_RGB_FROM_ALPHA) != 0;
+			bool bUseMaskFromBlendAlpha        = (Args.flags & FOperation::ImageLayerArgs::F_USE_MASK_FROM_BLENDED ) != 0;
 
 			FBlendFuncType* ColorBlendFunc = SelectBlendFunc(EBlendType(Args.blendType));
 			FBlendFuncType* AlphaBlendFunc = nullptr; //ColorBlendFunc;
@@ -710,7 +726,7 @@ namespace UE::Mutable::Private
 	class FImageLayerColourTask : public CodeRunner::FIssuedTask
 	{
 	public:
-		FImageLayerColourTask(const FScheduledOp& InOp, const OP::ImageLayerColourArgs& InArgs)
+		FImageLayerColourTask(const FScheduledOp& InOp, const FOperation::ImageLayerColourArgs& InArgs)
 			: FIssuedTask(InOp), Args(InArgs)
 		{}
 
@@ -721,7 +737,7 @@ namespace UE::Mutable::Private
 
 	private:
 		int32 ImageCompressionQuality = 0;
-		OP::ImageLayerColourArgs Args;
+		FOperation::ImageLayerColourArgs Args;
 		FVector4f Color;
 		TManagedPtr<const FImage> Mask;
 		TManagedPtr<FImage> Result;
@@ -818,18 +834,20 @@ namespace UE::Mutable::Private
 	{
 		// This runs in a worker thread
 		MUTABLE_CPUPROFILER_SCOPE(FImageLayerColourTask);
+		// This runs in a worker thread.
+		bool bOnlyOneMip = false;
 
 		const int32 LODBegin = 0;
-		const int32 LODEnd   = Result->GetLODCount();
+		const int32 LODEnd = Result->GetLODCount();
 
 		if (EBlendType(Args.blendType) != EBlendType::BT_NORMAL_COMBINE)
 		{
 			FBlendFuncType* ColorBlendFunc = SelectBlendFunc(EBlendType(Args.blendType));
 			FBlendFuncType* AlphaBlendFunc = ColorBlendFunc;
 			
-			bool bApplyColorBlendToAlpha       = (Args.flags & OP::ImageLayerArgs::F_APPLY_TO_ALPHA        ) != 0;
-			bool bUseBaseSourceFromBaseAlpha   = (Args.flags & OP::ImageLayerArgs::F_BASE_RGB_FROM_ALPHA   ) != 0;
-			bool bUseBlendSourceFromBlendAlpha = (Args.flags & OP::ImageLayerArgs::F_BLENDED_RGB_FROM_ALPHA) != 0;
+			bool bApplyColorBlendToAlpha       = (Args.flags & FOperation::ImageLayerArgs::F_APPLY_TO_ALPHA        ) != 0;
+			bool bUseBaseSourceFromBaseAlpha   = (Args.flags & FOperation::ImageLayerArgs::F_BASE_RGB_FROM_ALPHA   ) != 0;
+			bool bUseBlendSourceFromBlendAlpha = (Args.flags & FOperation::ImageLayerArgs::F_BLENDED_RGB_FROM_ALPHA) != 0;
 			bool bUseMaskFromBlendAlpha = false; 
 
 			if (!bApplyColorBlendToAlpha)
@@ -850,10 +868,11 @@ namespace UE::Mutable::Private
 		else
 		{
 			ImageLayerCombineConstant(
-					Result.Get(), Result.Get(), Mask.Get(), Color,
-					OpImageLayerCombineOps::NormalCombine, 
-					LODBegin, LODEnd);
+				Result.Get(), Result.Get(), Mask.Get(), Color,
+				OpImageLayerCombineOps::NormalCombine,
+				LODBegin, LODEnd);
 		}
+
 		// Reset relevancy map.
 		Result->Flags &= ~FImage::EImageFlags::IF_HAS_RELEVANCY_MAP;
 	}
@@ -881,7 +900,7 @@ namespace UE::Mutable::Private
 	class FImagePixelFormatTask : public CodeRunner::FIssuedTask
 	{
 	public:
-		FImagePixelFormatTask(const FScheduledOp& InOp, const OP::ImagePixelFormatArgs& InArgs)
+		FImagePixelFormatTask(const FScheduledOp& InOp, const FOperation::ImagePixelFormatArgs& InArgs)
 			: FIssuedTask(InOp), Args(InArgs)
 		{}
 
@@ -892,7 +911,7 @@ namespace UE::Mutable::Private
 
 	private:
 		int ImageCompressionQuality = 0;
-		OP::ImagePixelFormatArgs Args;
+		FOperation::ImagePixelFormatArgs Args;
 		EImageFormat TargetFormat = EImageFormat::None;
 		TManagedPtr<const FImage> Base;
 		TManagedPtr<FImage> Result;
@@ -980,7 +999,7 @@ namespace UE::Mutable::Private
 	class FImageMipmapTask : public CodeRunner::FIssuedTask
 	{
 	public:
-		FImageMipmapTask(const FScheduledOp& InOp, const OP::ImageMipmapArgs& InArgs)
+		FImageMipmapTask(const FScheduledOp& InOp, const FOperation::ImageMipmapArgs& InArgs)
 			:FIssuedTask(InOp), Args(InArgs)
 		{}
 
@@ -992,7 +1011,7 @@ namespace UE::Mutable::Private
 	private:
 		int32 ImageCompressionQuality = 0;
 		int32 StartLevel = -1;
-		OP::ImageMipmapArgs Args;
+		FOperation::ImageMipmapArgs Args;
 		TManagedPtr<const FImage> Base;
 		TManagedPtr<FImage> Result;
 		FImageOperator::FScratchImageMipmap Scratch;
@@ -1101,7 +1120,7 @@ namespace UE::Mutable::Private
 	class FImageSwizzleTask : public CodeRunner::FIssuedTask
 	{
 	public:
-		FImageSwizzleTask(const FScheduledOp& InOp, const OP::ImageSwizzleArgs& InArgs)
+		FImageSwizzleTask(const FScheduledOp& InOp, const FOperation::ImageSwizzleArgs& InArgs)
 			: FIssuedTask(InOp), Args(InArgs)
 		{}
 
@@ -1111,7 +1130,7 @@ namespace UE::Mutable::Private
 		virtual bool Complete(CodeRunner*) override;
 
 	private:
-		OP::ImageSwizzleArgs Args;
+		FOperation::ImageSwizzleArgs Args;
 		TManagedPtr<const FImage> Sources[MUTABLE_OP_MAX_SWIZZLE_CHANNELS];
 		TManagedPtr<FImage> Result;
 	};
@@ -1245,7 +1264,7 @@ namespace UE::Mutable::Private
 	class FImageSaturateTask : public CodeRunner::FIssuedTask
 	{
 	public:
-		FImageSaturateTask(const FScheduledOp&, const OP::ImageSaturateArgs&);
+		FImageSaturateTask(const FScheduledOp&, const FOperation::ImageSaturateArgs&);
 
 		// FIssuedTask interface
 		virtual bool Prepare(CodeRunner*, bool& bOutFailed) override;
@@ -1253,14 +1272,14 @@ namespace UE::Mutable::Private
 		virtual bool Complete(CodeRunner*) override;
 
 	private:
-		const OP::ImageSaturateArgs Args;
+		const FOperation::ImageSaturateArgs Args;
 		TManagedPtr<FImage> Result;
 		float Factor;
 	};
 
 
 	//---------------------------------------------------------------------------------------------
-	FImageSaturateTask::FImageSaturateTask(const FScheduledOp& InOp, const OP::ImageSaturateArgs& InArgs)
+	FImageSaturateTask::FImageSaturateTask(const FScheduledOp& InOp, const FOperation::ImageSaturateArgs& InArgs)
 		: FIssuedTask(InOp)
 		, Args(InArgs)
 	{
@@ -1325,7 +1344,7 @@ namespace UE::Mutable::Private
 	class FImageResizeTask : public CodeRunner::FIssuedTask
 	{
 	public:
-		FImageResizeTask(const FScheduledOp& InOp, const OP::ImageResizeArgs& InArgs)
+		FImageResizeTask(const FScheduledOp& InOp, const FOperation::ImageResizeArgs& InArgs)
 			: FIssuedTask(InOp), Args(InArgs)
 		{}
 
@@ -1336,7 +1355,7 @@ namespace UE::Mutable::Private
 
 	private:
 		int32 ImageCompressionQuality = 0;
-		OP::ImageResizeArgs Args;
+		FOperation::ImageResizeArgs Args;
 		TManagedPtr<const FImage> Base;
 		TManagedPtr<FImage> Result;
 		FImageOperator::FImagePixelFormatFunc ImagePixelFormatFunc;
@@ -1442,7 +1461,7 @@ namespace UE::Mutable::Private
 	class FImageResizeRelTask : public CodeRunner::FIssuedTask
 	{
 	public:
-		FImageResizeRelTask(const FScheduledOp& InOp, const OP::ImageResizeRelArgs& InArgs)
+		FImageResizeRelTask(const FScheduledOp& InOp, const FOperation::ImageResizeRelArgs& InArgs)
 			: FIssuedTask(InOp), Args(InArgs)
 		{}
 
@@ -1452,7 +1471,7 @@ namespace UE::Mutable::Private
 		virtual bool Complete(CodeRunner*) override;
 
 	private:
-		OP::ImageResizeRelArgs Args;
+		FOperation::ImageResizeRelArgs Args;
 		int32 ImageCompressionQuality=0;
 		TManagedPtr<const FImage> Base;
 		TManagedPtr<FImage> Result;
@@ -1553,7 +1572,7 @@ namespace UE::Mutable::Private
 	class FImageInvertTask : public CodeRunner::FIssuedTask
 	{
 	public:
-		FImageInvertTask(const FScheduledOp& InOp, const OP::ImageInvertArgs& InArgs)
+		FImageInvertTask(const FScheduledOp& InOp, const FOperation::ImageInvertArgs& InArgs)
 			: FIssuedTask(InOp), Args(InArgs)
 		{}
 
@@ -1564,7 +1583,7 @@ namespace UE::Mutable::Private
 
 	private:
 		TManagedPtr<FImage> Result;
-		OP::ImageInvertArgs Args;
+		FOperation::ImageInvertArgs Args;
 	};
 
 
@@ -1614,7 +1633,7 @@ namespace UE::Mutable::Private
 	class FImageComposeTask : public CodeRunner::FIssuedTask
 	{
 	public:
-		FImageComposeTask(const FScheduledOp& InOp, const OP::ImageComposeArgs& InArgs, const TManagedPtr<const FLayout>& InLayout)
+		FImageComposeTask(const FScheduledOp& InOp, const FOperation::ImageComposeArgs& InArgs, const TManagedPtr<const FLayout>& InLayout)
 			: FIssuedTask(InOp), Args(InArgs), Layout(InLayout)
 		{}
 
@@ -1625,7 +1644,7 @@ namespace UE::Mutable::Private
 
 	private:
 		int32 ImageCompressionQuality = 0;
-		OP::ImageComposeArgs Args;
+		FOperation::ImageComposeArgs Args;
 		TManagedPtr<const FLayout> Layout;
 		TManagedPtr<const FImage> Block;
 		TManagedPtr<const FImage> Mask;
@@ -1793,6 +1812,21 @@ namespace UE::Mutable::Private
 			return false;
 		}
 
+#if 1 // Enable load failure test.
+		{
+			float RomLoadFailureRate = CVarRomLoadFailureRate.GetValueOnAnyThread();	
+
+			if (!FMath::IsNearlyZero(RomLoadFailureRate))
+			{
+				if (FMath::FRand() < RomLoadFailureRate)
+				{
+					UE_LOG(LogMutableCore, Warning, TEXT("Mesh rom was forced to faile with CVar."));
+					bOutFailed = true;
+					return false;
+				}
+			}
+		}
+#endif
 		bOutFailed = false;
 
 		const FProgram& LocalProgram = Runner->Program;
@@ -1877,7 +1911,7 @@ namespace UE::Mutable::Private
 		// This runs in the Runner thread
 		MUTABLE_CPUPROFILER_SCOPE(FLoadMeshRomTask_Complete);
 		
-		OP::MeshConstantArgs Args = Runner->Program.GetOpArgs<OP::MeshConstantArgs>(Op.At);
+		FOperation::MeshConstantArgs Args = Runner->Program.GetOpArgs<FOperation::MeshConstantArgs>(Op.At);
 		
 		const TPassthroughObjectPtr<UClothingAssetCommon> ClothAsset = Runner->FindPassthrough<UClothingAssetCommon>(Args.ClothID);
 
@@ -1916,6 +1950,24 @@ namespace UE::Mutable::Private
 
 		// This runs in the mutable Runner thread
 		MUTABLE_CPUPROFILER_SCOPE(FLoadImageRomsTask_Prepare);
+
+
+#if 1 // Enable load failure test.
+		{
+			float RomLoadFailureRate = CVarRomLoadFailureRate.GetValueOnAnyThread();	
+
+			if (!FMath::IsNearlyZero(RomLoadFailureRate))
+			{
+				if (FMath::FRand() < RomLoadFailureRate)
+				{
+					UE_LOG(LogMutableCore, Warning, TEXT("Image rom was forced to fail with CVar."));
+					bOutFailed = true;
+					return false;
+				}
+			}
+		}
+#endif
+
 		bOutFailed = false;
 
 		const FProgram& LocalProgram = Runner->Program;
@@ -1960,7 +2012,7 @@ namespace UE::Mutable::Private
 	//---------------------------------------------------------------------------------------------
 	bool CodeRunner::FLoadImageRomsTask::Complete(CodeRunner* Runner)
 	{
-		OP::ResourceConstantArgs Args = Runner->Model.GetProgram().GetOpArgs<OP::ResourceConstantArgs>(Op.At);
+		FOperation::ResourceConstantArgs Args = Runner->Model.GetProgram().GetOpArgs<FOperation::ResourceConstantArgs>(Op.At);
 
 		const int32 MipsToSkip = Op.ExecutionOptions;
 		TManagedPtr<const FImage> Source;
@@ -2136,6 +2188,11 @@ namespace UE::Mutable::Private
 	//---------------------------------------------------------------------------------------------
 	TSharedPtr<CodeRunner::FIssuedTask> CodeRunner::IssueOp(FScheduledOp Item)
 	{
+		if (Item.Type != FScheduledOp::EType::Full)
+		{
+			return nullptr;
+		}
+
 		TSharedPtr<FIssuedTask> Issued;
 		
 		EOpType type = Program.GetOpType(Item.At);
@@ -2144,7 +2201,7 @@ namespace UE::Mutable::Private
 		{
 		case EOpType::ME_CONSTANT:
 		{
-			OP::MeshConstantArgs Args = Program.GetOpArgs<OP::MeshConstantArgs>(Item.At);
+			FOperation::MeshConstantArgs Args = Program.GetOpArgs<FOperation::MeshConstantArgs>(Item.At);
 		
 			const FMeshContentRange MeshContentRange = Program.ConstantMeshes[Args.Value];
 			const EMeshContentFlags ContentFilterFlags = static_cast<EMeshContentFlags>(Item.ExecutionOptions);
@@ -2156,7 +2213,7 @@ namespace UE::Mutable::Private
 
 		case EOpType::IM_CONSTANT:
 		{
-			OP::ResourceConstantArgs args = Program.GetOpArgs<OP::ResourceConstantArgs>(Item.At);
+			FOperation::ResourceConstantArgs args = Program.GetOpArgs<FOperation::ResourceConstantArgs>(Item.At);
 			int32 MipsToSkip = Item.ExecutionOptions;
 			int32 ImageIndex = args.value;
 			int32 ReallySkip = FMath::Min(MipsToSkip, Program.ConstantImages[ImageIndex].LODCount - 1);
@@ -2206,7 +2263,7 @@ namespace UE::Mutable::Private
 		{
 			if (Item.Stage == 1)
 			{
-				OP::ImageParameterConvertArgs Args = Program.GetOpArgs<OP::ImageParameterConvertArgs>(Item.At);
+				FOperation::ImageParameterConvertArgs Args = Program.GetOpArgs<FOperation::ImageParameterConvertArgs>(Item.At);
 			
 				TManagedPtr<const FImage> Image = LoadImage(FCacheAddress(Args.ImageParameter, Item));
 
@@ -2222,7 +2279,7 @@ namespace UE::Mutable::Private
 
 		case EOpType::IM_PARAMETER_FROM_MATERIAL:
 		{
-			const OP::MaterialBreakImageParameterArgs Args = Program.GetOpArgs<OP::MaterialBreakImageParameterArgs>(Item.At);
+			const FOperation::MaterialBreakImageParameterArgs Args = Program.GetOpArgs<FOperation::MaterialBreakImageParameterArgs>(Item.At);
 			TSharedPtr<FRangeIndex> Index = BuildCurrentOpRangeIndex(Item, Args.MaterialParameter);
 
 			// Get material parameter from the array of parameters
@@ -2257,7 +2314,7 @@ namespace UE::Mutable::Private
 
 		case EOpType::IM_REFERENCE:
 		{
-			OP::ResourceReferenceArgs Args = Program.GetOpArgs<OP::ResourceReferenceArgs>(Item.At);
+			FOperation::ResourceReferenceArgs Args = Program.GetOpArgs<FOperation::ResourceReferenceArgs>(Item.At);
 
 			// We only convert references to images if indicated in the operation.
 			if (Args.ForceLoad)
@@ -2278,7 +2335,7 @@ namespace UE::Mutable::Private
 		{
 			if (Item.Stage == 1)
 			{
-				OP::MeshSkeletalMeshBreakArgs Args = Program.GetOpArgs<OP::MeshSkeletalMeshBreakArgs>(Item.At);
+				FOperation::MeshSkeletalMeshBreakArgs Args = Program.GetOpArgs<FOperation::MeshSkeletalMeshBreakArgs>(Item.At);
 
 				TManagedPtr<const FSkeletalMesh> SkeletalMesh = LoadSkeletalMesh(FCacheAddress(Args.SkeletalMeshParameter, Item));
 
@@ -2297,7 +2354,7 @@ namespace UE::Mutable::Private
 		{
 			if (Item.Stage == 1)
 			{
-				OP::ImagePixelFormatArgs Args = Program.GetOpArgs<OP::ImagePixelFormatArgs>(Item.At);
+				FOperation::ImagePixelFormatArgs Args = Program.GetOpArgs<FOperation::ImagePixelFormatArgs>(Item.At);
 				Issued = MakeShared<FImagePixelFormatTask>(Item, Args);
 			}
 			break;
@@ -2307,7 +2364,7 @@ namespace UE::Mutable::Private
 		{
 			if (Item.Stage == 1)
 			{
-				OP::ImageLayerColourArgs Args = Program.GetOpArgs<OP::ImageLayerColourArgs>(Item.At);
+				FOperation::ImageLayerColourArgs Args = Program.GetOpArgs<FOperation::ImageLayerColourArgs>(Item.At);
 				Issued = MakeShared<FImageLayerColourTask>(Item, Args);
 			}
 			break;
@@ -2320,7 +2377,7 @@ namespace UE::Mutable::Private
 				(ExecutionStrategy != EExecutionStrategy::MinimizeMemory && Item.Stage == 1)
 				)
 			{
-				OP::ImageLayerArgs Args = Program.GetOpArgs<OP::ImageLayerArgs>(Item.At);
+				FOperation::ImageLayerArgs Args = Program.GetOpArgs<FOperation::ImageLayerArgs>(Item.At);
 				Issued = MakeShared<FImageLayerTask>(Item, Args);
 			}
 			break;
@@ -2330,7 +2387,7 @@ namespace UE::Mutable::Private
 		{
 			if (Item.Stage == 1)
 			{
-				OP::ImageMipmapArgs Args = Program.GetOpArgs<OP::ImageMipmapArgs>(Item.At);
+				FOperation::ImageMipmapArgs Args = Program.GetOpArgs<FOperation::ImageMipmapArgs>(Item.At);
 				Issued = MakeShared<FImageMipmapTask>(Item, Args);
 			}
 			break;
@@ -2340,7 +2397,7 @@ namespace UE::Mutable::Private
 		{
 			if (Item.Stage == 1)
 			{
-				OP::ImageSwizzleArgs Args = Program.GetOpArgs<OP::ImageSwizzleArgs>(Item.At);
+				FOperation::ImageSwizzleArgs Args = Program.GetOpArgs<FOperation::ImageSwizzleArgs>(Item.At);
 				Issued = MakeShared<FImageSwizzleTask>(Item, Args);
 			}
 			break;
@@ -2350,7 +2407,7 @@ namespace UE::Mutable::Private
 		{
 			if (Item.Stage == 1)
 			{
-				OP::ImageSaturateArgs Args = Program.GetOpArgs<OP::ImageSaturateArgs>(Item.At);
+				FOperation::ImageSaturateArgs Args = Program.GetOpArgs<FOperation::ImageSaturateArgs>(Item.At);
 				Issued = MakeShared<FImageSaturateTask>(Item, Args);
 			}
 			break;
@@ -2360,7 +2417,7 @@ namespace UE::Mutable::Private
 		{
 			if (Item.Stage == 1)
 			{
-				OP::ImageInvertArgs Args = Program.GetOpArgs<OP::ImageInvertArgs>(Item.At);
+				FOperation::ImageInvertArgs Args = Program.GetOpArgs<FOperation::ImageInvertArgs>(Item.At);
 				Issued = MakeShared<FImageInvertTask>(Item, Args);
 			}
 			break;
@@ -2370,7 +2427,7 @@ namespace UE::Mutable::Private
 		{
 			if (Item.Stage == 1)
 			{
-				OP::ImageResizeArgs Args = Program.GetOpArgs<OP::ImageResizeArgs>(Item.At);
+				FOperation::ImageResizeArgs Args = Program.GetOpArgs<FOperation::ImageResizeArgs>(Item.At);
 				Issued = MakeShared<FImageResizeTask>(Item, Args);
 			}
 			break;
@@ -2380,7 +2437,7 @@ namespace UE::Mutable::Private
 		{
 			if (Item.Stage == 1)
 			{
-				OP::ImageResizeRelArgs Args = Program.GetOpArgs<OP::ImageResizeRelArgs>(Item.At);
+				FOperation::ImageResizeRelArgs Args = Program.GetOpArgs<FOperation::ImageResizeRelArgs>(Item.At);
 				Issued = MakeShared<FImageResizeRelTask>(Item, Args);
 			}
 			break;
@@ -2391,7 +2448,7 @@ namespace UE::Mutable::Private
 			if ((ExecutionStrategy == EExecutionStrategy::MinimizeMemory && Item.Stage == 3) ||
 				(ExecutionStrategy != EExecutionStrategy::MinimizeMemory && Item.Stage == 2))
 			{
-				OP::ImageComposeArgs Args = Program.GetOpArgs<OP::ImageComposeArgs>(Item.At);
+				FOperation::ImageComposeArgs Args = Program.GetOpArgs<FOperation::ImageComposeArgs>(Item.At);
 				TManagedPtr<const FLayout> ComposeLayout = StaticCastManagedPtr<const FLayout>( HeapData[Item.CustomState].Resource);
 				Issued = MakeShared<FImageComposeTask>(Item, Args, ComposeLayout);
 			}
diff --git a/Engine/Plugins/Mutable/Source/MutableRuntime/Private/MuR/External/MaterialAdapter.cpp b/Engine/Plugins/Mutable/Source/MutableRuntime/Private/MuR/External/MaterialAdapter.cpp
--- a/Engine/Plugins/Mutable/Source/MutableRuntime/Private/MuR/External/MaterialAdapter.cpp
+++ b/Engine/Plugins/Mutable/Source/MutableRuntime/Private/MuR/External/MaterialAdapter.cpp
@@ -113,7 +113,7 @@ namespace UE::Mutable
 			return {};
 		}
 
-		TVariant<Private::OP::ADDRESS, Private::TManagedPtr<const Private::FImage>> Variant = Result->ImageParameter;
+		TVariant<Private::FOperation::ADDRESS, Private::TManagedPtr<const Private::FImage>> Variant = Result->ImageParameter;
 
 		if (ensure(Variant.IsType<Private::TManagedPtr<const Private::FImage>>()))
 		{
@@ -141,7 +141,7 @@ namespace UE::Mutable
 
 		if (Value)
 		{
-			TVariant<Private::OP::ADDRESS, Private::TManagedPtr<const Private::FImage>> Image;
+			TVariant<Private::FOperation::ADDRESS, Private::TManagedPtr<const Private::FImage>> Image;
 			Image.Set<Private::TManagedPtr<const Private::FImage>>(Value->Get().Image);
 			
 			Material->ImageParameters.Add(ParameterKey, { Image, ImageData->ImagePropertyIndex });
